Fix race in FCounters used in FileSystemCacheStore and AsyncLoading

#rnx

diff --git a/Engine/Source/Developer/DerivedDataCache/Private/FileSystemCacheStore.cpp b/Engine/Source/Developer/DerivedDataCache/Private/FileSystemCacheStore.cpp
--- a/Engine/Source/Developer/DerivedDataCache/Private/FileSystemCacheStore.cpp
+++ b/Engine/Source/Developer/DerivedDataCache/Private/FileSystemCacheStore.cpp
@@ -52,12 +52,12 @@
 namespace UE::DerivedData
 {
 
-TRACE_DECLARE_INT_COUNTER(FileSystemDDC_Get, TEXT("FileSystemDDC Get"));
-TRACE_DECLARE_INT_COUNTER(FileSystemDDC_GetHit, TEXT("FileSystemDDC Get Hit"));
-TRACE_DECLARE_INT_COUNTER(FileSystemDDC_Put, TEXT("FileSystemDDC Put"));
-TRACE_DECLARE_INT_COUNTER(FileSystemDDC_PutHit, TEXT("FileSystemDDC Put Hit"));
-TRACE_DECLARE_INT_COUNTER(FileSystemDDC_BytesRead, TEXT("FileSystemDDC Bytes Read"));
-TRACE_DECLARE_INT_COUNTER(FileSystemDDC_BytesWritten, TEXT("FileSystemDDC Bytes Written"));
+TRACE_DECLARE_ATOMIC_INT_COUNTER(FileSystemDDC_Get, TEXT("FileSystemDDC Get"));
+TRACE_DECLARE_ATOMIC_INT_COUNTER(FileSystemDDC_GetHit, TEXT("FileSystemDDC Get Hit"));
+TRACE_DECLARE_ATOMIC_INT_COUNTER(FileSystemDDC_Put, TEXT("FileSystemDDC Put"));
+TRACE_DECLARE_ATOMIC_INT_COUNTER(FileSystemDDC_PutHit, TEXT("FileSystemDDC Put Hit"));
+TRACE_DECLARE_ATOMIC_INT_COUNTER(FileSystemDDC_BytesRead, TEXT("FileSystemDDC Bytes Read"));
+TRACE_DECLARE_ATOMIC_INT_COUNTER(FileSystemDDC_BytesWritten, TEXT("FileSystemDDC Bytes Written"));
 
 ///////////////////////////////////////////////////////////////////////////////////////////////////////////////////////
 
diff --git a/Engine/Source/Runtime/Core/Public/ProfilingDebugging/CountersTrace.h b/Engine/Source/Runtime/Core/Public/ProfilingDebugging/CountersTrace.h
--- a/Engine/Source/Runtime/Core/Public/ProfilingDebugging/CountersTrace.h
+++ b/Engine/Source/Runtime/Core/Public/ProfilingDebugging/CountersTrace.h
@@ -8,6 +8,8 @@
 #include "Trace/Config.h"
 #include "Trace/Trace.h"
 
+#include <atomic>
+
 namespace UE { namespace Trace { class FChannel; } }
 
 #if !defined(COUNTERSTRACE_ENABLED)
@@ -40,7 +42,7 @@ struct FCountersTrace
 	CORE_API static void OutputSetValue(uint16 CounterId, int64 Value);
 	CORE_API static void OutputSetValue(uint16 CounterId, double Value);
 
-	template<typename ValueType, ETraceCounterType CounterType>
+	template<typename ValueType, ETraceCounterType CounterType, typename StoredType = ValueType>
 	class TCounter
 	{
 	public:
@@ -98,13 +100,14 @@ struct FCountersTrace
 		}
 
 	private:
-		ValueType Value;
+		StoredType Value;
 		const TCHAR* CounterName;
 		uint16 CounterId;
 		ETraceCounterDisplayHint CounterDisplayHint;
 	};
 
 	using FCounterInt = TCounter<int64, TraceCounterType_Int>;
+	using FCounterAtomicInt = TCounter<int64, TraceCounterType_Int, std::atomic<int64>>;
 	using FCounterFloat = TCounter<double, TraceCounterType_Float>;
 };
 
@@ -126,9 +129,15 @@ struct FCountersTrace
 #define TRACE_DECLARE_INT_COUNTER(CounterName, CounterDisplayName) \
 	FCountersTrace::FCounterInt PREPROCESSOR_JOIN(__GTraceCounter, CounterName)(CounterDisplayName, TraceCounterDisplayHint_None);
 
+#define TRACE_DECLARE_ATOMIC_INT_COUNTER(CounterName, CounterDisplayName) \
+	FCountersTrace::FCounterAtomicInt PREPROCESSOR_JOIN(__GTraceCounter, CounterName)(CounterDisplayName, TraceCounterDisplayHint_None);
+
 #define TRACE_DECLARE_INT_COUNTER_EXTERN(CounterName) \
 	extern FCountersTrace::FCounterInt PREPROCESSOR_JOIN(__GTraceCounter, CounterName);
 
+#define TRACE_DECLARE_ATOMIC_INT_COUNTER_EXTERN(CounterName) \
+	extern FCountersTrace::FCounterAtomicInt PREPROCESSOR_JOIN(__GTraceCounter, CounterName);
+
 #define TRACE_DECLARE_FLOAT_COUNTER(CounterName, CounterDisplayName) \
 	FCountersTrace::FCounterFloat PREPROCESSOR_JOIN(__GTraceCounter, CounterName)(CounterDisplayName, TraceCounterDisplayHint_None);
 
@@ -141,6 +150,12 @@ struct FCountersTrace
 #define TRACE_DECLARE_MEMORY_COUNTER_EXTERN(CounterName) \
 	TRACE_DECLARE_INT_COUNTER_EXTERN(CounterName)
 
+#define TRACE_DECLARE_ATOMIC_MEMORY_COUNTER(CounterName, CounterDisplayName) \
+	FCountersTrace::FCounterAtomicInt PREPROCESSOR_JOIN(__GTraceCounter, CounterName)(CounterDisplayName, TraceCounterDisplayHint_Memory);
+
+#define TRACE_DECLARE_ATOMIC_MEMORY_COUNTER_EXTERN(CounterName) \
+	TRACE_DECLARE_ATOMIC_INT_COUNTER_EXTERN(CounterName)
+
 #define TRACE_COUNTER_SET(CounterName, Value) \
 	PREPROCESSOR_JOIN(__GTraceCounter, CounterName).Set(Value);
 
@@ -163,10 +178,14 @@ struct FCountersTrace
 #define TRACE_MEMORY_VALUE(CounterDisplayName, Value)
 #define TRACE_DECLARE_INT_COUNTER(CounterName, CounterDisplayName)
 #define TRACE_DECLARE_INT_COUNTER_EXTERN(CounterName)
+#define TRACE_DECLARE_ATOMIC_INT_COUNTER(CounterName, CounterDisplayName)
+#define TRACE_DECLARE_ATOMIC_INT_COUNTER_EXTERN(CounterName)
 #define TRACE_DECLARE_FLOAT_COUNTER(CounterName, CounterDisplayName)
 #define TRACE_DECLARE_FLOAT_COUNTER_EXTERN(CounterName)
 #define TRACE_DECLARE_MEMORY_COUNTER(CounterName, CounterDisplayName)
 #define TRACE_DECLARE_MEMORY_COUNTER_EXTERN(CounterName)
+#define TRACE_DECLARE_ATOMIC_MEMORY_COUNTER(CounterName, CounterDisplayName)
+#define TRACE_DECLARE_ATOMIC_MEMORY_COUNTER_EXTERN(CounterName)
 #define TRACE_COUNTER_SET(CounterName, Value)
 #define TRACE_COUNTER_ADD(CounterName, Value)
 #define TRACE_COUNTER_SUBTRACT(CounterName, Value)
diff --git a/Engine/Source/Runtime/CoreUObject/Private/Serialization/AsyncLoading2.cpp b/Engine/Source/Runtime/CoreUObject/Private/Serialization/AsyncLoading2.cpp
--- a/Engine/Source/Runtime/CoreUObject/Private/Serialization/AsyncLoading2.cpp
+++ b/Engine/Source/Runtime/CoreUObject/Private/Serialization/AsyncLoading2.cpp
@@ -350,11 +350,11 @@ static FAutoConsoleVariableRef CVarGOnlyProcessRequiredPackagesWhenSyncLoading(
 CSV_DECLARE_CATEGORY_MODULE_EXTERN(CORE_API, Basic);
 CSV_DECLARE_CATEGORY_MODULE_EXTERN(CORE_API, FileIO);
 
-TRACE_DECLARE_INT_COUNTER(AsyncLoadingQueuedPackages, TEXT("AsyncLoading/PackagesQueued"));
-TRACE_DECLARE_INT_COUNTER(AsyncLoadingLoadingPackages, TEXT("AsyncLoading/PackagesLoading"));
-TRACE_DECLARE_INT_COUNTER(AsyncLoadingPackagesWithRemainingWork, TEXT("AsyncLoading/PackagesWithRemainingWork"));
-TRACE_DECLARE_INT_COUNTER(AsyncLoadingPendingIoRequests, TEXT("AsyncLoading/PendingIoRequests"));
-TRACE_DECLARE_MEMORY_COUNTER(AsyncLoadingTotalLoaded, TEXT("AsyncLoading/TotalLoaded"));
+TRACE_DECLARE_ATOMIC_INT_COUNTER(AsyncLoadingQueuedPackages, TEXT("AsyncLoading/PackagesQueued"));
+TRACE_DECLARE_ATOMIC_INT_COUNTER(AsyncLoadingLoadingPackages, TEXT("AsyncLoading/PackagesLoading"));
+TRACE_DECLARE_ATOMIC_INT_COUNTER(AsyncLoadingPackagesWithRemainingWork, TEXT("AsyncLoading/PackagesWithRemainingWork"));
+TRACE_DECLARE_ATOMIC_INT_COUNTER(AsyncLoadingPendingIoRequests, TEXT("AsyncLoading/PendingIoRequests"));
+TRACE_DECLARE_ATOMIC_MEMORY_COUNTER(AsyncLoadingTotalLoaded, TEXT("AsyncLoading/TotalLoaded"));
 
 struct FAsyncPackage2;
 class FAsyncLoadingThread2;
@@ -6941,8 +6941,7 @@ EAsyncPackageState::Type FAsyncLoadingThread2::ProcessLoadedPackagesFromGameThre
 	{
 #if WITH_EDITOR
 		// In editor builds, call the asset load callback. This happens in both editor and standalone to match EndLoad
-		TSet<FWeakObjectPtr> TempLoadedAssets = LoadedAssets;
-		LoadedAssets.Reset();
+		TSet<FWeakObjectPtr> TempLoadedAssets = MoveTemp(LoadedAssets);
 
 		// Make a copy because LoadedAssets could be modified by one of the OnAssetLoaded callbacks
 		for (const FWeakObjectPtr& WeakAsset : TempLoadedAssets)
