Add SetNum to TBitArray to be able to resize an array with a default value for any added bits.

diff --git a/Engine/Source/Runtime/Core/Public/Containers/BitArray.h b/Engine/Source/Runtime/Core/Public/Containers/BitArray.h
--- a/Engine/Source/Runtime/Core/Public/Containers/BitArray.h
+++ b/Engine/Source/Runtime/Core/Public/Containers/BitArray.h
@@ -742,6 +742,19 @@ public:
 		ClearPartialSlackBits();
 	}
 
+	/** Sets the number of bits, initializing any potentially added bits to the given value. */
+	template <typename ValueType>
+	FORCENOINLINE void SetNum(int32 InNumBits, ValueType bValue)
+	{
+		static_assert(std::is_same_v<ValueType, bool>, "TBitArray::SetNum: unexpected type passed as the bValue argument (expected bool)");
+		int32 PreviousNumBits = NumBits;
+		SetNumUninitialized(InNumBits);
+		if (InNumBits > PreviousNumBits)
+		{
+			SetRange(PreviousNumBits, InNumBits - PreviousNumBits, bValue);
+		}
+	}
+
 	/**
 	 * Sets or unsets a range of bits within the array.
 	 * @param Index  The index of the first bit to set; must be 0 <= Index <= Num().
