Tasks System: handling null prerequisites

diff --git a/Engine/Source/Runtime/Core/Public/Tasks/TaskPrivate.h b/Engine/Source/Runtime/Core/Public/Tasks/TaskPrivate.h
--- a/Engine/Source/Runtime/Core/Public/Tasks/TaskPrivate.h
+++ b/Engine/Source/Runtime/Core/Public/Tasks/TaskPrivate.h
@@ -218,7 +218,7 @@ namespace UE::Tasks
 			template<typename HigherLevelTaskType, decltype(std::declval<HigherLevelTaskType>().Pimpl)* = nullptr>
 			bool AddPrerequisites(const HigherLevelTaskType& Prerequisite)
 			{
-				return AddPrerequisites(*Prerequisite.Pimpl);
+				return Prerequisite.IsValid() ? AddPrerequisites(*Prerequisite.Pimpl) : false;
 			}
 
 			// The task will be executed only when all prerequisites are completed.
@@ -255,6 +255,11 @@ namespace UE::Tasks
 						Prerequisite = Prereq.Pimpl;
 					}
 
+					if (Prerequisite == nullptr)
+					{
+						continue;
+					}
+
 					if (Prerequisite->AddSubsequent(*this)) // acq_rel memory order
 					{
 						Prerequisite->AddRef(); // keep it alive until this task's execution
