Disabled TBB for ARM64.
Re-enabled mimalloc for ARM64.

diff --git a/Engine/Source/Runtime/Core/Public/Windows/WindowsPlatform.h b/Engine/Source/Runtime/Core/Public/Windows/WindowsPlatform.h
--- a/Engine/Source/Runtime/Core/Public/Windows/WindowsPlatform.h
+++ b/Engine/Source/Runtime/Core/Public/Windows/WindowsPlatform.h
@@ -79,12 +79,12 @@ typedef FWindowsPlatformTypes FPlatformTypes;
 #define PLATFORM_HAS_BSD_SOCKET_FEATURE_WINSOCKETS			1
 #define PLATFORM_USES_MICROSOFT_LIBC_FUNCTIONS				1
 #define PLATFORM_IS_ANSI_MALLOC_THREADSAFE					1
-#define PLATFORM_SUPPORTS_TBB								1
 #if PLATFORM_CPU_ARM_FAMILY
-#define PLATFORM_SUPPORTS_MIMALLOC							0
+#define PLATFORM_SUPPORTS_TBB								0
 #else
-#define PLATFORM_SUPPORTS_MIMALLOC							PLATFORM_64BITS
+#define PLATFORM_SUPPORTS_TBB								1
 #endif
+#define PLATFORM_SUPPORTS_MIMALLOC							PLATFORM_64BITS
 #define PLATFORM_SUPPORTS_NAMED_PIPES						1
 #define PLATFORM_COMPILER_HAS_TCHAR_WMAIN					1
 #define PLATFORM_SUPPORTS_EARLY_MOVIE_PLAYBACK				(!WITH_EDITOR) // movies will start before engine is initalized
