LLT test adapter: correct return of error state
#rnx

diff --git a/Engine/Source/Runtime/Core/Public/Misc/LowLevelTestAdapter.h b/Engine/Source/Runtime/Core/Public/Misc/LowLevelTestAdapter.h
--- a/Engine/Source/Runtime/Core/Public/Misc/LowLevelTestAdapter.h
+++ b/Engine/Source/Runtime/Core/Public/Misc/LowLevelTestAdapter.h
@@ -39,7 +39,7 @@
 			void TestBody(const FString& Parameters); \
 			virtual bool RunTest(const FString& Parameters) { \
 				TestBody(Parameters); \
-				return true; \
+				return HasAnyErrors(); \
 			} \
 			virtual FString GetBeautifiedTestName() const override { return PrettyNameDotNotation; } \
 		private:\
@@ -63,6 +63,7 @@
 #define LLT_STR(Macro) #Macro
 #define LLT_STR_EXPAND(Macro) LLT_STR(Macro)
 #define TEST_CASE_GENERATED_NAME_UNIQUE_STR LLT_STR_EXPAND(TEST_CASE_GENERATED_NAME_UNIQUE)
+// Note: TEST_CASE uses unique names which only work when used inside unique namespace in the same compilation unit?// Use TEST_CASE_NAMED instead and provide an unique global instance name
 #define TEST_CASE(PrettyName, TFlags) TEST_CASE_NAMED_STR(TEST_CASE_GENERATED_NAME_UNIQUE, TEST_CASE_GENERATED_NAME_UNIQUE_STR, PrettyName, TFlags)
 #define TEST_CASE_NAMED(ClassName, PrettyName, TFlags) TEST_CASE_NAMED_STR(ClassName, #ClassName, PrettyName, TFlags)
 
