Implemented async loading of localization data, and extended dynamic loading to handle explicitly loaded plugins

This change takes the existing async loading that was added for game localization data in CL# 14661627 and expands it to cover all localization data. Async localization tasks are chained together as dependencies to preserve the existing order of operations (plus going wide doesn't help due to contention on the live-table), and can be waited on at any point (eg, after a loading screen) by calling FTextLocalizationManager::WaitForAsyncTasks (though the engine init flow takes care of waiting for you unless you're loading additional plugin/chunked localization data).

This change takes the existing support for dynamically loading chunked localization data on PAK chunk load, and extends it to support dynamically loading any additional localization data. This is used by the engine to support dynamically loading the localization targets associated with explicitly loaded plugins. These loads are all async, so must be waited on if you need them to block.

Note: BeginInitGameTextLocalization and EndInitGameTextLocalization have been removed as part of this change. You would now just use InitGameTextLocalization (which is async) followed by FTextLocalizationManager::WaitForAsyncTasks.

diff --git a/Engine/Plugins/Experimental/GameFeatures/Source/GameFeatures/Private/GameFeaturePluginStateMachine.cpp b/Engine/Plugins/Experimental/GameFeatures/Source/GameFeatures/Private/GameFeaturePluginStateMachine.cpp
--- a/Engine/Plugins/Experimental/GameFeatures/Source/GameFeatures/Private/GameFeaturePluginStateMachine.cpp
+++ b/Engine/Plugins/Experimental/GameFeatures/Source/GameFeatures/Private/GameFeaturePluginStateMachine.cpp
@@ -2250,6 +2250,13 @@ struct FGameFeaturePluginState_Activating : public FGameFeaturePluginState
 
 		UGameFeaturesSubsystem::Get().OnGameFeatureActivating(StateProperties.GameFeatureData, StateProperties.PluginName, Context, StateProperties.PluginIdentifier.GetFullPluginURL());
 
+		// If this plugin caused localization data to load, wait for that here before marking it as active
+		if (TSharedPtr<IPlugin> Plugin = IPluginManager::Get().FindPlugin(StateProperties.PluginName);
+			Plugin && Plugin->GetDescriptor().bExplicitlyLoaded && Plugin->GetDescriptor().LocalizationTargets.Num() > 0)
+		{
+			FTextLocalizationManager::Get().WaitForAsyncTasks();
+		}
+
 		StateStatus.SetTransition(EGameFeaturePluginState::Active);
 	}
 };
diff --git a/Engine/Source/Runtime/Core/Private/Internationalization/PolyglotTextSource.cpp b/Engine/Source/Runtime/Core/Private/Internationalization/PolyglotTextSource.cpp
--- a/Engine/Source/Runtime/Core/Private/Internationalization/PolyglotTextSource.cpp
+++ b/Engine/Source/Runtime/Core/Private/Internationalization/PolyglotTextSource.cpp
@@ -9,11 +9,14 @@
 #include "Internationalization/TextLocalizationResource.h"
 #include "Misc/AssertionMacros.h"
 #include "Misc/EnumClassFlags.h"
+#include "Misc/ScopeLock.h"
 #include "Templates/Tuple.h"
 #include "Templates/UnrealTemplate.h"
 
 bool FPolyglotTextSource::GetNativeCultureName(const ELocalizedTextSourceCategory InCategory, FString& OutNativeCultureName)
 {
+	FScopeLock ScopeLock(&PolyglotDataCS);
+
 	if (FCultureInfo* CultureInfo = AvailableCultureInfo.Find(InCategory))
 	{
 		for (const auto& NativeCulturePair : CultureInfo->NativeCultures)
@@ -22,6 +25,7 @@ bool FPolyglotTextSource::GetNativeCultureName(const ELocalizedTextSourceCategor
 			return true;
 		}
 	}
+
 	return false;
 }
 
@@ -43,6 +47,8 @@ void FPolyglotTextSource::GetLocalizedCultureNames(const ELocalizationLoadFlags
 		}
 	};
 
+	FScopeLock ScopeLock(&PolyglotDataCS);
+
 	if (EnumHasAnyFlags(InLoadFlags, ELocalizationLoadFlags::Editor))
 	{
 		AppendCulturesForCategory(ELocalizedTextSourceCategory::Editor);
@@ -162,6 +168,8 @@ void FPolyglotTextSource::AddPolyglotDataToResource(const FPolyglotTextData& InP
 
 void FPolyglotTextSource::LoadLocalizedResources(const ELocalizationLoadFlags InLoadFlags, TArrayView<const FString> InPrioritizedCultures, FTextLocalizationResource& InOutNativeResource, FTextLocalizationResource& InOutLocalizedResource)
 {
+	FScopeLock ScopeLock(&PolyglotDataCS);
+
 	for (const auto& PolyglotTextDataPair : PolyglotTextDataMap)
 	{
 		const FPolyglotTextData& PolyglotTextData = PolyglotTextDataPair.Value;
@@ -171,11 +179,14 @@ void FPolyglotTextSource::LoadLocalizedResources(const ELocalizationLoadFlags In
 
 EQueryLocalizedResourceResult FPolyglotTextSource::QueryLocalizedResource(const ELocalizationLoadFlags InLoadFlags, TArrayView<const FString> InPrioritizedCultures, const FTextId InTextId, FTextLocalizationResource& InOutNativeResource, FTextLocalizationResource& InOutLocalizedResource)
 {
+	FScopeLock ScopeLock(&PolyglotDataCS);
+
 	if (const FPolyglotTextData* PolyglotTextData = PolyglotTextDataMap.Find(InTextId))
 	{
 		AddPolyglotDataToResource(*PolyglotTextData, InLoadFlags, InPrioritizedCultures, InOutNativeResource, InOutLocalizedResource);
 		return EQueryLocalizedResourceResult::Found;
 	}
+
 	return EQueryLocalizedResourceResult::NotFound;
 }
 
@@ -183,6 +194,8 @@ void FPolyglotTextSource::RegisterPolyglotTextData(const FPolyglotTextData& InPo
 {
 	check(InPolyglotTextData.IsValid());
 
+	FScopeLock ScopeLock(&PolyglotDataCS);
+
 	const FTextId Identity = FTextId(FTextKey(*InPolyglotTextData.GetNamespace()), FTextKey(*InPolyglotTextData.GetKey()));
 	if (FPolyglotTextData* CurrentPolyglotData = PolyglotTextDataMap.Find(Identity))
 	{
diff --git a/Engine/Source/Runtime/Core/Private/Internationalization/PolyglotTextSource.h b/Engine/Source/Runtime/Core/Private/Internationalization/PolyglotTextSource.h
--- a/Engine/Source/Runtime/Core/Private/Internationalization/PolyglotTextSource.h
+++ b/Engine/Source/Runtime/Core/Private/Internationalization/PolyglotTextSource.h
@@ -3,6 +3,7 @@
 #pragma once
 
 #include "CoreTypes.h"
+#include "HAL/CriticalSection.h"
 #include "Internationalization/ILocalizedTextSource.h"
 #include "Internationalization/LocKeyFuncs.h"
 #include "Internationalization/PolyglotTextData.h" // IWYU pragma: keep
@@ -56,4 +57,9 @@ private:
 	 * Mapping from a "{Namespace}::{Key}" string to a polyglot text data instance.
 	 */
 	TMap<FTextId, FPolyglotTextData> PolyglotTextDataMap;
+
+	/**
+	 * Critical section preventing concurrent access to the above data.
+	 */
+	mutable FCriticalSection PolyglotDataCS;
 };
diff --git a/Engine/Source/Runtime/Core/Private/Internationalization/TextLocalizationManager.cpp b/Engine/Source/Runtime/Core/Private/Internationalization/TextLocalizationManager.cpp
--- a/Engine/Source/Runtime/Core/Private/Internationalization/TextLocalizationManager.cpp
+++ b/Engine/Source/Runtime/Core/Private/Internationalization/TextLocalizationManager.cpp
@@ -26,10 +26,32 @@
 #include "Misc/App.h"
 #include "Misc/CoreDelegates.h"
 #include "Templates/UniquePtr.h"
+#include "Async/Async.h"
 #include "Async/TaskGraphInterfaces.h"
 
 DEFINE_LOG_CATEGORY_STATIC(LogTextLocalizationManager, Log, All);
 
+namespace TextLocalizationManager
+{
+static bool AsyncLoadLocalizationData = true;
+static FAutoConsoleVariableRef CVarAsyncLoadLocalizationData(TEXT("Localization.AsyncLoadLocalizationData"), AsyncLoadLocalizationData, TEXT("True to load localization data asynchronously (non-blocking), or False to load it synchronously (blocking)"));
+
+/**
+ * Note: This is disabled by default because we have existing code that conflates "the language changed" with "there is new localization data available".
+ * These places should be audited to replace FInternationalization::OnCultureChanged callbacks with FTextLocalizationManager::OnTextRevisionChangedEvent.
+ * 
+ * The most troublesome place to handle is the font and composite font caches, which flush on a language change (because the language can affect which fonts 
+ * will be used, and is generally a good point to clear any current font cache data), which results in the font cache re-filling with glyphs for the previous 
+ * language (as the new localization data hasn't loaded yet). Additionally composite fonts with per-language sub-fonts can briefly show text using the wrong 
+ * set of glyphs (eg, when switching from "ja" -> "zh-Hans" you'll see Japanese text using Chinese-style Han before the Chinese text loads in).
+ * 
+ * We don't want to flush these font caches every time new localization data is available, so we'll likely need extra context in OnTextRevisionChangedEvent 
+ * to know what caused the data to change (eg, LanguageChanged, RefreshRequested, AdditionalDataLoaded, etc) so that the callback can behave accordingly.
+ */
+static bool AsyncLoadLocalizationDataOnLanguageChange = false;
+static FAutoConsoleVariableRef CVarAsyncLoadLocalizationDataOnLanguageChange(TEXT("Localization.AsyncLoadLocalizationDataOnLanguageChange"), AsyncLoadLocalizationDataOnLanguageChange, TEXT("True to load localization data asynchronously (non-blocking) when the language changes, or False to load it synchronously (blocking)"));
+}
+
 enum class ERequestedCultureOverrideLevel : uint8
 {
 	CommandLine,
@@ -352,6 +374,9 @@ void InitEngineTextLocalization()
 	SCOPED_BOOT_TIMING("InitEngineTextLocalization");
 	DECLARE_SCOPE_CYCLE_COUNTER(TEXT("InitEngineTextLocalization"), STAT_InitEngineTextLocalization, STATGROUP_LoadTime);
 
+	// Make sure any async tasks have finished before we proceed, as this init function may update non-thread safe data
+	FTextLocalizationManager::Get().WaitForAsyncTasks();
+
 	// Make sure the String Table Registry is initialized as it may trigger module loads
 	FStringTableRegistry::Get();
 	FStringTableRedirects::InitStringTableRedirects();
@@ -388,13 +413,14 @@ void InitEngineTextLocalization()
 	TextLocalizationResourceUtil::ClearNativeEditorCultureName();
 #endif
 
-	FTextLocalizationManager::Get().LoadLocalizationResourcesForCulture(FInternationalization::Get().GetCurrentLanguage()->GetName(), LocLoadFlags);
-	FTextLocalizationManager::Get().InitializedFlags = FTextLocalizationManager::Get().InitializedFlags.load() | ETextLocalizationManagerInitializedFlags::Engine;
+	FTextLocalizationManager::Get().QueueAsyncTask([LocLoadFlags, AvailableTextSources = FTextLocalizationManager::Get().LocalizedTextSources]()
+	{
+		FTextLocalizationManager::Get().LoadLocalizationResourcesForCulture_Sync(AvailableTextSources, FInternationalization::Get().GetCurrentLanguage()->GetName(), LocLoadFlags);
+		FTextLocalizationManager::Get().InitializedFlags = FTextLocalizationManager::Get().InitializedFlags.load() | ETextLocalizationManagerInitializedFlags::Engine;
+	});
 }
 
-static FGraphEventRef InitGameTextLocalizationTask;
-
-void BeginInitGameTextLocalization()
+void InitGameTextLocalization()
 {
 	if (!FApp::IsGame())
 	{
@@ -403,9 +429,12 @@ void BeginInitGameTextLocalization()
 	}
 
 	LLM_SCOPE(ELLMTag::Localization);
-
+	SCOPED_BOOT_TIMING("InitGameTextLocalization");
 	DECLARE_SCOPE_CYCLE_COUNTER(TEXT("InitGameTextLocalization"), STAT_InitGameTextLocalization, STATGROUP_LoadTime);
 
+	// Make sure any async tasks have finished before we proceed, as this init function may update non-thread safe data
+	FTextLocalizationManager::Get().WaitForAsyncTasks();
+
 	// Refresh the cached config data before applying the default culture, as the game may have patched in new config data since the cache was built
 	FInternationalization::Get().RefreshCachedConfigData();
 
@@ -431,42 +460,16 @@ void BeginInitGameTextLocalization()
 		LocLoadFlags |= ELocalizationLoadFlags::Additional;
 	}
 
-	FTextLocalizationManager::Get().InitializedFlags = FTextLocalizationManager::Get().InitializedFlags.load() | ETextLocalizationManagerInitializedFlags::Initializing;
-	auto TaskLambda = [LocLoadFlags, InitializedFlags = FTextLocalizationManager::Get().InitializedFlags.load()]()
+	FTextLocalizationManager::Get().QueueAsyncTask([LocLoadFlags, AvailableTextSources = FTextLocalizationManager::Get().LocalizedTextSources]()
 	{
-		SCOPED_BOOT_TIMING("InitGameTextLocalization");
+		FTextLocalizationManager::Get().LoadLocalizationResourcesForCulture_Sync(AvailableTextSources, FInternationalization::Get().GetCurrentLanguage()->GetName(), LocLoadFlags);
+		FTextLocalizationManager::Get().InitializedFlags = FTextLocalizationManager::Get().InitializedFlags.load() | ETextLocalizationManagerInitializedFlags::Game;
 
-		FTextLocalizationManager::Get().LoadLocalizationResourcesForCulture(FInternationalization::Get().GetCurrentLanguage()->GetName(), LocLoadFlags);
-		FTextLocalizationManager::Get().InitializedFlags = (InitializedFlags & ~ETextLocalizationManagerInitializedFlags::Initializing) | ETextLocalizationManagerInitializedFlags::Game;
 		//FTextLocalizationManager::Get().DumpMemoryInfo();
 		//Worse when Compacting because we remove growth space and force new growth space to be reallocated the next time an entry is added which is going to be bigger than what we removed anyway...
 		//FTextLocalizationManager::Get().CompactDataStructures();
 		//FTextLocalizationManager::Get().DumpMemoryInfo();
-	};
-	if (FTaskGraphInterface::IsRunning())
-	{
-		InitGameTextLocalizationTask = FFunctionGraphTask::CreateAndDispatchWhenReady(MoveTemp(TaskLambda), TStatId());
-	}
-	else
-	{
-		TaskLambda();
-	}
-}
-
-void EndInitGameTextLocalization()
-{
-	SCOPED_BOOT_TIMING("WaitForInitGameTextLocalization");
-	if (InitGameTextLocalizationTask)
-	{
-		FTaskGraphInterface::Get().WaitUntilTaskCompletes(InitGameTextLocalizationTask);
-		InitGameTextLocalizationTask.SafeRelease();
-	}
-}
-
-void InitGameTextLocalization()
-{
-	BeginInitGameTextLocalization();
-	EndInitGameTextLocalization();
+	});
 }
 
 FTextLocalizationManager& FTextLocalizationManager::Get()
@@ -476,6 +479,8 @@ FTextLocalizationManager& FTextLocalizationManager::Get()
 
 void FTextLocalizationManager::TearDown()
 {
+	FTextLocalizationManager::Get().WaitForAsyncTasks();
+
 	FTextCache::TearDown();
 	TLazySingleton<FTextLocalizationManager>::TearDown();
 	FTextKey::TearDown();
@@ -491,6 +496,11 @@ FTextLocalizationManager::FTextLocalizationManager()
 	RegisterTextSource(PolyglotTextSource.ToSharedRef(), bRefreshResources);
 }
 
+FTextLocalizationManager::~FTextLocalizationManager()
+{
+	// Note: Explicit destructor needed since AsyncLocalizationTask is using a forward declared type in the header
+}
+
 void FTextLocalizationManager::DumpMemoryInfo() const
 {
 	{
@@ -569,8 +579,6 @@ TArray<FString> FTextLocalizationManager::GetLocalizedCultureNames(const ELocali
 
 void FTextLocalizationManager::RegisterTextSource(const TSharedRef<ILocalizedTextSource>& InLocalizedTextSource, const bool InRefreshResources)
 {
-	ensureMsgf(!IsInitializing(), TEXT("Localized text source registered during game text initialization"));
-
 	LocalizedTextSources.Add(InLocalizedTextSource);
 	LocalizedTextSources.StableSort([](const TSharedPtr<ILocalizedTextSource>& InLocalizedTextSourceOne, const TSharedPtr<ILocalizedTextSource>& InLocalizedTextSourceTwo)
 	{
@@ -901,7 +909,7 @@ void FTextLocalizationManager::UpdateFromLocalizationResource(const FString& Loc
 {
 	FTextLocalizationResource TextLocalizationResource;
 	TextLocalizationResource.LoadFromFile(LocalizationResourceFilePath, 0);
-	UpdateFromLocalizationResource(TextLocalizationResource);
+	UpdateFromLocalizations(MoveTemp(TextLocalizationResource));
 }
 
 void FTextLocalizationManager::UpdateFromLocalizationResource(const FTextLocalizationResource& TextLocalizationResource)
@@ -909,10 +917,23 @@ void FTextLocalizationManager::UpdateFromLocalizationResource(const FTextLocaliz
 	UpdateFromLocalizations(CopyTemp(TextLocalizationResource));
 }
 
-void FTextLocalizationManager::RefreshResources()
+void FTextLocalizationManager::WaitForAsyncTasks()
 {
-	ensureMsgf(!IsInitializing(), TEXT("Reloading text localization resources during game text initialization"));
+	SCOPED_BOOT_TIMING("FTextLocalizationManager::WaitForAsyncTasks");
+	DECLARE_SCOPE_CYCLE_COUNTER(TEXT("FTextLocalizationManager::WaitForAsyncTasks"), STAT_WaitForAsyncLocalizationTasks, STATGROUP_LoadTime);
+
+	if (AsyncLocalizationTask)
+	{
+		if (FTaskGraphInterface::IsRunning())
+		{
+			FTaskGraphInterface::Get().WaitUntilTaskCompletes(AsyncLocalizationTask);
+		}
+		AsyncLocalizationTask.SafeRelease();
+	}
+}
 
+void FTextLocalizationManager::RefreshResources()
+{
 	ELocalizationLoadFlags LocLoadFlags = ELocalizationLoadFlags::None;
 	LocLoadFlags |= (WITH_EDITOR ? ELocalizationLoadFlags::Editor : ELocalizationLoadFlags::None);
 	LocLoadFlags |= (FApp::IsGame() ? ELocalizationLoadFlags::Game : ELocalizationLoadFlags::None);
@@ -920,149 +941,53 @@ void FTextLocalizationManager::RefreshResources()
 	LocLoadFlags |= ELocalizationLoadFlags::Native;
 	LocLoadFlags |= ELocalizationLoadFlags::Additional;
 
-	LoadLocalizationResourcesForCulture(FInternationalization::Get().GetCurrentLanguage()->GetName(), LocLoadFlags);
+	LoadLocalizationResourcesForCulture_Async(FInternationalization::Get().GetCurrentLanguage()->GetName(), LocLoadFlags);
 }
 
-void FTextLocalizationManager::OnPakFileMounted(const IPakFile& PakFile)
+void FTextLocalizationManager::HandleLocalizationTargetsMounted(TArrayView<const FString> LocalizationTargetPaths)
 {
-	SCOPED_BOOT_TIMING("FTextLocalizationManager::OnPakFileMounted");
-	LLM_SCOPE(ELLMTag::Localization);
-
-	int32 ChunkId = PakFile.PakGetPakchunkIndex();
-	if (ChunkId == INDEX_NONE || ChunkId == 0 || PakFile.GetNumFiles() == 0)
-	{
-		// Skip empty (IoStore), non-chunked PAK files, and chunk 0 as that contains the standard localization data
-		return;
-	}
-
-	UE_LOG(LogTextLocalizationManager, Verbose, TEXT("Request to load localization data for chunk %d (from PAK '%s')"), ChunkId, *PakFile.PakGetPakFilename());
-
-	// Skip this request if we've already loaded the data for this chunk via the request for a previous PAK sub-file load notification
-	if (LocResTextSource->HasRegisteredChunkId(ChunkId))
-	{
-		UE_LOG(LogTextLocalizationManager, Verbose, TEXT("Skipped loading localization data for chunk %d (from PAK '%s') as this chunk has already been processed"), ChunkId, *PakFile.PakGetPakFilename());
-		return;
-	}
-	
-	// If we're being notified so early that even InitEngineTextLocalization hasn't run, then we can't safely make the queries below as things like GConfig may not be available yet!
 	if (!IsInitialized())
 	{
-		// Track this so that full resource refreshes (eg, changing culture) work as expected
-		LocResTextSource->RegisterChunkId(ChunkId);
-
-		UE_LOG(LogTextLocalizationManager, Verbose, TEXT("Skipped loading localization data for chunk %d (from PAK '%s') as the localization manager isn't ready"), ChunkId, *PakFile.PakGetPakFilename());
+		// If we've not yet loaded localization data then there's nothing to do
 		return;
 	}
 
-	ensureMsgf(!IsInitializing(), TEXT("Pak file mounted during game text initialization"));
-
-	// Note: We only allow game localization targets to be chunked, and the layout is assumed to follow our standard pattern (as used by the localization dashboard and FLocTextHelper)
-	const TArray<FString> ChunkedLocalizationTargets = FLocalizationResourceTextSource::GetChunkedLocalizationTargets();
-
-	// Check to see whether all the required localization data is now available
-	// This may not be the case if this PAK was split into multiple sub-files, and the localization data was split between them
-	TArray<FString> PrioritizedLocalizationPaths;
-	for (const FString& LocalizationTarget : ChunkedLocalizationTargets)
-	{
-		const FString ChunkedLocalizationTargetName = TextLocalizationResourceUtil::GetLocalizationTargetNameForChunkId(LocalizationTarget, ChunkId);
-
-		FString ChunkedLocalizationTargetPath = FPaths::ProjectContentDir() / TEXT("Localization") / ChunkedLocalizationTargetName;
-		if (!IFileManager::Get().DirectoryExists(*ChunkedLocalizationTargetPath))
-		{
-			UE_LOG(LogTextLocalizationManager, Verbose, TEXT("Skipped loading localization data for chunk %d (from PAK '%s') as the localization directory for '%s' was not yet available"), ChunkId, *PakFile.PakGetPakFilename(), *ChunkedLocalizationTargetName);
-			return;
-		}
-
-		FTextLocalizationMetaDataResource LocMetaResource;
-		{
-			const FString LocMetaFilename = ChunkedLocalizationTargetPath / FString::Printf(TEXT("%s.locmeta"), *ChunkedLocalizationTargetName);
-			if (!IFileManager::Get().FileExists(*LocMetaFilename))
-			{
-				UE_LOG(LogTextLocalizationManager, Verbose, TEXT("Skipped loading localization data for chunk %d (from PAK '%s') as the LocMeta file for '%s' was not yet available"), ChunkId, *PakFile.PakGetPakFilename(), *ChunkedLocalizationTargetName);
-				return;
-			}
-			if (!LocMetaResource.LoadFromFile(LocMetaFilename))
-			{
-				UE_LOG(LogTextLocalizationManager, Verbose, TEXT("Skipped loading localization data for chunk %d (from PAK '%s') as the LocMeta file for '%s' failed to load"), ChunkId, *PakFile.PakGetPakFilename(), *ChunkedLocalizationTargetName);
-				return;
-			}
-		}
-
-		for (const FString& CompiledCulture : LocMetaResource.CompiledCultures)
-		{
-			const FString LocResFilename = ChunkedLocalizationTargetPath / CompiledCulture / FString::Printf(TEXT("%s.locres"), *ChunkedLocalizationTargetName);
-			if (!IFileManager::Get().FileExists(*LocResFilename))
-			{
-				UE_LOG(LogTextLocalizationManager, Verbose, TEXT("Skipped loading localization data for chunk %d (from PAK '%s') as the '%s' LocRes file for '%s' was not yet available"), ChunkId, *PakFile.PakGetPakFilename(), *CompiledCulture, *ChunkedLocalizationTargetName);
-				return;
-			}
-		}
+	ELocalizationLoadFlags LocLoadFlags = ELocalizationLoadFlags::None;
+	LocLoadFlags |= (WITH_EDITOR ? ELocalizationLoadFlags::Editor : ELocalizationLoadFlags::None);
+	LocLoadFlags |= (FApp::IsGame() ? ELocalizationLoadFlags::Game : ELocalizationLoadFlags::None);
+	LocLoadFlags |= ELocalizationLoadFlags::Engine;
 
-		PrioritizedLocalizationPaths.Add(MoveTemp(ChunkedLocalizationTargetPath));
-	}
+	const TArray<FString> PrioritizedCultureNames = FInternationalization::Get().GetPrioritizedCultureNames(FInternationalization::Get().GetCurrentLanguage()->GetName());
 
-	// Track this so that full resource refreshes (eg, changing culture) work as expected
-	LocResTextSource->RegisterChunkId(ChunkId);
+	LoadLocalizationTargetsForPrioritizedCultures_Async(LocalizationTargetPaths, PrioritizedCultureNames, LocLoadFlags);
+}
 
-	if (!EnumHasAnyFlags(InitializedFlags.load(), ETextLocalizationManagerInitializedFlags::Game))
+void FTextLocalizationManager::HandleLocalizationTargetsUnmounted(TArrayView<const FString> LocalizationTargetPaths)
+{
+	if (!IsInitialized())
 	{
-		// If we've not yet initialized game localization then don't bother patching, as the full initialization path will load the data for this chunk
+		// If we've not yet loaded localization data then there's nothing to do
 		return;
 	}
+	
+	// Note: We don't track which LocRes text comes from (as it uses too much memory), so unloading text would require a full refresh (including flushing DisplayStringLookupTable)
+	//       RefreshResources cannot currently do a flush, so for now we ignore any unmount notifications
+	//RefreshResources();
+}
 
-	// Load the resources from each target in this chunk
-	const TArray<FString> PrioritizedCultureNames = FInternationalization::Get().GetPrioritizedCultureNames(FInternationalization::Get().GetCurrentLanguage()->GetName());
-	const ELocalizationLoadFlags LocLoadFlags = ELocalizationLoadFlags::Game;
-	FTextLocalizationResource UnusedNativeResource;
-	FTextLocalizationResource LocalizedResource;
-	for (const FString& PrioritizedLocalizationPath : PrioritizedLocalizationPaths)
-	{
-		UE_LOG(LogTextLocalizationManager, Verbose, TEXT("Loading chunked localization data from '%s'"), *PrioritizedLocalizationPath);
-	}
-	LocResTextSource->LoadLocalizedResourcesFromPaths(TArrayView<FString>(), PrioritizedLocalizationPaths, TArrayView<FString>(), LocLoadFlags, PrioritizedCultureNames, UnusedNativeResource, LocalizedResource);
+void FTextLocalizationManager::OnPakFileMounted(const IPakFile& PakFile)
+{
+	SCOPED_BOOT_TIMING("FTextLocalizationManager::OnPakFileMounted");
+	LLM_SCOPE(ELLMTag::Localization);
 
-	// Allow any higher priority text sources to override the text loaded for the chunk (eg, to allow polyglot hot-fixes to take priority)
-	// Note: If any text sources don't support dynamic queries, then we must do a much slower full refresh instead :(
-	bool bNeedsFullRefresh = false;
+	int32 ChunkId = PakFile.PakGetPakchunkIndex();
+	if (ChunkId == INDEX_NONE || ChunkId == 0 || PakFile.GetNumFiles() == 0)
 	{
-		// Copy the IDs array as QueryLocalizedResource can update the map
-		TArray<FTextId> ChunkTextIds;
-		LocalizedResource.Entries.GenerateKeyArray(ChunkTextIds);
-
-		for (const TSharedPtr<ILocalizedTextSource>& LocalizedTextSource : LocalizedTextSources)
-		{
-			if (LocalizedTextSource->GetPriority() <= LocResTextSource->GetPriority())
-			{
-				continue;
-			}
-
-			for (const FTextId& ChunkTextId : ChunkTextIds)
-			{
-				if (LocalizedTextSource->QueryLocalizedResource(LocLoadFlags, PrioritizedCultureNames, ChunkTextId, UnusedNativeResource, LocalizedResource) == EQueryLocalizedResourceResult::NotImplemented)
-				{
-					bNeedsFullRefresh = true;
-					break;
-				}
-			}
-
-			if (bNeedsFullRefresh)
-			{
-				break;
-			}
-		}
+		// Skip empty (IoStore), non-chunked PAK files, and chunk 0 as that contains the standard localization data
+		return;
 	}
 
-	// Apply the new data
-	if (bNeedsFullRefresh)
-	{
-		UE_LOG(LogTextLocalizationManager, Verbose, TEXT("Patching chunked localization data failed, performing full refresh"));
-		RefreshResources();
-	}
-	else
-	{
-		UE_LOG(LogTextLocalizationManager, Verbose, TEXT("Patching chunked localization data for %d entries"), LocalizedResource.Entries.Num());
-		UpdateFromLocalizations(MoveTemp(LocalizedResource), /*bDirtyTextRevision*/true);
-	}
+	LoadChunkedLocalizationResources_Async(ChunkId, PakFile.PakGetPakFilename());
 }
 
 void FTextLocalizationManager::OnCultureChanged()
@@ -1074,19 +999,15 @@ void FTextLocalizationManager::OnCultureChanged()
 		return;
 	}
 
-	ensureMsgf(!IsInitializing(), TEXT("Culture changed during game text initialization"));
-
-	ELocalizationLoadFlags LocLoadFlags = ELocalizationLoadFlags::None;
-	LocLoadFlags |= (WITH_EDITOR ? ELocalizationLoadFlags::Editor : ELocalizationLoadFlags::None);
-	LocLoadFlags |= (FApp::IsGame() ? ELocalizationLoadFlags::Game : ELocalizationLoadFlags::None);
-	LocLoadFlags |= ELocalizationLoadFlags::Engine;
-	LocLoadFlags |= ELocalizationLoadFlags::Native;
-	LocLoadFlags |= ELocalizationLoadFlags::Additional;
+	RefreshResources();
 
-	LoadLocalizationResourcesForCulture(FInternationalization::Get().GetCurrentLanguage()->GetName(), LocLoadFlags);
+	if (!TextLocalizationManager::AsyncLoadLocalizationDataOnLanguageChange)
+	{
+		WaitForAsyncTasks();
+	}
 }
 
-void FTextLocalizationManager::LoadLocalizationResourcesForCulture(const FString& CultureName, const ELocalizationLoadFlags LocLoadFlags)
+void FTextLocalizationManager::LoadLocalizationResourcesForCulture_Sync(TArrayView<const TSharedPtr<ILocalizedTextSource>> AvailableTextSources, const FString& CultureName, const ELocalizationLoadFlags LocLoadFlags)
 {
     LLM_SCOPE(ELLMTag::Localization);
 
@@ -1103,12 +1024,21 @@ void FTextLocalizationManager::LoadLocalizationResourcesForCulture(const FString
 		return;
 	}
 
-	LoadLocalizationResourcesForPrioritizedCultures(FInternationalization::Get().GetPrioritizedCultureNames(CultureName), LocLoadFlags);
+	LoadLocalizationResourcesForPrioritizedCultures_Sync(AvailableTextSources, FInternationalization::Get().GetPrioritizedCultureNames(CultureName), LocLoadFlags);
 }
 
-void FTextLocalizationManager::LoadLocalizationResourcesForPrioritizedCultures(TArrayView<const FString> PrioritizedCultureNames, const ELocalizationLoadFlags LocLoadFlags)
+void FTextLocalizationManager::LoadLocalizationResourcesForCulture_Async(const FString& CultureName, const ELocalizationLoadFlags LocLoadFlags)
+{
+	QueueAsyncTask([AvailableTextSources = LocalizedTextSources, CultureName, LocLoadFlags]()
+	{
+		FTextLocalizationManager::Get().LoadLocalizationResourcesForCulture_Sync(AvailableTextSources, CultureName, LocLoadFlags);
+	});
+}
+
+void FTextLocalizationManager::LoadLocalizationResourcesForPrioritizedCultures_Sync(TArrayView<const TSharedPtr<ILocalizedTextSource>> AvailableTextSources, TArrayView<const FString> PrioritizedCultureNames, const ELocalizationLoadFlags LocLoadFlags)
 {
 	LLM_SCOPE(ELLMTag::Localization);
+	TRACE_CPUPROFILER_EVENT_SCOPE(FTextLocalizationManager::LoadLocalizationResourcesForPrioritizedCultures);
 
 	// Nothing to do?
 	if (PrioritizedCultureNames.Num() == 0)
@@ -1130,7 +1060,7 @@ void FTextLocalizationManager::LoadLocalizationResourcesForPrioritizedCultures(T
 	// Load the resources from each text source
 	FTextLocalizationResource NativeResource;
 	FTextLocalizationResource LocalizedResource;
-	for (const TSharedPtr<ILocalizedTextSource>& LocalizedTextSource : LocalizedTextSources)
+	for (const TSharedPtr<ILocalizedTextSource>& LocalizedTextSource : AvailableTextSources)
 	{
 		LocalizedTextSource->LoadLocalizedResources(FinalLocLoadFlags, PrioritizedCultureNames, NativeResource, LocalizedResource);
 	}
@@ -1166,6 +1096,191 @@ void FTextLocalizationManager::LoadLocalizationResourcesForPrioritizedCultures(T
 	DirtyTextRevision();
 }
 
+void FTextLocalizationManager::LoadLocalizationResourcesForPrioritizedCultures_Async(TArrayView<const FString> PrioritizedCultureNames, const ELocalizationLoadFlags LocLoadFlags)
+{
+	QueueAsyncTask([AvailableTextSources = LocalizedTextSources, PrioritizedCultureNames = TArray<FString>(PrioritizedCultureNames), LocLoadFlags]()
+	{
+		FTextLocalizationManager::Get().LoadLocalizationResourcesForPrioritizedCultures_Sync(AvailableTextSources, PrioritizedCultureNames, LocLoadFlags);
+	});
+}
+
+void FTextLocalizationManager::LoadLocalizationTargetsForPrioritizedCultures_Sync(TArrayView<const TSharedPtr<ILocalizedTextSource>> AvailableTextSources, TArrayView<const FString> LocalizationTargetPaths, TArrayView<const FString> PrioritizedCultureNames, const ELocalizationLoadFlags LocLoadFlags)
+{
+	LLM_SCOPE(ELLMTag::Localization);
+	TRACE_CPUPROFILER_EVENT_SCOPE(FTextLocalizationManager::LoadLocalizationTargetsForPrioritizedCultures);
+
+	// Load the resources from each localization target
+	FTextLocalizationResource UnusedNativeResource;
+	FTextLocalizationResource LocalizedResource;
+	for (const FString& LocalizationTargetPath : LocalizationTargetPaths)
+	{
+		UE_LOG(LogTextLocalizationManager, Verbose, TEXT("Loading LocRes data from '%s'"), *LocalizationTargetPath);
+	}
+	LocResTextSource->LoadLocalizedResourcesFromPaths(TArrayView<FString>(), LocalizationTargetPaths, TArrayView<FString>(), LocLoadFlags, PrioritizedCultureNames, UnusedNativeResource, LocalizedResource);
+
+	// Allow any higher priority text sources to override the additional text loaded (eg, to allow polyglot hot-fixes to take priority)
+	// Note: If any text sources don't support dynamic queries, then we must do a much slower full refresh instead :(
+	bool bNeedsFullRefresh = false;
+	{
+		// Copy the IDs array as QueryLocalizedResource can update the map
+		TArray<FTextId> NewTextIds;
+		LocalizedResource.Entries.GenerateKeyArray(NewTextIds);
+
+		for (const TSharedPtr<ILocalizedTextSource>& LocalizedTextSource : AvailableTextSources)
+		{
+			if (LocalizedTextSource->GetPriority() <= LocResTextSource->GetPriority())
+			{
+				continue;
+			}
+
+			for (const FTextId& NewTextId : NewTextIds)
+			{
+				if (LocalizedTextSource->QueryLocalizedResource(LocLoadFlags, PrioritizedCultureNames, NewTextId, UnusedNativeResource, LocalizedResource) == EQueryLocalizedResourceResult::NotImplemented)
+				{
+					bNeedsFullRefresh = true;
+					break;
+				}
+			}
+
+			if (bNeedsFullRefresh)
+			{
+				break;
+			}
+		}
+	}
+
+	// Apply the new data
+	if (bNeedsFullRefresh)
+	{
+		UE_LOG(LogTextLocalizationManager, Verbose, TEXT("Patching LocRes data failed, performing full refresh"));
+		LoadLocalizationResourcesForPrioritizedCultures_Sync(AvailableTextSources, PrioritizedCultureNames, LocLoadFlags);
+	}
+	else
+	{
+		UE_LOG(LogTextLocalizationManager, Verbose, TEXT("Patching LocRes data for %d entries"), LocalizedResource.Entries.Num());
+		UpdateFromLocalizations(MoveTemp(LocalizedResource), /*bDirtyTextRevision*/true);
+	}
+}
+
+void FTextLocalizationManager::LoadLocalizationTargetsForPrioritizedCultures_Async(TArrayView<const FString> LocalizationTargetPaths, TArrayView<const FString> PrioritizedCultureNames, const ELocalizationLoadFlags LocLoadFlags)
+{
+	QueueAsyncTask([AvailableTextSources = LocalizedTextSources, LocalizationTargetPaths = TArray<FString>(LocalizationTargetPaths), PrioritizedCultureNames = TArray<FString>(PrioritizedCultureNames), LocLoadFlags]()
+	{
+		FTextLocalizationManager::Get().LoadLocalizationTargetsForPrioritizedCultures_Sync(AvailableTextSources, LocalizationTargetPaths, PrioritizedCultureNames, LocLoadFlags);;
+	});
+}
+
+void FTextLocalizationManager::LoadChunkedLocalizationResources_Sync(TArrayView<const TSharedPtr<ILocalizedTextSource>> AvailableTextSources, const int32 ChunkId, const FString& PakFilename)
+{
+	LLM_SCOPE(ELLMTag::Localization);
+	TRACE_CPUPROFILER_EVENT_SCOPE(FTextLocalizationManager::LoadChunkedLocalizationResources);
+
+	check(ChunkId != INDEX_NONE);
+
+	UE_LOG(LogTextLocalizationManager, Verbose, TEXT("Request to load localization data for chunk %d (from PAK '%s')"), ChunkId, *PakFilename);
+
+	// Skip this request if we've already loaded the data for this chunk via the request for a previous PAK sub-file load notification
+	if (LocResTextSource->HasRegisteredChunkId(ChunkId))
+	{
+		UE_LOG(LogTextLocalizationManager, Verbose, TEXT("Skipped loading localization data for chunk %d (from PAK '%s') as this chunk has already been processed"), ChunkId, *PakFilename);
+		return;
+	}
+
+	// If we're being notified so early that even InitEngineTextLocalization hasn't run, then we can't safely make the queries below as things like GConfig may not be available yet!
+	if (!IsInitialized())
+	{
+		// Track this so that full resource refreshes (eg, changing culture) work as expected
+		LocResTextSource->RegisterChunkId(ChunkId);
+
+		UE_LOG(LogTextLocalizationManager, Verbose, TEXT("Skipped loading localization data for chunk %d (from PAK '%s') as the localization manager isn't ready"), ChunkId, *PakFilename);
+		return;
+	}
+
+	// Note: We only allow game localization targets to be chunked, and the layout is assumed to follow our standard pattern (as used by the localization dashboard and FLocTextHelper)
+	const TArray<FString> ChunkedLocalizationTargets = FLocalizationResourceTextSource::GetChunkedLocalizationTargets();
+
+	// Check to see whether all the required localization data is now available
+	// This may not be the case if this PAK was split into multiple sub-files, and the localization data was split between them
+	TArray<FString> PrioritizedLocalizationPaths;
+	for (const FString& LocalizationTarget : ChunkedLocalizationTargets)
+	{
+		const FString ChunkedLocalizationTargetName = TextLocalizationResourceUtil::GetLocalizationTargetNameForChunkId(LocalizationTarget, ChunkId);
+
+		FString ChunkedLocalizationTargetPath = FPaths::ProjectContentDir() / TEXT("Localization") / ChunkedLocalizationTargetName;
+		if (!IFileManager::Get().DirectoryExists(*ChunkedLocalizationTargetPath))
+		{
+			UE_LOG(LogTextLocalizationManager, Verbose, TEXT("Skipped loading localization data for chunk %d (from PAK '%s') as the localization directory for '%s' was not yet available"), ChunkId, *PakFilename, *ChunkedLocalizationTargetName);
+			return;
+		}
+
+		FTextLocalizationMetaDataResource LocMetaResource;
+		{
+			const FString LocMetaFilename = ChunkedLocalizationTargetPath / FString::Printf(TEXT("%s.locmeta"), *ChunkedLocalizationTargetName);
+			if (!IFileManager::Get().FileExists(*LocMetaFilename))
+			{
+				UE_LOG(LogTextLocalizationManager, Verbose, TEXT("Skipped loading localization data for chunk %d (from PAK '%s') as the LocMeta file for '%s' was not yet available"), ChunkId, *PakFilename, *ChunkedLocalizationTargetName);
+				return;
+			}
+			if (!LocMetaResource.LoadFromFile(LocMetaFilename))
+			{
+				UE_LOG(LogTextLocalizationManager, Verbose, TEXT("Skipped loading localization data for chunk %d (from PAK '%s') as the LocMeta file for '%s' failed to load"), ChunkId, *PakFilename, *ChunkedLocalizationTargetName);
+				return;
+			}
+		}
+
+		for (const FString& CompiledCulture : LocMetaResource.CompiledCultures)
+		{
+			const FString LocResFilename = ChunkedLocalizationTargetPath / CompiledCulture / FString::Printf(TEXT("%s.locres"), *ChunkedLocalizationTargetName);
+			if (!IFileManager::Get().FileExists(*LocResFilename))
+			{
+				UE_LOG(LogTextLocalizationManager, Verbose, TEXT("Skipped loading localization data for chunk %d (from PAK '%s') as the '%s' LocRes file for '%s' was not yet available"), ChunkId, *PakFilename, *CompiledCulture, *ChunkedLocalizationTargetName);
+				return;
+			}
+		}
+
+		PrioritizedLocalizationPaths.Add(MoveTemp(ChunkedLocalizationTargetPath));
+	}
+
+	// Track this so that full resource refreshes (eg, changing culture) work as expected
+	LocResTextSource->RegisterChunkId(ChunkId);
+
+	if (!EnumHasAnyFlags(InitializedFlags.load(), ETextLocalizationManagerInitializedFlags::Game))
+	{
+		// If we've not yet initialized game localization then don't bother patching, as the full initialization path will load the data for this chunk
+		return;
+	}
+
+	const TArray<FString> PrioritizedCultureNames = FInternationalization::Get().GetPrioritizedCultureNames(FInternationalization::Get().GetCurrentLanguage()->GetName());
+	LoadLocalizationTargetsForPrioritizedCultures_Sync(AvailableTextSources, PrioritizedLocalizationPaths, PrioritizedCultureNames, ELocalizationLoadFlags::Game);
+}
+
+void FTextLocalizationManager::LoadChunkedLocalizationResources_Async(const int32 ChunkId, const FString& PakFilename)
+{
+	QueueAsyncTask([AvailableTextSources = LocalizedTextSources, ChunkId, PakFilename]()
+	{
+		FTextLocalizationManager::Get().LoadChunkedLocalizationResources_Sync(AvailableTextSources, ChunkId, PakFilename);
+	});
+}
+
+void FTextLocalizationManager::QueueAsyncTask(TUniqueFunction<void()>&& Task)
+{
+	if (TextLocalizationManager::AsyncLoadLocalizationData && FTaskGraphInterface::IsRunning())
+	{
+		if (AsyncLocalizationTask)
+		{
+			AsyncLocalizationTask = FFunctionGraphTask::CreateAndDispatchWhenReady(MoveTemp(Task), TStatId(), AsyncLocalizationTask);
+		}
+		else
+		{
+			AsyncLocalizationTask = FFunctionGraphTask::CreateAndDispatchWhenReady(MoveTemp(Task), TStatId());
+		}
+	}
+	else
+	{
+		Task();
+	}
+}
+
 void FTextLocalizationManager::UpdateFromNative(FTextLocalizationResource&& TextLocalizationResource, const bool bDirtyTextRevision)
 {
 	TRACE_CPUPROFILER_EVENT_SCOPE(FTextLocalizationManager::UpdateFromNative);
@@ -1419,7 +1534,17 @@ void FTextLocalizationManager::DirtyTextRevision()
 		LocalTextRevisions.Empty();
 	}
 
-	OnTextRevisionChangedEvent.Broadcast();
+	if (IsInGameThread())
+	{
+		OnTextRevisionChangedEvent.Broadcast();
+	}
+	else
+	{
+		AsyncTask(ENamedThreads::GameThread, []()
+		{
+			FTextLocalizationManager::Get().OnTextRevisionChangedEvent.Broadcast();
+		});
+	}
 }
 
 #if WITH_EDITOR
@@ -1460,7 +1585,7 @@ void FTextLocalizationManager::EnableGameLocalizationPreview(const FString& Cult
 	ELocalizationLoadFlags LocLoadFlags = ELocalizationLoadFlags::Game | ELocalizationLoadFlags::ForceLocalizedGame;
 	LocLoadFlags |= (bIsGameLocalizationPreviewEnabled ? ELocalizationLoadFlags::Native : ELocalizationLoadFlags::None);
 
-	LoadLocalizationResourcesForPrioritizedCultures(PrioritizedCultureNames, LocLoadFlags);
+	LoadLocalizationResourcesForPrioritizedCultures_Async(PrioritizedCultureNames, LocLoadFlags);
 }
 
 void FTextLocalizationManager::DisableGameLocalizationPreview()
diff --git a/Engine/Source/Runtime/Core/Public/Internationalization/TextLocalizationManager.h b/Engine/Source/Runtime/Core/Public/Internationalization/TextLocalizationManager.h
--- a/Engine/Source/Runtime/Core/Public/Internationalization/TextLocalizationManager.h
+++ b/Engine/Source/Runtime/Core/Public/Internationalization/TextLocalizationManager.h
@@ -2,6 +2,7 @@
 
 #pragma once
 
+#include "Async/TaskGraphFwd.h"
 #include "Containers/Array.h"
 #include "Containers/ArrayView.h"
 #include "Containers/ContainerAllocationPolicies.h"
@@ -17,6 +18,7 @@
 #include "Internationalization/TextKey.h"
 #include "Misc/Crc.h"
 #include "Misc/EnumClassFlags.h"
+#include "Templates/Function.h"
 #include "Templates/SharedPointer.h"
 
 #include <atomic>
@@ -31,7 +33,6 @@ enum class ETextLocalizationManagerInitializedFlags : uint8
 	None = 0,
 	Engine = 1<<0,
 	Game = 1<<1,
-	Initializing = 1<<2
 };
 ENUM_CLASS_FLAGS(ETextLocalizationManagerInitializedFlags);
 
@@ -41,8 +42,6 @@ class CORE_API FTextLocalizationManager
 	friend CORE_API void BeginPreInitTextLocalization();
 	friend CORE_API void BeginInitTextLocalization();
 	friend CORE_API void InitEngineTextLocalization();
-	friend CORE_API void BeginInitGameTextLocalization();
-	friend CORE_API void EndInitGameTextLocalization();
 	friend CORE_API void InitGameTextLocalization();
 
 private:
@@ -88,11 +87,6 @@ private:
 		return InitializedFlags != ETextLocalizationManagerInitializedFlags::None;
 	}
 
-	bool IsInitializing() const
-	{
-		return EnumHasAnyFlags(InitializedFlags.load(), ETextLocalizationManagerInitializedFlags::Initializing);
-	}
-
 	mutable FCriticalSection DisplayStringLookupTableCS;
 	FDisplayStringLookupTable DisplayStringLookupTable;
 
@@ -111,6 +105,9 @@ private:
 	
 public:
 
+	~FTextLocalizationManager();
+	UE_NONCOPYABLE(FTextLocalizationManager);
+
 	/** Singleton accessor */
 	static FTextLocalizationManager& Get();
 	static void TearDown();
@@ -183,9 +180,38 @@ public:
 	void UpdateFromLocalizationResource(const FString& LocalizationResourceFilePath);
 	void UpdateFromLocalizationResource(const FTextLocalizationResource& TextLocalizationResource);
 
-	/** Reloads resources for the current culture. */
+	/**
+	 * Wait for any current async tasks to finish.
+	 * 
+	 * Async tasks will be started for anything that loads localization data (eg, initialization, language changes, requests to refresh the localization data, requests 
+	 * to load additional localization data for chunked targets, or requests to load additional localization data for explicitly loaded plugins).
+	 * 
+	 * While the engine automatically waits at certain points during initialization, you may find that you need to add your own waits if you cause additional localization 
+	 * data to load post-initialization (eg, because you're using chunked localization targets or are explicitly loading plugins). A good place to add a wait in that case 
+	 * might be at the end of your loading screen, before showing your main menu or game world.
+	 */
+	void WaitForAsyncTasks();
+
+	/**
+	 * Reloads resources for the current culture.
+	 * @note The loading is async, see WaitUntilLoadingCompletes.
+	 */
 	void RefreshResources();
 
+	/**
+	 * Called when paths containing additional localization target data (LocRes) are mounted, to allow the display strings to dynamically update without waiting for a refresh.
+	 * @note The loading is async, see WaitUntilLoadingCompletes.
+	 * @see FCoreDelegates::GatherAdditionalLocResPathsCallback.
+	 */
+	void HandleLocalizationTargetsMounted(TArrayView<const FString> LocalizationTargetPaths);
+
+	 /**
+	  * Called when paths containing additional localization target data (LocRes) are unmounted, to allow the display strings to dynamically update without waiting for a refresh.
+	  * @note The loading is async, see WaitUntilLoadingCompletes.
+	  * @see FCoreDelegates::GatherAdditionalLocResPathsCallback.
+	  */
+	void HandleLocalizationTargetsUnmounted(TArrayView<const FString> LocalizationTargetPaths);
+
 	/**
 	 * Returns the current text revision number. This value can be cached when caching information from the text localization manager.
 	 * If the revision does not match, cached information may be invalid and should be recached.
@@ -269,17 +295,30 @@ public:
 	FTextRevisionChangedEvent OnTextRevisionChangedEvent;
 
 private:
-	/** Callback for when a PAK file is loaded. Loads any chunk specific localization resources. */
+	/** Callback for when a PAK file is loaded. Async loads any chunk specific localization resources. */
 	void OnPakFileMounted(const IPakFile& PakFile);
 
-	/** Callback for changes in culture. Loads the new culture's localization resources. */
+	/** Callback for changes in culture. Async loads the new culture's localization resources. */
 	void OnCultureChanged();
 
 	/** Loads localization resources for the specified culture, optionally loading localization resources that are editor-specific or game-specific. */
-	void LoadLocalizationResourcesForCulture(const FString& CultureName, const ELocalizationLoadFlags LocLoadFlags);
+	void LoadLocalizationResourcesForCulture_Sync(TArrayView<const TSharedPtr<ILocalizedTextSource>> AvailableTextSources, const FString& CultureName, const ELocalizationLoadFlags LocLoadFlags);
+	void LoadLocalizationResourcesForCulture_Async(const FString& CultureName, const ELocalizationLoadFlags LocLoadFlags);
 
 	/** Loads localization resources for the specified prioritized cultures, optionally loading localization resources that are editor-specific or game-specific. */
-	void LoadLocalizationResourcesForPrioritizedCultures(TArrayView<const FString> PrioritizedCultureNames, const ELocalizationLoadFlags LocLoadFlags);
+	void LoadLocalizationResourcesForPrioritizedCultures_Sync(TArrayView<const TSharedPtr<ILocalizedTextSource>> AvailableTextSources, TArrayView<const FString> PrioritizedCultureNames, const ELocalizationLoadFlags LocLoadFlags);
+	void LoadLocalizationResourcesForPrioritizedCultures_Async(TArrayView<const FString> PrioritizedCultureNames, const ELocalizationLoadFlags LocLoadFlags);
+
+	/** Loads the specified localization targets for the specified prioritized cultures, optionally loading localization resources that are editor-specific or game-specific. */
+	void LoadLocalizationTargetsForPrioritizedCultures_Sync(TArrayView<const TSharedPtr<ILocalizedTextSource>> AvailableTextSources, TArrayView<const FString> LocalizationTargetPaths, TArrayView<const FString> PrioritizedCultureNames, const ELocalizationLoadFlags LocLoadFlags);
+	void LoadLocalizationTargetsForPrioritizedCultures_Async(TArrayView<const FString> LocalizationTargetPaths, TArrayView<const FString> PrioritizedCultureNames, const ELocalizationLoadFlags LocLoadFlags);
+
+	/** Loads the chunked localization resource data for the specified chunk */
+	void LoadChunkedLocalizationResources_Sync(TArrayView<const TSharedPtr<ILocalizedTextSource>> AvailableTextSources, const int32 ChunkId, const FString& PakFilename);
+	void LoadChunkedLocalizationResources_Async(const int32 ChunkId, const FString& PakFilename);
+
+	/** Queue the given task to run async (where possible), chained to any existing AsyncLocalizationTask */
+	void QueueAsyncTask(TUniqueFunction<void()>&& Task);
 
 	/** Updates display string entries and adds new display string entries based on provided native text. */
 	void UpdateFromNative(FTextLocalizationResource&& TextLocalizationResource, const bool bDirtyTextRevision = true);
@@ -297,8 +336,7 @@ private:
 	void LeetifyAllDisplayStrings();
 	/** A helper function that converts all of the display strings to show the localization key associated with the string when the keys culture is active. */
 	void KeyifyAllDisplayStrings();
-#endif 
-
+#endif
 
 	/** Array of registered localized text sources, sorted by priority (@see RegisterTextSource) */
 	TArray<TSharedPtr<ILocalizedTextSource>> LocalizedTextSources;
@@ -308,4 +346,7 @@ private:
 
 	/** The polyglot text source (this is also added to LocalizedTextSources, but we keep a pointer to it directly so we can add new polyglot data to it at runtime) */
 	TSharedPtr<class FPolyglotTextSource> PolyglotTextSource;
+
+	/** The latest async localization task (if any). Additional requests are chained to this so that they run in sequence */
+	FGraphEventRef AsyncLocalizationTask;
 };
diff --git a/Engine/Source/Runtime/Core/Public/Internationalization/TextLocalizationManagerGlobals.h b/Engine/Source/Runtime/Core/Public/Internationalization/TextLocalizationManagerGlobals.h
--- a/Engine/Source/Runtime/Core/Public/Internationalization/TextLocalizationManagerGlobals.h
+++ b/Engine/Source/Runtime/Core/Public/Internationalization/TextLocalizationManagerGlobals.h
@@ -8,5 +8,3 @@ CORE_API void BeginPreInitTextLocalization();
 CORE_API void BeginInitTextLocalization();
 CORE_API void InitEngineTextLocalization();
 CORE_API void InitGameTextLocalization();
-CORE_API void BeginInitGameTextLocalization();
-CORE_API void EndInitGameTextLocalization();
diff --git a/Engine/Source/Runtime/Launch/Private/LaunchEngineLoop.cpp b/Engine/Source/Runtime/Launch/Private/LaunchEngineLoop.cpp
--- a/Engine/Source/Runtime/Launch/Private/LaunchEngineLoop.cpp
+++ b/Engine/Source/Runtime/Launch/Private/LaunchEngineLoop.cpp
@@ -50,6 +50,7 @@
 #include "ProfilingDebugging/BootProfiling.h"
 #if WITH_ENGINE
 #include "HAL/PlatformSplash.h"
+#include "SceneInterface.h"
 #include "StereoRenderUtils.h"
 #include "DataDrivenShaderPlatformInfo.h"
 #endif
@@ -3053,6 +3054,9 @@ PRAGMA_ENABLE_DEPRECATION_WARNINGS
 		FAudioThread::SetUseThreadedAudio(bUseThreadedAudio);
 	}
 
+	// Ensure engine localization has loaded before we show the splash
+	FTextLocalizationManager::Get().WaitForAsyncTasks();
+
 	// Are we creating a slate application?
 	bool bSlateApplication = !IsRunningDedicatedServer() && (bIsRegularClient || bHasEditorToken);
 	if (bSlateApplication)
@@ -3605,7 +3609,7 @@ int32 FEngineLoop::PreInitPostStartupScreen(const TCHAR* CmdLine)
 		}
 #endif
 
-		BeginInitGameTextLocalization();
+		InitGameTextLocalization();
 
 		DECLARE_SCOPE_CYCLE_COUNTER(TEXT("Initial UObject load"), STAT_InitialUObjectLoad, STATGROUP_LoadTime);
 
@@ -3644,7 +3648,8 @@ int32 FEngineLoop::PreInitPostStartupScreen(const TCHAR* CmdLine)
 			ProcessNewlyLoadedUObjects();
 		}
 
-		EndInitGameTextLocalization();
+		// Ensure game localization has loaded before we continue
+		FTextLocalizationManager::Get().WaitForAsyncTasks();
 
 		FDelayedAutoRegisterHelper::RunAndClearDelayedAutoRegisterDelegates(EDelayedRegisterRunPhase::ObjectSystemReady);
 
@@ -4240,6 +4245,7 @@ int32 FEngineLoop::PreInitPostStartupScreen(const TCHAR* CmdLine)
 #else // WITH_ENGINE
 	InitEngineTextLocalization();
 	InitGameTextLocalization();
+	FTextLocalizationManager::Get().WaitForAsyncTasks();
 #if USE_LOCALIZED_PACKAGE_CACHE
 	{
 		SCOPED_BOOT_TIMING("FPackageLocalizationManager::Get().InitializeFromDefaultCache");
diff --git a/Engine/Source/Runtime/Projects/Private/PluginManager.cpp b/Engine/Source/Runtime/Projects/Private/PluginManager.cpp
--- a/Engine/Source/Runtime/Projects/Private/PluginManager.cpp
+++ b/Engine/Source/Runtime/Projects/Private/PluginManager.cpp
@@ -26,6 +26,7 @@
 #include "Misc/ScopeRWLock.h"
 #include "Algo/Accumulate.h"
 #include "Containers/VersePath.h"
+#include "Internationalization/TextLocalizationManager.h"
 #include "ProfilingDebugging/CpuProfilerTrace.h"
 #if READ_TARGET_ENABLED_PLUGINS_FROM_RECEIPT
 #include "TargetReceipt.h"
@@ -340,6 +341,21 @@ namespace DiscoveredPluginMapUtils
 	}
 }
 
+namespace PluginLocalizationUtils
+{
+	void GetLocalizationPathsForPlugin(const TSharedRef<FPlugin>& Plugin, TArray<FString>& OutLocResPaths)
+	{
+		const FString PluginLocDir = Plugin->GetContentDir() / TEXT("Localization");
+		for (const FLocalizationTargetDescriptor& LocTargetDesc : Plugin->GetDescriptor().LocalizationTargets)
+		{
+			if (LocTargetDesc.ShouldLoadLocalizationTarget())
+			{
+				OutLocResPaths.Add(PluginLocDir / LocTargetDesc.Name);
+			}
+		}
+	}
+}
+
 FPlugin::FPlugin(const FString& InFileName, const FPluginDescriptor& InDescriptor, EPluginType InType)
 	: Name(FPaths::GetBaseFilename(InFileName))
 	, FileName(InFileName)
@@ -2308,14 +2324,7 @@ void FPluginManager::GetLocalizationPathsForEnabledPlugins( TArray<FString>& Out
 			continue;
 		}
 		
-		const FString PluginLocDir = Plugin->GetContentDir() / TEXT("Localization");
-		for (const FLocalizationTargetDescriptor& LocTargetDesc : Plugin->GetDescriptor().LocalizationTargets)
-		{
-			if (LocTargetDesc.ShouldLoadLocalizationTarget())
-			{
-				OutLocResPaths.Add(PluginLocDir / LocTargetDesc.Name);
-			}
-		}
+		PluginLocalizationUtils::GetLocalizationPathsForPlugin(Plugin, OutLocResPaths);
 	}
 }
 
@@ -2653,6 +2662,14 @@ void FPluginManager::MountPluginFromExternalSource(const TSharedRef<FPlugin>& Pl
 		}
 	}
 
+	// Notify that additional localization data should be loaded
+	if (Plugin->Descriptor.LocalizationTargets.Num() > 0)
+	{
+		TArray<FString> AdditionalLocResPaths;
+		PluginLocalizationUtils::GetLocalizationPathsForPlugin(Plugin, AdditionalLocResPaths);
+		FTextLocalizationManager::Get().HandleLocalizationTargetsMounted(AdditionalLocResPaths);
+	}
+
 	// If it's a code module, also load the modules for it
 	if (Plugin->Descriptor.Modules.Num() > 0)
 	{
@@ -2711,6 +2728,14 @@ bool FPluginManager::UnmountPluginFromExternalSource(const TSharedPtr<FPlugin>&
 		return false;
 	}
 
+	// Notify that additional localization data should be unloaded
+	if (Plugin->Descriptor.LocalizationTargets.Num() > 0 && !IsEngineExitRequested())
+	{
+		TArray<FString> AdditionalLocResPaths;
+		PluginLocalizationUtils::GetLocalizationPathsForPlugin(Plugin.ToSharedRef(), AdditionalLocResPaths);
+		FTextLocalizationManager::Get().HandleLocalizationTargetsUnmounted(AdditionalLocResPaths);
+	}
+
 	if ((Plugin->CanContainContent() || Plugin->CanContainVerse()) && ensure(UnRegisterMountPointDelegate.IsBound()))
 	{
 		// Remove this plugin's path from the list of content directories that the editor will search
diff --git a/Engine/Source/Runtime/SlateCore/Private/Application/SlateApplicationBase.cpp b/Engine/Source/Runtime/SlateCore/Private/Application/SlateApplicationBase.cpp
--- a/Engine/Source/Runtime/SlateCore/Private/Application/SlateApplicationBase.cpp
+++ b/Engine/Source/Runtime/SlateCore/Private/Application/SlateApplicationBase.cpp
@@ -35,7 +35,14 @@ FSlateApplicationBase::FSlateApplicationBase()
 , bIsSlateAsleep(false)
 , CustomSafeZoneState(ECustomSafeZoneState::Unset)
 {
-
+	FTextLocalizationManager::Get().OnTextRevisionChangedEvent.AddLambda([]()
+	{
+		if (FSlateApplicationBase::IsInitialized())
+		{
+			// Redraw widgets when new localized text data is loaded
+			FSlateApplicationBase::Get().InvalidateAllWidgets(false);
+		}
+	});
 }
 
 void FSlateApplicationBase::GetDisplayMetrics(FDisplayMetrics& OutDisplayMetrics) 
