Added a new simpler UE_REQUIRES macro which replaces the UE_CONSTRAINT* macros.

diff --git a/Engine/Source/Runtime/Core/Public/Math/UnrealMathUtility.h b/Engine/Source/Runtime/Core/Public/Math/UnrealMathUtility.h
--- a/Engine/Source/Runtime/Core/Public/Math/UnrealMathUtility.h
+++ b/Engine/Source/Runtime/Core/Public/Math/UnrealMathUtility.h
@@ -530,9 +530,7 @@ public:
 	 */
 	template <
 		typename IntegralType
-		UE_CONSTRAINTS_BEGIN
-			UE_CONSTRAINT(std::is_integral_v<IntegralType>)
-		UE_CONSTRAINTS_END
+		UE_REQUIRES(std::is_integral_v<IntegralType>)
 	>
 	UE_NODISCARD static constexpr FORCEINLINE IntegralType Floor(IntegralType I)
 	{
@@ -695,9 +693,7 @@ public:
 	*/
 	template <
 		typename T
-		UE_CONSTRAINTS_BEGIN
-			UE_CONSTRAINT(std::is_floating_point_v<T>)
-		UE_CONSTRAINTS_END
+		UE_REQUIRES(std::is_floating_point_v<T>)
 	>
 	static constexpr FORCEINLINE void SinCos(std::decay_t<T>* ScalarSin, std::decay_t<T>* ScalarCos, T  Value )
 	{
@@ -750,9 +746,7 @@ public:
 	template <
 		typename T,
 		typename U
-		UE_CONSTRAINTS_BEGIN
-			UE_CONSTRAINT(!std::is_same_v<T, U>)
-		UE_CONSTRAINTS_END
+		UE_REQUIRES(!std::is_same_v<T, U>)
 	>
 	static FORCEINLINE void SinCos(T* ScalarSin, T* ScalarCos, U Value)
 	{
@@ -843,9 +837,7 @@ public:
 	template <
 		typename T,
 		typename T2
-		UE_CONSTRAINTS_BEGIN
-			UE_CONSTRAINT(std::is_floating_point_v<T> || std::is_floating_point_v<T2>)
-		UE_CONSTRAINTS_END
+		UE_REQUIRES(std::is_floating_point_v<T> || std::is_floating_point_v<T2>)
 	>
 	UE_NODISCARD static constexpr auto FindDeltaAngleDegrees(T A1, T2 A2) -> decltype(A1 * A2)
 	{
@@ -873,9 +865,7 @@ public:
 	template <
 		typename T,
 		typename T2
-		UE_CONSTRAINTS_BEGIN
-			UE_CONSTRAINT(std::is_floating_point_v<T> || std::is_floating_point_v<T2>)
-		UE_CONSTRAINTS_END
+		UE_REQUIRES(std::is_floating_point_v<T> || std::is_floating_point_v<T2>)
 	>
 	UE_NODISCARD static constexpr auto FindDeltaAngleRadians(T A1, T2 A2) -> decltype(A1 * A2)
 	{
@@ -902,9 +892,7 @@ public:
 	/** Given a heading which may be outside the +/- PI range, 'unwind' it back into that range. */
 	template <
 		typename T
-		UE_CONSTRAINTS_BEGIN
-			UE_CONSTRAINT(std::is_floating_point_v<T>)
-		UE_CONSTRAINTS_END
+		UE_REQUIRES(std::is_floating_point_v<T>)
 	>
 	UE_NODISCARD static constexpr T UnwindRadians(T A)
 	{
@@ -924,9 +912,7 @@ public:
 	/** Utility to ensure angle is between +/- 180 degrees by unwinding. */
 	template <
 		typename T
-		UE_CONSTRAINTS_BEGIN
-			UE_CONSTRAINT(std::is_floating_point_v<T>)
-		UE_CONSTRAINTS_END
+		UE_REQUIRES(std::is_floating_point_v<T>)
 	>
 	UE_NODISCARD static constexpr T UnwindDegrees(T A)
 	{
@@ -1033,9 +1019,7 @@ public:
 	template <
 		typename T,
 		typename T2
-		UE_CONSTRAINTS_BEGIN
-			UE_CONSTRAINT(std::is_floating_point_v<T>)
-		UE_CONSTRAINTS_END
+		UE_REQUIRES(std::is_floating_point_v<T>)
 	>
 	UE_NODISCARD static constexpr FORCEINLINE auto GetRangePct(T MinValue, T MaxValue, T2 Value)
 	{
@@ -1055,9 +1039,7 @@ public:
 	template <
 		typename T,
 		typename T2
-		UE_CONSTRAINTS_BEGIN
-			UE_CONSTRAINT(std::is_floating_point_v<T>)
-		UE_CONSTRAINTS_END
+		UE_REQUIRES(std::is_floating_point_v<T>)
 	>
 	UE_NODISCARD static auto GetRangePct(UE::Math::TVector2<T> const& Range, T2 Value)
 	{
@@ -1068,9 +1050,7 @@ public:
 	template <
 		typename T,
 		typename T2
-		UE_CONSTRAINTS_BEGIN
-			UE_CONSTRAINT(std::is_floating_point_v<T>)
-		UE_CONSTRAINTS_END
+		UE_REQUIRES(std::is_floating_point_v<T>)
 	>
 	UE_NODISCARD static auto GetRangeValue(UE::Math::TVector2<T> const& Range, T2 Pct)
 	{
@@ -1116,9 +1096,7 @@ public:
 	template <
 		typename T,
 		typename U
-		UE_CONSTRAINTS_BEGIN
-			UE_CONSTRAINT(!TCustomLerp<T>::Value && (std::is_floating_point_v<U> || std::is_same_v<T, U>))
-		UE_CONSTRAINTS_END
+		UE_REQUIRES(!TCustomLerp<T>::Value && (std::is_floating_point_v<U> || std::is_same_v<T, U>))
 	>
 	UE_NODISCARD static constexpr FORCEINLINE_DEBUGGABLE T Lerp( const T& A, const T& B, const U& Alpha )
 	{
@@ -1129,9 +1107,7 @@ public:
 	template <
 		typename T,
 		typename U
-		UE_CONSTRAINTS_BEGIN
-			UE_CONSTRAINT(TCustomLerp<T>::Value)
-		UE_CONSTRAINTS_END
+		UE_REQUIRES(TCustomLerp<T>::Value)
 	>
 	UE_NODISCARD static FORCEINLINE_DEBUGGABLE T Lerp(const T& A, const T& B, const U& Alpha)
 	{
@@ -1143,9 +1119,7 @@ public:
 		typename T1,
 		typename T2,
 		typename T3
-		UE_CONSTRAINTS_BEGIN
-			UE_CONSTRAINT(!std::is_same_v<T1, T2> && !TCustomLerp<T1>::Value && !TCustomLerp<T2>::Value)
-		UE_CONSTRAINTS_END
+		UE_REQUIRES(!std::is_same_v<T1, T2> && !TCustomLerp<T1>::Value && !TCustomLerp<T2>::Value)
 	>
 	UE_NODISCARD static auto Lerp( const T1& A, const T2& B, const T3& Alpha ) -> decltype(A * B)
 	{
@@ -1172,9 +1146,7 @@ public:
 		typename T1,
 		typename T2,
 		typename T3
-		UE_CONSTRAINTS_BEGIN
-			UE_CONSTRAINT(!std::is_same_v<T1, T2>)
-		UE_CONSTRAINTS_END
+		UE_REQUIRES(!std::is_same_v<T1, T2>)
 	>
 	UE_NODISCARD static auto LerpStable( const T1& A, const T2& B, const T3& Alpha ) -> decltype(A * B)
 	{
@@ -1186,9 +1158,7 @@ public:
 	template <
 		typename T,
 		typename U
-		UE_CONSTRAINTS_BEGIN
-			UE_CONSTRAINT(!TCustomLerp<T>::Value && (std::is_floating_point_v<U> || std::is_same_v<T, U>))
-		UE_CONSTRAINTS_END
+		UE_REQUIRES(!TCustomLerp<T>::Value && (std::is_floating_point_v<U> || std::is_same_v<T, U>))
 	>
 	UE_NODISCARD static constexpr FORCEINLINE_DEBUGGABLE T BiLerp(const T& P00,const T& P10,const T& P01,const T& P11, const U& FracX, const U& FracY)
 	{
@@ -1203,9 +1173,7 @@ public:
 	template <
 		typename T,
 		typename U
-		UE_CONSTRAINTS_BEGIN
-			UE_CONSTRAINT(TCustomLerp<T>::Value)
-		UE_CONSTRAINTS_END
+		UE_REQUIRES(TCustomLerp<T>::Value)
 	>
 	UE_NODISCARD static FORCEINLINE_DEBUGGABLE T BiLerp(const T& P00, const T& P10, const T& P01, const T& P11, const U& FracX, const U& FracY)
 	{
@@ -1224,9 +1192,7 @@ public:
 	template <
 		typename T,
 		typename U
-		UE_CONSTRAINTS_BEGIN
-			UE_CONSTRAINT(!TCustomLerp<T>::Value && (std::is_floating_point_v<U> || std::is_same_v<T, U>))
-		UE_CONSTRAINTS_END
+		UE_REQUIRES(!TCustomLerp<T>::Value && (std::is_floating_point_v<U> || std::is_same_v<T, U>))
 	>
 	UE_NODISCARD static constexpr FORCEINLINE_DEBUGGABLE T CubicInterp( const T& P0, const T& T0, const T& P1, const T& T1, const U& A )
 	{
@@ -1240,9 +1206,7 @@ public:
 	template <
 		typename T,
 		typename U
-		UE_CONSTRAINTS_BEGIN
-			UE_CONSTRAINT(TCustomLerp<T>::Value)
-		UE_CONSTRAINTS_END
+		UE_REQUIRES(TCustomLerp<T>::Value)
 	>
 	UE_NODISCARD static FORCEINLINE_DEBUGGABLE T CubicInterp(const T& P0, const T& T0, const T& P1, const T& T1, const U& A)
 	{
@@ -1261,9 +1225,7 @@ public:
 	template <
 		typename T,
 		typename U
-		UE_CONSTRAINTS_BEGIN
-			UE_CONSTRAINT(std::is_floating_point_v<U>)
-		UE_CONSTRAINTS_END
+		UE_REQUIRES(std::is_floating_point_v<U>)
 	>
 	UE_NODISCARD static constexpr FORCEINLINE_DEBUGGABLE T CubicInterpDerivative( const T& P0, const T& T0, const T& P1, const T& T1, const U& A )
 	{
@@ -1288,9 +1250,7 @@ public:
 	template <
 		typename T,
 		typename U
-		UE_CONSTRAINTS_BEGIN
-			UE_CONSTRAINT(std::is_floating_point_v<U>)
-		UE_CONSTRAINTS_END
+		UE_REQUIRES(std::is_floating_point_v<U>)
 	>
 	UE_NODISCARD static constexpr FORCEINLINE_DEBUGGABLE T CubicInterpSecondDerivative( const T& P0, const T& T0, const T& P1, const T& T1, const U& A )
 	{
diff --git a/Engine/Source/Runtime/Core/Public/Templates/Function.h b/Engine/Source/Runtime/Core/Public/Templates/Function.h
--- a/Engine/Source/Runtime/Core/Public/Templates/Function.h
+++ b/Engine/Source/Runtime/Core/Public/Templates/Function.h
@@ -561,10 +561,8 @@ namespace UE::Core::Private::Function
 		}
 
 		template <
-			typename FunctorType,
-			typename = typename TEnableIf<
-				!std::is_same_v<TFunctionRefBase, typename TDecay<FunctorType>::Type>
-			>::Type
+			typename FunctorType
+			UE_REQUIRES(!std::is_same_v<TFunctionRefBase, std::decay_t<FunctorType>>)
 		>
 		TFunctionRefBase(FunctorType&& InFunc)
 		{
diff --git a/Engine/Source/Runtime/Core/Public/Templates/ResolveTypeAmbiguity.h b/Engine/Source/Runtime/Core/Public/Templates/ResolveTypeAmbiguity.h
--- a/Engine/Source/Runtime/Core/Public/Templates/ResolveTypeAmbiguity.h
+++ b/Engine/Source/Runtime/Core/Public/Templates/ResolveTypeAmbiguity.h
@@ -39,13 +39,11 @@
 	template < \
 		typename Arg1, \
 		typename Arg2 \
-		UE_CONSTRAINTS_BEGIN \
-			UE_CONSTRAINT( \
-				!std::is_same_v<Arg1, Arg2> && \
-				std::is_signed_v<Arg1> && std::is_integral_v<Arg1> && \
-				std::is_signed_v<Arg2> && std::is_integral_v<Arg2> \
-			) \
-		UE_CONSTRAINTS_END \
+		UE_REQUIRES( \
+			!std::is_same_v<Arg1, Arg2> && \
+			std::is_signed_v<Arg1> && std::is_integral_v<Arg1> && \
+			std::is_signed_v<Arg2> && std::is_integral_v<Arg2> \
+		) \
 	> \
 	static OptionalMarkup FORCEINLINE auto Func(Arg1 X, Arg2 Y) -> decltype(X * Y) \
 	{ \
@@ -62,18 +60,16 @@
 		typename Arg1, \
 		typename Arg2, \
 		typename Arg3 \
-		UE_CONSTRAINTS_BEGIN \
-			UE_CONSTRAINT( \
-				( \
-					!std::is_same_v<Arg1, Arg2>> || \
-					!std::is_same_v<Arg2, Arg3>> || \
-					!std::is_same_v<Arg1, Arg3>> \
-				) && \
-				std::is_signed_v<Arg1> && std::is_integral_v<Arg1> && \
-				std::is_signed_v<Arg2> && std::is_integral_v<Arg2> && \
-				std::is_signed_v<Arg3> && std::is_integral_v<Arg3> \
-			) \
-		UE_CONSTRAINTS_END \
+		UE_REQUIRES( \
+			( \
+				!std::is_same_v<Arg1, Arg2>> || \
+				!std::is_same_v<Arg2, Arg3>> || \
+				!std::is_same_v<Arg1, Arg3>> \
+			) && \
+			std::is_signed_v<Arg1> && std::is_integral_v<Arg1> && \
+			std::is_signed_v<Arg2> && std::is_integral_v<Arg2> && \
+			std::is_signed_v<Arg3> && std::is_integral_v<Arg3> \
+		) \
 	> \
 	static OptionalMarkup FORCEINLINE auto Func(Arg1 X, Arg2 Y, Arg3 Z) -> decltype(X * Y * Z) \
 	{ \
@@ -90,9 +86,7 @@
 	template < \
 		typename Arg1, \
 		typename Arg2 \
-		UE_CONSTRAINTS_BEGIN \
-			UE_CONSTRAINT((std::is_floating_point_v<Arg1> || std::is_floating_point_v<Arg2>) && !std::is_same_v<Arg1, Arg2>) \
-		UE_CONSTRAINTS_END \
+		UE_REQUIRES((std::is_floating_point_v<Arg1> || std::is_floating_point_v<Arg2>) && !std::is_same_v<Arg1, Arg2>) \
 	> \
 	static FORCEINLINE auto Func(Arg1 X, Arg2 Y) -> decltype(X * Y) \
 	{ \
@@ -106,12 +100,10 @@
 		typename Arg1, \
 		typename Arg2, \
 		typename Arg3 \
-		UE_CONSTRAINTS_BEGIN \
-			UE_CONSTRAINT( \
-				(std::is_floating_point_v<Arg1> || std::is_floating_point_v<Arg2> || std::is_floating_point_v<Arg3>) && \
-				(!std::is_same_v<Arg1, Arg2> || !std::is_same_v<Arg2, Arg3> || !std::is_same_v<Arg1, Arg3>) \
-			) \
-		UE_CONSTRAINTS_END \
+		UE_REQUIRES( \
+			(std::is_floating_point_v<Arg1> || std::is_floating_point_v<Arg2> || std::is_floating_point_v<Arg3>) && \
+			(!std::is_same_v<Arg1, Arg2> || !std::is_same_v<Arg2, Arg3> || !std::is_same_v<Arg1, Arg3>) \
+		) \
 	> \
 	static FORCEINLINE auto Func(Arg1 X, Arg2 Y, Arg3 Z) -> decltype(X * Y * Z) \
 	{ \
@@ -137,12 +129,10 @@
 	template < \
 		typename Arg1, \
 		typename Arg2 \
-		UE_CONSTRAINTS_BEGIN \
-			UE_CONSTRAINT( \
-				(std::is_floating_point_v<Arg1> || std::is_floating_point_v<Arg2>) && \
-				!std::is_same_v<Arg1, Arg2> \
-			) \
-		UE_CONSTRAINTS_END \
+		UE_REQUIRES( \
+			(std::is_floating_point_v<Arg1> || std::is_floating_point_v<Arg2>) && \
+			!std::is_same_v<Arg1, Arg2> \
+		) \
 	> \
 	static FORCEINLINE ReturnType Func(Arg1 X, Arg2 Y) \
 	{ \
@@ -156,12 +146,10 @@
 		typename Arg1, \
 		typename Arg2, \
 		typename Arg3 \
-		UE_CONSTRAINTS_BEGIN \
-			UE_CONSTRAINT( \
-				(std::is_floating_point_v<Arg1> || std::is_floating_point_v<Arg2> || std::is_floating_point_v<Arg3>) && \
-				(!std::is_same_v<Arg1, Arg2> || !std::is_same_v<Arg1, Arg3> || !std::is_same_v<Arg2, Arg3>) \
-			) \
-		UE_CONSTRAINTS_END \
+		UE_REQUIRES( \
+			(std::is_floating_point_v<Arg1> || std::is_floating_point_v<Arg2> || std::is_floating_point_v<Arg3>) && \
+			(!std::is_same_v<Arg1, Arg2> || !std::is_same_v<Arg1, Arg3> || !std::is_same_v<Arg2, Arg3>) \
+		) \
 	> \
 	static FORCEINLINE ReturnType Func(Arg1 X, Arg2 Y, Arg3 Z) \
 	{ \
diff --git a/Engine/Source/Runtime/Core/Public/Templates/Tuple.h b/Engine/Source/Runtime/Core/Public/Templates/Tuple.h
--- a/Engine/Source/Runtime/Core/Public/Templates/Tuple.h
+++ b/Engine/Source/Runtime/Core/Public/Templates/Tuple.h
@@ -71,9 +71,7 @@ namespace UE::Core::Private::Tuple
 	{
 		template <
 			typename ArgType
-			UE_CONSTRAINTS_BEGIN
-				UE_CONSTRAINT(std::is_constructible_v<T, ArgType&&>)
-			UE_CONSTRAINTS_END
+			UE_REQUIRES(std::is_constructible_v<T, ArgType&&>)
 		>
 		explicit TTupleBaseElement(EForwardingConstructor, ArgType&& Arg)
 			: Value(Forward<ArgType>(Arg))
@@ -98,9 +96,7 @@ namespace UE::Core::Private::Tuple
 	{
 		template <
 			typename ArgType
-			UE_CONSTRAINTS_BEGIN
-				UE_CONSTRAINT(std::is_constructible_v<T, ArgType&&>)
-			UE_CONSTRAINTS_END
+			UE_REQUIRES(std::is_constructible_v<T, ArgType&&>)
 		>
 		explicit TTupleBaseElement(EForwardingConstructor, ArgType&& Arg)
 			: Key(Forward<ArgType>(Arg))
@@ -394,10 +390,7 @@ namespace UE::Core::Private::Tuple
 		template <
 			typename KeyArgType,
 			typename ValueArgType
-			UE_CONSTRAINTS_BEGIN
-				UE_CONSTRAINT(std::is_constructible_v<KeyType,   KeyArgType  &&>)
-				UE_CONSTRAINT(std::is_constructible_v<ValueType, ValueArgType&&>)
-			UE_CONSTRAINTS_END
+			UE_REQUIRES(std::is_constructible_v<KeyType, KeyArgType&&> && std::is_constructible_v<ValueType, ValueArgType&&>)
 		>
 		explicit TTupleBase(EForwardingConstructor, KeyArgType&& KeyArg, ValueArgType&& ValueArg)
 			: Key  (Forward<KeyArgType  >(KeyArg  ))
@@ -408,10 +401,7 @@ namespace UE::Core::Private::Tuple
 		template <
 			typename TupleType,
 			typename std::decay_t<TupleType>::DummyPairIdentifier* = nullptr
-			UE_CONSTRAINTS_BEGIN
-				UE_CONSTRAINT(std::is_constructible_v<KeyType,   decltype(DeclVal<TupleType&&>().Get<0>())>)
-				UE_CONSTRAINT(std::is_constructible_v<ValueType, decltype(DeclVal<TupleType&&>().Get<1>())>)
-			UE_CONSTRAINTS_END
+			UE_REQUIRES(std::is_constructible_v<KeyType, decltype(DeclVal<TupleType&&>().Get<0>())> && std::is_constructible_v<ValueType, decltype(DeclVal<TupleType&&>().Get<1>())>)
 		>
 		explicit TTupleBase(EOtherTupleConstructor, TupleType&& Other)
 			: Key  (Forward<TupleType>(Other).Get<0>())
@@ -687,9 +677,7 @@ public:
 	template <
 		typename... ArgTypes,
 		decltype(UE::Core::Private::Tuple::ConstructibleConceptCheck<Types...>(DeclVal<ArgTypes&&>()...))* = nullptr
-		UE_CONSTRAINTS_BEGIN
-			UE_CONSTRAINT((std::is_convertible_v<ArgTypes&&, Types> && ...))
-		UE_CONSTRAINTS_END
+		UE_REQUIRES((std::is_convertible_v<ArgTypes&&, Types> && ...))
 	>
 	TTuple(ArgTypes&&... Args)
 		: Super(UE::Core::Private::Tuple::ForwardingConstructor, Forward<ArgTypes>(Args)...)
@@ -699,9 +687,7 @@ public:
 	template <
 		typename... ArgTypes,
 		decltype(UE::Core::Private::Tuple::ConstructibleConceptCheck<Types...>(DeclVal<ArgTypes&&>()...))* = nullptr
-		UE_CONSTRAINTS_BEGIN
-			UE_CONSTRAINT(!(std::is_convertible_v<ArgTypes&&, Types> && ...))
-		UE_CONSTRAINTS_END
+		UE_REQUIRES(!(std::is_convertible_v<ArgTypes&&, Types> && ...))
 	>
 	explicit TTuple(ArgTypes&&... Args)
 		: Super(UE::Core::Private::Tuple::ForwardingConstructor, Forward<ArgTypes>(Args)...)
diff --git a/Engine/Source/Runtime/Core/Public/Templates/UniquePtr.h b/Engine/Source/Runtime/Core/Public/Templates/UniquePtr.h
--- a/Engine/Source/Runtime/Core/Public/Templates/UniquePtr.h
+++ b/Engine/Source/Runtime/Core/Public/Templates/UniquePtr.h
@@ -194,8 +194,8 @@ public:
 	template <
 		typename OtherT,
 		typename OtherDeleter,
-		typename = typename TEnableIf<!TIsArray<OtherT>::Value>::Type,
 		typename = decltype(ImplicitConv<T*>((OtherT*)nullptr))
+		UE_REQUIRES(!std::is_array_v<OtherT>)
 	>
 	FORCEINLINE TUniquePtr(TUniquePtr<OtherT, OtherDeleter>&& Other)
 		: Deleter(MoveTemp(Other.GetDeleter()))
@@ -229,8 +229,8 @@ public:
 	template <
 		typename OtherT,
 		typename OtherDeleter,
-		typename = typename TEnableIf<!TIsArray<OtherT>::Value>::Type,
 		typename = decltype(ImplicitConv<T*>((OtherT*)nullptr))
+		UE_REQUIRES(!std::is_array_v<OtherT>)
 	>
 	FORCEINLINE TUniquePtr& operator=(TUniquePtr<OtherT, OtherDeleter>&& Other)
 	{
@@ -798,8 +798,12 @@ struct TIsBitwiseConstructible<TUniquePtr<T>, T*>
  *
  * @return A TUniquePtr which points to a newly-constructed T with the specified Args.
  */
-template <typename T, typename... TArgs>
-FORCEINLINE typename TEnableIf<!TIsArray<T>::Value, TUniquePtr<T>>::Type MakeUnique(TArgs&&... Args)
+template <
+	typename T,
+	typename... TArgs
+	UE_REQUIRES(!std::is_array_v<T>)
+>
+FORCEINLINE TUniquePtr<T> MakeUnique(TArgs&&... Args)
 {
 	return TUniquePtr<T>(new T(Forward<TArgs>(Args)...));
 }
@@ -813,7 +817,7 @@ FORCEINLINE typename TEnableIf<!TIsArray<T>::Value, TUniquePtr<T>>::Type MakeUni
  * @return A TUniquePtr which points to a newly-constructed T with the specified Args.
  */
 template <typename T, typename... TArgs>
-FORCEINLINE typename TEnableIf<!TIsArray<T>::Value, TUniquePtr<T>>::Type MakeUniqueForOverwrite()
+FORCEINLINE typename TEnableIf<!std::is_array_v<T>, TUniquePtr<T>>::Type MakeUniqueForOverwrite()
 {
 	return TUniquePtr<T>(new T);
 }
diff --git a/Engine/Source/Runtime/Core/Public/Templates/UnrealTypeTraits.h b/Engine/Source/Runtime/Core/Public/Templates/UnrealTypeTraits.h
--- a/Engine/Source/Runtime/Core/Public/Templates/UnrealTypeTraits.h
+++ b/Engine/Source/Runtime/Core/Public/Templates/UnrealTypeTraits.h
@@ -23,16 +23,13 @@
 #include "Templates/IsTriviallyCopyConstructible.h"
 
 /*-----------------------------------------------------------------------------
-	Readability macro for enable_if in template definitions, future-proofed
+	Readability macro for a constraint in template definitions, future-proofed
 	for C++ 20 concepts. Usage:
 
 	template <
 		typename T,
-		typename U  // note - no trailing comma before the constraints block
-		UE_CONSTRAINTS_BEGIN
-			UE_CONSTRAINT(std::is_integral_v<T>)
-			UE_CONSTRAINT(sizeof(U) <= 4)
-		UE_CONSTRAINTS_END
+		typename U  // note - no trailing comma before the constraint
+		UE_REQUIRES(std::is_integral_v<T> && sizeof(U) <= 4)
 	>
 	void IntegralUpTo32Bit(T Lhs, U Rhs) {}
  -----------------------------------------------------------------------------*/
@@ -40,10 +37,12 @@
     #define UE_CONSTRAINTS_BEGIN , std::enable_if_t<
     #define UE_CONSTRAINT(...) (__VA_ARGS__) &&
     #define UE_CONSTRAINTS_END true, int> = 0
+
+	#define UE_REQUIRES(...) , std::enable_if_t<(__VA_ARGS__), int> = 0
 #else
     namespace UE::Core::Private
     {
-        // Only needed for the UE_CONSTRAINT* macros to work
+        // Only needed for the UE_REQUIRES macro to work, to allow for a trailing > token after the macro
         template <bool B>
         concept BoolIdentityConcept = B;
     }
@@ -51,6 +50,8 @@
     #define UE_CONSTRAINTS_BEGIN > requires
     #define UE_CONSTRAINT(...) (__VA_ARGS__) &&
     #define UE_CONSTRAINTS_END UE::Core::Private::BoolIdentityConcept<true
+
+	#define UE_REQUIRES(...) > requires (__VA_ARGS__) && UE::Core::Private::BoolIdentityConcept<true
 #endif
 
 #define TEMPLATE_REQUIRES(...) typename TEnableIf<__VA_ARGS__, int>::type = 0
