Ensure that native text resources are always a lower priority than localized text resources when resolving the current set of translation data

#rnx

diff --git a/Engine/Source/Runtime/Core/Private/Internationalization/LocalizationResourceTextSource.cpp b/Engine/Source/Runtime/Core/Private/Internationalization/LocalizationResourceTextSource.cpp
--- a/Engine/Source/Runtime/Core/Private/Internationalization/LocalizationResourceTextSource.cpp
+++ b/Engine/Source/Runtime/Core/Private/Internationalization/LocalizationResourceTextSource.cpp
@@ -164,6 +164,7 @@ void FLocalizationResourceTextSource::LoadLocalizedResourcesFromPaths(TArrayView
 	SCOPED_BOOT_TIMING("LoadLocalizedResourcesFromPaths");
 
 	const int32 BaseResourcePriority = GetPriority() * -1; // Flip the priority as larger text source priorities are more important, but smaller text resource priorities are more important
+	const int32 NativeResourcePriority = BaseResourcePriority + InPrioritizedCultures.Num(); // Native resources always prioritize below any localized ones
 
 	static const FString PlatformLocalizationFolderName = FPaths::GetPlatformLocalizationFolderName();
 	static const FString PlatformName = ANSI_TO_TCHAR(FPlatformProperties::IniPlatformName());
@@ -207,7 +208,7 @@ void FLocalizationResourceTextSource::LoadLocalizedResourcesFromPaths(TArrayView
 				// We skip loading the native text if we're transitioning to the native culture as there's no extra work that needs to be done
 				if (!LocMetaResource.NativeCulture.IsEmpty() && !InPrioritizedCultures.Contains(LocMetaResource.NativeCulture))
 				{
-					LoadLocalizationResourcesForCulture(InOutNativeResource, LocalizationPath, LocMetaResource.NativeCulture, FPaths::GetCleanFilename(LocMetaResource.NativeLocRes), BaseResourcePriority);
+					LoadLocalizationResourcesForCulture(InOutNativeResource, LocalizationPath, LocMetaResource.NativeCulture, FPaths::GetCleanFilename(LocMetaResource.NativeLocRes), NativeResourcePriority);
 				}
 			}
 		}
@@ -227,7 +228,7 @@ void FLocalizationResourceTextSource::LoadLocalizedResourcesFromPaths(TArrayView
 				}
 
 				const FString LocResFilename = FPaths::GetBaseFilename(LocalizationPath) + TEXT(".locres");
-				LoadLocalizationResourcesForCulture(InOutLocalizedResource, LocalizationPath, NativeGameCulture, LocResFilename, BaseResourcePriority);
+				LoadLocalizationResourcesForCulture(InOutLocalizedResource, LocalizationPath, NativeGameCulture, LocResFilename, NativeResourcePriority);
 			}
 		}
 	}
@@ -246,7 +247,7 @@ void FLocalizationResourceTextSource::LoadLocalizedResourcesFromPaths(TArrayView
 				}
 
 				const FString LocResFilename = FPaths::GetBaseFilename(LocalizationPath) + TEXT(".locres");
-				LoadLocalizationResourcesForCulture(InOutLocalizedResource, LocalizationPath, PrioritizedCultureName, LocResFilename, BaseResourcePriority + CultureIndex + 1);
+				LoadLocalizationResourcesForCulture(InOutLocalizedResource, LocalizationPath, PrioritizedCultureName, LocResFilename, BaseResourcePriority + CultureIndex);
 			}
 		}
 	}
diff --git a/Engine/Source/Runtime/Core/Private/Internationalization/PolyglotTextSource.cpp b/Engine/Source/Runtime/Core/Private/Internationalization/PolyglotTextSource.cpp
--- a/Engine/Source/Runtime/Core/Private/Internationalization/PolyglotTextSource.cpp
+++ b/Engine/Source/Runtime/Core/Private/Internationalization/PolyglotTextSource.cpp
@@ -112,6 +112,7 @@ void FPolyglotTextSource::AddPolyglotDataToResource(const FPolyglotTextData& InP
 	};
 
 	const int32 BaseResourcePriority = GetPriority() * -1; // Flip the priority as larger text source priorities are more important, but smaller text resource priorities are more important
+	const int32 NativeResourcePriority = BaseResourcePriority + InPrioritizedCultures.Num(); // Native resources always prioritize below any localized ones
 	const FString NativeCulture = InPolyglotTextData.ResolveNativeCulture();
 
 	// We skip loading the native text if we're transitioning to the native culture as there's no extra work that needs to be done
@@ -125,7 +126,7 @@ void FPolyglotTextSource::AddPolyglotDataToResource(const FPolyglotTextData& InP
 				InPolyglotTextData.GetKey(),
 				InPolyglotTextData.GetNativeString(),
 				LocalizedString,
-				BaseResourcePriority
+				NativeResourcePriority
 				);
 		}
 	}
@@ -143,7 +144,7 @@ void FPolyglotTextSource::AddPolyglotDataToResource(const FPolyglotTextData& InP
 					InPolyglotTextData.GetKey(),
 					InPolyglotTextData.GetNativeString(),
 					LocalizedString,
-					BaseResourcePriority
+					NativeResourcePriority
 					);
 			}
 		}
@@ -159,7 +160,7 @@ void FPolyglotTextSource::AddPolyglotDataToResource(const FPolyglotTextData& InP
 					InPolyglotTextData.GetKey(),
 					InPolyglotTextData.GetNativeString(),
 					LocalizedString,
-					BaseResourcePriority + LocalizedPriority + 1
+					BaseResourcePriority + LocalizedPriority
 					);
 			}
 		}
