Filter chunked localization on PAK mount against known "game" targets

This matches the logic used by FLocalizationResourceTextSource::LoadLocalizedResources when determining which localization paths may contain chunked target data

#rnx

diff --git a/Engine/Source/Runtime/Core/Private/Internationalization/TextLocalizationManager.cpp b/Engine/Source/Runtime/Core/Private/Internationalization/TextLocalizationManager.cpp
--- a/Engine/Source/Runtime/Core/Private/Internationalization/TextLocalizationManager.cpp
+++ b/Engine/Source/Runtime/Core/Private/Internationalization/TextLocalizationManager.cpp
@@ -1196,8 +1196,22 @@ void FTextLocalizationManager::LoadChunkedLocalizationResources_Sync(TArrayView<
 		return;
 	}
 
+	TArray<FString> GameLocalizationPaths;
+	GameLocalizationPaths += FPaths::GetGameLocalizationPaths();
+#if UE_IS_COOKED_EDITOR
+	if (GIsEditor)
+	{
+		// Cooked editors may also load game localization targets
+		GameLocalizationPaths += FPaths::GetCookedEditorLocalizationPaths();
+	}
+#endif
+
 	// Note: We only allow game localization targets to be chunked, and the layout is assumed to follow our standard pattern (as used by the localization dashboard and FLocTextHelper)
-	const TArray<FString> ChunkedLocalizationTargets = FLocalizationResourceTextSource::GetChunkedLocalizationTargets();
+	TArray<FString> ChunkedLocalizationTargets = FLocalizationResourceTextSource::GetChunkedLocalizationTargets();
+	ChunkedLocalizationTargets.RemoveAll([&GameLocalizationPaths](const FString& LocalizationTarget)
+	{
+		return !GameLocalizationPaths.Contains(FPaths::ProjectContentDir() / TEXT("Localization") / LocalizationTarget);
+	});
 
 	// Check to see whether all the required localization data is now available
 	// This may not be the case if this PAK was split into multiple sub-files, and the localization data was split between them
