Added clamp when parsing EntriesCount from a LocRes

#rnx

diff --git a/Engine/Source/Runtime/Core/Private/Internationalization/TextLocalizationResource.cpp b/Engine/Source/Runtime/Core/Private/Internationalization/TextLocalizationResource.cpp
--- a/Engine/Source/Runtime/Core/Private/Internationalization/TextLocalizationResource.cpp
+++ b/Engine/Source/Runtime/Core/Private/Internationalization/TextLocalizationResource.cpp
@@ -295,7 +295,7 @@ bool FTextLocalizationResource::LoadFromArchive(FArchive& Archive, const FTextKe
 	{
 		uint32 EntriesCount;
 		Archive << EntriesCount;
-		Entries.Reserve(Entries.Num() + EntriesCount);
+		Entries.Reserve(Entries.Num() + FMath::Max(0, (int32)EntriesCount));
 	}
 
 	// Read namespace count
