Fixed so client can build with non unity/pch

diff --git a/Engine/Source/Runtime/Core/Private/HAL/MallocStomp2.cpp b/Engine/Source/Runtime/Core/Private/HAL/MallocStomp2.cpp
--- a/Engine/Source/Runtime/Core/Private/HAL/MallocStomp2.cpp
+++ b/Engine/Source/Runtime/Core/Private/HAL/MallocStomp2.cpp
@@ -1,13 +1,15 @@
 // Copyright Epic Games, Inc. All Rights Reserved.
 
 #include "HAL/MallocStomp2.h"
+
+#if WITH_MALLOC_STOMP2
 #include "Math/UnrealMathUtility.h"
+#include "Misc/Parse.h"
 #include "HAL/UnrealMemory.h"
 #include "HAL/IConsoleManager.h"
 #include "GenericPlatform/GenericPlatformMemory.h"
 #include "Algo/BinarySearch.h"
 
-#if WITH_MALLOC_STOMP2
 CORE_API FMallocStomp2* GMallocStomp2 = nullptr;
 CORE_API bool GMallocStomp2Enabled = false;
 
diff --git a/Engine/Source/Runtime/Core/Public/Windows/COMPointer.h b/Engine/Source/Runtime/Core/Public/Windows/COMPointer.h
--- a/Engine/Source/Runtime/Core/Public/Windows/COMPointer.h
+++ b/Engine/Source/Runtime/Core/Public/Windows/COMPointer.h
@@ -2,6 +2,8 @@
 
 #pragma once
 
+// HEADER_UNIT_SKIP - Should not be included
+
 // Producing a warning that is not treated as an error here calls non-deterministic incremental build warnings
 //#pragma message(__FILE__"(5): warning : #include Microsoft/COMPointer.h instead of Windows/COMPointer.h")
 
