Fix CIS static analysis warnings (reccuring check)

diff --git a/Engine/Source/Runtime/Core/Tests/Containers/BitArrayTests.cpp b/Engine/Source/Runtime/Core/Tests/Containers/BitArrayTests.cpp
--- a/Engine/Source/Runtime/Core/Tests/Containers/BitArrayTests.cpp
+++ b/Engine/Source/Runtime/Core/Tests/Containers/BitArrayTests.cpp
@@ -1239,11 +1239,13 @@ public:
 					ZeroesText[DisplayIndex] = TEXT('\0');
 					if (!bOnesMatchesExpected)
 					{
+						//-V571
 						CHECK_MESSAGE(FString::Printf(TEXT("SetRange bool BitWidth=%d WriteOffset=%d, Ones\nExpected=%s\nActual  =%s"), BitWidth, WriteOffset, OnesExpectedText, OnesText),
 							bOnesMatchesExpected);
 					}
 					else if (!bZeroesMatchesExpected)
 					{
+						//-V571
 						CHECK_MESSAGE(FString::Printf(TEXT("SetRange bool BitWidth=%d WriteOffset=%d, Zeroes\nExpected=%s\nActual  =%s"), BitWidth, WriteOffset, ZeroesExpectedText, ZeroesText),
 							bZeroesMatchesExpected);
 					}
@@ -2191,11 +2193,13 @@ public:
 						ZeroesText[DisplayIndex] = TEXT('\0');
 						if (!bOnesMatchesExpected)
 						{
+							//-V571
 							CHECK_MESSAGE(FString::Printf(TEXT("MemmoveBitsWordOrder BitWidth=%d ReadOffset=%d, WriteOffset=%d, Overlap=%d, Ones\nExpected=%s\nActual  =%s"), BitWidth, ReadOffset, WriteOffset, Overlap, OnesExpectedText, OnesText),
 								bOnesMatchesExpected);
 						}
 						else if (!bZeroesMatchesExpected)
 						{
+							//-V571
 							CHECK_MESSAGE(FString::Printf(TEXT("MemmoveBitsWordOrder BitWidth=%d ReadOffset=%d, WriteOffset=%d, Overlap=%d, Zeroes\nExpected=%s\nActual  =%s"), BitWidth, ReadOffset, WriteOffset, Overlap, ZeroesExpectedText, ZeroesText),
 								bZeroesMatchesExpected);
 						}
