Logging: Extended UE_SLOG to handle both positional and named fields and up to 16 fields

Added UE_SLOG_EX to handle more than 16 fields.
Added UE_SLOG_VALUE to handle values in UE_SLOG_EX.
Deprecated UE_SLOG_NAMED and UE_SLOG_NAMED_EX.
Renamed static variables to reduce the likelihood of shadowing warnings.

Optimized field values to be passed by value when appropriate, which reduces instruction count at the site of UE_SLOG at the cost of added instructions in LogWithFields.

#rnx

diff --git a/Engine/Source/Runtime/Core/Public/Logging/StructuredLog.h b/Engine/Source/Runtime/Core/Public/Logging/StructuredLog.h
--- a/Engine/Source/Runtime/Core/Public/Logging/StructuredLog.h
+++ b/Engine/Source/Runtime/Core/Public/Logging/StructuredLog.h
@@ -9,6 +9,7 @@
 #include "Serialization/CompactBinary.h"
 #include "Serialization/CompactBinaryWriter.h"
 #include "Templates/IsArrayOrRefOfType.h"
+#include "Templates/UnrealTypeTraits.h"
 
 #include <atomic>
 #include <type_traits>
@@ -22,54 +23,39 @@ class FTextFormat;
 /**
  * Records a structured log event if this category is active at this level of verbosity.
  *
- * Example:
- * UE_SLOG(LogCore, Warning, "Loading '{Name}' failed with error {Error}", Package->GetName(), ErrorCode);
- *
- * Field names must match "[A-Za-z0-9_]+" and must be unique within this log event.
- * Field values will be serialized using SerializeForLog or operator<<(FCbWriter&, FieldType).
- * The field values must exactly match the fields referenced by Format.
+ * Supports either positional or named parameters, but not a mix of these styles.
  *
- * @param CategoryName   Name of a log category declared by DECLARE_LOG_CATEGORY_*.
- * @param Verbosity      Name of a log verbosity level from ELogVerbosity.
- * @param Format         Format string in the style of FLogTemplate.
- * @param FieldValues*   Zero or more field values that correspond to the names in Format.
- */
-#define UE_SLOG(CategoryName, Verbosity, Format, ...) UE_PRIVATE_SLOG(CategoryName, Verbosity, Format, Anonymous, ##__VA_ARGS__)
-
-/**
- * Records a structured log event if this category is active at this level of verbosity.
+ * Positional: The field values must exactly match the fields referenced by Format.
+ * UE_SLOG(LogCore, Warning, "Loading '{Name}' failed with error {Error}", Package->GetName(), ErrorCode);
  *
- * Example:
- * UE_SLOG_NAMED(LogCore, Warning, "Loading '{PackageName}' failed with error {Error}",
+ * Named: The fields must contain every field referenced by Format. Order is irrelevant and extra fields are permitted.
+ * UE_SLOG(LogCore, Warning, "Loading '{Name}' failed with error {Error}",
  *     ("Name", Package->GetName()), ("Error", ErrorCode), ("Flags", LoadFlags));
  *
  * Field names must match "[A-Za-z0-9_]+" and must be unique within this log event.
  * Field values will be serialized using SerializeForLog or operator<<(FCbWriter&, FieldType).
- * The fields must contain every field referenced by Format. Order is irrelevant and extra fields are permitted.
  *
  * @param CategoryName   Name of a log category declared by DECLARE_LOG_CATEGORY_*.
  * @param Verbosity      Name of a log verbosity level from ELogVerbosity.
  * @param Format         Format string in the style of FLogTemplate.
- * @param Fields[0-9]    Zero to nine fields in the format ("Name", Value).
+ * @param Fields[0-16]   Zero to sixteen fields or field values.
  */
-#define UE_SLOG_NAMED(CategoryName, Verbosity, Format, ...) UE_PRIVATE_SLOG_CALL(UE_SLOG_NAMED_EX, (CategoryName, Verbosity, Format UE_PRIVATE_SLOG_FIELDS(__VA_ARGS__)))
+#define UE_SLOG(CategoryName, Verbosity, Format, ...) UE_PRIVATE_SLOG_CALL(UE_SLOG_EX, (CategoryName, Verbosity, Format UE_PRIVATE_SLOG_FIELDS(__VA_ARGS__)))
 
 /**
  * Records a structured log event if this category is active at this level of verbosity.
  *
- * Example:
- * UE_SLOG_NAMED_EX(LogCore, Warning, "Loading '{PackageName}' failed with error {Error}",
- *     UE_SLOG_FIELD("Name", Package->GetName()), UE_SLOG_FIELD("Error", ErrorCode), UE_SLOG_FIELD("Flags", LoadFlags));
+ * This has the same functionality as UE_SLOG but removes the limit on field count.
  *
- * Same as UE_SLOG_NAMED and works for any number of fields.
- * Fields must be written as UE_SLOG_FIELD("Name", Value).
- */
-#define UE_SLOG_NAMED_EX(CategoryName, Verbosity, Format, ...) UE_PRIVATE_SLOG(CategoryName, Verbosity, Format, Named, ##__VA_ARGS__)
-
-/**
- * Expands to a named structured log field. UE_SLOG_FIELD("Name", Value)
+ * Positional: Values must be wrapped in UE_SLOG_VALUE.
+ * UE_SLOG_EX(LogCore, Warning, "Loading '{Name}' failed with error {Error}",
+ *     UE_SLOG_VALUE(Package->GetName()), UE_SLOG_VALUE(ErrorCode));
+ *
+ * Named: Fields must be wrapped in UE_SLOG_FIELD.
+ * UE_SLOG_EX(LogCore, Warning, "Loading '{Name}' failed with error {Error}",
+ *     UE_SLOG_FIELD("Name", Package->GetName()), UE_SLOG_FIELD("Error", ErrorCode), UE_SLOG_FIELD("Flags", LoadFlags));
  */
-#define UE_SLOG_FIELD(Name, Value) ::UE::Logging::Private::CheckFieldName(Name), Value
+#define UE_SLOG_EX(CategoryName, Verbosity, Format, ...) UE_PRIVATE_SLOG(CategoryName, Verbosity, Format, ##__VA_ARGS__)
 
 /**
  * Records a localized structured log event if this category is active at this level of verbosity.
@@ -87,7 +73,7 @@ class FTextFormat;
  * @param Namespace      Namespace for the format FText, or LOCTEXT_NAMESPACE for the non-NS macro.
  * @param Key            Key for the format FText that is unique within the namespace.
  * @param Format         Format string in the style of FTextFormat.
- * @param Fields[0-9]    Zero to nine fields in the format ("Name", Value).
+ * @param Fields[0-16]   Zero to sixteen fields in the format ("Name", Value).
  */
 #define UE_SLOG_LOC(CategoryName, Verbosity, Key, Format, ...) \
 	UE_SLOG_NSLOC(CategoryName, Verbosity, LOCTEXT_NAMESPACE, Key, Format, ##__VA_ARGS__)
@@ -109,20 +95,30 @@ class FTextFormat;
 #define UE_SLOG_NSLOC_EX(CategoryName, Verbosity, Namespace, Key, Format, ...) \
 	UE_PRIVATE_SLOG_LOC(CategoryName, Verbosity, Namespace, Key, Format, ##__VA_ARGS__)
 
+/** Expands to a named structured log field. UE_SLOG_FIELD("Name", Value) */
+#define UE_SLOG_FIELD(Name, Value) UE_PRIVATE_SLOG_FIELD((Name, Value))
+
+/** Expands to a structured log value. UE_SLOG_VALUE(Value) */
+#define UE_SLOG_VALUE(Value) Value
+
 #else // #if !NO_LOGGING
 
 #define UE_SLOG(...)
-#define UE_SLOG_NAMED(...)
-#define UE_SLOG_NAMED_EX(...)
-#define UE_SLOG_FIELD(Name, Value)
+#define UE_SLOG_EX(...)
 
 #define UE_SLOG_LOC(...)
 #define UE_SLOG_NSLOC(...)
 #define UE_SLOG_LOC_EX(...)
 #define UE_SLOG_NSLOC_EX(...)
 
+#define UE_SLOG_FIELD(Name, Value)
+#define UE_SLOG_VALUE(Value)
+
 #endif
 
+#define UE_SLOG_NAMED UE_DEPRECATED_MACRO(5.2, "UE_SLOG_NAMED is deprecated. Use UE_SLOG.") UE_SLOG
+#define UE_SLOG_NAMED_EX UE_DEPRECATED_MACRO(5.2, "UE_SLOG_NAMED_EX is deprecated. Use UE_SLOG_EX.") UE_SLOG_EX
+
 struct FDateTime;
 
 namespace UE
@@ -223,7 +219,7 @@ private:
 template <typename ValueType>
 inline void SerializeForLog(FCbWriter& Writer, ValueType&& Value)
 {
-	Writer << Forward<ValueType>(Value);
+	Writer << (ValueType&&)Value;
 }
 
 } // UE
@@ -290,63 +286,77 @@ UE_API void LogWithFieldArray(const FStaticLogRecord& Log, const FLogField* Fiel
 UE_API void LogWithNoFields(const FStaticLocalizedLogRecord& Log);
 UE_API void LogWithFieldArray(const FStaticLocalizedLogRecord& Log, const FLogField* Fields, int32 FieldCount);
 
-/** Creates log fields with no name from value arguments. */
-struct FLogFieldsAnonymous
+/** Wrapper to identify field names interleaved with field values. */
+template <typename NameType>
+struct TLogFieldName
+{
+	NameType Name;
+};
+
+/** Verify that the name is likely a string literal and forward it on. */
+template <typename NameType>
+inline constexpr TLogFieldName<NameType> CheckFieldName(NameType&& Name)
 {
+	static_assert(TIsArrayOrRefOfType<NameType, ANSICHAR>::Value, "Name must be an ANSICHAR string literal.");
+	return {(NameType&&)Name};
+}
+
+/** Create log fields from values optionally preceded by names. */
+struct FLogFieldCreator
+{
+	template <typename T> inline constexpr static int32 ValueCount = 1;
+	template <typename T> inline constexpr static int32 ValueCount<TLogFieldName<T>> = 0;
+
 	template <typename... FieldArgTypes>
-	static constexpr int32 GetCount()
+	inline constexpr static int32 GetCount()
 	{
-		return sizeof...(FieldArgTypes);
+		return (ValueCount<FieldArgTypes> + ...);
 	}
 
-	static inline void Create(FLogField* Fields)
+	inline static void Create(FLogField* Fields)
 	{
 	}
 
-	template <typename ValueType, typename... FieldArgTypes>
-	static inline void Create(FLogField* Fields, const ValueType& Value, FieldArgTypes&&... FieldArgs)
+	template <typename ValueType, typename... FieldArgTypes, std::enable_if_t<ValueCount<ValueType>>* = nullptr>
+	inline static void Create(FLogField* Fields, const ValueType& Value, FieldArgTypes&&... FieldArgs)
 	{
 		new(Fields) FLogField{nullptr, &Value, FLogField::Write<std::remove_reference_t<ValueType>>};
-		Create(Fields + 1, Forward<FieldArgTypes>(FieldArgs)...);
+		Create(Fields + 1, (FieldArgTypes&&)FieldArgs...);
 	}
-};
 
-/** Creates log fields from pairs of name/value arguments. */
-struct FLogFieldsNamed
-{
-	template <typename... FieldArgTypes>
-	static constexpr int32 GetCount()
+	template <typename NameType, typename ValueType, typename... FieldArgTypes>
+	inline static void Create(FLogField* Fields, TLogFieldName<NameType> Name, const ValueType& Value, FieldArgTypes&&... FieldArgs)
 	{
-		constexpr int32 FieldArgsCount = sizeof...(FieldArgTypes);
-		static_assert(FieldArgsCount % 2 == 0);
-		return FieldArgsCount / 2;
+		new(Fields) FLogField{Name.Name, &Value, FLogField::Write<std::remove_reference_t<ValueType>>};
+		Create(Fields + 1, (FieldArgTypes&&)FieldArgs...);
 	}
+};
 
-	static inline void Create(FLogField* Fields)
-	{
-	}
+template <typename T>
+struct TFieldArgType
+{
+	using Type = typename TCallTraits<T>::ParamType;
+};
 
-	template <typename ValueType, typename... FieldArgTypes>
-	static inline void Create(FLogField* Fields, const ANSICHAR* Name, const ValueType& Value, FieldArgTypes&&... FieldArgs)
-	{
-		new(Fields) FLogField{Name, &Value, FLogField::Write<std::remove_reference_t<ValueType>>};
-		Create(Fields + 1, Forward<FieldArgTypes>(FieldArgs)...);
-	}
+template <typename NameType>
+struct TFieldArgType<TLogFieldName<NameType>>
+{
+	using Type = TLogFieldName<NameType>;
 };
 
 /** Log with fields created from the arguments, which may be values or pairs of name/value. */
-template <const auto& Log, typename LogFieldsType, typename... FieldArgTypes>
-FORCENOINLINE UE_DEBUG_SECTION void LogWithFields(FieldArgTypes&&... FieldArgs)
+template <const auto& Log, typename... FieldArgTypes>
+FORCENOINLINE UE_DEBUG_SECTION void LogWithFields(typename TFieldArgType<FieldArgTypes>::Type... FieldArgs)
 {
-	constexpr int32 FieldCount = LogFieldsType::template GetCount<FieldArgTypes...>();
+	constexpr int32 FieldCount = FLogFieldCreator::template GetCount<FieldArgTypes...>();
 	static_assert(FieldCount > 0);
 	FLogField Fields[FieldCount];
-	LogFieldsType::Create(Fields, Forward<FieldArgTypes>(FieldArgs)...);
+	FLogFieldCreator::Create(Fields, (FieldArgTypes&&)FieldArgs...);
 	LogWithFieldArray(Log, Fields, FieldCount);
 }
 
 /** Log if the category is active at this level of verbosity. */
-template <const auto& Log, typename LogCategoryType, ELogVerbosity::Type Verbosity, typename LogFieldsType, typename... FieldArgTypes>
+template <const auto& Log, typename LogCategoryType, ELogVerbosity::Type Verbosity, typename... FieldArgTypes>
 inline void LogIfActive(FieldArgTypes&&... FieldArgs)
 {
 	static_assert(Verbosity != ELogVerbosity::Fatal, "Fatal verbosity is not supported by this API at this time.");
@@ -364,7 +374,7 @@ inline void LogIfActive(FieldArgTypes&&... FieldArgs)
 			}
 			else
 			{
-				LogWithFields<Log, LogFieldsType>(Forward<FieldArgTypes>(FieldArgs)...);
+				LogWithFields<Log, FieldArgTypes...>((FieldArgTypes&&)FieldArgs...);
 			}
 		}
 	}
@@ -377,49 +387,62 @@ inline void LogIfActive(FieldArgTypes&&... FieldArgs)
 	}
 }
 
-/** Verify that the name is likely a string literal and forward it on. */
-template <typename NameType>
-constexpr decltype(auto) CheckFieldName(NameType&& Name)
-{
-	static_assert(TIsArrayOrRefOfType<NameType, ANSICHAR>::Value, "Name must be an ANSICHAR string literal.");
-	return Forward<NameType>(Name);
-}
-
 } // UE::Logging::Private
 
-#define UE_PRIVATE_SLOG(CategoryName, Verbosity, Format, FieldsType, ...) \
+#define UE_PRIVATE_SLOG(CategoryName, Verbosity, Format, ...) \
 	do \
 	{ \
-		static ::UE::Logging::Private::FStaticLogDynamicData DynamicData; \
-		static constexpr ::UE::Logging::Private::FStaticLogRecord Log{CategoryName, TEXT(Format), __FILE__, __LINE__, ::ELogVerbosity::Verbosity, DynamicData}; \
-		::UE::Logging::Private::LogIfActive<Log, FLogCategory##CategoryName, ::ELogVerbosity::Verbosity, ::UE::Logging::Private::FLogFields##FieldsType>(__VA_ARGS__); \
+		static ::UE::Logging::Private::FStaticLogDynamicData LOG_Dynamic; \
+		static constexpr ::UE::Logging::Private::FStaticLogRecord LOG_Static{CategoryName, TEXT(Format), __FILE__, __LINE__, ::ELogVerbosity::Verbosity, LOG_Dynamic}; \
+		::UE::Logging::Private::LogIfActive<LOG_Static, FLogCategory##CategoryName, ::ELogVerbosity::Verbosity>(__VA_ARGS__); \
 	} \
 	while (false)
 
 #define UE_PRIVATE_SLOG_LOC(CategoryName, Verbosity, Namespace, Key, Format, ...) \
 	do \
 	{ \
-		static ::UE::Logging::Private::FStaticLogDynamicData DynamicData; \
-		static constexpr ::UE::Logging::Private::FStaticLocalizedLogRecord Log{CategoryName, TEXT(Namespace), TEXT(Key), TEXT(Format), __FILE__, __LINE__, ::ELogVerbosity::Verbosity, DynamicData}; \
-		::UE::Logging::Private::LogIfActive<Log, FLogCategory##CategoryName, ::ELogVerbosity::Verbosity, ::UE::Logging::Private::FLogFieldsNamed>(__VA_ARGS__); \
+		static ::UE::Logging::Private::FStaticLogDynamicData LOG_Dynamic; \
+		static constexpr ::UE::Logging::Private::FStaticLocalizedLogRecord LOG_Static{CategoryName, TEXT(Namespace), TEXT(Key), TEXT(Format), __FILE__, __LINE__, ::ELogVerbosity::Verbosity, LOG_Dynamic}; \
+		::UE::Logging::Private::LogIfActive<LOG_Static, FLogCategory##CategoryName, ::ELogVerbosity::Verbosity>(__VA_ARGS__); \
 	} \
 	while (false)
 
+// This macro expands a field from either `(Name, Value)` or `Value`
+// A `(Name, Value)` field is converted to `CheckFieldName(Name), Value`
+// A `Value` field is passed through as `Value`
+#define UE_PRIVATE_SLOG_FIELD(Field) UE_PRIVATE_SLOG_FIELD_EXPAND(UE_PRIVATE_SLOG_NAMED_FIELD Field)
+// This macro is only called when the field was parenthesized.
+#define UE_PRIVATE_SLOG_NAMED_FIELD(Name, ...) UE_PRIVATE_SLOG_NAMED_FIELD ::UE::Logging::Private::CheckFieldName(Name), __VA_ARGS__
+// The next three macros remove UE_PRIVATE_SLOG_NAMED_FIELD from the expanded expression.
+#define UE_PRIVATE_SLOG_FIELD_EXPAND(...) UE_PRIVATE_SLOG_FIELD_EXPAND_INNER(__VA_ARGS__)
+#define UE_PRIVATE_SLOG_FIELD_EXPAND_INNER(...) UE_PRIVATE_SLOG_STRIP_ ## __VA_ARGS__
+#define UE_PRIVATE_SLOG_STRIP_UE_PRIVATE_SLOG_NAMED_FIELD
+
+// This macro expands `Arg1, Arg2` to `UE_PRIVATE_SLOG_FIELD(Arg1), UE_PRIVATE_SLOG_FIELD(Arg2), ...`
+// This macro expands `("Name1", Arg1), ("Name2", Arg2)` to `UE_PRIVATE_SLOG_FIELD(("Name1", Arg1)), UE_PRIVATE_SLOG_FIELD(("Name2", Arg2)), ...
 #define UE_PRIVATE_SLOG_FIELDS(...) UE_PRIVATE_SLOG_CALL(PREPROCESSOR_JOIN(UE_PRIVATE_SLOG_FIELDS_, UE_PRIVATE_SLOG_COUNT(__VA_ARGS__)), (__VA_ARGS__))
 
 #define UE_PRIVATE_SLOG_FIELDS_0()
-#define UE_PRIVATE_SLOG_FIELDS_1(A)                 , UE_SLOG_FIELD A
-#define UE_PRIVATE_SLOG_FIELDS_2(A,B)               , UE_SLOG_FIELD A, UE_SLOG_FIELD B
-#define UE_PRIVATE_SLOG_FIELDS_3(A,B,C)             , UE_SLOG_FIELD A, UE_SLOG_FIELD B, UE_SLOG_FIELD C
-#define UE_PRIVATE_SLOG_FIELDS_4(A,B,C,D)           , UE_SLOG_FIELD A, UE_SLOG_FIELD B, UE_SLOG_FIELD C, UE_SLOG_FIELD D
-#define UE_PRIVATE_SLOG_FIELDS_5(A,B,C,D,E)         , UE_SLOG_FIELD A, UE_SLOG_FIELD B, UE_SLOG_FIELD C, UE_SLOG_FIELD D, UE_SLOG_FIELD E
-#define UE_PRIVATE_SLOG_FIELDS_6(A,B,C,D,E,F)       , UE_SLOG_FIELD A, UE_SLOG_FIELD B, UE_SLOG_FIELD C, UE_SLOG_FIELD D, UE_SLOG_FIELD E, UE_SLOG_FIELD F
-#define UE_PRIVATE_SLOG_FIELDS_7(A,B,C,D,E,F,G)     , UE_SLOG_FIELD A, UE_SLOG_FIELD B, UE_SLOG_FIELD C, UE_SLOG_FIELD D, UE_SLOG_FIELD E, UE_SLOG_FIELD F, UE_SLOG_FIELD G
-#define UE_PRIVATE_SLOG_FIELDS_8(A,B,C,D,E,F,G,H)   , UE_SLOG_FIELD A, UE_SLOG_FIELD B, UE_SLOG_FIELD C, UE_SLOG_FIELD D, UE_SLOG_FIELD E, UE_SLOG_FIELD F, UE_SLOG_FIELD G, UE_SLOG_FIELD H
-#define UE_PRIVATE_SLOG_FIELDS_9(A,B,C,D,E,F,G,H,I) , UE_SLOG_FIELD A, UE_SLOG_FIELD B, UE_SLOG_FIELD C, UE_SLOG_FIELD D, UE_SLOG_FIELD E, UE_SLOG_FIELD F, UE_SLOG_FIELD G, UE_SLOG_FIELD H, UE_SLOG_FIELD I
-
-#define UE_PRIVATE_SLOG_COUNT(...) UE_PRIVATE_SLOG_CALL(UE_PRIVATE_SLOG_COUNT_IMPL, (_, ##__VA_ARGS__, 9,8,7,6,5,4,3,2,1,0))
-#define UE_PRIVATE_SLOG_COUNT_IMPL(_, A,B,C,D,E,F,G,H,I, Count, ...) Count
+#define UE_PRIVATE_SLOG_FIELDS_1(A)                 , UE_PRIVATE_SLOG_FIELD(A)
+#define UE_PRIVATE_SLOG_FIELDS_2(A,B)               , UE_PRIVATE_SLOG_FIELD(A), UE_PRIVATE_SLOG_FIELD(B)
+#define UE_PRIVATE_SLOG_FIELDS_3(A,B,C)             , UE_PRIVATE_SLOG_FIELD(A), UE_PRIVATE_SLOG_FIELD(B), UE_PRIVATE_SLOG_FIELD(C)
+#define UE_PRIVATE_SLOG_FIELDS_4(A,B,C,D)           , UE_PRIVATE_SLOG_FIELD(A), UE_PRIVATE_SLOG_FIELD(B), UE_PRIVATE_SLOG_FIELD(C), UE_PRIVATE_SLOG_FIELD(D)
+#define UE_PRIVATE_SLOG_FIELDS_5(A,B,C,D,E)         , UE_PRIVATE_SLOG_FIELD(A), UE_PRIVATE_SLOG_FIELD(B), UE_PRIVATE_SLOG_FIELD(C), UE_PRIVATE_SLOG_FIELD(D), UE_PRIVATE_SLOG_FIELD(E)
+#define UE_PRIVATE_SLOG_FIELDS_6(A,B,C,D,E,F)       , UE_PRIVATE_SLOG_FIELD(A), UE_PRIVATE_SLOG_FIELD(B), UE_PRIVATE_SLOG_FIELD(C), UE_PRIVATE_SLOG_FIELD(D), UE_PRIVATE_SLOG_FIELD(E), UE_PRIVATE_SLOG_FIELD(F)
+#define UE_PRIVATE_SLOG_FIELDS_7(A,B,C,D,E,F,G)     , UE_PRIVATE_SLOG_FIELD(A), UE_PRIVATE_SLOG_FIELD(B), UE_PRIVATE_SLOG_FIELD(C), UE_PRIVATE_SLOG_FIELD(D), UE_PRIVATE_SLOG_FIELD(E), UE_PRIVATE_SLOG_FIELD(F), UE_PRIVATE_SLOG_FIELD(G)
+#define UE_PRIVATE_SLOG_FIELDS_8(A,B,C,D,E,F,G,H)   , UE_PRIVATE_SLOG_FIELD(A), UE_PRIVATE_SLOG_FIELD(B), UE_PRIVATE_SLOG_FIELD(C), UE_PRIVATE_SLOG_FIELD(D), UE_PRIVATE_SLOG_FIELD(E), UE_PRIVATE_SLOG_FIELD(F), UE_PRIVATE_SLOG_FIELD(G), UE_PRIVATE_SLOG_FIELD(H)
+#define UE_PRIVATE_SLOG_FIELDS_9(A,B,C,D,E,F,G,H,I) , UE_PRIVATE_SLOG_FIELD(A), UE_PRIVATE_SLOG_FIELD(B), UE_PRIVATE_SLOG_FIELD(C), UE_PRIVATE_SLOG_FIELD(D), UE_PRIVATE_SLOG_FIELD(E), UE_PRIVATE_SLOG_FIELD(F), UE_PRIVATE_SLOG_FIELD(G), UE_PRIVATE_SLOG_FIELD(H), UE_PRIVATE_SLOG_FIELD(I)
+
+#define UE_PRIVATE_SLOG_FIELDS_10(A,B,C,D,E,F,G,H,I,J)             , UE_PRIVATE_SLOG_FIELD(A), UE_PRIVATE_SLOG_FIELD(B), UE_PRIVATE_SLOG_FIELD(C), UE_PRIVATE_SLOG_FIELD(D), UE_PRIVATE_SLOG_FIELD(E), UE_PRIVATE_SLOG_FIELD(F), UE_PRIVATE_SLOG_FIELD(G), UE_PRIVATE_SLOG_FIELD(H), UE_PRIVATE_SLOG_FIELD(I), UE_PRIVATE_SLOG_FIELD(J)
+#define UE_PRIVATE_SLOG_FIELDS_11(A,B,C,D,E,F,G,H,I,J,K)           , UE_PRIVATE_SLOG_FIELD(A), UE_PRIVATE_SLOG_FIELD(B), UE_PRIVATE_SLOG_FIELD(C), UE_PRIVATE_SLOG_FIELD(D), UE_PRIVATE_SLOG_FIELD(E), UE_PRIVATE_SLOG_FIELD(F), UE_PRIVATE_SLOG_FIELD(G), UE_PRIVATE_SLOG_FIELD(H), UE_PRIVATE_SLOG_FIELD(I), UE_PRIVATE_SLOG_FIELD(J), UE_PRIVATE_SLOG_FIELD(K)
+#define UE_PRIVATE_SLOG_FIELDS_12(A,B,C,D,E,F,G,H,I,J,K,L)         , UE_PRIVATE_SLOG_FIELD(A), UE_PRIVATE_SLOG_FIELD(B), UE_PRIVATE_SLOG_FIELD(C), UE_PRIVATE_SLOG_FIELD(D), UE_PRIVATE_SLOG_FIELD(E), UE_PRIVATE_SLOG_FIELD(F), UE_PRIVATE_SLOG_FIELD(G), UE_PRIVATE_SLOG_FIELD(H), UE_PRIVATE_SLOG_FIELD(I), UE_PRIVATE_SLOG_FIELD(J), UE_PRIVATE_SLOG_FIELD(K), UE_PRIVATE_SLOG_FIELD(L)
+#define UE_PRIVATE_SLOG_FIELDS_13(A,B,C,D,E,F,G,H,I,J,K,L,M)       , UE_PRIVATE_SLOG_FIELD(A), UE_PRIVATE_SLOG_FIELD(B), UE_PRIVATE_SLOG_FIELD(C), UE_PRIVATE_SLOG_FIELD(D), UE_PRIVATE_SLOG_FIELD(E), UE_PRIVATE_SLOG_FIELD(F), UE_PRIVATE_SLOG_FIELD(G), UE_PRIVATE_SLOG_FIELD(H), UE_PRIVATE_SLOG_FIELD(I), UE_PRIVATE_SLOG_FIELD(J), UE_PRIVATE_SLOG_FIELD(K), UE_PRIVATE_SLOG_FIELD(L), UE_PRIVATE_SLOG_FIELD(M)
+#define UE_PRIVATE_SLOG_FIELDS_14(A,B,C,D,E,F,G,H,I,J,K,L,M,N)     , UE_PRIVATE_SLOG_FIELD(A), UE_PRIVATE_SLOG_FIELD(B), UE_PRIVATE_SLOG_FIELD(C), UE_PRIVATE_SLOG_FIELD(D), UE_PRIVATE_SLOG_FIELD(E), UE_PRIVATE_SLOG_FIELD(F), UE_PRIVATE_SLOG_FIELD(G), UE_PRIVATE_SLOG_FIELD(H), UE_PRIVATE_SLOG_FIELD(I), UE_PRIVATE_SLOG_FIELD(J), UE_PRIVATE_SLOG_FIELD(K), UE_PRIVATE_SLOG_FIELD(L), UE_PRIVATE_SLOG_FIELD(M), UE_PRIVATE_SLOG_FIELD(N)
+#define UE_PRIVATE_SLOG_FIELDS_15(A,B,C,D,E,F,G,H,I,J,K,L,M,N,O)   , UE_PRIVATE_SLOG_FIELD(A), UE_PRIVATE_SLOG_FIELD(B), UE_PRIVATE_SLOG_FIELD(C), UE_PRIVATE_SLOG_FIELD(D), UE_PRIVATE_SLOG_FIELD(E), UE_PRIVATE_SLOG_FIELD(F), UE_PRIVATE_SLOG_FIELD(G), UE_PRIVATE_SLOG_FIELD(H), UE_PRIVATE_SLOG_FIELD(I), UE_PRIVATE_SLOG_FIELD(J), UE_PRIVATE_SLOG_FIELD(K), UE_PRIVATE_SLOG_FIELD(L), UE_PRIVATE_SLOG_FIELD(M), UE_PRIVATE_SLOG_FIELD(N), UE_PRIVATE_SLOG_FIELD(O)
+#define UE_PRIVATE_SLOG_FIELDS_16(A,B,C,D,E,F,G,H,I,J,K,L,M,N,O,P) , UE_PRIVATE_SLOG_FIELD(A), UE_PRIVATE_SLOG_FIELD(B), UE_PRIVATE_SLOG_FIELD(C), UE_PRIVATE_SLOG_FIELD(D), UE_PRIVATE_SLOG_FIELD(E), UE_PRIVATE_SLOG_FIELD(F), UE_PRIVATE_SLOG_FIELD(G), UE_PRIVATE_SLOG_FIELD(H), UE_PRIVATE_SLOG_FIELD(I), UE_PRIVATE_SLOG_FIELD(J), UE_PRIVATE_SLOG_FIELD(K), UE_PRIVATE_SLOG_FIELD(L), UE_PRIVATE_SLOG_FIELD(M), UE_PRIVATE_SLOG_FIELD(N), UE_PRIVATE_SLOG_FIELD(O), UE_PRIVATE_SLOG_FIELD(P)
+
+#define UE_PRIVATE_SLOG_COUNT(...) UE_PRIVATE_SLOG_CALL(UE_PRIVATE_SLOG_COUNT_IMPL, (_, ##__VA_ARGS__, 16,15,14,13,12,11,10,9,8,7,6,5,4,3,2,1,0))
+#define UE_PRIVATE_SLOG_COUNT_IMPL(_, A,B,C,D,E,F,G,H,I,J,K,L,M,N,O,P, Count, ...) Count
 
 #define UE_PRIVATE_SLOG_CALL(F, A) UE_PRIVATE_SLOG_EXPAND(F A)
 #define UE_PRIVATE_SLOG_EXPAND(X) X
