Removed usage of TRemoveConst.

diff --git a/Engine/Plugins/Developer/Concert/ConcertMain/Source/ConcertTransport/Public/ConcertMessageContext.h b/Engine/Plugins/Developer/Concert/ConcertMain/Source/ConcertTransport/Public/ConcertMessageContext.h
--- a/Engine/Plugins/Developer/Concert/ConcertMain/Source/ConcertTransport/Public/ConcertMessageContext.h
+++ b/Engine/Plugins/Developer/Concert/ConcertMain/Source/ConcertTransport/Public/ConcertMessageContext.h
@@ -28,10 +28,10 @@ struct FConcertMessageContext
 
 	/** Utility to get the message data cast to the correct type */
 	template <typename T>
-	const typename TRemoveConst<T>::Type* GetMessage() const
+	const T* GetMessage() const
 	{
 		check(MessageType->IsChildOf(T::StaticStruct()));
-		return static_cast<const typename TRemoveConst<T>::Type*>(Message);
+		return static_cast<const T*>(Message);
 	}
 
 	FGuid SenderConcertEndpointId;
diff --git a/Engine/Plugins/Experimental/AnimNext/Source/AnimNextInterface/Public/AnimNextInterfaceContext.h b/Engine/Plugins/Experimental/AnimNext/Source/AnimNextInterface/Public/AnimNextInterfaceContext.h
--- a/Engine/Plugins/Experimental/AnimNext/Source/AnimNextInterface/Public/AnimNextInterfaceContext.h
+++ b/Engine/Plugins/Experimental/AnimNext/Source/AnimNextInterface/Public/AnimNextInterfaceContext.h
@@ -163,7 +163,7 @@ public:
 		}
 		else
 		{
-			using MutableValueType = typename TRemoveConst<ValueType>::Type;
+			using MutableValueType = std::remove_const_t<ValueType>;
 
 			// If non const set the mutable flag
 			if constexpr (TIsConst<ValueType>::Value == false)
diff --git a/Engine/Plugins/Experimental/AnimNext/Source/AnimNextInterface/Public/AnimNextInterfaceParam.h b/Engine/Plugins/Experimental/AnimNext/Source/AnimNextInterface/Public/AnimNextInterfaceParam.h
--- a/Engine/Plugins/Experimental/AnimNext/Source/AnimNextInterface/Public/AnimNextInterfaceParam.h
+++ b/Engine/Plugins/Experimental/AnimNext/Source/AnimNextInterface/Public/AnimNextInterfaceParam.h
@@ -170,7 +170,7 @@ struct TParamType
 {
 	FORCEINLINE static const FParamType& GetType()
 	{
-		return TParamTypeImpl<typename TRemoveConst<ValueType>::Type>::GetType();
+		return TParamTypeImpl<std::remove_const_t<ValueType>>::GetType();
 	}
 };
 
diff --git a/Engine/Plugins/Experimental/PCG/Source/PCG/Public/Helpers/PCGSettingsHelpers.h b/Engine/Plugins/Experimental/PCG/Source/PCG/Public/Helpers/PCGSettingsHelpers.h
--- a/Engine/Plugins/Experimental/PCG/Source/PCG/Public/Helpers/PCGSettingsHelpers.h
+++ b/Engine/Plugins/Experimental/PCG/Source/PCG/Public/Helpers/PCGSettingsHelpers.h
@@ -204,4 +204,4 @@ PRAGMA_ENABLE_DEPRECATION_WARNINGS
 }
 
 // Deprecated macro, not necessary anymore. Cf. GetValue
-#define PCG_GET_OVERRIDEN_VALUE(Settings, Variable, Params) PCGSettingsHelpers::GetValue(GET_MEMBER_NAME_CHECKED(TRemovePointer<TRemoveConst<decltype(Settings)>::Type>::Type, Variable), (Settings)->Variable, Params)
\ No newline at end of file
+#define PCG_GET_OVERRIDEN_VALUE(Settings, Variable, Params) PCGSettingsHelpers::GetValue(GET_MEMBER_NAME_CHECKED(TRemovePointer<std::remove_const_t<decltype(Settings)>>::Type, Variable), (Settings)->Variable, Params)
\ No newline at end of file
diff --git a/Engine/Plugins/Experimental/PythonScriptPlugin/Source/PythonScriptPlugin/Private/PyUtil.h b/Engine/Plugins/Experimental/PythonScriptPlugin/Source/PythonScriptPlugin/Private/PyUtil.h
--- a/Engine/Plugins/Experimental/PythonScriptPlugin/Source/PythonScriptPlugin/Private/PyUtil.h
+++ b/Engine/Plugins/Experimental/PythonScriptPlugin/Source/PythonScriptPlugin/Private/PyUtil.h
@@ -185,7 +185,7 @@ namespace PyUtil
 		{
 			if (Prop)
 			{
-				((typename TRemoveConst<TPropType>::Type*)Prop)->AddReferencedObjects(Collector);
+				((std::remove_const_t<TPropType>*)Prop)->AddReferencedObjects(Collector);
 			}
 		}
 
diff --git a/Engine/Plugins/Online/OnlineFramework/Source/Party/Public/Party/PartyTypes.h b/Engine/Plugins/Online/OnlineFramework/Source/Party/Public/Party/PartyTypes.h
--- a/Engine/Plugins/Online/OnlineFramework/Source/Party/Public/Party/PartyTypes.h
+++ b/Engine/Plugins/Online/OnlineFramework/Source/Party/Public/Party/PartyTypes.h
@@ -475,7 +475,7 @@ protected:
 #define EXPOSE_REP_DATA_PROPERTY_NO_SETTER(Owner, PropertyType, PropertyName, PropertyAccess)	\
 public:	\
 	/** If the property is a POD or ptr type, we'll work with it by copy. Otherwise, by const ref */	\
-	using Mutable##PropertyName##Type = typename TRemoveConst<PropertyType>::Type;	\
+	using Mutable##PropertyName##Type = std::remove_const_t<PropertyType>;	\
 	using PropertyName##ArgType = typename TChooseClass<TOr<TIsPODType<PropertyType>, TIsPointer<PropertyType>>::Value, PropertyType, const Mutable##PropertyName##Type&>::Result;	\
 	\
 private:	\
@@ -494,7 +494,7 @@ private:	\
 	{	\
 		Compare##PropertyName(OldData.PropertyAccess); \
 	}	\
-	void Compare##PropertyName(const typename TRemoveConst<PropertyType>::Type& OldData) const	\
+	void Compare##PropertyName(const std::remove_const_t<PropertyType>& OldData) const	\
 	{	\
 		if (PropertyAccess != OldData)	\
 		{	\
diff --git a/Engine/Plugins/Runtime/MassEntity/Source/MassEntity/Public/MassEntityTypes.h b/Engine/Plugins/Runtime/MassEntity/Source/MassEntity/Public/MassEntityTypes.h
--- a/Engine/Plugins/Runtime/MassEntity/Source/MassEntity/Public/MassEntityTypes.h
+++ b/Engine/Plugins/Runtime/MassEntity/Source/MassEntity/Public/MassEntityTypes.h
@@ -430,7 +430,7 @@ namespace UE::Mass
 	struct TMultiTypeList : TMultiTypeList<TOthers...>
 	{
 		using Super = TMultiTypeList<TOthers...>;
-		using FType = typename TRemoveConst<typename TRemoveReference<T>::Type>::Type;
+		using FType = std::remove_const_t<typename TRemoveReference<T>::Type>;
 		enum
 		{
 			Ordinal = Super::Ordinal + 1
@@ -448,7 +448,7 @@ namespace UE::Mass
 	template<typename T>
 	struct TMultiTypeList<T>
 	{
-		using FType = typename TRemoveConst<typename TRemoveReference<T>::Type>::Type;
+		using FType = std::remove_const_t<typename TRemoveReference<T>::Type>;
 		enum
 		{
 			Ordinal = 0
@@ -469,7 +469,7 @@ namespace UE::Mass
 	template<typename T, typename... TOthers>
 	struct TMultiArray : TMultiArray<TOthers...>
 	{
-		using FType = typename TRemoveConst<typename TRemoveReference<T>::Type>::Type;
+		using FType = std::remove_const_t<typename TRemoveReference<T>::Type>;
 		using Super = TMultiArray<TOthers...>;
 
 		enum
@@ -509,7 +509,7 @@ namespace UE::Mass
 	template<typename T>
 	struct TMultiArray<T>
 	{
-		using FType = typename TRemoveConst<typename TRemoveReference<T>::Type>::Type;
+		using FType = std::remove_const_t<typename TRemoveReference<T>::Type>;
 		enum { Ordinal = 0 };
 
 		SIZE_T GetAllocatedSize() const
diff --git a/Engine/Source/Runtime/Core/Public/Containers/ArrayView.h b/Engine/Source/Runtime/Core/Public/Containers/ArrayView.h
--- a/Engine/Source/Runtime/Core/Public/Containers/ArrayView.h
+++ b/Engine/Source/Runtime/Core/Public/Containers/ArrayView.h
@@ -633,9 +633,9 @@ public:
 	 *          the subset of elements for which the functor returns true.
 	 */
 	template <typename Predicate>
-	TArray<typename TRemoveConst<ElementType>::Type> FilterByPredicate(Predicate Pred) const
+	TArray<std::remove_const_t<ElementType>> FilterByPredicate(Predicate Pred) const
 	{
-		TArray<typename TRemoveConst<ElementType>::Type> FilterResults;
+		TArray<std::remove_const_t<ElementType>> FilterResults;
 		for (const ElementType* RESTRICT Data = GetData(), *RESTRICT DataEnd = Data + ArrayNum; Data != DataEnd; ++Data)
 		{
 			if (::Invoke(Pred, *Data))
diff --git a/Engine/Source/Runtime/Core/Public/Containers/UnrealString.h b/Engine/Source/Runtime/Core/Public/Containers/UnrealString.h
--- a/Engine/Source/Runtime/Core/Public/Containers/UnrealString.h
+++ b/Engine/Source/Runtime/Core/Public/Containers/UnrealString.h
@@ -377,7 +377,7 @@ public:
 	FORCEINLINE FString& Append(/* no const! */ CharType* Str)
 	{
 		checkSlow(Str);
-		AppendChars(Str, TCString<typename TRemoveConst<CharType>::Type>::Strlen(Str));
+		AppendChars(Str, TCString<std::remove_const_t<CharType>>::Strlen(Str));
 		return *this;
 	}
 
diff --git a/Engine/Source/Runtime/Core/Public/Delegates/DelegateInstancesImpl.h b/Engine/Source/Runtime/Core/Public/Delegates/DelegateInstancesImpl.h
--- a/Engine/Source/Runtime/Core/Public/Delegates/DelegateInstancesImpl.h
+++ b/Engine/Source/Runtime/Core/Public/Delegates/DelegateInstancesImpl.h
@@ -259,7 +259,7 @@ public:
 
 	RetValType Execute(ParamTypes... Params) const final
 	{
-		using MutableUserClass = typename TRemoveConst<UserClass>::Type;
+		using MutableUserClass = std::remove_const_t<UserClass>;
 
 		// Verify that the user object is still valid.  We only have a weak reference to it.
 		TSharedPtr<UserClass, SPMode> SharedUserObject = UserObject.Pin();
@@ -280,7 +280,7 @@ public:
 		// Verify that the user object is still valid.  We only have a weak reference to it.
 		if (TSharedPtr<UserClass, SPMode> SharedUserObject = this->UserObject.Pin())
 		{
-			using MutableUserClass = typename TRemoveConst<UserClass>::Type;
+			using MutableUserClass = std::remove_const_t<UserClass>;
 
 			// Safely remove const to work around a compiler issue with instantiating template permutations for 
 			// overloaded functions that take a function pointer typedef as a member of a templated class.  In
@@ -395,7 +395,7 @@ public:
 
 	RetValType Execute(ParamTypes... Params) const final
 	{
-		using MutableUserClass = typename TRemoveConst<UserClass>::Type;
+		using MutableUserClass = std::remove_const_t<UserClass>;
 
 		// Safely remove const to work around a compiler issue with instantiating template permutations for 
 		// overloaded functions that take a function pointer typedef as a member of a templated class.  In
@@ -410,7 +410,7 @@ public:
 
 	bool ExecuteIfSafe(ParamTypes... Params) const final
 	{
-		using MutableUserClass = typename TRemoveConst<UserClass>::Type;
+		using MutableUserClass = std::remove_const_t<UserClass>;
 
 		// Safely remove const to work around a compiler issue with instantiating template permutations for 
 		// overloaded functions that take a function pointer typedef as a member of a templated class.  In
@@ -519,7 +519,7 @@ public:
 
 	RetValType Execute(ParamTypes... Params) const final
 	{
-		using MutableUserClass = typename TRemoveConst<UserClass>::Type;
+		using MutableUserClass = std::remove_const_t<UserClass>;
 
 		// Verify that the user object is still valid.  We only have a weak reference to it.
 		checkSlow(UserObject.IsValid());
@@ -538,7 +538,7 @@ public:
 	{
 		if (UserClass* ActualUserObject = this->UserObject.Get())
 		{
-			using MutableUserClass = typename TRemoveConst<UserClass>::Type;
+			using MutableUserClass = std::remove_const_t<UserClass>;
 
 			// Safely remove const to work around a compiler issue with instantiating template permutations for 
 			// overloaded functions that take a function pointer typedef as a member of a templated class.  In
@@ -751,7 +751,7 @@ private:
 	// We make this mutable to allow mutable lambdas to be bound and executed.  We don't really want to
 	// model the Functor as being a direct subobject of the delegate (which would maintain transivity of
 	// const - because the binding doesn't affect the substitutability of a copied delegate.
-	mutable typename TRemoveConst<FunctorType>::Type Functor;
+	mutable std::remove_const_t<FunctorType> Functor;
 };
 
 /**
@@ -851,5 +851,5 @@ private:
 	// We make this mutable to allow mutable lambdas to be bound and executed.  We don't really want to
 	// model the Functor as being a direct subobject of the delegate (which would maintain transivity of
 	// const - because the binding doesn't affect the substitutability of a copied delegate.
-	mutable typename TRemoveConst<FunctorType>::Type Functor;
+	mutable std::remove_const_t<FunctorType> Functor;
 };
diff --git a/Engine/Source/Runtime/Core/Public/Misc/ExpressionParserTypes.inl b/Engine/Source/Runtime/Core/Public/Misc/ExpressionParserTypes.inl
--- a/Engine/Source/Runtime/Core/Public/Misc/ExpressionParserTypes.inl
+++ b/Engine/Source/Runtime/Core/Public/Misc/ExpressionParserTypes.inl
@@ -11,7 +11,6 @@
 class FExpressionNode;
 class FExpressionToken;
 struct FOperatorFunctionID;
-template< class T > struct TRemoveConst;
 
 namespace Impl
 {
@@ -246,7 +245,7 @@ template<typename ContextType>
 template<typename OperatorType, typename FuncType>
 void TOperatorJumpTable<ContextType>::MapPreUnary(FuncType InFunc)
 {
-	typedef typename TRemoveConst<typename TRemoveReference<typename Impl::TCallableInfo<FuncType>::Arg1>::Type>::Type OperandType;
+	typedef std::remove_const_t<typename TRemoveReference<typename Impl::TCallableInfo<FuncType>::Arg1>::Type> OperandType;
 
 	FOperatorFunctionID ID = {
 		TGetExpressionNodeTypeId<OperatorType>::GetTypeId(),
@@ -261,7 +260,7 @@ template<typename ContextType>
 template<typename OperatorType, typename FuncType>
 void TOperatorJumpTable<ContextType>::MapPostUnary(FuncType InFunc)
 {
-	typedef typename TRemoveConst<typename TRemoveReference<typename Impl::TCallableInfo<FuncType>::Arg1>::Type>::Type OperandType;
+	typedef std::remove_const_t<typename TRemoveReference<typename Impl::TCallableInfo<FuncType>::Arg1>::Type> OperandType;
 
 	FOperatorFunctionID ID = {
 		TGetExpressionNodeTypeId<OperatorType>::GetTypeId(),
@@ -276,8 +275,8 @@ template<typename ContextType>
 template<typename OperatorType, typename FuncType>
 void TOperatorJumpTable<ContextType>::MapBinary(FuncType InFunc)
 {
-	typedef typename TRemoveConst<typename TRemoveReference<typename Impl::TCallableInfo<FuncType>::Arg1>::Type>::Type LeftOperandType;
-	typedef typename TRemoveConst<typename TRemoveReference<typename Impl::TCallableInfo<FuncType>::Arg2>::Type>::Type RightOperandType;
+	typedef std::remove_const_t<typename TRemoveReference<typename Impl::TCallableInfo<FuncType>::Arg1>::Type> LeftOperandType;
+	typedef std::remove_const_t<typename TRemoveReference<typename Impl::TCallableInfo<FuncType>::Arg2>::Type> RightOperandType;
 
 	FOperatorFunctionID ID = {
 		TGetExpressionNodeTypeId<OperatorType>::GetTypeId(),
@@ -292,7 +291,7 @@ template<typename ContextType>
 template<typename OperatorType, typename FuncType>
 void TOperatorJumpTable<ContextType>::MapShortCircuit(FuncType InFunc)
 {
-	typedef typename TRemoveConst<typename TRemoveReference<typename Impl::TCallableInfo<FuncType>::Arg1>::Type>::Type OperandType;
+	typedef std::remove_const_t<typename TRemoveReference<typename Impl::TCallableInfo<FuncType>::Arg1>::Type> OperandType;
 
 	FOperatorFunctionID ID = {
 		TGetExpressionNodeTypeId<OperatorType>::GetTypeId(),
diff --git a/Engine/Source/Runtime/Core/Public/Templates/UnrealTypeTraits.h b/Engine/Source/Runtime/Core/Public/Templates/UnrealTypeTraits.h
--- a/Engine/Source/Runtime/Core/Public/Templates/UnrealTypeTraits.h
+++ b/Engine/Source/Runtime/Core/Public/Templates/UnrealTypeTraits.h
@@ -299,15 +299,10 @@ struct TCallTraitsParamTypeHelper<T*, true>
  *		MutablePtr = const_cast< RemoveConst< ConstPtrType >::Type >( ConstPtr );
  */
 template< class T >
-struct TRemoveConst
+struct /*UE_DEPRECATED(5.3, "TRemoveConst has been deprecated, please use std::remove_const_t instead.")*/ TRemoveConst
 {
-	typedef T Type;
+	using Type = std::remove_const_t<T>;
 };
-template< class T >
-struct TRemoveConst<const T>
-{
-	typedef T Type;
-};	
 
 
 /*-----------------------------------------------------------------------------
diff --git a/Engine/Source/Runtime/MeshDescription/Public/MeshAttributeArray.h b/Engine/Source/Runtime/MeshDescription/Public/MeshAttributeArray.h
--- a/Engine/Source/Runtime/MeshDescription/Public/MeshAttributeArray.h
+++ b/Engine/Source/Runtime/MeshDescription/Public/MeshAttributeArray.h
@@ -2294,7 +2294,7 @@ void TAttributesSet<ElementIDType>::ForEachByType(ForEachFunc Func)
 {
 	for (auto& MapEntry : this->Map)
 	{
-		ForEachByTypeImpl::DispatchFunctor<ElementIDType, typename TRemoveConst<AttributeType>::Type, ForEachFunc>()(MapEntry.Key, Func, MapEntry.Value.Get());
+		ForEachByTypeImpl::DispatchFunctor<ElementIDType, std::remove_const_t<AttributeType>, ForEachFunc>()(MapEntry.Key, Func, MapEntry.Value.Get());
 	}
 }
 
@@ -2304,7 +2304,7 @@ void TAttributesSet<ElementIDType>::ForEachByType(ForEachFunc Func) const
 {
 	for (const auto& MapEntry : this->Map)
 	{
-		ForEachByTypeImpl::ConstDispatchFunctor<ElementIDType, typename TRemoveConst<AttributeType>::Type, ForEachFunc>()(MapEntry.Key, Func, MapEntry.Value.Get());
+		ForEachByTypeImpl::ConstDispatchFunctor<ElementIDType, std::remove_const_t<AttributeType>, ForEachFunc>()(MapEntry.Key, Func, MapEntry.Value.Get());
 	}
 }
 
