Deprecated TRemoveCV.

diff --git a/Engine/Source/Editor/UnrealEd/Private/UnrealEdPrivatePCH.h b/Engine/Source/Editor/UnrealEd/Private/UnrealEdPrivatePCH.h
--- a/Engine/Source/Editor/UnrealEd/Private/UnrealEdPrivatePCH.h
+++ b/Engine/Source/Editor/UnrealEd/Private/UnrealEdPrivatePCH.h
@@ -26,7 +26,6 @@
 #include "Templates/UnrealTypeTraits.h"
 #include "Templates/AndOrNot.h"
 #include "Templates/IsArithmetic.h"
-#include "Templates/RemoveCV.h"
 #include "Templates/IsPODType.h"
 #include "Templates/IsTriviallyCopyConstructible.h"
 #include "Templates/UnrealTemplate.h"
diff --git a/Engine/Source/Editor/UnrealEd/Public/UnrealEdSharedPCH.h b/Engine/Source/Editor/UnrealEd/Public/UnrealEdSharedPCH.h
--- a/Engine/Source/Editor/UnrealEd/Public/UnrealEdSharedPCH.h
+++ b/Engine/Source/Editor/UnrealEd/Public/UnrealEdSharedPCH.h
@@ -26,7 +26,6 @@
 #include "Templates/UnrealTypeTraits.h"
 #include "Templates/AndOrNot.h"
 #include "Templates/IsArithmetic.h"
-#include "Templates/RemoveCV.h"
 #include "Templates/IsPODType.h"
 #include "Templates/IsTriviallyCopyConstructible.h"
 #include "Templates/UnrealTemplate.h"
diff --git a/Engine/Source/Runtime/Core/Public/Containers/StringConv.h b/Engine/Source/Runtime/Core/Public/Containers/StringConv.h
--- a/Engine/Source/Runtime/Core/Public/Containers/StringConv.h
+++ b/Engine/Source/Runtime/Core/Public/Containers/StringConv.h
@@ -13,7 +13,6 @@
 #include "Templates/AndOrNot.h"
 #include "Templates/EnableIf.h"
 #include "Templates/IsArray.h"
-#include "Templates/RemoveCV.h"
 #include "Templates/RemoveReference.h"
 #include "Templates/UnrealTemplate.h"
 #include "Templates/UnrealTypeTraits.h"
diff --git a/Engine/Source/Runtime/Core/Public/Core.h b/Engine/Source/Runtime/Core/Public/Core.h
--- a/Engine/Source/Runtime/Core/Public/Core.h
+++ b/Engine/Source/Runtime/Core/Public/Core.h
@@ -19,7 +19,6 @@ MONOLITHIC_HEADER_BOILERPLATE()
 #include "Misc/Char.h"
 #include "Templates/EnableIf.h"
 #include "Templates/AndOrNot.h"
-#include "Templates/RemoveCV.h"
 #include "Templates/IsArithmetic.h"
 #include "Templates/IsFloatingPoint.h"
 #include "Templates/IsIntegral.h"
diff --git a/Engine/Source/Runtime/Core/Public/CoreMinimal.h b/Engine/Source/Runtime/Core/Public/CoreMinimal.h
--- a/Engine/Source/Runtime/Core/Public/CoreMinimal.h
+++ b/Engine/Source/Runtime/Core/Public/CoreMinimal.h
@@ -34,7 +34,6 @@
 #include "HAL/UnrealMemory.h"
 #include "Templates/IsArithmetic.h"
 #include "Templates/AndOrNot.h"
-#include "Templates/RemoveCV.h"
 #include "Templates/IsPODType.h"
 #include "Templates/IsUECoreType.h"
 #include "Templates/IsTriviallyCopyConstructible.h"
diff --git a/Engine/Source/Runtime/Core/Public/CoreSharedPCH.h b/Engine/Source/Runtime/Core/Public/CoreSharedPCH.h
--- a/Engine/Source/Runtime/Core/Public/CoreSharedPCH.h
+++ b/Engine/Source/Runtime/Core/Public/CoreSharedPCH.h
@@ -194,7 +194,6 @@
 #include "Templates/MemoryOps.h"
 #include "Templates/PointerIsConvertibleFromTo.h"
 #include "Templates/RefCounting.h"
-#include "Templates/RemoveCV.h"
 #include "Templates/RemoveExtent.h"
 #include "Templates/RemoveReference.h"
 #include "Templates/SharedPointer.h"
diff --git a/Engine/Source/Runtime/Core/Public/Templates/RemoveCV.h b/Engine/Source/Runtime/Core/Public/Templates/RemoveCV.h
--- a/Engine/Source/Runtime/Core/Public/Templates/RemoveCV.h
+++ b/Engine/Source/Runtime/Core/Public/Templates/RemoveCV.h
@@ -11,7 +11,7 @@
  * note: won't remove the const from "const int*", as the pointer is not const
  */
 template <typename T>
-struct /*UE_DEPRECATED(5.3, "TRemoveCV has been deprecated, please use std::remove_cv_t instead.")*/ TRemoveCV
+struct UE_DEPRECATED(5.3, "TRemoveCV has been deprecated, please use std::remove_cv_t instead.") TRemoveCV
 {
 	using Type = std::remove_cv_t<T>;
 };
diff --git a/Engine/Source/Runtime/Core/Public/Templates/UnrealTypeTraits.h b/Engine/Source/Runtime/Core/Public/Templates/UnrealTypeTraits.h
--- a/Engine/Source/Runtime/Core/Public/Templates/UnrealTypeTraits.h
+++ b/Engine/Source/Runtime/Core/Public/Templates/UnrealTypeTraits.h
@@ -15,7 +15,6 @@
 #include "Templates/EnableIf.h"
 #include "Templates/IsArithmetic.h"
 #include "Templates/IsEnum.h"
-#include "Templates/RemoveCV.h"
 #include "Templates/Models.h"
 
 #include "Templates/IsPODType.h"
diff --git a/Engine/Source/Runtime/Engine/Private/EnginePrivatePCH.h b/Engine/Source/Runtime/Engine/Private/EnginePrivatePCH.h
--- a/Engine/Source/Runtime/Engine/Private/EnginePrivatePCH.h
+++ b/Engine/Source/Runtime/Engine/Private/EnginePrivatePCH.h
@@ -26,7 +26,6 @@
 #include "Templates/UnrealTypeTraits.h"
 #include "Templates/AndOrNot.h"
 #include "Templates/IsArithmetic.h"
-#include "Templates/RemoveCV.h"
 #include "Templates/IsPODType.h"
 #include "Templates/IsTriviallyCopyConstructible.h"
 #include "Templates/UnrealTemplate.h"
diff --git a/Engine/Source/Runtime/Engine/Public/EngineSharedPCH.h b/Engine/Source/Runtime/Engine/Public/EngineSharedPCH.h
--- a/Engine/Source/Runtime/Engine/Public/EngineSharedPCH.h
+++ b/Engine/Source/Runtime/Engine/Public/EngineSharedPCH.h
@@ -26,7 +26,6 @@
 #include "Templates/UnrealTypeTraits.h"
 #include "Templates/AndOrNot.h"
 #include "Templates/IsArithmetic.h"
-#include "Templates/RemoveCV.h"
 #include "Templates/IsPODType.h"
 #include "Templates/IsTriviallyCopyConstructible.h"
 #include "Templates/UnrealTemplate.h"
diff --git a/Engine/Source/Runtime/MeshDescription/Public/MeshAttributeArray.h b/Engine/Source/Runtime/MeshDescription/Public/MeshAttributeArray.h
--- a/Engine/Source/Runtime/MeshDescription/Public/MeshAttributeArray.h
+++ b/Engine/Source/Runtime/MeshDescription/Public/MeshAttributeArray.h
@@ -27,7 +27,6 @@
 #include "Templates/CopyQualifiersFromTo.h"
 #include "Templates/EnableIf.h"
 #include "Templates/IsArray.h"
-#include "Templates/RemoveCV.h"
 #include "Templates/Tuple.h"
 #include "Templates/UniquePtr.h"
 #include "Templates/UnrealTemplate.h"
diff --git a/Engine/Source/Runtime/Slate/Public/SlateSharedPCH.h b/Engine/Source/Runtime/Slate/Public/SlateSharedPCH.h
--- a/Engine/Source/Runtime/Slate/Public/SlateSharedPCH.h
+++ b/Engine/Source/Runtime/Slate/Public/SlateSharedPCH.h
@@ -26,7 +26,6 @@
 #include "Templates/UnrealTypeTraits.h"
 #include "Templates/AndOrNot.h"
 #include "Templates/IsArithmetic.h"
-#include "Templates/RemoveCV.h"
 #include "Templates/IsPODType.h"
 #include "Templates/IsTriviallyCopyConstructible.h"
 #include "Templates/UnrealTemplate.h"
