[Runtime]
Compile time optimizations by changing template arguments to use std::enable_if_t

diff --git a/Engine/Source/Runtime/Core/Public/Serialization/Archive.h b/Engine/Source/Runtime/Core/Public/Serialization/Archive.h
--- a/Engine/Source/Runtime/Core/Public/Serialization/Archive.h
+++ b/Engine/Source/Runtime/Core/Public/Serialization/Archive.h
@@ -2252,7 +2252,7 @@ FORCEINLINE FArchive& operator<<(FArchive& Ar, TEnumAsByte<TEnum>& Value)
 */
 template <
 	typename EnumType,
-	typename = typename TEnableIf<TIsEnumClass<EnumType>::Value>::Type
+	std::enable_if_t<TIsEnumClass<EnumType>::Value, int> = 0
 >
 FORCEINLINE FArchive& operator<<(FArchive& Ar, EnumType& Value)
 {
diff --git a/Engine/Source/Runtime/Core/Public/Serialization/StructuredArchiveAdapters.h b/Engine/Source/Runtime/Core/Public/Serialization/StructuredArchiveAdapters.h
--- a/Engine/Source/Runtime/Core/Public/Serialization/StructuredArchiveAdapters.h
+++ b/Engine/Source/Runtime/Core/Public/Serialization/StructuredArchiveAdapters.h
@@ -120,11 +120,9 @@ private:
  *
  * @return  A reference to the same archive as Ar.
  */
-template <typename T>
-typename TEnableIf<
-	!TModels_V<CInsertable<FArchive&>, T> && TModels_V<CInsertable<FStructuredArchiveSlot>, T>,
-	FArchive&
->::Type operator<<(FArchive& Ar, T& Obj)
+template <typename T,
+	std::enable_if_t<!TModels_V<CInsertable<FArchive&>, T> && TModels_V<CInsertable<FStructuredArchiveSlot>, T>, int> = 0>
+FArchive& operator<<(FArchive& Ar, T& Obj)
 {
 	FStructuredArchiveFromArchive ArAdapt(Ar);
 	ArAdapt.GetSlot() << Obj;
@@ -137,11 +135,9 @@ typename TEnableIf<
  * @param  Slot  The slot to read from or write to.
  * @param  Obj   The object to read or write.
  */
-template <typename T>
-typename TEnableIf<
-	TModels_V<CInsertable<FArchive&>, T> &&
-	!TModels_V<CInsertable<FStructuredArchiveSlot>, T>
->::Type operator<<(FStructuredArchiveSlot Slot, T& Obj)
+template <typename T,
+	std::enable_if_t<TModels_V<CInsertable<FArchive&>, T> && !TModels_V<CInsertable<FStructuredArchiveSlot>, T>, int> = 0>
+void operator<<(FStructuredArchiveSlot Slot, T& Obj)
 {
 #if WITH_TEXT_ARCHIVE_SUPPORT
 	FArchiveFromStructuredArchive Adapter(Slot);
diff --git a/Engine/Source/Runtime/Core/Public/Serialization/StructuredArchiveSlots.h b/Engine/Source/Runtime/Core/Public/Serialization/StructuredArchiveSlots.h
--- a/Engine/Source/Runtime/Core/Public/Serialization/StructuredArchiveSlots.h
+++ b/Engine/Source/Runtime/Core/Public/Serialization/StructuredArchiveSlots.h
@@ -85,7 +85,7 @@ public:
 
 	template <
 		typename EnumType,
-		typename = typename TEnableIf<TIsEnumClass<EnumType>::Value>::Type
+		std::enable_if_t<TIsEnumClass<EnumType>::Value, int> = 0
 	>
 	FORCEINLINE void operator<<(EnumType& Value)
 	{
