Updated FPaths::Combine to work with FStringView

diff --git a/Engine/Source/Runtime/Core/Private/Misc/Paths.cpp b/Engine/Source/Runtime/Core/Private/Misc/Paths.cpp
--- a/Engine/Source/Runtime/Core/Private/Misc/Paths.cpp
+++ b/Engine/Source/Runtime/Core/Private/Misc/Paths.cpp
@@ -1597,24 +1597,26 @@ const FString& FPaths::GetRelativePathToRoot()
 	return StaticData.RelativePathToRoot;
 }
 
-void FPaths::CombineInternal(FString& OutPath, const FStringView* Paths, int32 NumPaths)
+FString FPaths::CombineInternal(const FStringView* Paths, int32 NumPaths)
 {
-	check(Paths != NULL && NumPaths > 0);
-
-	int32 OutStringSize = 0;
+	check(Paths && NumPaths > 0);
 
+	int32 CombinedPathLen = 0;
 	for (int32 i=0; i < NumPaths; ++i)
 	{
-		OutStringSize += Paths[i].Len() + 1;
+		CombinedPathLen += Paths[i].Len() + 1;
 	}
 
-	OutPath.Empty(OutStringSize);
-	OutPath += Paths[0];
+	FString CombinedPath;
+	CombinedPath.Reserve(CombinedPathLen);
+	CombinedPath += Paths[0];
 	
 	for (int32 i=1; i < NumPaths; ++i)
 	{
-		OutPath /= Paths[i];
+		CombinedPath /= Paths[i];
 	}
+
+	return CombinedPath;
 }
 
 bool FPaths::IsSamePath(const FString& PathA, const FString& PathB)
diff --git a/Engine/Source/Runtime/Core/Public/Misc/Paths.h b/Engine/Source/Runtime/Core/Public/Misc/Paths.h
--- a/Engine/Source/Runtime/Core/Public/Misc/Paths.h
+++ b/Engine/Source/Runtime/Core/Public/Misc/Paths.h
@@ -19,11 +19,11 @@ namespace UE::Core::Private
 	const FString& GetBindingType(const FString& Str);
 	const FStringView& GetBindingType(const FStringView& StringView);
 
-	// This is used to force the arguments of FString::Combine to implicitly convert (if necessary)
+	// This is used to force the arguments of FPaths::Combine to implicitly convert (if necessary)
 	// to FString when calling CombineImpl(), allowing them to remain as temporaries on the stack
 	// so that they stay allocated during the combining process.
 	//
-	// Pointer arguments are passed without causing FString temporaries to be created,
+	// Pointer/view arguments are passed without causing FString temporaries to be created,
 	// and FString arguments are referenced directly without creating extra copies.
 	template <typename T>
 	using TToStringType_T = decltype(GetBindingType(std::declval<T>()));
@@ -648,26 +648,6 @@ public:
 		return CombineImpl<PathTypes...>(Forward<PathTypes>(InPaths)...);
 	}
 
-#if 0
-	FORCEINLINE static FString Combine(const FString& Path1, const FStringView& Path2)
-	{
-		const FStringView Paths[] = { InternalGetStringView(Path1), InternalGetStringView(Path2) };
-		FString Out;
-
-		CombineInternal(Out, Paths, UE_ARRAY_COUNT(Paths));
-		return Out;
-	}
-
-	FORCEINLINE static FString Combine(const FString& Path1, const TCHAR* PathMid, const FStringView& Path2)
-	{
-		const FStringView Paths[] = { InternalGetStringView(Path1), InternalGetStringView(PathMid), InternalGetStringView(Path2) };
-		FString Out;
-
-		CombineInternal(Out, Paths, UE_ARRAY_COUNT(Paths));
-		return Out;
-	}
-#endif
-
 	/**
 	 * Frees any memory retained by FPaths.
 	 */
@@ -675,7 +655,7 @@ public:
 
 protected:
 
-	static void CombineInternal(FString& OutPath, const FStringView* Paths, int32 NumPaths);
+	static FString CombineInternal(const FStringView* Paths, int32 NumPaths);
 
 private:
 	struct FStaticData;
@@ -684,10 +664,7 @@ private:
 	FORCEINLINE static FString CombineImpl(UE::Core::Private::TToStringType_T<std::decay_t<PathTypes>>... InPaths)
 	{
 		const FStringView Paths[] = { FStringView(InPaths)... };
-		FString Out;
-		
-		CombineInternal(Out, Paths, UE_ARRAY_COUNT(Paths));
-		return Out;
+		return CombineInternal(Paths, UE_ARRAY_COUNT(Paths));
 	}
 
 	/** Returns, if any, the value of the -userdir command line argument. This can be used to sandbox artifacts to a desired location */
