Fixed text created from FText::FindText with a null source string failing to resolve correctly if the language changes

This was a regression from the display string optimization work, though use of FText::FindText is generally discouraged

diff --git a/Engine/Source/Runtime/Core/Private/Internationalization/TextHistory.cpp b/Engine/Source/Runtime/Core/Private/Internationalization/TextHistory.cpp
--- a/Engine/Source/Runtime/Core/Private/Internationalization/TextHistory.cpp
+++ b/Engine/Source/Runtime/Core/Private/Internationalization/TextHistory.cpp
@@ -904,7 +904,7 @@ bool FTextHistory_Base::CanUpdateDisplayString()
 void FTextHistory_Base::UpdateDisplayString()
 {
 	check(!TextId.IsEmpty()); // CanUpdateDisplayString should prevent UpdateDisplayString being called
-	LocalizedString = FTextLocalizationManager::Get().GetDisplayString(TextId.GetNamespace(), TextId.GetKey(), &SourceString);
+	LocalizedString = FTextLocalizationManager::Get().GetDisplayString(TextId.GetNamespace(), TextId.GetKey(), SourceString.IsEmpty() ? nullptr : &SourceString);
 }
 
 bool FTextHistory_Base::StaticShouldReadFromBuffer(const TCHAR* Buffer)
