Migrate ClosableMpscQueueTest, ConcurrentQueuesTest, DepletableMpscQueueTest, RingBuffer tests to Catch2

diff --git a/Engine/Source/Runtime/Core/Private/Tests/Containers/ConcurrentQueuesTest.cpp b/Engine/Source/Runtime/Core/Private/Tests/Containers/ConcurrentQueuesTest.cpp
deleted file mode 100644
diff --git a/Engine/Source/Runtime/Core/Private/Tests/Containers/RingBufferTest.cpp b/Engine/Source/Runtime/Core/Private/Tests/Containers/RingBufferTest.cpp
deleted file mode 100644
diff --git a/Engine/Source/Runtime/Core/Tests/Containers/ConcurrentQueuesTest.cpp b/Engine/Source/Runtime/Core/Tests/Containers/ConcurrentQueuesTest.cpp
new file mode 100644
--- /dev/null
+++ b/Engine/Source/Runtime/Core/Tests/Containers/ConcurrentQueuesTest.cpp
@@ -0,0 +1,1142 @@
+// Copyright Epic Games, Inc. All Rights Reserved.
+
+#include "CoreTypes.h"
+#include "Async/TaskGraphInterfaces.h"
+#include "Containers/Array.h"
+#include "Containers/CircularQueue.h"
+#include "Containers/Queue.h"
+#include "Containers/SpscQueue.h"
+#include "Containers/MpscQueue.h"
+#include "Containers/ClosableMpscQueue.h"
+#include "Containers/DepletableMpmcQueue.h"
+#include "Tests/Benchmark.h"
+
+#include <atomic>
+
+#if WITH_TESTS
+
+#include "Tests/TestHarnessAdapter.h"
+
+
+namespace UE { namespace ConcurrentQueuesTests
+{
+	template<typename QueueType>
+	class TQueueAdapter : public QueueType
+	{
+	public:
+		template<typename... ArgTypes>
+		TQueueAdapter(ArgTypes... Args)
+			: QueueType(Forward<ArgTypes>(Args)...)
+		{
+		}
+
+		TOptional<typename QueueType::FElementType> Dequeue()
+		{
+			typename QueueType::FElementType Value;
+			if (QueueType::Dequeue(Value))
+			{
+				return Value;
+			}
+			else
+			{
+				return {};
+			}
+		}
+	};
+
+	const uint32 CircularQueueSize = 1024;
+
+	// measures performance of a queue when the producer and the consumer are run in the same thread
+	template<uint32 Num, typename QueueType>
+	void TestSpscQueueSingleThreadImp(QueueType& Queue)
+	{
+		const uint32 BatchSize = CircularQueueSize;
+		const uint32 BatchNum = Num / BatchSize;
+		for (uint32 BatchIndex = 0; BatchIndex != BatchNum; ++BatchIndex)
+		{
+			for (uint32 i = 0; i != BatchSize; ++i)
+			{
+				Queue.Enqueue(i);
+			}
+
+			for (uint32 i = 0; i != BatchSize; ++i)
+			{
+				TOptional<uint32> Consumed{ Queue.Dequeue() };
+				checkSlow(Consumed.IsSet() && Consumed.GetValue() == i);
+			}
+		}
+	}
+
+	template<uint32 Num>
+	void TestTCircularQueueSingleThread()
+	{
+		TQueueAdapter<TCircularQueue<uint32>> Queue{ CircularQueueSize + 1 };
+		TestSpscQueueSingleThreadImp<Num>(Queue);
+	}
+
+	template<uint32 Num, typename QueueType>
+	void TestQueueSingleThread()
+	{
+		QueueType Queue;
+		TestSpscQueueSingleThreadImp<Num>(Queue);
+	}
+
+	template<uint32 Num, typename QueueType>
+	void TestSpscQueue_Impl(QueueType& Queue)
+	{
+		std::atomic<bool> bStop{ false };
+
+		FGraphEventRef Producer = FFunctionGraphTask::CreateAndDispatchWhenReady(
+			[&bStop, &Queue]
+			{
+				while (!bStop)
+				{
+					Queue.Enqueue(0);
+				}
+			}
+			);
+
+		// consumer
+		uint32 It = 0;
+		while (It != Num)
+		{
+			if (Queue.Dequeue().IsSet())
+			{
+				++It;
+			}
+		}
+
+		bStop = true;
+
+		Producer->Wait(ENamedThreads::GameThread);
+	}
+
+	template<uint32 Num>
+	void TestTCircularQueue()
+	{
+		TQueueAdapter<TCircularQueue<uint32>> Queue{ CircularQueueSize + 1 };
+		TestSpscQueue_Impl<Num>(Queue);
+	}
+
+	template<uint32 Num, typename QueueType>
+	void TestSpscQueue()
+	{
+		QueueType Queue;
+		TestSpscQueue_Impl<Num>(Queue);
+	}
+
+	template<uint32 Num, typename QueueType>
+	void TestSpscQueueCorrectness()
+	{
+		QueueType Queue;
+		int32 NumProduced = 0;
+
+		// producer
+		FGraphEventRef Task = FFunctionGraphTask::CreateAndDispatchWhenReady(
+			[&Queue, &NumProduced]
+			{
+				do
+				{
+					Queue.Enqueue(0);
+				} while (++NumProduced != Num);
+			}
+		);
+
+		// consumer
+		uint32 NumConcumed = 0;
+		while (NumConcumed != Num)
+		{
+			if (Queue.Dequeue().IsSet())
+			{
+				++NumConcumed;
+			}
+		}
+
+		while (Queue.Dequeue().IsSet())
+		{
+			++NumConcumed;
+		}
+
+		Task->Wait();
+
+		check(NumProduced == NumConcumed);
+	}
+
+	template<uint32 Num, typename QueueType>
+	void TestMpscQueue()
+	{
+		QueueType Queue;
+		std::atomic<bool> bStop{ false };
+
+		int32 NumProducers = FTaskGraphInterface::Get().GetNumWorkerThreads();
+		FGraphEventArray Producers;
+		for (int32 i = 0; i != NumProducers; ++i)
+		{
+			Producers.Add(FFunctionGraphTask::CreateAndDispatchWhenReady(
+				[&bStop, &Queue]
+				{
+					while (!bStop)
+					{
+						Queue.Enqueue(0);
+					}
+				}
+				));
+		}
+
+		uint32 It = 0;
+		while (It != Num)
+		{
+			if (Queue.Dequeue().IsSet())
+			{
+				++It;
+			}
+		}
+
+		bStop = true;
+
+		FTaskGraphInterface::Get().WaitUntilTasksComplete(MoveTemp(Producers), ENamedThreads::GameThread);
+	}
+
+	template<uint32 Num, typename QueueType>
+	void TestMpscQueueCorrectness()
+	{
+		struct alignas(PLATFORM_CACHE_LINE_SIZE) FCounter
+		{
+			uint32 Count = 0;
+		};
+
+		QueueType Queue;
+
+		const int32 NumProducers = FTaskGraphInterface::Get().GetNumWorkerThreads();
+		const uint32 NumPerProducer = Num / NumProducers;
+		TArray<FCounter, TAlignedHeapAllocator<PLATFORM_CACHE_LINE_SIZE>> NumsProduced;
+		NumsProduced.AddDefaulted(NumProducers);
+		for (int32 i = 0; i != NumProducers; ++i)
+		{
+			FFunctionGraphTask::CreateAndDispatchWhenReady(
+				[&Queue, &NumsProduced, i, NumPerProducer]()
+			{
+				do
+				{
+					Queue.Enqueue(0);
+				} while (++NumsProduced[i].Count < NumPerProducer);
+			}
+			);
+		}
+
+		// consumer
+		uint32 NumConcumed = 0;
+		while (NumConcumed != NumPerProducer * NumProducers)
+		{
+			if (Queue.Dequeue().IsSet())
+			{
+				++NumConcumed;
+			}
+		}
+
+		while (Queue.Dequeue().IsSet())
+		{
+			++NumConcumed;
+		}
+
+		uint32 Produced = 0;
+		for (int i = 0; i != NumProducers; ++i)
+		{
+			Produced += NumsProduced[i].Count;
+		}
+		check(Produced == NumConcumed);
+	}
+
+	TEST_CASE_NAMED(FConcurrentQueuesTest, "System::Core::Async::ConcurrentQueuesTest", "[ApplicationContextMask][EngineFilter]")
+	{
+		{	// test for support of not default constructible types
+			struct FNonDefaultConstructable
+			{
+				int Value;
+
+				explicit FNonDefaultConstructable(int InValue)
+					: Value(InValue)
+				{
+					UE_LOG(LogTemp, Display, TEXT("ctor"));
+				}
+
+				~FNonDefaultConstructable()
+				{
+					UE_LOG(LogTemp, Display, TEXT("dctor"));
+				}
+
+				FNonDefaultConstructable(const FNonDefaultConstructable&) = delete;
+				FNonDefaultConstructable& operator=(const FNonDefaultConstructable&) = delete;
+
+				FNonDefaultConstructable(FNonDefaultConstructable&& Other)
+					: Value(Other.Value)
+				{
+					UE_LOG(LogTemp, Display, TEXT("move-ctor"));
+				}
+
+				FNonDefaultConstructable& operator=(FNonDefaultConstructable&& Other)
+				{
+					Value = Other.Value;
+					UE_LOG(LogTemp, Display, TEXT("move="));
+					return *this;
+				}
+			};
+
+			{
+				TSpscQueue<FNonDefaultConstructable> Q;
+				Q.Enqueue(1);
+				TOptional<FNonDefaultConstructable> Res{ Q.Dequeue() };
+				verify(Res.IsSet() && Res.GetValue().Value == 1);
+			}
+
+			{
+				TMpscQueue<FNonDefaultConstructable> Q;
+				Q.Enqueue(1);
+				TOptional<FNonDefaultConstructable> Res{ Q.Dequeue() };
+				verify(Res.IsSet() && Res.GetValue().Value == 1);
+			}
+		}
+
+		{	// test queue destruction and not default movable types
+			struct FNonTrivial
+			{
+				int* Value;
+
+				explicit FNonTrivial(TUniquePtr<int> InValue)
+					: Value(InValue.Release())
+				{}
+
+				~FNonTrivial()
+				{
+					// check for double delete
+					check(Value != (int*)0xcdcdcdcdcdcdcdcd);
+					delete Value;
+					Value = (int*)0xcdcdcdcdcdcdcdcd;
+				}
+
+				FNonTrivial(const FNonTrivial&) = delete;
+				FNonTrivial& operator=(const FNonTrivial&) = delete;
+
+				FNonTrivial(FNonTrivial&& Other)
+				{
+					Value = Other.Value;
+					Other.Value = nullptr;
+				}
+
+				FNonTrivial& operator=(FNonTrivial&& Other)
+				{
+					Swap(Value, Other.Value);
+					return *this;
+				}
+			};
+
+			// SPSC
+
+			{	// destroy queue while it's holding one unconsumed item
+				TSpscQueue<FNonTrivial> Q;
+				Q.Enqueue(MakeUnique<int>(1));
+			}
+
+			{	// destroy queue while it's holding one cached consumed time
+				TSpscQueue<FNonTrivial> Q;
+				Q.Enqueue(MakeUnique<int>(1));
+				TOptional<FNonTrivial> Res{ Q.Dequeue() };
+				verify(Res.IsSet() && *Res.GetValue().Value == 1);
+			}
+
+			{	// destroy queue while it's holding one chached consumed item and one unconsumed item
+				TSpscQueue<FNonTrivial> Q;
+				Q.Enqueue(MakeUnique<int>(1));
+				Q.Enqueue(MakeUnique<int>(2));
+				TOptional<FNonTrivial> Res{ Q.Dequeue() };
+				verify(Res.IsSet() && *Res.GetValue().Value == 1);
+			}
+
+			// MPSC
+			{	// destroy untouched queue
+				TMpscQueue<FNonTrivial> Q;
+			}
+
+			{	// destroy never consumed queue with one unconsumed item
+				TMpscQueue<FNonTrivial> Q;
+				Q.Enqueue(MakeUnique<int>(1));
+			}
+
+			{	// destroy empty queue
+				TMpscQueue<FNonTrivial> Q;
+				Q.Enqueue(MakeUnique<int>(1));
+				TOptional<FNonTrivial> Res{ Q.Dequeue() };
+				verify(Res.IsSet() && *Res.GetValue().Value == 1);
+			}
+
+			{	// destroy queue with one unconsumed item
+				TMpscQueue<FNonTrivial> Q;
+				Q.Enqueue(MakeUnique<int>(1));
+				Q.Enqueue(MakeUnique<int>(2));
+				TOptional<FNonTrivial> Res{ Q.Dequeue() };
+				verify(Res.IsSet() && *Res.GetValue().Value == 1);
+			}
+
+			{	// destroy queue with two items
+				TMpscQueue<FNonTrivial> Q;
+				Q.Enqueue(MakeUnique<int>(1));
+				Q.Enqueue(MakeUnique<int>(2));
+			}
+
+			{	// enqueue and dequeue multiple items
+				TMpscQueue<FNonTrivial> Q;
+				Q.Enqueue(MakeUnique<int>(1));
+				Q.Enqueue(MakeUnique<int>(2));
+				TOptional<FNonTrivial> Res{ Q.Dequeue() };
+				verify(Res.IsSet() && *Res.GetValue().Value == 1);
+				Res = Q.Dequeue();
+				verify(Res.IsSet() && *Res.GetValue().Value == 2);
+			}
+
+			{	// enqueue and dequeue (interleaved) multiple items
+				TMpscQueue<FNonTrivial> Q;
+				Q.Enqueue(MakeUnique<int>(1));
+				TOptional<FNonTrivial> Res{ Q.Dequeue() };
+				verify(Res.IsSet() && *Res.GetValue().Value == 1);
+				Q.Enqueue(MakeUnique<int>(2));
+				Res = Q.Dequeue();
+				verify(Res.IsSet() && *Res.GetValue().Value == 2);
+			}
+		}
+
+		{ // test `TSpscQueue::IsEmpty()`
+			TSpscQueue<int> Q;
+			check(Q.IsEmpty());
+			Q.Enqueue(42);
+			check(!Q.IsEmpty());
+			Q.Dequeue();
+			check(Q.IsEmpty());
+		}
+
+		{ // test `TMpscQueue::IsEmpty()`
+			TMpscQueue<int> Q;
+			check(Q.IsEmpty());
+			Q.Enqueue(42);
+			check(!Q.IsEmpty());
+			Q.Dequeue();
+			check(Q.IsEmpty());
+		}
+
+		{	// test `Peek()`
+			struct FMovableOnly
+			{
+				explicit FMovableOnly(int32 InId)
+					: Id(InId)
+				{
+				}
+
+				FMovableOnly(FMovableOnly&&) = default;
+				FMovableOnly& operator=(FMovableOnly&&) = default;
+
+				int32 Id;
+			};
+
+			{
+				TSpscQueue<FMovableOnly> Q;
+				Q.Enqueue(42);
+				FMovableOnly* Peeked = Q.Peek();
+				check(Peeked);
+				check(Peeked->Id == 42);
+				TOptional<FMovableOnly> Dequeued = Q.Dequeue();
+				check(Dequeued);
+				check(Dequeued.GetValue().Id == Peeked->Id);
+			}
+
+			{
+				TMpscQueue<FMovableOnly> Q;
+				Q.Enqueue(42);
+				FMovableOnly* Peeked = Q.Peek();
+				check(Peeked);
+				check(Peeked->Id == 42);
+				TOptional<FMovableOnly> Dequeued = Q.Dequeue();
+				check(Dequeued);
+				check(Dequeued.GetValue().Id == Peeked->Id);
+			}
+		}
+
+		UE_BENCHMARK(5, TestTCircularQueueSingleThread<5'000'000>);
+		UE_BENCHMARK(5, TestQueueSingleThread<5'000'000, TQueueAdapter<TQueue<uint32, EQueueMode::Spsc>>>);
+		UE_BENCHMARK(5, TestQueueSingleThread<5'000'000, TQueueAdapter<TQueue<uint32, EQueueMode::Mpsc>>>);
+		UE_BENCHMARK(5, TestQueueSingleThread<5'000'000, TMpscQueue<uint32>>);
+		UE_BENCHMARK(5, TestQueueSingleThread<5'000'000, TSpscQueue<uint32>>);
+
+		UE_BENCHMARK(5, TestSpscQueueCorrectness<1'000'000, TMpscQueue<uint32>>);
+		UE_BENCHMARK(5, TestSpscQueueCorrectness<1'000'000, TSpscQueue<uint32>>);
+
+		UE_BENCHMARK(5, TestTCircularQueue<5'000'000>);
+		UE_BENCHMARK(5, TestSpscQueue<5'000'000, TQueueAdapter<TQueue<uint32, EQueueMode::Spsc>>>);
+		UE_BENCHMARK(5, TestSpscQueue<5'000'000, TQueueAdapter<TQueue<uint32, EQueueMode::Mpsc>>>);
+		UE_BENCHMARK(5, TestSpscQueue<5'000'000, TMpscQueue<uint32>>);
+		UE_BENCHMARK(5, TestSpscQueue<5'000'000, TSpscQueue<uint32>>);
+
+		UE_BENCHMARK(5, TestMpscQueueCorrectness<5'000'000, TQueueAdapter<TQueue<uint32, EQueueMode::Mpsc>>>);
+		UE_BENCHMARK(5, TestMpscQueueCorrectness<5'000'000, TMpscQueue<uint32>>);
+
+		UE_BENCHMARK(5, TestMpscQueue<100'000, TQueueAdapter<TQueue<uint32, EQueueMode::Mpsc>>>);
+		UE_BENCHMARK(5, TestMpscQueue<100'000, TMpscQueue<uint32>>);
+	}
+}
+
+namespace ClosableMpscQueueTests
+{
+	// straightforward implementation over MPSC TQueue for comparison with TClosableMpscQueue
+	template<typename T>
+	class TClosableMpscQueueOnTQueue
+	{
+	public:
+		bool Enqueue(T Value)
+		{
+			bool bRes = false;
+
+			// `acquire` to make it "happen-before" reading `bIsComplete` and `release` to "synchronise-with" reading from other threads
+			verify(SubscribingThreadsNum.fetch_add(1, std::memory_order_acq_rel) >= 0);
+			// `acquire` to "synchronise-with" the consumer
+			if (!bIsComplete.load(std::memory_order_acquire))
+			{
+				ensureAlways(Queue.Enqueue(MoveTemp(Value)));
+				bRes = true;
+			}
+			// `release` to "synchonise-with" other threads
+			verify(SubscribingThreadsNum.fetch_sub(1, std::memory_order_release) >= 1);
+
+			return bRes;
+		}
+
+		template<typename F>
+		void Close(F&& Consumer)
+		{
+			// `seq_cst` to make it "happen-before" reading `SubscribingThreadsNum`, and to "synchronise-with" reading from other threads
+			bIsComplete.store(true, std::memory_order_seq_cst);
+
+			// `acquire` to "synchronise-with" other threads and to make it "happen-before" dequeueing
+			while (SubscribingThreadsNum.load(std::memory_order_acquire) != 0)
+			{
+				FPlatformProcess::Yield();
+			}
+
+			T Value;
+			while (Queue.Dequeue(Value))
+			{
+				Consumer(MoveTemp(Value));
+			}
+		}
+
+	private:
+		TQueue<T, EQueueMode::Mpsc> Queue;
+		std::atomic<bool> bIsComplete{ false };
+		std::atomic<uint32> SubscribingThreadsNum{ 0 };
+	};
+
+	// straightforward implementation over TMpscQueue for comparison with TClosableMpscQueue
+	template<typename T>
+	class TClosableMpscQueueOnTMpscQueue
+	{
+	public:
+		bool Enqueue(T Value)
+		{
+			bool bRes = false;
+
+			// `acquire` to make it "happen-before" reading `bIsComplete` and `release` to "synchronise-with" reading from other threads
+			verify(SubscribingThreadsNum.fetch_add(1, std::memory_order_acq_rel) >= 0);
+			// `acquire` to "synchronise-with" the consumer
+			if (!bIsComplete.load(std::memory_order_acquire))
+			{
+				Queue.Enqueue(MoveTemp(Value));
+				bRes = true;
+			}
+			// `release` to "synchonise-with" other threads
+			verify(SubscribingThreadsNum.fetch_sub(1, std::memory_order_release) >= 1);
+
+			return bRes;
+		}
+
+		template<typename F>
+		void Close(F&& Consumer)
+		{
+			// `seq_cst` to make it "happen-before" reading `SubscribingThreadsNum`, and to "synchronise-with" reading from other threads
+			bIsComplete.store(true, std::memory_order_seq_cst);
+
+			// `acquire` to "synchronise-with" other threads and to make it "happen-before" dequeueing
+			while (SubscribingThreadsNum.load(std::memory_order_acquire) != 0)
+			{
+				FPlatformProcess::Yield();
+			}
+
+			while (TOptional<T> Value = Queue.Dequeue())
+			{
+				Consumer(MoveTemp(Value.GetValue()));
+			}
+		}
+
+	private:
+		TMpscQueue<T> Queue;
+		std::atomic<bool> bIsComplete{ false };
+		std::atomic<uint32> SubscribingThreadsNum{ 0 };
+	};
+
+	// wrapper over `TClosableLockFreePointerListUnorderedSingleConsumer`
+	template<typename T>
+	class TClosableLockFreePointerListUnorderedSingleConsumer_Adapter
+	{
+	public:
+		bool Enqueue(T* Value)
+		{
+			return Queue.PushIfNotClosed(MoveTemp(Value));
+		}
+
+		template<typename F>
+		void Close(F&& Consumer)
+		{
+			TArray<T*> Items;
+			Queue.PopAllAndClose(Items);
+			for (T* Item : Items)
+			{
+				Consumer(MoveTemp(Item));
+			}
+		}
+
+	private:
+		TClosableLockFreePointerListUnorderedSingleConsumer<T, 0 /* as was used in FGraphEvent */> Queue;
+	};
+
+	template<uint32 Num, typename QueueT>
+	void TestClosingEmptyQueue()
+	{
+		for (uint32 i = 0; i != Num; ++i)
+		{
+			QueueT{}.Close([](void*) {});
+		}
+	}
+
+	// launches a number of producer tasks that push incremental numbers into a queue until it's closed. then, on the main thread,
+	// the queue is closed, supposedly concurrently to producers trying pushing into it. after that it's checked that the same total sum was 
+	// pushed and popped from the queue
+	template<int32 Num, typename QueueT>
+	void TestCorrectness()
+	{
+		for (int32 Run = 0; Run != Num; ++Run)
+		{
+			QueueT Queue;
+
+			int32 ProducersNum = FPlatformMisc::NumberOfWorkerThreadsToSpawn();
+
+			FGraphEventArray Producers;
+			Producers.Reserve(ProducersNum);
+
+			TArray<uint32> NumsProduced;
+			NumsProduced.AddZeroed(ProducersNum);
+
+			for (int32 i = 0; i != ProducersNum; ++i)
+			{
+				Producers.Add(FFunctionGraphTask::CreateAndDispatchWhenReady(
+					[&Queue, NumProduced = &NumsProduced[i]]
+					{
+						uint32 i = 0;
+						while (Queue.Enqueue((void*)(intptr_t)(++i)))
+						{
+						}
+
+						*NumProduced = i - 1;
+					}
+				));
+			}
+
+			FPlatformProcess::Yield();
+
+			uint64 Consumed = 0;
+			Queue.Close([&Consumed](void* Value) { Consumed += (uint32)(intptr_t)Value; });
+
+			FTaskGraphInterface::Get().WaitUntilTasksComplete(Producers);
+
+			uint64 Produced = 0;
+			for (int32 i = 0; i != ProducersNum; ++i)
+			{
+				uint32 N = NumsProduced[i];
+				Produced += N * (N + 1) / 2;
+			}
+
+			checkf(Produced == Consumed, TEXT("%" INT64_FMT " - %" INT64_FMT ", %d run"), Produced, Consumed, Run);
+
+			if (Produced == 0)
+			{
+				Run -= 1; // discard this run
+			}
+		}
+	}
+
+	template<uint32 Num, typename QueueType>
+	void TestPerf()
+	{
+		TArray<QueueType, TAlignedHeapAllocator<alignof(QueueType)>> Queues;
+		Queues.AddDefaulted(Num);
+		std::atomic<int> Index{ 0 };
+
+		const int32 NumProducers = FPlatformMisc::NumberOfWorkerThreadsToSpawn();
+		TArray<int32> ProducersResults;
+		ProducersResults.AddDefaulted(NumProducers);
+
+		FGraphEventArray Events;
+		for (int ProducerIndex = 0; ProducerIndex != NumProducers; ++ProducerIndex)
+		{
+			Events.Add(FFunctionGraphTask::CreateAndDispatchWhenReady(
+				[&Queues, &Index, &ProducersResults, ProducerIndex]
+				{
+					int NumProduced = 0;
+					int LocalIndex = Index.load(std::memory_order_acquire);
+
+					while (LocalIndex != Num && Queues[LocalIndex].Enqueue((void*)((intptr_t)NumProduced + 1)))
+					{
+						++NumProduced;
+						LocalIndex = Index.load(std::memory_order_acquire);
+					}
+
+					ProducersResults[ProducerIndex] = NumProduced;
+				}
+				));
+			Events.Last()->SetDebugName(TEXT("ClosableMpscQueueTest.TestPerf.Producer"));
+		}
+
+		int32 NumConsumed = 0;
+		for (int LocalIndex = 0; LocalIndex != Num; ++LocalIndex)
+		{
+			Queues[LocalIndex].Close([&NumConsumed](void*) { ++NumConsumed; });
+			Index.store(LocalIndex + 1, std::memory_order_release);
+			FPlatformProcess::Yield();
+			FPlatformProcess::Yield();
+			FPlatformProcess::Yield();
+		}
+
+		FTaskGraphInterface::Get().WaitUntilTasksComplete(Events);
+
+		int32 NumProduced = 0;
+		for (int i = 0; i != NumProducers; ++i)
+		{
+			NumProduced += ProducersResults[i];
+		}
+
+		check(NumProduced == NumConsumed);
+		UE_LOG(LogTemp, Display, TEXT("items queued %d"), NumProduced);
+	}
+
+	TEST_CASE_NAMED(FClosableMpscQueueTest, "System::Core::Async::ClosableMpscQueueTest", "[ApplicationContextMask][EngineFilter]")
+	{
+		{
+			TClosableMpscQueue<int> Q;
+			Q.Close([](int) { check(false); });
+		}
+		{
+			TClosableMpscQueue<int> Q;
+			Q.Enqueue(1);
+			bool Done = false;
+			Q.Close([&Done](int Item) { check(!Done && Item == 1); Done = true; });
+		}
+
+		{
+			TClosableMpscQueue<int> Q;
+			Q.Enqueue(1);
+			Q.Enqueue(2);
+			int Step = 0;
+			Q.Close(
+				[&Step](int Item)
+				{
+					switch (Step)
+					{
+					case 0:
+						check(Item == 1);
+						++Step;
+						break;
+					case 1:
+						check(Item == 2);
+						++Step;
+						break;
+					default:
+						check(false);
+					}
+				}
+			);
+		}
+
+		//for (int i = 0; i != 1000000; ++i)
+		{
+			TClosableMpscQueue<int> Q;
+			Q.Close([](int) { check(false); });
+			Q.Close([](int) { check(false); });
+		}
+
+		//for (int i = 0; i != 1000000; ++i)
+		{
+			TClosableMpscQueue<int> Q;
+			check(!Q.IsClosed());
+			Q.Close([](int) { check(false); });
+			check(Q.IsClosed());
+		}
+
+		UE_BENCHMARK(5, TestClosingEmptyQueue<10'000'000, TClosableMpscQueue<void*>>);
+		UE_BENCHMARK(5, TestClosingEmptyQueue<10'000'000, TClosableMpscQueueOnTQueue<void*>>);
+		UE_BENCHMARK(5, TestClosingEmptyQueue<10'000'000, TClosableMpscQueueOnTMpscQueue<void*>>);
+		UE_BENCHMARK(5, TestClosingEmptyQueue<10'000'000, TClosableLockFreePointerListUnorderedSingleConsumer_Adapter<void>>);
+
+		UE_BENCHMARK(5, TestCorrectness<20, TClosableMpscQueue<void*>>);
+		UE_BENCHMARK(5, TestCorrectness<20, TClosableMpscQueueOnTQueue<void*>>);
+		UE_BENCHMARK(5, TestCorrectness<20, TClosableMpscQueueOnTMpscQueue<void*>>);
+		UE_BENCHMARK(5, TestCorrectness<20, TClosableLockFreePointerListUnorderedSingleConsumer_Adapter<void>>);
+
+		UE_BENCHMARK(5, TestPerf<5'000'000, TClosableMpscQueue<void*>>);
+		UE_BENCHMARK(5, TestPerf<5'000'000, TClosableMpscQueueOnTQueue<void*>>);
+		UE_BENCHMARK(5, TestPerf<5'000'000, TClosableMpscQueueOnTMpscQueue<void*>>);
+		UE_BENCHMARK(5, TestPerf<5'000'000, TClosableLockFreePointerListUnorderedSingleConsumer_Adapter<void>>);
+	}
+}
+
+namespace DepletableMpscQueueTests
+{
+	// straightforward implementation over MPSC TQueue for comparison with TResettableMpscQueue
+	template<typename T>
+	class TDepletableMpmcQueueOnTQueue
+	{
+		using FQueue = TQueue<T, EQueueMode::Mpsc>;
+
+	public:
+		~TDepletableMpmcQueueOnTQueue()
+		{
+			checkf(SubscribingThreadsNum.load(std::memory_order_relaxed) == 0, TEXT("No other threads can use the object during its destruction"));
+			delete Queue.load(std::memory_order_relaxed);
+		}
+
+		void Enqueue(T Value)
+		{
+			// `acquire`/`release` to keep `Queue` usage inside the scope
+			verify(SubscribingThreadsNum.fetch_add(1, std::memory_order_acquire) >= 0);
+			Queue.load(std::memory_order_relaxed)->Enqueue(MoveTemp(Value));
+			// `release` to "synchonise-with" other threads
+			verify(SubscribingThreadsNum.fetch_sub(1, std::memory_order_release) >= 1);
+		}
+
+		template<typename F>
+		void Deplete(const F& Consumer)
+		{
+			// `acquire` to make it "happen-before" reading `SubscribingThreadsNum`
+			FQueue* LocalQueue = Queue.exchange(new FQueue, std::memory_order_acquire);
+
+			T Value;
+			bool bHasValue = false;
+			// it's possible that a queue is exchanged and depleted after a producer has gotten the queue but before it has finished enqueuing. 
+			// we use the following counter to make the consumer wait until all producers have finished enqueing to the current queue
+			do
+			{
+				bHasValue = LocalQueue->Dequeue(Value);
+				if (bHasValue)
+				{
+					Consumer(MoveTemp(Value));
+				}
+			} while (bHasValue || SubscribingThreadsNum.load(std::memory_order_relaxed) != 0);
+
+			delete LocalQueue;
+		}
+
+	private:
+		std::atomic<FQueue*> Queue{ new FQueue };
+		std::atomic<uint32> SubscribingThreadsNum{ 0 };
+	};
+
+	// straightforward implementation over MPSC TQueue for comparison with TResettableMpscQueue
+	template<typename T>
+	class TDepletableMpmcQueueOnTMpscQueue
+	{
+	public:
+		~TDepletableMpmcQueueOnTMpscQueue()
+		{
+			checkf(SubscribingThreadsNum.load(std::memory_order_relaxed) == 0, TEXT("No other threads can use the object during its destration"));
+			delete Queue.load(std::memory_order_relaxed);
+		}
+
+		void Enqueue(T Value)
+		{
+			// `acquire`/`release` to keep `Queue` usage inside the scope
+			verify(SubscribingThreadsNum.fetch_add(1, std::memory_order_acquire) >= 0);
+			Queue.load(std::memory_order_relaxed)->Enqueue(MoveTemp(Value));
+			// `release` to "synchonise-with" other threads
+			verify(SubscribingThreadsNum.fetch_sub(1, std::memory_order_release) >= 1);
+		}
+
+		template<typename F>
+		void Deplete(const F& Consumer)
+		{
+			// `acquire` to make it "happen-before" reading `SubscribingThreadsNum`
+			TMpscQueue<T>* LocalQueue = Queue.exchange(new TMpscQueue<T>, std::memory_order_acquire);
+
+			TOptional<T> Value;
+			// it's possible that a queue is exchanged and depleted after a producer has gotten the queue but before it has finished enqueuing. 
+			// we use the following counter to make the consumer wait until all producers have finished enqueing to the current queue
+			do
+			{
+				Value = LocalQueue->Dequeue();
+				if (Value.IsSet())
+				{
+					Consumer(MoveTemp(Value.GetValue()));
+				}
+			} while (Value.IsSet() || SubscribingThreadsNum.load(std::memory_order_relaxed) != 0);
+
+			delete LocalQueue;
+		}
+
+	private:
+		std::atomic<TMpscQueue<T>*> Queue{ new TMpscQueue<T> };
+		std::atomic<uint32> SubscribingThreadsNum{ 0 };
+	};
+
+	// wrapper over `TClosableLockFreePointerListUnorderedSingleConsumer`
+	template<typename T>
+	class TDepletableMpmcQueueOnTClosableLockFreePointerListUnorderedSingleConsumer
+	{
+		using FQueue = TClosableLockFreePointerListUnorderedSingleConsumer<T, 0 /* as was used in FGraphEvent */>;
+
+	public:
+		~TDepletableMpmcQueueOnTClosableLockFreePointerListUnorderedSingleConsumer()
+		{
+			checkf(SubscribingThreadsNum.load(std::memory_order_relaxed) == 0, TEXT("No other threads can use the object during its destration"));
+			delete Queue.load(std::memory_order_relaxed);
+		}
+
+		void Enqueue(T* Value)
+		{
+			// `acquire`/`release` to keep `Queue` usage inside the scope
+			verify(SubscribingThreadsNum.fetch_add(1, std::memory_order_acquire) >= 0);
+			ensureAlways(Queue.load(std::memory_order_relaxed)->PushIfNotClosed(MoveTemp(Value)));
+			// `release` to "synchonise-with" other threads
+			verify(SubscribingThreadsNum.fetch_sub(1, std::memory_order_release) >= 1);
+		}
+
+		template<typename F>
+		void Deplete(const F& Consumer)
+		{
+			// `acquire` to make it "happen-before" reading `SubscribingThreadsNum`
+			FQueue* LocalQueue = Queue.exchange(new FQueue, std::memory_order_acquire);
+
+			// wait for all producers to finish enqueuing to this queue instance before closing it, because enqueuing to a closed queue would fail
+			// `acquire` to make it "happen-before" dequeueing
+			while (SubscribingThreadsNum.load(std::memory_order_acquire) != 0)
+			{
+				FPlatformProcess::Yield();
+			}
+
+			TArray<T*> Items;
+			LocalQueue->PopAllAndClose(Items);
+			for (T* Item : Items)
+			{
+				Consumer(MoveTemp(Item));
+			}
+
+			delete LocalQueue;
+		}
+
+	private:
+		std::atomic<FQueue*> Queue{ new FQueue };
+		std::atomic<uint32> SubscribingThreadsNum{ 0 };
+	};
+
+	// wrapper over `TClosableMpscQueue`
+	template<typename T>
+	class TDepletableMpmcQueueOnTClosableMpscQueue
+	{
+		using FQueue = TClosableMpscQueue<T>;
+
+	public:
+		~TDepletableMpmcQueueOnTClosableMpscQueue()
+		{
+			checkf(SubscribingThreadsNum.load(std::memory_order_relaxed) == 0, TEXT("No other threads can use the object during its destration"));
+			delete Queue.load(std::memory_order_relaxed);
+		}
+
+		void Enqueue(T Value)
+		{
+			// `acquire`/`release` to keep `Queue` usage inside the scope
+			verify(SubscribingThreadsNum.fetch_add(1, std::memory_order_acquire) >= 0);
+			ensureAlways(Queue.load(std::memory_order_relaxed)->Enqueue(MoveTemp(Value)));
+			// `release` to "synchonise-with" other threads
+			verify(SubscribingThreadsNum.fetch_sub(1, std::memory_order_release) >= 1);
+		}
+
+		template<typename F>
+		void Deplete(const F& Consumer)
+		{
+			// `acquire` to make it "happen-before" reading `SubscribingThreadsNum`
+			FQueue* LocalQueue = Queue.exchange(new FQueue, std::memory_order_acquire);
+
+			// wait for all producers to finish enqueuing to this queue instance before closing it, because enqueuing to a closed queue would fail
+			// `acquire` to make it "happen-before" dequeueing
+			while (SubscribingThreadsNum.load(std::memory_order_acquire) != 0)
+			{
+				FPlatformProcess::Yield();
+			}
+
+			LocalQueue->Close(Consumer);
+
+			delete LocalQueue;
+		}
+
+	private:
+		std::atomic<FQueue*> Queue{ new FQueue };
+		std::atomic<uint32> SubscribingThreadsNum{ 0 };
+	};
+
+	template<uint32 Num, typename QueueT>
+	void TestConsumingEmptyQueue()
+	{
+		QueueT Q;
+		for (uint32 i = 0; i != Num; ++i)
+		{
+			Q.Deplete([](void*) { checkNoEntry(); });
+		}
+	}
+
+	template<uint64 Num, typename QueueT>
+	void TestPerf()
+	{
+		// spawns producers num = workers num. Every producer pushes `Num` items into the queue. The consumer continuously
+		// consumes queue until all produced items are consumed
+		QueueT Queue;
+
+		int32 ProducersNum = FPlatformMisc::NumberOfWorkerThreadsToSpawn();
+
+		FGraphEventArray Producers;
+		Producers.Reserve(ProducersNum);
+
+		for (int32 i = 0; i != ProducersNum; ++i)
+		{
+			Producers.Add(FFunctionGraphTask::CreateAndDispatchWhenReady
+			(
+				[&Queue]
+				{
+					uint64 Total = 0;
+					for (uint64 Product = 1; Product != Num + 1; ++Product)
+					{
+						Queue.Enqueue((void*)(intptr_t)Product);
+						Total += Product;
+					}
+					check(Total == Num * (Num + 1) / 2);
+				}
+			));
+		}
+
+		uint64 Produced = Num * (Num + 1) / 2 * ProducersNum;
+		uint64 Consumed = 0;
+		while (Consumed != Produced)
+		{
+			Queue.Deplete
+			(
+				[&Consumed](void* Value) 
+				{ 
+					uint64 IntValue = (uint64)(intptr_t)Value;
+					Consumed += IntValue;
+				}
+			);
+			checkf(Produced >= Consumed, TEXT("%llu - %llu"), Produced, Consumed);
+		}
+
+		checkf(Produced == Consumed, TEXT("%llu - %llu"), Produced, Consumed);
+
+		FTaskGraphInterface::Get().WaitUntilTasksComplete(Producers);
+
+		Queue.Deplete([](void*) { checkNoEntry(); }); // empty
+	}
+
+	TEST_CASE_NAMED(FDepletableMpscQueueTest, "System::Core::Async::DepletableMpscQueueTest", "[ApplicationContextMask][EngineFilter]")
+	{
+		{
+			TDepletableMpmcQueue<int> Q;
+			Q.Enqueue(1);
+			Q.Enqueue(2);
+			Q.Deplete([](int) {});
+			Q.Enqueue(3);
+			Q.Deplete([](int) {});
+			Q.Deplete([](int) {});
+			Q.Enqueue(4); // to destroy non-empty queue
+		}
+
+		{
+			TDepletableMpmcQueue<int> Q;
+			Q.Deplete([](int) { check(false); });
+		}
+		{
+			TDepletableMpmcQueue<int> Q;
+			Q.Enqueue(1);
+			bool Done = false;
+			Q.Deplete([&Done](int Item) { check(!Done && Item == 1); Done = true; });
+		}
+
+		{
+			TDepletableMpmcQueue<int> Q;
+			Q.Enqueue(1);
+			Q.Enqueue(2);
+			int Step = 0;
+			Q.Deplete(
+				[&Step](int Item)
+				{
+					switch (Step)
+					{
+					case 0:
+						check(Item == 1);
+						++Step;
+						break;
+					case 1:
+						check(Item == 2);
+						++Step;
+						break;
+					default:
+						check(false);
+					}
+				}
+			);
+		}
+
+		//for (int i = 0; i != 1000000; ++i)
+		{
+			TDepletableMpmcQueue<int> Q;
+			Q.Deplete([](int) { check(false); });
+			Q.Deplete([](int) { check(false); });
+		}
+
+		//for (int i = 0; i != 1000000; ++i)
+		{
+			int Res = 0;
+			TDepletableMpmcQueue<int> Q;
+			Q.Enqueue(1);
+			Q.Deplete([&Res](int i) { Res += i; });
+			Q.Enqueue(2);
+			Q.Deplete([&Res](int i) { Res += i; });
+			checkf(Res == 3, TEXT("%d"), Res);
+		}
+
+		{	// EnqueueAndReturnWasEmpty
+			TDepletableMpmcQueue<int> Q;
+			verify(Q.EnqueueAndReturnWasEmpty(1));
+			verify(!Q.EnqueueAndReturnWasEmpty(2));
+			Q.Deplete([](int) {});
+			verify(Q.EnqueueAndReturnWasEmpty(1));
+			verify(!Q.EnqueueAndReturnWasEmpty(2));
+			Q.Deplete([](int) {});
+			verify(Q.EnqueueAndReturnWasEmpty(1));
+			verify(!Q.EnqueueAndReturnWasEmpty(2));
+		}
+
+		UE_BENCHMARK(5, TestPerf<10'000, TDepletableMpmcQueue<void*>>);
+		UE_BENCHMARK(5, TestPerf<10'000, TDepletableMpmcQueueOnTQueue<void*>>);
+		UE_BENCHMARK(5, TestPerf<10'000, TDepletableMpmcQueueOnTMpscQueue<void*>>);
+		UE_BENCHMARK(5, TestPerf<10'000, TDepletableMpmcQueueOnTClosableLockFreePointerListUnorderedSingleConsumer<void>>);
+		UE_BENCHMARK(5, TestPerf<10'000, TDepletableMpmcQueueOnTClosableMpscQueue<void*>>);
+
+		UE_BENCHMARK(5, TestConsumingEmptyQueue<10'000'000, TDepletableMpmcQueue<void*>>);
+		UE_BENCHMARK(5, TestConsumingEmptyQueue<10'000'000, TDepletableMpmcQueueOnTQueue<void*>>);
+		UE_BENCHMARK(5, TestConsumingEmptyQueue<10'000'000, TDepletableMpmcQueueOnTMpscQueue<void*>>);
+		UE_BENCHMARK(5, TestConsumingEmptyQueue<10'000'000, TDepletableMpmcQueueOnTClosableLockFreePointerListUnorderedSingleConsumer<void>>);
+		UE_BENCHMARK(5, TestConsumingEmptyQueue<10'000'000, TDepletableMpmcQueueOnTClosableMpscQueue<void*>>);
+	}
+}}
+
+#endif // WITH_TESTS
diff --git a/Engine/Source/Runtime/Core/Tests/Containers/RingBufferTest.cpp b/Engine/Source/Runtime/Core/Tests/Containers/RingBufferTest.cpp
new file mode 100644
--- /dev/null
+++ b/Engine/Source/Runtime/Core/Tests/Containers/RingBufferTest.cpp
@@ -0,0 +1,1128 @@
+// Copyright Epic Games, Inc. All Rights Reserved.
+
+#include "Containers/RingBuffer.h"
+
+#if WITH_TESTS
+
+#include "Tests/TestHarnessAdapter.h"
+
+
+class FRingBufferTest
+{
+public:
+	static bool IsIntegerRange(const TRingBuffer<uint32>& Queue, uint32 Start, uint32 End, bool bForward = true)
+	{
+		if (Queue.Num() != End - Start)
+		{
+			return false;
+		}
+
+		// Peek elements in queue at given offset, peek from back to front
+		for (int32 It = 0; It < Queue.Num(); ++It)
+		{
+			uint32 QueueValue = bForward ? Queue[It] : Queue[Queue.Num() - 1 - It];
+			if (QueueValue != Start + It)
+			{
+				return false;
+			}
+		}
+
+		return true;
+	}
+
+	struct Counter
+	{
+		Counter(uint32 InValue = 0x12345)
+			:Value(InValue)
+		{
+			++NumVoid;
+		}
+		Counter(const Counter& Other)
+			:Value(Other.Value)
+		{
+			++NumCopy;
+		}
+		Counter(Counter&& Other)
+			:Value(Other.Value)
+		{
+			++NumMove;
+		}
+		~Counter()
+		{
+			++NumDestruct;
+		}
+
+		operator uint32() const
+		{
+			return Value;
+		}
+
+		bool operator==(const Counter& Other) const
+		{
+			return Value == Other.Value;
+		}
+		static int NumVoid;
+		static int NumCopy;
+		static int NumMove;
+		static int NumDestruct;
+		static void Clear()
+		{
+			NumVoid = NumCopy = NumMove = NumDestruct = 0;
+		}
+
+		uint32 Value;
+	};
+
+	FRingBufferTest() = default;
+
+	bool Test()
+	{
+		// Test empty
+		{
+			TRingBuffer<uint32> Q(0);
+
+			CHECK_MESSAGE(TEXT("Test empty - IsEmpty"), Q.IsEmpty() == true);
+			CHECK_MESSAGE(TEXT("Test empty - Size"), Q.Num() == 0);
+			CHECK_MESSAGE(TEXT("Test empty - Capacity"), Q.Max() == 0);
+			CHECK_MESSAGE(TEXT("Test empty - Iterator"), Q.begin() == Q.end());
+			CHECK_MESSAGE(TEXT("Test empty - ConvertPointerToIndex"), Q.ConvertPointerToIndex(nullptr) == INDEX_NONE);
+			CHECK_MESSAGE(TEXT("Test empty - ConvertPointerToIndex"), Q.ConvertPointerToIndex(reinterpret_cast<uint32*>(this)) == INDEX_NONE);
+			Q.Trim();
+			CHECK_MESSAGE(TEXT("Test Trim From empty - Size"), Q.Num() == 0);
+			CHECK_MESSAGE(TEXT("Test Trim From empty - Capacity"), Q.Max() == 0);
+			Q.Reset();
+			CHECK_MESSAGE(TEXT("Test Reset From empty - Size"), Q.Num() == 0);
+			CHECK_MESSAGE(TEXT("Test Reset From empty - Capacity"), Q.Max() == 0);
+			Q.Empty(0);
+			CHECK_MESSAGE(TEXT("Test Empty From empty - Size"), Q.Num() == 0);
+			CHECK_MESSAGE(TEXT("Test Empty From empty - Capacity"), Q.Max() == 0);
+			Q.PopFront(0);
+			Q.Pop(0);
+			CHECK_MESSAGE(TEXT("Test Pop on empty - Size"), Q.Num() == 0);
+			CHECK_MESSAGE(TEXT("Test Pop on empty - Capacity"), Q.Max() == 0);
+			CHECK_MESSAGE(TEXT("Test empty - IsValidIndex"), Q.IsValidIndex(0) == false);
+
+
+			const TRingBuffer<uint32> ConstQ(0);
+			CHECK_MESSAGE(TEXT("Test const empty - IsEmpty"), ConstQ.IsEmpty() == true);
+			CHECK_MESSAGE(TEXT("Test const empty - Size"), ConstQ.Num() == 0);
+			CHECK_MESSAGE(TEXT("Test const empty - Capacity"), ConstQ.Max() == 0);
+			CHECK_MESSAGE(TEXT("Test const empty - Iterator"), ConstQ.begin() == ConstQ.end());
+			CHECK_MESSAGE(TEXT("Test const empty - ConvertPointerToIndex"), ConstQ.ConvertPointerToIndex(reinterpret_cast<uint32*>(this)) == INDEX_NONE);
+		}
+
+		// Test Adding a sequence of elements
+		{
+			const TRingBuffer<int32>::IndexType FirstSize = 8;
+
+			TRingBuffer<int32> Q(0);
+
+			CHECK_MESSAGE(TEXT("Test AddSequence - Capacity (Implementation Detail)"), Q.Max() == 0);
+			Q.Emplace(0);
+			CHECK_MESSAGE(TEXT("Test AddSequence - Size"), Q.Num() == 1);
+			CHECK_MESSAGE(TEXT("Test AddSequence - Capacity (Implementation Detail)"), Q.Max() == 1);
+			Q.Emplace(1);
+			CHECK_MESSAGE(TEXT("Test AddSequence - Size"), Q.Num() == 2);
+			CHECK_MESSAGE(TEXT("Implementation Detail - These tests expect that growing size will set capacity to successive powers of 2."), Q.Max() == 2);
+			for (int32 It = 2; It < FirstSize; ++It)
+			{
+				Q.Emplace(It);
+				CHECK_MESSAGE(TEXT("Test AddSequence - Size"), Q.Num() == It + 1);
+				CHECK_MESSAGE(TEXT("Test AddSequence - Capacity (Implementation Detail)"),
+					static_cast<uint32>(Q.Max()) == FMath::RoundUpToPowerOfTwo(It + 1));
+			}
+
+			for (int32 Index = 0; Index < FirstSize; ++Index)
+			{
+				CHECK_MESSAGE(TEXT("Test AddSequence - Expected values"), Q[Index] == Index);
+				CHECK_MESSAGE(TEXT("Test AddSequence const- Expected values"), const_cast<TRingBuffer<int32>&>(Q)[Index] == Index);
+			}
+
+			const TRingBuffer<int32>::IndexType SecondSize = 13;
+			for (int32 It = FirstSize; It < SecondSize; ++It)
+			{
+				Q.Emplace(It);
+				CHECK_MESSAGE(TEXT("Test AddSequence non powerof2 - Size"), Q.Num() == It + 1);
+				CHECK_MESSAGE(TEXT("Test AddSequence non powerof2 const - Capacity (Implementation Detail)"),
+					static_cast<uint32>(Q.Max()) == FMath::RoundUpToPowerOfTwo(It + 1));
+			}
+
+			for (int32 Index = 0; Index < FirstSize; ++Index)
+			{
+				CHECK_MESSAGE(TEXT("Test AddSequence non powerof2 - Expected values"), Q[Index] == Index);
+				CHECK_MESSAGE(TEXT("Test AddSequence non powerof2 const - Expected values"), const_cast<TRingBuffer<int32>&>(Q)[Index] == Index);
+			}
+		}
+
+		// Test Add under/over Capacity
+		{
+			const TRingBuffer<int32>::IndexType FirstElementsToAdd = 3;
+			const TRingBuffer<int32>::IndexType InitialCapacity = 8;
+			const TRingBuffer<int32>::IndexType SecondElementsToAdd = 9;
+
+			TRingBuffer<int32> Q(InitialCapacity);
+
+			for (int32 It = 0; It < FirstElementsToAdd; ++It)
+			{
+				Q.Emplace(It);
+			}
+
+			CHECK_MESSAGE(TEXT("Test Add under Capacity - Size"), Q.Num() == FirstElementsToAdd);
+			CHECK_MESSAGE(TEXT("Test Add under Capacity - Capacity"), Q.Max() == InitialCapacity);
+			for (int32 Index = 0; Index < FirstElementsToAdd; ++Index)
+			{
+				CHECK_MESSAGE(TEXT("Test Add under Capacity - Expected values"), Q[Index] == Index);
+				CHECK_MESSAGE(TEXT("Test Add under Capacity const - Expected values"), const_cast<TRingBuffer<int32>&>(Q)[Index] == Index);
+			}
+
+			for (int32 It = FirstElementsToAdd; It < SecondElementsToAdd; ++It)
+			{
+				Q.Emplace(It);
+			}
+
+			CHECK_MESSAGE(TEXT("Test Add over Capacity - Size"), Q.Num() == SecondElementsToAdd);
+			CHECK_MESSAGE(TEXT("Test Add over Capacity - Capacity (Implementation Detail)"),
+				static_cast<uint32>(Q.Max()) == FMath::RoundUpToPowerOfTwo(SecondElementsToAdd));
+			for (int32 Index = 0; Index < SecondElementsToAdd; ++Index)
+			{
+				CHECK_MESSAGE(TEXT("Test Add over Capacity - Expected values"), Q[Index] == Index);
+				CHECK_MESSAGE(TEXT("Test Add over Capacity const - Expected values"), const_cast<TRingBuffer<int32>&>(Q)[Index] == Index);
+			}
+		}
+
+		// Test Last/First
+		{
+			TRingBuffer<uint32> Q({ 0,1,2,3 });
+			CHECK_MESSAGE(TEXT("Test Last"), 3 == Q.Last());
+			Q.Last() = 4;
+			CHECK_MESSAGE(TEXT("Test Last const"), 4 == const_cast<TRingBuffer<uint32>&>(Q).Last());
+			CHECK_MESSAGE(TEXT("Test First"), 0 == Q.First());
+			Q.First() = 5;
+			CHECK_MESSAGE(TEXT("Test First const"), 5 == const_cast<TRingBuffer<uint32>&>(Q).First());
+		}
+
+		// Test PopFrontValue/PopValue
+		{
+			TRingBuffer<Counter> Q({ 31,32,33 });
+			Q.AddFront(30);
+
+			Counter::Clear();
+			Counter C(Q.PopFrontValue());
+			CHECK_MESSAGE(TEXT("PopFrontValue - PoppedValue"), C.Value == 30);
+			CHECK_MESSAGE(TEXT("PopFrontValue - ConstructorCounts"), (Counter::NumMove > 0 && Counter::NumCopy == 0) == true);
+			CHECK_MESSAGE(TEXT("PopFrontValue - Remaining Values"), Q == TRingBuffer<Counter>({ 31,32,33 }));
+			Counter::Clear();
+			CHECK_MESSAGE(TEXT("PopFrontValue Inline - PoppedValue"), Q.PopFrontValue().Value == 31);
+			CHECK_MESSAGE(TEXT("PopFrontValue Inline - ConstructorCounts"), (Counter::NumCopy == 0) == true);
+			CHECK_MESSAGE(TEXT("PopFrontValue Inline - Remaining Values"), Q == TRingBuffer<Counter>({ 32,33 }));
+
+			Counter::Clear();
+			Counter D(Q.PopValue());
+			CHECK_MESSAGE(TEXT("PopValue - PoppedValue"), D.Value == 33);
+			CHECK_MESSAGE(TEXT("PopValue - ConstructorCounts"), (Counter::NumMove > 0 && Counter::NumCopy == 0) == true);
+			CHECK_MESSAGE(TEXT("PopValue - Remaining Values"), Q == TRingBuffer<Counter>({ Counter(32) }));
+			Counter::Clear();
+			CHECK_MESSAGE(TEXT("PopValue Inline - PoppedValue"), Q.PopValue().Value == 32);
+			CHECK_MESSAGE(TEXT("PopValue Inline - ConstructorCounts"), (Counter::NumCopy == 0) == true);
+			CHECK_MESSAGE(TEXT("PopValue Inline - Remaining Values"), Q.IsEmpty() == true);
+		}
+
+		// Test Initializer_List
+		{
+			const TRingBuffer<int32>::IndexType InitializerSize = 9;
+			TRingBuffer<int32> Q({ 0, 1, 2, 3, 4, 5, 6, 7, 8 });
+
+			CHECK_MESSAGE(TEXT("Test Initializer_List - Size"), Q.Num() == InitializerSize);
+			CHECK_MESSAGE(TEXT("Test Initializer_List - Capacity (Implementation Detail)"),
+				static_cast<uint32>(Q.Max()) == FMath::RoundUpToPowerOfTwo(InitializerSize));
+			for (int32 Index = 0; Index < InitializerSize; ++Index)
+			{
+				CHECK_MESSAGE(TEXT("Test Initializer_List - Expected values"), Q[Index] == Index);
+			}
+		}
+
+		// Test RingBuffer's Copy Constructors et al
+		{
+			TRingBuffer<uint32> Original({ 0,1,2,3,4,5,6,7 });
+			TRingBuffer<uint32> Copy(Original);
+			CHECK_MESSAGE(TEXT("Copy Constructor"), Original == Copy);
+			TRingBuffer<uint32> Moved(MoveTemp(Copy));
+			CHECK_MESSAGE(TEXT("Move Constructor"), Original == Moved);
+			CHECK_MESSAGE(TEXT("Move Constructor did in fact move"), Copy.Max() == 0);
+			TRingBuffer<uint32> AssignCopy;
+			AssignCopy = Original;
+			CHECK_MESSAGE(TEXT("Copy Assignment"), Original == AssignCopy);
+			TRingBuffer<uint32> AssignMove;
+			AssignMove = MoveTemp(AssignCopy);
+			CHECK_MESSAGE(TEXT("Move Assignment"), Original == AssignMove);
+			CHECK_MESSAGE(TEXT("Move Assignment did in fact move"), AssignCopy.Max() == 0);
+		}
+
+		// Test Equality 
+		{
+			auto TestEquality = [this](const TCHAR* Message, bool ExpectedEqual, const TRingBuffer<int32>& A, const TRingBuffer<int32>& B)
+			{
+				CHECK_MESSAGE(*FString::Printf(TEXT("Test equality - %s - A == B"), Message), (A == B) == ExpectedEqual);
+				CHECK_MESSAGE(*FString::Printf(TEXT("Test equality - %s - B == A"), Message), (B == A) == ExpectedEqual);
+				CHECK_MESSAGE(*FString::Printf(TEXT("Test equality - %s - A != B"), Message), (A != B) == !ExpectedEqual);
+				CHECK_MESSAGE(*FString::Printf(TEXT("Test equality - %s - B != A"), Message), (B != A) == !ExpectedEqual);
+			};
+
+			TestEquality(TEXT("empty"), true, TRingBuffer<int32>(0), TRingBuffer<int32>(0));
+			TestEquality(TEXT("empty different capacities"), true, TRingBuffer<int32>(0), TRingBuffer<int32>(8));
+			TestEquality(TEXT("equal nonempty powerof2"), true, TRingBuffer<int32>({ 0, 1, 2, 3 }), TRingBuffer<int32>({ 0, 1, 2, 3 }));
+			TestEquality(TEXT("equal nonempty nonpowerof2"), true, TRingBuffer<int32>({ 0, 1, 2, 3, 4, 5 }), TRingBuffer<int32>({ 0, 1, 2, 3, 4, 5 }));
+			{
+				TRingBuffer<int32> QNum6Cap16(16);
+				for (int32 Index = 0; Index < 6; ++Index)
+				{
+					QNum6Cap16.Add(Index);
+				}
+				TestEquality(TEXT("equal nonempty different capacities"), true, QNum6Cap16, TRingBuffer<int32>({ 0, 1, 2, 3, 4, 5 }));
+			}
+
+			TestEquality(TEXT("empty to nonempty"), false, TRingBuffer<int32>(0), TRingBuffer<int32>({ 0, 1, 2, 3, 4, 5 }));
+			TestEquality(TEXT("smaller size to bigger size"), false, TRingBuffer<int32>({ 0, 1, 2 }), TRingBuffer<int32>({ 0, 1, 2, 3, 4, 5 }));
+			TestEquality(TEXT("same size different elements"), false, TRingBuffer<int32>({ 0, 1, 2 }), TRingBuffer<int32>({ 0, 1, 3 }));
+			TestEquality(TEXT("same elements different order"), false, TRingBuffer<int32>({ 0, 1, 2 }), TRingBuffer<int32>({ 0, 2, 1 }));
+
+			for (int HasPow2 = 0; HasPow2 < 2; ++HasPow2)
+			{
+				const int Count = HasPow2 ? 8 : 7;
+				TRingBuffer<int32> Q0Pop;
+				TRingBuffer<int32> Q1PopFront;
+				TRingBuffer<int32> Q2PopFront;
+				TRingBuffer<int32> Q1Pop;
+				TRingBuffer<int32> Q2Pop;
+				TRingBuffer<int32> Q2PopFront3Pop;
+				Q1PopFront.Add(47);
+				Q2PopFront.Add(576);
+				Q2PopFront.Add(-5);
+				Q2PopFront3Pop.Add(84);
+				Q2PopFront3Pop.Add(1000);
+				for (int Index = 0; Index < Count; ++Index)
+				{
+					Q0Pop.Add(Index);
+					Q1PopFront.Add(Index);
+					Q2PopFront.Add(Index);
+					Q1Pop.Add(Index);
+					Q2Pop.Add(Index);
+					Q2PopFront3Pop.Add(Index);
+				}
+				Q1PopFront.PopFront();
+				Q2PopFront.PopFront();
+				Q2PopFront.PopFront();
+				Q1Pop.Add(-18);
+				Q1Pop.Pop();
+				Q2Pop.Add(105);
+				Q2Pop.Add(219);
+				Q2Pop.Pop();
+				Q2Pop.Pop();
+				Q2PopFront3Pop.Add(456);
+				Q2PopFront3Pop.Add(654);
+				Q2PopFront3Pop.Add(8888888);
+				Q2PopFront3Pop.PopFront();
+				Q2PopFront3Pop.Pop();
+				Q2PopFront3Pop.PopFront();
+				Q2PopFront3Pop.Pop();
+				Q2PopFront3Pop.Pop();
+
+				const TCHAR* Names[] =
+				{
+					TEXT("Q0Pop"),
+					TEXT("Q1PopFront"),
+					TEXT("Q2PopFront"),
+					TEXT("Q1Pop"),
+					TEXT("Q2Pop"),
+					TEXT("Q2PopFront3Pop"),
+				};
+				TRingBuffer<int32>* Pops[] =
+				{
+					&Q0Pop,
+					&Q1PopFront,
+					&Q2PopFront,
+					&Q1Pop,
+					&Q2Pop,
+					&Q2PopFront3Pop
+				};
+
+
+				auto TestThesePops = [this, HasPow2, &TestEquality, &Names, &Pops](int TrialA, int TrialB)
+				{
+					TestEquality(*FString::Printf(TEXT("%s - %s - %s"), Names[TrialA], Names[TrialB], (HasPow2 ? TEXT("powerof2") : TEXT("nonpowerof2"))), true, *Pops[TrialA], *Pops[TrialB]);
+				};
+
+				for (int TrialA = 0; TrialA < UE_ARRAY_COUNT(Names); ++TrialA)
+				{
+					for (int TrialB = TrialA /* test each against itself as well */; TrialB < UE_ARRAY_COUNT(Names); ++TrialB)
+					{
+						TestThesePops(TrialA, TrialB);
+					}
+				}
+			}
+		}
+
+		// Test Add and pop all
+		for (int Direction = 0; Direction < 2; ++Direction)
+		{
+			bool bIsAddBack = Direction == 0;
+			auto GetMessage = [&bIsAddBack](const TCHAR* Message)
+			{
+				return FString::Printf(TEXT("Test %s (%s)"), Message, (bIsAddBack ? TEXT("AddBack") : TEXT("AddFront")));
+			};
+
+			// Test Mixed Adds and Pops
+			{
+				const TRingBuffer<uint32>::IndexType ElementsToAdd = 256;
+				const TRingBuffer<uint32>::IndexType ElementPopMod = 16;
+				const TRingBuffer<uint32>::IndexType ExpectedSize = 256 - ElementPopMod;
+				const TRingBuffer<uint32>::IndexType ExpectedCapacity = 256;
+
+				TRingBuffer<uint32> Q(4);
+
+				uint32 ExpectedPoppedValue = 0;
+				for (uint32 It = 0; It < 256; ++It)
+				{
+					if (bIsAddBack)
+					{
+						Q.Add(It);
+						CHECK_MESSAGE(*GetMessage(TEXT("Add and pop - Add")), It == Q[Q.Num() - 1]);
+					}
+					else
+					{
+						Q.AddFront(It);
+						CHECK_MESSAGE(*GetMessage(TEXT("Add and pop - Add")), It == Q[0]);
+					}
+
+					if (It % ElementPopMod == 0)
+					{
+						uint32 PoppedValue;
+						if (bIsAddBack)
+						{
+							PoppedValue = Q[0];
+							Q.PopFront();
+						}
+						else
+						{
+							PoppedValue = Q[Q.Num() - 1];
+							Q.Pop();
+						}
+						CHECK_MESSAGE(*GetMessage(TEXT("Add and pop - Pop")), ExpectedPoppedValue == PoppedValue);
+						++ExpectedPoppedValue;
+					}
+				}
+
+				CHECK_MESSAGE(*GetMessage(TEXT("Add and pop - Size")), Q.Num() == ExpectedSize);
+				CHECK_MESSAGE(*GetMessage(TEXT("Add and pop - Capacity")), Q.Max() == ExpectedCapacity);
+				CHECK_MESSAGE(*GetMessage(TEXT("Add and pop - IntegerRange")),
+					IsIntegerRange(Q, ExpectedPoppedValue, ExpectedPoppedValue + ExpectedSize, bIsAddBack) == true);
+			}
+
+
+			// Popping down to empty
+			{
+				const TRingBuffer<uint32>::IndexType ElementsToAdd = 256;
+
+				TRingBuffer<uint32> Q(ElementsToAdd);
+
+				CHECK_MESSAGE(*GetMessage(TEXT("Add and pop all - IsEmpty before")), Q.IsEmpty() == true);
+				CHECK_MESSAGE(*GetMessage(TEXT("Add and pop all - Size before")), Q.Num() == 0);
+
+				for (TRingBuffer<int32>::IndexType It = 0; It < ElementsToAdd; ++It)
+				{
+					if (bIsAddBack)
+					{
+						Q.Add(It);
+					}
+					else
+					{
+						Q.AddFront(It);
+					}
+				}
+
+				CHECK_MESSAGE(*GetMessage(TEXT("Add and pop all - Size")), Q.Num() == ElementsToAdd);
+				CHECK_MESSAGE(*GetMessage(TEXT("Add and pop all - Capacity")), Q.Max() == ElementsToAdd);
+				CHECK_MESSAGE(*GetMessage(TEXT("Add and pop all - Expected")), IsIntegerRange(Q, 0, ElementsToAdd, bIsAddBack) == true);
+
+				for (TRingBuffer<int32>::IndexType It = 0; It < ElementsToAdd; ++It)
+				{
+					if (bIsAddBack)
+					{
+						Q.PopFront();
+					}
+					else
+					{
+						Q.Pop();
+					}
+				}
+
+				CHECK_MESSAGE(*GetMessage(TEXT("Add and pop all - IsEmpty after")), Q.IsEmpty() == true);
+				CHECK_MESSAGE(*GetMessage(TEXT("Add and pop all - Size after")), Q.Num() == 0);
+				CHECK_MESSAGE(*GetMessage(TEXT("Add and pop all - Capacity after")), Q.Max() == ElementsToAdd);
+			}
+
+			// Test index wrap
+			{
+				for (int32 Offset : {-12, -8, -5, -1, 0, 2, 7, 8, 15})
+				{
+					const TRingBuffer<uint32>::IndexType ElementsToAdd = 256;
+					const TRingBuffer<uint32>::IndexType ElementPopMod = 16;
+					const TRingBuffer<uint32>::IndexType ExpectedSize = 256 - ElementPopMod;
+					const TRingBuffer<uint32>::IndexType ExpectedCapacity = 256;
+
+					TRingBuffer<uint32> Q(8);
+
+					// Set front and afterback to an arbitrary offset
+					// Note that AfterBack is always exactly equal to Front + Num()
+					Q.Front = Offset;
+					Q.AfterBack = Q.Front;
+
+					CHECK_MESSAGE(*GetMessage(TEXT("index wrap - IsEmpty before")), Q.IsEmpty() == true);
+					CHECK_MESSAGE(*GetMessage(TEXT("index wrap - Size before")), Q.Num() == 0);
+
+					for (TRingBuffer<uint32>::IndexType It = 0; It < ElementsToAdd; ++It)
+					{
+						if (bIsAddBack)
+						{
+							Q.Add(It);
+						}
+						else
+						{
+							Q.AddFront(It);
+						}
+					}
+
+					CHECK_MESSAGE(*GetMessage(TEXT("index wrap - Size")), Q.Num() == ElementsToAdd);
+					CHECK_MESSAGE(*GetMessage(TEXT("index wrap - Capacity")), Q.Max() == ElementsToAdd);
+					CHECK_MESSAGE(*GetMessage(TEXT("index wrap - Expected")), IsIntegerRange(Q, 0, ElementsToAdd, bIsAddBack) == true);
+
+					for (TRingBuffer<int32>::IndexType It = 0; It < ElementsToAdd; ++It)
+					{
+						if (bIsAddBack)
+						{
+							Q.PopFront();
+						}
+						else
+						{
+							Q.Pop();
+						}
+					}
+
+					CHECK_MESSAGE(*GetMessage(TEXT("index wrap - IsEmpty after")), Q.IsEmpty() == true);
+					CHECK_MESSAGE(*GetMessage(TEXT("index wrap - Size after")), Q.Num() == 0);
+					CHECK_MESSAGE(*GetMessage(TEXT("index wrap - Capacity after")), Q.Max() == ElementsToAdd);
+				}
+			}
+		}
+
+		// Test Trim
+		{
+			const TRingBuffer<int32>::IndexType ElementsToAdd = 9;
+			const TRingBuffer<int32>::IndexType ElementsToPop = 5;
+			const TRingBuffer<int32>::IndexType ExpectedCapacity = 16;
+			const TRingBuffer<int32>::IndexType ExpectedCapacityAfterTrim = 4;
+
+			TRingBuffer<uint32> Q(0);
+
+			for (TRingBuffer<int32>::IndexType It = 0; It < ElementsToAdd; ++It)
+			{
+				Q.Add(It);
+			}
+
+			CHECK_MESSAGE(TEXT("Test Trim - Size"), Q.Num() == ElementsToAdd);
+			CHECK_MESSAGE(TEXT("Test Trim - Capacity"), Q.Max() == ExpectedCapacity);
+			CHECK_MESSAGE(TEXT("Test Trim - Expected"), IsIntegerRange(Q, 0, ElementsToAdd) == true);
+
+			for (TRingBuffer<int32>::IndexType It = 0; It < ElementsToPop; ++It)
+			{
+				Q.PopFront();
+			}
+
+			Q.Trim();
+
+			CHECK_MESSAGE(TEXT("Test Trim - Size"), Q.Num() == ElementsToAdd - ElementsToPop);
+			CHECK_MESSAGE(TEXT("Test Trim - Capacity"), Q.Max() == ExpectedCapacityAfterTrim);
+			CHECK_MESSAGE(TEXT("Test Trim - Expected"), IsIntegerRange(Q, ElementsToPop, ElementsToAdd) == true);
+		}
+
+		// Test First and Last acting as two stacks
+		{
+			TRingBuffer<uint32> Q;
+
+			const uint32 ElementsToAdd = 64;
+			const uint32 ElementPopMod = 5;
+
+			for (uint32 It = 0; It < ElementsToAdd; ++It)
+			{
+				Q.Add(It);
+				CHECK_MESSAGE(TEXT("Test TwoStacks - AddBack"), Q.Last() == It);
+				Q.AddFront(It);
+				CHECK_MESSAGE(TEXT("Test TwoStacks - AddFront"), Q.First() == It);
+				if (It % ElementPopMod == 0)
+				{
+					uint32 AddValue = 0xfefefefe;
+					Q.Add(AddValue);
+					CHECK_MESSAGE(TEXT("Test TwoStacks - Sporadic Pop"), Q.Last() == AddValue);
+					Q.Pop();
+					Q.AddFront(AddValue);
+					CHECK_MESSAGE(TEXT("Test TwoStacks - Sporadic PopFront"), Q.First() == AddValue);
+					Q.PopFront();
+				}
+			}
+
+			CHECK_MESSAGE(TEXT("Test TwoStacks - MiddleSize"), Q.Num() == ElementsToAdd * 2);
+			for (uint32 It = 0; It < ElementsToAdd * 2; ++It)
+			{
+				CHECK_MESSAGE(*FString::Printf(TEXT("TwoStacks - Middle value %d"), It),
+					Q[It] == (It < ElementsToAdd ? ElementsToAdd - 1 - It : It - ElementsToAdd));
+			}
+
+			for (uint32 It = 0; It < ElementsToAdd; ++It)
+			{
+				CHECK_MESSAGE(TEXT("Test TwoStacks - Final Pop"), Q.Last() == ElementsToAdd - 1 - It);
+				Q.Pop();
+				CHECK_MESSAGE(TEXT("Test TwoStacks - Final PopFront"), Q.First() == ElementsToAdd - 1 - It);
+				Q.PopFront();
+			}
+
+			CHECK_MESSAGE(TEXT("Test TwoStacks - FinalSize"), Q.Num() == 0);
+		}
+
+		// Test adding into space that has been cleared from popping on the other side
+		{
+			for (int Direction = 0; Direction < 2; ++Direction)
+			{
+				bool bIsAddBack = Direction == 0;
+				auto GetMessage = [bIsAddBack](const TCHAR* Message)
+				{
+					return FString::Printf(TEXT("Test AddIntoPop - %s (%s)"), Message, (bIsAddBack ? TEXT("AddBack") : TEXT("AddFront")));
+				};
+				TRingBuffer<uint32> Q({ 0,1,2,3,4,5,6,7 });
+				TRingBuffer<int32>::IndexType InitialSize = 8;
+				CHECK_MESSAGE(*GetMessage(TEXT("InitialSize")), InitialSize == Q.Num());
+				CHECK_MESSAGE(*GetMessage(TEXT("InitialCapacity (Implementation Detail)")), InitialSize == Q.Max());
+
+				if (bIsAddBack)
+				{
+					Q.Pop();
+				}
+				else
+				{
+					Q.PopFront();
+				}
+				CHECK_MESSAGE(*GetMessage(TEXT("PoppedSize")), InitialSize - 1 == Q.Num());
+				CHECK_MESSAGE(*GetMessage(TEXT("PoppedCapacity")), InitialSize == Q.Max());
+
+				if (bIsAddBack)
+				{
+					Q.AddFront(8);
+				}
+				else
+				{
+					Q.Add(8);
+				}
+				CHECK_MESSAGE(*GetMessage(TEXT("AddedSize")), InitialSize == Q.Num());
+				CHECK_MESSAGE(*GetMessage(TEXT("AddedCapacity")), InitialSize == Q.Max());
+				if (bIsAddBack)
+				{
+					CHECK_MESSAGE(*GetMessage(TEXT("AddedValues")), Q == TRingBuffer<uint32>({ 8,0,1,2,3,4,5,6 }));
+				}
+				else
+				{
+					CHECK_MESSAGE(*GetMessage(TEXT("AddedValues")), Q == TRingBuffer<uint32>({ 1,2,3,4,5,6,7,8 }));
+				}
+
+				if (bIsAddBack)
+				{
+					Q.AddFront(9);
+				}
+				else
+				{
+					Q.Add(9);
+				}
+				CHECK_MESSAGE(*GetMessage(TEXT("Second AddedSize")), InitialSize + 1 == Q.Num());
+				CHECK_MESSAGE(*GetMessage(TEXT("Second AddedCapacity")), static_cast<uint32>(FMath::RoundUpToPowerOfTwo(InitialSize + 1)) == Q.Max());
+				if (bIsAddBack)
+				{
+					CHECK_MESSAGE(*GetMessage(TEXT("Second AddedValues")), Q == TRingBuffer<uint32>({ 9,8,0,1,2,3,4,5,6 }));
+				}
+				else
+				{
+					CHECK_MESSAGE(*GetMessage(TEXT("Second AddedValues")), Q == TRingBuffer<uint32>({ 1,2,3,4,5,6,7,8,9 }));
+				}
+			}
+		}
+
+		// Test Empty to a capacity
+		{
+			TRingBuffer<uint32> Q(16);
+			CHECK_MESSAGE(TEXT("Test EmptyToCapacity - InitialCapacity"), 16 == Q.Max());
+			Q.Empty(8);
+			CHECK_MESSAGE(TEXT("Test EmptyToCapacity - Lower"), 8 == Q.Max());
+			Q.Empty(32);
+			CHECK_MESSAGE(TEXT("Test EmptyToCapacity - Higher"), 32 == Q.Max());
+		}
+
+		// Test Different Add constructors
+		{
+			auto Clear = []()
+			{
+				Counter::Clear();
+			};
+			auto TestCounts = [this](const TCHAR* Message, int32 NumVoid, int32 NumCopy, int32 NumMove, int32 NumDestruct)
+			{
+				CHECK_MESSAGE(Message, (NumVoid == Counter::NumVoid && NumCopy == Counter::NumCopy && NumMove == Counter::NumMove && NumDestruct == Counter::NumDestruct) == true);
+			};
+
+			Clear();
+			{
+				TRingBuffer<Counter> QEmpty(4);
+				QEmpty.Reserve(8);
+				QEmpty.Empty();
+				TRingBuffer<Counter> QEmpty2(4);
+			}
+			TestCounts(TEXT("Test Add Constructors - Unallocated elements call no constructors/destructors"), 0, 0, 0, 0);
+			{
+				TRingBuffer<Counter> QEmpty(4);
+				QEmpty.Emplace();
+				QEmpty.Pop();
+				Clear();
+			}
+			TestCounts(TEXT("Test Add Constructors - Already removed element calls no destructors"), 0, 0, 0, 0);
+
+
+			uint32 MarkerValue = 0x54321;
+			Counter CounterA(MarkerValue);
+
+			TRingBuffer<Counter> Q(4);
+			Clear();
+			for (int Direction = 0; Direction < 2; ++Direction)
+			{
+				bool bAddBack = Direction == 0;
+				auto TestDirCounts = [this, bAddBack, &TestCounts, &Q, &Clear, MarkerValue](const TCHAR* Message, int32 NumVoid, int32 NumCopy, int32 NumMove, int32 NumDestruct, bool bWasInitialized = true)
+				{
+					const TCHAR* DirectionText = bAddBack ? TEXT("Back") : TEXT("Front");
+					bool bElementExists = Q.Num() == 1;
+					CHECK_MESSAGE(*FString::Printf(TEXT("Test Add Constructors - %s%s ElementExists"), Message, DirectionText), bElementExists == true);
+					if (bWasInitialized && bElementExists)
+					{
+						CHECK_MESSAGE(*FString::Printf(TEXT("Test Add Constructors - %s%s ValueEquals"), Message, DirectionText),
+							(Q.First().Value == MarkerValue) == true);
+					}
+					Q.PopFront();
+					TestCounts(*FString::Printf(TEXT("Test Add Constructors - %s%s CountsEqual"), Message, DirectionText), NumVoid, NumCopy, NumMove, NumDestruct);
+					Clear();
+				};
+
+				if (bAddBack) Q.Add(CounterA); else Q.AddFront(CounterA);
+				TestDirCounts(TEXT("Copy Add"), 0, 1, 0, 1);
+				if (bAddBack) Q.Add_GetRef(CounterA); else Q.AddFront_GetRef(CounterA);
+				TestDirCounts(TEXT("Copy GetRef Add"), 0, 1, 0, 1);
+				if (bAddBack) Q.Add(MoveTemp(CounterA)); else Q.AddFront(MoveTemp(CounterA));
+				TestDirCounts(TEXT("Move Add"), 0, 0, 1, 1);
+				if (bAddBack) Q.Add_GetRef(MoveTemp(CounterA)); else Q.AddFront_GetRef(MoveTemp(CounterA));
+				TestDirCounts(TEXT("Move GetRef Add"), 0, 0, 1, 1);
+				if (bAddBack) Q.Emplace(MarkerValue); else Q.EmplaceFront(MarkerValue);
+				TestDirCounts(TEXT("Emplace"), 1, 0, 0, 1);
+				if (bAddBack) Q.Emplace_GetRef(MarkerValue); else Q.EmplaceFront_GetRef(MarkerValue);
+				TestDirCounts(TEXT("GetRef Emplace"), 1, 0, 0, 1);
+				if (bAddBack) Q.AddUninitialized(); else Q.AddFrontUninitialized();
+				TestDirCounts(TEXT("Uninitialized Add"), 0, 0, 0, 1, false);
+				if (bAddBack) Q.AddUninitialized_GetRef(); else Q.AddFrontUninitialized_GetRef();
+				TestDirCounts(TEXT("Uninitialized GetRef Add"), 0, 0, 0, 1, false);
+			}
+		}
+
+		TestShiftIndex<uint32>();
+		TestShiftIndex<Counter>();
+
+		// Test RemoveAt
+		{
+			{
+				TRingBuffer<uint32> Q{ 0,1,2,3,4,5,6,7 };
+				Q.RemoveAt(2);
+				CHECK_MESSAGE(TEXT("Test RemoveAt Front Closest"), TRingBuffer<uint32>({ 0,1,3,4,5,6,7 }) == Q);
+			}
+			{
+				TRingBuffer<uint32> Q{ 0,1,2,3,4,5,6,7 };
+				Q.RemoveAt(5);
+				CHECK_MESSAGE(TEXT("Test RemoveAt Back Closest"), TRingBuffer<uint32>({ 0,1,2,3,4,6,7 }) == Q);
+			}
+			{
+				TRingBuffer<uint32> Q{ 0,1,2,3,4,5,6,7 };
+				int32 Offset = 4;
+				Q.Front += Offset;
+				Q.AfterBack += Offset;
+				//Now equal to: TRingBuffer<uint32> Q{ 4,5,6,7,0,1,2,3 };
+				Q.RemoveAt(2);
+				CHECK_MESSAGE(TEXT("Test RemoveAt Front Closest With Offset"), TRingBuffer<uint32>({ 4,5,7,0,1,2,3 }) == Q);
+			}
+			{
+				TRingBuffer<uint32> Q{ 0,1,2,3,4,5,6,7 };
+				int32 Offset = 4;
+				Q.Front += Offset;
+				Q.AfterBack += Offset;
+				//Now equal to: TRingBuffer<uint32> Q{ 4,5,6,7,0,1,2,3 };
+				Q.RemoveAt(5);
+				CHECK_MESSAGE(TEXT("Test RemoveAt Back Closest With Offset"), TRingBuffer<uint32>({ 4,5,6,7,0,2,3 }) == Q);
+			}
+		}
+
+		// Test Iteration
+		{
+			{
+				TRingBuffer<uint32> Q{ 0,1,2,3,4,5,6,7 };
+				uint32 Counter = 0;
+				for (uint32 Value : Q)
+				{
+					CHECK_MESSAGE(TEXT("Test Iteration - Value"), Counter++ == Value);
+				}
+				CHECK_MESSAGE(TEXT("Test Iteration - Num"), Counter == 8);
+			}
+			{
+				TRingBuffer<uint32> Q{ 4,5,6,7,0,1,2,3 };
+				int32 Offset = 4;
+				Q.Front += Offset;
+				Q.AfterBack += Offset;
+				// Now equal to 0,1,2,3,4,5,6,7
+				uint32 Counter = 0;
+				for (uint32 Value : Q)
+				{
+					CHECK_MESSAGE(TEXT("Test Iteration with Offset - Value"), Counter++ == Value);
+				}
+				CHECK_MESSAGE(TEXT("Test Iteration with Offset  - Num"), Counter == 8);
+			}
+		}
+
+		// Test ConvertPointerToIndex
+		{
+			{
+				TRingBuffer<uint32> Q{ 4,5,6,7,0,1,2,3 };
+				int32 Offset = 4;
+				Q.Front += Offset;
+				Q.AfterBack += Offset;
+				// Now equal to 0,1,2,3,4,5,6,7
+				CHECK_MESSAGE(TEXT("Test ConvertPointerToIndex - before array"), Q.ConvertPointerToIndex(&Q[0] - 100) == INDEX_NONE);
+				CHECK_MESSAGE(TEXT("Test ConvertPointerToIndex - after array"), Q.ConvertPointerToIndex(&Q[0] + 100) == INDEX_NONE);
+				for (int32 It = 0; It < 8; ++It)
+				{
+					CHECK_MESSAGE(TEXT("Test ConvertPointerToIndex - Values"), Q.ConvertPointerToIndex(&Q[It]) == It);
+				}
+			}
+
+			{
+				TRingBuffer<uint32> Q(16);
+				for (int32 It = 7; It >= 0; --It)
+				{
+					Q.AddFront(It);
+				}
+				Q.Pop();
+				// 8 Invalids, followed by 0,1,2,3,4,5,6, followed by Invalid
+				for (int32 It = 0; It < 7; ++It)
+				{
+					CHECK_MESSAGE(TEXT("Test ConvertPointerToIndex - Cap - Values"), Q.ConvertPointerToIndex(&Q[It]) == It);
+				}
+				CHECK_MESSAGE(TEXT("Test ConvertPointerToIndex - Cap - After End"), Q.ConvertPointerToIndex(&Q[6] + 1) == INDEX_NONE);
+				CHECK_MESSAGE(TEXT("Test ConvertPointerToIndex - Cap - Before Start"), Q.ConvertPointerToIndex(&Q[0] - 1) == INDEX_NONE);
+			}
+		}
+
+		// Test that setting Front to its maximum value and then popping the maximum number of elements does not break the contract that Front < capacity in StorageModulo space
+		{
+			TRingBuffer<uint32> Q(8);
+			Q.AddFront(0);
+			for (uint32 It = 1; It < 8; ++It)
+			{
+				Q.Add(It);
+			}
+			CHECK_MESSAGE(TEXT("Test Front<Capacity - Setup"), ((Q.Front & Q.IndexMask) == Q.IndexMask && Q.Num() == Q.Max()) == true);
+			Q.PopFront(8);
+			CHECK_MESSAGE(TEXT("Test Front<Capacity - Contract is true"), (static_cast<uint32>(Q.Front) < static_cast<uint32>(Q.Max())) == true);
+		}
+
+		// Test IsValidIndex
+		{
+			TRingBuffer<uint32> Q({ 0,1,2,3,4 });
+			for (int32 It = 0; It < Q.Num(); ++It)
+			{
+				CHECK_MESSAGE(TEXT("IsValidIndex - InRange"), Q.IsValidIndex(It) == true);
+			}
+			CHECK_MESSAGE(TEXT("IsValidIndex - Negative"), Q.IsValidIndex(-1) == false);
+			CHECK_MESSAGE(TEXT("IsValidIndex - Num()"), Q.IsValidIndex(Q.Num() + 1) == false);
+			CHECK_MESSAGE(TEXT("IsValidIndex - Capacity"), Q.IsValidIndex(Q.Max()) == false);
+			CHECK_MESSAGE(TEXT("IsValidIndex - Capacity + 1"), Q.IsValidIndex(Q.Max() + 1) == false);
+		}
+
+		// Test Compact
+		{
+			{
+				TRingBuffer<uint32> QEmpty;
+				CHECK_MESSAGE(TEXT("Compact - Empty zero capacity"), QEmpty.Compact().Num() == 0);
+				QEmpty.Add(1);
+				QEmpty.PopFront();
+				CHECK_MESSAGE(TEXT("Compact - Empty non-zero capacity"), QEmpty.Compact().Num() == 0);
+			}
+			{
+				TArrayView<uint32> View;
+				TRingBuffer<uint32> Q(8);
+				Q.AddFront(37);
+				View = Q.Compact();
+				CHECK_MESSAGE(TEXT("Compact - Front at end"), ArrayViewsEqual(View, TArrayView<const uint32>({ 37 })) == true);
+			}
+			{
+				TArrayView<uint32> View;
+				TRingBuffer<uint32> Q(8);
+				for (TRingBuffer<uint32>::IndexType It = 0; It < 6; ++It)
+				{
+					Q.Add(It);
+				}
+				Q.PopFront();
+				TRingBuffer<uint32>::StorageModuloType SavedFront = Q.Front;
+				CHECK_MESSAGE(TEXT("Compact - Front in middle - setup"), (SavedFront > 0) == true);
+				View = Q.Compact();
+				CHECK_MESSAGE(TEXT("Compact - Front in middle - values"), ArrayViewsEqual(View, TArrayView<const uint32>({ 1,2,3,4,5 })) == true);
+				CHECK_MESSAGE(TEXT("Compact - Front in middle - no reallocate"), (Q.Front == SavedFront) == true);
+			}
+			{
+				TArrayView<uint32> View;
+				TRingBuffer<uint32> Q(8);
+				for (TRingBuffer<uint32>::IndexType It = 1; It < 8; ++It)
+				{
+					Q.Add(It);
+				}
+				Q.AddFront(0);
+				CHECK_MESSAGE(TEXT("Compact - Full array front at end - setup"), ((Q.Front & Q.IndexMask) == 7) == true);
+				View = Q.Compact();
+				CHECK_MESSAGE(TEXT("Compact - Full array front at end - values"),
+					ArrayViewsEqual(View, TArrayView<const uint32>({ 0,1,2,3,4,5,6,7 })) == true);
+				CHECK_MESSAGE(TEXT("Compact - Full array front at end - reallocated"), (Q.Front == 0) == true);
+			}
+			{
+				TArrayView<uint32> View;
+				TRingBuffer<uint32> Q(8);
+				for (TRingBuffer<uint32>::IndexType It = 0; It < 8; ++It)
+				{
+					Q.Add(It);
+				}
+				uint32* SavedData = Q.AllocationData;
+				CHECK_MESSAGE(TEXT("Compact - Full array front at start - setup"), (Q.Front == 0) == true);
+				View = Q.Compact();
+				CHECK_MESSAGE(TEXT("Compact - Full array front at start - values"),
+					ArrayViewsEqual(View, TArrayView<const uint32>({ 0,1,2,3,4,5,6,7 })) == true);
+				CHECK_MESSAGE(TEXT("Compact - Full array front at start - no reallocate"), (Q.AllocationData == SavedData) == true);
+			}
+		}
+
+		// Test Remove
+		{
+			Counter Value;
+			{
+				TRingBuffer<Counter> Q;
+				Value.Value = 2;
+				Counter::Clear();
+				CHECK_MESSAGE(TEXT("Remove - empty"), Q.Remove(Value) == 0);
+				CHECK_MESSAGE(TEXT("Remove - empty - destructor count"), Counter::NumDestruct == 0);
+			}
+			{
+				TRingBuffer<Counter> Q({ 0,1,2,3,4 });
+				Value.Value = 5;
+				Counter::Clear();
+				CHECK_MESSAGE(TEXT("Remove - no hits"), Q.Remove(Value) == 0);
+				CHECK_MESSAGE(TEXT("Remove - no hits - destructor count"), Counter::NumDestruct == 0);
+				Q.Add(5);
+				CHECK_MESSAGE(TEXT("Remove - no hits - values"), (Q == TRingBuffer<Counter>({ 0,1,2,3,4,5 })) == true);
+			}
+			{
+				TRingBuffer<Counter> Q({ 1,2,3,4 });
+				Q.AddFront(0);
+				Value.Value = 0;
+				Counter::Clear();
+				CHECK_MESSAGE(TEXT("Remove - one element at front - num"), Q.Remove(Value) == 1);
+				CHECK_MESSAGE(TEXT("Remove - one element at front - destructor count"), Counter::NumDestruct == 5);
+				Q.Add(5);
+				CHECK_MESSAGE(TEXT("Remove - one element at front - values"), (Q == TRingBuffer<Counter>({ 1,2,3,4,5 })) == true);
+			}
+			{
+				TRingBuffer<Counter> Q({ 0,1,2,3,4 });
+				Value.Value = 2;
+				Counter::Clear();
+				CHECK_MESSAGE(TEXT("Remove - one element in mid - num"), Q.Remove(Value) == 1);
+				CHECK_MESSAGE(TEXT("Remove - one element in mid - destructor count"), Counter::NumDestruct == 3);
+				Q.Add(5);
+				CHECK_MESSAGE(TEXT("Remove - one element in mid - values"), (Q == TRingBuffer<Counter>({ 0,1,3,4,5 })) == true);
+			}
+			{
+				TRingBuffer<Counter> Q({ 1,2,3,4 });
+				Q.AddFront(0);
+				Value.Value = 2;
+				Counter::Clear();
+				CHECK_MESSAGE(TEXT("Remove - one element in mid - front at end"), Q.Remove(Value) == 1);
+				CHECK_MESSAGE(TEXT("Remove - one element in mid - front at end - destructor count"), Counter::NumDestruct == 3);
+				Q.Add(5);
+				CHECK_MESSAGE(TEXT("Remove - one element in mid - front at end - values"), (Q == TRingBuffer<Counter>({ 0,1,3,4,5 })) == true);
+			}
+			{
+				TRingBuffer<Counter> Q({ 0,1,2,3,4 });
+				Value.Value = 4;
+				Counter::Clear();
+				CHECK_MESSAGE(TEXT("Remove - one element - element at end - num"), Q.Remove(Value) == 1);
+				CHECK_MESSAGE(TEXT("Remove - one element - element at end - destructor count"), Counter::NumDestruct == 1);
+				Q.Add(5);
+				CHECK_MESSAGE(TEXT("Remove - one element - element at end - values"), (Q == TRingBuffer<Counter>({ 0,1,2,3,5 })) == true);
+			}
+			{
+				TRingBuffer<Counter> Q({ 1,2,3,4 });
+				Q.AddFront(4);
+				Value.Value = 4;
+				Counter::Clear();
+				CHECK_MESSAGE(TEXT("Remove - one element at front one at end - num"), Q.Remove(Value) == 2);
+				CHECK_MESSAGE(TEXT("Remove - one element at front one at end - destructor count"), Counter::NumDestruct == 5);
+				Q.Add(5);
+				CHECK_MESSAGE(TEXT("Remove - one element at front one at end - values"), (Q == TRingBuffer<Counter>({ 1,2,3,5 })) == true);
+			}
+			{
+				TRingBuffer<Counter> Q({ 1,2,3,4 });
+				Q.AddFront(1);
+				Value.Value = 1;
+				Counter::Clear();
+				CHECK_MESSAGE(TEXT("Remove - two elements - front at end - num"), Q.Remove(Value) == 2);
+				CHECK_MESSAGE(TEXT("Remove - two elements - front at end - destructor count"), Counter::NumDestruct == 5);
+				Q.Add(5);
+				CHECK_MESSAGE(TEXT("Remove - two elements - front at end - values"), (Q == TRingBuffer<Counter>({ 2,3,4,5 })) == true);
+			}
+		}
+
+		return true;
+	}
+
+	template <typename T>
+	void TestShiftIndex()
+	{
+		// Test shifts at specific points
+		{
+			{
+				TRingBuffer<T> Q{ 0,1,2,3,4,5,6,7 };
+				Q.ShiftIndexToFront(5);
+				CHECK_MESSAGE(TEXT("ShiftIndexToFront"), TRingBuffer<T>({ 5,0,1,2,3,4,6,7 }) == Q);
+				Q.ShiftIndexToBack(3);
+				CHECK_MESSAGE(TEXT("ShiftIndexToBack"), TRingBuffer<T>({ 5,0,1,3,4,6,7,2 }) == Q);
+			}
+
+			{
+				TRingBuffer<T> Q{ 0,1,2,3,4,5,6,7 };
+				int32 Offset = 4;
+				Q.Front += Offset;
+				Q.AfterBack += Offset;
+				//Now equal to: TRingBuffer<uint32> Q{ 4,5,6,7,0,1,2,3 };
+				Q.ShiftIndexToFront(5);
+				CHECK_MESSAGE(TEXT("ShiftIndexToFront With Offset"), TRingBuffer<T>({ 1,4,5,6,7,0,2,3 }) == Q);
+				Q.ShiftIndexToBack(3);
+				CHECK_MESSAGE(TEXT("ShiftIndexToBack With Offset"), TRingBuffer<T>({ 1,4,5,7,0,2,3,6 }) == Q);
+			}
+
+			{
+				TRingBuffer<T> Q{ 0,1,2,3,4,5,6,7,8 };
+				CHECK_MESSAGE(TEXT("ShiftIndexToFront Cap - Capacity"), Q.Max() == 16);
+				Q.ShiftIndexToFront(5);
+				CHECK_MESSAGE(TEXT("ShiftIndexToFront Cap"), TRingBuffer<T>({ 5,0,1,2,3,4,6,7,8 }) == Q);
+				Q.ShiftIndexToBack(3);
+				CHECK_MESSAGE(TEXT("ShiftIndexToBack Cap"), TRingBuffer<T>({ 5,0,1,3,4,6,7,8,2 }) == Q);
+			}
+
+			{
+				TRingBuffer<T> Q(16);
+				for (int32 It = 7; It >= 0; --It)
+				{
+					Q.AddFront(It);
+				}
+				Q.Pop();
+				// 8 Invalids, followed by 0,1,2,3,4,5,6, followed by Invalid
+				Q.ShiftIndexToFront(5);
+				CHECK_MESSAGE(TEXT("ShiftIndexToFront Cap With Offset"), TRingBuffer<T>({ 5,0,1,2,3,4,6 }) == Q);
+				Q.ShiftIndexToBack(3);
+				CHECK_MESSAGE(TEXT("ShiftIndexToBack Cap With Offset"), TRingBuffer<T>({ 5,0,1,3,4,6,2 }) == Q);
+			}
+
+			{
+				TRingBuffer<T> Q(16);
+				for (int32 It = 7; It >= 0; --It)
+				{
+					Q.AddFront(It);
+				}
+				Q.Add(8);
+				// 8, (AfterBack), followed by 7 Invalids, followed by (Start) 0,1,2,3,4,5,6,7
+				Q.ShiftIndexToFront(8);
+				CHECK_MESSAGE(TEXT("ShiftIndexToFront Cap With Wrapped"), TRingBuffer<T>({ 8,0,1,2,3,4,5,6,7 }) == Q);
+				Q.ShiftIndexToBack(0);
+				CHECK_MESSAGE(TEXT("ShiftIndexToBack Cap With Wrapped"), TRingBuffer<T>({ 0,1,2,3,4,5,6,7,8 }) == Q);
+			}
+		}
+
+		// Test ShiftIndex of each possible index
+		{
+			int32 Count = 8;
+			for (int32 It = 0; It < Count; ++It)
+			{
+				TRingBuffer<T> Q({ 0,1,2,3,4,5,6,7 });
+				Q.ShiftIndexToBack(It);
+				int32 CheckIndex = 0;
+				for (; CheckIndex < It; ++CheckIndex)
+				{
+					CHECK_MESSAGE(*FString::Printf(TEXT("ShiftIndexToBack Entire Array Values (%d,%d)"), It, CheckIndex), CheckIndex == Q[CheckIndex]);
+				}
+				for (; CheckIndex < Count - 1; ++CheckIndex)
+				{
+					CHECK_MESSAGE(*FString::Printf(TEXT("ShiftIndexToBack Entire Array Values (%d,%d)"), It, CheckIndex), CheckIndex + 1 == Q[CheckIndex]);
+				}
+				CHECK_MESSAGE(*FString::Printf(TEXT("ShiftIndexToBack Entire Array Values (%d,%d)"), It, Count - 1), It == Q[Count - 1]);
+			}
+			for (int32 It = 0; It < Count; ++It)
+			{
+				TRingBuffer<T> Q({ 0,1,2,3,4,5,6,7 });
+				Q.ShiftIndexToFront(It);
+
+				CHECK_MESSAGE(*FString::Printf(TEXT("ShiftIndexToFront Entire Array Values (%d,%d)"), It, 0), It == Q[0]);
+				int32 CheckIndex = 1;
+				for (; CheckIndex <= It; ++CheckIndex)
+				{
+					CHECK_MESSAGE(*FString::Printf(TEXT("ShiftIndexToFront Entire Array Values (%d,%d)"), It, CheckIndex), CheckIndex - 1 == Q[CheckIndex]);
+				}
+				for (; CheckIndex < Count; ++CheckIndex)
+				{
+					CHECK_MESSAGE(*FString::Printf(TEXT("ShiftIndexToFront Entire Array Values (%d,%d)"), It, CheckIndex), CheckIndex == Q[CheckIndex]);
+				}
+			}
+		}
+	}
+
+	template <typename T, typename U>
+	static bool ArrayViewsEqual(const TArrayView<T>& A, const TArrayView<U>& B)
+	{
+		int32 Num = A.Num();
+		if (Num != B.Num())
+		{
+			return false;
+		}
+		for (int It = 0; It < Num; ++It)
+		{
+			if (A[It] != B[It])
+			{
+				return false;
+			}
+		}
+		return true;
+	}
+};
+
+int FRingBufferTest::Counter::NumVoid = 0;
+int FRingBufferTest::Counter::NumCopy = 0;
+int FRingBufferTest::Counter::NumMove = 0;
+int FRingBufferTest::Counter::NumDestruct = 0;
+
+TEST_CASE_NAMED(FRingBufferTestSubClass, "System::Core::Containers::RingBuffer", "[ApplicationContextMask][EngineFilter]")
+{
+	FRingBufferTest instance;
+	instance.Test();
+}
+
+#endif // #if WITH_TESTS
\ No newline at end of file
