Implement REQUIRE_MESSAGE macro

diff --git a/Engine/Source/Developer/LowLevelTestsRunner/Public/TestHarness.h b/Engine/Source/Developer/LowLevelTestsRunner/Public/TestHarness.h
--- a/Engine/Source/Developer/LowLevelTestsRunner/Public/TestHarness.h
+++ b/Engine/Source/Developer/LowLevelTestsRunner/Public/TestHarness.h
@@ -142,6 +142,11 @@ struct Catch::StringMaker<TTuple<KeyType, ValueType>>
 #define REQUIRE_NOT_EQUAL(Actual, Expected)\
 	REQUIRE(Actual != Expected)
 
+#define REQUIRE_MESSAGE(What, Value) do { \
+	INFO(What); \
+	REQUIRE((Value)==true); \
+} while (false)
+
 #define TEST_TRUE(What, Value)\
 	INFO(What);\
 	CHECK((Value)==true)
diff --git a/Engine/Source/Runtime/Core/Public/Misc/LowLevelTestAdapter.h b/Engine/Source/Runtime/Core/Public/Misc/LowLevelTestAdapter.h
--- a/Engine/Source/Runtime/Core/Public/Misc/LowLevelTestAdapter.h
+++ b/Engine/Source/Runtime/Core/Public/Misc/LowLevelTestAdapter.h
@@ -72,6 +72,7 @@
 #define CHECK_FALSE(Expr) if (Expr) { FAutomationTestFramework::Get().GetCurrentTest()->AddError(TEXT("Condition expected to return false but returned true")); }
 #define CHECK_MESSAGE(Message, Expr) if (!(Expr)) { FAutomationTestFramework::Get().GetCurrentTest()->AddError(Message); }
 #define REQUIRE(Expr) if (!(Expr)) { FAutomationTestFramework::Get().GetCurrentTest()->AddError(TEXT("Required condition failed, interrupting test")); return; }
+#define REQUIRE_MESSAGE(Message, Expr) if (!(Expr)) { FAutomationTestFramework::Get().GetCurrentTest()->AddError(Message); return; }
 
 #define SECTION(Text) AddInfo(TEXT(Text));
 
