Fixed Clang static analysis false positives in our operator delete overrides by disabling them.

diff --git a/Engine/Source/Runtime/Core/Public/Modules/Boilerplate/ModuleBoilerplate.h b/Engine/Source/Runtime/Core/Public/Modules/Boilerplate/ModuleBoilerplate.h
--- a/Engine/Source/Runtime/Core/Public/Modules/Boilerplate/ModuleBoilerplate.h
+++ b/Engine/Source/Runtime/Core/Public/Modules/Boilerplate/ModuleBoilerplate.h
@@ -28,7 +28,9 @@ class FChunkedFixedUObjectArray;
 	#define OPERATOR_NEW_MSVC_PRAGMA
 #endif
 
-#if !FORCE_ANSI_ALLOCATOR
+// Disable the replacement new/delete when running the Clang static analyzer, due to false positives in 15.0.x:
+// [REDACTED_PROFILE_URL]/llvm-project/issues/58820
+#if !FORCE_ANSI_ALLOCATOR && !defined(__clang_analyzer__)
 static_assert(__STDCPP_DEFAULT_NEW_ALIGNMENT__ <= 16, "Expecting 16-byte default operator new alignment - alignments > 16 may have bloat");
 #define REPLACEMENT_OPERATOR_NEW_AND_DELETE \
 	OPERATOR_NEW_MSVC_PRAGMA void* operator new  ( size_t Size                                                    ) OPERATOR_NEW_THROW_SPEC      { return FMemory::Malloc( Size ? Size : 1, __STDCPP_DEFAULT_NEW_ALIGNMENT__ ); } \
