Removed the now-unused UE_CONSTRAINT* macros.

diff --git a/Engine/Source/Runtime/Core/Public/Templates/UnrealTypeTraits.h b/Engine/Source/Runtime/Core/Public/Templates/UnrealTypeTraits.h
--- a/Engine/Source/Runtime/Core/Public/Templates/UnrealTypeTraits.h
+++ b/Engine/Source/Runtime/Core/Public/Templates/UnrealTypeTraits.h
@@ -33,22 +33,14 @@
 	void IntegralUpTo32Bit(T Lhs, U Rhs) {}
  -----------------------------------------------------------------------------*/
 #if __cplusplus < 202000
-    #define UE_CONSTRAINTS_BEGIN , std::enable_if_t<
-    #define UE_CONSTRAINT(...) (__VA_ARGS__) &&
-    #define UE_CONSTRAINTS_END true, int> = 0
-
 	#define UE_REQUIRES(...) , std::enable_if_t<(__VA_ARGS__), int> = 0
 #else
-    namespace UE::Core::Private
-    {
-        // Only needed for the UE_REQUIRES macro to work, to allow for a trailing > token after the macro
-        template <bool B>
-        concept BoolIdentityConcept = B;
-    }
-
-    #define UE_CONSTRAINTS_BEGIN > requires
-    #define UE_CONSTRAINT(...) (__VA_ARGS__) &&
-    #define UE_CONSTRAINTS_END UE::Core::Private::BoolIdentityConcept<true
+	namespace UE::Core::Private
+	{
+		// Only needed for the UE_REQUIRES macro to work, to allow for a trailing > token after the macro
+		template <bool B>
+		concept BoolIdentityConcept = B;
+	}
 
 	#define UE_REQUIRES(...) > requires (__VA_ARGS__) && UE::Core::Private::BoolIdentityConcept<true
 #endif
