Fixed false positive static analysis warnings (reccuring check) for the CHECK/REQUIRE_MESSAGE macros

diff --git a/Engine/Source/Developer/LowLevelTestsRunner/Public/TestHarness.h b/Engine/Source/Developer/LowLevelTestsRunner/Public/TestHarness.h
--- a/Engine/Source/Developer/LowLevelTestsRunner/Public/TestHarness.h
+++ b/Engine/Source/Developer/LowLevelTestsRunner/Public/TestHarness.h
@@ -131,6 +131,7 @@ struct Catch::StringMaker<TTuple<KeyType, ValueType>>
 #define CHECK_NOT_EQUAL(Actual, Expected)\
 	CHECK(Actual != Expected)
 
+//-V:CHECK_MESSAGE:571
 #define CHECK_MESSAGE(What, Value) do { \
 	INFO(What); \
 	CHECK((Value)==true); \
@@ -142,6 +143,7 @@ struct Catch::StringMaker<TTuple<KeyType, ValueType>>
 #define REQUIRE_NOT_EQUAL(Actual, Expected)\
 	REQUIRE(Actual != Expected)
 
+//-V:REQUIRE_MESSAGE:571
 #define REQUIRE_MESSAGE(What, Value) do { \
 	INFO(What); \
 	REQUIRE((Value)==true); \
diff --git a/Engine/Source/Runtime/Core/Public/Misc/LowLevelTestAdapter.h b/Engine/Source/Runtime/Core/Public/Misc/LowLevelTestAdapter.h
--- a/Engine/Source/Runtime/Core/Public/Misc/LowLevelTestAdapter.h
+++ b/Engine/Source/Runtime/Core/Public/Misc/LowLevelTestAdapter.h
@@ -70,8 +70,10 @@
 #define INFO(What)
 #define CHECK(Expr) if (!(Expr)) { FAutomationTestFramework::Get().GetCurrentTest()->AddError(TEXT("Condition failed")); }
 #define CHECK_FALSE(Expr) if (Expr) { FAutomationTestFramework::Get().GetCurrentTest()->AddError(TEXT("Condition expected to return false but returned true")); }
+//-V:CHECK_MESSAGE:571
 #define CHECK_MESSAGE(Message, Expr) if (!(Expr)) { FAutomationTestFramework::Get().GetCurrentTest()->AddError(Message); }
 #define REQUIRE(Expr) if (!(Expr)) { FAutomationTestFramework::Get().GetCurrentTest()->AddError(TEXT("Required condition failed, interrupting test")); return; }
+//-V:REQUIRE_MESSAGE:571
 #define REQUIRE_MESSAGE(Message, Expr) if (!(Expr)) { FAutomationTestFramework::Get().GetCurrentTest()->AddError(Message); return; }
 
 #define SECTION(Text) AddInfo(TEXT(Text));
diff --git a/Engine/Source/Runtime/Core/Tests/Containers/BitArrayTests.cpp b/Engine/Source/Runtime/Core/Tests/Containers/BitArrayTests.cpp
--- a/Engine/Source/Runtime/Core/Tests/Containers/BitArrayTests.cpp
+++ b/Engine/Source/Runtime/Core/Tests/Containers/BitArrayTests.cpp
@@ -1239,13 +1239,11 @@ public:
 					ZeroesText[DisplayIndex] = TEXT('\0');
 					if (!bOnesMatchesExpected)
 					{
-						//-V571
 						CHECK_MESSAGE(FString::Printf(TEXT("SetRange bool BitWidth=%d WriteOffset=%d, Ones\nExpected=%s\nActual  =%s"), BitWidth, WriteOffset, OnesExpectedText, OnesText),
 							bOnesMatchesExpected);
 					}
 					else if (!bZeroesMatchesExpected)
 					{
-						//-V571
 						CHECK_MESSAGE(FString::Printf(TEXT("SetRange bool BitWidth=%d WriteOffset=%d, Zeroes\nExpected=%s\nActual  =%s"), BitWidth, WriteOffset, ZeroesExpectedText, ZeroesText),
 							bZeroesMatchesExpected);
 					}
@@ -2193,13 +2191,11 @@ public:
 						ZeroesText[DisplayIndex] = TEXT('\0');
 						if (!bOnesMatchesExpected)
 						{
-							//-V571
 							CHECK_MESSAGE(FString::Printf(TEXT("MemmoveBitsWordOrder BitWidth=%d ReadOffset=%d, WriteOffset=%d, Overlap=%d, Ones\nExpected=%s\nActual  =%s"), BitWidth, ReadOffset, WriteOffset, Overlap, OnesExpectedText, OnesText),
 								bOnesMatchesExpected);
 						}
 						else if (!bZeroesMatchesExpected)
 						{
-							//-V571
 							CHECK_MESSAGE(FString::Printf(TEXT("MemmoveBitsWordOrder BitWidth=%d ReadOffset=%d, WriteOffset=%d, Overlap=%d, Zeroes\nExpected=%s\nActual  =%s"), BitWidth, ReadOffset, WriteOffset, Overlap, ZeroesExpectedText, ZeroesText),
 								bZeroesMatchesExpected);
 						}
