Logging: Fixed the static asserts about the format string to catch non-const format strings

diff --git a/Engine/Platforms/Hololens/Source/Runtime/Core/Private/HoloLensPlatformAtomics.cpp b/Engine/Platforms/Hololens/Source/Runtime/Core/Private/HoloLensPlatformAtomics.cpp
--- a/Engine/Platforms/Hololens/Source/Runtime/Core/Private/HoloLensPlatformAtomics.cpp
+++ b/Engine/Platforms/Hololens/Source/Runtime/Core/Private/HoloLensPlatformAtomics.cpp
@@ -5,17 +5,11 @@
 
 #if PLATFORM_HOLOLENS
 
-void FHoloLensAtomics::HandleAtomicsFailure( const TCHAR* InFormat, ... )
-{	
+void FHoloLensAtomics::HandleAtomicsFailure(const TCHAR* InFormat, ...)
+{
 	TCHAR TempStr[1024];
-	va_list Ptr;
-
-	va_start( Ptr, InFormat );	
-	FPlatformString::GetVarArgs( TempStr, UE_ARRAY_COUNT(TempStr), InFormat, Ptr );
-	va_end( Ptr );
-
-	UE_LOG(LogTemp, Log, TempStr);
-	check( 0 );
+	GET_VARARGS(TempStr, UE_ARRAY_COUNT(TempStr), UE_ARRAY_COUNT(TempStr) - 1, InFormat, InFormat);
+	UE_LOG(LogTemp, Fatal, TEXT("%s"), TempStr);
 }
 
 #endif //PLATFORM_HOLOLENS
\ No newline at end of file
diff --git a/Engine/Source/Runtime/Core/Private/Misc/CoreMisc.cpp b/Engine/Source/Runtime/Core/Private/Misc/CoreMisc.cpp
--- a/Engine/Source/Runtime/Core/Private/Misc/CoreMisc.cpp
+++ b/Engine/Source/Runtime/Core/Private/Misc/CoreMisc.cpp
@@ -531,7 +531,7 @@ namespace UEAsserts_Private
 		const int32 TempStrSize = 4096;
 		TCHAR TempStr[TempStrSize];
 		GET_VARARGS(TempStr, TempStrSize, TempStrSize - 1, FormattedMsg, FormattedMsg);
-		UE_LOG(LogCore, Error, TempStr);
+		UE_LOG(LogCore, Error, TEXT("%s"), TempStr);
 	}
 }
 #endif
diff --git a/Engine/Source/Runtime/Core/Private/Misc/MessageDialog.cpp b/Engine/Source/Runtime/Core/Private/Misc/MessageDialog.cpp
--- a/Engine/Source/Runtime/Core/Private/Misc/MessageDialog.cpp
+++ b/Engine/Source/Runtime/Core/Private/Misc/MessageDialog.cpp
@@ -54,17 +54,16 @@ void FMessageDialog::Debugf( const FText& Message, const FText* OptTitle )
 void FMessageDialog::ShowLastError()
 {
 	uint32 LastError = FPlatformMisc::GetLastError();
-
-	TCHAR TempStr[MAX_SPRINTF] = {};
 	TCHAR ErrorBuffer[1024];
-	FCString::Sprintf( TempStr, TEXT("GetLastError : %d\n\n%s"), LastError, FPlatformMisc::GetSystemErrorMessage(ErrorBuffer, 1024, 0) );
-	if( FApp::IsUnattended() == true )
+	if (FApp::IsUnattended())
 	{
-		UE_LOG(LogOutputDevice, Fatal, TempStr);
+		UE_LOG(LogOutputDevice, Fatal, TEXT("GetLastError : %d\n\n%s"), LastError, FPlatformMisc::GetSystemErrorMessage(ErrorBuffer, 1024, 0));
 	}
 	else
 	{
-		FPlatformMisc::MessageBoxExt( EAppMsgType::Ok, TempStr, *NSLOCTEXT("MessageDialog", "DefaultSystemErrorTitle", "System Error").ToString() );
+		TCHAR TempStr[MAX_SPRINTF] = {};
+		FCString::Sprintf(TempStr, TEXT("GetLastError : %d\n\n%s"), LastError, FPlatformMisc::GetSystemErrorMessage(ErrorBuffer, 1024, 0));
+		FPlatformMisc::MessageBoxExt(EAppMsgType::Ok, TempStr, *NSLOCTEXT("MessageDialog", "DefaultSystemErrorTitle", "System Error").ToString());
 	}
 }
 
diff --git a/Engine/Source/Runtime/Core/Private/Serialization/BitWriter.cpp b/Engine/Source/Runtime/Core/Private/Serialization/BitWriter.cpp
--- a/Engine/Source/Runtime/Core/Private/Serialization/BitWriter.cpp
+++ b/Engine/Source/Runtime/Core/Private/Serialization/BitWriter.cpp
@@ -136,7 +136,7 @@ void FBitWriter::SerializeInt(uint32& Value, uint32 ValueMax)
 
 	if (WriteValue >= ValueMax)
 	{
-		const auto& Msg = TEXT("FBitWriter::SerializeInt(): Value out of bounds (Value: %u, ValueMax: %u)");
+		constexpr static const auto& Msg = TEXT("FBitWriter::SerializeInt(): Value out of bounds (Value: %u, ValueMax: %u)");
 
 		UE_LOG(LogSerialization, Error, Msg, WriteValue, ValueMax);
 		ensureMsgf(false, Msg, WriteValue, ValueMax);
diff --git a/Engine/Source/Runtime/Core/Private/Windows/WindowsPlatformAtomics.cpp b/Engine/Source/Runtime/Core/Private/Windows/WindowsPlatformAtomics.cpp
--- a/Engine/Source/Runtime/Core/Private/Windows/WindowsPlatformAtomics.cpp
+++ b/Engine/Source/Runtime/Core/Private/Windows/WindowsPlatformAtomics.cpp
@@ -7,15 +7,9 @@
 #include "CoreGlobals.h"
 
 
-void FWindowsPlatformAtomics::HandleAtomicsFailure( const TCHAR* InFormat, ... )
-{	
+void FWindowsPlatformAtomics::HandleAtomicsFailure(const TCHAR* InFormat, ...)
+{
 	TCHAR TempStr[1024];
-	va_list Ptr;
-
-	va_start( Ptr, InFormat );	
-	FCString::GetVarArgs( TempStr, UE_ARRAY_COUNT(TempStr), InFormat, Ptr );
-	va_end( Ptr );
-
-	UE_LOG(LogWindows, Log,  TempStr );
-	check( 0 );
+	GET_VARARGS(TempStr, UE_ARRAY_COUNT(TempStr), UE_ARRAY_COUNT(TempStr) - 1, InFormat, InFormat);
+	UE_LOG(LogWindows, Fatal, TEXT("%s"), TempStr);
 }
diff --git a/Engine/Source/Runtime/Core/Public/Logging/LogMacros.h b/Engine/Source/Runtime/Core/Public/Logging/LogMacros.h
--- a/Engine/Source/Runtime/Core/Public/Logging/LogMacros.h
+++ b/Engine/Source/Runtime/Core/Public/Logging/LogMacros.h
@@ -39,7 +39,7 @@ struct CORE_API FMsg
 	template <typename FmtType, typename... Types>
 	static void SendNotificationStringf(const FmtType& Fmt, Types... Args)
 	{
-		static_assert(TIsArrayOrRefOfTypeByPredicate<FmtType, TIsCharEncodingCompatibleWithTCHAR>::Value, "Formatting string must be a TCHAR array.");
+		static_assert(TIsArrayOrRefOfTypeByPredicate<FmtType, TIsCharEncodingCompatibleWithTCHAR>::Value, "Formatting string must be a const TCHAR array.");
 		static_assert(TAnd<TIsValidVariadicFunctionArg<Types>...>::Value, "Invalid argument(s) passed to FMsg::SendNotificationStringf");
 
 		SendNotificationStringfImpl((const TCHAR*)Fmt, Args...);
@@ -49,7 +49,7 @@ struct CORE_API FMsg
 	template <typename FmtType, typename... Types>
 	static void Logf(const ANSICHAR* File, int32 Line, const FLogCategoryName& Category, ELogVerbosity::Type Verbosity, const FmtType& Fmt, Types... Args)
 	{
-		static_assert(TIsArrayOrRefOfTypeByPredicate<FmtType, TIsCharEncodingCompatibleWithTCHAR>::Value, "Formatting string must be a TCHAR array.");
+		static_assert(TIsArrayOrRefOfTypeByPredicate<FmtType, TIsCharEncodingCompatibleWithTCHAR>::Value, "Formatting string must be a const TCHAR array.");
 		static_assert(TAnd<TIsValidVariadicFunctionArg<Types>...>::Value, "Invalid argument(s) passed to FMsg::Logf");
 
 		LogfImpl(File, Line, Category, Verbosity, (const TCHAR*)Fmt, Args...);
@@ -59,7 +59,7 @@ struct CORE_API FMsg
 	template <typename FmtType, typename... Types>
 	static void Logf_Internal(const ANSICHAR* File, int32 Line, const FLogCategoryName& Category, ELogVerbosity::Type Verbosity, const FmtType& Fmt, Types... Args)
 	{
-		static_assert(TIsArrayOrRefOfTypeByPredicate<FmtType, TIsCharEncodingCompatibleWithTCHAR>::Value, "Formatting string must be a TCHAR array.");
+		static_assert(TIsArrayOrRefOfTypeByPredicate<FmtType, TIsCharEncodingCompatibleWithTCHAR>::Value, "Formatting string must be a const TCHAR array.");
 		static_assert(TAnd<TIsValidVariadicFunctionArg<Types>...>::Value, "Invalid argument(s) passed to FMsg::Logf_Internal");
 
 		Logf_InternalImpl(File, Line, Category, Verbosity, (const TCHAR*)Fmt, Args...);
@@ -198,6 +198,7 @@ private:
 	 * INTERNAL IMPLEMENTATION. DO NOT CALL DIRECTLY!
 	**/
 	#define UE_INTERNAL_LOG_IMPL(CategoryName, Verbosity, Format, ...) \
+		static_assert(std::is_const_v<std::remove_reference_t<decltype(Format)>>, "Formatting string must be a const TCHAR array."); \
 		static_assert(TIsArrayOrRefOfTypeByPredicate<decltype(Format), TIsCharEncodingCompatibleWithTCHAR>::Value, "Formatting string must be a TCHAR array."); \
 		static_assert((ELogVerbosity::Verbosity & ELogVerbosity::VerbosityMask) < ELogVerbosity::NumVerbosity && ELogVerbosity::Verbosity > 0, "Verbosity must be constant and in range."); \
 		UE_LOG_EXPAND_IS_FATAL(Verbosity, PREPROCESSOR_NOTHING, if (!CategoryName.IsSuppressed(ELogVerbosity::Verbosity))) \
@@ -256,6 +257,7 @@ private:
 	 ***/
 	#define UE_LOG_CLINKAGE(CategoryName, Verbosity, Format, ...) \
 	{ \
+		static_assert(std::is_const_v<std::remove_reference_t<decltype(Format)>>, "Formatting string must be a const TCHAR array."); \
 		static_assert(TIsArrayOrRefOfTypeByPredicate<decltype(Format), TIsCharEncodingCompatibleWithTCHAR>::Value, "Formatting string must be a TCHAR array."); \
 		static_assert((ELogVerbosity::Verbosity & ELogVerbosity::VerbosityMask) < ELogVerbosity::NumVerbosity && ELogVerbosity::Verbosity > 0, "Verbosity must be constant and in range."); \
 		CA_CONSTANT_IF((ELogVerbosity::Verbosity & ELogVerbosity::VerbosityMask) <= ELogVerbosity::COMPILED_IN_MINIMUM_VERBOSITY && (ELogVerbosity::Warning & ELogVerbosity::VerbosityMask) <= FLogCategory##CategoryName::CompileTimeVerbosity) \
@@ -300,6 +302,7 @@ private:
 	// Conditional logging. Will only log if Condition is met.
 	#define UE_CLOG(Condition, CategoryName, Verbosity, Format, ...) \
 	{ \
+		static_assert(std::is_const_v<std::remove_reference_t<decltype(Format)>>, "Formatting string must be a const TCHAR array."); \
 		static_assert(TIsArrayOrRefOfTypeByPredicate<decltype(Format), TIsCharEncodingCompatibleWithTCHAR>::Value, "Formatting string must be a TCHAR array."); \
 		static_assert((ELogVerbosity::Verbosity & ELogVerbosity::VerbosityMask) < ELogVerbosity::NumVerbosity && ELogVerbosity::Verbosity > 0, "Verbosity must be constant and in range."); \
 		CA_CONSTANT_IF((ELogVerbosity::Verbosity & ELogVerbosity::VerbosityMask) <= ELogVerbosity::COMPILED_IN_MINIMUM_VERBOSITY && (ELogVerbosity::Warning & ELogVerbosity::VerbosityMask) <= FLogCategory##CategoryName::CompileTimeVerbosity) \
diff --git a/Engine/Source/Runtime/Engine/Private/ShaderCompiler/ShaderCompiler.cpp b/Engine/Source/Runtime/Engine/Private/ShaderCompiler/ShaderCompiler.cpp
--- a/Engine/Source/Runtime/Engine/Private/ShaderCompiler/ShaderCompiler.cpp
+++ b/Engine/Source/Runtime/Engine/Private/ShaderCompiler/ShaderCompiler.cpp
@@ -3069,7 +3069,7 @@ FShaderCompilerStats* GShaderCompilerStats = nullptr;
 void FShaderCompilerStats::WriteStats(FOutputDevice* Ar)
 {
 #if ALLOW_DEBUG_FILES
-	static TCHAR DebugText[] = TEXT("Wrote shader compile stats to file '%s'.");
+	constexpr static const TCHAR DebugText[] = TEXT("Wrote shader compile stats to file '%s'.");
 	{
 		FlushRenderingCommands();
 
diff --git a/Engine/Source/Runtime/RHI/Private/RHIBreadcrumbs.cpp b/Engine/Source/Runtime/RHI/Private/RHIBreadcrumbs.cpp
--- a/Engine/Source/Runtime/RHI/Private/RHIBreadcrumbs.cpp
+++ b/Engine/Source/Runtime/RHI/Private/RHIBreadcrumbs.cpp
@@ -152,7 +152,7 @@ void FRHIBreadcrumbStack::WriteRenderBreadcrumbs(FCrashContextExtendedWriter& Wr
 	TCHAR StaticBreadcrumbName[MAX_BREADCRUMB_NAME_STRING];
 	FCString::Snprintf(StaticBreadcrumbName, MAX_BREADCRUMB_NAME_STRING, TEXT("Breadcrumbs_%s_%d"), ThreadName, BreadcrumbId++);
 	Writer.AddString(StaticBreadcrumbName, StaticBreadcrumbStackString);
-	UE_LOG(LogRHI, Error, StaticBreadcrumbStackString);
+	UE_LOG(LogRHI, Error, TEXT("%s"), StaticBreadcrumbStackString);
 }
 #endif
 
