Consistency: TEST_CASE_NAMED for Deque and PagedArray tests
#rnx

diff --git a/Engine/Source/Runtime/Core/Tests/Containers/DequeTest.cpp b/Engine/Source/Runtime/Core/Tests/Containers/DequeTest.cpp
--- a/Engine/Source/Runtime/Core/Tests/Containers/DequeTest.cpp
+++ b/Engine/Source/Runtime/Core/Tests/Containers/DequeTest.cpp
@@ -86,7 +86,7 @@ bool Test::EmplaceFirstPopLast(TDeque<int32Token>& Deque, int32 Count)
 // Unit tests
 //---------------------------------------------------------------------------------------------------------------------
 
-TEST_CASE("System::Core::Containers::TDeque::Reserve without data", "[SmokeFilter][Core][Containers][Deque]")
+TEST_CASE_NAMED(FDequeReserveWithoutDataTest, "System::Core::Containers::TDeque::Reserve without data", "[SmokeFilter][Core][Containers][Deque]")
 {
 	int32Token::Reset();
 	TDeque<int32Token> Deque;
@@ -100,7 +100,7 @@ TEST_CASE("System::Core::Containers::TDeque::Reserve without data", "[SmokeFilte
 	CHECK(int32Token::EvenConstructionDestructionCalls(0));
 }
 
-TEST_CASE("System::Core::Containers::TDeque::Reserve EmplaceLast single element", "[SmokeFilter][Core][Containers][Deque]")
+TEST_CASE_NAMED(FDequeReserveEmplaceLastSingleElementTest, "System::Core::Containers::TDeque::Reserve EmplaceLast single element", "[SmokeFilter][Core][Containers][Deque]")
 {
 	int32Token::Reset();
 	{
@@ -119,7 +119,7 @@ TEST_CASE("System::Core::Containers::TDeque::Reserve EmplaceLast single element"
 	CHECK(int32Token::EvenConstructionDestructionCalls(1));
 }
 
-TEST_CASE("System::Core::Containers::TDeque::Reset", "[SmokeFilter][Core][Containers][Deque]")
+TEST_CASE_NAMED(FDequeResetTest, "System::Core::Containers::TDeque::Reset", "[SmokeFilter][Core][Containers][Deque]")
 {
 	int32Token::Reset();
 	{
@@ -140,7 +140,7 @@ TEST_CASE("System::Core::Containers::TDeque::Reset", "[SmokeFilter][Core][Contai
 	CHECK(int32Token::EvenConstructionDestructionCalls(1));
 }
 
-TEST_CASE("System::Core::Containers::TDeque::Empty", "[SmokeFilter][Core][Containers][Deque]")
+TEST_CASE_NAMED(FDequeEmptyTest, "System::Core::Containers::TDeque::Empty", "[SmokeFilter][Core][Containers][Deque]")
 {
 	int32Token::Reset();
 	{
@@ -153,7 +153,7 @@ TEST_CASE("System::Core::Containers::TDeque::Empty", "[SmokeFilter][Core][Contai
 	CHECK(int32Token::EvenConstructionDestructionCalls(0));
 }
 
-TEST_CASE("System::Core::Containers::TDeque::Empty after single element EmplaceLast", "[SmokeFilter][Core][Containers][Deque]")
+TEST_CASE_NAMED(FDequeEmptyAfterSingleElementEmplaceLastTest, "System::Core::Containers::TDeque::Empty after single element EmplaceLast", "[SmokeFilter][Core][Containers][Deque]")
 {
 	int32Token::Reset();
 	{
@@ -170,7 +170,7 @@ TEST_CASE("System::Core::Containers::TDeque::Empty after single element EmplaceL
 	CHECK(int32Token::EvenConstructionDestructionCalls(1));
 }
 
-TEST_CASE("System::Core::Containers::TDeque::EmplaceLast single element", "[SmokeFilter][Core][Containers][Deque]")
+TEST_CASE_NAMED(FDequeEmplaceLastSingleElementTest, "System::Core::Containers::TDeque::EmplaceLast single element", "[SmokeFilter][Core][Containers][Deque]")
 {
 	int32Token::Reset();
 	{
@@ -186,7 +186,7 @@ TEST_CASE("System::Core::Containers::TDeque::EmplaceLast single element", "[Smok
 	CHECK(int32Token::EvenConstructionDestructionCalls(1));
 }
 
-TEST_CASE("System::Core::Containers::TDeque::EmplaceLast range to capacity", "[SmokeFilter][Core][Containers][Deque]")
+TEST_CASE_NAMED(FDequeEmplaceLastRangeToCapacityTest, "System::Core::Containers::TDeque::EmplaceLast range to capacity", "[SmokeFilter][Core][Containers][Deque]")
 {
 	int32Token::Reset();
 	{
@@ -205,7 +205,7 @@ TEST_CASE("System::Core::Containers::TDeque::EmplaceLast range to capacity", "[S
 	CHECK(int32Token::EvenConstructionDestructionCalls(Test::DefaultCapacity * 10));
 }
 
-TEST_CASE("System::Core::Containers::TDeque::EmplaceLast range past capacity", "[SmokeFilter][Core][Containers][Deque]")
+TEST_CASE_NAMED(FDequeEmplaceLastRangePastCapacityTest, "System::Core::Containers::TDeque::EmplaceLast range past capacity", "[SmokeFilter][Core][Containers][Deque]")
 {
 	int32Token::Reset();
 	{
@@ -227,7 +227,7 @@ TEST_CASE("System::Core::Containers::TDeque::EmplaceLast range past capacity", "
 	CHECK(int32Token::EvenConstructionDestructionCalls(Test::DefaultCapacity + 1));
 }
 
-TEST_CASE("System::Core::Containers::TDeque::EmplaceFirst single element", "[SmokeFilter][Core][Containers][Deque]")
+TEST_CASE_NAMED(FDequeEmplaceFirstSingleElementTest, "System::Core::Containers::TDeque::EmplaceFirst single element", "[SmokeFilter][Core][Containers][Deque]")
 {
 	int32Token::Reset();
 	{
@@ -243,7 +243,7 @@ TEST_CASE("System::Core::Containers::TDeque::EmplaceFirst single element", "[Smo
 	CHECK(int32Token::EvenConstructionDestructionCalls(1));
 }
 
-TEST_CASE("System::Core::Containers::TDeque::EmplaceFirst range to capacity", "[SmokeFilter][Core][Containers][Deque]")
+TEST_CASE_NAMED(FDequeEmplaceFirstRangeToCapacityTest, "System::Core::Containers::TDeque::EmplaceFirst range to capacity", "[SmokeFilter][Core][Containers][Deque]")
 {
 	int32Token::Reset();
 	{
@@ -262,7 +262,7 @@ TEST_CASE("System::Core::Containers::TDeque::EmplaceFirst range to capacity", "[
 	CHECK(int32Token::EvenConstructionDestructionCalls(Test::DefaultCapacity * 10));
 }
 
-TEST_CASE("System::Core::Containers::TDeque::PushLast single element (implicit move)", "[SmokeFilter][Core][Containers][Deque]")
+TEST_CASE_NAMED(FDequePushLastSingleElementTest, "System::Core::Containers::TDeque::PushLast single element (implicit move)", "[SmokeFilter][Core][Containers][Deque]")
 {
 	int32Token::Reset();
 	{
@@ -281,7 +281,7 @@ TEST_CASE("System::Core::Containers::TDeque::PushLast single element (implicit m
 	CHECK(int32Token::NumMoveConstructorCalls() == 1);
 }
 
-TEST_CASE("System::Core::Containers::TDeque::PushLast single element from move", "[SmokeFilter][Core][Containers][Deque]")
+TEST_CASE_NAMED(FDequePushLastSingleElementFromMoveTest, "System::Core::Containers::TDeque::PushLast single element from move", "[SmokeFilter][Core][Containers][Deque]")
 {
 	int32Token::Reset();
 	{
@@ -301,7 +301,7 @@ TEST_CASE("System::Core::Containers::TDeque::PushLast single element from move",
 	CHECK(int32Token::NumMoveConstructorCalls() == 1);
 }
 
-TEST_CASE("System::Core::Containers::TDeque::PushLast single element from copy", "[SmokeFilter][Core][Containers][Deque]")
+TEST_CASE_NAMED(FDequePushLastSingleElementFromCopyTest, "System::Core::Containers::TDeque::PushLast single element from copy", "[SmokeFilter][Core][Containers][Deque]")
 {
 	int32Token::Reset();
 	{
@@ -321,7 +321,7 @@ TEST_CASE("System::Core::Containers::TDeque::PushLast single element from copy",
 	CHECK(int32Token::NumMoveConstructorCalls() == 0);
 }
 
-TEST_CASE("System::Core::Containers::TDeque::PushFirst single element (implicit move)", "[SmokeFilter][Core][Containers][Deque]")
+TEST_CASE_NAMED(FDequePushFirstSingleElementImplicitMoveTest, "System::Core::Containers::TDeque::PushFirst single element (implicit move)", "[SmokeFilter][Core][Containers][Deque]")
 {
 	int32Token::Reset();
 	{
@@ -340,7 +340,7 @@ TEST_CASE("System::Core::Containers::TDeque::PushFirst single element (implicit
 	CHECK(int32Token::NumMoveConstructorCalls() == 1);
 }
 
-TEST_CASE("System::Core::Containers::TDeque::PushFirst single element from move", "[SmokeFilter][Core][Containers][Deque]")
+TEST_CASE_NAMED(FDequePushFirstSingleElementFromMoveTest, "System::Core::Containers::TDeque::PushFirst single element from move", "[SmokeFilter][Core][Containers][Deque]")
 {
 	int32Token::Reset();
 	{
@@ -360,7 +360,7 @@ TEST_CASE("System::Core::Containers::TDeque::PushFirst single element from move"
 	CHECK(int32Token::NumMoveConstructorCalls() == 1);
 }
 
-TEST_CASE("System::Core::Containers::TDeque::PushFirst single element from copy", "[SmokeFilter][Core][Containers][Deque]")
+TEST_CASE_NAMED(FDequePushFirstSingleElementFromCopyTest, "System::Core::Containers::TDeque::PushFirst single element from copy", "[SmokeFilter][Core][Containers][Deque]")
 {
 	int32Token::Reset();
 	{
@@ -385,7 +385,7 @@ void PopOne(TDeque<int32Token>& Deque)
 	Deque.PopFirst();
 }
 
-TEST_CASE("System::Core::Containers::TDeque::EmplaceLast/PopFirst single element", "[SmokeFilter][Core][Containers][Deque]")
+TEST_CASE_NAMED(FDequeEmplaceLastPopFirstSingleElementTest, "System::Core::Containers::TDeque::EmplaceLast/PopFirst single element", "[SmokeFilter][Core][Containers][Deque]")
 {
 	int32Token::Reset();
 	{
@@ -405,7 +405,7 @@ TEST_CASE("System::Core::Containers::TDeque::EmplaceLast/PopFirst single element
 	CHECK(int32Token::EvenConstructionDestructionCalls(1));
 }
 
-TEST_CASE("System::Core::Containers::TDeque::EmplaceLast/PopFirst single element multiple times causing head/tail wrap around", "[SmokeFilter][Core][Containers][Deque]")
+TEST_CASE_NAMED(FDequeEmplaceLastPopFirstSingleElementMultipleTest, "System::Core::Containers::TDeque::EmplaceLast/PopFirst single element multiple times causing head/tail wrap around", "[SmokeFilter][Core][Containers][Deque]")
 {
 	int32Token::Reset();
 	{
@@ -421,7 +421,7 @@ TEST_CASE("System::Core::Containers::TDeque::EmplaceLast/PopFirst single element
 	CHECK(int32Token::EvenConstructionDestructionCalls(Test::DefaultCapacity * 2));
 }
 
-TEST_CASE("System::Core::Containers::TDeque::EmplaceLast/PopFirst range without head/tail wrap around", "[SmokeFilter][Core][Containers][Deque]")
+TEST_CASE_NAMED(FDequeEmplaceLastPopFirstRangeNoWrapTest, "System::Core::Containers::TDeque::EmplaceLast/PopFirst range without head/tail wrap around", "[SmokeFilter][Core][Containers][Deque]")
 {
 	int32Token::Reset();
 	{
@@ -433,7 +433,7 @@ TEST_CASE("System::Core::Containers::TDeque::EmplaceLast/PopFirst range without
 	CHECK(int32Token::EvenConstructionDestructionCalls(Test::DefaultCapacity));
 }
 
-TEST_CASE("System::Core::Containers::TDeque::EmplaceLast/PopFirst range with reallocation without head/tail wrap around", "[SmokeFilter][Core][Containers][Deque]")
+TEST_CASE_NAMED(FDequeEmplaceLastPopFirstRangeWithReallocationNoWrapTest, "System::Core::Containers::TDeque::EmplaceLast/PopFirst range with reallocation without head/tail wrap around", "[SmokeFilter][Core][Containers][Deque]")
 {
 	int32Token::Reset();
 	{
@@ -446,7 +446,7 @@ TEST_CASE("System::Core::Containers::TDeque::EmplaceLast/PopFirst range with rea
 	CHECK(int32Token::EvenConstructionDestructionCalls(Test::DefaultCapacity + 1));
 }
 
-TEST_CASE("System::Core::Containers::TDeque::EmplaceLast/PopFirst range with head/tail wrap around", "[SmokeFilter][Core][Containers][Deque]")
+TEST_CASE_NAMED(FDequeEmplaceLastPopFirstRangeWithWrapTest, "System::Core::Containers::TDeque::EmplaceLast/PopFirst range with head/tail wrap around", "[SmokeFilter][Core][Containers][Deque]")
 {
 	int32Token::Reset();
 	{
@@ -461,7 +461,7 @@ TEST_CASE("System::Core::Containers::TDeque::EmplaceLast/PopFirst range with hea
 	CHECK(int32Token::EvenConstructionDestructionCalls(Test::DefaultCapacity * (Test::DefaultCapacity - 1)));
 }
 
-TEST_CASE("System::Core::Containers::TDeque::EmplaceFirst/PopLast single element", "[SmokeFilter][Core][Containers][Deque]")
+TEST_CASE_NAMED(FDequeEmplaceFirstPopLastSingleElementTest, "System::Core::Containers::TDeque::EmplaceFirst/PopLast single element", "[SmokeFilter][Core][Containers][Deque]")
 {
 	int32Token::Reset();
 	{
@@ -481,7 +481,7 @@ TEST_CASE("System::Core::Containers::TDeque::EmplaceFirst/PopLast single element
 	CHECK(int32Token::EvenConstructionDestructionCalls(1));
 }
 
-TEST_CASE("System::Core::Containers::TDeque::EmplaceFirst/PopLast range", "[SmokeFilter][Core][Containers][Deque]")
+TEST_CASE_NAMED(FDequeEmplaceFirstPopLastRangeTest, "System::Core::Containers::TDeque::EmplaceFirst/PopLast range", "[SmokeFilter][Core][Containers][Deque]")
 {
 	int32Token::Reset();
 	{
@@ -493,7 +493,7 @@ TEST_CASE("System::Core::Containers::TDeque::EmplaceFirst/PopLast range", "[Smok
 	CHECK(int32Token::EvenConstructionDestructionCalls(Test::DefaultCapacity));
 }
 
-TEST_CASE("System::Core::Containers::TDeque::EmplaceFirst/PopLast range with reallocation", "[SmokeFilter][Core][Containers][Deque]")
+TEST_CASE_NAMED(FDequeEmplaceFirstPopLastRangeWithReallocationTest, "System::Core::Containers::TDeque::EmplaceFirst/PopLast range with reallocation", "[SmokeFilter][Core][Containers][Deque]")
 {
 	int32Token::Reset();
 	{
@@ -506,7 +506,7 @@ TEST_CASE("System::Core::Containers::TDeque::EmplaceFirst/PopLast range with rea
 	CHECK(int32Token::EvenConstructionDestructionCalls(Test::DefaultCapacity + 1));
 }
 
-TEST_CASE("System::Core::Containers::TDeque::TryPopFirst", "[SmokeFilter][Core][Containers][Deque]")
+TEST_CASE_NAMED(FDequeTryPopFirstTest, "System::Core::Containers::TDeque::TryPopFirst", "[SmokeFilter][Core][Containers][Deque]")
 {
 	int32Token::Reset();
 	{
@@ -529,7 +529,7 @@ TEST_CASE("System::Core::Containers::TDeque::TryPopFirst", "[SmokeFilter][Core][
 	CHECK(int32Token::EvenConstructionDestructionCalls(Test::DefaultCapacity + 1));  // + 1 for Result
 }
 
-TEST_CASE("System::Core::Containers::TDeque::TryPopFirst with reallocation", "[SmokeFilter][Core][Containers][Deque]")
+TEST_CASE_NAMED(FDequeTryPopFirstWithReallocationTest, "System::Core::Containers::TDeque::TryPopFirst with reallocation", "[SmokeFilter][Core][Containers][Deque]")
 {
 	int32Token::Reset();
 	{
@@ -555,7 +555,7 @@ TEST_CASE("System::Core::Containers::TDeque::TryPopFirst with reallocation", "[S
 	CHECK(int32Token::EvenConstructionDestructionCalls(Test::DefaultCapacity + 2));
 }
 
-TEST_CASE("System::Core::Containers::TDeque::TryPopLast", "[SmokeFilter][Core][Containers][Deque]")
+TEST_CASE_NAMED(FDequeTryPopLastTest, "System::Core::Containers::TDeque::TryPopLast", "[SmokeFilter][Core][Containers][Deque]")
 {
 	int32Token::Reset();
 	{
@@ -578,7 +578,7 @@ TEST_CASE("System::Core::Containers::TDeque::TryPopLast", "[SmokeFilter][Core][C
 	CHECK(int32Token::EvenConstructionDestructionCalls(Test::DefaultCapacity + 1));  // + 1 for Result
 }
 
-TEST_CASE("System::Core::Containers::TDeque::TryPopLast with reallocation", "[SmokeFilter][Core][Containers][Deque]")
+TEST_CASE_NAMED(FDequeTryPopWithReallocationTest, "System::Core::Containers::TDeque::TryPopLast with reallocation", "[SmokeFilter][Core][Containers][Deque]")
 {
 	int32Token::Reset();
 	{
@@ -604,7 +604,7 @@ TEST_CASE("System::Core::Containers::TDeque::TryPopLast with reallocation", "[Sm
 	CHECK(int32Token::EvenConstructionDestructionCalls(Test::DefaultCapacity + 2));
 }
 
-TEST_CASE("System::Core::Containers::TDeque::Comparison simple", "[SmokeFilter][Core][Containers][Deque]")
+TEST_CASE_NAMED(FDequeComparisonSimpleTest, "System::Core::Containers::TDeque::Comparison simple", "[SmokeFilter][Core][Containers][Deque]")
 {
 	int32Token::Reset();
 	{
@@ -623,7 +623,7 @@ TEST_CASE("System::Core::Containers::TDeque::Comparison simple", "[SmokeFilter][
 	CHECK(int32Token::EvenConstructionDestructionCalls(Test::DefaultCapacity * 2));
 }
 
-TEST_CASE("System::Core::Containers::TDeque::Comparison with head/tail wrap around", "[SmokeFilter][Core][Containers][Deque]")
+TEST_CASE_NAMED(FDequeComparisonWithHeadTailWrapTest, "System::Core::Containers::TDeque::Comparison with head/tail wrap around", "[SmokeFilter][Core][Containers][Deque]")
 {
 	int32Token::Reset();
 	{
@@ -651,7 +651,7 @@ TEST_CASE("System::Core::Containers::TDeque::Comparison with head/tail wrap arou
 	CHECK(int32Token::EvenConstructionDestructionCalls(Test::DefaultCapacity * 6));
 }
 
-TEST_CASE("System::Core::Containers::TDeque::Copy simple", "[SmokeFilter][Core][Containers][Deque]")
+TEST_CASE_NAMED(FDequeCopySimpleTest, "System::Core::Containers::TDeque::Copy simple", "[SmokeFilter][Core][Containers][Deque]")
 {
 	int32Token::Reset();
 	{
@@ -669,7 +669,7 @@ TEST_CASE("System::Core::Containers::TDeque::Copy simple", "[SmokeFilter][Core][
 	CHECK(int32Token::EvenConstructionDestructionCalls(Test::DefaultCapacity * 2));
 }
 
-TEST_CASE("System::Core::Containers::TDeque::Copy with head/tail wrap around", "[SmokeFilter][Core][Containers][Deque]")
+TEST_CASE_NAMED(FDequeCopyWithHeadTailWrapTest, "System::Core::Containers::TDeque::Copy with head/tail wrap around", "[SmokeFilter][Core][Containers][Deque]")
 {
 	int32Token::Reset();
 	{
@@ -692,7 +692,7 @@ TEST_CASE("System::Core::Containers::TDeque::Copy with head/tail wrap around", "
 	CHECK(int32Token::EvenConstructionDestructionCalls(Test::DefaultCapacity * 9));
 }
 
-TEST_CASE("System::Core::Containers::TDeque::Copy variable size with head/tail wrap around", "[SmokeFilter][Core][Containers][Deque]")
+TEST_CASE_NAMED(FDequeCopyVariableSizeWithHeadTailWrapTest, "System::Core::Containers::TDeque::Copy variable size with head/tail wrap around", "[SmokeFilter][Core][Containers][Deque]")
 {
 	int32Token::Reset();
 	{
@@ -720,7 +720,7 @@ TEST_CASE("System::Core::Containers::TDeque::Copy variable size with head/tail w
 	CHECK(int32Token::EvenConstructionDestructionCalls());
 }
 
-TEST_CASE("System::Core::Containers::TDeque::Move simple", "[SmokeFilter][Core][Containers][Deque]")
+TEST_CASE_NAMED(FDequeMoveSimpleTest, "System::Core::Containers::TDeque::Move simple", "[SmokeFilter][Core][Containers][Deque]")
 {
 	int32Token::Reset();
 	{
@@ -744,7 +744,7 @@ TEST_CASE("System::Core::Containers::TDeque::Move simple", "[SmokeFilter][Core][
 	CHECK(int32Token::EvenConstructionDestructionCalls(Test::DefaultCapacity + 1));
 }
 
-TEST_CASE("System::Core::Containers::TDeque::Move with head/tail wrap around", "[SmokeFilter][Core][Containers][Deque]")
+TEST_CASE_NAMED(FDequeMoveWithHEadTailWrapTest, "System::Core::Containers::TDeque::Move with head/tail wrap around", "[SmokeFilter][Core][Containers][Deque]")
 {
 	int32Token::Reset();
 	{
@@ -772,7 +772,7 @@ TEST_CASE("System::Core::Containers::TDeque::Move with head/tail wrap around", "
 	CHECK(int32Token::EvenConstructionDestructionCalls(Test::DefaultCapacity * 6));
 }
 
-TEST_CASE("System::Core::Containers::TDeque::Move variable size with head/tail wrap around", "[SmokeFilter][Core][Containers][Deque]")
+TEST_CASE_NAMED(FDequeMoveVariableSizeWithHeadTailWrapTest, "System::Core::Containers::TDeque::Move variable size with head/tail wrap around", "[SmokeFilter][Core][Containers][Deque]")
 {
 	int32Token::Reset();
 	{
@@ -804,7 +804,7 @@ TEST_CASE("System::Core::Containers::TDeque::Move variable size with head/tail w
 	CHECK(int32Token::EvenConstructionDestructionCalls());
 }
 
-TEST_CASE("System::Core::Containers::TDeque::Iteration without head/tail wrap around", "[SmokeFilter][Core][Containers][Deque]")
+TEST_CASE_NAMED(FDequeIterationWIthoutWrapTest, "System::Core::Containers::TDeque::Iteration without head/tail wrap around", "[SmokeFilter][Core][Containers][Deque]")
 {
 	int32Token::Reset();
 	{
@@ -829,7 +829,7 @@ TEST_CASE("System::Core::Containers::TDeque::Iteration without head/tail wrap ar
 	CHECK(int32Token::EvenConstructionDestructionCalls(Test::DefaultCapacity));
 }
 
-TEST_CASE("System::Core::Containers::TDeque::Iteration with head/tail wrap around", "[SmokeFilter][Core][Containers][Deque]")
+TEST_CASE_NAMED(FDequeIterationWithWrapTest, "System::Core::Containers::TDeque::Iteration with head/tail wrap around", "[SmokeFilter][Core][Containers][Deque]")
 {
 	int32Token::Reset();
 	{
@@ -859,7 +859,7 @@ TEST_CASE("System::Core::Containers::TDeque::Iteration with head/tail wrap aroun
 	CHECK(int32Token::EvenConstructionDestructionCalls(Test::DefaultCapacity * (Test::DefaultCapacity + 1)));
 }
 
-TEST_CASE("System::Core::Containers::TDeque::Iterator arithmetic", "[SmokeFilter][Core][Containers][Deque]")
+TEST_CASE_NAMED(FDequeIteratorArithmeticTest, "System::Core::Containers::TDeque::Iterator arithmetic", "[SmokeFilter][Core][Containers][Deque]")
 {
 	TDeque<int32Token> Deque;
 	Deque.Reserve(Test::DefaultCapacity);
@@ -886,7 +886,7 @@ TEST_CASE("System::Core::Containers::TDeque::Iterator arithmetic", "[SmokeFilter
 	}
 }
 
-TEST_CASE("System::Core::Containers::TDeque::Construct from std initializer_list", "[SmokeFilter][Core][Containers][Deque]")
+TEST_CASE_NAMED(FDequeConstructFromStdInitializerListTest, "System::Core::Containers::TDeque::Construct from std initializer_list", "[SmokeFilter][Core][Containers][Deque]")
 {
 	int32Token::Reset();
 	{
@@ -902,7 +902,7 @@ TEST_CASE("System::Core::Containers::TDeque::Construct from std initializer_list
 	CHECK(int32Token::EvenConstructionDestructionCalls(6 * 2 + 1));
 }
 
-TEST_CASE("System::Core::Containers::TDeque::Construct from empty std initializer_list", "[SmokeFilter][Core][Containers][Deque]")
+TEST_CASE_NAMED(FDequeConstructFromEmptyStdInitializerListTest, "System::Core::Containers::TDeque::Construct from empty std initializer_list", "[SmokeFilter][Core][Containers][Deque]")
 {
 	int32Token::Reset();
 	{
@@ -912,7 +912,7 @@ TEST_CASE("System::Core::Containers::TDeque::Construct from empty std initialize
 	CHECK(int32Token::EvenConstructionDestructionCalls(0));
 }
 
-TEST_CASE("System::Core::Containers::TDeque::Assign from std initializer_list", "[SmokeFilter][Core][Containers][Deque]")
+TEST_CASE_NAMED(FDequeAssignFromStdInitializerListTest, "System::Core::Containers::TDeque::Assign from std initializer_list", "[SmokeFilter][Core][Containers][Deque]")
 {
 	int32Token::Reset();
 	{
diff --git a/Engine/Source/Runtime/Core/Tests/Containers/PagedArrayTest.cpp b/Engine/Source/Runtime/Core/Tests/Containers/PagedArrayTest.cpp
--- a/Engine/Source/Runtime/Core/Tests/Containers/PagedArrayTest.cpp
+++ b/Engine/Source/Runtime/Core/Tests/Containers/PagedArrayTest.cpp
@@ -30,7 +30,7 @@ struct FDummyTestStruct
 // Unit tests
 //---------------------------------------------------------------------------------------------------------------------
 
-TEST_CASE("System::Core::Containers::TPagedArray::Reserve without data", "[SmokeFilter][Core][Containers][PagedArray]")
+TEST_CASE_NAMED(FPagedArrayReserveWithoutDataTest, "System::Core::Containers::TPagedArray::Reserve without data", "[SmokeFilter][Core][Containers][PagedArray]")
 {
 	int32Token::Reset();
 	using ArrayType = TPagedArray<int32Token, 128>;
@@ -45,7 +45,8 @@ TEST_CASE("System::Core::Containers::TPagedArray::Reserve without data", "[Smoke
 	CHECK(int32Token::EvenConstructionDestructionCalls(0));
 }
 
-TEST_CASE(
+TEST_CASE_NAMED(
+	FPagedArrayReserveEmplaceSingleElementTest,
 	"System::Core::Containers::TPagedArray::Reserve Emplace single element",
 	"[SmokeFilter][Core][Containers][PagedArray]")
 {
@@ -74,7 +75,7 @@ TEST_CASE(
 	CHECK(int32Token::EvenConstructionDestructionCalls(1));
 }
 
-TEST_CASE("System::Core::Containers::TPagedArray::Add", "[SmokeFilter][Core][Containers][PagedArray]")
+TEST_CASE_NAMED(FPagedArrayAddTest, "System::Core::Containers::TPagedArray::Add", "[SmokeFilter][Core][Containers][PagedArray]")
 {
 	int32Token::Reset();
 	{
@@ -94,7 +95,8 @@ TEST_CASE("System::Core::Containers::TPagedArray::Add", "[SmokeFilter][Core][Con
 	CHECK(int32Token::EvenConstructionDestructionCalls(4));
 }
 
-TEST_CASE(
+TEST_CASE_NAMED(
+	FPagedArrayEmplaceMultipleElementsTest,
 	"System::Core::Containers::TPagedArray::Emplace multiple elements",
 	"[SmokeFilter][Core][Containers][PagedArray]")
 {
@@ -121,7 +123,7 @@ TEST_CASE(
 	CHECK(int32Token::EvenConstructionDestructionCalls(5));
 }
 
-TEST_CASE("System::Core::Containers::TPagedArray::Reset", "[SmokeFilter][Core][Containers][PagedArray]")
+TEST_CASE_NAMED(FPagedArrayResetTest, "System::Core::Containers::TPagedArray::Reset", "[SmokeFilter][Core][Containers][PagedArray]")
 {
 	int32Token::Reset();
 	{
@@ -153,7 +155,8 @@ TEST_CASE("System::Core::Containers::TPagedArray::Reset", "[SmokeFilter][Core][C
 	CHECK(int32Token::EvenConstructionDestructionCalls(1));
 }
 
-TEST_CASE(
+TEST_CASE_NAMED(
+	FPagedArrayResetMultiplePageAddressCheckTest,
 	"System::Core::Containers::TPagedArray::Reset multiple-page plus consistent address check",
 	"[SmokeFilter][Core][Containers][PagedArray]")
 {
@@ -185,7 +188,7 @@ TEST_CASE(
 	CHECK(int32Token::EvenConstructionDestructionCalls(10 * 4));
 }
 
-TEST_CASE("System::Core::Containers::TPagedArray::Empty", "[SmokeFilter][Core][Containers][PagedArray]")
+TEST_CASE_NAMED(FPagedArrayEmptyTest, "System::Core::Containers::TPagedArray::Empty", "[SmokeFilter][Core][Containers][PagedArray]")
 {
 	int32Token::Reset();
 	{
@@ -199,7 +202,8 @@ TEST_CASE("System::Core::Containers::TPagedArray::Empty", "[SmokeFilter][Core][C
 	CHECK(int32Token::EvenConstructionDestructionCalls(0));
 }
 
-TEST_CASE(
+TEST_CASE_NAMED(
+	FPagedArrayEmptyAfterSingleEmplaceTest,
 	"System::Core::Containers::TPagedArray::Empty after single Emplace",
 	"[SmokeFilter][Core][Containers][PagedArray]")
 {
@@ -219,7 +223,7 @@ TEST_CASE(
 	CHECK(int32Token::EvenConstructionDestructionCalls(1));
 }
 
-TEST_CASE("System::Core::Containers::TPagedArray::Indexing", "[SmokeFilter][Core][Containers][PagedArray]")
+TEST_CASE_NAMED(FPagedArrayIndexingTest, "System::Core::Containers::TPagedArray::Indexing", "[SmokeFilter][Core][Containers][PagedArray]")
 {
 	using ArrayType = TPagedArray<int32Token, 128>;
 	int32Token::Reset();
@@ -251,7 +255,7 @@ TEST_CASE("System::Core::Containers::TPagedArray::Indexing", "[SmokeFilter][Core
 	CHECK(int32Token::EvenConstructionDestructionCalls(ArrayType::MaxPerPage() * 2));
 }
 
-TEST_CASE("System::Core::Containers::TPagedArray::Iteration", "[SmokeFilter][Core][Containers][PagedArray]")
+TEST_CASE_NAMED(FPagedArrayIterationTest, "System::Core::Containers::TPagedArray::Iteration", "[SmokeFilter][Core][Containers][PagedArray]")
 {
 	using ArrayType = TPagedArray<int32Token, 128>;
 	int32Token::Reset();
@@ -285,7 +289,7 @@ TEST_CASE("System::Core::Containers::TPagedArray::Iteration", "[SmokeFilter][Cor
 	CHECK(int32Token::EvenConstructionDestructionCalls(ArrayType::MaxPerPage() * 2));
 }
 
-TEST_CASE("System::Core::Containers::TPagedArray::Last", "[SmokeFilter][Core][Containers][PagedArray]")
+TEST_CASE_NAMED(FPagedArrayGetLastTest, "System::Core::Containers::TPagedArray::Last", "[SmokeFilter][Core][Containers][PagedArray]")
 {
 	using ArrayType = TPagedArray<int32Token, 128>;
 	int32Token::Reset();
@@ -300,7 +304,7 @@ TEST_CASE("System::Core::Containers::TPagedArray::Last", "[SmokeFilter][Core][Co
 	CHECK(int32Token::EvenConstructionDestructionCalls(ArrayType::MaxPerPage() * 2));
 }
 
-TEST_CASE("System::Core::Containers::TPagedArray::SetNum", "[SmokeFilter][Core][Containers][PagedArray]")
+TEST_CASE_NAMED(FPagedArraySetNumTest, "System::Core::Containers::TPagedArray::SetNum", "[SmokeFilter][Core][Containers][PagedArray]")
 {
 	using ArrayType = TPagedArray<int32Token, 128>;
 	int32Token::Reset();
@@ -396,7 +400,7 @@ TEST_CASE("System::Core::Containers::TPagedArray::SetNum", "[SmokeFilter][Core][
 	CHECK(int32Token::EvenConstructionDestructionCalls(ArrayType::MaxPerPage() * 6 * 2));
 }
 
-TEST_CASE("System::Core::Containers::TPagedArray::SetNum after Reset", "[SmokeFilter][Core][Containers][PagedArray]")
+TEST_CASE_NAMED(FPagedArraySetNumAfterResetTest, "System::Core::Containers::TPagedArray::SetNum after Reset", "[SmokeFilter][Core][Containers][PagedArray]")
 {
 	using ArrayType = TPagedArray<int32Token, 16>;
 	int32Token::Reset();
@@ -440,7 +444,8 @@ TEST_CASE("System::Core::Containers::TPagedArray::SetNum after Reset", "[SmokeFi
 	CHECK(int32Token::EvenConstructionDestructionCalls(ArrayType::MaxPerPage() * (3 * 3 + 5)));
 }
 
-TEST_CASE(
+TEST_CASE_NAMED(
+	FPagedArraySetNumSmallPageSizeTest,
 	"System::Core::Containers::TPagedArray::SetNum small page size",
 	"[SmokeFilter][Core][Containers][PagedArray]")
 {
@@ -477,7 +482,7 @@ TEST_CASE(
 	CHECK(int32Token::EvenConstructionDestructionCalls(ArrayType::MaxPerPage() * 10 - 1));
 }
 
-TEST_CASE("System::Core::Containers::TPagedArray::SetZero", "[SmokeFilter][Core][Containers][PagedArray]")
+TEST_CASE_NAMED(FPagedArraySetZeroTest, "System::Core::Containers::TPagedArray::SetZero", "[SmokeFilter][Core][Containers][PagedArray]")
 {
 	using ArrayType = TPagedArray<int32Token, 16>;
 	int32Token::Reset();
@@ -528,7 +533,7 @@ TEST_CASE("System::Core::Containers::TPagedArray::SetZero", "[SmokeFilter][Core]
 	CHECK(int32Token::EvenConstructionDestructionCalls(ArrayType::MaxPerPage() * 7 + 1));
 }
 
-TEST_CASE("System::Core::Containers::TPagedArray::RemoveAtSwap", "[SmokeFilter][Core][Containers][PagedArray]")
+TEST_CASE_NAMED(FPagedArrayRemoveAtSwapTest, "System::Core::Containers::TPagedArray::RemoveAtSwap", "[SmokeFilter][Core][Containers][PagedArray]")
 {
 	using ArrayType = TPagedArray<int32Token, 128>;
 	int32Token::Reset();
@@ -592,7 +597,7 @@ TEST_CASE("System::Core::Containers::TPagedArray::RemoveAtSwap", "[SmokeFilter][
 	CHECK(int32Token::EvenConstructionDestructionCalls(ArrayType::MaxPerPage() * 9 + 2));
 }
 
-TEST_CASE("System::Core::Containers::TPagedArray::Pop", "[SmokeFilter][Core][Containers][PagedArray]")
+TEST_CASE_NAMED(FPagedArrayPopTest, "System::Core::Containers::TPagedArray::Pop", "[SmokeFilter][Core][Containers][PagedArray]")
 {
 	using ArrayType = TPagedArray<int32Token, 128>;
 	int32Token::Reset();
@@ -622,7 +627,7 @@ TEST_CASE("System::Core::Containers::TPagedArray::Pop", "[SmokeFilter][Core][Con
 	CHECK(int32Token::EvenConstructionDestructionCalls(ArrayType::MaxPerPage() * 2 + 2));
 }
 
-TEST_CASE("System::Core::Containers::TPagedArray::Copy", "[SmokeFilter][Core][Containers][PagedArray]")
+TEST_CASE_NAMED(FPagedArrayCopyTest, "System::Core::Containers::TPagedArray::Copy", "[SmokeFilter][Core][Containers][PagedArray]")
 {
 	using ArrayType = TPagedArray<int32Token, 128>;
 	int32Token::Reset();
@@ -674,7 +679,7 @@ TEST_CASE("System::Core::Containers::TPagedArray::Copy", "[SmokeFilter][Core][Co
 	CHECK(int32Token::EvenConstructionDestructionCalls(ArrayType::MaxPerPage() * 8));
 }
 
-TEST_CASE("System::Core::Containers::TPagedArray::Move", "[SmokeFilter][Core][Containers][PagedArray]")
+TEST_CASE_NAMED(FPagedArrayMoveTest, "System::Core::Containers::TPagedArray::Move", "[SmokeFilter][Core][Containers][PagedArray]")
 {
 	using ArrayType = TPagedArray<int32Token, 128>;
 	int32Token::Reset();
@@ -729,7 +734,8 @@ TEST_CASE("System::Core::Containers::TPagedArray::Move", "[SmokeFilter][Core][Co
 	CHECK(int32Token::EvenConstructionDestructionCalls(ArrayType::MaxPerPage() * 4));
 }
 
-TEST_CASE(
+TEST_CASE_NAMED(
+	FPagedArrayConstructFromStdInitializerListTest,
 	"System::Core::Containers::TPagedArray::Construct from std initializer_list",
 	"[SmokeFilter][Core][Containers][PagedArray]")
 {
@@ -765,7 +771,8 @@ TEST_CASE(
 	CHECK(int32Token::EvenConstructionDestructionCalls((4 + 10) * 2));
 }
 
-TEST_CASE(
+TEST_CASE_NAMED(
+	FPagedArrayAssignFromStdInitializerListTest,
 	"System::Core::Containers::TPagedArray::Assign from std initializer_list",
 	"[SmokeFilter][Core][Containers][PagedArray]")
 {
@@ -805,7 +812,7 @@ TEST_CASE(
 	CHECK(int32Token::EvenConstructionDestructionCalls((4 + 10) * 2));
 }
 
-TEST_CASE("System::Core::Containers::TPagedArray::Comparison", "[SmokeFilter][Core][Containers][PagedArray]")
+TEST_CASE_NAMED(FPagedArrayComparisonTest, "System::Core::Containers::TPagedArray::Comparison", "[SmokeFilter][Core][Containers][PagedArray]")
 {
 	using ArrayType = TPagedArray<int32Token, 128>;
 	int32Token::Reset();
@@ -842,7 +849,7 @@ TEST_CASE("System::Core::Containers::TPagedArray::Comparison", "[SmokeFilter][Co
 	CHECK(int32Token::EvenConstructionDestructionCalls(ArrayType::MaxPerPage() * 2 + 2));
 }
 
-TEST_CASE("System::Core::Containers::TPagedArray::Assign from TArray", "[SmokeFilter][Core][Containers][PagedArray]")
+TEST_CASE_NAMED(FPagedArrayAssignFromArrayTest, "System::Core::Containers::TPagedArray::Assign from TArray", "[SmokeFilter][Core][Containers][PagedArray]")
 {
 	using ArrayType = TPagedArray<int32Token, 16>;
 	int32Token::Reset();
@@ -885,7 +892,7 @@ TEST_CASE("System::Core::Containers::TPagedArray::Assign from TArray", "[SmokeFi
 	CHECK(int32Token::EvenConstructionDestructionCalls((4 + 10) * 3));
 }
 
-TEST_CASE("System::Core::Containers::TPagedArray::Copy to TArray", "[SmokeFilter][Core][Containers][PagedArray]")
+TEST_CASE_NAMED(FPagedArrayCopyToArrayTest, "System::Core::Containers::TPagedArray::Copy to TArray", "[SmokeFilter][Core][Containers][PagedArray]")
 {
 	using ArrayType = TPagedArray<int32Token, 16>;
 	int32Token::Reset();
@@ -932,7 +939,7 @@ TEST_CASE("System::Core::Containers::TPagedArray::Copy to TArray", "[SmokeFilter
 	CHECK(int32Token::EvenConstructionDestructionCalls(4 * 4 + 10 * 3));
 }
 
-TEST_CASE("System::Core::Containers::TPagedArray::Move to TArray", "[SmokeFilter][Core][Containers][PagedArray]")
+TEST_CASE_NAMED(FPagedArrayMoveToArrayTest, "System::Core::Containers::TPagedArray::Move to TArray", "[SmokeFilter][Core][Containers][PagedArray]")
 {
 	using ArrayType = TPagedArray<int32Token, 16>;
 	int32Token::Reset();
@@ -983,7 +990,8 @@ TEST_CASE("System::Core::Containers::TPagedArray::Move to TArray", "[SmokeFilter
 	CHECK(int32Token::EvenConstructionDestructionCalls(4 * 4 + 10 * 2));
 }
 
-TEST_CASE(
+TEST_CASE_NAMED(
+	FPagedArrayAssignFromCompatiblePagedArrayTest,
 	"System::Core::Containers::TPagedArray::Assign from compatible paged array",
 	"[SmokeFilter][Core][Containers][PagedArray]")
 {
@@ -1059,7 +1067,8 @@ TEST_CASE(
 	CHECK(int32Token::EvenConstructionDestructionCalls((10 + 8 + 8 + 4) * 3));
 }
 
-TEST_CASE(
+TEST_CASE_NAMED(
+	FPagedArrayAppendFromCompatiblePagedArrayTest,
 	"System::Core::Containers::TPagedArray::Append from compatible paged array",
 	"[SmokeFilter][Core][Containers][PagedArray]")
 {
@@ -1167,7 +1176,7 @@ TEST_CASE(
 	CHECK(int32Token::EvenConstructionDestructionCalls((8 + 16) * 6));
 }
 
-TEST_CASE("System::Core::Containers::TPagedArray::Append multiple", "[SmokeFilter][Core][Containers][PagedArray]")
+TEST_CASE_NAMED(FPagedArrayAppendMultipleTest, "System::Core::Containers::TPagedArray::Append multiple", "[SmokeFilter][Core][Containers][PagedArray]")
 {
 	using ArrayType = TPagedArray<int32Token, 16>;
 	int32Token::Reset();
@@ -1199,7 +1208,7 @@ TEST_CASE("System::Core::Containers::TPagedArray::Append multiple", "[SmokeFilte
 	CHECK(int32Token::EvenConstructionDestructionCalls(8 * 2));
 }
 
-TEST_CASE("System::Core::Containers::TPagedArray::Iterator arithmetic", "[SmokeFilter][Core][Containers][PagedArray]")
+TEST_CASE_NAMED(FPagedArrayIteratorArithmeticTest, "System::Core::Containers::TPagedArray::Iterator arithmetic", "[SmokeFilter][Core][Containers][PagedArray]")
 {
 	using ArrayType = TPagedArray<int32Token, 16>;
 	int32Token::Reset();
