Reduce amount of malloc calls during pak creation

1) Move few containers outside of inner loop to reuse same memory across iterations in AssignPackagesDiskOrder
2) Change FString to FStringView when list of filenames is prepared during FIoStoreWriter::Finalize

diff --git a/Engine/Source/Developer/IoStoreUtilities/Private/IoStoreUtilities.cpp b/Engine/Source/Developer/IoStoreUtilities/Private/IoStoreUtilities.cpp
--- a/Engine/Source/Developer/IoStoreUtilities/Private/IoStoreUtilities.cpp
+++ b/Engine/Source/Developer/IoStoreUtilities/Private/IoStoreUtilities.cpp
@@ -1122,6 +1122,13 @@ static void AssignPackagesDiskOrder(
 		});
 	}
 
+	// Keep these containers outside of inner loops to reuse allocated memory across iterations
+	// No need to allocate & free them on every single iteration, just reset element count before using them
+	TSet<FCluster*> ClustersToRead;
+	TSet<FLegacyCookedPackage*> VisitedDeps;
+	TArray<FLegacyCookedPackage*> DepQueue;
+	TArray<FCluster*> OrderedClustersToRead;
+
 	int32 ClusterSequence = 0;
 	TMap<FLegacyCookedPackage*, FCluster*> PackageToCluster;
 	for (FPackageAndOrder& Entry : SortedPackages)
@@ -1177,10 +1184,11 @@ static void AssignPackagesDiskOrder(
 			for (FLegacyCookedPackage* Package : Cluster->Packages)
 			{
 				int64 BytesToRead = 0;
-				TSet<FCluster*> ClustersToRead;
 
-				TSet<FLegacyCookedPackage*> VisitedDeps;
-				TArray<FLegacyCookedPackage*> DepQueue;
+				ClustersToRead.Reset();
+				VisitedDeps.Reset();
+				DepQueue.Reset();
+
 				DepQueue.Push(Package);
 				while (DepQueue.Num() > 0)
 				{
@@ -1201,11 +1209,15 @@ static void AssignPackagesDiskOrder(
 							{
 								DepQueue.Push(FindReferencedPackage);
 							}
-		}
-	}
+						}
+					}
 				}
 
-				TArray<FCluster*> OrderedClustersToRead = ClustersToRead.Array();
+				OrderedClustersToRead.Reset(ClustersToRead.Num());
+				for (FCluster* ClusterToRead : ClustersToRead)
+				{
+					OrderedClustersToRead.Add(ClusterToRead);
+				}
 				Algo::SortBy(OrderedClustersToRead, [](FCluster* C) { return C->ClusterSequence; }, TLess<int32>());
 
 				int32 NumClustersToRead = 1; // Could replace with "min seeks"
diff --git a/Engine/Source/Runtime/Core/Internal/IO/IoDirectoryIndex.h b/Engine/Source/Runtime/Core/Internal/IO/IoDirectoryIndex.h
--- a/Engine/Source/Runtime/Core/Internal/IO/IoDirectoryIndex.h
+++ b/Engine/Source/Runtime/Core/Internal/IO/IoDirectoryIndex.h
@@ -3,6 +3,7 @@
 #pragma once
 
 #include "Containers/Array.h"
+#include "Containers/ChunkedArray.h"
 #include "Containers/Map.h"
 #include "Containers/StringFwd.h"
 #include "Containers/UnrealString.h"
@@ -13,7 +14,7 @@ class FArchive;
 
 namespace IoDirectoryIndexUtils
 {
-	FString GetCommonRootPath(const TArray<FString>& Filenames);
+	FString GetCommonRootPath(const TArray<FStringView>& Filenames);
 }
 
 struct FIoDirectoryIndexEntry
@@ -51,8 +52,11 @@ public:
 	FIoDirectoryIndexWriter();
 
 	void SetMountPoint(FString InMountPoint);
-	uint32 AddFile(const FString& InFileName);
+	uint32 AddFile(const FStringView& InFileName);
 	void SetFileUserData(uint32 InFileEntryIndex, uint32 InUserData);
+
+	// Flush() can be called only once, because afterwards contents of Strings member
+	// will be empty as all strings have been moved away
 	void Flush(TArray<uint8>& OutBuffer, FAES::FAESKey InEncryptionKey);
 
 private:
@@ -69,6 +73,11 @@ private:
 	FString MountPoint;
 	TArray<FIoDirectoryIndexEntry> DirectoryEntries;
 	TArray<FIoFileIndexEntry> FileEntries;
-	TMap<FString, uint32> StringToIndex;
-	TArray<FString> Strings;
+
+	// These FStringView's are referencing stable FString locations inside Strings
+	// member, and Strings member is only appended, never removed. It is safe to
+	// use StringToIndex during lifetime of this object while Strings member is
+	// populated by AddFile() until Flush() is called
+	TMap<FStringView, uint32> StringToIndex;
+	TChunkedArray<FString> Strings;
 };
diff --git a/Engine/Source/Runtime/Core/Private/IO/IoDirectoryIndex.cpp b/Engine/Source/Runtime/Core/Private/IO/IoDirectoryIndex.cpp
--- a/Engine/Source/Runtime/Core/Private/IO/IoDirectoryIndex.cpp
+++ b/Engine/Source/Runtime/Core/Private/IO/IoDirectoryIndex.cpp
@@ -20,12 +20,12 @@ namespace IoDirectoryIndexUtils
 {
 
 // See PakFileUtilities.cpp 
-FString GetLongestPath(const TArray<FString>& Filenames)
+static FString GetLongestPath(const TArray<FStringView>& Filenames)
 {
 	FString LongestPath;
 	int32 MaxNumDirectories = 0;
 
-	for (const FString& Filename : Filenames)
+	for (const FStringView& Filename : Filenames)
 	{
 		int32 NumDirectories = 0;
 		for (int32 Index = 0; Index < Filename.Len(); Index++)
@@ -45,12 +45,12 @@ FString GetLongestPath(const TArray<FString>& Filenames)
 }
 
 // See PakFileUtilities.cpp 
-FString GetCommonRootPath(const TArray<FString>& Filenames)
+FString GetCommonRootPath(const TArray<FStringView>& Filenames)
 {
 	FString Root = GetLongestPath(Filenames);
-	for (const FString& Filename : Filenames)
+	for (const FStringView& Filename : Filenames)
 	{
-		FString Path = FPaths::GetPath(Filename) + TEXT("/");
+		FString Path = FPaths::GetPath(FString(Filename)) + TEXT("/");
 		int32 CommonSeparatorIndex = -1;
 		int32 SeparatorIndex = Path.Find(TEXT("/"), ESearchCase::CaseSensitive);
 		while (SeparatorIndex >= 0)
@@ -78,7 +78,7 @@ FString GetCommonRootPath(const TArray<FString>& Filenames)
 }
 
 // See IPlatformFilePak.h
-bool SplitPathInline(FString& InOutPath, FString& OutFilename)
+static bool SplitPathInline(FStringView& InOutPath, FStringView& OutFilename)
 {
 	// FPaths::GetPath doesn't handle our / at the end of directories, so we have to do string manipulation ourselves
 	// The manipulation is less complicated than GetPath deals with, since we have normalized/path/strings, we have relative paths only, and we don't care about extensions
@@ -92,7 +92,7 @@ bool SplitPathInline(FString& InOutPath, FString& OutFilename)
 		if (InOutPath[0] == TEXT('/'))
 		{
 			// The root directory; it has no parent.
-			OutFilename.Empty();
+			OutFilename = FStringView();
 			return false;
 		}
 		else
@@ -108,7 +108,7 @@ bool SplitPathInline(FString& InOutPath, FString& OutFilename)
 		if (InOutPath[InOutPath.Len() - 1] == TEXT('/'))
 		{
 			// The input was a Directory; remove the trailing / since we don't keep those on the CleanFilename
-			InOutPath.LeftChopInline(1, false /* bAllowShrinking */);
+			InOutPath.LeftChopInline(1);
 		}
 
 		int32 Offset = 0;
@@ -116,7 +116,7 @@ bool SplitPathInline(FString& InOutPath, FString& OutFilename)
 		{
 			int32 FilenameStart = Offset + 1;
 			OutFilename = InOutPath.Mid(FilenameStart);
-			InOutPath.LeftInline(FilenameStart, false /* bAllowShrinking */); // The Parent Directory keeps the / at the end
+			InOutPath.LeftInline(FilenameStart); // The Parent Directory keeps the / at the end
 		}
 		else
 		{
@@ -181,13 +181,13 @@ void FIoDirectoryIndexWriter::SetMountPoint(FString InMountPoint)
 	MountPoint = MoveTemp(InMountPoint);
 }
 
-uint32 FIoDirectoryIndexWriter::AddFile(const FString& InFileName)
+uint32 FIoDirectoryIndexWriter::AddFile(const FStringView& InFileName)
 {
 	uint32 Directory = 0; //Root
 
-	FString RelativePathFromMount = InFileName.Mid(MountPoint.Len());
-	FString RelativeDirectoryFromMount = RelativePathFromMount;
-	FString CleanFileName;
+	FStringView RelativePathFromMount = InFileName.Mid(MountPoint.Len());
+	FStringView RelativeDirectoryFromMount = RelativePathFromMount;
+	FStringView CleanFileName;
 	IoDirectoryIndexUtils::SplitPathInline(RelativeDirectoryFromMount, CleanFileName);
 
 	FStringView Path = RelativeDirectoryFromMount;
@@ -213,12 +213,20 @@ void FIoDirectoryIndexWriter::SetFileUserData(uint32 InFileEntryIndex, uint32 In
 
 void FIoDirectoryIndexWriter::Flush(TArray<uint8>& OutBuffer, FAES::FAESKey EncryptionKey)
 {
+	TArray<FString> StringTable;
+	StringTable.Reserve(Strings.Num());
+	for (FString& String : Strings)
+	{
+		StringTable.Emplace(MoveTemp(String));
+	}
+	Strings.Empty();
+
 	//TODO: Sort entries based on directory hierarchy
 	FIoDirectoryIndexResource DirectoryIndex;
 	DirectoryIndex.MountPoint = MoveTemp(MountPoint);
 	DirectoryIndex.DirectoryEntries = MoveTemp(DirectoryEntries);
 	DirectoryIndex.FileEntries = MoveTemp(FileEntries);
-	DirectoryIndex.StringTable = MoveTemp(Strings);
+	DirectoryIndex.StringTable = MoveTemp(StringTable);
 	
 	FMemoryWriter Ar(OutBuffer);
 	Ar << DirectoryIndex;
@@ -281,17 +289,15 @@ uint32 FIoDirectoryIndexWriter::CreateDirectory(const FStringView& DirectoryName
 
 uint32 FIoDirectoryIndexWriter::GetNameIndex(const FStringView& String)
 {
-	FString Tmp(String);
-
-	if (const uint32* Index = StringToIndex.Find(Tmp))
+	if (const uint32* Index = StringToIndex.Find(String))
 	{
 		return *Index;
 	}
 	else
 	{
-		const uint32 NewIndex = Strings.Num();
-		Strings.Emplace(Tmp);
-		StringToIndex.Add(Tmp, NewIndex);
+		const uint32 NewIndex = Strings.Add();
+		const FString& NewString = Strings[NewIndex] = FString(String);
+		StringToIndex.Add(NewString, NewIndex);
 		return NewIndex;
 	}
 }
diff --git a/Engine/Source/Runtime/Core/Private/IO/IoStore.cpp b/Engine/Source/Runtime/Core/Private/IO/IoStore.cpp
--- a/Engine/Source/Runtime/Core/Private/IO/IoStore.cpp
+++ b/Engine/Source/Runtime/Core/Private/IO/IoStore.cpp
@@ -581,9 +581,13 @@ public:
 		return nullptr;
 	}
 
-	void GetFileNamesToIndex(TArray<FString>& OutFileNames) const
+	void GetFileNamesToIndex(TArray<FStringView>& OutFileNames) const
 	{
-		ChunkIdToFileName.GenerateValueArray(OutFileNames);
+		OutFileNames.Empty(ChunkIdToFileName.Num());
+		for (auto& ChinkIdAndFileName : ChunkIdToFileName)
+		{
+			OutFileNames.Emplace(ChinkIdAndFileName.Value);
+		}
 	}
 
 	const FString* GetFileName(const FIoChunkId& ChunkId) const
@@ -1159,7 +1163,7 @@ public:
 		if (ContainerSettings.IsIndexed())
 		{
 			TRACE_CPUPROFILER_EVENT_SCOPE(BuildIndex);
-			TArray<FString> FilesToIndex;
+			TArray<FStringView> FilesToIndex;
 			Toc.GetFileNamesToIndex(FilesToIndex);
 			
 			FString MountPoint = IoDirectoryIndexUtils::GetCommonRootPath(FilesToIndex);
