CompactBinary: Fixed AsFieldView() to exclude names as expected

diff --git a/Engine/Source/Runtime/Core/Public/Serialization/CompactBinary.h b/Engine/Source/Runtime/Core/Public/Serialization/CompactBinary.h
--- a/Engine/Source/Runtime/Core/Public/Serialization/CompactBinary.h
+++ b/Engine/Source/Runtime/Core/Public/Serialization/CompactBinary.h
@@ -620,6 +620,15 @@ public:
 	/** Construct a field from a value, without access to the name. */
 	inline explicit FCbFieldView(const FCbValue& Value);
 
+	/** Returns a copy of the field with the name removed. */
+	constexpr inline FCbFieldView RemoveName() const
+	{
+		FCbFieldView Field;
+		Field.TypeWithFlags = TypeWithFlags & ~(ECbFieldType::HasFieldType | ECbFieldType::HasFieldName);
+		Field.Value = Value;
+		return Field;
+	}
+
 	/** Returns the name of the field if it has a name, otherwise an empty view. */
 	constexpr inline FUtf8StringView GetName() const
 	{
@@ -939,7 +948,7 @@ public:
 	CORE_API uint64 Num() const;
 
 	/** Access the array as an array field. */
-	inline FCbFieldView AsFieldView() const { return static_cast<const FCbFieldView&>(*this); }
+	inline FCbFieldView AsFieldView() const { return RemoveName(); }
 
 	/** Construct an array from an array field. No type check is performed! */
 	static inline FCbArrayView FromFieldNoCheck(const FCbFieldView& Field) { return FCbArrayView(Field); }
@@ -1036,7 +1045,7 @@ public:
 	inline FCbFieldView operator[](FUtf8StringView Name) const { return FindView(Name); }
 
 	/** Access the object as an object field. */
-	inline FCbFieldView AsFieldView() const { return static_cast<const FCbFieldView&>(*this); }
+	inline FCbFieldView AsFieldView() const { return RemoveName(); }
 
 	/** Construct an object from an object field. No type check is performed! */
 	static inline FCbObjectView FromFieldNoCheck(const FCbFieldView& Field) { return FCbObjectView(Field); }
