Fixes non-unity build, previous change requires FStringView type visible in header

diff --git a/Engine/Source/Runtime/Core/Internal/IO/IoDirectoryIndex.h b/Engine/Source/Runtime/Core/Internal/IO/IoDirectoryIndex.h
--- a/Engine/Source/Runtime/Core/Internal/IO/IoDirectoryIndex.h
+++ b/Engine/Source/Runtime/Core/Internal/IO/IoDirectoryIndex.h
@@ -6,6 +6,7 @@
 #include "Containers/ChunkedArray.h"
 #include "Containers/Map.h"
 #include "Containers/StringFwd.h"
+#include "Containers/StringView.h"
 #include "Containers/UnrealString.h"
 #include "HAL/Platform.h"
 #include "Misc/AES.h"
