Remove unused FNAME_Replace_Not_Safe_For_Threading and FNAME_WRITE_PROTECT_PAGES

diff --git a/Engine/Source/Runtime/Core/Private/UObject/UnrealNames.cpp b/Engine/Source/Runtime/Core/Private/UObject/UnrealNames.cpp
--- a/Engine/Source/Runtime/Core/Private/UObject/UnrealNames.cpp
+++ b/Engine/Source/Runtime/Core/Private/UObject/UnrealNames.cpp
@@ -40,16 +40,6 @@
 
 PRAGMA_DISABLE_UNSAFE_TYPECAST_WARNINGS
 
-// Page protection to catch FNameEntry stomps
-#ifndef FNAME_WRITE_PROTECT_PAGES
-#define FNAME_WRITE_PROTECT_PAGES 0
-#endif
-#if FNAME_WRITE_PROTECT_PAGES
-#	define FNAME_BLOCK_ALIGNMENT FPlatformMemory::GetConstants().PageSize
-#else
-#	define FNAME_BLOCK_ALIGNMENT alignof(FNameEntry)
-#endif
-
 DEFINE_LOG_CATEGORY_STATIC(LogUnrealNames, Log, All);
 
 // Console command declarations
@@ -444,7 +434,7 @@ public:
 	FNameEntryAllocator()
 	{
 		LLM_SCOPE(ELLMTag::FName);
-		Blocks[0] = (uint8*)FMemory::MallocPersistentAuxiliary(BlockSizeBytes, FNAME_BLOCK_ALIGNMENT);
+		Blocks[0] = (uint8*)FMemory::MallocPersistentAuxiliary(BlockSizeBytes, alignof(FNameEntry));
 	}
 
 	~FNameEntryAllocator()
@@ -664,7 +654,7 @@ private:
 
 	static uint8* AllocBlock()
 	{
-		return (uint8*)FMemory::MallocPersistentAuxiliary(BlockSizeBytes, FNAME_BLOCK_ALIGNMENT);
+		return (uint8*)FMemory::MallocPersistentAuxiliary(BlockSizeBytes, alignof(FNameEntry));
 	}
 	
 	void AllocateNewBlock()
@@ -687,9 +677,6 @@ private:
 		}
 #endif
 
-#if FNAME_WRITE_PROTECT_PAGES
-		FPlatformMemory::PageProtect(Blocks[CurrentBlock], BlockSizeBytes, /* read */ true, /* write */ false);
-#endif
 		++CurrentBlock;
 		CurrentByteCursor = 0;
 
@@ -2896,8 +2883,9 @@ struct FNameHelper
 			FNameEntryId DisplayId = Pool.StoreWithNumber(BaseIds, InternalNumber);
 			return FinalConstruct(FNameEntryIds{ ResolveComparisonId(DisplayId), DisplayId });
 		}
-		else if (FindType == FNAME_Find)
+		else
 		{
+			check(FindType == FNAME_Find);
 			FNameEntryId DisplayId = Pool.FindWithNumber(BaseIds.DisplayId, InternalNumber);
 			if (DisplayId)
 			{
@@ -2909,11 +2897,6 @@ struct FNameHelper
 				return FName();
 			}
 		}
-		else
-		{
-			checkf(false, TEXT("FNAME_Replace_Not_Safe_For_Threading not supported for numbered names"));
-			return FName();
-		}
 #else
 		// Number is just stored in the FName pass it straight on
 		return FinalConstruct(BaseIds, InternalNumber);
@@ -3052,24 +3035,12 @@ private:
 			FNameEntryId DisplayId = Pool.Store(View);
 			return FNameEntryIds{ ResolveComparisonId(DisplayId), DisplayId };
 		}
-		else if (FindType == FNAME_Find)
+		else
 		{
+			check(FindType == FNAME_Find);
 			FNameEntryId DisplayId = Pool.Find(View);
 			return FNameEntryIds{ ResolveComparisonId(DisplayId), DisplayId };
 		}
-		else
-		{
-			check(FindType == FNAME_Replace_Not_Safe_For_Threading);
-
-#if FNAME_WRITE_PROTECT_PAGES
-			checkf(false, TEXT("FNAME_Replace_Not_Safe_For_Threading can't be used together with page protection."));
-#endif
-
-			FNameEntryId DisplayId = Pool.Store(View);
-			FNameEntryId ComparisonId = ResolveComparisonId(DisplayId);
-			ReplaceName(Pool.Resolve(ComparisonId), View);
-			return FNameEntryIds{ ComparisonId, DisplayId };
-		}
 	}
 
 #if UE_FNAME_OUTLINE_NUMBER
@@ -3901,25 +3872,7 @@ PRAGMA_DISABLE_DEPRECATION_WARNINGS
 		check(FName("UniqueUnicorn!!", 17, FNAME_Find) != FName());		// Now we can find it
 		check(FName("UniqueUnicorn!!", 127, FNAME_Find) == FName());	// But we can't find one with a different number
 PRAGMA_ENABLE_DEPRECATION_WARNINGS
-#endif
-
-#if !FNAME_WRITE_PROTECT_PAGES
-		// Check FNAME_Replace_Not_Safe_For_Threading updates casing
-		check(0 != UniqueName.GetPlainNameString().Compare("UNIQUEunicorn!!", ESearchCase::CaseSensitive));
-		const FName UniqueNameReplaced("UNIQUEunicorn!!", FNAME_Replace_Not_Safe_For_Threading);
-		check(0 == UniqueName.GetPlainNameString().Compare("UNIQUEunicorn!!", ESearchCase::CaseSensitive));
-		check(UniqueNameReplaced == UniqueName);
-
-		// Check FNAME_Replace_Not_Safe_For_Threading works with wide string
-		check(0 != UniqueName.GetPlainNameString().Compare("uniqueunicorn!!", ESearchCase::CaseSensitive));
-		const FName UpdatedCasing(TEXT("uniqueunicorn!!"), FNAME_Replace_Not_Safe_For_Threading);
-		check(0 == UniqueName.GetPlainNameString().Compare("uniqueunicorn!!", ESearchCase::CaseSensitive));
-
-		// Check FNAME_Replace_Not_Safe_For_Threading adds entries that do not exist
-		const FName AddedByReplace("WasAdded!!", FNAME_Replace_Not_Safe_For_Threading);
-		check(FName("WasAdded!!", FNAME_Find) == AddedByReplace);
-#endif
-	
+#endif	
 		Once = false;
 	}
 
diff --git a/Engine/Source/Runtime/Core/Public/UObject/NameTypes.h b/Engine/Source/Runtime/Core/Public/UObject/NameTypes.h
--- a/Engine/Source/Runtime/Core/Public/UObject/NameTypes.h
+++ b/Engine/Source/Runtime/Core/Public/UObject/NameTypes.h
@@ -183,12 +183,7 @@ enum EFindName
 	FNAME_Find,
 
 	/** Find a name or add it if it doesn't exist. */
-	FNAME_Add,
-
-	/** Finds a name and replaces it. Adds it if missing. This is only used by UHT and is generally not safe for threading. 
-	 * All this really is used for is correcting the case of names. In MT conditions you might get a half-changed name.
-	 */
-	FNAME_Replace_Not_Safe_For_Threading,
+	FNAME_Add
 };
 
 /*----------------------------------------------------------------------------
