Unshelved from pending changelist '23736295':

Fixed documentation commentary for TRotator::Equals

diff --git a/Engine/Source/Runtime/Core/Public/Math/Rotator.h b/Engine/Source/Runtime/Core/Public/Math/Rotator.h
--- a/Engine/Source/Runtime/Core/Public/Math/Rotator.h
+++ b/Engine/Source/Runtime/Core/Public/Math/Rotator.h
@@ -222,6 +222,9 @@ public:
 	/**
 	 * Checks whether two rotators are equal within specified tolerance, when treated as an orientation.
 	 * This means that TRotator(0, 0, 360).Equals(TRotator(0,0,0)) is true, because they represent the same final orientation.
+	 * It can compare only wound rotators (i.e. multiples of 360 degrees) that end up in a same rotation
+	 * Rotators that represent the same final rotation, but get there via different intermediate rotations aren't equal
+	 * i.e. TRotator(0, 45, 0).Equals(TRotator(180, 135, 180)) is false
 	 *
 	 * @param R The other rotator.
 	 * @param Tolerance Error Tolerance.
