String functions: Delete ParseIntoArray; it is redundant with the existing ParseTokens.
#rnx

diff --git a/Engine/Source/Runtime/Core/Private/Containers/String.cpp b/Engine/Source/Runtime/Core/Private/Containers/String.cpp
--- a/Engine/Source/Runtime/Core/Private/Containers/String.cpp
+++ b/Engine/Source/Runtime/Core/Private/Containers/String.cpp
@@ -18,7 +18,7 @@
 #include "Misc/VarArgs.h"
 #include "Serialization/Archive.h"
 #include "String/HexToBytes.h"
-#include "String/Parse.h"
+#include "String/ParseTokens.h"
 #include "Templates/MemoryOps.h"
 #include "Templates/RemoveReference.h"
 #include "Templates/UnrealTemplate.h"
@@ -1281,7 +1281,17 @@ bool FString::IsNumeric() const
 int32 FString::ParseIntoArray( TArray<FString>& OutArray, const TCHAR* pchDelim, const bool InCullEmpty ) const
 {
 	check(pchDelim);
-	return UE::String::ParseIntoArray(OutArray, *this, pchDelim, InCullEmpty);
+	OutArray.Reset();
+	// TODO Legacy behavior: if DelimLength is 0 we return an empty array. Can we change this to return { Text }?
+	if (pchDelim[0] != '\0')
+	{
+		UE::String::EParseTokensOptions ParseOptions = UE::String::EParseTokensOptions::IgnoreCase |
+			(InCullEmpty ? UE::String::EParseTokensOptions::SkipEmpty : UE::String::EParseTokensOptions::None);
+		UE::String::ParseTokens(FStringView(*this), FStringView(pchDelim),
+			[&OutArray](FStringView Token) { OutArray.Emplace(Token); },
+			ParseOptions);
+	}
+	return OutArray.Num();
 }
 
 bool FString::MatchesWildcard(const TCHAR* InWildcard, int32 InWildcardLen, ESearchCase::Type SearchCase) const
diff --git a/Engine/Source/Runtime/Core/Private/Misc/Paths.cpp b/Engine/Source/Runtime/Core/Private/Misc/Paths.cpp
--- a/Engine/Source/Runtime/Core/Private/Misc/Paths.cpp
+++ b/Engine/Source/Runtime/Core/Private/Misc/Paths.cpp
@@ -18,7 +18,7 @@
 #include "Misc/LazySingleton.h"
 #include "Misc/Parse.h"
 #include "Misc/ScopeLock.h"
-#include "String/Parse.h"
+#include "String/ParseTokens.h"
 #include "UObject/NameTypes.h"
 
 DEFINE_LOG_CATEGORY_STATIC(LogPaths, Log, All);
@@ -1297,11 +1297,13 @@ bool FPaths::MakePathRelativeTo( FString& InPath, const TCHAR* InRelativeTo )
 	Source.ReplaceInline(TEXT("\\"), TEXT("/"), ESearchCase::CaseSensitive);
 	Target.ReplaceInline(TEXT("\\"), TEXT("/"), ESearchCase::CaseSensitive);
 
+	const UE::String::EParseTokensOptions ParseOptions = UE::String::EParseTokensOptions::IgnoreCase |
+		UE::String::EParseTokensOptions::SkipEmpty;
 	TArray<FStringView, TInlineAllocator<16>> TargetArrayBuffer;
-	UE::String::ParseIntoArray(TargetArrayBuffer, Target, TEXT("/"), true);
+	UE::String::ParseTokens(Target, TEXTVIEW("/"), TargetArrayBuffer, ParseOptions);
 	TArrayView<FStringView> TargetArray(TargetArrayBuffer);
 	TArray<FStringView, TInlineAllocator<16>> SourceArrayBuffer;
-	UE::String::ParseIntoArray(SourceArrayBuffer, Source, TEXT("/"), true);
+	UE::String::ParseTokens(Source, TEXTVIEW("/"), SourceArrayBuffer, ParseOptions);
 	TArrayView<FStringView> SourceArray(SourceArrayBuffer);
 
 	if (TargetArray.Num() && SourceArray.Num())
diff --git a/Engine/Source/Runtime/Core/Public/String/Parse.h b/Engine/Source/Runtime/Core/Public/String/Parse.h
deleted file mode 100644
