Add support for specifying priority when mounting iodispatcher and packagestore backends

diff --git a/Engine/Source/Runtime/Core/Private/IO/IoDispatcher.cpp b/Engine/Source/Runtime/Core/Private/IO/IoDispatcher.cpp
--- a/Engine/Source/Runtime/Core/Private/IO/IoDispatcher.cpp
+++ b/Engine/Source/Runtime/Core/Private/IO/IoDispatcher.cpp
@@ -19,6 +19,7 @@
 #include "ProfilingDebugging/CsvProfiler.h"
 #include "Containers/Ticker.h"
 #include "IO/IoDispatcherBackend.h"
+#include "Templates/Greater.h"
 
 DEFINE_LOG_CATEGORY(LogIoDispatcher);
 
@@ -345,9 +346,9 @@ public:
 
 	~FIoDispatcherImpl()
 	{
-		for (const TSharedRef<IIoDispatcherBackend>& Backend : Backends)
+		for (const FBackendAndPriority& Backend : Backends)
 		{
-			Backend->Shutdown();
+			Backend.Value->Shutdown();
 		}
 		FCoreDelegates::GetMemoryTrimDelegate().Remove(MemoryTrimDelegateHandle);
 		BackendContext->WakeUpDispatcherThreadDelegate.Unbind();
@@ -367,9 +368,9 @@ public:
 		bIsInitialized = true;
 		if (!Backends.IsEmpty())
 		{
-			for (const TSharedRef<IIoDispatcherBackend>& Backend : Backends)
+			for (const FBackendAndPriority& Backend : Backends)
 			{
-				Backend->Initialize(BackendContext);
+				Backend.Value->Initialize(BackendContext);
 			}
 			// If there are no mounted backends the resolve thread is not needed
 			StartThread();
@@ -424,9 +425,9 @@ public:
 	{
 		if (ChunkId.IsValid())
 		{
-			for (const TSharedRef<IIoDispatcherBackend>& Backend : Backends)
+			for (const FBackendAndPriority& Backend : Backends)
 			{
-				TIoStatusOr<FIoMappedRegion> Result = Backend->OpenMapped(ChunkId, Options);
+				TIoStatusOr<FIoMappedRegion> Result = Backend.Value->OpenMapped(ChunkId, Options);
 				if (Result.IsOk())
 				{
 					return Result;
@@ -440,10 +441,12 @@ public:
 		}
 	}
 
-	void Mount(TSharedRef<IIoDispatcherBackend> Backend)
+	void Mount(TSharedRef<IIoDispatcherBackend> Backend, int32 Priority)
 	{
 		check(IsInGameThread());
-		Backends.Add(Backend);
+
+		int32 Index = Algo::LowerBoundBy(Backends, Priority, &FBackendAndPriority::Key, TGreater<>());
+		Backends.Insert(MakeTuple(Priority, Backend), Index);
 		if (bIsInitialized)
 		{
 			Backend->Initialize(BackendContext);
@@ -456,9 +459,9 @@ public:
 
 	bool DoesChunkExist(const FIoChunkId& ChunkId) const
 	{
-		for (const TSharedRef<IIoDispatcherBackend>& Backend : Backends)
+		for (const FBackendAndPriority& Backend : Backends)
 		{
-			if (Backend->DoesChunkExist(ChunkId))
+			if (Backend.Value->DoesChunkExist(ChunkId))
 			{
 				return true;
 			}
@@ -471,9 +474,9 @@ public:
 		// Only attempt to find the size if the FIoChunkId is valid
 		if (ChunkId.IsValid())
 		{
-			for (const TSharedRef<IIoDispatcherBackend>& Backend : Backends)
+			for (const FBackendAndPriority& Backend : Backends)
 			{
-				TIoStatusOr<uint64> Result = Backend->GetSizeForChunk(ChunkId);
+				TIoStatusOr<uint64> Result = Backend.Value->GetSizeForChunk(ChunkId);
 				if (Result.IsOk())
 				{
 					return Result;
@@ -615,9 +618,9 @@ private:
 	{
 		//TRACE_CPUPROFILER_EVENT_SCOPE(ProcessCompletedRequests);
 
-		for (const TSharedRef<IIoDispatcherBackend>& Backend : Backends)
+		for (const FBackendAndPriority& Backend : Backends)
 		{
-			FIoRequestImpl* CompletedRequestsHead = Backend->GetCompletedRequests();
+			FIoRequestImpl* CompletedRequestsHead = Backend.Value->GetCompletedRequests();
 			while (CompletedRequestsHead)
 			{
 				FIoRequestImpl* NextRequest = CompletedRequestsHead->NextRequest;
@@ -775,12 +778,12 @@ private:
 			{
 				TRACE_CPUPROFILER_EVENT_SCOPE(ResolveRequest);
 				bool bResolved = false;
-				for (const TSharedRef<IIoDispatcherBackend>& Backend : Backends)
+				for (const FBackendAndPriority& Backend : Backends)
 				{
-					if (Backend->Resolve(Request))
+					if (Backend.Value->Resolve(Request))
 					{
 						bResolved = true;
-						Request->Backend = &Backend.Get();
+						Request->Backend = &Backend.Value.Get();
 						break;
 					}
 				}
@@ -836,10 +839,11 @@ private:
 	}
 
 	using FBatchAllocator = TBlockAllocator<FIoBatchImpl, 4096>;
+	using FBackendAndPriority = TTuple<int32, TSharedRef<IIoDispatcherBackend>>;
 
 	TSharedRef<FIoDispatcherBackendContext> BackendContext;
 	FDelegateHandle MemoryTrimDelegateHandle;
-	TArray<TSharedRef<IIoDispatcherBackend>> Backends;
+	TArray<FBackendAndPriority> Backends;
 	FIoRequestAllocator* RequestAllocator = nullptr;
 	FBatchAllocator BatchAllocator;
 	FRunnableThread* Thread = nullptr;
@@ -869,9 +873,9 @@ FIoDispatcher::~FIoDispatcher()
 }
 
 void
-FIoDispatcher::Mount(TSharedRef<IIoDispatcherBackend> Backend)
+FIoDispatcher::Mount(TSharedRef<IIoDispatcherBackend> Backend, int32 Priority)
 {
-	Impl->Mount(Backend);
+	Impl->Mount(Backend, Priority);
 }
 
 FIoBatch
diff --git a/Engine/Source/Runtime/Core/Private/IO/PackageStore.cpp b/Engine/Source/Runtime/Core/Private/IO/PackageStore.cpp
--- a/Engine/Source/Runtime/Core/Private/IO/PackageStore.cpp
+++ b/Engine/Source/Runtime/Core/Private/IO/PackageStore.cpp
@@ -12,6 +12,7 @@
 #include "Serialization/Archive.h"
 #include "Serialization/CompactBinary.h"
 #include "Serialization/CompactBinaryWriter.h"
+#include "Templates/Greater.h"
 
 FArchive& operator<<(FArchive& Ar, FPackageStoreExportInfo& ExportInfo)
 {
@@ -132,9 +133,9 @@ FPackageStoreReadScope::FPackageStoreReadScope(FPackageStore& InPackageStore)
 {
 	if (!PackageStore.ThreadReadCount)
 	{
-		for (const TSharedRef<IPackageStoreBackend>& Backend : PackageStore.Backends)
+		for (const FPackageStore::FBackendAndPriority& Backend : PackageStore.Backends)
 		{
-			Backend->BeginRead();
+			Backend.Value->BeginRead();
 		}
 	}
 	++PackageStore.ThreadReadCount;
@@ -145,9 +146,9 @@ FPackageStoreReadScope::~FPackageStoreReadScope()
 	check(PackageStore.ThreadReadCount > 0);
 	if (--PackageStore.ThreadReadCount == 0)
 	{
-		for (const TSharedRef<IPackageStoreBackend>& Backend : PackageStore.Backends)
+		for (const FPackageStore::FBackendAndPriority& Backend : PackageStore.Backends)
 		{
-			Backend->EndRead();
+			Backend.Value->EndRead();
 		}
 	}
 }
@@ -165,19 +166,20 @@ FPackageStore& FPackageStore::Get()
 }
 
 
-void FPackageStore::Mount(TSharedRef<IPackageStoreBackend> Backend)
+void FPackageStore::Mount(TSharedRef<IPackageStoreBackend> Backend, int32 Priority)
 {
 	check(IsInGameThread());
-	Backends.Add(Backend);
+	int32 Index = Algo::LowerBoundBy(Backends, Priority, &FBackendAndPriority::Key, TGreater<>());
+	Backends.Insert(MakeTuple(Priority, Backend), Index);
 	Backend->OnMounted(BackendContext);
 }
 
 EPackageStoreEntryStatus FPackageStore::GetPackageStoreEntry(FPackageId PackageId, FPackageStoreEntry& OutPackageStoreEntry)
 {
 	check(ThreadReadCount);
-	for (const TSharedRef<IPackageStoreBackend>& Backend : Backends)
+	for (const FBackendAndPriority& Backend : Backends)
 	{
-		EPackageStoreEntryStatus Status = Backend->GetPackageStoreEntry(PackageId, OutPackageStoreEntry);
+		EPackageStoreEntryStatus Status = Backend.Value->GetPackageStoreEntry(PackageId, OutPackageStoreEntry);
 		if (Status >= EPackageStoreEntryStatus::Pending)
 		{
 			return Status;
@@ -189,9 +191,9 @@ EPackageStoreEntryStatus FPackageStore::GetPackageStoreEntry(FPackageId PackageI
 bool FPackageStore::GetPackageRedirectInfo(FPackageId PackageId, FName& OutSourcePackageName, FPackageId& OutRedirectedToPackageId)
 {
 	check(ThreadReadCount);
-	for (const TSharedRef<IPackageStoreBackend>& Backend : Backends)
+	for (const FBackendAndPriority& Backend : Backends)
 	{
-		if (Backend->GetPackageRedirectInfo(PackageId, OutSourcePackageName, OutRedirectedToPackageId))
+		if (Backend.Value->GetPackageRedirectInfo(PackageId, OutSourcePackageName, OutRedirectedToPackageId))
 		{
 			return true;
 		}
diff --git a/Engine/Source/Runtime/Core/Public/IO/IoDispatcher.h b/Engine/Source/Runtime/Core/Public/IO/IoDispatcher.h
--- a/Engine/Source/Runtime/Core/Public/IO/IoDispatcher.h
+++ b/Engine/Source/Runtime/Core/Public/IO/IoDispatcher.h
@@ -840,7 +840,7 @@ public:
 	CORE_API						FIoDispatcher();
 	CORE_API						~FIoDispatcher();
 
-	CORE_API void					Mount(TSharedRef<IIoDispatcherBackend> Backend);
+	CORE_API void					Mount(TSharedRef<IIoDispatcherBackend> Backend, int32 Priority = 0);
 
 	CORE_API FIoBatch				NewBatch();
 
diff --git a/Engine/Source/Runtime/Core/Public/IO/PackageStore.h b/Engine/Source/Runtime/Core/Public/IO/PackageStore.h
--- a/Engine/Source/Runtime/Core/Public/IO/PackageStore.h
+++ b/Engine/Source/Runtime/Core/Public/IO/PackageStore.h
@@ -224,7 +224,7 @@ public:
 	CORE_API static FPackageStore& Get();
 
 	/* Mount a package store backend. */
-	CORE_API void Mount(TSharedRef<IPackageStoreBackend> Backend);
+	CORE_API void Mount(TSharedRef<IPackageStoreBackend> Backend, int32 Priority = 0);
 
 	/* Returns the package store entry data with export info and imported packages for the specified package ID. */
 	CORE_API EPackageStoreEntryStatus GetPackageStoreEntry(FPackageId PackageId, FPackageStoreEntry& OutPackageStoreEntry);
@@ -240,7 +240,9 @@ private:
 	friend class FPackageStoreReadScope;
 
 	TSharedRef<FPackageStoreBackendContext> BackendContext;
-	TArray<TSharedRef<IPackageStoreBackend>> Backends;
+	
+	using FBackendAndPriority = TTuple<int32, TSharedRef<IPackageStoreBackend>>;
+	TArray<FBackendAndPriority> Backends;
 
 	static thread_local int32 ThreadReadCount;
 };
