Adjust network version setup due to static init ordering

diff --git a/Engine/Source/Runtime/Core/Private/Misc/NetworkVersion.cpp b/Engine/Source/Runtime/Core/Private/Misc/NetworkVersion.cpp
--- a/Engine/Source/Runtime/Core/Private/Misc/NetworkVersion.cpp
+++ b/Engine/Source/Runtime/Core/Private/Misc/NetworkVersion.cpp
@@ -24,10 +24,65 @@ FNetworkVersion::FGetLocalNetworkVersionOverride FNetworkVersion::GetLocalNetwor
 FNetworkVersion::FIsNetworkCompatibleOverride FNetworkVersion::IsNetworkCompatibleOverride;
 FNetworkVersion::FGetReplayCompatibleChangeListOverride FNetworkVersion::GetReplayCompatibleChangeListOverride;
 
-const FGuid FEngineNetworkCustomVersion::Guid = FGuid(0x62915CA3, 0x1C8E4BF7, 0xA30E12C7, 0xC8219DF7);
+namespace UE::Net::Private
+{
+	FGuid GetEngineNetworkVersionGuid()
+	{
+		static FGuid Guid = FGuid(0x62915CA3, 0x1C8E4BF7, 0xA30E12C7, 0xC8219DF7);
+		return Guid;
+	}
+
+	FGuid GetGameNetworkVersionGuid()
+	{
+		static FGuid Guid = FGuid(0xCC400D24, 0xE0E94E7B, 0x9BF9A283, 0xDCC0C027);
+		return Guid;
+	}
+
+	FCustomVersionContainer& GetNetworkCustomVersions()
+	{
+		static FCustomVersionContainer NetworkCustomVersions;
+
+		static bool bInitialized = false;
+
+		if (!bInitialized)
+		{
+			const FGuid EngineNetworkGuid = UE::Net::Private::GetEngineNetworkVersionGuid();
+			const FGuid GameNetworkGuid = UE::Net::Private::GetGameNetworkVersionGuid();
+
+			NetworkCustomVersions.SetVersion(EngineNetworkGuid, FEngineNetworkCustomVersion::LatestVersion, TEXT("EngineNetworkVersion"));
+			NetworkCustomVersions.SetVersion(GameNetworkGuid, FGameNetworkCustomVersion::LatestVersion, TEXT("GameNetworkVersion"));
+
+			bInitialized = true;
+		}
+
+		return NetworkCustomVersions;
+	}
+
+	FCustomVersionContainer& GetCompatibleNetworkCustomVersions()
+	{
+		static FCustomVersionContainer CompatibleNetworkCustomVersions;
+
+		static bool bInitialized = false;
+
+		if (!bInitialized)
+		{
+			const FGuid EngineNetworkGuid = UE::Net::Private::GetEngineNetworkVersionGuid();
+			const FGuid GameNetworkGuid = UE::Net::Private::GetGameNetworkVersionGuid();
+
+			CompatibleNetworkCustomVersions.SetVersion(EngineNetworkGuid, FEngineNetworkCustomVersion::ReplayBackwardsCompat, TEXT("EngineNetworkVersion"));
+			CompatibleNetworkCustomVersions.SetVersion(GameNetworkGuid, FGameNetworkCustomVersion::LatestVersion, TEXT("GameNetworkVersion"));
+
+			bInitialized = true;
+		}
+
+		return CompatibleNetworkCustomVersions;
+	}
+}
+
+const FGuid FEngineNetworkCustomVersion::Guid = UE::Net::Private::GetEngineNetworkVersionGuid();
 FCustomVersionRegistration GRegisterEngineNetworkCustomVersion(FEngineNetworkCustomVersion::Guid, FEngineNetworkCustomVersion::LatestVersion, TEXT("EngineNetworkVersion"));
 
-const FGuid FGameNetworkCustomVersion::Guid = FGuid(0xCC400D24, 0xE0E94E7B, 0x9BF9A283, 0xDCC0C027);
+const FGuid FGameNetworkCustomVersion::Guid = UE::Net::Private::GetGameNetworkVersionGuid();
 FCustomVersionRegistration GRegisterGameNetworkCustomVersion(FGameNetworkCustomVersion::Guid, FGameNetworkCustomVersion::LatestVersion, TEXT("GameNetworkVersion"));
 
 FString& FNetworkVersion::GetProjectVersion_Internal()
@@ -48,11 +103,6 @@ uint32 FNetworkVersion::GameNetworkProtocolVersion		= 0;
 uint32 FNetworkVersion::EngineCompatibleNetworkProtocolVersion		= FEngineNetworkCustomVersion::ReplayBackwardsCompat;
 uint32 FNetworkVersion::GameCompatibleNetworkProtocolVersion		= 0;
 
-FCustomVersionContainer FNetworkVersion::NetworkCustomVersions;
-FCustomVersionContainer FNetworkVersion::CompatibleNetworkCustomVersions;
-
-bool FNetworkVersion::bInitializedCustomVersions = false;
-
 void FNetworkVersion::SetProjectVersion(const TCHAR* InVersion)
 {
 	if (ensureMsgf(InVersion != nullptr && FCString::Strlen(InVersion), TEXT("ProjectVersion used for network version must be a valid string!")))
@@ -68,7 +118,7 @@ void FNetworkVersion::SetProjectVersion(const TCHAR* InVersion)
 
 void FNetworkVersion::SetGameNetworkProtocolVersion(const uint32 InGameNetworkProtocolVersion)
 {
-	InitCustomVersions();
+	FCustomVersionContainer& NetworkCustomVersions = UE::Net::Private::GetNetworkCustomVersions();
 
 	PRAGMA_DISABLE_DEPRECATION_WARNINGS
 	GameNetworkProtocolVersion = InGameNetworkProtocolVersion;
@@ -81,7 +131,7 @@ void FNetworkVersion::SetGameNetworkProtocolVersion(const uint32 InGameNetworkPr
 
 void FNetworkVersion::SetGameCompatibleNetworkProtocolVersion(const uint32 InGameCompatibleNetworkProtocolVersion)
 {
-	InitCustomVersions();
+	FCustomVersionContainer& CompatibleNetworkCustomVersions = UE::Net::Private::GetCompatibleNetworkCustomVersions();
 
 	PRAGMA_DISABLE_DEPRECATION_WARNINGS
 	GameCompatibleNetworkProtocolVersion = InGameCompatibleNetworkProtocolVersion;
@@ -136,16 +186,14 @@ uint32 FNetworkVersion::GetReplayCompatibleChangelist()
 
 uint32 FNetworkVersion::GetNetworkProtocolVersion(const FGuid& VersionGuid)
 {
-	InitCustomVersions();
-
+	const FCustomVersionContainer& NetworkCustomVersions = UE::Net::Private::GetNetworkCustomVersions();
 	const FCustomVersion* CustomVersion = NetworkCustomVersions.GetVersion(VersionGuid);
 	return CustomVersion ? CustomVersion->Version : 0;
 }
 
 uint32 FNetworkVersion::GetCompatibleNetworkProtocolVersion(const FGuid& VersionGuid)
 {
-	InitCustomVersions();
-
+	const FCustomVersionContainer& CompatibleNetworkCustomVersions = UE::Net::Private::GetCompatibleNetworkCustomVersions();
 	const FCustomVersion* CustomVersion = CompatibleNetworkCustomVersions.GetVersion(VersionGuid);
 	return CustomVersion ? CustomVersion->Version : 0;
 }
@@ -172,30 +220,14 @@ uint32 FNetworkVersion::GetGameCompatibleNetworkProtocolVersion()
 
 const FCustomVersionContainer& FNetworkVersion::GetNetworkCustomVersions()
 {
-	InitCustomVersions();
-
-	return NetworkCustomVersions;
-}
-
-void FNetworkVersion::InitCustomVersions()
-{
-	if (!bInitializedCustomVersions)
-	{
-		FCustomVersion RegisteredEngineVersion = FCurrentCustomVersions::Get(FEngineNetworkCustomVersion::Guid).GetValue();
-		FCustomVersion RegisteredGameVersion = FCurrentCustomVersions::Get(FGameNetworkCustomVersion::Guid).GetValue();
-
-		NetworkCustomVersions.SetVersion(FEngineNetworkCustomVersion::Guid, RegisteredEngineVersion.Version, RegisteredEngineVersion.GetFriendlyName());
-		NetworkCustomVersions.SetVersion(FGameNetworkCustomVersion::Guid, RegisteredGameVersion.Version, RegisteredGameVersion.GetFriendlyName());
-
-		CompatibleNetworkCustomVersions.SetVersion(FEngineNetworkCustomVersion::Guid, FEngineNetworkCustomVersion::ReplayBackwardsCompat, RegisteredEngineVersion.GetFriendlyName());
-		CompatibleNetworkCustomVersions.SetVersion(FGameNetworkCustomVersion::Guid, RegisteredGameVersion.Version, RegisteredGameVersion.GetFriendlyName());
-
-		bInitializedCustomVersions = true;
-	}
+	return UE::Net::Private::GetNetworkCustomVersions();
 }
 
 void FNetworkVersion::RegisterNetworkCustomVersion(const FGuid& VersionGuid, int32 Version, int32 CompatibleVersion, const FName& FriendlyName)
 {
+	FCustomVersionContainer& NetworkCustomVersions = UE::Net::Private::GetNetworkCustomVersions();
+	FCustomVersionContainer& CompatibleNetworkCustomVersions = UE::Net::Private::GetCompatibleNetworkCustomVersions();
+
 	if (const FCustomVersion* ExistingVersion = NetworkCustomVersions.GetVersion(VersionGuid))
 	{
 		if (ExistingVersion->Version != Version)
@@ -238,13 +270,13 @@ uint32 FNetworkVersion::GetLocalNetworkVersion( bool AllowOverrideDelegate /*=tr
 		return CachedNetworkChecksum;
 	}
 
-	InitCustomVersions();
-
 	FString VersionString = FString::Printf(TEXT("%s %s, NetCL: %d"),
 		FApp::GetProjectName(),
 		*FNetworkVersion::GetProjectVersion(),
 		GetNetworkCompatibleChangelist());
 
+	const FCustomVersionContainer& NetworkCustomVersions = UE::Net::Private::GetNetworkCustomVersions();
+
 	const FCustomVersionArray& CustomVers = NetworkCustomVersions.GetAllVersions();
 	for (const FCustomVersion& CustomVer : CustomVers)
 	{
@@ -276,6 +308,8 @@ FNetworkReplayVersion FNetworkVersion::GetReplayVersion()
 	{
 		FString VersionString = FApp::GetProjectName();
 
+		const FCustomVersionContainer& CompatibleNetworkCustomVersions = UE::Net::Private::GetCompatibleNetworkCustomVersions();
+
 		const FCustomVersionArray& CustomVers = CompatibleNetworkCustomVersions.GetAllVersions();
 		for (const FCustomVersion& CustomVer : CustomVers)
 		{
diff --git a/Engine/Source/Runtime/Core/Public/Misc/NetworkVersion.h b/Engine/Source/Runtime/Core/Public/Misc/NetworkVersion.h
--- a/Engine/Source/Runtime/Core/Public/Misc/NetworkVersion.h
+++ b/Engine/Source/Runtime/Core/Public/Misc/NetworkVersion.h
@@ -65,9 +65,9 @@ struct CORE_API FNetworkVersion
 	static uint32 GetNetworkCompatibleChangelist();
 	static uint32 GetReplayCompatibleChangelist();
 
-	UE_DEPRECATED(5.2, "Please use GetCompatibleNetworkProtocolVersion instead.")
+	UE_DEPRECATED(5.2, "Please use GetNetworkProtocolVersion instead.")
 	static uint32 GetEngineNetworkProtocolVersion();
-	UE_DEPRECATED(5.2, "Please use GetCompatibleNetworkProtocolVersion instead.")
+	UE_DEPRECATED(5.2, "Please use GetNetworkProtocolVersion instead.")
 	static uint32 GetGameNetworkProtocolVersion();
 
 	UE_DEPRECATED(5.2, "Please use GetCompatibleNetworkProtocolVersion instead.")
@@ -168,11 +168,4 @@ protected:
 	static uint32	EngineCompatibleNetworkProtocolVersion;
 	UE_DEPRECATED(5.2, "Now storing this value in CompatibleNetworkCustomVersions, do not use directly.")
 	static uint32	GameCompatibleNetworkProtocolVersion;
-
-private:
-	static FCustomVersionContainer NetworkCustomVersions;
-	static FCustomVersionContainer CompatibleNetworkCustomVersions;
-
-	static bool bInitializedCustomVersions;
-	static void InitCustomVersions();
 };
