FPaths::MakePathRelativeTo: Fix negative number passed into reserve.
#rnx

diff --git a/Engine/Source/Runtime/Core/Private/Misc/Paths.cpp b/Engine/Source/Runtime/Core/Private/Misc/Paths.cpp
--- a/Engine/Source/Runtime/Core/Private/Misc/Paths.cpp
+++ b/Engine/Source/Runtime/Core/Private/Misc/Paths.cpp
@@ -1327,7 +1327,7 @@ bool FPaths::MakePathRelativeTo( FString& InPath, const TCHAR* InRelativeTo )
 
 	FStringView ParentDirSeparator = TEXTVIEW("../");
 	FStringView DirSeparator = TEXTVIEW("/");
-	int32 ResultsLen = SourceArray.Num() * ParentDirSeparator.Len() + (TargetArray.Num() - 1)* DirSeparator.Len();
+	int32 ResultsLen = SourceArray.Num() * ParentDirSeparator.Len() + (FMath::Max(TargetArray.Num(),1) - 1)*DirSeparator.Len();
 	for (const FStringView& TargetDir : TargetArray)
 	{
 		ResultsLen += TargetDir.Len();
