Tag allocations related to Editor - Framework components

Added LLM tags for many allocations in Editor Framework components.  Used best judgement for use of existing tags.
Added EngineMisc/MessageLog tag to capture allocations that the logger does since it is a well defined core system that all other modules tend to use. This will pull the disparate allocations within the logger attributed to a lot of other tags into it's own bespoke tag.
Added a LevelEditor tag since it is also a fairly large component. There isn't a generic 'Editor' tag, though I think this might be worthwhile to explore as there are other editors in UEd (PropertyEditor for example).

diff --git a/Engine/Plugins/Editor/ContentBrowser/ContentBrowserAliasDataSoure/Source/ContentBrowserAliasDataSource/Private/ContentBrowserAliasDataSource.cpp b/Engine/Plugins/Editor/ContentBrowser/ContentBrowserAliasDataSoure/Source/ContentBrowserAliasDataSource/Private/ContentBrowserAliasDataSource.cpp
--- a/Engine/Plugins/Editor/ContentBrowser/ContentBrowserAliasDataSoure/Source/ContentBrowserAliasDataSource/Private/ContentBrowserAliasDataSource.cpp
+++ b/Engine/Plugins/Editor/ContentBrowser/ContentBrowserAliasDataSoure/Source/ContentBrowserAliasDataSource/Private/ContentBrowserAliasDataSource.cpp
@@ -832,6 +832,7 @@ FContentBrowserItemData UContentBrowserAliasDataSource::CreateAssetFileItem(cons
 	// See UContentBrowserAliasDataSource::DoesItemPassFilter for more information on how this could fail
 	if (AliasData)
 	{
+		LLM_SCOPE(ELLMTag::UI);
 		FName VirtualizedPath;
 PRAGMA_DISABLE_DEPRECATION_WARNINGS
 		TryConvertInternalPathToVirtual(AliasData->ObjectPath.ToFName(), VirtualizedPath);
diff --git a/Engine/Source/Developer/MessageLog/Private/MessageLogModule.cpp b/Engine/Source/Developer/MessageLog/Private/MessageLogModule.cpp
--- a/Engine/Source/Developer/MessageLog/Private/MessageLogModule.cpp
+++ b/Engine/Source/Developer/MessageLog/Private/MessageLogModule.cpp
@@ -27,6 +27,7 @@ IMPLEMENT_MODULE( FMessageLogModule, MessageLog );
 
 TSharedRef<SDockTab> SpawnMessageLog( const FSpawnTabArgs& Args, TSharedRef<FMessageLogViewModel> MessageLogViewModel )
 {
+	LLM_SCOPE(ELLMTag::UI);
 	return SNew(SDockTab)
 		.TabRole(ETabRole::NomadTab)
 		[
@@ -80,16 +81,19 @@ void FMessageLogModule::ShutdownModule()
 
 TSharedRef<IMessageLogListing> FMessageLogModule::GetLogListing(const FName& LogName)
 {
+	LLM_SCOPE_BYTAG(EngineMisc_MessageLog);
 	return MessageLogViewModel->GetLogListingViewModel(LogName);
 }
 
 void FMessageLogModule::RegisterLogListing(const FName& LogName, const FText& LogLabel, const FMessageLogInitializationOptions& InitializationOptions)
 {
+	LLM_SCOPE_BYTAG(EngineMisc_MessageLog);
 	MessageLogViewModel->RegisterLogListingViewModel(LogName, LogLabel, InitializationOptions);
 }
 
 bool FMessageLogModule::UnregisterLogListing(const FName& LogName)
 {
+	LLM_SCOPE_BYTAG(EngineMisc_MessageLog);
 	return MessageLogViewModel->UnregisterLogListingViewModel(LogName);
 }	
 
@@ -100,12 +104,14 @@ bool FMessageLogModule::IsRegisteredLogListing(const FName& LogName) const
 
 TSharedRef<IMessageLogListing> FMessageLogModule::CreateLogListing(const FName& InLogName, const FMessageLogInitializationOptions& InitializationOptions)
 {
+	LLM_SCOPE_BYTAG(EngineMisc_MessageLog);
 	TSharedRef<FMessageLogListingModel> MessageLogListingModel = FMessageLogListingModel::Create( InLogName );
 	return FMessageLogListingViewModel::Create( MessageLogListingModel, FText(), InitializationOptions );
 }
 
 TSharedRef<SWidget> FMessageLogModule::CreateLogListingWidget(const TSharedRef<IMessageLogListing>& InMessageLogListing)
 {
+	LLM_SCOPE(ELLMTag::UI);
 	return SNew(SMessageLogListing, InMessageLogListing);
 }
 
@@ -130,6 +136,7 @@ void FMessageLogModule::EnableMessageLogDisplay(bool bInCanDisplayMessageLog)
 
 void FMessageLogModule::RegisterMessageLogSpawner(const TSharedRef<FWorkspaceItem>& InGroup)
 {
+	LLM_SCOPE(ELLMTag::UI);
 	FGlobalTabmanager::Get()->RegisterNomadTabSpawner("MessageLog", FOnSpawnTab::CreateStatic(&SpawnMessageLog, MessageLogViewModel.ToSharedRef()))
 		.SetDisplayName(NSLOCTEXT("UnrealEditor", "MessageLogTab", "Message Log"))
 		.SetTooltipText(NSLOCTEXT("UnrealEditor", "MessageLogTooltipText", "Open the Message Log tab."))
diff --git a/Engine/Source/Editor/LevelEditor/Private/LevelEditor.cpp b/Engine/Source/Editor/LevelEditor/Private/LevelEditor.cpp
--- a/Engine/Source/Editor/LevelEditor/Private/LevelEditor.cpp
+++ b/Engine/Source/Editor/LevelEditor/Private/LevelEditor.cpp
@@ -49,6 +49,7 @@
 #define LOCTEXT_NAMESPACE "LevelEditor"
 
 IMPLEMENT_MODULE( FLevelEditorModule, LevelEditor );
+LLM_DEFINE_TAG(LevelEditor);
 
 const FName LevelEditorApp = FName(TEXT("LevelEditorApp"));
 const FName MainFrame("MainFrame");
@@ -195,6 +196,8 @@ private:
 
 TSharedRef<SDockTab> FLevelEditorModule::SpawnLevelEditor( const FSpawnTabArgs& InArgs )
 {
+	LLM_SCOPE_BYTAG(LevelEditor);
+	
 	TSharedRef<SDockTab> LevelEditorTab = SNew(SDockTab)
 		.TabRole(ETabRole::MajorTab)
 		.ContentPadding(FMargin(0))
diff --git a/Engine/Source/Editor/PropertyEditor/Private/PropertyEditorModule.cpp b/Engine/Source/Editor/PropertyEditor/Private/PropertyEditorModule.cpp
--- a/Engine/Source/Editor/PropertyEditor/Private/PropertyEditorModule.cpp
+++ b/Engine/Source/Editor/PropertyEditor/Private/PropertyEditorModule.cpp
@@ -253,6 +253,8 @@ TSharedPtr<FAssetThumbnailPool> FPropertyEditorModule::GetThumbnailPool()
 
 TSharedRef<IDetailsView> FPropertyEditorModule::CreateDetailView( const FDetailsViewArgs& DetailsViewArgs )
 {
+	LLM_SCOPE(ELLMTag::UI);
+
 	// Compact the list of detail view instances
 	for( int32 ViewIndex = 0; ViewIndex < AllDetailViews.Num(); ++ViewIndex )
 	{
diff --git a/Engine/Source/Editor/StatusBar/Private/StatusBarSubsystem.cpp b/Engine/Source/Editor/StatusBar/Private/StatusBarSubsystem.cpp
--- a/Engine/Source/Editor/StatusBar/Private/StatusBarSubsystem.cpp
+++ b/Engine/Source/Editor/StatusBar/Private/StatusBarSubsystem.cpp
@@ -440,6 +440,8 @@ bool UStatusBarSubsystem::ToggleContentBrowser(TSharedRef<SWindow> ParentWindow)
 
 TSharedRef<SWidget> UStatusBarSubsystem::MakeStatusBarWidget(FName StatusBarName, const TSharedRef<SDockTab>& InParentTab)
 {
+	LLM_SCOPE(ELLMTag::UI);
+
 	CreateContentBrowserIfNeeded();
 
 	TSharedRef<SStatusBar> StatusBar =
diff --git a/Engine/Source/Runtime/Core/Private/Logging/MessageLog.cpp b/Engine/Source/Runtime/Core/Private/Logging/MessageLog.cpp
--- a/Engine/Source/Runtime/Core/Private/Logging/MessageLog.cpp
+++ b/Engine/Source/Runtime/Core/Private/Logging/MessageLog.cpp
@@ -10,6 +10,8 @@ FMessageLog::FMessageSelectionChanged FMessageLog::MessageSelectionChanged;
 
 #define LOCTEXT_NAMESPACE "MessageLog"
 
+LLM_DEFINE_TAG(EngineMisc_MessageLog);
+
 class FBasicMessageLog : public IMessageLog, public TSharedFromThis<FBasicMessageLog>
 {
 public:
@@ -94,6 +96,8 @@ private:
 FMessageLog::FMessageLog( const FName& InLogName )
 	: bSuppressLoggingToOutputLog(false)
 {
+	LLM_SCOPE_BYTAG(EngineMisc_MessageLog);
+	
 	if(GetLog.IsBound())
 	{
 		MessageLog = GetLog.Execute(InLogName);
@@ -106,11 +110,15 @@ FMessageLog::FMessageLog( const FName& InLogName )
 
 FMessageLog::~FMessageLog()
 {
+	LLM_SCOPE_BYTAG(EngineMisc_MessageLog);
+	
 	Flush();
 }
 
 const TSharedRef<FTokenizedMessage>& FMessageLog::AddMessage( const TSharedRef<FTokenizedMessage>& InMessage )
 {
+	LLM_SCOPE_BYTAG(EngineMisc_MessageLog);
+	
 	Messages.Add(InMessage);
 
 	return InMessage;
@@ -118,11 +126,15 @@ const TSharedRef<FTokenizedMessage>& FMessageLog::AddMessage( const TSharedRef<F
 
 void FMessageLog::AddMessages( const TArray< TSharedRef<FTokenizedMessage> >& InMessages )
 {
+	LLM_SCOPE_BYTAG(EngineMisc_MessageLog);
+	
 	Messages.Append(InMessages);
 }
 
 TSharedRef<FTokenizedMessage> FMessageLog::Message( EMessageSeverity::Type InSeverity, const FText& InMessage )
 {
+	LLM_SCOPE_BYTAG(EngineMisc_MessageLog);
+	
 	TSharedRef<FTokenizedMessage> Message = FTokenizedMessage::Create(InSeverity, InMessage);
 	Messages.Add(Message);
 	return Message;
@@ -135,6 +147,8 @@ TSharedRef<FTokenizedMessage> FMessageLog::CriticalError( const FText& InMessage
 	checkf(false, 
 		TEXT("[Deprecated] FMessageLog::CriticalError has been deprecated. Please update callsites. The original error message: \n%s"), 
 		*InMessage.ToString());
+
+	LLM_SCOPE_BYTAG(EngineMisc_MessageLog);
 	
 	TSharedRef<FTokenizedMessage> Message = FTokenizedMessage::Create(EMessageSeverity::Error, InMessage);
 	Messages.Add(Message);
@@ -143,6 +157,8 @@ TSharedRef<FTokenizedMessage> FMessageLog::CriticalError( const FText& InMessage
 
 TSharedRef<FTokenizedMessage> FMessageLog::Error( const FText& InMessage )
 {
+	LLM_SCOPE_BYTAG(EngineMisc_MessageLog);
+	
 	TSharedRef<FTokenizedMessage> Message = FTokenizedMessage::Create(EMessageSeverity::Error, InMessage);
 	Messages.Add(Message);
 	return Message;
@@ -150,6 +166,8 @@ TSharedRef<FTokenizedMessage> FMessageLog::Error( const FText& InMessage )
 
 TSharedRef<FTokenizedMessage> FMessageLog::PerformanceWarning( const FText& InMessage )
 {
+	LLM_SCOPE_BYTAG(EngineMisc_MessageLog);
+	
 	TSharedRef<FTokenizedMessage> Message = FTokenizedMessage::Create(EMessageSeverity::PerformanceWarning, InMessage);
 #if !PLATFORM_UNIX  // @todo: these are too spammy for now on Linux
 	Messages.Add(Message);
@@ -159,6 +177,8 @@ TSharedRef<FTokenizedMessage> FMessageLog::PerformanceWarning( const FText& InMe
 
 TSharedRef<FTokenizedMessage> FMessageLog::Warning( const FText& InMessage )
 {
+	LLM_SCOPE_BYTAG(EngineMisc_MessageLog);
+	
 	TSharedRef<FTokenizedMessage> Message = FTokenizedMessage::Create(EMessageSeverity::Warning, InMessage);
 #if !PLATFORM_UNIX // @todo: these are too spammy for now on Linux
 	Messages.Add(Message); // TODO These are too spammy for now
@@ -168,6 +188,8 @@ TSharedRef<FTokenizedMessage> FMessageLog::Warning( const FText& InMessage )
 
 TSharedRef<FTokenizedMessage> FMessageLog::Info( const FText& InMessage )
 {
+	LLM_SCOPE_BYTAG(EngineMisc_MessageLog);
+	
 	TSharedRef<FTokenizedMessage> Message = FTokenizedMessage::Create(EMessageSeverity::Info, InMessage);
 	Messages.Add(Message);
 	return Message;
@@ -175,12 +197,16 @@ TSharedRef<FTokenizedMessage> FMessageLog::Info( const FText& InMessage )
 
 int32 FMessageLog::NumMessages( EMessageSeverity::Type InSeverityFilter )
 {
+	LLM_SCOPE_BYTAG(EngineMisc_MessageLog);
+	
 	Flush();
 	return MessageLog->NumMessages(InSeverityFilter);
 }
 
 void FMessageLog::Open( EMessageSeverity::Type InSeverityFilter, bool bOpenEvenIfEmpty )
 {
+	LLM_SCOPE_BYTAG(EngineMisc_MessageLog);
+	
 	Flush();
 	if(bOpenEvenIfEmpty)
 	{
@@ -194,18 +220,24 @@ void FMessageLog::Open( EMessageSeverity::Type InSeverityFilter, bool bOpenEvenI
 
 void FMessageLog::Notify( const FText& InMessage, EMessageSeverity::Type InSeverityFilter, bool bForce )
 {
+	LLM_SCOPE_BYTAG(EngineMisc_MessageLog);
+	
 	Flush();
 	MessageLog->NotifyIfAnyMessages(InMessage, InSeverityFilter, bForce);
 }
 
 void FMessageLog::NewPage( const FText& InLabel )
 {
+	LLM_SCOPE_BYTAG(EngineMisc_MessageLog);
+	
 	Flush();
 	MessageLog->NewPage(InLabel);
 }
 
 void FMessageLog::SetCurrentPage(const FText& InLabel)
 {
+	LLM_SCOPE_BYTAG(EngineMisc_MessageLog);
+	
 	Flush();
 	MessageLog->SetCurrentPage(InLabel);
 }
diff --git a/Engine/Source/Runtime/Core/Public/Logging/MessageLog.h b/Engine/Source/Runtime/Core/Public/Logging/MessageLog.h
--- a/Engine/Source/Runtime/Core/Public/Logging/MessageLog.h
+++ b/Engine/Source/Runtime/Core/Public/Logging/MessageLog.h
@@ -10,9 +10,12 @@
 #include "Logging/TokenizedMessage.h"
 #include "Templates/SharedPointer.h"
 #include "UObject/NameTypes.h"
+#include "HAL/LowLevelMemTracker.h"
 
 class IMessageLog;
 
+LLM_DECLARE_TAG_API(EngineMisc_MessageLog, CORE_API);
+
 class FMessageLog
 {
 public:
diff --git a/Engine/Source/Runtime/Slate/Private/Framework/Commands/UICommandList.cpp b/Engine/Source/Runtime/Slate/Private/Framework/Commands/UICommandList.cpp
--- a/Engine/Source/Runtime/Slate/Private/Framework/Commands/UICommandList.cpp
+++ b/Engine/Source/Runtime/Slate/Private/Framework/Commands/UICommandList.cpp
@@ -53,6 +53,8 @@ void FUICommandList::MapAction( const TSharedPtr< const FUICommandInfo > InUICom
 
 void FUICommandList::MapAction( const TSharedPtr< const FUICommandInfo > InUICommandInfo, const FUIAction& InUIAction )
 {
+	LLM_SCOPE_BYTAG(UI_Slate);
+
 	if (!ensure(InUICommandInfo.IsValid()))
 	{
 		return;
@@ -64,6 +66,8 @@ void FUICommandList::MapAction( const TSharedPtr< const FUICommandInfo > InUICom
 
 void FUICommandList::MapAction(const TSharedPtr< const FUICommandInfo > InUICommandInfo, const FUIAction& InUIAction, const FUIActionContext& InUIActionContext)
 {
+	LLM_SCOPE_BYTAG(UI_Slate);
+	
 	if (!ensure(InUICommandInfo.IsValid()))
 	{
 		return;
@@ -76,6 +80,8 @@ void FUICommandList::MapAction(const TSharedPtr< const FUICommandInfo > InUIComm
 
 void FUICommandList::Append( const TSharedRef<FUICommandList>& InCommandsToAppend )
 {
+	LLM_SCOPE_BYTAG(UI_Slate);
+	
 	check( &InCommandsToAppend.Get() != this );
 
 	// Clear out any invalid parents or children
diff --git a/Engine/Source/Runtime/TypedElementFramework/Private/Elements/Framework/TypedElementRegistry.cpp b/Engine/Source/Runtime/TypedElementFramework/Private/Elements/Framework/TypedElementRegistry.cpp
--- a/Engine/Source/Runtime/TypedElementFramework/Private/Elements/Framework/TypedElementRegistry.cpp
+++ b/Engine/Source/Runtime/TypedElementFramework/Private/Elements/Framework/TypedElementRegistry.cpp
@@ -36,6 +36,7 @@ TStrongObjectPtr<UTypedElementRegistry>& GetTypedElementRegistryInstance()
 
 UTypedElementRegistry::UTypedElementRegistry()
 {
+	LLM_SCOPE(ELLMTag::EngineMisc);
 	if (!HasAnyFlags(RF_ClassDefaultObject))
 	{
 		FCoreDelegates::OnBeginFrame.AddUObject(this, &UTypedElementRegistry::OnBeginFrame);
@@ -174,6 +175,8 @@ void UTypedElementRegistry::RegisterElementTypeImpl(const FName InElementTypeNam
 
 void UTypedElementRegistry::RegisterElementInterfaceImpl(const FName InElementTypeName, UObject* InElementInterface, const TSubclassOf<UInterface>& InBaseInterfaceType, const bool InAllowOverride)
 {
+	LLM_SCOPE(ELLMTag::EngineMisc);
+
 	checkf(InElementInterface->GetClass()->ImplementsInterface(InBaseInterfaceType), TEXT("Interface '%s' of type '%s' does not derive from '%s'!"), *InElementInterface->GetPathName(), *InElementInterface->GetClass()->GetName(), *InBaseInterfaceType->GetName());
 
 	FRegisteredElementType* RegisteredElementType = GetRegisteredElementTypeFromName(InElementTypeName);
@@ -202,6 +205,7 @@ void UTypedElementRegistry::ProcessDeferredElementsToDestroy()
 
 	FReadScopeLock RegisteredElementTypesLock(RegisteredElementTypesRW);
 
+	LLM_SCOPE(ELLMTag::EngineMisc);
 	for (TUniquePtr<FRegisteredElementType>& RegisteredElementType : RegisteredElementTypes)
 	{
 		if (RegisteredElementType)
diff --git a/Engine/Source/Runtime/TypedElementFramework/Public/Elements/Framework/TypedElementRegistry.h b/Engine/Source/Runtime/TypedElementFramework/Public/Elements/Framework/TypedElementRegistry.h
--- a/Engine/Source/Runtime/TypedElementFramework/Public/Elements/Framework/TypedElementRegistry.h
+++ b/Engine/Source/Runtime/TypedElementFramework/Public/Elements/Framework/TypedElementRegistry.h
@@ -638,6 +638,7 @@ private:
 		{
 			FWriteScopeLock RegisteredElementTypesLock(RegisteredElementTypesRW);
 
+			LLM_SCOPE(ELLMTag::EngineMisc);
 			RegisteredElementTypesNameToId.Add(InRegisteredElementType->TypeName, InRegisteredElementType->TypeId);
 			RegisteredElementTypes[InRegisteredElementType->TypeId - 1] = MoveTemp(InRegisteredElementType);
 		}
