Localization: Improved debug localization culture.
- The keys culture now also displays the localization namespace of text on screen. It is displayed as key, namespace.
- Added the Slate.LogRenderedText CVar. When set to true, this CVar logs all the text to be rendered on screen via FSlateDrawElement::MakeText and code that calls FSlateDrawElement::MakeShapedText.
- Updated the automated tests for the keys culture to reflect the updated behavior of displaying text.

diff --git a/Engine/Source/Runtime/Core/Private/Internationalization/TextLocalizationManager.cpp b/Engine/Source/Runtime/Core/Private/Internationalization/TextLocalizationManager.cpp
--- a/Engine/Source/Runtime/Core/Private/Internationalization/TextLocalizationManager.cpp
+++ b/Engine/Source/Runtime/Core/Private/Internationalization/TextLocalizationManager.cpp
@@ -1674,9 +1674,10 @@ void FTextLocalizationManager::KeyifyAllDisplayStrings()
 	{
 		FDisplayStringEntry& LiveEntry = DisplayStringPair.Value;
 		LiveEntry.NativeStringBackup = LiveEntry.DisplayString;
-
-		FString Key = DisplayStringPair.Key.GetKey().GetChars();
-		FTextDisplayStringRef TmpDisplayString = MakeTextDisplayString(MoveTemp(Key));
+		// We want to show the identity in terms of key, namespace. This is to try and fit into the constraints of UI text blocks and at least let the key component be visible to easily identify a piece of text.
+		// If the key/namespace pair is too long, the Slate.LogPaintedText cvar can be used to see the entire thing.
+		FString KeyNamespaceDisplay = FString::Printf(TEXT("%s, %s"), DisplayStringPair.Key.GetKey().GetChars(), DisplayStringPair.Key.GetNamespace().GetChars());
+		FTextDisplayStringRef TmpDisplayString = MakeTextDisplayString(MoveTemp(KeyNamespaceDisplay));
 		LiveEntry.DisplayString = TmpDisplayString;
 	}
 }
diff --git a/Engine/Source/Runtime/Core/Private/Tests/Internationalization/KeysCultureTest.cpp b/Engine/Source/Runtime/Core/Private/Tests/Internationalization/KeysCultureTest.cpp
--- a/Engine/Source/Runtime/Core/Private/Tests/Internationalization/KeysCultureTest.cpp
+++ b/Engine/Source/Runtime/Core/Private/Tests/Internationalization/KeysCultureTest.cpp
@@ -31,6 +31,18 @@ struct FRestoreCultureStateGuard
 private:
 	FInternationalization::FCultureStateSnapshot OriginalCultureState;
 };
+/**
+* A helper function that generates the expected output string when the keys culture is active.
+* If you change the behavior of what gets displayed by the keys culture, be sure to update this function. FTextLocalizationManager::KeyifyAllDisplayStrings is currently the place that is done.
+*/
+FString CreateExpectedOutputFromPolyglotData(const FPolyglotTextData& InData)
+{
+	FString Key;
+	FString Namespace;
+	InData.GetIdentity(Namespace, Key);
+	// Right now, we display it as key, namespace
+	return FString::Printf(TEXT("%s, %s"), *Key, *Namespace);
+}
 
 IMPLEMENT_SIMPLE_AUTOMATION_TEST(FKeysCultureTest, "System.Core.Misc.KeysCulture", EAutomationTestFlags::EditorContext | EAutomationTestFlags::ClientContext | EAutomationTestFlags::EngineFilter)
 bool FKeysCultureTest::RunTest (const FString& Parameters)
@@ -45,7 +57,7 @@ bool FKeysCultureTest::RunTest (const FString& Parameters)
 	// 2. We register text data that as a native string and a target culture translation.
 	// 3. We register text data that as a native string and no target culture translation. 
 	// 4. We change to the target culture. The text with translation should have its display string the same as the localized string in the target culture. The display string for the text with no translation should match the native string.
-	// 5. We change to the keys culture. Both text with and without translations in a target culture should have their display strings show the localization key of the text.
+	// 5. We change to the keys culture. Both text with and without translations in a target culture should have their display strings show the localization key and namespace of the text.
 	// 6. We revert back to the native culture. Both text with and without translations should have their display strings the same as their native strings with no change.
 	
 	FInternationalization& I18N = FInternationalization::Get();
@@ -110,7 +122,7 @@ bool FKeysCultureTest::RunTest (const FString& Parameters)
 	TestEqual(TEXT("Text without translation display string is the same as the native string in target culture."), Test2TargetCultureDisplayString, Test2.GetNativeString());
 	
 	// Change to the keys culture
-	// This culture should output all the display strings as the localization key of the text.
+	// This culture should output all the display strings as the localization key and namespace of the text.
 	// There will be logs that say the localization resource files associated with the keys culture can't be read. That is normal and can be ignored.
 	I18N.SetCurrentCulture(FKeysCulture::StaticGetName());
 	// Test 5 - Can we successfully change from the native culture to the keys culture 
@@ -120,13 +132,13 @@ bool FKeysCultureTest::RunTest (const FString& Parameters)
 		return false;
 	}
 	
-	// Test 6 - The display string for the text with translation should be the localization key for the text in the keys culture 
+	// Test 6 - The display string for the text with translation should be the localization key and localization namespace for the text in the keys culture 
 	FString Test1KeysCultureDisplayString = Test1Text.ToString();
-	TestEqual(TEXT("Text with translation now has display string that matches its localization key in keys culture."), Test1KeysCultureDisplayString, Test1.GetKey());
+	TestEqual(TEXT("Text with translation now has display string that matches its localization key and namespace in keys culture."), Test1KeysCultureDisplayString, CreateExpectedOutputFromPolyglotData(Test1));
 	
-	// Test 7 - The display string for the text with no translation should be its localization key in the keys culture.
+	// Test 7 - The display string for the text with no translation should be its localization key and localization namespace in the keys culture.
 	FString Test2KeysCultureDisplayString = Test2Text.ToString();
-	TestEqual(TEXT("Text with blank translation has display string that matches its localization key in keys culture."), Test2KeysCultureDisplayString, Test2.GetKey());
+	TestEqual(TEXT("Text with blank translation has display string that matches its localization key and namespace in keys culture."), Test2KeysCultureDisplayString, CreateExpectedOutputFromPolyglotData(Test2));
 
 	// Revert to native culture. All of the text should still show up as the originals.
 	I18N.SetCurrentCulture(NativeCulture);
diff --git a/Engine/Source/Runtime/Core/Public/Internationalization/TextLocalizationManager.h b/Engine/Source/Runtime/Core/Public/Internationalization/TextLocalizationManager.h
--- a/Engine/Source/Runtime/Core/Public/Internationalization/TextLocalizationManager.h
+++ b/Engine/Source/Runtime/Core/Public/Internationalization/TextLocalizationManager.h
@@ -334,7 +334,7 @@ private:
 #if ENABLE_LOC_TESTING
 	/** A helper function that leetifies all of the display strings when the LEET culture is active. */
 	void LeetifyAllDisplayStrings();
-	/** A helper function that converts all of the display strings to show the localization key associated with the string when the keys culture is active. */
+	/** A helper function that converts all of the display strings to show the localization key and namespace associated with the string when the keys culture is active. */
 	void KeyifyAllDisplayStrings();
 #endif
 
diff --git a/Engine/Source/Runtime/Slate/Private/Framework/Text/SlateTextRun.cpp b/Engine/Source/Runtime/Slate/Private/Framework/Text/SlateTextRun.cpp
--- a/Engine/Source/Runtime/Slate/Private/Framework/Text/SlateTextRun.cpp
+++ b/Engine/Source/Runtime/Slate/Private/Framework/Text/SlateTextRun.cpp
@@ -75,9 +75,15 @@ TSharedRef< ILayoutBlock > FSlateTextRun::CreateBlock( int32 BeginIndex, int32 E
 {
 	return FDefaultLayoutBlock::Create( SharedThis( this ), FTextRange( BeginIndex, EndIndex ), Size, TextContext, Renderer );
 }
-
 int32 FSlateTextRun::OnPaint(const FPaintArgs& PaintArgs, const FTextArgs& TextArgs, const FGeometry& AllottedGeometry, const FSlateRect& MyCullingRect, FSlateWindowElementList& OutDrawElements, int32 LayerId, const FWidgetStyle& InWidgetStyle, bool bParentEnabled) const
 {
+#if !UE_BUILD_SHIPPING
+	static const auto CVar = IConsoleManager::Get().FindConsoleVariable(TEXT("Slate.LogPaintedText"));
+	if (CVar->GetBool())
+	{
+		UE_LOG(LogSlate, Log, TEXT("FSlateTextRun: '%s'."), **Text);
+	}
+#endif
 	const ESlateDrawEffect DrawEffects = bParentEnabled ? ESlateDrawEffect::None : ESlateDrawEffect::DisabledEffect;
 
 	const TSharedRef<ILayoutBlock>& Block = TextArgs.Block;
diff --git a/Engine/Source/Runtime/SlateCore/Private/Rendering/DrawElements.cpp b/Engine/Source/Runtime/SlateCore/Private/Rendering/DrawElements.cpp
--- a/Engine/Source/Runtime/SlateCore/Private/Rendering/DrawElements.cpp
+++ b/Engine/Source/Runtime/SlateCore/Private/Rendering/DrawElements.cpp
@@ -343,6 +343,17 @@ void FSlateDrawElement::MakeRotatedBox(
 		DrawElement.SetRenderTransform(Concatenate(RotationTransform, DrawElement.GetRenderTransform()));
 	}
 }
+#if !UE_BUILD_SHIPPING
+namespace UE::Slate::Private
+{
+	bool GLogPaintedText = false;
+	FAutoConsoleVariableRef LogPaintedTextRef(
+		TEXT("Slate.LogPaintedText"),
+		GLogPaintedText,
+		TEXT("If true, all text that is visible to the user will be logged when it is painted. This will log the full text to be painted, not the truncated or clipped version based on UI constraints.")
+	);
+}
+#endif
 
 void FSlateDrawElement::MakeText( FSlateWindowElementList& ElementList, uint32 InLayer, const FPaintGeometry& PaintGeometry, const FString& InText, const int32 StartIndex, const int32 EndIndex, const FSlateFontInfo& InFontInfo, ESlateDrawEffect InDrawEffects, const FLinearColor& InTint )
 {
@@ -353,7 +364,12 @@ void FSlateDrawElement::MakeText( FSlateWindowElementList& ElementList, uint32 I
 	{
 		return;
 	}
-
+#if !UE_BUILD_SHIPPING
+	if (UE::Slate::Private::GLogPaintedText)
+	{
+		UE_LOG(LogSlate, Log, TEXT("MakeText: '%s'."), *InText);
+	}
+#endif
 	FSlateDrawElement& Element = ElementList.AddUninitialized();
 	FSlateTextPayload& DataPayload = ElementList.CreatePayload<FSlateTextPayload>(Element);
 
@@ -384,7 +400,12 @@ void FSlateDrawElement::MakeText( FSlateWindowElementList& ElementList, uint32 I
 	{
 		return;
 	}
-
+#if !UE_BUILD_SHIPPING
+	if (UE::Slate::Private::GLogPaintedText)
+	{
+		UE_LOG(LogSlate, Log, TEXT("MakeText: '%s'."), *InText);
+	}
+#endif
 	FSlateDrawElement& Element = ElementList.AddUninitialized();
 
 	FSlateTextPayload& DataPayload = ElementList.CreatePayload<FSlateTextPayload>(Element);
