TArray GetTypeHash to return the product of array items

diff --git a/Engine/Plugins/Animation/DeformerGraph/Source/OptimusCore/Private/OptimusComputeDataInterface.cpp b/Engine/Plugins/Animation/DeformerGraph/Source/OptimusCore/Private/OptimusComputeDataInterface.cpp
--- a/Engine/Plugins/Animation/DeformerGraph/Source/OptimusCore/Private/OptimusComputeDataInterface.cpp
+++ b/Engine/Plugins/Animation/DeformerGraph/Source/OptimusCore/Private/OptimusComputeDataInterface.cpp
@@ -3,10 +3,10 @@
 #include "OptimusComputeDataInterface.h"
 
 #include "DataInterfaces/OptimusDataInterfaceRawBuffer.h"
-#include "OptimusTemplates.h"
 
-#include "UObject/UObjectIterator.h"
 #include "Templates/SubclassOf.h"
+#include "UObject/UObjectIterator.h"
+
 
 #include UE_INLINE_GENERATED_CPP_BY_NAME(OptimusComputeDataInterface)
 
diff --git a/Engine/Plugins/Animation/DeformerGraph/Source/OptimusCore/Public/OptimusNode.h b/Engine/Plugins/Animation/DeformerGraph/Source/OptimusCore/Public/OptimusNode.h
--- a/Engine/Plugins/Animation/DeformerGraph/Source/OptimusCore/Public/OptimusNode.h
+++ b/Engine/Plugins/Animation/DeformerGraph/Source/OptimusCore/Public/OptimusNode.h
@@ -4,11 +4,10 @@
 
 #include "OptimusCoreNotify.h"
 #include "OptimusDiagnostic.h"
-#include "OptimusTemplates.h"
 
-#include "UObject/Object.h"
 #include "CoreMinimal.h"
 #include "OptimusNodePin.h"
+#include "UObject/Object.h"
 
 #include "OptimusNode.generated.h"
 
diff --git a/Engine/Plugins/Animation/DeformerGraph/Source/OptimusCore/Public/OptimusTemplates.h b/Engine/Plugins/Animation/DeformerGraph/Source/OptimusCore/Public/OptimusTemplates.h
deleted file mode 100644
diff --git a/Engine/Source/Runtime/Core/Public/Containers/Array.h b/Engine/Source/Runtime/Core/Public/Containers/Array.h
--- a/Engine/Source/Runtime/Core/Public/Containers/Array.h
+++ b/Engine/Source/Runtime/Core/Public/Containers/Array.h
@@ -17,6 +17,7 @@
 #include "Algo/HeapSort.h"
 #include "Algo/IsHeap.h"
 #include "Algo/Impl/BinaryHeap.h"
+#include "Concepts/GetTypeHashable.h"
 #include "Templates/AndOrNot.h"
 #include "Templates/IdentityFunctor.h"
 #include "Templates/Invoke.h"
@@ -3561,3 +3562,15 @@ FArchive& operator<<(FArchive& Ar, TArray<ElementType, AllocatorType>& A)
 {
 	return TArrayPrivateFriend::Serialize(Ar, A);
 }
+
+/** Returns a unique hash by combining those of each array element. */
+template<typename InElementType, typename InAllocatorType>
+uint32 GetTypeHash(const TArray<InElementType, InAllocatorType>& A)
+{
+	uint32 Hash = 0;
+	for (const InElementType& V : A)
+	{
+		Hash = HashCombineFast(Hash, GetTypeHash(V));
+	}
+	return Hash;
+}
