[Core]
* RingBufferTest - Fixed c++20 compile error

diff --git a/Engine/Source/Runtime/Core/Tests/Containers/RingBufferTest.cpp b/Engine/Source/Runtime/Core/Tests/Containers/RingBufferTest.cpp
--- a/Engine/Source/Runtime/Core/Tests/Containers/RingBufferTest.cpp
+++ b/Engine/Source/Runtime/Core/Tests/Containers/RingBufferTest.cpp
@@ -1068,28 +1068,28 @@ public:
 				int32 CheckIndex = 0;
 				for (; CheckIndex < It; ++CheckIndex)
 				{
-					CHECK_MESSAGE(*FString::Printf(TEXT("ShiftIndexToBack Entire Array Values (%d,%d)"), It, CheckIndex), CheckIndex == Q[CheckIndex]);
+					CHECK_MESSAGE(*FString::Printf(TEXT("ShiftIndexToBack Entire Array Values (%d,%d)"), It, CheckIndex), CheckIndex == (int32)Q[CheckIndex]);
 				}
 				for (; CheckIndex < Count - 1; ++CheckIndex)
 				{
-					CHECK_MESSAGE(*FString::Printf(TEXT("ShiftIndexToBack Entire Array Values (%d,%d)"), It, CheckIndex), CheckIndex + 1 == Q[CheckIndex]);
+					CHECK_MESSAGE(*FString::Printf(TEXT("ShiftIndexToBack Entire Array Values (%d,%d)"), It, CheckIndex), CheckIndex + 1 == (int32)Q[CheckIndex]);
 				}
-				CHECK_MESSAGE(*FString::Printf(TEXT("ShiftIndexToBack Entire Array Values (%d,%d)"), It, Count - 1), It == Q[Count - 1]);
+				CHECK_MESSAGE(*FString::Printf(TEXT("ShiftIndexToBack Entire Array Values (%d,%d)"), It, Count - 1), It == (int32)Q[Count - 1]);
 			}
 			for (int32 It = 0; It < Count; ++It)
 			{
 				TRingBuffer<T> Q({ 0,1,2,3,4,5,6,7 });
 				Q.ShiftIndexToFront(It);
 
-				CHECK_MESSAGE(*FString::Printf(TEXT("ShiftIndexToFront Entire Array Values (%d,%d)"), It, 0), It == Q[0]);
+				CHECK_MESSAGE(*FString::Printf(TEXT("ShiftIndexToFront Entire Array Values (%d,%d)"), It, 0), It == (int32)Q[0]);
 				int32 CheckIndex = 1;
 				for (; CheckIndex <= It; ++CheckIndex)
 				{
-					CHECK_MESSAGE(*FString::Printf(TEXT("ShiftIndexToFront Entire Array Values (%d,%d)"), It, CheckIndex), CheckIndex - 1 == Q[CheckIndex]);
+					CHECK_MESSAGE(*FString::Printf(TEXT("ShiftIndexToFront Entire Array Values (%d,%d)"), It, CheckIndex), CheckIndex - 1 == (int32)Q[CheckIndex]);
 				}
 				for (; CheckIndex < Count; ++CheckIndex)
 				{
-					CHECK_MESSAGE(*FString::Printf(TEXT("ShiftIndexToFront Entire Array Values (%d,%d)"), It, CheckIndex), CheckIndex == Q[CheckIndex]);
+					CHECK_MESSAGE(*FString::Printf(TEXT("ShiftIndexToFront Entire Array Values (%d,%d)"), It, CheckIndex), CheckIndex == (int32)Q[CheckIndex]);
 				}
 			}
 		}
