Added check for invalid reserves in TSet.

diff --git a/Engine/Source/Runtime/Core/Private/Containers/Array.cpp b/Engine/Source/Runtime/Core/Private/Containers/Array.cpp
deleted file mode 100644
diff --git a/Engine/Source/Runtime/Core/Private/Containers/ContainerHelpers.cpp b/Engine/Source/Runtime/Core/Private/Containers/ContainerHelpers.cpp
new file mode 100644
--- /dev/null
+++ b/Engine/Source/Runtime/Core/Private/Containers/ContainerHelpers.cpp
@@ -0,0 +1,13 @@
+// Copyright Epic Games, Inc. All Rights Reserved.
+
+#include "Containers/ContainerHelpers.h"
+
+void UE::Core::Private::OnInvalidArrayNum(unsigned long long NewNum)
+{
+	UE_LOG(LogCore, Fatal, TEXT("Trying to resize TArray to an invalid size of %llu"), NewNum);
+}
+
+void UE::Core::Private::OnInvalidSetNum(unsigned long long NewNum)
+{
+	UE_LOG(LogCore, Fatal, TEXT("Trying to resize TSet to an invalid size of %llu"), NewNum);
+}
diff --git a/Engine/Source/Runtime/Core/Public/Containers/Array.h b/Engine/Source/Runtime/Core/Public/Containers/Array.h
--- a/Engine/Source/Runtime/Core/Public/Containers/Array.h
+++ b/Engine/Source/Runtime/Core/Public/Containers/Array.h
@@ -332,8 +332,11 @@ namespace UE4Array_Private
 	{
 		enum { Value = TArrayElementsAreCompatible<ElementType, TElementType_T<RangeType>>::Value };
 	};
+}
 
-	CORE_API void OnInvalidArrayNum(const TCHAR* ArrayNameSuffix, unsigned long long NewNum);
+namespace UE::Core::Private
+{
+	CORE_API void OnInvalidArrayNum(unsigned long long NewNum);
 }
 
 
@@ -362,16 +365,6 @@ public:
 private:
 	using USizeType = typename TMakeUnsigned<SizeType>::Type;
 
-	FORCENOINLINE static void OnInvalidNum(USizeType NewNum)
-	{
-		const TCHAR* ArrayNameSuffix = TEXT("");
-		if constexpr (sizeof(SizeType) == 8)
-		{
-			ArrayNameSuffix = TEXT("64");
-		}
-		UE4Array_Private::OnInvalidArrayNum(ArrayNameSuffix, (unsigned long long)NewNum);
-	}
-
 public:
 	UE_DEPRECATED(5.0, "TArray::Allocator type is deprecated, please use TArray::AllocatorType instead.")
 	typedef InAllocatorType Allocator;
@@ -403,7 +396,7 @@ public:
 	{
 		if (Count < 0)
 		{
-			OnInvalidNum((USizeType)Count);
+			UE::Core::Private::OnInvalidArrayNum((unsigned long long)Count);
 		}
 
 		check(Ptr != nullptr || Count == 0);
@@ -550,7 +543,7 @@ private:
 			{
 				if (ToArray.ArrayNum != FromArray.ArrayNum || ToArray.ArrayMax != FromArray.ArrayMax)
 				{
-					OnInvalidNum((USizeType)ToArray.ArrayNum);
+					UE::Core::Private::OnInvalidArrayNum((unsigned long long)ToArray.ArrayNum);
 				}
 			}
 
@@ -589,7 +582,7 @@ private:
 			// This should only happen when we've underflowed or overflowed SizeType
 			if ((SizeType)NewMax < LocalArrayNum)
 			{
-				OnInvalidNum((USizeType)ExtraSlack);
+				UE::Core::Private::OnInvalidArrayNum((unsigned long long)ExtraSlack);
 			}
 
 			ToArray.Reserve(NewMax);
@@ -1792,7 +1785,7 @@ public:
 	{
 		if (NewSize < 0)
 		{
-			OnInvalidNum((USizeType)NewSize);
+			UE::Core::Private::OnInvalidArrayNum((unsigned long long)NewSize);
 		}
 
 		// If we have space to hold the excepted size, then don't reallocate
@@ -1816,7 +1809,7 @@ public:
 	{
 		if (Slack < 0)
 		{
-			OnInvalidNum((USizeType)Slack);
+			UE::Core::Private::OnInvalidArrayNum((unsigned long long)Slack);
 		}
 
 		DestructItems(GetData(), ArrayNum);
@@ -1846,7 +1839,7 @@ public:
 		}
 		else if (NewNum < 0)
 		{
-			OnInvalidNum((USizeType)NewNum);
+			UE::Core::Private::OnInvalidArrayNum((unsigned long long)NewNum);
 		}
 		else if (NewNum < Num())
 		{
@@ -1869,7 +1862,7 @@ public:
 		}
 		else if (NewNum < 0)
 		{
-			OnInvalidNum((USizeType)NewNum);
+			UE::Core::Private::OnInvalidArrayNum((unsigned long long)NewNum);
 		}
 		else if (NewNum < Num())
 		{
@@ -1890,7 +1883,7 @@ public:
 		}
 		else if (NewNum < 0)
 		{
-			OnInvalidNum((USizeType)NewNum);
+			UE::Core::Private::OnInvalidArrayNum((unsigned long long)NewNum);
 		}
 		else if (NewNum < Num())
 		{
@@ -2478,7 +2471,7 @@ public:
 		checkSlow(Number >= 0);
 		if (Number < 0)
 		{
-			OnInvalidNum((USizeType)Number);
+			UE::Core::Private::OnInvalidArrayNum((unsigned long long)Number);
 		}
 		else if (Number > ArrayMax)
 		{
@@ -2938,7 +2931,7 @@ private:
 		// This should only happen when we've underflowed or overflowed SizeType in the caller
 		if (LocalArrayNum < OldNum)
 		{
-			OnInvalidNum((USizeType)LocalArrayNum - (USizeType)OldNum);
+			UE::Core::Private::OnInvalidArrayNum((unsigned long long)LocalArrayNum - (unsigned long long)OldNum);
 		}
 		ArrayMax = AllocatorCalculateSlackGrow(LocalArrayNum, ArrayMax);
 		AllocatorResizeAllocation(OldNum, ArrayMax);
@@ -3031,7 +3024,7 @@ private:
 			// This should only happen when we've underflowed or overflowed SizeType
 			if ((SizeType)NewMax < NewNum)
 			{
-				OnInvalidNum((USizeType)NewMax);
+				UE::Core::Private::OnInvalidArrayNum((unsigned long long)NewMax);
 			}
 
 			ResizeForCopy(NewNum + ExtraSlack, PrevMax);
diff --git a/Engine/Source/Runtime/Core/Public/Containers/ContainerHelpers.h b/Engine/Source/Runtime/Core/Public/Containers/ContainerHelpers.h
new file mode 100644
--- /dev/null
+++ b/Engine/Source/Runtime/Core/Public/Containers/ContainerHelpers.h
@@ -0,0 +1,7 @@
+// Copyright Epic Games, Inc. All Rights Reserved.
+
+namespace UE::Core::Private
+{
+	CORE_API void OnInvalidArrayNum(unsigned long long NewNum);
+	CORE_API void OnInvalidSetNum  (unsigned long long NewNum);
+}
diff --git a/Engine/Source/Runtime/Core/Public/Containers/Set.h b/Engine/Source/Runtime/Core/Public/Containers/Set.h
--- a/Engine/Source/Runtime/Core/Public/Containers/Set.h
+++ b/Engine/Source/Runtime/Core/Public/Containers/Set.h
@@ -258,6 +258,11 @@ public:
 	}
 };
 
+namespace UE::Core::Private
+{
+	CORE_API void OnInvalidSetNum(unsigned long long NewNum);
+}
+
 /**
  * A set with an optional KeyFuncs parameters for customizing how the elements are compared and searched.  
  * E.g. You can specify a mapping from elements to keys if you want to find elements by specifying a subset of 
@@ -290,6 +295,8 @@ public:
 	static_assert(std::is_same_v<SizeType, int32>, "TSet currently only supports 32-bit allocators");
 
 private:
+	using USizeType = std::make_unsigned_t<SizeType>;
+
 	template <typename, typename>
 	friend class TScriptSet;
 
@@ -497,8 +504,14 @@ public:
 	{
 		// makes sense only when Number > Elements.Num() since TSparseArray::Reserve 
 		// does any work only if that's the case
-		if (Number > Elements.Num())
+		if ((USizeType)Number > (USizeType)Elements.Num())
 		{
+			// Trap negative reserves
+			if (Number < 0)
+			{
+				UE::Core::Private::OnInvalidSetNum((unsigned long long)Number);
+			}
+
 			// Preallocates memory for array of elements
 			Elements.Reserve(Number);
 
