Changed LowLevelTasks::TEventStack to support ARM's HWASan and MTE

diff --git a/Engine/Source/Runtime/Core/Private/Async/Fundamental/ReserveScheduler.cpp b/Engine/Source/Runtime/Core/Private/Async/Fundamental/ReserveScheduler.cpp
--- a/Engine/Source/Runtime/Core/Private/Async/Fundamental/ReserveScheduler.cpp
+++ b/Engine/Source/Runtime/Core/Private/Async/Fundamental/ReserveScheduler.cpp
@@ -17,11 +17,8 @@ namespace LowLevelTasks
 
 FReserveScheduler FReserveScheduler::Singleton;
 
-TUniquePtr<FThread> FReserveScheduler::CreateWorker(FThread::EForkable IsForkable, EThreadPriority Priority)
+TUniquePtr<FThread> FReserveScheduler::CreateWorker(FThread::EForkable IsForkable, FYieldedWork* ReserveEvent, EThreadPriority Priority)
 {
-	FYieldedWork* ReserveEvent = new FYieldedWork;
-	ReserveEvents.Emplace(ReserveEvent);
-
 	uint32 WorkerId = NextWorkerId++;
 	return MakeUnique<FThread>
 	(
@@ -84,10 +81,12 @@ void FReserveScheduler::StartWorkers(uint32 NumWorkers, FThread::EForkable IsFor
 		check(!WorkerThreads.Num());
 		check(NextWorkerId == 0);
 
+		ReserveEvents.Reserve(NumWorkers);
 		UE::Trace::ThreadGroupBegin(TEXT("Reserve Workers"));
 		for (uint32 WorkerId = 0; WorkerId < NumWorkers; ++WorkerId)
 		{
-			WorkerThreads.Add(CreateWorker(IsForkable, WorkerPriority));
+			ReserveEvents.Emplace();
+			WorkerThreads.Add(CreateWorker(IsForkable, &ReserveEvents.Last(), &WorkerLocalQueues.Last(), WorkerPriority));
 		}
 		UE::Trace::ThreadGroupEnd();
 	}
@@ -100,9 +99,9 @@ void FReserveScheduler::StopWorkers()
 	{
 		FScopeLock Lock(&WorkerThreadsCS);
 
-		for (TUniquePtr<FYieldedWork>& Event : ReserveEvents)
+		for (FYieldedWork& Event : ReserveEvents)
 		{
-			Event->SleepEvent->Trigger();
+			Event.SleepEvent->Trigger();
 		}
 
 		while (FYieldedWork* Event = EventStack.Pop())
diff --git a/Engine/Source/Runtime/Core/Public/Async/Fundamental/ReserveScheduler.h b/Engine/Source/Runtime/Core/Public/Async/Fundamental/ReserveScheduler.h
--- a/Engine/Source/Runtime/Core/Public/Async/Fundamental/ReserveScheduler.h
+++ b/Engine/Source/Runtime/Core/Public/Async/Fundamental/ReserveScheduler.h
@@ -26,19 +26,8 @@ namespace LowLevelTasks
 		FReserveScheduler() = default;
 		~FReserveScheduler();
 
-	public:
 		using FConditional = TTaskDelegate<bool(), 16>;
 
-	private:
-		struct alignas(8) FYieldedWork
-		{
-			FEventRef		SleepEvent;
-			FConditional	CompletedDelegate = []() { return true; };
-			std::atomic<FYieldedWork*> Next {nullptr};
-			bool			bPermitBackgroundWork = false;
-		};
-	
-	public:
 		FORCEINLINE_DEBUGGABLE static FReserveScheduler& Get();
 
 		//start number of reserve workers where 0 is the system default
@@ -49,11 +38,18 @@ namespace LowLevelTasks
 		CORE_API bool DoReserveWorkUntil(FConditional&& Condition);
 
 	private: 
-		TUniquePtr<FThread> CreateWorker(FThread::EForkable IsForkable = FThread::NonForkable, EThreadPriority Priority = EThreadPriority::TPri_Normal);
+		struct alignas(64) FYieldedWork
+		{
+			FEventRef		SleepEvent;
+			FConditional	CompletedDelegate = []() { return true; };
+			std::atomic<FYieldedWork*> Next { nullptr };
+			bool			bPermitBackgroundWork = false;
+		};
+
+		TUniquePtr<FThread> CreateWorker(FThread::EForkable IsForkable = FThread::NonForkable, FYieldedWork* ReserveEvent = nullptr, EThreadPriority Priority = EThreadPriority::TPri_Normal);
 
-	private:
-		TEventStack<FYieldedWork> 						EventStack;
-		TArray<TUniquePtr<FYieldedWork>>				ReserveEvents;
+		TEventStack<FYieldedWork> 						EventStack = { ReserveEvents };
+		TAlignedArray<FYieldedWork>						ReserveEvents;
 		FCriticalSection 								WorkerThreadsCS;
 		TArray<TUniquePtr<FThread>>						WorkerThreads;
 		std::atomic_uint								ActiveWorkers { 0 };
diff --git a/Engine/Source/Runtime/Core/Public/Async/Fundamental/Scheduler.h b/Engine/Source/Runtime/Core/Public/Async/Fundamental/Scheduler.h
--- a/Engine/Source/Runtime/Core/Public/Async/Fundamental/Scheduler.h
+++ b/Engine/Source/Runtime/Core/Public/Async/Fundamental/Scheduler.h
@@ -30,42 +30,56 @@ namespace LowLevelTasks
 		DefaultPreference = LocalQueuePreference,
 	};
 
+	template<typename NodeType>
+	using TAlignedArray = TArray<NodeType, TAlignedHeapAllocator<alignof(NodeType)>>;
 
 	//implementation of a treiber stack
 	//(https://en.wikipedia.org/wiki/Treiber_stack)
 	template<typename NodeType>
 	class TEventStack
 	{
-	private:
+		static constexpr uint32 EVENT_INDEX_NONE = -1;
+
 		struct FTopNode
 		{
-			uintptr_t Address  : 45; //all CPUs we care about use less than 48 bits for their addressing, and the lower 3 bits are unused due to alignment
-			uintptr_t Revision : 19; //Tagging is used to avoid ABA (the wrap around is several minutes for this use-case (https://en.wikipedia.org/wiki/ABA_problem#Tagged_state_reference))
+			uint32 EventIndex;		// All events are stored in an array, so this allows us to pack node info into 32 bit as array index
+			uint32 Revision;		// Tagging is used to avoid ABA (https://en.wikipedia.org/wiki/ABA_problem#Tagged_state_reference)
 		};
-		std::atomic<FTopNode> Top { FTopNode{0, 0} };
+		std::atomic<FTopNode> Top { FTopNode{EVENT_INDEX_NONE, 0} };
+		TAlignedArray<NodeType>& NodesArray;
+
+		uint32 GetNodeIndex(const NodeType* Node) const
+		{
+			return (Node == nullptr) ? EVENT_INDEX_NONE : (Node - NodesArray.GetData());
+		}
 
 	public:
+		TEventStack(TAlignedArray<NodeType>& NodesArray)
+			: NodesArray(NodesArray)
+		{
+		}
+
 		NodeType* Pop()
 		{
 			FTopNode LocalTop = Top.load(std::memory_order_acquire);
 			while (true) 
 			{			
-				if (LocalTop.Address == 0)
+				if (LocalTop.EventIndex == EVENT_INDEX_NONE)
 				{
 					return nullptr;
 				}
 #if DO_CHECK
-				int64 LastRevision = int64(LocalTop.Revision); 
+				const int64 LastRevision = int64(LocalTop.Revision);
 #endif
-				NodeType* Item = reinterpret_cast<NodeType*>(LocalTop.Address << 3);
-				if (Top.compare_exchange_weak(LocalTop, FTopNode { reinterpret_cast<uintptr_t>(Item->Next.load(std::memory_order_relaxed)) >> 3, uintptr_t(LocalTop.Revision + 1) }, std::memory_order_acq_rel))
+				NodeType* Item = &NodesArray[LocalTop.EventIndex];
+				if (Top.compare_exchange_weak(LocalTop, FTopNode { GetNodeIndex(Item->Next.load(std::memory_order_relaxed)), uint32(LocalTop.Revision + 1) }, std::memory_order_relaxed, std::memory_order_relaxed))
 				{
 					Item->Next.store(nullptr, std::memory_order_relaxed);
 					return Item;
 				}
 #if DO_CHECK
-				int64 NewRevision = int64(LocalTop.Revision) < LastRevision ? ((1ll << 19) + int64(LocalTop.Revision)) : int64(LocalTop.Revision);
-				ensureMsgf((NewRevision - LastRevision) < (1ll << 18), TEXT("Dangerously close to the wraparound: %d, %d"), LastRevision, NewRevision);
+				const int64 NewRevision = int64(LocalTop.Revision) < LastRevision ? (int64(UINT32_MAX) + int64(LocalTop.Revision)) : int64(LocalTop.Revision);
+				ensureMsgf((NewRevision - LastRevision) < (1ll << 31), TEXT("Dangerously close to the wraparound: %d, %d"), LastRevision, NewRevision);
 #endif
 			}
 		}
@@ -73,26 +87,23 @@ namespace LowLevelTasks
 		void Push(NodeType* Item)
 		{
 			checkSlow(Item != nullptr);
-#if !USING_CODE_ANALYSIS //MS SA thowing warning C6011 on Item->Next access, even when it is validated or branched over
-			checkSlow(reinterpret_cast<uintptr_t>(Item) < (1ull << 48));
-			checkSlow((reinterpret_cast<uintptr_t>(Item) & 0x7) == 0);
-#endif
 			checkSlow(Item->Next.load(std::memory_order_relaxed) == nullptr);
+			checkfSlow((Item >= NodesArray.GetData()) && (Item < (NodesArray.GetData() + NodesArray.Num())), TEXT("Item doesn't belong to a Nodes Array"));
 
 			FTopNode LocalTop = Top.load(std::memory_order_relaxed);
 			while (true) 
 			{
 #if DO_CHECK
-				int64 LastRevision = int64(LocalTop.Revision); 
+				const int64 LastRevision = int64(LocalTop.Revision);
 #endif
-				Item->Next.store(reinterpret_cast<NodeType*>(LocalTop.Address << 3), std::memory_order_relaxed);
-				if (Top.compare_exchange_weak(LocalTop, FTopNode { reinterpret_cast<uintptr_t>(Item) >> 3, uintptr_t(LocalTop.Revision + 1) }, std::memory_order_acq_rel))
+				Item->Next.store((LocalTop.EventIndex == EVENT_INDEX_NONE) ? nullptr : &NodesArray[LocalTop.EventIndex], std::memory_order_relaxed);
+				if (Top.compare_exchange_weak(LocalTop, FTopNode { GetNodeIndex(Item), uint32(LocalTop.Revision + 1) }, std::memory_order_release, std::memory_order_relaxed))
 				{
 					return;
 				}
 #if DO_CHECK
-				int64 NewRevision = int64(LocalTop.Revision) < LastRevision ? ((1ll << 19) + int64(LocalTop.Revision)) : int64(LocalTop.Revision);
-				ensureMsgf((NewRevision - LastRevision) < (1ll << 18), TEXT("Dangerously close to the wraparound: %d, %d"), LastRevision, NewRevision);
+				const int64 NewRevision = int64(LocalTop.Revision) < LastRevision ? (int64(UINT32_MAX) + int64(LocalTop.Revision)) : int64(LocalTop.Revision);
+				ensureMsgf((NewRevision - LastRevision) < (1ll << 31), TEXT("Dangerously close to the wraparound: %d, %d"), LastRevision, NewRevision);
 #endif
 			}
 		}
@@ -195,9 +206,7 @@ namespace LowLevelTasks
 		bool TryExecuteTaskFrom(QueueType* Queue, FSchedulerTls::FQueueRegistry::FOutOfWork& OutOfWork, bool bPermitBackgroundWork, bool bDisableThrottleStealing);
 
 	private:
-		template<typename ElementType>
-		using TAlignedArray = TArray<ElementType, TAlignedHeapAllocator<alignof(ElementType)>>;
-		TEventStack<FSleepEvent> 						SleepEventStack[2];
+		TEventStack<FSleepEvent> 						SleepEventStack[2] = { WorkerEvents , WorkerEvents };
 		FSchedulerTls::FQueueRegistry 					QueueRegistry;
 		FCriticalSection 								WorkerThreadsCS;
 		TArray<TUniquePtr<FThread>>						WorkerThreads;
