Fixed compilation error in Scheduler.h

diff --git a/Engine/Source/Runtime/Core/Public/Async/Fundamental/Scheduler.h b/Engine/Source/Runtime/Core/Public/Async/Fundamental/Scheduler.h
--- a/Engine/Source/Runtime/Core/Public/Async/Fundamental/Scheduler.h
+++ b/Engine/Source/Runtime/Core/Public/Async/Fundamental/Scheduler.h
@@ -50,7 +50,7 @@ namespace LowLevelTasks
 
 		uint32 GetNodeIndex(const NodeType* Node) const
 		{
-			return (Node == nullptr) ? EVENT_INDEX_NONE : (Node - NodesArray.GetData());
+			return (Node == nullptr) ? EVENT_INDEX_NONE : uint32(Node - NodesArray.GetData());
 		}
 
 	public:
