Fixed compilation in ReserveScheduler.cpp

diff --git a/Engine/Source/Runtime/Core/Private/Async/Fundamental/ReserveScheduler.cpp b/Engine/Source/Runtime/Core/Private/Async/Fundamental/ReserveScheduler.cpp
--- a/Engine/Source/Runtime/Core/Private/Async/Fundamental/ReserveScheduler.cpp
+++ b/Engine/Source/Runtime/Core/Private/Async/Fundamental/ReserveScheduler.cpp
@@ -86,7 +86,7 @@ void FReserveScheduler::StartWorkers(uint32 NumWorkers, FThread::EForkable IsFor
 		for (uint32 WorkerId = 0; WorkerId < NumWorkers; ++WorkerId)
 		{
 			ReserveEvents.Emplace();
-			WorkerThreads.Add(CreateWorker(IsForkable, &ReserveEvents.Last(), &WorkerLocalQueues.Last(), WorkerPriority));
+			WorkerThreads.Add(CreateWorker(IsForkable, &ReserveEvents.Last(), WorkerPriority));
 		}
 		UE::Trace::ThreadGroupEnd();
 	}
