Copied over unity build fixes.

diff --git a/Engine/Source/Runtime/Core/Private/Containers/ContainerHelpers.cpp b/Engine/Source/Runtime/Core/Private/Containers/ContainerHelpers.cpp
--- a/Engine/Source/Runtime/Core/Private/Containers/ContainerHelpers.cpp
+++ b/Engine/Source/Runtime/Core/Private/Containers/ContainerHelpers.cpp
@@ -1,6 +1,7 @@
 // Copyright Epic Games, Inc. All Rights Reserved.
 
 #include "Containers/ContainerHelpers.h"
+#include "CoreGlobals.h"
 
 void UE::Core::Private::OnInvalidArrayNum(unsigned long long NewNum)
 {
diff --git a/Engine/Source/Runtime/Core/Public/Containers/ContainerHelpers.h b/Engine/Source/Runtime/Core/Public/Containers/ContainerHelpers.h
--- a/Engine/Source/Runtime/Core/Public/Containers/ContainerHelpers.h
+++ b/Engine/Source/Runtime/Core/Public/Containers/ContainerHelpers.h
@@ -1,5 +1,9 @@
 // Copyright Epic Games, Inc. All Rights Reserved.
 
+#pragma once
+
+#include "CoreTypes.h"
+
 namespace UE::Core::Private
 {
 	CORE_API void OnInvalidArrayNum(unsigned long long NewNum);
