Implement FAutomationTestBase::TestNearlyEqual methods

#rn adam.miezianko, chris.constantinescru, jerome.delattre

diff --git a/Engine/Source/Runtime/Core/Private/Misc/AutomationTest.cpp b/Engine/Source/Runtime/Core/Private/Misc/AutomationTest.cpp
--- a/Engine/Source/Runtime/Core/Private/Misc/AutomationTest.cpp
+++ b/Engine/Source/Runtime/Core/Private/Misc/AutomationTest.cpp
@@ -1532,6 +1532,31 @@ bool FAutomationTestBase::TestEqualInsensitive(const TCHAR* What, const TCHAR* A
 	return true;
 }
 
+bool FAutomationTestBase::TestNearlyEqual(const TCHAR* What, const float Actual, const float Expected, float Tolerance)
+{
+	return TestEqual(What, Actual, Expected, Tolerance);
+}
+
+bool FAutomationTestBase::TestNearlyEqual(const TCHAR* What, const double Actual, const double Expected, double Tolerance)
+{
+	return TestEqual(What, Actual, Expected, Tolerance);
+}
+
+bool FAutomationTestBase::TestNearlyEqual(const TCHAR* What, const FVector Actual, const FVector Expected, float Tolerance)
+{
+	return TestEqual(What, Actual, Expected, Tolerance);
+}
+
+bool FAutomationTestBase::TestNearlyEqual(const TCHAR* What, const FTransform Actual, const FTransform Expected, float Tolerance)
+{
+	return TestEqual(What, Actual, Expected, Tolerance);
+}
+
+bool FAutomationTestBase::TestNearlyEqual(const TCHAR* What, const FRotator Actual, const FRotator Expected, float Tolerance)
+{
+	return TestEqual(What, Actual, Expected, Tolerance);
+}
+
 bool FAutomationTestBase::TestFalse(const TCHAR* What, bool Value)
 {
 	if (Value)
diff --git a/Engine/Source/Runtime/Core/Public/Misc/AutomationTest.h b/Engine/Source/Runtime/Core/Public/Misc/AutomationTest.h
--- a/Engine/Source/Runtime/Core/Public/Misc/AutomationTest.h
+++ b/Engine/Source/Runtime/Core/Public/Misc/AutomationTest.h
@@ -1607,6 +1607,11 @@ public:
 		return TestEqual(*What, Actual, Expected, Tolerance);
 	}
 
+	bool TestEqual(const FString& What, const FTransform Actual, const FTransform Expected, float Tolerance = UE_KINDA_SMALL_NUMBER)
+	{
+		return TestEqual(*What, Actual, Expected, Tolerance);
+	}
+
 	bool TestEqual(const FString& What, const FRotator Actual, const FRotator Expected, float Tolerance = UE_KINDA_SMALL_NUMBER)
 	{
 		return TestEqual(*What, Actual, Expected, Tolerance);
@@ -1678,6 +1683,37 @@ public:
 		return TestEqual(*What, Actual, Expected);
 	}
 
+	bool TestNearlyEqual(const TCHAR* What, const float Actual, const float Expected, float Tolerance = UE_KINDA_SMALL_NUMBER);
+	bool TestNearlyEqual(const TCHAR* What, const double Actual, const double Expected, double Tolerance = UE_KINDA_SMALL_NUMBER);
+	bool TestNearlyEqual(const TCHAR* What, const FVector Actual, const FVector Expected, float Tolerance = UE_KINDA_SMALL_NUMBER);
+	bool TestNearlyEqual(const TCHAR* What, const FTransform Actual, const FTransform Expected, float Tolerance = UE_KINDA_SMALL_NUMBER);
+	bool TestNearlyEqual(const TCHAR* What, const FRotator Actual, const FRotator Expected, float Tolerance = UE_KINDA_SMALL_NUMBER);
+
+	bool TestNearlyEqual(const FString& What, const float Actual, const float Expected, float Tolerance = UE_KINDA_SMALL_NUMBER)
+	{
+		return TestNearlyEqual(*What, Actual, Expected, Tolerance);
+	}
+
+	bool TestNearlyEqual(const FString& What, const double Actual, const double Expected, double Tolerance = UE_KINDA_SMALL_NUMBER)
+	{
+		return TestNearlyEqual(*What, Actual, Expected, Tolerance);
+	}
+
+	bool TestNearlyEqual(const FString& What, const FVector Actual, const FVector Expected, float Tolerance = UE_KINDA_SMALL_NUMBER)
+	{
+		return TestNearlyEqual(*What, Actual, Expected, Tolerance);
+	}
+
+	bool TestNearlyEqual(const FString& What, const FTransform Actual, const FTransform Expected, float Tolerance = UE_KINDA_SMALL_NUMBER)
+	{
+		return TestNearlyEqual(*What, Actual, Expected, Tolerance);
+	}
+
+	bool TestNearlyEqual(const FString& What, const FRotator Actual, const FRotator Expected, float Tolerance = UE_KINDA_SMALL_NUMBER)
+	{
+		return TestNearlyEqual(*What, Actual, Expected, Tolerance);
+	}
+
 
 	/**
 	 * Logs an error if the specified Boolean value is not false.
