handling of invalid prerequisites and nested tasks in "tasks integration"

diff --git a/Engine/Source/Runtime/Core/Public/Async/TaskGraphInterfaces.h b/Engine/Source/Runtime/Core/Public/Async/TaskGraphInterfaces.h
--- a/Engine/Source/Runtime/Core/Public/Async/TaskGraphInterfaces.h
+++ b/Engine/Source/Runtime/Core/Public/Async/TaskGraphInterfaces.h
@@ -468,7 +468,10 @@ public:
 		{
 			for (const FGraphEventRef& Prereq : *InPrerequisites)
 			{
-				AddPrerequisites(*Prereq);
+				if (Prereq)
+				{
+					AddPrerequisites(*Prereq);
+				}
 			}
 		}
 	}
@@ -499,8 +502,20 @@ public:
 
 	void DontCompleteUntil(FGraphEventRef NestedTask)
 	{
-		checkSlow(UE::Tasks::Private::GetCurrentTask() == this); // a nested task can be added only from inside of parent's execution
-		AddNested(*NestedTask);
+		if (!NestedTask)
+		{
+			return;
+		}
+
+		if (GetExtendedPriority() == UE::Tasks::EExtendedTaskPriority::TaskEvent)
+		{	// TaskEvent can't have nested tasks, add it as a prerequisite instead
+			AddPrerequisites(*NestedTask);
+		}
+		else
+		{
+			checkSlow(UE::Tasks::Private::GetCurrentTask() == this); // a nested task can be added only from inside of parent's execution
+			AddNested(*NestedTask);
+		}
 	}
 
 	bool IsComplete() const
@@ -512,6 +527,7 @@ public:
 
 	void DispatchSubsequents(ENamedThreads::Type CurrentThreadIfKnown = ENamedThreads::AnyThread)
 	{
+		//check(CurrentThreadIfKnown == ENamedThreads::AnyThread); // the feature is not used
 		AddRef(); // scheduler's reference
 		TryLaunch();
 	}
@@ -529,7 +545,7 @@ public:
 
 	void Wait(ENamedThreads::Type CurrentThreadIfKnown = ENamedThreads::AnyThread)
 	{
-		// local queue have to be handled by the original TaskGraph implementation. The new frontend doesn't support local queues
+		// local queue have to be handled by the original TaskGraph implementation. Tasks System doesn't support local queues
 		if (ENamedThreads::GetQueueIndex(CurrentThreadIfKnown) != ENamedThreads::MainQueue)
 		{
 			return FTaskGraphInterface::Get().WaitUntilTaskCompletes(this, CurrentThreadIfKnown);
