Add TSparseArray::SortFreeList (required for determinism in Chaos)

diff --git a/Engine/Source/Runtime/Core/Private/Containers/ContainersTest.cpp b/Engine/Source/Runtime/Core/Private/Containers/ContainersTest.cpp
--- a/Engine/Source/Runtime/Core/Private/Containers/ContainersTest.cpp
+++ b/Engine/Source/Runtime/Core/Private/Containers/ContainersTest.cpp
@@ -396,6 +396,49 @@ namespace
 		}
 	}
 
+	// Test container element address consistency when using SortFreeList 
+	// (see TSparseArray::SortFreeList for comments)
+	template <typename ContainerType, typename KeyType>
+	void RunContainerConsistencyTests()
+	{
+		ContainerType Cont;
+
+		{
+			// Add 3 elements, then remove 2 in the same order they were added
+			const KeyType Key0 = GenerateTestKey<KeyType>(0);
+			const KeyType Key1 = GenerateTestKey<KeyType>(1);
+			const KeyType Key2 = GenerateTestKey<KeyType>(2);
+			const FContainerTestValueType Value = FContainerTestValueType(TEXT("New Value"));
+			Cont.Add(Key0, Value);
+			Cont.Add(Key1, Value);
+			Cont.Add(Key2, Value);
+			const FContainerTestValueType* ValuePtr0 = Cont.Find(Key0);
+			const FContainerTestValueType* ValuePtr1 = Cont.Find(Key1);
+			const FContainerTestValueType* ValuePtr2 = Cont.Find(Key2);
+			check(ValuePtr0 != nullptr);
+			check(ValuePtr1 != nullptr);
+			check(ValuePtr2 != nullptr);
+			Cont.Remove(Key1);
+			Cont.Remove(Key2);
+
+			// Re-add the 2 elements in the same order. Without the call to SortFreeList() 
+			// the elements would end up in a different locations/order compared to the 
+			// original insertions. SortFreeList() should ensure that re-adding the elements 
+			// gives use the same container layout as long as we perform the same operations 
+			// in the same order as before.
+			Cont.SortFreeList();
+			Cont.Add(Key1, Value);
+			Cont.Add(Key2, Value);
+			const FContainerTestValueType* NewValuePtr0 = Cont.Find(Key0);
+			const FContainerTestValueType* NewValuePtr1 = Cont.Find(Key1);
+			const FContainerTestValueType* NewValuePtr2 = Cont.Find(Key2);
+
+			check(ValuePtr0 == NewValuePtr0);
+			check(ValuePtr1 == NewValuePtr1);
+			check(ValuePtr2 == NewValuePtr2);
+		}
+	}
+
 	template <typename Container>
 	void RunEmptyContainerSelfEqualityTest()
 	{
@@ -570,6 +613,8 @@ bool FContainersFullTest::RunTest(const FString& Parameters)
 	RunContainerTests<TMap<int32, FContainerTestValueType, TFixedSetAllocator<64>>, int32>();
 	RunContainerTests<TMap<FString, FContainerTestValueType, FDefaultSetAllocator, FCaseSensitiveLookupKeyFuncs<FContainerTestValueType>>, FString>();
 
+	RunContainerConsistencyTests<TMap<int32, FContainerTestValueType>, int32>();
+
 	RunContainerTests<TSortedMap<int32, FContainerTestValueType>, int32>();
 	RunContainerTests<TSortedMap<FName, FContainerTestValueType, FDefaultAllocator, FNameLexicalLess>, FName>();
 	RunContainerTests<TSortedMap<FString, FContainerTestValueType>, FString>();
diff --git a/Engine/Source/Runtime/Core/Public/Containers/Map.h b/Engine/Source/Runtime/Core/Public/Containers/Map.h
--- a/Engine/Source/Runtime/Core/Public/Containers/Map.h
+++ b/Engine/Source/Runtime/Core/Public/Containers/Map.h
@@ -1044,6 +1044,15 @@ public:
 		Super::Pairs.StableSort(FValueComparisonClass<PREDICATE_CLASS>(Predicate));
 	}
 
+	/**
+	 * Sort the free element list so that subsequent additions will occur in the lowest available
+	 * TSet index resulting in tighter packing without moving any existing items. Also useful for
+	 * some types of determinism. @see TSparseArray::SortFreeList() for more info.
+	 */
+	void SortFreeList()
+	{
+		Super::Pairs.SortFreeList();
+	}
 
 private:
 
diff --git a/Engine/Source/Runtime/Core/Public/Containers/Set.h b/Engine/Source/Runtime/Core/Public/Containers/Set.h
--- a/Engine/Source/Runtime/Core/Public/Containers/Set.h
+++ b/Engine/Source/Runtime/Core/Public/Containers/Set.h
@@ -1104,6 +1104,16 @@ public:
 		Rehash();
 	}
 
+	/**
+	* Sort the free element list so that subsequent additions will occur in the lowest available
+	* TSparseArray index resulting in tighter packing without moving any existing items. Also useful for
+	* some types of determinism. @see TSparseArray::SortFreeList() for more info.
+	*/
+	void SortFreeList()
+	{
+		Elements.SortFreeList();
+	}
+
 	/**
 	 * Describes the set's contents through an output device.
 	 * @param Ar - The output device to describe the set's contents through.
diff --git a/Engine/Source/Runtime/Core/Public/Containers/SparseArray.h b/Engine/Source/Runtime/Core/Public/Containers/SparseArray.h
--- a/Engine/Source/Runtime/Core/Public/Containers/SparseArray.h
+++ b/Engine/Source/Runtime/Core/Public/Containers/SparseArray.h
@@ -618,6 +618,50 @@ public:
 		StableSort(TLess< ElementType >());
 	}
 
+	/**
+	* Sort the free element list so that subsequent allocations will occur in the lowest available 
+	* position resulting in tighter packing without moving any existing items. This also means 
+	* that assigned indices no longer depend on the order in which old items were removed, making
+	* it easier to use the container when determinism is required (without a container reset).
+	* 
+	* E.g., call SortFreeList() each frame to make the container assign the same indices when we 
+	* perform the following operations:
+	*	Frame1: Add(A) -> [0]; Add(B) -> [1]; Remove(A); Remove(B); (Free list is now {[1],[0],...})
+	*	Frame2: Add(A) -> [1]; Add(B) -> [0]; Remove(A); Remove(B); (Free list is now {[0],[1],...})
+	* 
+	* NOTE: This is operation is currently O(N) with N = GetMaxIndex(). This could be improved for
+	* large mostly-full arrays if necessary.
+	*/
+	void SortFreeList()
+	{
+		FElementOrFreeListLink* DataPtr = (FElementOrFreeListLink*)Data.GetData();
+		int32 CurrentHeadIndex = INDEX_NONE;
+		int32 NumFreeIndicesProcessed = 0;
+
+		// Reverse iteration to build the list from low to high indices
+		for (int32 Index = Data.Num() - 1; NumFreeIndicesProcessed < NumFreeIndices; --Index)
+		{
+			// If we have an unused element, add it to the free list
+			if (!IsValidIndex(Index))
+			{
+				DataPtr[Index].PrevFreeIndex = INDEX_NONE;
+				DataPtr[Index].NextFreeIndex = INDEX_NONE;
+
+				if (CurrentHeadIndex != INDEX_NONE)
+				{
+					DataPtr[CurrentHeadIndex].PrevFreeIndex = Index;
+					DataPtr[Index].NextFreeIndex = CurrentHeadIndex;
+				}
+
+				CurrentHeadIndex = Index;
+				++NumFreeIndicesProcessed;
+			}
+		}
+
+		// Set up the new head
+		FirstFreeIndex = CurrentHeadIndex;
+	}
+
 	/** 
 	 * Helper function to return the amount of memory allocated by this container 
 	 * Only returns the size of allocations made directly by the container, not the elements themselves.
