Moved get() out of TTuple and simplified it into a single template in order to reduce compiler overhead.
Removed structured binding guards which were for C++14 compatibility.

diff --git a/Engine/Source/Runtime/Core/Public/Templates/Tuple.h b/Engine/Source/Runtime/Core/Public/Templates/Tuple.h
--- a/Engine/Source/Runtime/Core/Public/Templates/Tuple.h
+++ b/Engine/Source/Runtime/Core/Public/Templates/Tuple.h
@@ -37,12 +37,6 @@
 	#define UE_TUPLE_STATIC_ANALYSIS_WORKAROUND 0
 #endif
 
-#if defined(__cpp_structured_bindings)
-	#define UE_TUPLE_STRUCTURED_BINDING_SUPPORT 1
-#else
-	#define UE_TUPLE_STRUCTURED_BINDING_SUPPORT 0
-#endif
-
 class FArchive;
 
 template <typename... Types>
@@ -738,19 +732,6 @@ public:
 		UE::Core::Private::Tuple::Assign(*this, MoveTemp(Other), TMakeIntegerSequence<uint32, sizeof...(OtherTypes)>{});
 		return *this;
 	}
-
-#if UE_TUPLE_STRUCTURED_BINDING_SUPPORT
-	// TTuple support for structured binding - not intended to be called directly
-	template <int N> friend decltype(auto) get(               TTuple&  val) { return static_cast<               TTuple& >(val).template Get<N>(); }
-	template <int N> friend decltype(auto) get(const          TTuple&  val) { return static_cast<const          TTuple& >(val).template Get<N>(); }
-	template <int N> friend decltype(auto) get(      volatile TTuple&  val) { return static_cast<      volatile TTuple& >(val).template Get<N>(); }
-	template <int N> friend decltype(auto) get(const volatile TTuple&  val) { return static_cast<const volatile TTuple& >(val).template Get<N>(); }
-	template <int N> friend decltype(auto) get(               TTuple&& val) { return static_cast<               TTuple&&>(val).template Get<N>(); }
-	template <int N> friend decltype(auto) get(const          TTuple&& val) { return static_cast<const          TTuple&&>(val).template Get<N>(); }
-	template <int N> friend decltype(auto) get(      volatile TTuple&& val) { return static_cast<      volatile TTuple&&>(val).template Get<N>(); }
-	template <int N> friend decltype(auto) get(const volatile TTuple&& val) { return static_cast<const volatile TTuple&&>(val).template Get<N>(); }
-#endif
-
 };
 
 template <typename... Types>
@@ -956,8 +937,22 @@ FORCEINLINE TTuple<Types&...> Tie(Types&... Args)
 	return TTuple<Types&...>(Args...);
 }
 
-#if UE_TUPLE_STRUCTURED_BINDING_SUPPORT
-// TTuple support for structured bindings
+template <typename T> constexpr bool TIsTuple_V = false;
+
+template <typename... Types> constexpr bool TIsTuple_V<               TTuple<Types...>> = true;
+template <typename... Types> constexpr bool TIsTuple_V<const          TTuple<Types...>> = true;
+template <typename... Types> constexpr bool TIsTuple_V<      volatile TTuple<Types...>> = true;
+template <typename... Types> constexpr bool TIsTuple_V<const volatile TTuple<Types...>> = true;
+
+template <typename T>
+struct TIsTuple
+{
+	enum { Value = TIsTuple_V<T> };
+};
+
+///////////////////////////////////////////////////////////////////////////////
+// TTuple support for structured bindings - not intended to be used directly //
+///////////////////////////////////////////////////////////////////////////////
 template <typename... ArgTypes>
 struct std::tuple_size<TTuple<ArgTypes...>>
 	: std::integral_constant<std::size_t, sizeof...(ArgTypes)>
@@ -975,20 +970,19 @@ struct std::tuple_element<N, TTuple<ArgTypes...>>
 public:
 	using type = typename TTupleElement<N, TTuple<ArgTypes...>>::Type;
 };
-#endif
-
-template <typename T> constexpr bool TIsTuple_V = false;
-
-template <typename... Types> constexpr bool TIsTuple_V<               TTuple<Types...>> = true;
-template <typename... Types> constexpr bool TIsTuple_V<const          TTuple<Types...>> = true;
-template <typename... Types> constexpr bool TIsTuple_V<      volatile TTuple<Types...>> = true;
-template <typename... Types> constexpr bool TIsTuple_V<const volatile TTuple<Types...>> = true;
 
-template <typename T>
-struct TIsTuple
+template <
+	int N,
+	typename TupleType
+	UE_REQUIRES(TIsTuple_V<std::decay_t<TupleType>>)
+>
+decltype(auto) get(TupleType&& val)
 {
-	enum { Value = TIsTuple_V<T> };
-};
+	return ((TupleType&&)val).template Get<N>();
+}
+///////////////////////////////////////
+// End of structured binding support //
+///////////////////////////////////////
 
 template <typename... Types>
 FORCEINLINE FArchive& operator<<(FArchive& Ar, TTuple<Types...>& Tuple)
