Fixed Run Enterprise UnrealEditor Python Tests - UnrealEditor timeout on test_dataprep_import
  - Root cause: Call to FSelfRegisteringExec::FSelfRegisteringExec() subsequent to a call to FSelfRegisteringExec::StaticExec in the same thread
  - Solution: Replace use of FRWLock with FCriticalSection which allows such scenario

diff --git a/Engine/Source/Runtime/Core/Private/Misc/CoreMisc.cpp b/Engine/Source/Runtime/Core/Private/Misc/CoreMisc.cpp
--- a/Engine/Source/Runtime/Core/Private/Misc/CoreMisc.cpp
+++ b/Engine/Source/Runtime/Core/Private/Misc/CoreMisc.cpp
@@ -44,10 +44,12 @@ DEFINE_LOG_CATEGORY(LogCore);
 using FSelfRegisteredExecArray = TArray<FSelfRegisteringExec*, TInlineAllocator<8>>;
 
 // Lazy because pthread implementation doesn't like static initialization.
-FRWLock& GetExecRegistryLock()
+FCriticalSection* GetExecRegistryLock()
 {
-	static FRWLock ExecRegistryLock;
-	return ExecRegistryLock;
+	// Note: Using FCriticalSection still allows calls to addition or removal
+	// to/from Execs and FSelfRegisteringExec::StaticExec on the same thread
+	static FCriticalSection ExecRegistryLock;
+	return &ExecRegistryLock;
 }
 
 FSelfRegisteredExecArray& GetExecRegistry()
@@ -59,20 +61,20 @@ FSelfRegisteredExecArray& GetExecRegistry()
 /** Constructor, registering this instance. */
 FSelfRegisteringExec::FSelfRegisteringExec()
 {
-	FWriteScopeLock ScopeLock(GetExecRegistryLock());
+	FScopeLock ScopeLock(GetExecRegistryLock());
 	GetExecRegistry().Add( this );
 }
 
 /** Destructor, unregistering this instance. */
 FSelfRegisteringExec::~FSelfRegisteringExec()
 {
-	FWriteScopeLock ScopeLock(GetExecRegistryLock());
+	FScopeLock ScopeLock(GetExecRegistryLock());
 	verify(GetExecRegistry().Remove( this ) == 1 );
 }
 
 bool FSelfRegisteringExec::StaticExec( UWorld* InWorld, const TCHAR* Cmd, FOutputDevice& Ar )
 {
-	FReadScopeLock ScopeLock(GetExecRegistryLock());
+	FScopeLock ScopeLock(GetExecRegistryLock());
 	for (FSelfRegisteringExec* Exe : GetExecRegistry())
 	{
 		if (Exe->Exec( InWorld, Cmd,Ar ))
