PR #9886: Fixed linker error when using operator+ in `FResourceSizeEx` outside of Core (Contributed by [REDACTED_PERSON])

#github https://github.com/EpicGames/UnrealEngine/pull/9886

diff --git a/Engine/Source/Runtime/Core/Public/ProfilingDebugging/ResourceSize.h b/Engine/Source/Runtime/Core/Public/ProfilingDebugging/ResourceSize.h
--- a/Engine/Source/Runtime/Core/Public/ProfilingDebugging/ResourceSize.h
+++ b/Engine/Source/Runtime/Core/Public/ProfilingDebugging/ResourceSize.h
@@ -111,11 +111,6 @@ public:
 	 */
 	FResourceSizeEx& operator+=(const FResourceSizeEx& InRHS);
 
-	/**
-	 * Add two FResourceSizeEx instances together and return a copy.
-	 */
-	friend FResourceSizeEx operator+(FResourceSizeEx InLHS, const FResourceSizeEx& InRHS);
-
 private:
 	/**
 	 * Type of resource size held in this struct.
@@ -140,3 +135,8 @@ private:
 	 */
 	TMap<FName, SIZE_T> UnknownMemoryBytesMap;
 };
+
+/**
+* Add two FResourceSizeEx instances together and return a copy.
+*/
+CORE_API FResourceSizeEx operator+(FResourceSizeEx InLHS, const FResourceSizeEx& InRHS);
