Runtime/Mac: Undefine PF_MAX so it can be defined by UE for pixelformat instead of socket.h

With XCode14.x, there were changes made to PrintCore.h to include networking code. When we include Cocoa.h, we end up pulling this code in via a long include chain. This results in now pulling in socket.h when we don't need it and PF_MAX getting defined by socket.h instead of by UE.

diff --git a/Engine/Source/Runtime/Core/Public/Mac/MacSystemIncludes.h b/Engine/Source/Runtime/Core/Public/Mac/MacSystemIncludes.h
--- a/Engine/Source/Runtime/Core/Public/Mac/MacSystemIncludes.h
+++ b/Engine/Source/Runtime/Core/Public/Mac/MacSystemIncludes.h
@@ -7,6 +7,7 @@
 #define FVector FVectorWorkaround
 #ifdef __OBJC__
 #import <Cocoa/Cocoa.h>
+#undef PF_MAX
 #endif
 #include <Carbon/Carbon.h>
 #include <CoreServices/CoreServices.h>
