LLT test adapter: fix flipped bool in return of error state
#rnx

diff --git a/Engine/Source/Runtime/Core/Public/Misc/LowLevelTestAdapter.h b/Engine/Source/Runtime/Core/Public/Misc/LowLevelTestAdapter.h
--- a/Engine/Source/Runtime/Core/Public/Misc/LowLevelTestAdapter.h
+++ b/Engine/Source/Runtime/Core/Public/Misc/LowLevelTestAdapter.h
@@ -39,7 +39,7 @@
 			void TestBody(const FString& Parameters); \
 			virtual bool RunTest(const FString& Parameters) { \
 				TestBody(Parameters); \
-				return HasAnyErrors(); \
+				return !HasAnyErrors(); \
 			} \
 			virtual FString GetBeautifiedTestName() const override { return PrettyNameDotNotation; } \
 		private:\
