Migrate tests from files RobinHoodHashTableTest.cpp, StridedViewTest.cpp, TickerTests.cpp

diff --git a/Engine/Source/Runtime/Core/Tests/Containers/RobinHoodHashTableTest.cpp b/Engine/Source/Runtime/Core/Tests/Containers/RobinHoodHashTableTest.cpp
--- a/Engine/Source/Runtime/Core/Tests/Containers/RobinHoodHashTableTest.cpp
+++ b/Engine/Source/Runtime/Core/Tests/Containers/RobinHoodHashTableTest.cpp
@@ -36,7 +36,7 @@ TEST_CASE_NAMED(FRobinHoodHashTableTest, "System::Core::Containers::RobinHoodSha
 			bool bAlreadyExists = false;
 			const int32* FoundPtr = TableConst.Find(Index);
 			CHECK_MESSAGE(TEXT("FoundPtr"), FoundPtr != nullptr);
-			CHECK_MESSAGE(TEXT("*FoundPtr"), *FoundPtr == Index);
+			CHECK_MESSAGE(TEXT("*FoundPtr"), ((FoundPtr != nullptr) && (*FoundPtr == Index)));
 		}
 	}
 }
