I/O - hex to FIoChunkId

#rnx

diff --git a/Engine/Source/Runtime/Core/Private/IO/IoChunkId.cpp b/Engine/Source/Runtime/Core/Private/IO/IoChunkId.cpp
--- a/Engine/Source/Runtime/Core/Private/IO/IoChunkId.cpp
+++ b/Engine/Source/Runtime/Core/Private/IO/IoChunkId.cpp
@@ -6,6 +6,7 @@
 #include "Memory/MemoryView.h"
 #include "Serialization/Archive.h"
 #include "Serialization/CompactBinaryWriter.h"
+#include "String/HexToBytes.h"
 
 FString LexToString(const EIoChunkType Type)
 {
@@ -54,6 +55,18 @@ void FIoChunkId::ToString(FString& Output) const
 	CharArray.Last() = TCHAR('\0');
 }
 
+FIoChunkId FIoChunkId::FromHex(FStringView Hex)
+{
+	if (Hex.Len() == sizeof(FIoChunkId) * 2)
+	{
+		FIoChunkId ChunkId;
+		UE::String::HexToBytes(Hex, ChunkId.Id);
+		return ChunkId;
+	}
+
+	return InvalidChunkId;
+}
+
 FArchive& operator<<(FArchive& Ar, FIoChunkId& ChunkId)
 {
 	Ar.Serialize(&ChunkId.Id, sizeof FIoChunkId::Id);
diff --git a/Engine/Source/Runtime/Core/Public/IO/IoChunkId.h b/Engine/Source/Runtime/Core/Public/IO/IoChunkId.h
--- a/Engine/Source/Runtime/Core/Public/IO/IoChunkId.h
+++ b/Engine/Source/Runtime/Core/Public/IO/IoChunkId.h
@@ -113,6 +113,12 @@ public:
 
 	friend class FIoStoreReaderImpl;
 
+	/**
+	 * Creates an I/O chunk ID from a 24 character long hexadecimal string.
+	 * @return The corresponding I/O chunk ID or an invalid ID if the input string is not in the correct format.
+	 */
+	CORE_API static FIoChunkId FromHex(FStringView Hex);
+
 private:
 	static inline FIoChunkId CreateEmptyId()
 	{
