Add UniqueBy to remove duplicates from a range via a projection to apply to the elements, similar to 'std::ranges::unique'

diff --git a/Engine/Source/Runtime/Core/Public/Algo/Unique.h b/Engine/Source/Runtime/Core/Public/Algo/Unique.h
--- a/Engine/Source/Runtime/Core/Public/Algo/Unique.h
+++ b/Engine/Source/Runtime/Core/Public/Algo/Unique.h
@@ -3,13 +3,14 @@
 #pragma once
 
 #include "Templates/UnrealTemplate.h"
+#include "Templates/IdentityFunctor.h"
 #include "Templates/Invoke.h"
 #include "Templates/EqualTo.h"
 
 namespace AlgoImpl
 {
-	template <typename T, typename SizeType, typename BinaryPredicate>
-	SizeType Unique(T* Array, SizeType ArraySize, BinaryPredicate Predicate)
+	template <typename T, typename SizeType, typename ProjectionType, typename BinaryPredicate>
+	SizeType Unique(T* Array, SizeType ArraySize, ProjectionType Proj, BinaryPredicate Predicate)
 	{
 		if (ArraySize <= 1)
 		{
@@ -19,7 +20,7 @@ namespace AlgoImpl
 		T* Result = Array;
 		for (T* Iter = Array + 1; Iter != Array + ArraySize; ++Iter)
 		{
-			if (!Invoke(Predicate, *Result, *Iter)) 
+			if (!Invoke(Predicate, Invoke(Proj, *Result), Invoke(Proj, *Iter))) 
 			{
 				++Result;
 				if (Result != Iter)
@@ -55,14 +56,26 @@ namespace Algo
 	 * See https://en.cppreference.com/w/cpp/algorithm/unique
 	 */
 	template<typename RangeType>
-	auto Unique(RangeType&& Range) -> decltype(AlgoImpl::Unique(GetData(Range), GetNum(Range), TEqualTo<>{}))
+	auto Unique(RangeType&& Range) -> decltype(AlgoImpl::Unique(GetData(Range), GetNum(Range), FIdentityFunctor(), TEqualTo<>{}))
 	{
-		return AlgoImpl::Unique(GetData(Range), GetNum(Range), TEqualTo<>{});
+		return AlgoImpl::Unique(GetData(Range), GetNum(Range), FIdentityFunctor(), TEqualTo<>{});
 	}
 
 	template<typename RangeType, typename BinaryPredicate>
-	auto Unique(RangeType&& Range, BinaryPredicate Predicate) -> decltype(AlgoImpl::Unique(GetData(Range), GetNum(Range), MoveTemp(Predicate)))
+	auto Unique(RangeType&& Range, BinaryPredicate Predicate) -> decltype(AlgoImpl::Unique(GetData(Range), GetNum(Range), FIdentityFunctor(), MoveTemp(Predicate)))
 	{
-		return AlgoImpl::Unique(GetData(Range), GetNum(Range), MoveTemp(Predicate));
+		return AlgoImpl::Unique(GetData(Range), GetNum(Range), FIdentityFunctor(), MoveTemp(Predicate));
+	}
+
+	template<typename RangeType, typename ProjectionType>
+	auto UniqueBy(RangeType&& Range, ProjectionType Proj) -> decltype(AlgoImpl::Unique(GetData(Range), GetNum(Range), MoveTemp(Proj), TEqualTo<>{}))
+	{
+		return AlgoImpl::Unique(GetData(Range), GetNum(Range), MoveTemp(Proj), TEqualTo<>{});
+	}
+
+	template<typename RangeType, typename ProjectionType, typename BinaryPredicate>
+	auto UniqueBy(RangeType&& Range, ProjectionType Proj, BinaryPredicate Predicate) -> decltype(AlgoImpl::Unique(GetData(Range), GetNum(Range), MoveTemp(Proj), MoveTemp(Predicate)))
+	{
+		return AlgoImpl::Unique(GetData(Range), GetNum(Range), MoveTemp(Proj), MoveTemp(Predicate));
 	}
 }
diff --git a/Engine/Source/Runtime/Core/Tests/Algo/UniqueTest.cpp b/Engine/Source/Runtime/Core/Tests/Algo/UniqueTest.cpp
--- a/Engine/Source/Runtime/Core/Tests/Algo/UniqueTest.cpp
+++ b/Engine/Source/Runtime/Core/Tests/Algo/UniqueTest.cpp
@@ -49,4 +49,48 @@ TEST_CASE_NAMED(FUniqueTest, "System::Core::Algo::Unique", "[ApplicationContextM
 	}
 }
 
+TEST_CASE_NAMED(FUniqueByTest, "System::Core::Algo::UniqueBy", "[ApplicationContextMask][SmokeFilter]")
+{
+	using namespace Algo;
+
+	using FIntPair = TPair<int32, int32>;
+
+	{
+		TArray<FIntPair> Array;
+		int32 RemoveFrom = UniqueBy(Array, &FIntPair::Key);
+		CHECK_MESSAGE(TEXT("`UniqueBy` must handle an empty container"), RemoveFrom == 0);
+	}
+	{
+		TArray<FIntPair> Array{ { 1, 0 }, { 2, 0 }, { 3, 0} };
+		Array.SetNum(UniqueBy(Array, &FIntPair::Key));
+		CHECK_MESSAGE(TEXT("`UniqueBy` with no duplicates must remain unchanged"), (Array == TArray<FIntPair>{ { 1, 0 }, { 2, 0 }, { 3, 0 } }));
+	}
+	{
+		TArray<FIntPair> Array{ { 1, 0 }, { 1, 0 }, { 2, 0}, { 2, 0 }, { 2, 0 }, { 3, 0 }, { 3, 0 }, { 3, 0 }, { 3, 0 }};
+		Array.SetNum(UniqueBy(Array, &FIntPair::Key));
+		CHECK_MESSAGE(TEXT("`UniqueBy` with multiple duplicates must return correct result"), (Array == TArray<FIntPair>{ { 1, 0 }, { 2, 0 }, { 3, 0 } }));
+	}
+	{
+		TArray<FIntPair> Array{ { 1, 0 }, { 1, 0 }, { 2, 0}, { 3, 0 }, { 3, 0 }, { 3, 0 } };
+		Array.SetNum(UniqueBy(Array, &FIntPair::Key));
+		CHECK_MESSAGE(TEXT("`UniqueBy` with duplicates and unique items must return correct result"), (Array == TArray<FIntPair>{ { 1, 0 }, { 2, 0 }, { 3, 0 } }));
+	}
+	{
+		FString Str = TEXT("aa");
+		Str = Str.Mid(0, UniqueBy(Str, FIdentityFunctor()));
+		CHECK_MESSAGE(TEXT("`UniqueBy` on `FString` as an example of arbitrary random-access container must compile and return correct result"),
+			Str == FString(TEXT("a")));
+	}
+	{
+		FIntPair Array[] = { { 1, 0 } };
+		int32 NewSize = (int32)UniqueBy(Array, &FIntPair::Key);
+		CHECK_MESSAGE(TEXT("`UniqueBy` must support C arrays"), NewSize == 1);
+	}
+	{
+		TArray<FIntPair> Array = { { 1, 0 }, { 1, 0 } };
+		int32 NewSize = UniqueBy(MakeArrayView(Array.GetData() + 1, 1), &FIntPair::Key);
+		CHECK_MESSAGE(TEXT("`UniqueBy` must support ranges"), NewSize == 1);
+	}
+}
+
 #endif // WITH_TESTS
