StridedView fixes
- Adding IsValidIndex and IsEmpty methods to match ArrayView.
- Switching to using std::is_base_of_v instead of TIsDerivedFrom since the latter can't handle const types.
- Adding TConstStridedView
- Adding forward declarations of TStridedView and TConstStridedView to ContainersFwd.h

diff --git a/Engine/Source/Runtime/Core/Public/Containers/ContainersFwd.h b/Engine/Source/Runtime/Core/Public/Containers/ContainersFwd.h
--- a/Engine/Source/Runtime/Core/Public/Containers/ContainersFwd.h
+++ b/Engine/Source/Runtime/Core/Public/Containers/ContainersFwd.h
@@ -31,4 +31,6 @@ template <typename> struct TTypeTraits;
 template<typename KeyType, typename ValueType, typename ArrayAllocator = FDefaultAllocator, typename SortPredicate = TLess<typename TTypeTraits<KeyType>::ConstPointerType> > class TSortedMap;
 template<typename ElementType,bool bInAllowDuplicateKeys = false> struct DefaultKeyFuncs;
 template<typename InElementType, typename KeyFuncs = DefaultKeyFuncs<InElementType>, typename Allocator = FDefaultSetAllocator> class TSet;
+template<typename InElementType, typename InSizeType = int32> class TStridedView;
+template<typename T, typename SizeType = int32> using TConstStridedView = TStridedView<const T, SizeType>;
 /// @endcond
diff --git a/Engine/Source/Runtime/Core/Public/Containers/StridedView.h b/Engine/Source/Runtime/Core/Public/Containers/StridedView.h
--- a/Engine/Source/Runtime/Core/Public/Containers/StridedView.h
+++ b/Engine/Source/Runtime/Core/Public/Containers/StridedView.h
@@ -3,6 +3,7 @@
 #pragma once
 
 #include "CoreTypes.h"
+#include "Containers/ContainersFwd.h"
 #include "Templates/UnrealTemplate.h"
 #include "Misc/AssertionMacros.h"
 #include "Containers/ArrayView.h"
@@ -38,7 +39,7 @@
 */
 
 /** Pointer with extent and a stride, similar to TArrayView. Designed to allow functions to take pointers to arbitrarily structured data. */
-template<typename InElementType, typename InSizeType = int32>
+template<typename InElementType, typename InSizeType>
 class TStridedView
 {
 public:
@@ -73,6 +74,16 @@ public:
 		}
 	}
 
+	FORCEINLINE bool IsValidIndex(SizeType Index) const
+	{
+		return (Index >= 0) && (Index < NumElements);
+	}
+
+	FORCEINLINE bool IsEmpty() const
+	{
+		return NumElements == 0;
+	}
+
 	FORCEINLINE SizeType Num() const
 	{
 		return NumElements;
@@ -160,7 +171,7 @@ TStridedView<ElementType> MakeStridedView(int32 BytesBetweenElements, ElementTyp
 template <typename BaseStructureType, typename DerivedStructureType>
 TStridedView<BaseStructureType> MakeStridedViewOfBase(TArrayView<DerivedStructureType> StructuredView)
 {
-	static_assert(TIsDerivedFrom<DerivedStructureType, BaseStructureType>::IsDerived, "Expecting derived structure type");
+	static_assert(std::is_base_of_v<BaseStructureType, DerivedStructureType>, "Expecting derived structure type");
 	return MakeStridedView<BaseStructureType>((int32)sizeof(DerivedStructureType), GetData(StructuredView), (int32)GetNum(StructuredView));
 }
 
