Added a console variable to control whether the crash reporter dialogs should be forced off.

#rnx

diff --git a/Engine/Source/Runtime/Core/Private/Windows/WindowsPlatformCrashContext.cpp b/Engine/Source/Runtime/Core/Private/Windows/WindowsPlatformCrashContext.cpp
--- a/Engine/Source/Runtime/Core/Private/Windows/WindowsPlatformCrashContext.cpp
+++ b/Engine/Source/Runtime/Core/Private/Windows/WindowsPlatformCrashContext.cpp
@@ -160,6 +160,15 @@ struct FAssertInfo
 
 const TCHAR* const FWindowsPlatformCrashContext::UEGPUAftermathMinidumpName = TEXT("UEAftermathD3D12.nv-gpudmp");
 
+namespace UE::Core::Private
+{
+	static TAutoConsoleVariable<bool> CVarForceCrashReportDialogOff(
+		TEXT("WindowsPlatformCrashContext.ForceCrashReportDialogOff"),
+		false,
+		TEXT("If true, force the crash report dialog to not be displayed in the event of a crash."),
+		ECVF_Default);
+}
+
 /**
 * Implement platform specific static cleanup function
 */
@@ -515,9 +524,11 @@ int32 ReportCrashForMonitor(
 	SharedContext->ExceptionProgramCounter = ExceptionInfo->ExceptionRecord->ExceptionAddress;
 
 	// Determine UI settings for the crash report. Suppress the user input dialog if we're running in unattended mode
+	// or if it is forced of through the WindowsPlatformCrashContext.ForceCrashReportDialogOff console variable.
 	// Usage data controls if we want analytics in the crash report client
 	// Finally we cannot call some of these functions if we crash during static init, so check if they are initialized.
 	bool bNoDialog = ReportUI == EErrorReportUI::ReportInUnattendedMode;
+	bNoDialog |= UE::Core::Private::CVarForceCrashReportDialogOff.GetValueOnAnyThread() == true;
 	bool bSendUnattendedBugReports = true;
 	bool bSendUsageData = true;
 	bool bCanSendCrashReport = true;
@@ -725,8 +736,10 @@ int32 ReportCrashUsingCrashReportClient(FWindowsPlatformCrashContext& InContext,
 	const TCHAR* ExecutableName = FPlatformProcess::ExecutableName();
 	bool bCanRunCrashReportClient = FCString::Stristr( ExecutableName, TEXT( "CrashReportClient" ) ) == nullptr;
 
-	// Suppress the user input dialog if we're running in unattended mode
+	// Suppress the user input dialog if we're running in unattended mode or if it is forced of through the 
+	// WindowsPlatformCrashContext.ForceCrashReportDialogOff console variable.
 	bool bNoDialog = FApp::IsUnattended() || ReportUI == EErrorReportUI::ReportInUnattendedMode || IsRunningDedicatedServer();
+	bNoDialog |= UE::Core::Private::CVarForceCrashReportDialogOff.GetValueOnAnyThread() == true;
 
 	bool bImplicitSend = false;
 #if !UE_EDITOR
