Added missing HAL/IConsoleManager.h include.

#rnx

diff --git a/Engine/Source/Runtime/Core/Private/Windows/WindowsPlatformCrashContext.cpp b/Engine/Source/Runtime/Core/Private/Windows/WindowsPlatformCrashContext.cpp
--- a/Engine/Source/Runtime/Core/Private/Windows/WindowsPlatformCrashContext.cpp
+++ b/Engine/Source/Runtime/Core/Private/Windows/WindowsPlatformCrashContext.cpp
@@ -9,6 +9,7 @@
 #include "HAL/ThreadHeartBeat.h"
 #include "HAL/PlatformProcess.h"
 #include "HAL/FileManager.h"
+#include "HAL/IConsoleManager.h"
 #include "HAL/PlatformOutputDevices.h"
 #include "HAL/PlatformTLS.h"
 #include "HAL/PlatformTime.h"
