CompactBinary: Fixed a smoke test failure that slipped in due to UE-161113

#rnx

diff --git a/Engine/Source/Runtime/Core/Private/Tests/Serialization/CompactBinaryTest.cpp b/Engine/Source/Runtime/Core/Private/Tests/Serialization/CompactBinaryTest.cpp
--- a/Engine/Source/Runtime/Core/Private/Tests/Serialization/CompactBinaryTest.cpp
+++ b/Engine/Source/Runtime/Core/Private/Tests/Serialization/CompactBinaryTest.cpp
@@ -404,7 +404,7 @@ bool FCbFieldObjectTest::RunTest(const FString& Parameters)
 		FCbFieldView NamedField(NamedValue, ECbFieldType::UniformObject | ECbFieldType::HasFieldName);
 		TestTrue(TEXT("FCbObjectView::Equals()"), Field.AsObjectView().Equals(NamedField.AsObjectView()));
 		TestTrue(TEXT("FCbObjectView::AsFieldView().Equals()"), Field.Equals(Field.AsObjectView().AsFieldView()));
-		TestTrue(TEXT("FCbObjectView::AsFieldView().Equals()"), NamedField.Equals(NamedField.AsObjectView().AsFieldView()));
+		TestTrue(TEXT("FCbObjectView::AsFieldView().Equals()"), NamedField.RemoveName().Equals(NamedField.AsObjectView().AsFieldView()));
 
 		// CopyTo
 		uint8 CopyBytes[sizeof(Value) + 1];
@@ -599,7 +599,7 @@ bool FCbFieldArrayTest::RunTest(const FString& Parameters)
 		FCbFieldView NamedField(NamedValue, ECbFieldType::UniformArray | ECbFieldType::HasFieldName);
 		TestTrue(TEXT("FCbArrayView::Equals()"), Field.AsArrayView().Equals(NamedField.AsArrayView()));
 		TestTrue(TEXT("FCbArrayView::AsFieldView().Equals()"), Field.Equals(Field.AsArrayView().AsFieldView()));
-		TestTrue(TEXT("FCbArrayView::AsFieldView().Equals()"), NamedField.Equals(NamedField.AsArrayView().AsFieldView()));
+		TestTrue(TEXT("FCbArrayView::AsFieldView().Equals()"), NamedField.RemoveName().Equals(NamedField.AsArrayView().AsFieldView()));
 
 		// CopyTo
 		uint8 CopyBytes[sizeof(Value) + 1];
