AutomationTest: Fixed calculation of the client and server contexts

#rnx

diff --git a/Engine/Source/Runtime/Core/Private/Misc/AutomationTest.cpp b/Engine/Source/Runtime/Core/Private/Misc/AutomationTest.cpp
--- a/Engine/Source/Runtime/Core/Private/Misc/AutomationTest.cpp
+++ b/Engine/Source/Runtime/Core/Private/Misc/AutomationTest.cpp
@@ -612,10 +612,10 @@ void FAutomationTestFramework::GetValidTestNames( TArray<FAutomationTestInfo>& T
 	TestInfo.Empty();
 
 	// Determine required application type (Editor, Game, or Commandlet)
-	const bool bRunningEditor = GIsEditor && !IsRunningCommandlet();
-	const bool bRunningClient = IsRunningGame() || IsRunningClientOnly();
-	const bool bRunningServer = IsRunningGame() || IsRunningDedicatedServer();
 	const bool bRunningCommandlet = IsRunningCommandlet();
+	const bool bRunningEditor = GIsEditor && !bRunningCommandlet;
+	const bool bRunningClient = !GIsEditor && !IsRunningDedicatedServer();
+	const bool bRunningServer = !GIsEditor && !IsRunningClientOnly();
 
 	//application flags
 	uint32 ApplicationSupportFlags = 0;
