AutomationTest expected message verbosity check fix (the warning was treated as verbose!)

diff --git a/Engine/Source/Runtime/Core/Private/Misc/AutomationTest.cpp b/Engine/Source/Runtime/Core/Private/Misc/AutomationTest.cpp
--- a/Engine/Source/Runtime/Core/Private/Misc/AutomationTest.cpp
+++ b/Engine/Source/Runtime/Core/Private/Misc/AutomationTest.cpp
@@ -1599,7 +1599,7 @@ bool FAutomationTestBase::IsExpectedMessage(
 		// Maintains previous behavior: Adjust so that error and fatal messages are tested against when the input verbosity is "Warning"
 		// Similarly, any message above warning should be considered an "info" message
 		const ELogVerbosity::Type AdjustedMessageVerbosity =
-			ExpectedMessage.Verbosity < ELogVerbosity::Warning 
+			ExpectedMessage.Verbosity <= ELogVerbosity::Warning 
 			? ELogVerbosity::Warning
 			: ELogVerbosity::VeryVerbose;
 
