Removed GIsUCCMakeStandaloneHeaderGenerator since C++ UHT no longer uses engine types.

diff --git a/Engine/Source/Editor/UnrealEd/Private/EditorObject.cpp b/Engine/Source/Editor/UnrealEd/Private/EditorObject.cpp
--- a/Engine/Source/Editor/UnrealEd/Private/EditorObject.cpp
+++ b/Engine/Source/Editor/UnrealEd/Private/EditorObject.cpp
@@ -232,7 +232,6 @@ static const TCHAR* ImportProperties(
 	const TMap<AActor*, AActor*>* ActorRemapper
 	)
 {
-	check(!GIsUCCMakeStandaloneHeaderGenerator);
 	check(ObjectStruct!=NULL);
 	check(DestData!=NULL);
 
@@ -889,7 +888,6 @@ static const TCHAR* ImportCreateSubObjectsStep(
 	TMap<FString, UObject*>* ObjectRemapper
 	)
 {
-	check(!GIsUCCMakeStandaloneHeaderGenerator);
 	check(ObjectStruct!=nullptr);
 
 	if (SourceText.IsEmpty())
@@ -1281,7 +1279,6 @@ static const TCHAR* ImportPropertiesStep(
 	TSet<FProperty*>* PropertiesToSkip
 	)
 {
-	check(!GIsUCCMakeStandaloneHeaderGenerator);
 	check(ObjectStruct!=nullptr);
 	check(DestData!=nullptr);
 
diff --git a/Engine/Source/Programs/UnrealHeaderTool/Private/CodeGenerator.cpp b/Engine/Source/Programs/UnrealHeaderTool/Private/CodeGenerator.cpp
--- a/Engine/Source/Programs/UnrealHeaderTool/Private/CodeGenerator.cpp
+++ b/Engine/Source/Programs/UnrealHeaderTool/Private/CodeGenerator.cpp
@@ -7328,8 +7328,6 @@ ECompilationResult::Type UnrealHeaderTool_Main(const FString& ModuleInfoFilename
 	bIncludeDebugOutput = FParse::Param(FCommandLine::Get(), TEXT("IncludeDebugOutput"));
 	bGoWide = !FParse::Param(FCommandLine::Get(), TEXT("NoGoWide"));
 
-	check(GIsUCCMakeStandaloneHeaderGenerator);
-
 	FString ModuleInfoPath = FPaths::GetPath(ModuleInfoFilename);
 
 	// The meta data keywords must be initialized prior to going wide
diff --git a/Engine/Source/Programs/UnrealHeaderTool/Private/UnrealHeaderToolMain.cpp b/Engine/Source/Programs/UnrealHeaderTool/Private/UnrealHeaderToolMain.cpp
--- a/Engine/Source/Programs/UnrealHeaderTool/Private/UnrealHeaderToolMain.cpp
+++ b/Engine/Source/Programs/UnrealHeaderTool/Private/UnrealHeaderToolMain.cpp
@@ -82,7 +82,6 @@ INT32_MAIN_INT32_ARGC_TCHAR_ARGV()
 		FEngineLoop::AppExit();
 	};
 
-	GIsUCCMakeStandaloneHeaderGenerator = true;
 	if (GEngineLoop.PreInit(*ShortCmdLine) != 0)
 	{
 		UE_LOG(LogCompile, Error, TEXT("Failed to initialize the engine (PreInit failed)."));
diff --git a/Engine/Source/Runtime/Core/Private/Misc/CoreGlobals.cpp b/Engine/Source/Runtime/Core/Private/Misc/CoreGlobals.cpp
--- a/Engine/Source/Runtime/Core/Private/Misc/CoreGlobals.cpp
+++ b/Engine/Source/Runtime/Core/Private/Misc/CoreGlobals.cpp
@@ -178,14 +178,12 @@ bool					PRIVATE_GAllowCommandletAudio 		= false;				/** If true, allow audio ev
 #if WITH_EDITORONLY_DATA
 bool					GIsEditor						= false;					/* Whether engine was launched for editing */
 bool					GIsImportingT3D					= false;					/* Whether editor is importing T3D */
-bool					GIsUCCMakeStandaloneHeaderGenerator = false;				/* Are we rebuilding script via the standalone header generator? */
 bool					GIsTransacting					= false;					/* true if there is an undo/redo operation in progress. */
 bool					GIntraFrameDebuggingGameThread	= false;					/* Indicates that the game thread is currently paused deep in a call stack; do not process any game thread tasks */
 bool					GFirstFrameIntraFrameDebugging	= false;					/* Indicates that we're currently processing the first frame of intra-frame debugging */
 #elif USING_CODE_ANALYSIS
 // These are always false during 'non-editor code analysis', just like they would be when #defined.
 bool					GIsEditor						= false;
-bool					GIsUCCMakeStandaloneHeaderGenerator = false;
 bool					GIntraFrameDebuggingGameThread	= false;
 bool					GFirstFrameIntraFrameDebugging	= false;
 #endif // !WITH_EDITORONLY_DATA
diff --git a/Engine/Source/Runtime/Core/Public/CoreGlobals.h b/Engine/Source/Runtime/Core/Public/CoreGlobals.h
--- a/Engine/Source/Runtime/Core/Public/CoreGlobals.h
+++ b/Engine/Source/Runtime/Core/Public/CoreGlobals.h
@@ -212,7 +212,6 @@ extern CORE_API bool PRIVATE_GAllowCommandletAudio;
 */
 extern CORE_API bool GIsEditor;
 extern CORE_API bool GIsImportingT3D;
-extern CORE_API bool GIsUCCMakeStandaloneHeaderGenerator;
 extern CORE_API bool GIsTransacting;
 
 /** Indicates that the game thread is currently paused deep in a call stack,
@@ -226,14 +225,12 @@ extern CORE_API bool			GFirstFrameIntraFrameDebugging;
 
 // Defined as variables during code analysis to prevent lots of '<constant> && <expr>' warnings
 extern CORE_API bool GIsEditor;
-extern CORE_API bool GIsUCCMakeStandaloneHeaderGenerator;
 extern CORE_API bool GIntraFrameDebuggingGameThread;
 extern CORE_API bool GFirstFrameIntraFrameDebugging;
 
 #else
 
 #define GIsEditor								false
-#define GIsUCCMakeStandaloneHeaderGenerator		false
 #define GIntraFrameDebuggingGameThread			false
 #define GFirstFrameIntraFrameDebugging			false
 
diff --git a/Engine/Source/Runtime/CoreUObject/Private/UObject/Class.cpp b/Engine/Source/Runtime/CoreUObject/Private/UObject/Class.cpp
--- a/Engine/Source/Runtime/CoreUObject/Private/UObject/Class.cpp
+++ b/Engine/Source/Runtime/CoreUObject/Private/UObject/Class.cpp
@@ -2563,7 +2563,7 @@ void UScriptStruct::PrepareCppStructOps()
 		CppStructOps = GetDeferredCppStructOps().FindRef(GetFlattenedStructPathName());
 		if (!CppStructOps)
 		{
-			if (!GIsUCCMakeStandaloneHeaderGenerator && (StructFlags&STRUCT_Native))
+			if (StructFlags&STRUCT_Native)
 			{
 				UE_LOG(LogClass, Fatal, TEXT("Couldn't bind to native struct %s. Headers need to be rebuilt, or a noexport class is missing a IMPLEMENT_STRUCT."),*GetName());
 			}
@@ -4242,7 +4242,7 @@ void UClass::Bind()
 {
 	UStruct::Bind();
 
-	if( !GIsUCCMakeStandaloneHeaderGenerator && !ClassConstructor && IsNative() )
+	if( !ClassConstructor && IsNative() )
 	{
 		UE_LOG(LogClass, Fatal, TEXT("Can't bind to native class %s"), *GetPathName() );
 	}
diff --git a/Engine/Source/Runtime/Engine/Private/Level.cpp b/Engine/Source/Runtime/Engine/Private/Level.cpp
--- a/Engine/Source/Runtime/Engine/Private/Level.cpp
+++ b/Engine/Source/Runtime/Engine/Private/Level.cpp
@@ -1633,7 +1633,7 @@ void ULevel::IncrementalUpdateComponents(int32 NumComponentsToUpdate, bool bReru
 			if (IncrementalRegisterComponents(true, NumComponentsToUpdate, Context))
 			{
 #if WITH_EDITOR
-				const bool bShouldRunConstructionScripts = !bHasRerunConstructionScripts && bRerunConstructionScripts && !IsTemplate() && !GIsUCCMakeStandaloneHeaderGenerator;
+				const bool bShouldRunConstructionScripts = !bHasRerunConstructionScripts && bRerunConstructionScripts && !IsTemplate();
 				IncrementalComponentState = bShouldRunConstructionScripts ? EIncrementalComponentState::RunConstructionScripts : EIncrementalComponentState::Finalize;
 #else
 				IncrementalComponentState = EIncrementalComponentState::Finalize;
diff --git a/Engine/Source/Runtime/Launch/Private/LaunchEngineLoop.cpp b/Engine/Source/Runtime/Launch/Private/LaunchEngineLoop.cpp
--- a/Engine/Source/Runtime/Launch/Private/LaunchEngineLoop.cpp
+++ b/Engine/Source/Runtime/Launch/Private/LaunchEngineLoop.cpp
@@ -2284,21 +2284,9 @@ int32 FEngineLoop::PreInitPreStartupScreen(const TCHAR* CmdLine)
 	};
 #endif
 
-#if WITH_EDITOR
-	if (GIsUCCMakeStandaloneHeaderGenerator)
-	{
-		// Header generation is treated as a commandlet
-		checkf(!IsRunningDedicatedServer(), TEXT("GIsUCCMakeStandaloneHeaderGenerator is not supported in a ServerOnly configuration, or with -server on commandline."));
-
-		SetIsRunningAsCommandlet(TEXT("GIsUCCMakeStandaloneHeaderGenerator"));
-	}
-	else
-#endif
+	if (IsRunningDedicatedServer())
 	{
-		if (IsRunningDedicatedServer())
-		{
-			SetIsRunningAsDedicatedServer();
-		}
+		SetIsRunningAsDedicatedServer();
 	}
 
 #if UE_EDITOR || WITH_ENGINE || WITH_EDITOR
