UnrealPak: Skip DDC initialization when running without -compressionEnableDDC.

#rnx

diff --git a/Engine/Source/Runtime/Core/Private/IO/IoStore.cpp b/Engine/Source/Runtime/Core/Private/IO/IoStore.cpp
--- a/Engine/Source/Runtime/Core/Private/IO/IoStore.cpp
+++ b/Engine/Source/Runtime/Core/Private/IO/IoStore.cpp
@@ -291,7 +291,6 @@ class FIoStoreWriterContextImpl
 public:
 	FIoStoreWriterContextImpl()
 	{
-		DDC = GetDerivedDataCache();
 	}
 
 	~FIoStoreWriterContextImpl()
@@ -321,6 +320,11 @@ public:
 		WriterSettings = InWriterSettings;
 		CompressionBufferAvailableEvent = FPlatformProcess::GetSynchEventFromPool(false);
 
+		if (WriterSettings.bCompressionEnableDDC)
+		{
+			DDC = GetDerivedDataCache();
+		}
+
 		if (WriterSettings.CompressionMethod != NAME_None)
 		{
 			CompressionBufferSize = FCompression::CompressMemoryBound(WriterSettings.CompressionMethod, static_cast<int32>(WriterSettings.CompressionBlockSize));
