Removed redundant HAL/ExceptionHandling.h include and reordered includes alphabetically to track them easier.

#rnx

diff --git a/Engine/Source/Runtime/Core/Private/Windows/WindowsPlatformCrashContext.cpp b/Engine/Source/Runtime/Core/Private/Windows/WindowsPlatformCrashContext.cpp
--- a/Engine/Source/Runtime/Core/Private/Windows/WindowsPlatformCrashContext.cpp
+++ b/Engine/Source/Runtime/Core/Private/Windows/WindowsPlatformCrashContext.cpp
@@ -1,35 +1,35 @@
 // Copyright Epic Games, Inc. All Rights Reserved.
 
 #include "Windows/WindowsPlatformCrashContext.h"
-#include "HAL/PlatformMallocCrash.h"
-#include "HAL/ExceptionHandling.h"
-#include "Misc/EngineVersion.h"
-#include "Misc/EngineBuildSettings.h"
+
+#include "BuildSettings.h"
+#include "CoreGlobals.h"
 #include "HAL/ExceptionHandling.h"
-#include "HAL/ThreadHeartBeat.h"
-#include "HAL/PlatformProcess.h"
 #include "HAL/FileManager.h"
 #include "HAL/IConsoleManager.h"
+#include "HAL/PlatformMallocCrash.h"
 #include "HAL/PlatformOutputDevices.h"
-#include "HAL/PlatformTLS.h"
+#include "HAL/PlatformProcess.h"
 #include "HAL/PlatformTime.h"
+#include "HAL/PlatformTLS.h"
+#include "HAL/ThreadHeartBeat.h"
+#include "HAL/ThreadManager.h"
 #include "Internationalization/Internationalization.h"
 #include "Misc/App.h"
-#include "Misc/Paths.h"
+#include "Misc/ConfigCacheIni.h"
+#include "Misc/CoreDelegates.h"
+#include "Misc/EngineBuildSettings.h"
+#include "Misc/EngineVersion.h"
 #include "Misc/FeedbackContext.h"
 #include "Misc/MessageDialog.h"
-#include "Misc/CoreDelegates.h"
-#include "Misc/ConfigCacheIni.h"
-#include "Misc/OutputDeviceRedirector.h"
+#include "Misc/OutputDeviceArchiveWrapper.h"
 #include "Misc/OutputDeviceFile.h"
+#include "Misc/OutputDeviceRedirector.h"
+#include "Misc/Paths.h"
 #include "Serialization/Archive.h"
-#include "Windows/WindowsPlatformStackWalk.h"
 #include "Templates/UniquePtr.h"
 #include "Templates/UnrealTemplate.h"
-#include "Misc/OutputDeviceArchiveWrapper.h"
-#include "HAL/ThreadManager.h"
-#include "BuildSettings.h"
-#include "CoreGlobals.h"
+#include "Windows/WindowsPlatformStackWalk.h"
 #include <atomic>
 #include <signal.h>
 
@@ -41,7 +41,7 @@
 #include <psapi.h>
 #include <tlhelp32.h>
 #include <shellapi.h>
-// Windows platform types intentionall not hidden til later
+// Windows platform types intentionally not hidden till later
 
 #ifndef UE_LOG_CRASH_CALLSTACK
 	#define UE_LOG_CRASH_CALLSTACK 1
