[Core]
* EnumerateRange - Changed code to work around compiler error

diff --git a/Engine/Source/Runtime/Core/Public/Misc/EnumerateRange.h b/Engine/Source/Runtime/Core/Public/Misc/EnumerateRange.h
--- a/Engine/Source/Runtime/Core/Public/Misc/EnumerateRange.h
+++ b/Engine/Source/Runtime/Core/Public/Misc/EnumerateRange.h
@@ -52,7 +52,7 @@ namespace UE::Core::Private
 		IteratorType Iterator;
 		SizeType     Index = 0;
 
-		auto operator*() const -> TEnumerateRef<std::remove_reference_t<decltype(*this->Iterator)>, SizeType>
+		auto operator*() const -> TEnumerateRef<std::remove_reference_t<decltype(*std::declval<const IteratorType&>())>, SizeType>
 		{
 			return { *this->Iterator, this->Index };
 		}
