Log SlowTask messages in out-of-order destruction ensure

diff --git a/Engine/Source/Runtime/Core/Private/Misc/SlowTask.cpp b/Engine/Source/Runtime/Core/Private/Misc/SlowTask.cpp
--- a/Engine/Source/Runtime/Core/Private/Misc/SlowTask.cpp
+++ b/Engine/Source/Runtime/Core/Private/Misc/SlowTask.cpp
@@ -114,16 +114,17 @@ void FSlowTask::Destroy()
 		}
 
 		FSlowTaskStack& Stack = Context.ScopeStack;
-		checkSlow(Stack.Num() != 0 && Stack.Last() == this);
-
-		FSlowTask* Task = Stack.Last();
-		if (ensureMsgf(Task == this, TEXT("Out-of-order scoped task construction/destruction")))
-		{
-			Stack.Pop(false);
-		}
-		else
+		if (ensure(Stack.Num() != 0))
 		{
-			Stack.RemoveSingleSwap(this, false);
+			FSlowTask* Task = Stack.Last();
+			if (ensureMsgf(Task == this, TEXT("Out-of-order slow task construction/destruction: destroying '%s' but '%s' is at the top of the stack"), *DefaultMessage.ToString(), *Task->DefaultMessage.ToString()))
+			{
+				Stack.Pop(false);
+			}
+			else
+			{
+				Stack.RemoveSingleSwap(this, false);
+			}
 		}
 
 		if (Stack.Num() != 0)
