Fix static analysis warnings in updated lz4 library code.

diff --git a/Engine/Source/Runtime/Core/Private/Compression/lz4.cpp b/Engine/Source/Runtime/Core/Private/Compression/lz4.cpp
--- a/Engine/Source/Runtime/Core/Private/Compression/lz4.cpp
+++ b/Engine/Source/Runtime/Core/Private/Compression/lz4.cpp
@@ -521,6 +521,7 @@ LZ4_memcpy_using_offset(BYTE* dstPtr, const BYTE* srcPtr, BYTE* dstEnd, const si
 #  pragma warning(push)
 #  pragma warning(disable : 6385) /* warning C6385: Reading invalid data from 'v'. */
 #endif
+        CA_SUPPRESS(6385);
         LZ4_memcpy(&v[4], v, 4);
 #if defined(_MSC_VER) && (_MSC_VER <= 1933) /* MSVC 2022 ver 17.3 or earlier */
 #  pragma warning(pop)
@@ -552,6 +553,7 @@ static unsigned LZ4_NbCommonBytes (reg_t val)
 {
     assert(val != 0);
     if (LZ4_isLittleEndian()) {
+        CA_SUPPRESS(6239);
         if (sizeof(val) == 8) {
 #       if defined(_MSC_VER) && (_MSC_VER >= 1800) && (defined(_M_AMD64) && !defined(_M_ARM64EC)) && !defined(LZ4_FORCE_SW_BITCOUNT)
 /*-*************************************************************************************************
@@ -595,6 +597,7 @@ static unsigned LZ4_NbCommonBytes (reg_t val)
 #       endif
         }
     } else   /* Big Endian CPU */ {
+        CA_SUPPRESS(6239);
         if (sizeof(val)==8) {
 #       if (defined(__clang__) || (defined(__GNUC__) && ((__GNUC__ > 3) || \
                             ((__GNUC__ == 3) && (__GNUC_MINOR__ >= 4))))) && \
@@ -668,6 +671,7 @@ unsigned LZ4_count(const BYTE* pIn, const BYTE* pMatch, const BYTE* pInLimit)
         return (unsigned)(pIn - pStart);
     }
 
+    CA_SUPPRESS(6239);
     if ((STEPSIZE==8) && (pIn<(pInLimit-3)) && (LZ4_read32(pMatch) == LZ4_read32(pIn))) { pIn+=4; pMatch+=4; }
     if ((pIn<(pInLimit-1)) && (LZ4_read16(pMatch) == LZ4_read16(pIn))) { pIn+=2; pMatch+=2; }
     if ((pIn<pInLimit) && (*pMatch == *pIn)) pIn++;
@@ -768,6 +772,7 @@ LZ4_FORCE_INLINE U32 LZ4_hash5(U64 sequence, tableType_t const tableType)
 
 LZ4_FORCE_INLINE U32 LZ4_hashPosition(const void* const p, tableType_t const tableType)
 {
+    CA_SUPPRESS(6239);
     if ((sizeof(reg_t)==8) && (tableType != byU16)) return LZ4_hash5(LZ4_read_ARCH(p), tableType);
     return LZ4_hash4(LZ4_read32(p), tableType);
 }
