Instead of rehashing the package name every time in sort fns, save it off.

diff --git a/Engine/Source/Developer/IoStoreUtilities/Private/PackageStoreOptimizer.cpp b/Engine/Source/Developer/IoStoreUtilities/Private/PackageStoreOptimizer.cpp
--- a/Engine/Source/Developer/IoStoreUtilities/Private/PackageStoreOptimizer.cpp
+++ b/Engine/Source/Developer/IoStoreUtilities/Private/PackageStoreOptimizer.cpp
@@ -1761,6 +1761,7 @@ FPackageStoreEntryResource FPackageStoreOptimizer::CreatePackageStoreEntry(const
 	}
 	
 	Result.PackageName = Package->Name;
+	Result.PackageId = FPackageId::FromName(Package->Name);
 	Result.SourcePackageName = Package->SourceName;
 	Result.Region = FName(*Package->Region);
 	Result.ExportInfo.ExportCount = Package->Exports.Num();
diff --git a/Engine/Source/Runtime/Core/Public/IO/PackageStore.h b/Engine/Source/Runtime/Core/Public/IO/PackageStore.h
--- a/Engine/Source/Runtime/Core/Public/IO/PackageStore.h
+++ b/Engine/Source/Runtime/Core/Public/IO/PackageStore.h
@@ -85,6 +85,7 @@ struct FPackageStoreEntryResource
 	EPackageStoreEntryFlags Flags = EPackageStoreEntryFlags::None;
 	/** The package name. */
 	FName PackageName;
+	FPackageId PackageId;
 	/** Used for localized and redirected packages. */
 	FName SourcePackageName;
 	/** Region name for localized packages. */
@@ -103,7 +104,7 @@ struct FPackageStoreEntryResource
 	/** Returns the package ID. */
 	FPackageId GetPackageId() const
 	{
-		return FPackageId::FromName(PackageName);
+		return PackageId;
 	}
 
 	/** Returns the source package ID. */
