[Compression]
* Fixed linker error in oodledatacompressionutil

diff --git a/Engine/Source/Runtime/Core/Public/Compression/OodleDataCompressionUtil.h b/Engine/Source/Runtime/Core/Public/Compression/OodleDataCompressionUtil.h
--- a/Engine/Source/Runtime/Core/Public/Compression/OodleDataCompressionUtil.h
+++ b/Engine/Source/Runtime/Core/Public/Compression/OodleDataCompressionUtil.h
@@ -62,7 +62,7 @@ namespace FOodleCompressedArray
 	* @return						0 if the FCompressedArray isn't valid, or if the compressed array requires PeekSizes64.
 	*								Otherwise, the offset to the compressed data in the array.
 	*/
-	int32 PeekSizes(TArray<uint8> const& InCompressed, int32& OutCompressedSize, int32& OutDecompressedSize)
+	inline int32 PeekSizes(TArray<uint8> const& InCompressed, int32& OutCompressedSize, int32& OutDecompressedSize)
 	{
 		// At minimum we need 8 bytes.
 		if (InCompressed.Num() < sizeof(int32)*2)
@@ -88,7 +88,7 @@ namespace FOodleCompressedArray
 
 		return sizeof(int32) * 2;
 	}
-	int32 PeekSizes64(TArray64<uint8> const& InCompressed, int64& OutCompressedSize, int64& OutDecompressedSize)
+	inline int32 PeekSizes64(TArray64<uint8> const& InCompressed, int64& OutCompressedSize, int64& OutDecompressedSize)
 	{
 		if (InCompressed.Num() < sizeof(int32)*2)
 		{
