Fixed missing dllexport on FName::SanitizeWhitespace

diff --git a/Engine/Source/Runtime/Core/Public/UObject/NameTypes.h b/Engine/Source/Runtime/Core/Public/UObject/NameTypes.h
--- a/Engine/Source/Runtime/Core/Public/UObject/NameTypes.h
+++ b/Engine/Source/Runtime/Core/Public/UObject/NameTypes.h
@@ -839,7 +839,7 @@ public:
 	* @return the new FName based upon the calling FName, but with any Whitespace characters potentially problematic for
 	 * showing in a log or on screen omitted.
 	 */
-	static FString SanitizeWhitespace(const FString& FNameString);
+	CORE_API static FString SanitizeWhitespace(const FString& FNameString);
 	
 	/**
 	 * Compares name to passed in one. Sort is alphabetical ascending.
