ignoring task traces until task channel is initialised. this was asserting on early editor exit

diff --git a/Engine/Source/Runtime/Core/Private/Async/TaskTrace.cpp b/Engine/Source/Runtime/Core/Private/Async/TaskTrace.cpp
--- a/Engine/Source/Runtime/Core/Private/Async/TaskTrace.cpp
+++ b/Engine/Source/Runtime/Core/Private/Async/TaskTrace.cpp
@@ -94,7 +94,11 @@ namespace TaskTrace
 
 	void Created(FId TaskId)
 	{
-		check(bGTaskTraceInitialized);
+		if (!bGTaskTraceInitialized)
+		{
+			return;
+		}
+
 		check(TaskId != InvalidId);
 
 		UE_TRACE_LOG(TaskTrace, Created, TaskChannel)
@@ -104,7 +108,11 @@ namespace TaskTrace
 
 	void Launched(FId TaskId, const TCHAR* DebugName, bool bTracked, ENamedThreads::Type ThreadToExecuteOn)
 	{
-		check(bGTaskTraceInitialized);
+		if (!bGTaskTraceInitialized)
+		{
+			return;
+		}
+
 		check(TaskId != InvalidId);
 
 		UE_TRACE_LOG(TaskTrace, Launched, TaskChannel)
@@ -117,7 +125,11 @@ namespace TaskTrace
 
 	void Scheduled(FId TaskId)
 	{
-		check(bGTaskTraceInitialized);
+		if (!bGTaskTraceInitialized)
+		{
+			return;
+		}
+
 		check(TaskId != InvalidId);
 
 		UE_TRACE_LOG(TaskTrace, Scheduled, TaskChannel)
@@ -127,7 +139,10 @@ namespace TaskTrace
 
 	void SubsequentAdded(FId TaskId, FId SubsequentId)
 	{
-		check(bGTaskTraceInitialized);
+		if (!bGTaskTraceInitialized)
+		{
+			return;
+		}
 
 		// "empty" FGraphEvent is used for synchronisation only, to wait for a notification. It doesn't have an associated task and ID.
 		if (TaskId == InvalidId)
@@ -144,7 +159,11 @@ namespace TaskTrace
 
 	void Started(FId TaskId)
 	{
-		check(bGTaskTraceInitialized);
+		if (!bGTaskTraceInitialized)
+		{
+			return;
+		}
+
 		check(TaskId != InvalidId);
 
 		UE_TRACE_LOG(TaskTrace, Started, TaskChannel)
@@ -154,7 +173,11 @@ namespace TaskTrace
 
 	void Finished(FId TaskId)
 	{
-		check(bGTaskTraceInitialized);
+		if (!bGTaskTraceInitialized)
+		{
+			return;
+		}
+
 		check(TaskId != InvalidId);
 
 		UE_TRACE_LOG(TaskTrace, Finished, TaskChannel)
@@ -164,7 +187,10 @@ namespace TaskTrace
 
 	void Completed(FId TaskId)
 	{
-		check(bGTaskTraceInitialized);
+		if (!bGTaskTraceInitialized)
+		{
+			return;
+		}
 
 		// "empty" FGraphEvent is used for synchronisation only, to wait for a notification. It doesn't have an associated task and ID.
 		if (TaskId == InvalidId)
@@ -179,7 +205,10 @@ namespace TaskTrace
 
 	void Destroyed(FId TaskId)
 	{
-		check(bGTaskTraceInitialized);
+		if (!bGTaskTraceInitialized)
+		{
+			return;
+		}
 
 		UE_TRACE_LOG(TaskTrace, Destroyed, TaskChannel)
 			<< Destroyed.Timestamp(FPlatformTime::Cycles64())
@@ -188,7 +217,10 @@ namespace TaskTrace
 
 	FWaitingScope::FWaitingScope(const TArray<FId>& Tasks)
 	{
-		check(bGTaskTraceInitialized);
+		if (!bGTaskTraceInitialized)
+		{
+			return;
+		}
 
 		UE_TRACE_LOG(TaskTrace, WaitingStarted, TaskChannel)
 			<< WaitingStarted.Timestamp(FPlatformTime::Cycles64())
@@ -197,7 +229,10 @@ namespace TaskTrace
 
 	FWaitingScope::FWaitingScope(FId TaskId)
 	{
-		check(bGTaskTraceInitialized);
+		if (!bGTaskTraceInitialized)
+		{
+			return;
+		}
 
 		UE_TRACE_LOG(TaskTrace, WaitingStarted, TaskChannel)
 			<< WaitingStarted.Timestamp(FPlatformTime::Cycles64())
@@ -206,7 +241,10 @@ namespace TaskTrace
 
 	FWaitingScope::~FWaitingScope()
 	{
-		check(bGTaskTraceInitialized);
+		if (!bGTaskTraceInitialized)
+		{
+			return;
+		}
 
 		UE_TRACE_LOG(TaskTrace, WaitingFinished, TaskChannel)
 			<< WaitingFinished.Timestamp(FPlatformTime::Cycles64());
