Added FMRSWRecursiveAccessDetector that is zero-initializable, supports upgrading a reader access to a writer access recursively, and destruction from inside an access

diff --git a/Engine/Source/Runtime/Core/Private/Misc/MTAccessDetector.cpp b/Engine/Source/Runtime/Core/Private/Misc/MTAccessDetector.cpp
--- a/Engine/Source/Runtime/Core/Private/Misc/MTAccessDetector.cpp
+++ b/Engine/Source/Runtime/Core/Private/Misc/MTAccessDetector.cpp
@@ -1,11 +1,14 @@
 // Copyright Epic Games, Inc. All Rights Reserved.
 
 #include "Misc/MTAccessDetector.h"
+
+#if ENABLE_MT_DETECTOR && WITH_DEV_AUTOMATION_TESTS
+
 #include "Misc/AutomationTest.h"
 #include "HAL/Thread.h"
 #include "GenericPlatform/GenericPlatformProcess.h"
-
-#if ENABLE_MT_DETECTOR
+#include "Tasks/Task.h"
+#include "Tests/Benchmark.h"
 
 //----------------------------------------------------------------------//
 // FRWAccessDetector tests
@@ -325,4 +328,177 @@ bool FRWAccessDetector_FullyRecursiveReadReentrance::RunTest(const FString& Para
 	return Success;
 }
 
-#endif // ENABLE_MT_DETECTOR
+/////////////////////////////////////////////////////////////////////////
+
+namespace MTAccessDetector_Private
+{
+	using namespace UE::Tasks;
+
+	// helper for testing concurrent cases
+	class FParallel
+	{
+	private:
+		FTask Task;
+		FTaskEvent TaskStarted{ UE_SOURCE_LOCATION };
+		FTaskEvent FinishTask{ UE_SOURCE_LOCATION };
+
+	public:
+		template<typename PrologType, typename EpilogType>
+		FParallel(PrologType&& Prolog, EpilogType&& Epilog)
+		{
+			Task.Launch(UE_SOURCE_LOCATION, [this, Prolog = MoveTemp(Prolog), Epilog = MoveTemp(Epilog)]
+			{
+				Prolog();
+
+				TaskStarted.Trigger();
+				FinishTask.Wait();
+
+				Epilog();
+			});
+
+			TaskStarted.Wait();
+		}
+
+		~FParallel()
+		{
+			FinishTask.Trigger();
+			Task.Wait();
+		}
+	};
+}
+
+IMPLEMENT_SIMPLE_AUTOMATION_TEST(FRWRecursiveAccessDetectorTest, "System.Core.RWRecursiveAccessDetector", EAutomationTestFlags::ApplicationContextMask | EAutomationTestFlags::EngineFilter);
+
+bool FRWRecursiveAccessDetectorTest::RunTest(const FString& Parameters)
+{
+	// single thread
+
+	{	// single instance
+		UE_MT_DECLARE_MRSW_RECURSIVE_ACCESS_DETECTOR(AD);
+
+		{
+			UE_MT_SCOPED_READ_ACCESS(AD);
+			UE_MT_SCOPED_READ_ACCESS(AD);
+			UE_MT_SCOPED_WRITE_ACCESS(AD);
+			UE_MT_SCOPED_WRITE_ACCESS(AD);
+			UE_MT_SCOPED_READ_ACCESS(AD);
+		}
+	}
+
+	{	// multiple nested instances
+		UE_MT_DECLARE_MRSW_RECURSIVE_ACCESS_DETECTOR(Outer);
+		UE_MT_DECLARE_MRSW_RECURSIVE_ACCESS_DETECTOR(Inner);
+
+		UE_MT_SCOPED_READ_ACCESS(Outer);
+		UE_MT_SCOPED_READ_ACCESS(Inner);
+		UE_MT_SCOPED_WRITE_ACCESS(Outer);
+		UE_MT_SCOPED_WRITE_ACCESS(Inner);
+		UE_MT_SCOPED_WRITE_ACCESS(Outer);
+		UE_MT_SCOPED_WRITE_ACCESS(Inner);
+	}
+
+	{	// destroying detector from inside access scope
+		auto* AD = new FMRSWRecursiveAccessDetector;
+
+		AD->AcquireReadAccess();
+
+		FMRSWRecursiveAccessDetector::FDestructionSentinel DestructionSentinel{ FMRSWRecursiveAccessDetector::EAccessType::Reader };
+		FMRSWRecursiveAccessDetector::FDestructionSentinel* PrevDestructionSentinel = AD->SetDestructionSentinel(&DestructionSentinel);
+
+		delete AD;
+
+		if (!DestructionSentinel.bDestroyed)
+		{
+			AD->SetDestructionSentinel(PrevDestructionSentinel);
+			AD->ReleaseReadAccess();
+		}
+	}
+
+	// multithread
+
+	{	// multireader
+		UE_MT_DECLARE_MRSW_RECURSIVE_ACCESS_DETECTOR(AD);
+		MTAccessDetector_Private::FParallel ParallelReaderScope([&AD] { AD.AcquireReadAccess(); }, [&AD] { AD.ReleaseReadAccess(); });
+		UE_MT_SCOPED_READ_ACCESS(AD);
+	}
+
+#if 0
+	{	// reader/writer, asserts
+		UE_MT_DECLARE_MRSW_RECURSIVE_ACCESS_DETECTOR(AD);
+		MTAccessDetector_Private::FParallel ParallelReaderScope([&AD] { AD.AcquireReadAccess(); }, [&AD] { AD.ReleaseReadAccess(); });
+		UE_MT_SCOPED_WRITE_ACCESS(AD); // must assert on acquiring write access
+	}
+#endif 
+
+#if 0
+	{	// multiwriter, asserts
+		UE_MT_DECLARE_MRSW_RECURSIVE_ACCESS_DETECTOR(AD);
+		MTAccessDetector_Private::FParallel ParallelReaderScope([&AD] { AD.AcquireWriteAccess(); }, [&AD] { AD.ReleaseWriteAccess(); });
+		UE_MT_SCOPED_WRITE_ACCESS(AD); // must assert on acquiring write access
+	}
+#endif
+
+	return true;
+}
+
+template<uint64 Num, typename AccessDetectorType>
+void RWAccessDetectorPerfTest()
+{
+	AccessDetectorType AD;
+	for (int i = 0; i != Num; ++i)
+	{
+		AD.AcquireReadAccess();
+		AD.ReleaseReadAccess();
+	}
+	for (int i = 0; i != Num; ++i)
+	{
+		AD.AcquireWriteAccess();
+		AD.ReleaseWriteAccess();
+	}
+}
+
+template<uint64 Num, typename AccessDetectorType>
+void RWAccessDetectorConcurrentReadersPerfTest()
+{
+	using namespace UE::Tasks;
+
+	AccessDetectorType AD;
+	auto Reader = [&AD]
+	{
+		for (uint64 i = 0; i != Num; ++i)
+		{
+			AD.AcquireReadAccess();
+			AD.ReleaseReadAccess();
+		}
+	};
+
+	TArray<FTask> Tasks;
+	for (int i = 0; i != 3; ++i)
+	{
+		auto TaskBody = Reader;
+		Tasks.Add(Launch(UE_SOURCE_LOCATION, MoveTemp(TaskBody)));
+	}
+
+	Reader();
+
+	Wait(Tasks);
+}
+
+IMPLEMENT_SIMPLE_AUTOMATION_TEST(FRWRecursiveAccessDetectorPerfTest, "System.Core.RWRecursiveAccessDetector.Perf", EAutomationTestFlags::ApplicationContextMask | EAutomationTestFlags::EngineFilter);
+
+bool FRWRecursiveAccessDetectorPerfTest::RunTest(const FString& Parameters)
+{
+	UE_BENCHMARK(5, RWAccessDetectorPerfTest<1'000'000, FRWAccessDetector>);
+	UE_BENCHMARK(5, RWAccessDetectorPerfTest<1'000'000, FRWRecursiveAccessDetector>);
+	UE_BENCHMARK(5, RWAccessDetectorPerfTest<1'000'000, FRWFullyRecursiveAccessDetector>);
+	UE_BENCHMARK(5, RWAccessDetectorPerfTest<1'000'000, FMRSWRecursiveAccessDetector>);
+
+	UE_BENCHMARK(5, RWAccessDetectorConcurrentReadersPerfTest<1'000'000, FRWAccessDetector>);
+	UE_BENCHMARK(5, RWAccessDetectorConcurrentReadersPerfTest<1'000'000, FRWRecursiveAccessDetector>);
+	UE_BENCHMARK(5, RWAccessDetectorConcurrentReadersPerfTest<1'000'000, FRWFullyRecursiveAccessDetector>);
+	UE_BENCHMARK(5, RWAccessDetectorConcurrentReadersPerfTest<1'000'000, FMRSWRecursiveAccessDetector>);
+
+	return true;
+}
+
+#endif // ENABLE_MT_DETECTOR && WITH_DEV_AUTOMATION_TESTS
diff --git a/Engine/Source/Runtime/Core/Public/Misc/MTAccessDetector.h b/Engine/Source/Runtime/Core/Public/Misc/MTAccessDetector.h
--- a/Engine/Source/Runtime/Core/Public/Misc/MTAccessDetector.h
+++ b/Engine/Source/Runtime/Core/Public/Misc/MTAccessDetector.h
@@ -3,17 +3,20 @@
 #pragma once
 
 #include "CoreTypes.h"
-#include "HAL/PlatformTLS.h"
-#include "HAL/PreprocessorHelpers.h"
 #include "Misc/AssertionMacros.h"
-#include "Misc/Build.h"
-#include <atomic>
-
 
 #define ENABLE_MT_DETECTOR DO_CHECK
 
 #if ENABLE_MT_DETECTOR
 
+#include "Containers/Array.h"
+#include "Containers/ContainerAllocationPolicies.h"
+#include "HAL/PlatformTLS.h"
+#include "HAL/PreprocessorHelpers.h"
+#include "Misc/Build.h"
+#include <atomic>
+
+
 extern CORE_API bool GIsAutomationTesting;
 
 /**
@@ -31,7 +34,7 @@ public:
 
 	~FRWAccessDetector()
 	{
-		checkf(AtomicValue.load(std::memory_order_relaxed) == 0, TEXT("Detector cannot be destroyed while other threads access it"))
+		checkf(AtomicValue.load(std::memory_order_relaxed) == 0, TEXT("Detector cannot be destroyed while other threads access it"));
 	}
 
 	FRWAccessDetector(FRWAccessDetector&& Other)
@@ -218,7 +221,7 @@ struct TScopedReaderAccessDetector : public FBaseScopedAccessDetector
 {
 public:
 
-	FORCEINLINE TScopedReaderAccessDetector(const RWAccessDetector& InAccessDetector)
+	FORCEINLINE TScopedReaderAccessDetector(RWAccessDetector& InAccessDetector)
 	: AccessDetector(InAccessDetector)
 	{
 		AccessDetector.AcquireReadAccess();
@@ -229,7 +232,7 @@ public:
 		AccessDetector.ReleaseReadAccess();
 	}
 private:
-	const RWAccessDetector& AccessDetector;
+	RWAccessDetector& AccessDetector;
 };
 
 template<typename RWAccessDetector>
@@ -242,8 +245,7 @@ template<typename RWAccessDetector>
 struct TScopedWriterDetector : public FBaseScopedAccessDetector
 {
 public:
-
-	FORCEINLINE TScopedWriterDetector(const RWAccessDetector& InAccessDetector)
+	FORCEINLINE TScopedWriterDetector(RWAccessDetector& InAccessDetector)
 		: AccessDetector(InAccessDetector)
 	{
 		AccessDetector.AcquireWriteAccess();
@@ -254,7 +256,7 @@ public:
 		AccessDetector.ReleaseWriteAccess();
 	}
 private:
-	const RWAccessDetector& AccessDetector;
+	RWAccessDetector& AccessDetector;
 };
 
 template<typename RWAccessDetector>
@@ -263,9 +265,262 @@ FORCEINLINE TScopedWriterDetector<RWAccessDetector> MakeScopedWriterAccessDetect
 	return TScopedWriterDetector<RWAccessDetector>(InAccessDetector);
 }
 
+/** 
+ * race detector supporting multiple readers (MR) single writer (SW) recursive access, a write from inside a read, a read from inside a write 
+ * and all other combinations. is zero initializable. supports destruction while being "accessed"
+ */
+class FMRSWRecursiveAccessDetector
+{
+private:
+	// despite 0 is a valid TID on some platforms, we store actual TID + 1 to avoid collisions. it's required to use 0 as 
+	// an invalid TID for zero-initilization
+	static constexpr uint32 InvalidThreadId = 0;
+
+	union FState
+	{
+		uint64 Value;
+
+		struct
+		{
+			uint32 ReaderNum : 20;
+			uint32 WriterNum : 12;
+			uint32 WriterThreadId;
+		};
+
+		constexpr FState()
+			: Value(0)
+		{
+		}
+
+		constexpr FState(uint64 InValue)
+			: Value(InValue)
+		{
+		}
+
+		constexpr FState(uint32 InReaderNum, uint32 InWriterNum, uint32 InWriterThreadId)
+			: ReaderNum(InReaderNum)
+			, WriterNum(InWriterNum)
+			, WriterThreadId(InWriterThreadId)
+		{
+		}
+	};
+
+	// all atomic ops are relaxed to preserve the original memory order, as the detector is compiled out in non-dev builds
+	mutable std::atomic<uint64> State{ 0 }; // it's actually FState, but we need to do `atomic::fetch_add` on it
+
+	FORCEINLINE FState LoadState() const
+	{
+		return FState(State.load(std::memory_order_relaxed));
+	}
+
+	FORCEINLINE FState ExchangeState(FState NewState)
+	{
+		return FState(State.exchange(NewState.Value, std::memory_order_relaxed));
+	}
+
+	FORCEINLINE FState IncrementReaderNum() const
+	{
+		constexpr FState OneReader{ 1, 0, 0 };
+		return FState(State.fetch_add(OneReader.Value, std::memory_order_relaxed));
+	}
+
+	FORCEINLINE FState DecrementReaderNum() const
+	{
+		constexpr FState OneReader{ 1, 0, 0 };
+		return FState(State.fetch_sub(OneReader.Value, std::memory_order_relaxed));
+	}
+
+	FORCEINLINE friend bool operator==(FState L, FState R)
+	{
+		return L.Value == R.Value;
+	}
+
+	FORCEINLINE static void CheckOtherThreadWriters(FState InState)
+	{
+		if (InState.WriterNum != 0)
+		{
+			uint32 CurrentThreadId = FPlatformTLS::GetCurrentThreadId() + 1; // to shift from 0 TID that is considered "invalid"
+			checkf(InState.WriterThreadId == CurrentThreadId, TEXT("Data race detected! Writer on thread %u while acquiring read access on thread %u"), InState.WriterThreadId, CurrentThreadId);
+		}
+	}
+
+public: // DestructionSentinel
+	enum class EAccessType { Reader, Writer };
+
+	struct FDestructionSentinel
+	{
+		FORCEINLINE explicit FDestructionSentinel(EAccessType InAccessType)
+			: AccessType(InAccessType)
+		{
+		}
+
+		EAccessType AccessType;
+		bool bDestroyed = false;
+	};
+
+private:
+	mutable FDestructionSentinel* DestructionSentinel = nullptr;
+
+	// if access detector was destroyed while holding read accessed (same thread) it can't release the access, but is still registered in TLS.
+	// use this static method to clear it
+	FORCEINLINE static void RemoveReader(FMRSWRecursiveAccessDetector* Reader)
+	{
+		uint32 ReaderIndex = ReadersTls.IndexOfByPredicate([Reader](FReaderNum ReaderNum) { return ReaderNum.Reader == Reader; });
+		checkfSlow(ReaderIndex != INDEX_NONE, TEXT("Can't find the reader in the TLS table"));
+		checkfSlow(ReadersTls[ReaderIndex].Num == 1, TEXT("More than one accesses detected while destoying access detector: %d"), ReadersTls[ReaderIndex].Num);
+		ReadersTls.RemoveAtSwap(ReaderIndex, 1, /* bAllowShrinking = */ false);
+	}
+
+private:
+	///////////////////////////////////////////////
+	// to support a write access from inside a read access on the same thread, we need to know that there're no readers on other threads.
+	// As we can't have TLS slot per access detector instance, a single one is shared between multiple instances with the assumption that
+	// there rarely will be more than one reader registered.
+	struct FReaderNum
+	{
+		const FMRSWRecursiveAccessDetector* Reader;
+		uint32 Num;
+	};
+
+	inline static thread_local TArray<FReaderNum, TInlineAllocator<4>> ReadersTls;
+	///////////////////////////////////////////////
+
+public:
+	FMRSWRecursiveAccessDetector() = default;
+
+	FORCEINLINE FMRSWRecursiveAccessDetector(const FMRSWRecursiveAccessDetector& Other)
+		// just default initialisation, the copy is not being accessed
+	{
+		CheckOtherThreadWriters(Other.LoadState());
+	}
+
+	FORCEINLINE FMRSWRecursiveAccessDetector& operator=(const FMRSWRecursiveAccessDetector& Other)
+		// do not alter the state, it can be accessed
+	{
+		CheckOtherThreadWriters(Other.LoadState());
+		return *this;
+	}
+
+	// use copy construction/assignment ^^ for move semantics too
+
+	FORCEINLINE ~FMRSWRecursiveAccessDetector()
+	{
+		if (DestructionSentinel)
+		{
+			DestructionSentinel->bDestroyed = true;
+			if (DestructionSentinel->AccessType == EAccessType::Reader)
+			{	// readers are registered in TLS table and must be cleared on destruction
+				RemoveReader(this);
+			}
+		}
+		else
+		{
+			checkf(LoadState().Value == 0,
+				TEXT("Race detector destroyed while being accessed: %u readers, %u writers on thread %d"),
+				LoadState().ReaderNum, LoadState().WriterNum, LoadState().WriterThreadId - 1);
+		}
+	}
+
+	// supports destroying an instance while it's being accessed by notifying the user about destruction so the user don't release access to
+	// an already destroyed instance. Should be used only from inside an access scope.
+	// @return previously set destruction sentinel that should be set back before leaving an access scope
+	FORCEINLINE FDestructionSentinel* SetDestructionSentinel(FDestructionSentinel* InDestructionSentinel) const
+	{
+		FDestructionSentinel* PrevDestructionSentinel = DestructionSentinel;
+		DestructionSentinel = InDestructionSentinel;
+		return PrevDestructionSentinel;
+	}
+
+	FORCEINLINE void AcquireReadAccess() const
+	{
+		FState PrevState = IncrementReaderNum();
+
+		CheckOtherThreadWriters(PrevState);
+
+		// register the reader in TLS
+		uint32 ReaderIndex = PrevState.ReaderNum == 0 ? 
+			INDEX_NONE : 
+			ReadersTls.IndexOfByPredicate([this](FReaderNum ReaderNum) { return ReaderNum.Reader == this; });
+		if (ReaderIndex == INDEX_NONE)
+		{
+			ReadersTls.Add(FReaderNum{ this, 1 });
+		}
+		else
+		{
+			++ReadersTls[ReaderIndex].Num;
+		}
+	}
+
+	FORCEINLINE void ReleaseReadAccess() const
+	{
+		// unregister the reader from TLS
+		uint32 ReaderIndex = ReadersTls.IndexOfByPredicate([this](FReaderNum ReaderNum) { return ReaderNum.Reader == this; });
+		checkfSlow(ReaderIndex != INDEX_NONE, 
+			TEXT("Invalid usage of the race detector! No matching AcquireReadAccess(): %u readers, %u writers on thread %d"),
+			LoadState().ReaderNum, LoadState().WriterNum, LoadState().WriterThreadId - 1);
+		uint32 ReaderNum = --ReadersTls[ReaderIndex].Num;
+		if (ReaderNum == 0)
+		{
+			ReadersTls.RemoveAtSwap(ReaderIndex, 1, /* bAllowShrinking = */ false);
+		}
+
+		DecrementReaderNum();
+		// no need to check for writers
+	}
+
+	FORCEINLINE void AcquireWriteAccess()
+	{
+		FState LocalState = LoadState();
+
+		if (LocalState.ReaderNum >= 1)
+		{	// check that all readers are on the current thread
+			uint32 ReaderIndex = ReadersTls.IndexOfByPredicate([this](FReaderNum ReaderNum) { return ReaderNum.Reader == this; });
+			checkf(ReaderIndex != INDEX_NONE && ReadersTls[ReaderIndex].Num == LocalState.ReaderNum, 
+				TEXT("Data race detected: %d reader(s) on another thread(s) while acquiring write access"), 
+				LocalState.ReaderNum - ReadersTls[ReaderIndex].Num);
+		}
+
+		uint32 CurrentThreadId = FPlatformTLS::GetCurrentThreadId() + 1;
+		if (LocalState.WriterNum != 0)
+		{
+			checkf(LocalState.WriterThreadId == CurrentThreadId, 
+				TEXT("Data race detected: writer on thread %d during acquiring write access on thread %d"), 
+				LocalState.WriterThreadId - 1, CurrentThreadId - 1);
+		}
+
+		FState NewState{ LocalState.ReaderNum, LocalState.WriterNum + 1u, CurrentThreadId };
+		FState PrevState = ExchangeState(NewState);
+
+		checkf(LocalState == PrevState,
+			TEXT("Data race detected: other thread(s) activity during acquiring write access on thread %d: %u -> %u readers, %u -> %u writers on thread %d -> %d"),
+			CurrentThreadId - 1,
+			LocalState.ReaderNum, PrevState.ReaderNum,
+			LocalState.WriterNum, PrevState.WriterNum,
+			LocalState.WriterThreadId - 1, PrevState.WriterThreadId - 1);
+	}
+
+	FORCEINLINE void ReleaseWriteAccess()
+	{
+		FState LocalState = LoadState();
+
+		uint32 WriterThreadId = LocalState.WriterNum != 1 ? LocalState.WriterThreadId : InvalidThreadId;
+		FState NewState{ LocalState.ReaderNum, LocalState.WriterNum - 1u,WriterThreadId };
+		FState PrevState = ExchangeState(NewState);
+
+		checkf(LocalState == PrevState,
+			TEXT("Data race detected: other thread(s) activity during releasing write access on thread %d: %u -> %u readers, %u -> %u writers on thread %d -> %d"),
+			LocalState.ReaderNum, PrevState.ReaderNum,
+			LocalState.WriterNum, PrevState.WriterNum,
+			LocalState.WriterThreadId - 1, PrevState.WriterThreadId - 1);
+	}
+};
+
+///////////////////////////////////////////////////////////////////
+
 #define UE_MT_DECLARE_RW_ACCESS_DETECTOR(AccessDetector) FRWAccessDetector AccessDetector;
 #define UE_MT_DECLARE_RW_RECURSIVE_ACCESS_DETECTOR(AccessDetector) FRWRecursiveAccessDetector AccessDetector;
 #define UE_MT_DECLARE_RW_FULLY_RECURSIVE_ACCESS_DETECTOR(AccessDetector) FRWFullyRecursiveAccessDetector AccessDetector;
+#define UE_MT_DECLARE_MRSW_RECURSIVE_ACCESS_DETECTOR(AccessDetector) FMRSWRecursiveAccessDetector AccessDetector;
 
 #define UE_MT_SCOPED_READ_ACCESS(AccessDetector) const FBaseScopedAccessDetector& PREPROCESSOR_JOIN(ScopedMTAccessDetector_,__LINE__) = MakeScopedReaderAccessDetector(AccessDetector);
 #define UE_MT_SCOPED_WRITE_ACCESS(AccessDetector) const FBaseScopedAccessDetector& PREPROCESSOR_JOIN(ScopedMTAccessDetector_,__LINE__) = MakeScopedWriterAccessDetector(AccessDetector);
