suppressed MSVC static analysis warning in MTAccessDetector unit-test
https://horde.devtools.epicgames.com/job/63e3c73edd21ba4a6d6ebbf6?step=fb38&issue=260112

diff --git a/Engine/Source/Runtime/Core/Private/Misc/MTAccessDetector.cpp b/Engine/Source/Runtime/Core/Private/Misc/MTAccessDetector.cpp
--- a/Engine/Source/Runtime/Core/Private/Misc/MTAccessDetector.cpp
+++ b/Engine/Source/Runtime/Core/Private/Misc/MTAccessDetector.cpp
@@ -398,7 +398,10 @@ bool FRWRecursiveAccessDetectorTest::RunTest(const FString& Parameters)
 	}
 
 	{	// destroying detector from inside access scope
+#pragma warning(push)
+#pragma warning(disable:6001) // "Using uninitialized memory" warning
 		auto* AD = new FMRSWRecursiveAccessDetector;
+#pragma warning(pop)
 
 		AD->AcquireReadAccess();
 
