suppressed false-positive PVS-Studio static analysis warning
https://horde.devtools.epicgames.com/job/63e3c73edd21ba4a6d6ebbf6?step=fe3b&issue=260146

diff --git a/Engine/Source/Runtime/Core/Private/Misc/MTAccessDetector.cpp b/Engine/Source/Runtime/Core/Private/Misc/MTAccessDetector.cpp
--- a/Engine/Source/Runtime/Core/Private/Misc/MTAccessDetector.cpp
+++ b/Engine/Source/Runtime/Core/Private/Misc/MTAccessDetector.cpp
@@ -400,7 +400,7 @@ bool FRWRecursiveAccessDetectorTest::RunTest(const FString& Parameters)
 	{	// destroying detector from inside access scope
 #pragma warning(push)
 #pragma warning(disable:6001) // "Using uninitialized memory" warning
-		auto* AD = new FMRSWRecursiveAccessDetector;
+		auto* AD = new FMRSWRecursiveAccessDetector; //-V774: "The 'AD' pointer was used after the memory was released." - no, it wasn't because of the `if`
 #pragma warning(pop)
 
 		AD->AcquireReadAccess();
