Fix case-insensitive system include of Semaphore.h (renamed)

#android

diff --git a/Engine/Source/Runtime/Core/Private/Tests/Async/SemaphoreTest.cpp b/Engine/Source/Runtime/Core/Private/Tests/Async/SemaphoreTest.cpp
--- a/Engine/Source/Runtime/Core/Private/Tests/Async/SemaphoreTest.cpp
+++ b/Engine/Source/Runtime/Core/Private/Tests/Async/SemaphoreTest.cpp
@@ -4,7 +4,7 @@
 #include "Misc/AutomationTest.h"
 #include "Tests/Benchmark.h"
 #include "HAL/Thread.h"
-#include "HAL/Semaphore.h"
+#include "HAL/UESemaphore.h"
 #include "HAL/Event.h"
 #include "HAL/PlatformProcess.h"
 #include "Containers/Array.h"
diff --git a/Engine/Source/Runtime/Core/Public/HAL/Semaphore.h b/Engine/Source/Runtime/Core/Public/HAL/Semaphore.h
deleted file mode 100644
diff --git a/Engine/Source/Runtime/Core/Public/HAL/UESemaphore.h b/Engine/Source/Runtime/Core/Public/HAL/UESemaphore.h
new file mode 100644
--- /dev/null
+++ b/Engine/Source/Runtime/Core/Public/HAL/UESemaphore.h
@@ -0,0 +1,7 @@
+// Copyright Epic Games, Inc. All Rights Reserved.
+
+#pragma once
+
+#include "CoreTypes.h"
+
+#include COMPILED_PLATFORM_HEADER(Semaphore.h)
