Removed some includes to ExpressionParserTypes.inl

diff --git a/Engine/Source/Runtime/Core/Private/Internationalization/TextFormatter.cpp b/Engine/Source/Runtime/Core/Private/Internationalization/TextFormatter.cpp
--- a/Engine/Source/Runtime/Core/Private/Internationalization/TextFormatter.cpp
+++ b/Engine/Source/Runtime/Core/Private/Internationalization/TextFormatter.cpp
@@ -17,7 +17,6 @@
 #include "Misc/Char.h"
 #include "Misc/EnumClassFlags.h"
 #include "Misc/ExpressionParser.h"
-#include "Misc/ExpressionParserTypes.inl"
 #include "Misc/Guid.h"
 #include "Misc/Optional.h"
 #include "Misc/ScopeLock.h"
diff --git a/Engine/Source/Runtime/Core/Private/Math/UnitConversion.cpp b/Engine/Source/Runtime/Core/Private/Math/UnitConversion.cpp
--- a/Engine/Source/Runtime/Core/Private/Math/UnitConversion.cpp
+++ b/Engine/Source/Runtime/Core/Private/Math/UnitConversion.cpp
@@ -11,8 +11,6 @@
 #include "Misc/AssertionMacros.h"
 #include "Misc/CString.h"
 #include "Misc/ExpressionParser.h"
-#include "Misc/ExpressionParserTypes.h"
-#include "Misc/ExpressionParserTypes.inl"
 #include "Misc/Guid.h"
 #include "Templates/UniquePtr.h"
 #include "Templates/UnrealTemplate.h"
diff --git a/Engine/Source/Runtime/Core/Private/Misc/FrameRate.cpp b/Engine/Source/Runtime/Core/Private/Misc/FrameRate.cpp
--- a/Engine/Source/Runtime/Core/Private/Misc/FrameRate.cpp
+++ b/Engine/Source/Runtime/Core/Private/Misc/FrameRate.cpp
@@ -9,8 +9,6 @@
 #include "Internationalization/Internationalization.h"
 #include "Math/BasicMathExpressionEvaluator.h"
 #include "Misc/ExpressionParser.h"
-#include "Misc/ExpressionParserTypes.h"
-#include "Misc/ExpressionParserTypes.inl"
 #include "Misc/Guid.h"
 #include "Templates/UnrealTemplate.h"
 
diff --git a/Engine/Source/Runtime/Core/Private/Misc/TextFilterExpressionEvaluator.cpp b/Engine/Source/Runtime/Core/Private/Misc/TextFilterExpressionEvaluator.cpp
--- a/Engine/Source/Runtime/Core/Private/Misc/TextFilterExpressionEvaluator.cpp
+++ b/Engine/Source/Runtime/Core/Private/Misc/TextFilterExpressionEvaluator.cpp
@@ -5,7 +5,7 @@
 #include "Containers/SparseArray.h"
 #include "Math/BasicMathExpressionEvaluator.h"
 #include "Misc/CString.h"
-#include "Misc/ExpressionParserTypes.inl"
+#include "Misc/ExpressionParserTypes.h"
 #include "Misc/Guid.h"
 #include "Templates/Function.h"
 #include "Templates/ValueOrError.h"
