Allow FAutomationSpecBase to skip stack walking through project settings

diff --git a/Engine/Source/Runtime/Core/Private/Misc/AutomationTest.cpp b/Engine/Source/Runtime/Core/Private/Misc/AutomationTest.cpp
--- a/Engine/Source/Runtime/Core/Private/Misc/AutomationTest.cpp
+++ b/Engine/Source/Runtime/Core/Private/Misc/AutomationTest.cpp
@@ -29,6 +29,12 @@ namespace AutomationTest
 		TEXT("Automation.CaptureLogEvents"),
 		bCaptureLogEvents,
 		TEXT("Consider warning/error log events during a test as impacting the test itself"));
+
+	static bool bSkipStackWalk = true;
+	static FAutoConsoleVariableRef CVarAutomationSkipStackWalk(
+		TEXT("Automation.SkipStackWalk"),
+		bSkipStackWalk,
+		TEXT("Whether to skip stack walk while iterating for listing the tests"));
 };
 
 
@@ -258,6 +264,12 @@ FString FAutomationTestFramework::GetUserAutomationDirectory() const
 	return FString(FPlatformProcess::UserDir()) + DefaultAutomationSubFolder;
 }
 
+bool FAutomationTestFramework::NeedSkipStackWalk()
+{
+	return AutomationTest::bSkipStackWalk;
+}
+
+
 bool FAutomationTestFramework::RegisterAutomationTest( const FString& InTestNameToRegister, FAutomationTestBase* InTestToRegister )
 {
 	const bool bAlreadyRegistered = AutomationTestClassNameToInstanceMap.Contains( InTestNameToRegister );
diff --git a/Engine/Source/Runtime/Core/Public/Misc/AutomationTest.h b/Engine/Source/Runtime/Core/Public/Misc/AutomationTest.h
--- a/Engine/Source/Runtime/Core/Public/Misc/AutomationTest.h
+++ b/Engine/Source/Runtime/Core/Public/Misc/AutomationTest.h
@@ -1023,6 +1023,11 @@ public:
 		return CurrentTest;
 	}
 
+	/**
+	* Whether to skip stack walk while iterating for listing the tests
+	*/
+	static bool NeedSkipStackWalk();
+
 	void NotifyScreenshotComparisonComplete(const FAutomationScreenshotCompareResults& CompareResults);
 	void NotifyTestDataRetrieved(bool bWasNew, const FString& JsonData);
 	void NotifyPerformanceDataRetrieved(bool bSuccess, const FString& ErrorMessage);
@@ -2625,70 +2630,70 @@ public:
 	void It(const FString& InDescription, TFunction<void()> DoWork)
 	{
 		const TSharedRef<FSpecDefinitionScope> CurrentScope = DefinitionScopeStack.Last();
-		SAFE_GETSTACK(Stack, 1, 1);
+		const auto Stack = GetStack();
 
 		PushDescription(InDescription);
-		CurrentScope->It.Push(MakeShareable(new FSpecIt(GetDescription(), GetId(), Stack[0].Filename, Stack[0].LineNumber, MakeShareable(new FSingleExecuteLatentCommand(this, DoWork, bEnableSkipIfError)))));
+		CurrentScope->It.Push(MakeShareable(new FSpecIt(GetDescription(), GetId(), Stack.Get()[0].Filename, Stack.Get()[0].LineNumber, MakeShareable(new FSingleExecuteLatentCommand(this, DoWork, bEnableSkipIfError)))));
 		PopDescription(InDescription);
 	}
 
 	void It(const FString& InDescription, EAsyncExecution Execution, TFunction<void()> DoWork)
 	{
 		const TSharedRef<FSpecDefinitionScope> CurrentScope = DefinitionScopeStack.Last();
-		SAFE_GETSTACK(Stack, 1, 1);
+		const auto Stack = GetStack();
 
 		PushDescription(InDescription);
-		CurrentScope->It.Push(MakeShareable(new FSpecIt(GetDescription(), GetId(), Stack[0].Filename, Stack[0].LineNumber, MakeShareable(new FAsyncLatentCommand(this, Execution, DoWork, DefaultTimeout, bEnableSkipIfError)))));
+		CurrentScope->It.Push(MakeShareable(new FSpecIt(GetDescription(), GetId(), Stack.Get()[0].Filename, Stack.Get()[0].LineNumber, MakeShareable(new FAsyncLatentCommand(this, Execution, DoWork, DefaultTimeout, bEnableSkipIfError)))));
 		PopDescription(InDescription);
 	}
 
 	void It(const FString& InDescription, EAsyncExecution Execution, const FTimespan& Timeout, TFunction<void()> DoWork)
 	{
 		const TSharedRef<FSpecDefinitionScope> CurrentScope = DefinitionScopeStack.Last();
-		SAFE_GETSTACK(Stack, 1, 1);
+		const auto Stack = GetStack();
 
 		PushDescription(InDescription);
-		CurrentScope->It.Push(MakeShareable(new FSpecIt(GetDescription(), GetId(), Stack[0].Filename, Stack[0].LineNumber, MakeShareable(new FAsyncLatentCommand(this, Execution, DoWork, Timeout, bEnableSkipIfError)))));
+		CurrentScope->It.Push(MakeShareable(new FSpecIt(GetDescription(), GetId(), Stack.Get()[0].Filename, Stack.Get()[0].LineNumber, MakeShareable(new FAsyncLatentCommand(this, Execution, DoWork, Timeout, bEnableSkipIfError)))));
 		PopDescription(InDescription);
 	}
 
 	void LatentIt(const FString& InDescription, TFunction<void(const FDoneDelegate&)> DoWork)
 	{
 		const TSharedRef<FSpecDefinitionScope> CurrentScope = DefinitionScopeStack.Last();
-		SAFE_GETSTACK(Stack, 1, 1);
+		const auto Stack = GetStack();
 
 		PushDescription(InDescription);
-		CurrentScope->It.Push(MakeShareable(new FSpecIt(GetDescription(), GetId(), Stack[0].Filename, Stack[0].LineNumber, MakeShareable(new FUntilDoneLatentCommand(this, DoWork, DefaultTimeout, bEnableSkipIfError)))));
+		CurrentScope->It.Push(MakeShareable(new FSpecIt(GetDescription(), GetId(), Stack.Get()[0].Filename, Stack.Get()[0].LineNumber, MakeShareable(new FUntilDoneLatentCommand(this, DoWork, DefaultTimeout, bEnableSkipIfError)))));
 		PopDescription(InDescription);
 	}
 
 	void LatentIt(const FString& InDescription, const FTimespan& Timeout, TFunction<void(const FDoneDelegate&)> DoWork)
 	{
 		const TSharedRef<FSpecDefinitionScope> CurrentScope = DefinitionScopeStack.Last();
-		SAFE_GETSTACK(Stack, 1, 1);
+		const auto Stack = GetStack();
 
 		PushDescription(InDescription);
-		CurrentScope->It.Push(MakeShareable(new FSpecIt(GetDescription(), GetId(), Stack[0].Filename, Stack[0].LineNumber, MakeShareable(new FUntilDoneLatentCommand(this, DoWork, Timeout, bEnableSkipIfError)))));
+		CurrentScope->It.Push(MakeShareable(new FSpecIt(GetDescription(), GetId(), Stack.Get()[0].Filename, Stack.Get()[0].LineNumber, MakeShareable(new FUntilDoneLatentCommand(this, DoWork, Timeout, bEnableSkipIfError)))));
 		PopDescription(InDescription);
 	}
 
 	void LatentIt(const FString& InDescription, EAsyncExecution Execution, TFunction<void(const FDoneDelegate&)> DoWork)
 	{
 		const TSharedRef<FSpecDefinitionScope> CurrentScope = DefinitionScopeStack.Last();
-		SAFE_GETSTACK(Stack, 1, 1);
+		const auto Stack = GetStack();
 
 		PushDescription(InDescription);
-		CurrentScope->It.Push(MakeShareable(new FSpecIt(GetDescription(), GetId(), Stack[0].Filename, Stack[0].LineNumber, MakeShareable(new FAsyncUntilDoneLatentCommand(this, Execution, DoWork, DefaultTimeout, bEnableSkipIfError)))));
+		CurrentScope->It.Push(MakeShareable(new FSpecIt(GetDescription(), GetId(), Stack.Get()[0].Filename, Stack.Get()[0].LineNumber, MakeShareable(new FAsyncUntilDoneLatentCommand(this, Execution, DoWork, DefaultTimeout, bEnableSkipIfError)))));
 		PopDescription(InDescription);
 	}
 
 	void LatentIt(const FString& InDescription, EAsyncExecution Execution, const FTimespan& Timeout, TFunction<void(const FDoneDelegate&)> DoWork)
 	{
 		const TSharedRef<FSpecDefinitionScope> CurrentScope = DefinitionScopeStack.Last();
-		SAFE_GETSTACK(Stack, 1, 1);
+		const auto Stack = GetStack();
 
 		PushDescription(InDescription);
-		CurrentScope->It.Push(MakeShareable(new FSpecIt(GetDescription(), GetId(), Stack[0].Filename, Stack[0].LineNumber, MakeShareable(new FAsyncUntilDoneLatentCommand(this, Execution, DoWork, Timeout, bEnableSkipIfError)))));
+		CurrentScope->It.Push(MakeShareable(new FSpecIt(GetDescription(), GetId(), Stack.Get()[0].Filename, Stack.Get()[0].LineNumber, MakeShareable(new FAsyncUntilDoneLatentCommand(this, Execution, DoWork, Timeout, bEnableSkipIfError)))));
 		PopDescription(InDescription);
 	}
 
@@ -3020,6 +3025,48 @@ private:
 		return CompleteId;
 	}
 
+	static TArray<FProgramCounterSymbolInfo> StackWalk()
+	{
+		QUICK_SCOPE_CYCLE_COUNTER(STAT_FAutomationSpecBase_StackWalk);
+
+		LLM_SCOPE_BYNAME(TEXT("AutomationTest/Framework"));
+		SAFE_GETSTACK(Stack, 1, 1);
+		return Stack;
+	}
+
+	static TArray<FProgramCounterSymbolInfo> SkipStackWalk()
+	{
+		QUICK_SCOPE_CYCLE_COUNTER(STAT_FAutomationSpecBase_SkipStackWalk);
+
+		LLM_SCOPE_BYNAME(TEXT("AutomationTest/Framework"));
+		TArray<FProgramCounterSymbolInfo> Stack;
+		FProgramCounterSymbolInfo First;
+		TCString<ANSICHAR>::Strcpy(First.Filename, FProgramCounterSymbolInfo::MAX_NAME_LENGTH, "Unknown");
+		First.LineNumber = 0;
+		Stack.Add(First);
+
+		return Stack;
+	}
+
+	static TSharedRef<TArray<FProgramCounterSymbolInfo>> GetStack()
+	{
+		QUICK_SCOPE_CYCLE_COUNTER(STAT_FAutomationSpecBase_GetStack);
+
+		static bool CurrentNeedSkipStackWalk(FAutomationTestFramework::NeedSkipStackWalk());
+		static TSharedRef<TArray<FProgramCounterSymbolInfo>> Stack = MakeShared<TArray<FProgramCounterSymbolInfo>>(
+			CurrentNeedSkipStackWalk ? SkipStackWalk() : StackWalk());
+
+		if (CurrentNeedSkipStackWalk != FAutomationTestFramework::NeedSkipStackWalk())
+		{
+			// This block is to react on changes of corresponding configuration variable */
+			CurrentNeedSkipStackWalk = !CurrentNeedSkipStackWalk;
+			Stack = MakeShared<TArray<FProgramCounterSymbolInfo>>(
+				CurrentNeedSkipStackWalk ? SkipStackWalk() : StackWalk());
+		}
+		
+		return Stack;
+	}
+
 private:
 
 	TArray<FString> Description;
