suppressed PVS-studio false-positive warnings

diff --git a/Engine/Source/Runtime/Core/Private/Misc/MTAccessDetector.cpp b/Engine/Source/Runtime/Core/Private/Misc/MTAccessDetector.cpp
--- a/Engine/Source/Runtime/Core/Private/Misc/MTAccessDetector.cpp
+++ b/Engine/Source/Runtime/Core/Private/Misc/MTAccessDetector.cpp
@@ -412,8 +412,8 @@ bool FRWRecursiveAccessDetectorTest::RunTest(const FString& Parameters)
 
 		if (!DestructionSentinel.bDestroyed)
 		{
-			AD->SetDestructionSentinel(PrevDestructionSentinel);
-			AD->ReleaseReadAccess();
+			AD->SetDestructionSentinel(PrevDestructionSentinel);  //-V774: same as above
+			AD->ReleaseReadAccess();  //-V774
 		}
 	}
 
