ZenCook: Add missing generation of package ids when serializing package store entries

diff --git a/Engine/Source/Runtime/Core/Private/IO/PackageStore.cpp b/Engine/Source/Runtime/Core/Private/IO/PackageStore.cpp
--- a/Engine/Source/Runtime/Core/Private/IO/PackageStore.cpp
+++ b/Engine/Source/Runtime/Core/Private/IO/PackageStore.cpp
@@ -55,6 +55,7 @@ FArchive& operator<<(FArchive& Ar, FPackageStoreEntryResource& PackageStoreEntry
 
 	if (Ar.IsLoading())
 	{
+		PackageStoreEntry.PackageId = FPackageId::FromName(PackageStoreEntry.PackageName);
 		PackageStoreEntry.Flags = static_cast<EPackageStoreEntryFlags>(Flags);
 	}
 
@@ -102,6 +103,7 @@ FPackageStoreEntryResource FPackageStoreEntryResource::FromCbObject(const FCbObj
 
 	Entry.Flags				= static_cast<EPackageStoreEntryFlags>(Obj["flags"].AsUInt32());
 	Entry.PackageName		= FName(Obj["packagename"].AsString());
+	Entry.PackageId			= FPackageId::FromName(Entry.PackageName);
 	Entry.SourcePackageName	= FName(Obj["sourcepackagename"].AsString());
 	Entry.Region			= FName(Obj["region"].AsString());
 	Entry.ExportInfo		= FPackageStoreExportInfo::FromCbObject(Obj["exportinfo"].AsObject());
