suppressed PVS-studio static analysis warning
https://horde.devtools.epicgames.com/job/63e3c73edd21ba4a6d6ebbf6?step=fe3b&issue=260146

diff --git a/Engine/Source/Runtime/Core/Private/Misc/MTAccessDetector.cpp b/Engine/Source/Runtime/Core/Private/Misc/MTAccessDetector.cpp
--- a/Engine/Source/Runtime/Core/Private/Misc/MTAccessDetector.cpp
+++ b/Engine/Source/Runtime/Core/Private/Misc/MTAccessDetector.cpp
@@ -412,7 +412,7 @@ bool FRWRecursiveAccessDetectorTest::RunTest(const FString& Parameters)
 
 		if (!DestructionSentinel.bDestroyed)
 		{
-			AD->SetDestructionSentinel(PrevDestructionSentinel);  //-V774: same as above
+			AD->SetDestructionSentinel(PrevDestructionSentinel);  //-V774
 			AD->ReleaseReadAccess();  //-V774
 		}
 	}
