[Insights] Fixed scope of metadata and mem tag for mem allocations made by the trace system (i.e. extends scope over the original allocator calls). So, heap allocations will inherit same "trace" tag.

diff --git a/Engine/Source/Runtime/Core/Private/ProfilingDebugging/MemoryTrace.cpp b/Engine/Source/Runtime/Core/Private/ProfilingDebugging/MemoryTrace.cpp
--- a/Engine/Source/Runtime/Core/Private/ProfilingDebugging/MemoryTrace.cpp
+++ b/Engine/Source/Runtime/Core/Private/ProfilingDebugging/MemoryTrace.cpp
@@ -697,6 +697,9 @@ FTraceMalloc::~FTraceMalloc()
 void* FTraceMalloc::Malloc(SIZE_T Count, uint32 Alignment)
 {
 #if UE_MEMORY_TRACE_ENABLED
+	UE_TRACE_METADATA_CLEAR_SCOPE();
+	UE_MEMSCOPE(TRACE_TAG);
+
 	void* NewPtr;
 	{
 		TGuardValue<bool> _(GDoNotAllocateInTrace, true);
@@ -707,9 +710,6 @@ void* FTraceMalloc::Malloc(SIZE_T Count, uint32 Alignment)
 	const uint32 AlignmentPow2 = uint32(FPlatformMath::CountTrailingZeros(Alignment));
 	const uint32 Alignment_SizeLower = (AlignmentPow2 << SizeShift) | uint32(Size & ((1 << SizeShift) - 1));
 
-	UE_TRACE_METADATA_CLEAR_SCOPE();
-	UE_MEMSCOPE(TRACE_TAG);
-
 	UE_TRACE_LOG(Memory, Alloc, MemAllocChannel)
 		<< Alloc.Address(uint64(NewPtr))
 		<< Alloc.CallstackId(0)
@@ -727,11 +727,6 @@ void* FTraceMalloc::Malloc(SIZE_T Count, uint32 Alignment)
 void* FTraceMalloc::Realloc(void* Original, SIZE_T Count, uint32 Alignment)
 {
 #if UE_MEMORY_TRACE_ENABLED
-	void* NewPtr = nullptr;
-	const uint64 Size = Count;
-	const uint32 AlignmentPow2 = uint32(FPlatformMath::CountTrailingZeros(Alignment));
-	const uint32 Alignment_SizeLower = (AlignmentPow2 << SizeShift) | uint32(Size & ((1 << SizeShift) - 1));
-
 	UE_TRACE_METADATA_CLEAR_SCOPE();
 	UE_MEMSCOPE(TRACE_TAG);
 
@@ -739,11 +734,16 @@ void* FTraceMalloc::Realloc(void* Original, SIZE_T Count, uint32 Alignment)
 		<< ReallocFree.Address(uint64(Original))
 		<< ReallocFree.RootHeap(uint8(EMemoryTraceRootHeap::SystemMemory));
 
+	void* NewPtr;
 	{
 		TGuardValue<bool> _(GDoNotAllocateInTrace, true);
 		NewPtr = WrappedMalloc->Realloc(Original, Count, Alignment);
 	}
 
+	const uint64 Size = Count;
+	const uint32 AlignmentPow2 = uint32(FPlatformMath::CountTrailingZeros(Alignment));
+	const uint32 Alignment_SizeLower = (AlignmentPow2 << SizeShift) | uint32(Size & ((1 << SizeShift) - 1));
+
 	UE_TRACE_LOG(Memory, ReallocAlloc, MemAllocChannel)
 		<< ReallocAlloc.Address(uint64(NewPtr))
 		<< ReallocAlloc.CallstackId(0)
