Fix for metadata scopes applied in tasks.

When launching a task, the metadata stack is captured at the calling site, and reapplied as a scope as soon as the task begins to execute. The macro that performed this reapplication was misnamed and didn't imply that it was actually creating a scope. This change removes the macro and adds a custom scope type which can be embedded in the inherited context for the task.

diff --git a/Engine/Source/Runtime/Core/Private/Async/InheritedContext.cpp b/Engine/Source/Runtime/Core/Private/Async/InheritedContext.cpp
--- a/Engine/Source/Runtime/Core/Private/Async/InheritedContext.cpp
+++ b/Engine/Source/Runtime/Core/Private/Async/InheritedContext.cpp
@@ -7,17 +7,21 @@ namespace UE
 	// this method must be defined in a .cpp file to avoid exporting dependency on trace module
 	UE_NODISCARD FInheritedContextScope FInheritedContextBase::RestoreInheritedContext()
 	{
-		UE_TRACE_METADATA_RESTORE_STACK(InheritedMetadataId);
-
 		return FInheritedContextScope{
 #if ENABLE_LOW_LEVEL_MEM_TRACKER
 				InheritedLLMTags
-	#if UE_MEMORY_TAGS_TRACE_ENABLED && UE_TRACE_ENABLED
+	#if (UE_MEMORY_TAGS_TRACE_ENABLED||UE_TRACE_METADATA_ENABLED) && UE_TRACE_ENABLED
 				,
 	#endif
 #endif
 #if UE_MEMORY_TAGS_TRACE_ENABLED && UE_TRACE_ENABLED
 				InheritedMemTag
+	#if UE_TRACE_METADATA_ENABLED && UE_TRACE_ENABLED
+				,
+	#endif
+#endif
+#if UE_TRACE_METADATA_ENABLED && UE_TRACE_ENABLED
+				InheritedMetadataId
 #endif
 		};
 	}
diff --git a/Engine/Source/Runtime/Core/Private/ProfilingDebugging/MetadataTrace.cpp b/Engine/Source/Runtime/Core/Private/ProfilingDebugging/MetadataTrace.cpp
--- a/Engine/Source/Runtime/Core/Private/ProfilingDebugging/MetadataTrace.cpp
+++ b/Engine/Source/Runtime/Core/Private/ProfilingDebugging/MetadataTrace.cpp
@@ -19,6 +19,32 @@ uint32 FMetadataTrace::SaveStack()
 	return Id;
 }
 
+////////////////////////////////////////////////////////////////////////////////////////////////////
+FMetadataRestoreScope::FMetadataRestoreScope(uint32 SavedMetadataIdentifier)
+{
+	if (SavedMetadataIdentifier)
+	{
+		ActivateScope(SavedMetadataIdentifier);
+	}
+}
+
+////////////////////////////////////////////////////////////////////////////////////////////////////
+FMetadataRestoreScope::~FMetadataRestoreScope()
+{}
+
+////////////////////////////////////////////////////////////////////////////////////////////////////
+void FMetadataRestoreScope::ActivateScope(uint32 InStackId)
+{
+	if (auto LogScope = FMetadataStackRestoreStackFields::LogScopeType::ScopedEnter<FMetadataStackRestoreStackFields>())
+	{
+		if (const auto& __restrict MemoryScope = *(FMetadataStackRestoreStackFields*)(&LogScope))
+		{
+			Inner.SetActive();
+			LogScope += LogScope << MemoryScope.Id(InStackId);
+		}
+	}
+}
+
 #endif // UE_TRACE_METADATA_ENABLED
 
 
diff --git a/Engine/Source/Runtime/Core/Public/Async/InheritedContext.h b/Engine/Source/Runtime/Core/Public/Async/InheritedContext.h
--- a/Engine/Source/Runtime/Core/Public/Async/InheritedContext.h
+++ b/Engine/Source/Runtime/Core/Public/Async/InheritedContext.h
@@ -54,29 +54,44 @@ namespace UE
 		UE_NONCOPYABLE(FInheritedContextScope);
 
 		friend class FInheritedContextBase; // allow construction only by `FInheritedContextBase`
-
+		
 		FInheritedContextScope(
 #if ENABLE_LOW_LEVEL_MEM_TRACKER
 			const FLLMActiveTagsCapture& InInheritedLLMTag
-	#if UE_MEMORY_TAGS_TRACE_ENABLED && UE_TRACE_ENABLED
+	#if UE_TRACE_ENABLED && (UE_MEMORY_TAGS_TRACE_ENABLED || UE_TRACE_METADATA_ENABLED)
 			,
 	#endif
 #endif
-#if UE_MEMORY_TAGS_TRACE_ENABLED && UE_TRACE_ENABLED
+#if (UE_MEMORY_TAGS_TRACE_ENABLED && UE_TRACE_ENABLED)
 			int32 InInheritedMemTag
+	#if (UE_TRACE_METADATA_ENABLED && UE_TRACE_ENABLED)
+			,
+	#endif
+#endif
+#if (UE_TRACE_METADATA_ENABLED && UE_TRACE_ENABLED)
+			uint32 InInheritedMetadataId
 #endif
 		)
 #if ENABLE_LOW_LEVEL_MEM_TRACKER
 			: LLMScopes(InInheritedLLMTag)
-	#if UE_MEMORY_TAGS_TRACE_ENABLED && UE_TRACE_ENABLED
+	#if UE_TRACE_ENABLED && (UE_MEMORY_TAGS_TRACE_ENABLED || UE_TRACE_METADATA_ENABLED)
 			,
 	#endif
 #endif
-#if UE_MEMORY_TAGS_TRACE_ENABLED && UE_TRACE_ENABLED
+#if (UE_MEMORY_TAGS_TRACE_ENABLED && UE_TRACE_ENABLED)
 	#if !ENABLE_LOW_LEVEL_MEM_TRACKER
 			:
 	#endif
 			MemScope(InInheritedMemTag)
+	#if (UE_TRACE_METADATA_ENABLED && UE_TRACE_ENABLED)
+			,
+	#endif
+#endif
+#if (UE_TRACE_METADATA_ENABLED && UE_TRACE_ENABLED)
+	#if !ENABLE_LOW_LEVEL_MEM_TRACKER && !(UE_MEMORY_TAGS_TRACE_ENABLED && UE_TRACE_ENABLED)
+			:
+	#endif
+			MetaScope(InInheritedMetadataId)
 #endif
 		{
 		}
@@ -86,9 +101,13 @@ namespace UE
 		FLLMActiveTagsScope LLMScopes;
 #endif
 
-#if UE_MEMORY_TAGS_TRACE_ENABLED && UE_TRACE_ENABLED
+#if (UE_MEMORY_TAGS_TRACE_ENABLED && UE_TRACE_ENABLED)
 		FMemScope MemScope;
 #endif
+
+#if (UE_TRACE_METADATA_ENABLED && UE_TRACE_ENABLED) 
+		FMetadataRestoreScope MetaScope;
+#endif
 	};
 
 	// this class extends the inherited context (see private members for what the inherited context is) to cover async execution. 
@@ -107,7 +126,7 @@ namespace UE
 			InheritedMemTag = MemoryTrace_GetActiveTag();
 #endif
 
-#if UE_TRACE_METADATA_ENABLED
+#if UE_TRACE_METADATA_ENABLED && UE_TRACE_ENABLED
 			InheritedMetadataId = UE_TRACE_METADATA_SAVE_STACK();
 #endif
 		}
@@ -124,8 +143,8 @@ namespace UE
 		int32 InheritedMemTag;
 #endif
 
-#if UE_TRACE_METADATA_ENABLED
+#if UE_TRACE_METADATA_ENABLED && UE_TRACE_ENABLED
 		uint32 InheritedMetadataId;
 #endif
 	};
-}
\ No newline at end of file
+}
diff --git a/Engine/Source/Runtime/Core/Public/ProfilingDebugging/MetadataTrace.h b/Engine/Source/Runtime/Core/Public/ProfilingDebugging/MetadataTrace.h
--- a/Engine/Source/Runtime/Core/Public/ProfilingDebugging/MetadataTrace.h
+++ b/Engine/Source/Runtime/Core/Public/ProfilingDebugging/MetadataTrace.h
@@ -73,12 +73,17 @@ public:
 	FMetadataTrace::SaveStack();
 
 /**
- * Restore a previously saved stack.
+ * A scope that restores a meta stack previously stored with UE_TRACE_METADATA_SAVE_STACK during it's lifetime.
  */
-#define UE_TRACE_METADATA_RESTORE_STACK(InId) \
-	UE_TRACE_LOG_SCOPED(MetadataStack, RestoreStack, MetadataChannel) \
-		<< RestoreStack.Id(InId);
-
+class FMetadataRestoreScope
+{
+public:
+	CORE_API FMetadataRestoreScope(uint32 SavedMetadataIdentifier);
+	CORE_API ~FMetadataRestoreScope();;
+private:
+	void ActivateScope(uint32 InStackId);
+	UE::Trace::Private::FScopedLogScope Inner;
+};
 ////////////////////////////////////////////////////////////////////////////////////////////////////
 
 #else
