[Insights] Fixed typo in warning message for SnapshotSend (trace command).

diff --git a/Engine/Source/Runtime/Core/Private/ProfilingDebugging/TraceAuxiliary.cpp b/Engine/Source/Runtime/Core/Private/ProfilingDebugging/TraceAuxiliary.cpp
--- a/Engine/Source/Runtime/Core/Private/ProfilingDebugging/TraceAuxiliary.cpp
+++ b/Engine/Source/Runtime/Core/Private/ProfilingDebugging/TraceAuxiliary.cpp
@@ -972,9 +972,9 @@ static void TraceAuxiliarySnapshotSend(const TArray<FString>& Args)
 	{
 		LexFromString(Port, *Args[1]);
 	}
-	if (Args.Num() > 2) 
+	if (Args.Num() > 2)
 	{
-		UE_LOG(LogConsoleResponse, Warning, TEXT("Invalid arguments. Usage: Trace.SnapshotFile <Host> <Port>"));
+		UE_LOG(LogConsoleResponse, Warning, TEXT("Invalid arguments. Usage: Trace.SnapshotSend <Host> <Port>"));
 		return;
 	}
 
@@ -1379,7 +1379,7 @@ bool FTraceAuxiliary::Start(EConnectionType Type, const TCHAR* Target, const TCH
 	}
 
 	uint16 SendFlags = (Options && Options->bExcludeTail) ? UE::Trace::FSendFlags::ExcludeTail : 0;
-	
+
 	if (Type == EConnectionType::File)
 	{
 		return GTraceAuxiliary.Connect(ETraceConnectType::File, Target, LogCategory, SendFlags);
@@ -1545,7 +1545,7 @@ void FTraceAuxiliary::Initialize(const TCHAR* CommandLine)
 	Desc.TailSizeBytes = 32 << 20;
 #endif
 	SetupInitFromConfig(Desc);
-	
+
 	Desc.bUseWorkerThread = bShouldStartWorkerThread;
 	Desc.bUseImportantCache = (FParse::Param(CommandLine, TEXT("tracenocache")) == false);
 	Desc.OnConnectionFunc = &OnConnectionCallback;
