Fix for missing NaNs in matrix Inverse().

diff --git a/Engine/Source/Runtime/Core/Public/Math/Matrix.inl b/Engine/Source/Runtime/Core/Public/Math/Matrix.inl
--- a/Engine/Source/Runtime/Core/Public/Math/Matrix.inl
+++ b/Engine/Source/Runtime/Core/Public/Math/Matrix.inl
@@ -335,7 +335,11 @@ inline TMatrix<T> TMatrix<T>::Inverse() const
 	{
 		const T	Det = Determinant();
 
-		if (Det == 0.0f)
+		// Under fp:fast this comparison can (and has been) getting inverted such that NaNs
+		// get an Identity matrix out. Under fp:precise any NaNs (correctly) get treated as != 0 
+		// and carry through to the inverse and crash, when this function is supposed to handle 
+		// that (vs InverseFast which expects well-formed matrices)
+		if (FMath::IsFinite(Det) == false || Det == 0.0f)
 		{
 			Result = Identity;
 		}
