Logging: Renamed structured logging macros from UE_SLOG to UE_LOGFMT

The new names are more discoverable and easier to understand.

#rnx

diff --git a/Engine/Source/Editor/UnrealEd/Private/Commandlets/GatherTextFromSourceCommandlet.cpp b/Engine/Source/Editor/UnrealEd/Private/Commandlets/GatherTextFromSourceCommandlet.cpp
--- a/Engine/Source/Editor/UnrealEd/Private/Commandlets/GatherTextFromSourceCommandlet.cpp
+++ b/Engine/Source/Editor/UnrealEd/Private/Commandlets/GatherTextFromSourceCommandlet.cpp
@@ -288,10 +288,10 @@ int32 UGatherTextFromSourceCommandlet::Main( const FString& Params )
 
 	Parsables.Add(new FStringTableEntryMetaDataMacroDescriptor());
 
-	Parsables.Add(new FStructuredLogMacroDescriptor(TEXT("UE_SLOG_LOC"), FStructuredLogMacroDescriptor::EFlags::None));
-	Parsables.Add(new FStructuredLogMacroDescriptor(TEXT("UE_SLOG_LOC_EX"), FStructuredLogMacroDescriptor::EFlags::None));
-	Parsables.Add(new FStructuredLogMacroDescriptor(TEXT("UE_SLOG_NSLOC"), FStructuredLogMacroDescriptor::EFlags::Namespace));
-	Parsables.Add(new FStructuredLogMacroDescriptor(TEXT("UE_SLOG_NSLOC_EX"), FStructuredLogMacroDescriptor::EFlags::Namespace));
+	Parsables.Add(new FStructuredLogMacroDescriptor(TEXT("UE_LOGFMT_LOC"), FStructuredLogMacroDescriptor::EFlags::None));
+	Parsables.Add(new FStructuredLogMacroDescriptor(TEXT("UE_LOGFMT_LOC_EX"), FStructuredLogMacroDescriptor::EFlags::None));
+	Parsables.Add(new FStructuredLogMacroDescriptor(TEXT("UE_LOGFMT_NSLOC"), FStructuredLogMacroDescriptor::EFlags::Namespace));
+	Parsables.Add(new FStructuredLogMacroDescriptor(TEXT("UE_LOGFMT_NSLOC_EX"), FStructuredLogMacroDescriptor::EFlags::Namespace));
 
 	Parsables.Add(new FIniNamespaceDescriptor());
 
@@ -2302,8 +2302,8 @@ UGatherTextFromSourceCommandlet::FStructuredLogMacroDescriptor::FStructuredLogMa
 void UGatherTextFromSourceCommandlet::FStructuredLogMacroDescriptor::TryParse(const FString& Text, FSourceFileParseContext& Context) const
 {
 	// Attempt to parse something of the format
-	// UE_SLOG_LOC[_EX](CategoryName, Verbosity, Key, Format, ...)
-	// UE_SLOG_NSLOC[_EX](CategoryName, Verbosity, Namespace, Key, Format, ...)
+	// UE_LOGFMT_LOC[_EX](CategoryName, Verbosity, Key, Format, ...)
+	// UE_LOGFMT_NSLOC[_EX](CategoryName, Verbosity, Namespace, Key, Format, ...)
 	int32 NextArg = 2;
 
 	if (!Context.ExcludedRegion && !Context.WithinBlockComment && !Context.WithinLineComment && !Context.WithinStringLiteral)
diff --git a/Engine/Source/Runtime/Core/Public/Logging/StructuredLog.h b/Engine/Source/Runtime/Core/Public/Logging/StructuredLog.h
--- a/Engine/Source/Runtime/Core/Public/Logging/StructuredLog.h
+++ b/Engine/Source/Runtime/Core/Public/Logging/StructuredLog.h
@@ -26,10 +26,10 @@ class FTextFormat;
  * Supports either positional or named parameters, but not a mix of these styles.
  *
  * Positional: The field values must exactly match the fields referenced by Format.
- * UE_SLOG(LogCore, Warning, "Loading '{Name}' failed with error {Error}", Package->GetName(), ErrorCode);
+ * UE_LOGFMT(LogCore, Warning, "Loading '{Name}' failed with error {Error}", Package->GetName(), ErrorCode);
  *
  * Named: The fields must contain every field referenced by Format. Order is irrelevant and extra fields are permitted.
- * UE_SLOG(LogCore, Warning, "Loading '{Name}' failed with error {Error}",
+ * UE_LOGFMT(LogCore, Warning, "Loading '{Name}' failed with error {Error}",
  *     ("Name", Package->GetName()), ("Error", ErrorCode), ("Flags", LoadFlags));
  *
  * Field names must match "[A-Za-z0-9_]+" and must be unique within this log event.
@@ -40,28 +40,28 @@ class FTextFormat;
  * @param Format         Format string in the style of FLogTemplate.
  * @param Fields[0-16]   Zero to sixteen fields or field values.
  */
-#define UE_SLOG(CategoryName, Verbosity, Format, ...) UE_PRIVATE_SLOG_CALL(UE_SLOG_EX, (CategoryName, Verbosity, Format UE_PRIVATE_SLOG_FIELDS(__VA_ARGS__)))
+#define UE_LOGFMT(CategoryName, Verbosity, Format, ...) UE_PRIVATE_LOGFMT_CALL(UE_LOGFMT_EX, (CategoryName, Verbosity, Format UE_PRIVATE_LOGFMT_FIELDS(__VA_ARGS__)))
 
 /**
  * Records a structured log event if this category is active at this level of verbosity.
  *
- * This has the same functionality as UE_SLOG but removes the limit on field count.
+ * This has the same functionality as UE_LOGFMT but removes the limit on field count.
  *
- * Positional: Values must be wrapped in UE_SLOG_VALUE.
- * UE_SLOG_EX(LogCore, Warning, "Loading '{Name}' failed with error {Error}",
- *     UE_SLOG_VALUE(Package->GetName()), UE_SLOG_VALUE(ErrorCode));
+ * Positional: Values must be wrapped in UE_LOGFMT_VALUE.
+ * UE_LOGFMT_EX(LogCore, Warning, "Loading '{Name}' failed with error {Error}",
+ *     UE_LOGFMT_VALUE(Package->GetName()), UE_LOGFMT_VALUE(ErrorCode));
  *
- * Named: Fields must be wrapped in UE_SLOG_FIELD.
- * UE_SLOG_EX(LogCore, Warning, "Loading '{Name}' failed with error {Error}",
- *     UE_SLOG_FIELD("Name", Package->GetName()), UE_SLOG_FIELD("Error", ErrorCode), UE_SLOG_FIELD("Flags", LoadFlags));
+ * Named: Fields must be wrapped in UE_LOGFMT_FIELD.
+ * UE_LOGFMT_EX(LogCore, Warning, "Loading '{Name}' failed with error {Error}",
+ *     UE_LOGFMT_FIELD("Name", Package->GetName()), UE_LOGFMT_FIELD("Error", ErrorCode), UE_LOGFMT_FIELD("Flags", LoadFlags));
  */
-#define UE_SLOG_EX(CategoryName, Verbosity, Format, ...) UE_PRIVATE_SLOG(CategoryName, Verbosity, Format, ##__VA_ARGS__)
+#define UE_LOGFMT_EX(CategoryName, Verbosity, Format, ...) UE_PRIVATE_LOGFMT(CategoryName, Verbosity, Format, ##__VA_ARGS__)
 
 /**
  * Records a localized structured log event if this category is active at this level of verbosity.
  *
  * Example:
- * UE_SLOG_LOC(LogCore, Warning, "LoadingFailed", "Loading '{Name}' failed with error {Error}",
+ * UE_LOGFMT_LOC(LogCore, Warning, "LoadingFailed", "Loading '{Name}' failed with error {Error}",
  *     ("Name", Package->GetName()), ("Error", ErrorCode), ("Flags", LoadFlags));
  *
  * Field names must match "[A-Za-z0-9_]+" and must be unique within this log event.
@@ -75,50 +75,47 @@ class FTextFormat;
  * @param Format         Format string in the style of FTextFormat.
  * @param Fields[0-16]   Zero to sixteen fields in the format ("Name", Value).
  */
-#define UE_SLOG_LOC(CategoryName, Verbosity, Key, Format, ...) \
-	UE_SLOG_NSLOC(CategoryName, Verbosity, LOCTEXT_NAMESPACE, Key, Format, ##__VA_ARGS__)
-#define UE_SLOG_NSLOC(CategoryName, Verbosity, Namespace, Key, Format, ...) \
-	UE_PRIVATE_SLOG_CALL(UE_SLOG_NSLOC_EX, (CategoryName, Verbosity, Namespace, Key, Format UE_PRIVATE_SLOG_FIELDS(__VA_ARGS__)))
+#define UE_LOGFMT_LOC(CategoryName, Verbosity, Key, Format, ...) \
+	UE_LOGFMT_NSLOC(CategoryName, Verbosity, LOCTEXT_NAMESPACE, Key, Format, ##__VA_ARGS__)
+#define UE_LOGFMT_NSLOC(CategoryName, Verbosity, Namespace, Key, Format, ...) \
+	UE_PRIVATE_LOGFMT_CALL(UE_LOGFMT_NSLOC_EX, (CategoryName, Verbosity, Namespace, Key, Format UE_PRIVATE_LOGFMT_FIELDS(__VA_ARGS__)))
 
 /**
  * Records a localized structured log event if this category is active at this level of verbosity.
  *
  * Example:
- * UE_SLOG_LOC_EX(LogCore, Warning, "LoadingFailed", "Loading '{PackageName}' failed with error {Error}",
- *     UE_SLOG_FIELD("Name", Package->GetName()), UE_SLOG_FIELD("Error", ErrorCode), UE_SLOG_FIELD("Flags", LoadFlags));
+ * UE_LOGFMT_LOC_EX(LogCore, Warning, "LoadingFailed", "Loading '{PackageName}' failed with error {Error}",
+ *     UE_LOGFMT_FIELD("Name", Package->GetName()), UE_LOGFMT_FIELD("Error", ErrorCode), UE_LOGFMT_FIELD("Flags", LoadFlags));
  *
- * Same as UE_SLOG_LOC and works for any number of fields.
- * Fields must be written as UE_SLOG_FIELD("Name", Value).
+ * Same as UE_LOGFMT_LOC and works for any number of fields.
+ * Fields must be written as UE_LOGFMT_FIELD("Name", Value).
  */
-#define UE_SLOG_LOC_EX(CategoryName, Verbosity, Key, Format, ...) \
-	UE_SLOG_NSLOC_EX(CategoryName, Verbosity, LOCTEXT_NAMESPACE, Key, Format, ##__VA_ARGS__)
-#define UE_SLOG_NSLOC_EX(CategoryName, Verbosity, Namespace, Key, Format, ...) \
-	UE_PRIVATE_SLOG_LOC(CategoryName, Verbosity, Namespace, Key, Format, ##__VA_ARGS__)
+#define UE_LOGFMT_LOC_EX(CategoryName, Verbosity, Key, Format, ...) \
+	UE_LOGFMT_NSLOC_EX(CategoryName, Verbosity, LOCTEXT_NAMESPACE, Key, Format, ##__VA_ARGS__)
+#define UE_LOGFMT_NSLOC_EX(CategoryName, Verbosity, Namespace, Key, Format, ...) \
+	UE_PRIVATE_LOGFMT_LOC(CategoryName, Verbosity, Namespace, Key, Format, ##__VA_ARGS__)
 
-/** Expands to a named structured log field. UE_SLOG_FIELD("Name", Value) */
-#define UE_SLOG_FIELD(Name, Value) UE_PRIVATE_SLOG_FIELD((Name, Value))
+/** Expands to a named structured log field. UE_LOGFMT_FIELD("Name", Value) */
+#define UE_LOGFMT_FIELD(Name, Value) UE_PRIVATE_LOGFMT_FIELD((Name, Value))
 
-/** Expands to a structured log value. UE_SLOG_VALUE(Value) */
-#define UE_SLOG_VALUE(Value) Value
+/** Expands to a structured log value. UE_LOGFMT_VALUE(Value) */
+#define UE_LOGFMT_VALUE(Value) Value
 
 #else // #if !NO_LOGGING
 
-#define UE_SLOG(...)
-#define UE_SLOG_EX(...)
+#define UE_LOGFMT(...)
+#define UE_LOGFMT_EX(...)
 
-#define UE_SLOG_LOC(...)
-#define UE_SLOG_NSLOC(...)
-#define UE_SLOG_LOC_EX(...)
-#define UE_SLOG_NSLOC_EX(...)
+#define UE_LOGFMT_LOC(...)
+#define UE_LOGFMT_NSLOC(...)
+#define UE_LOGFMT_LOC_EX(...)
+#define UE_LOGFMT_NSLOC_EX(...)
 
-#define UE_SLOG_FIELD(Name, Value)
-#define UE_SLOG_VALUE(Value)
+#define UE_LOGFMT_FIELD(Name, Value)
+#define UE_LOGFMT_VALUE(Value)
 
 #endif
 
-#define UE_SLOG_NAMED UE_DEPRECATED_MACRO(5.2, "UE_SLOG_NAMED is deprecated. Use UE_SLOG.") UE_SLOG
-#define UE_SLOG_NAMED_EX UE_DEPRECATED_MACRO(5.2, "UE_SLOG_NAMED_EX is deprecated. Use UE_SLOG_EX.") UE_SLOG_EX
-
 struct FDateTime;
 
 namespace UE
@@ -389,7 +386,7 @@ inline void LogIfActive(FieldArgTypes&&... FieldArgs)
 
 } // UE::Logging::Private
 
-#define UE_PRIVATE_SLOG(CategoryName, Verbosity, Format, ...) \
+#define UE_PRIVATE_LOGFMT(CategoryName, Verbosity, Format, ...) \
 	do \
 	{ \
 		static ::UE::Logging::Private::FStaticLogDynamicData LOG_Dynamic; \
@@ -398,7 +395,7 @@ inline void LogIfActive(FieldArgTypes&&... FieldArgs)
 	} \
 	while (false)
 
-#define UE_PRIVATE_SLOG_LOC(CategoryName, Verbosity, Namespace, Key, Format, ...) \
+#define UE_PRIVATE_LOGFMT_LOC(CategoryName, Verbosity, Namespace, Key, Format, ...) \
 	do \
 	{ \
 		static ::UE::Logging::Private::FStaticLogDynamicData LOG_Dynamic; \
@@ -410,42 +407,42 @@ inline void LogIfActive(FieldArgTypes&&... FieldArgs)
 // This macro expands a field from either `(Name, Value)` or `Value`
 // A `(Name, Value)` field is converted to `CheckFieldName(Name), Value`
 // A `Value` field is passed through as `Value`
-#define UE_PRIVATE_SLOG_FIELD(Field) UE_PRIVATE_SLOG_FIELD_EXPAND(UE_PRIVATE_SLOG_NAMED_FIELD Field)
+#define UE_PRIVATE_LOGFMT_FIELD(Field) UE_PRIVATE_LOGFMT_FIELD_EXPAND(UE_PRIVATE_LOGFMT_NAMED_FIELD Field)
 // This macro is only called when the field was parenthesized.
-#define UE_PRIVATE_SLOG_NAMED_FIELD(Name, ...) UE_PRIVATE_SLOG_NAMED_FIELD ::UE::Logging::Private::CheckFieldName(Name), __VA_ARGS__
-// The next three macros remove UE_PRIVATE_SLOG_NAMED_FIELD from the expanded expression.
-#define UE_PRIVATE_SLOG_FIELD_EXPAND(...) UE_PRIVATE_SLOG_FIELD_EXPAND_INNER(__VA_ARGS__)
-#define UE_PRIVATE_SLOG_FIELD_EXPAND_INNER(...) UE_PRIVATE_SLOG_STRIP_ ## __VA_ARGS__
-#define UE_PRIVATE_SLOG_STRIP_UE_PRIVATE_SLOG_NAMED_FIELD
-
-// This macro expands `Arg1, Arg2` to `UE_PRIVATE_SLOG_FIELD(Arg1), UE_PRIVATE_SLOG_FIELD(Arg2), ...`
-// This macro expands `("Name1", Arg1), ("Name2", Arg2)` to `UE_PRIVATE_SLOG_FIELD(("Name1", Arg1)), UE_PRIVATE_SLOG_FIELD(("Name2", Arg2)), ...
-#define UE_PRIVATE_SLOG_FIELDS(...) UE_PRIVATE_SLOG_CALL(PREPROCESSOR_JOIN(UE_PRIVATE_SLOG_FIELDS_, UE_PRIVATE_SLOG_COUNT(__VA_ARGS__)), (__VA_ARGS__))
-
-#define UE_PRIVATE_SLOG_FIELDS_0()
-#define UE_PRIVATE_SLOG_FIELDS_1(A)                 , UE_PRIVATE_SLOG_FIELD(A)
-#define UE_PRIVATE_SLOG_FIELDS_2(A,B)               , UE_PRIVATE_SLOG_FIELD(A), UE_PRIVATE_SLOG_FIELD(B)
-#define UE_PRIVATE_SLOG_FIELDS_3(A,B,C)             , UE_PRIVATE_SLOG_FIELD(A), UE_PRIVATE_SLOG_FIELD(B), UE_PRIVATE_SLOG_FIELD(C)
-#define UE_PRIVATE_SLOG_FIELDS_4(A,B,C,D)           , UE_PRIVATE_SLOG_FIELD(A), UE_PRIVATE_SLOG_FIELD(B), UE_PRIVATE_SLOG_FIELD(C), UE_PRIVATE_SLOG_FIELD(D)
-#define UE_PRIVATE_SLOG_FIELDS_5(A,B,C,D,E)         , UE_PRIVATE_SLOG_FIELD(A), UE_PRIVATE_SLOG_FIELD(B), UE_PRIVATE_SLOG_FIELD(C), UE_PRIVATE_SLOG_FIELD(D), UE_PRIVATE_SLOG_FIELD(E)
-#define UE_PRIVATE_SLOG_FIELDS_6(A,B,C,D,E,F)       , UE_PRIVATE_SLOG_FIELD(A), UE_PRIVATE_SLOG_FIELD(B), UE_PRIVATE_SLOG_FIELD(C), UE_PRIVATE_SLOG_FIELD(D), UE_PRIVATE_SLOG_FIELD(E), UE_PRIVATE_SLOG_FIELD(F)
-#define UE_PRIVATE_SLOG_FIELDS_7(A,B,C,D,E,F,G)     , UE_PRIVATE_SLOG_FIELD(A), UE_PRIVATE_SLOG_FIELD(B), UE_PRIVATE_SLOG_FIELD(C), UE_PRIVATE_SLOG_FIELD(D), UE_PRIVATE_SLOG_FIELD(E), UE_PRIVATE_SLOG_FIELD(F), UE_PRIVATE_SLOG_FIELD(G)
-#define UE_PRIVATE_SLOG_FIELDS_8(A,B,C,D,E,F,G,H)   , UE_PRIVATE_SLOG_FIELD(A), UE_PRIVATE_SLOG_FIELD(B), UE_PRIVATE_SLOG_FIELD(C), UE_PRIVATE_SLOG_FIELD(D), UE_PRIVATE_SLOG_FIELD(E), UE_PRIVATE_SLOG_FIELD(F), UE_PRIVATE_SLOG_FIELD(G), UE_PRIVATE_SLOG_FIELD(H)
-#define UE_PRIVATE_SLOG_FIELDS_9(A,B,C,D,E,F,G,H,I) , UE_PRIVATE_SLOG_FIELD(A), UE_PRIVATE_SLOG_FIELD(B), UE_PRIVATE_SLOG_FIELD(C), UE_PRIVATE_SLOG_FIELD(D), UE_PRIVATE_SLOG_FIELD(E), UE_PRIVATE_SLOG_FIELD(F), UE_PRIVATE_SLOG_FIELD(G), UE_PRIVATE_SLOG_FIELD(H), UE_PRIVATE_SLOG_FIELD(I)
-
-#define UE_PRIVATE_SLOG_FIELDS_10(A,B,C,D,E,F,G,H,I,J)             , UE_PRIVATE_SLOG_FIELD(A), UE_PRIVATE_SLOG_FIELD(B), UE_PRIVATE_SLOG_FIELD(C), UE_PRIVATE_SLOG_FIELD(D), UE_PRIVATE_SLOG_FIELD(E), UE_PRIVATE_SLOG_FIELD(F), UE_PRIVATE_SLOG_FIELD(G), UE_PRIVATE_SLOG_FIELD(H), UE_PRIVATE_SLOG_FIELD(I), UE_PRIVATE_SLOG_FIELD(J)
-#define UE_PRIVATE_SLOG_FIELDS_11(A,B,C,D,E,F,G,H,I,J,K)           , UE_PRIVATE_SLOG_FIELD(A), UE_PRIVATE_SLOG_FIELD(B), UE_PRIVATE_SLOG_FIELD(C), UE_PRIVATE_SLOG_FIELD(D), UE_PRIVATE_SLOG_FIELD(E), UE_PRIVATE_SLOG_FIELD(F), UE_PRIVATE_SLOG_FIELD(G), UE_PRIVATE_SLOG_FIELD(H), UE_PRIVATE_SLOG_FIELD(I), UE_PRIVATE_SLOG_FIELD(J), UE_PRIVATE_SLOG_FIELD(K)
-#define UE_PRIVATE_SLOG_FIELDS_12(A,B,C,D,E,F,G,H,I,J,K,L)         , UE_PRIVATE_SLOG_FIELD(A), UE_PRIVATE_SLOG_FIELD(B), UE_PRIVATE_SLOG_FIELD(C), UE_PRIVATE_SLOG_FIELD(D), UE_PRIVATE_SLOG_FIELD(E), UE_PRIVATE_SLOG_FIELD(F), UE_PRIVATE_SLOG_FIELD(G), UE_PRIVATE_SLOG_FIELD(H), UE_PRIVATE_SLOG_FIELD(I), UE_PRIVATE_SLOG_FIELD(J), UE_PRIVATE_SLOG_FIELD(K), UE_PRIVATE_SLOG_FIELD(L)
-#define UE_PRIVATE_SLOG_FIELDS_13(A,B,C,D,E,F,G,H,I,J,K,L,M)       , UE_PRIVATE_SLOG_FIELD(A), UE_PRIVATE_SLOG_FIELD(B), UE_PRIVATE_SLOG_FIELD(C), UE_PRIVATE_SLOG_FIELD(D), UE_PRIVATE_SLOG_FIELD(E), UE_PRIVATE_SLOG_FIELD(F), UE_PRIVATE_SLOG_FIELD(G), UE_PRIVATE_SLOG_FIELD(H), UE_PRIVATE_SLOG_FIELD(I), UE_PRIVATE_SLOG_FIELD(J), UE_PRIVATE_SLOG_FIELD(K), UE_PRIVATE_SLOG_FIELD(L), UE_PRIVATE_SLOG_FIELD(M)
-#define UE_PRIVATE_SLOG_FIELDS_14(A,B,C,D,E,F,G,H,I,J,K,L,M,N)     , UE_PRIVATE_SLOG_FIELD(A), UE_PRIVATE_SLOG_FIELD(B), UE_PRIVATE_SLOG_FIELD(C), UE_PRIVATE_SLOG_FIELD(D), UE_PRIVATE_SLOG_FIELD(E), UE_PRIVATE_SLOG_FIELD(F), UE_PRIVATE_SLOG_FIELD(G), UE_PRIVATE_SLOG_FIELD(H), UE_PRIVATE_SLOG_FIELD(I), UE_PRIVATE_SLOG_FIELD(J), UE_PRIVATE_SLOG_FIELD(K), UE_PRIVATE_SLOG_FIELD(L), UE_PRIVATE_SLOG_FIELD(M), UE_PRIVATE_SLOG_FIELD(N)
-#define UE_PRIVATE_SLOG_FIELDS_15(A,B,C,D,E,F,G,H,I,J,K,L,M,N,O)   , UE_PRIVATE_SLOG_FIELD(A), UE_PRIVATE_SLOG_FIELD(B), UE_PRIVATE_SLOG_FIELD(C), UE_PRIVATE_SLOG_FIELD(D), UE_PRIVATE_SLOG_FIELD(E), UE_PRIVATE_SLOG_FIELD(F), UE_PRIVATE_SLOG_FIELD(G), UE_PRIVATE_SLOG_FIELD(H), UE_PRIVATE_SLOG_FIELD(I), UE_PRIVATE_SLOG_FIELD(J), UE_PRIVATE_SLOG_FIELD(K), UE_PRIVATE_SLOG_FIELD(L), UE_PRIVATE_SLOG_FIELD(M), UE_PRIVATE_SLOG_FIELD(N), UE_PRIVATE_SLOG_FIELD(O)
-#define UE_PRIVATE_SLOG_FIELDS_16(A,B,C,D,E,F,G,H,I,J,K,L,M,N,O,P) , UE_PRIVATE_SLOG_FIELD(A), UE_PRIVATE_SLOG_FIELD(B), UE_PRIVATE_SLOG_FIELD(C), UE_PRIVATE_SLOG_FIELD(D), UE_PRIVATE_SLOG_FIELD(E), UE_PRIVATE_SLOG_FIELD(F), UE_PRIVATE_SLOG_FIELD(G), UE_PRIVATE_SLOG_FIELD(H), UE_PRIVATE_SLOG_FIELD(I), UE_PRIVATE_SLOG_FIELD(J), UE_PRIVATE_SLOG_FIELD(K), UE_PRIVATE_SLOG_FIELD(L), UE_PRIVATE_SLOG_FIELD(M), UE_PRIVATE_SLOG_FIELD(N), UE_PRIVATE_SLOG_FIELD(O), UE_PRIVATE_SLOG_FIELD(P)
-
-#define UE_PRIVATE_SLOG_COUNT(...) UE_PRIVATE_SLOG_CALL(UE_PRIVATE_SLOG_COUNT_IMPL, (_, ##__VA_ARGS__, 16,15,14,13,12,11,10,9,8,7,6,5,4,3,2,1,0))
-#define UE_PRIVATE_SLOG_COUNT_IMPL(_, A,B,C,D,E,F,G,H,I,J,K,L,M,N,O,P, Count, ...) Count
-
-#define UE_PRIVATE_SLOG_CALL(F, A) UE_PRIVATE_SLOG_EXPAND(F A)
-#define UE_PRIVATE_SLOG_EXPAND(X) X
+#define UE_PRIVATE_LOGFMT_NAMED_FIELD(Name, ...) UE_PRIVATE_LOGFMT_NAMED_FIELD ::UE::Logging::Private::CheckFieldName(Name), __VA_ARGS__
+// The next three macros remove UE_PRIVATE_LOGFMT_NAMED_FIELD from the expanded expression.
+#define UE_PRIVATE_LOGFMT_FIELD_EXPAND(...) UE_PRIVATE_LOGFMT_FIELD_EXPAND_INNER(__VA_ARGS__)
+#define UE_PRIVATE_LOGFMT_FIELD_EXPAND_INNER(...) UE_PRIVATE_LOGFMT_STRIP_ ## __VA_ARGS__
+#define UE_PRIVATE_LOGFMT_STRIP_UE_PRIVATE_LOGFMT_NAMED_FIELD
+
+// This macro expands `Arg1, Arg2` to `UE_PRIVATE_LOGFMT_FIELD(Arg1), UE_PRIVATE_LOGFMT_FIELD(Arg2), ...`
+// This macro expands `("Name1", Arg1), ("Name2", Arg2)` to `UE_PRIVATE_LOGFMT_FIELD(("Name1", Arg1)), UE_PRIVATE_LOGFMT_FIELD(("Name2", Arg2)), ...
+#define UE_PRIVATE_LOGFMT_FIELDS(...) UE_PRIVATE_LOGFMT_CALL(PREPROCESSOR_JOIN(UE_PRIVATE_LOGFMT_FIELDS_, UE_PRIVATE_LOGFMT_COUNT(__VA_ARGS__)), (__VA_ARGS__))
+
+#define UE_PRIVATE_LOGFMT_FIELDS_0()
+#define UE_PRIVATE_LOGFMT_FIELDS_1(A)                 , UE_PRIVATE_LOGFMT_FIELD(A)
+#define UE_PRIVATE_LOGFMT_FIELDS_2(A,B)               , UE_PRIVATE_LOGFMT_FIELD(A), UE_PRIVATE_LOGFMT_FIELD(B)
+#define UE_PRIVATE_LOGFMT_FIELDS_3(A,B,C)             , UE_PRIVATE_LOGFMT_FIELD(A), UE_PRIVATE_LOGFMT_FIELD(B), UE_PRIVATE_LOGFMT_FIELD(C)
+#define UE_PRIVATE_LOGFMT_FIELDS_4(A,B,C,D)           , UE_PRIVATE_LOGFMT_FIELD(A), UE_PRIVATE_LOGFMT_FIELD(B), UE_PRIVATE_LOGFMT_FIELD(C), UE_PRIVATE_LOGFMT_FIELD(D)
+#define UE_PRIVATE_LOGFMT_FIELDS_5(A,B,C,D,E)         , UE_PRIVATE_LOGFMT_FIELD(A), UE_PRIVATE_LOGFMT_FIELD(B), UE_PRIVATE_LOGFMT_FIELD(C), UE_PRIVATE_LOGFMT_FIELD(D), UE_PRIVATE_LOGFMT_FIELD(E)
+#define UE_PRIVATE_LOGFMT_FIELDS_6(A,B,C,D,E,F)       , UE_PRIVATE_LOGFMT_FIELD(A), UE_PRIVATE_LOGFMT_FIELD(B), UE_PRIVATE_LOGFMT_FIELD(C), UE_PRIVATE_LOGFMT_FIELD(D), UE_PRIVATE_LOGFMT_FIELD(E), UE_PRIVATE_LOGFMT_FIELD(F)
+#define UE_PRIVATE_LOGFMT_FIELDS_7(A,B,C,D,E,F,G)     , UE_PRIVATE_LOGFMT_FIELD(A), UE_PRIVATE_LOGFMT_FIELD(B), UE_PRIVATE_LOGFMT_FIELD(C), UE_PRIVATE_LOGFMT_FIELD(D), UE_PRIVATE_LOGFMT_FIELD(E), UE_PRIVATE_LOGFMT_FIELD(F), UE_PRIVATE_LOGFMT_FIELD(G)
+#define UE_PRIVATE_LOGFMT_FIELDS_8(A,B,C,D,E,F,G,H)   , UE_PRIVATE_LOGFMT_FIELD(A), UE_PRIVATE_LOGFMT_FIELD(B), UE_PRIVATE_LOGFMT_FIELD(C), UE_PRIVATE_LOGFMT_FIELD(D), UE_PRIVATE_LOGFMT_FIELD(E), UE_PRIVATE_LOGFMT_FIELD(F), UE_PRIVATE_LOGFMT_FIELD(G), UE_PRIVATE_LOGFMT_FIELD(H)
+#define UE_PRIVATE_LOGFMT_FIELDS_9(A,B,C,D,E,F,G,H,I) , UE_PRIVATE_LOGFMT_FIELD(A), UE_PRIVATE_LOGFMT_FIELD(B), UE_PRIVATE_LOGFMT_FIELD(C), UE_PRIVATE_LOGFMT_FIELD(D), UE_PRIVATE_LOGFMT_FIELD(E), UE_PRIVATE_LOGFMT_FIELD(F), UE_PRIVATE_LOGFMT_FIELD(G), UE_PRIVATE_LOGFMT_FIELD(H), UE_PRIVATE_LOGFMT_FIELD(I)
+
+#define UE_PRIVATE_LOGFMT_FIELDS_10(A,B,C,D,E,F,G,H,I,J)             , UE_PRIVATE_LOGFMT_FIELD(A), UE_PRIVATE_LOGFMT_FIELD(B), UE_PRIVATE_LOGFMT_FIELD(C), UE_PRIVATE_LOGFMT_FIELD(D), UE_PRIVATE_LOGFMT_FIELD(E), UE_PRIVATE_LOGFMT_FIELD(F), UE_PRIVATE_LOGFMT_FIELD(G), UE_PRIVATE_LOGFMT_FIELD(H), UE_PRIVATE_LOGFMT_FIELD(I), UE_PRIVATE_LOGFMT_FIELD(J)
+#define UE_PRIVATE_LOGFMT_FIELDS_11(A,B,C,D,E,F,G,H,I,J,K)           , UE_PRIVATE_LOGFMT_FIELD(A), UE_PRIVATE_LOGFMT_FIELD(B), UE_PRIVATE_LOGFMT_FIELD(C), UE_PRIVATE_LOGFMT_FIELD(D), UE_PRIVATE_LOGFMT_FIELD(E), UE_PRIVATE_LOGFMT_FIELD(F), UE_PRIVATE_LOGFMT_FIELD(G), UE_PRIVATE_LOGFMT_FIELD(H), UE_PRIVATE_LOGFMT_FIELD(I), UE_PRIVATE_LOGFMT_FIELD(J), UE_PRIVATE_LOGFMT_FIELD(K)
+#define UE_PRIVATE_LOGFMT_FIELDS_12(A,B,C,D,E,F,G,H,I,J,K,L)         , UE_PRIVATE_LOGFMT_FIELD(A), UE_PRIVATE_LOGFMT_FIELD(B), UE_PRIVATE_LOGFMT_FIELD(C), UE_PRIVATE_LOGFMT_FIELD(D), UE_PRIVATE_LOGFMT_FIELD(E), UE_PRIVATE_LOGFMT_FIELD(F), UE_PRIVATE_LOGFMT_FIELD(G), UE_PRIVATE_LOGFMT_FIELD(H), UE_PRIVATE_LOGFMT_FIELD(I), UE_PRIVATE_LOGFMT_FIELD(J), UE_PRIVATE_LOGFMT_FIELD(K), UE_PRIVATE_LOGFMT_FIELD(L)
+#define UE_PRIVATE_LOGFMT_FIELDS_13(A,B,C,D,E,F,G,H,I,J,K,L,M)       , UE_PRIVATE_LOGFMT_FIELD(A), UE_PRIVATE_LOGFMT_FIELD(B), UE_PRIVATE_LOGFMT_FIELD(C), UE_PRIVATE_LOGFMT_FIELD(D), UE_PRIVATE_LOGFMT_FIELD(E), UE_PRIVATE_LOGFMT_FIELD(F), UE_PRIVATE_LOGFMT_FIELD(G), UE_PRIVATE_LOGFMT_FIELD(H), UE_PRIVATE_LOGFMT_FIELD(I), UE_PRIVATE_LOGFMT_FIELD(J), UE_PRIVATE_LOGFMT_FIELD(K), UE_PRIVATE_LOGFMT_FIELD(L), UE_PRIVATE_LOGFMT_FIELD(M)
+#define UE_PRIVATE_LOGFMT_FIELDS_14(A,B,C,D,E,F,G,H,I,J,K,L,M,N)     , UE_PRIVATE_LOGFMT_FIELD(A), UE_PRIVATE_LOGFMT_FIELD(B), UE_PRIVATE_LOGFMT_FIELD(C), UE_PRIVATE_LOGFMT_FIELD(D), UE_PRIVATE_LOGFMT_FIELD(E), UE_PRIVATE_LOGFMT_FIELD(F), UE_PRIVATE_LOGFMT_FIELD(G), UE_PRIVATE_LOGFMT_FIELD(H), UE_PRIVATE_LOGFMT_FIELD(I), UE_PRIVATE_LOGFMT_FIELD(J), UE_PRIVATE_LOGFMT_FIELD(K), UE_PRIVATE_LOGFMT_FIELD(L), UE_PRIVATE_LOGFMT_FIELD(M), UE_PRIVATE_LOGFMT_FIELD(N)
+#define UE_PRIVATE_LOGFMT_FIELDS_15(A,B,C,D,E,F,G,H,I,J,K,L,M,N,O)   , UE_PRIVATE_LOGFMT_FIELD(A), UE_PRIVATE_LOGFMT_FIELD(B), UE_PRIVATE_LOGFMT_FIELD(C), UE_PRIVATE_LOGFMT_FIELD(D), UE_PRIVATE_LOGFMT_FIELD(E), UE_PRIVATE_LOGFMT_FIELD(F), UE_PRIVATE_LOGFMT_FIELD(G), UE_PRIVATE_LOGFMT_FIELD(H), UE_PRIVATE_LOGFMT_FIELD(I), UE_PRIVATE_LOGFMT_FIELD(J), UE_PRIVATE_LOGFMT_FIELD(K), UE_PRIVATE_LOGFMT_FIELD(L), UE_PRIVATE_LOGFMT_FIELD(M), UE_PRIVATE_LOGFMT_FIELD(N), UE_PRIVATE_LOGFMT_FIELD(O)
+#define UE_PRIVATE_LOGFMT_FIELDS_16(A,B,C,D,E,F,G,H,I,J,K,L,M,N,O,P) , UE_PRIVATE_LOGFMT_FIELD(A), UE_PRIVATE_LOGFMT_FIELD(B), UE_PRIVATE_LOGFMT_FIELD(C), UE_PRIVATE_LOGFMT_FIELD(D), UE_PRIVATE_LOGFMT_FIELD(E), UE_PRIVATE_LOGFMT_FIELD(F), UE_PRIVATE_LOGFMT_FIELD(G), UE_PRIVATE_LOGFMT_FIELD(H), UE_PRIVATE_LOGFMT_FIELD(I), UE_PRIVATE_LOGFMT_FIELD(J), UE_PRIVATE_LOGFMT_FIELD(K), UE_PRIVATE_LOGFMT_FIELD(L), UE_PRIVATE_LOGFMT_FIELD(M), UE_PRIVATE_LOGFMT_FIELD(N), UE_PRIVATE_LOGFMT_FIELD(O), UE_PRIVATE_LOGFMT_FIELD(P)
+
+#define UE_PRIVATE_LOGFMT_COUNT(...) UE_PRIVATE_LOGFMT_CALL(UE_PRIVATE_LOGFMT_COUNT_IMPL, (_, ##__VA_ARGS__, 16,15,14,13,12,11,10,9,8,7,6,5,4,3,2,1,0))
+#define UE_PRIVATE_LOGFMT_COUNT_IMPL(_, A,B,C,D,E,F,G,H,I,J,K,L,M,N,O,P, Count, ...) Count
+
+#define UE_PRIVATE_LOGFMT_CALL(F, A) UE_PRIVATE_LOGFMT_EXPAND(F A)
+#define UE_PRIVATE_LOGFMT_EXPAND(X) X
 
 #endif // #if !NO_LOGGING
 
