Enhance memory pressure enum on Android to support 3 stages - OK, low, critical

diff --git a/Engine/Source/Runtime/Core/Private/Android/AndroidPlatformMemory.cpp b/Engine/Source/Runtime/Core/Private/Android/AndroidPlatformMemory.cpp
--- a/Engine/Source/Runtime/Core/Private/Android/AndroidPlatformMemory.cpp
+++ b/Engine/Source/Runtime/Core/Private/Android/AndroidPlatformMemory.cpp
@@ -381,6 +381,7 @@ FGenericPlatformMemoryStats::EMemoryPressureStatus FPlatformMemoryStats::GetMemo
 		case MEMORYADVICE_STATE_OK:
 			return FGenericPlatformMemoryStats::EMemoryPressureStatus::Nominal;
 		case MEMORYADVICE_STATE_APPROACHING_LIMIT:
+			return FGenericPlatformMemoryStats::EMemoryPressureStatus::Warning;
 		case MEMORYADVICE_STATE_CRITICAL:
 			return FGenericPlatformMemoryStats::EMemoryPressureStatus::Critical;
 		default:
@@ -400,10 +401,11 @@ FGenericPlatformMemoryStats::EMemoryPressureStatus FPlatformMemoryStats::GetMemo
 		case FAndroidPlatformMemory::ETrimValues::Unknown:
 			return FGenericPlatformMemoryStats::EMemoryPressureStatus::Unknown;
 		case FAndroidPlatformMemory::ETrimValues::Complete:
+		case FAndroidPlatformMemory::ETrimValues::Running_Low:
+			return FGenericPlatformMemoryStats::EMemoryPressureStatus::Warning;
 		case FAndroidPlatformMemory::ETrimValues::Moderate:
 		case FAndroidPlatformMemory::ETrimValues::Background:
 		case FAndroidPlatformMemory::ETrimValues::UI_Hidden:
-		case FAndroidPlatformMemory::ETrimValues::Running_Low:
 		case FAndroidPlatformMemory::ETrimValues::Running_Moderate:
 		default:
 			return FGenericPlatformMemoryStats::EMemoryPressureStatus::Nominal;
diff --git a/Engine/Source/Runtime/Core/Public/GenericPlatform/GenericPlatformMemory.h b/Engine/Source/Runtime/Core/Public/GenericPlatform/GenericPlatformMemory.h
--- a/Engine/Source/Runtime/Core/Public/GenericPlatform/GenericPlatformMemory.h
+++ b/Engine/Source/Runtime/Core/Public/GenericPlatform/GenericPlatformMemory.h
@@ -156,7 +156,8 @@ struct CORE_API FGenericPlatformMemoryStats : public FPlatformMemoryConstants
 	enum class EMemoryPressureStatus : uint8 
 	{ 
 		Unknown,
-		Nominal, 
+		Nominal,
+		Warning,
 		Critical, // high risk of OOM conditions
 	};
 	EMemoryPressureStatus GetMemoryPressureStatus() const;
