Added Fwd.h headers for WeakObjectPtrs.

diff --git a/Engine/Source/Runtime/Core/Public/UObject/WeakObjectPtrTemplates.h b/Engine/Source/Runtime/Core/Public/UObject/WeakObjectPtrTemplates.h
--- a/Engine/Source/Runtime/Core/Public/UObject/WeakObjectPtrTemplates.h
+++ b/Engine/Source/Runtime/Core/Public/UObject/WeakObjectPtrTemplates.h
@@ -8,22 +8,10 @@
 #include "Templates/AndOrNot.h"
 #include "Templates/LosesQualifiersFromTo.h"
 #include "Containers/Map.h"
+#include "UObject/WeakObjectPtrTemplatesFwd.h"
 
 #include <type_traits>
 
-/**
- * FWeakObjectPtr is a weak pointer to a UObject. 
- * It can return nullptr later if the object is garbage collected.
- * It has no impact on if the object is garbage collected or not.
- * It can't be directly used across a network.
- *
- * Most often it is used when you explicitly do NOT want to prevent something from being garbage collected.
- */
-struct FWeakObjectPtr;
-
-template<class T=UObject, class TWeakObjectPtrBase=FWeakObjectPtr>
-struct TWeakObjectPtr;
-
 /**
  * TWeakObjectPtr is the templated version of the generic FWeakObjectPtr
  */
@@ -320,10 +308,6 @@ FORCEINLINE TWeakObjectPtr<T> MakeWeakObjectPtr(T* Ptr)
 	return TWeakObjectPtr<T>(Ptr);
 }
 
-template<class T> struct TIsPODType<TWeakObjectPtr<T> > { enum { Value = true }; };
-template<class T> struct TIsZeroConstructType<TWeakObjectPtr<T> > { enum { Value = true }; };
-template<class T> struct TIsWeakPointerType<TWeakObjectPtr<T> > { enum { Value = true }; };
-
 
 /**
  * SetKeyFuncs for TWeakObjectPtrs which allow the key to become stale without invalidating the set.
diff --git a/Engine/Source/Runtime/Core/Public/UObject/WeakObjectPtrTemplatesFwd.h b/Engine/Source/Runtime/Core/Public/UObject/WeakObjectPtrTemplatesFwd.h
new file mode 100644
--- /dev/null
+++ b/Engine/Source/Runtime/Core/Public/UObject/WeakObjectPtrTemplatesFwd.h
@@ -0,0 +1,26 @@
+// Copyright Epic Games, Inc. All Rights Reserved.
+
+#pragma once
+
+class UObject;
+
+/**
+ * FWeakObjectPtr is a weak pointer to a UObject. 
+ * It can return nullptr later if the object is garbage collected.
+ * It has no impact on if the object is garbage collected or not.
+ * It can't be directly used across a network.
+ *
+ * Most often it is used when you explicitly do NOT want to prevent something from being garbage collected.
+ */
+struct FWeakObjectPtr;
+
+template<class T=UObject, class TWeakObjectPtrBase=FWeakObjectPtr>
+struct TWeakObjectPtr;
+
+template <typename T> struct TIsPODType;
+template <typename T> struct TIsWeakPointerType;
+template <typename T> struct TIsZeroConstructType;
+
+template<class T> struct TIsPODType<TWeakObjectPtr<T> > { enum { Value = true }; };
+template<class T> struct TIsZeroConstructType<TWeakObjectPtr<T> > { enum { Value = true }; };
+template<class T> struct TIsWeakPointerType<TWeakObjectPtr<T> > { enum { Value = true }; };
diff --git a/Engine/Source/Runtime/CoreUObject/Public/UObject/WeakObjectPtr.h b/Engine/Source/Runtime/CoreUObject/Public/UObject/WeakObjectPtr.h
--- a/Engine/Source/Runtime/CoreUObject/Public/UObject/WeakObjectPtr.h
+++ b/Engine/Source/Runtime/CoreUObject/Public/UObject/WeakObjectPtr.h
@@ -11,14 +11,12 @@
 #include "Templates/UnrealTemplate.h"
 #include "UObject/ScriptDelegates.h"
 #include "UObject/UObjectArray.h"
+#include "UObject/WeakObjectPtrFwd.h"
 
 #include <type_traits>
 
 class FArchive;
 class UObject;
-template <typename T> struct TIsPODType;
-template <typename T> struct TIsWeakPointerType;
-template <typename T> struct TIsZeroConstructType;
 
 /** Invalid FWeakObjectPtr ObjectIndex values must be 0 to support zeroed initialization (this used to be INDEX_NONE, leading to subtle bugs). */
 #ifndef UE_WEAKOBJECTPTR_ZEROINIT_FIX
@@ -325,14 +323,6 @@ private:
 #endif
 };
 
-template<> struct TIsPODType<FWeakObjectPtr> { enum { Value = true }; };
-template<> struct TIsZeroConstructType<FWeakObjectPtr> { enum { Value = true }; };
-template<> struct TIsWeakPointerType<FWeakObjectPtr> { enum { Value = true }; };
-
-// Typedef script delegates for convenience.
-typedef TScriptDelegate<> FScriptDelegate;
-typedef TMulticastScriptDelegate<> FMulticastScriptDelegate;
-
 /** Hash function. */
 FORCEINLINE uint32 GetTypeHash(const FWeakObjectPtr& WeakObjectPtr)
 {
diff --git a/Engine/Source/Runtime/CoreUObject/Public/UObject/WeakObjectPtrFwd.h b/Engine/Source/Runtime/CoreUObject/Public/UObject/WeakObjectPtrFwd.h
new file mode 100644
--- /dev/null
+++ b/Engine/Source/Runtime/CoreUObject/Public/UObject/WeakObjectPtrFwd.h
@@ -0,0 +1,23 @@
+// Copyright Epic Games, Inc. All Rights Reserved.
+
+/*=============================================================================
+	WeakObjectPtrFwd.h: Weak object pointer forward declarations
+=============================================================================*/
+
+#pragma once
+
+#include "UObject/WeakObjectPtrTemplatesFwd.h"
+
+struct FWeakObjectPtr;
+
+template <typename T> struct TIsPODType;
+template <typename T> struct TIsWeakPointerType;
+template <typename T> struct TIsZeroConstructType;
+
+template<> struct TIsPODType<FWeakObjectPtr> { enum { Value = true }; };
+template<> struct TIsZeroConstructType<FWeakObjectPtr> { enum { Value = true }; };
+template<> struct TIsWeakPointerType<FWeakObjectPtr> { enum { Value = true }; };
+
+// Typedef script delegates for convenience.
+typedef TScriptDelegate<> FScriptDelegate;
+typedef TMulticastScriptDelegate<> FMulticastScriptDelegate;
