IoStore: Error checking when reading compressed data from DDC

diff --git a/Engine/Source/Runtime/Core/Private/IO/IoStore.cpp b/Engine/Source/Runtime/Core/Private/IO/IoStore.cpp
--- a/Engine/Source/Runtime/Core/Private/IO/IoStore.cpp
+++ b/Engine/Source/Runtime/Core/Private/IO/IoStore.cpp
@@ -1446,22 +1446,32 @@ private:
 		}
 	}
 
-	void SerializeCompressedDDCData(FIoStoreWriteQueueEntry* Entry, FArchive& Ar)
+	bool SerializeCompressedDDCData(FIoStoreWriteQueueEntry* Entry, FArchive& Ar)
 	{
 		uint32 NumBlocks = Entry->ChunkBlocks.Num();
 		Ar << NumBlocks;
-		check(NumBlocks == Entry->ChunkBlocks.Num());
+		if (NumBlocks != Entry->ChunkBlocks.Num())
+		{
+			return false;
+		}
 		for (FChunkBlock& Block : Entry->ChunkBlocks)
 		{
 			Ar << Block.CompressedSize;
-			check(Block.CompressedSize <= Block.UncompressedSize);
+			if (Block.CompressedSize > Block.UncompressedSize)
+			{
+				return false;
+			}
 			if (Ar.IsLoading() && Block.CompressedSize == Block.UncompressedSize)
 			{
 				Block.CompressionMethod = NAME_None;
 			}
-			check(Block.IoBuffer->DataSize() >= Block.CompressedSize);
+			if (Block.IoBuffer->DataSize() < Block.CompressedSize)
+			{
+				return false;
+			}
 			Ar.Serialize(Block.IoBuffer->Data(), Block.CompressedSize);
 		}
+		return !Ar.IsError();
 	}
 
 	void BeginCompress(FIoStoreWriteQueueEntry* Entry)
@@ -1587,11 +1597,15 @@ private:
 
 			TArray<uint8> DDCData;
 			TRACE_CPUPROFILER_EVENT_SCOPE(ReadFromDDC);
+			bool bFoundInDDC = false;
 			if (WriterContext->DDC->GetSynchronous(*Entry->DDCCacheKey, DDCData, Entry->Options.FileName))
 			{
-				WriterContext->CompressionDDCHitCount.IncrementExchange();
 				FLargeMemoryReader DDCDataReader(DDCData.GetData(), DDCData.Num());
-				SerializeCompressedDDCData(Entry, DDCDataReader);
+				bFoundInDDC = SerializeCompressedDDCData(Entry, DDCDataReader);
+			}
+			if (bFoundInDDC)
+			{
+				WriterContext->CompressionDDCHitCount.IncrementExchange();
 				Entry->FinishCompressionBarrier->DispatchSubsequents();
 				return;
 			}
@@ -1989,8 +2003,10 @@ void FIoStoreWriterContextImpl::BeginEncryptionAndSigningThreadFunc()
 				TRACE_CPUPROFILER_EVENT_SCOPE(StoreInDDC);
 				TArray<uint8> DDCData;
 				FMemoryWriter DDCDataWriter(DDCData, true);
-				Entry->Writer->SerializeCompressedDDCData(Entry, DDCDataWriter);
-				DDC->Put(*Entry->DDCCacheKey, DDCData, Entry->Options.FileName);
+				if (Entry->Writer->SerializeCompressedDDCData(Entry, DDCDataWriter))
+				{
+					DDC->Put(*Entry->DDCCacheKey, DDCData, Entry->Options.FileName);
+				}
 			}
 			FinishEncryptionAndSigningQueue.Enqueue(Entry);
 			Entry->Writer->BeginEncryptAndSign(Entry);
