Detect HW CRC instructions support on ARM platforms

diff --git a/Engine/Source/Runtime/Core/Private/Misc/Crc.cpp b/Engine/Source/Runtime/Core/Private/Misc/Crc.cpp
--- a/Engine/Source/Runtime/Core/Private/Misc/Crc.cpp
+++ b/Engine/Source/Runtime/Core/Private/Misc/Crc.cpp
@@ -4,6 +4,14 @@
 #include "Templates/AlignmentTemplates.h"
 #include "Templates/UnrealTemplate.h"
 #include "Misc/ByteSwap.h"
+#include <signal.h>
+#include <setjmp.h>
+#if PLATFORM_ANDROID_ARM64
+#	include <cpu-features.h>
+#elif PLATFORM_LINUXARM64
+#	include <sys/auxv.h>
+#	include <asm/hwcap.h>
+#endif
 
 /** CRC 32 polynomial */
 enum { Crc32Poly = 0x04c11db7 };
@@ -324,6 +332,72 @@ uint32 FCrc::CRCTablesSB8[8][256] =
 	}
 };
 
+static uint32 MemCrc32SW(const void* InData, int32 Length, uint32 CRC/*=0 */);
+
+#if defined(__clang__) && PLATFORM_CPU_ARM_FAMILY
+static uint32 MemCrc32ArmHW(const void* InData, int32 Length, uint32 CRC/*=0 */);
+#endif
+
+#if DETECT_HW_CRC32_SUPPORT_IN_RUNTIME
+FCrc::MemCrc32Functor FCrc::MemCrc32Func = &MemCrc32SW;
+
+static sigjmp_buf CrcHwJump;
+static void CrcHwSignalHandler(int Sig)
+{
+	siglongjmp(CrcHwJump, Sig);
+}
+
+struct CRCRuntimeHWSupportDetector
+{
+	CRCRuntimeHWSupportDetector()
+	{
+#if PLATFORM_ANDROID_ARM64
+		if (android_getCpuFeatures() & ANDROID_CPU_ARM64_FEATURE_CRC32)
+		{
+			FCrc::MemCrc32Func = &MemCrc32ArmHW;
+			FPlatformMisc::LowLevelOutputDebugString(TEXT("Hardware CRC support detected!\n"));
+		}
+#elif PLATFORM_LINUXARM64
+		if (getauxval(AT_HWCAP) & HWCAP_CRC32)
+		{
+			FCrc::MemCrc32Func = &MemCrc32ArmHW;
+			FPlatformMisc::LowLevelOutputDebugString(TEXT("Hardware CRC support detected!\n"));
+		}
+#else
+		sigset_t SignalSet;
+		sigemptyset(&SignalSet);
+		sigaddset(&SignalSet, SIGILL);
+
+		struct sigaction SigAction;
+		struct sigaction OldSigAction;
+		sigset_t OldSignalSet;
+		memset(&SigAction, 0, sizeof(SigAction));
+		SigAction.sa_handler = CrcHwSignalHandler;
+		SigAction.sa_mask = SignalSet;
+
+		sigprocmask(SIG_SETMASK, &SigAction.sa_mask, &OldSignalSet);
+		sigaction(SIGILL, &SigAction, &OldSigAction);
+
+		if (sigsetjmp(CrcHwJump, 1) == 0)
+		{
+			FCrc::MemCrc32Func = &MemCrc32ArmHW;
+			FCrc::MemCrc32(&SigAction, sizeof(SigAction));
+			FPlatformMisc::LowLevelOutputDebugString(TEXT("Hardware CRC support detected!\n"));
+		}
+		else
+		{
+			FCrc::MemCrc32Func = &MemCrc32SW;
+		}
+
+		sigaction(SIGILL, &OldSigAction, nullptr);
+		sigprocmask(SIG_SETMASK, &OldSignalSet, nullptr);
+#endif
+	}
+};
+
+static CRCRuntimeHWSupportDetector CRCRuntimeHWSupportDetect;
+#endif
+
 void FCrc::Init()
 {
 #if !UE_BUILD_SHIPPING
@@ -384,20 +458,16 @@ void FCrc::Init()
 #endif // !UE_BUILD_SHIPPING
 }
 
-#if PLATFORM_HAS_CRC_INTRINSICS
-// Need to include the header file with the intrinsic definitions per platform until conformity comes
-#if PLATFORM_SWITCH
-#include <arm_acle.h>
-#endif // PLATFORM_SWITCH
-
-uint32 FCrc::MemCrc32(const void* InData, int32 Length, uint32 CRC/*=0 */)
+#if defined(__clang__) && PLATFORM_CPU_ARM_FAMILY
+__attribute__((target("crc")))
+static uint32 MemCrc32ArmHW(const void* InData, int32 Length, uint32 CRC/*=0 */)
 {
 	CRC = ~CRC;
 
 	const uint8_t* __restrict Data = reinterpret_cast<const uint8_t*>(InData);
 
-	// First we need to align to 32-bits
-	int32 InitBytes = Align(Data, 4) - Data;
+	// First we need to align to 64-bits
+	int32 InitBytes = Align(Data, 8) - Data;
 
 	if (Length > InitBytes)
 	{
@@ -405,14 +475,28 @@ uint32 FCrc::MemCrc32(const void* InData, int32 Length, uint32 CRC/*=0 */)
 
 		for (; InitBytes; --InitBytes)
 		{
-			CRC = (CRC >> 8) ^ __crc32b(CRC & 0xFF, *Data++);
+			CRC = (CRC >> 8) ^ __builtin_arm_crc32b(CRC & 0xFF, *Data++);
 		}
 
-		auto Data8 = reinterpret_cast<const uint64_t*>(Data);
+		const uint64_t* Data8 = reinterpret_cast<const uint64_t*>(Data);
 
+		while (Length >= 64)
+		{
+			CRC = __builtin_arm_crc32d(CRC, *Data8++);
+			CRC = __builtin_arm_crc32d(CRC, *Data8++);
+			CRC = __builtin_arm_crc32d(CRC, *Data8++);
+			CRC = __builtin_arm_crc32d(CRC, *Data8++);
+			CRC = __builtin_arm_crc32d(CRC, *Data8++);
+			CRC = __builtin_arm_crc32d(CRC, *Data8++);
+			CRC = __builtin_arm_crc32d(CRC, *Data8++);
+			CRC = __builtin_arm_crc32d(CRC, *Data8++);
+
+			Length -= 64;
+		}
+		
 		for (uint32 Repeat = Length / 8; Repeat; --Repeat)
 		{
-			CRC = __crc32d(CRC, *Data8++);
+			CRC = __builtin_arm_crc32d(CRC, *Data8++);
 		}
 
 		Data = reinterpret_cast<const uint8_t*>(Data8);
@@ -422,15 +506,14 @@ uint32 FCrc::MemCrc32(const void* InData, int32 Length, uint32 CRC/*=0 */)
 
 	for (; Length; --Length)
 	{
-		CRC = (CRC >> 8) ^ __crc32b(CRC & 0xFF, *Data++);
+		CRC = (CRC >> 8) ^ __builtin_arm_crc32b(CRC & 0xFF, *Data++);
 	}
 
 	return ~CRC;
 }
+#endif
 
-#else
-
-uint32 FCrc::MemCrc32( const void* InData, int32 Length, uint32 CRC/*=0 */ )
+static uint32 MemCrc32SW(const void* InData, int32 Length, uint32 CRC/*=0 */)
 {
 	// Based on the Slicing-by-8 implementation found here:
 	// http://slicing-by-8.sourceforge.net/
@@ -448,7 +531,7 @@ uint32 FCrc::MemCrc32( const void* InData, int32 Length, uint32 CRC/*=0 */ )
 
 		for (; InitBytes; --InitBytes)
 		{
-			CRC = (CRC >> 8) ^ CRCTablesSB8[0][(CRC & 0xFF) ^ *Data++];
+			CRC = (CRC >> 8) ^ FCrc::CRCTablesSB8[0][(CRC & 0xFF) ^ *Data++];
 		}
 
 		auto Data4 = (const uint32*)Data;
@@ -457,14 +540,14 @@ uint32 FCrc::MemCrc32( const void* InData, int32 Length, uint32 CRC/*=0 */ )
 			uint32 V1 = *Data4++ ^ CRC;
 			uint32 V2 = *Data4++;
 			CRC =
-				CRCTablesSB8[7][ V1         & 0xFF] ^
-				CRCTablesSB8[6][(V1 >> 8)   & 0xFF] ^
-				CRCTablesSB8[5][(V1 >> 16)  & 0xFF] ^
-				CRCTablesSB8[4][ V1 >> 24         ] ^
-				CRCTablesSB8[3][ V2         & 0xFF] ^
-				CRCTablesSB8[2][(V2 >> 8)   & 0xFF] ^
-				CRCTablesSB8[1][(V2 >> 16)  & 0xFF] ^
-				CRCTablesSB8[0][ V2 >> 24         ];
+				FCrc::CRCTablesSB8[7][ V1         & 0xFF] ^
+				FCrc::CRCTablesSB8[6][(V1 >> 8)   & 0xFF] ^
+				FCrc::CRCTablesSB8[5][(V1 >> 16)  & 0xFF] ^
+				FCrc::CRCTablesSB8[4][ V1 >> 24         ] ^
+				FCrc::CRCTablesSB8[3][ V2         & 0xFF] ^
+				FCrc::CRCTablesSB8[2][(V2 >> 8)   & 0xFF] ^
+				FCrc::CRCTablesSB8[1][(V2 >> 16)  & 0xFF] ^
+				FCrc::CRCTablesSB8[0][ V2 >> 24         ];
 		}
 		Data = (const uint8*)Data4;
 
@@ -473,11 +556,21 @@ uint32 FCrc::MemCrc32( const void* InData, int32 Length, uint32 CRC/*=0 */ )
 
 	for (; Length; --Length)
 	{
-		CRC = (CRC >> 8) ^ CRCTablesSB8[0][(CRC & 0xFF) ^ *Data++];
+		CRC = (CRC >> 8) ^ FCrc::CRCTablesSB8[0][(CRC & 0xFF) ^ *Data++];
 	}
 
 	return ~CRC;
 }
+
+#if !DETECT_HW_CRC32_SUPPORT_IN_RUNTIME
+uint32 FCrc::MemCrc32(const void* InData, int32 Length, uint32 CRC/*=0 */)
+{
+#if PLATFORM_HAS_CRC_INTRINSICS && PLATFORM_CPU_ARM_FAMILY
+	return MemCrc32ArmHW(InData, Length, CRC);
+#else
+	return MemCrc32SW(InData, Length, CRC);
+#endif
+}
 #endif
 
 uint32 FCrc::MemCrc_DEPRECATED(const void* InData, int32 Length, uint32 CRC/*=0 */)
diff --git a/Engine/Source/Runtime/Core/Public/Misc/Crc.h b/Engine/Source/Runtime/Core/Public/Misc/Crc.h
--- a/Engine/Source/Runtime/Core/Public/Misc/Crc.h
+++ b/Engine/Source/Runtime/Core/Public/Misc/Crc.h
@@ -9,6 +9,11 @@
 #include "Misc/Char.h"
 #include "Templates/EnableIf.h"
 
+#if PLATFORM_ANDROID_ARM64 || PLATFORM_LINUXARM64
+#	define DETECT_HW_CRC32_SUPPORT_IN_RUNTIME 1
+#else
+#	define DETECT_HW_CRC32_SUPPORT_IN_RUNTIME 0
+#endif
 
 /** 
  * CRC hash generation for different types of input data
@@ -23,7 +28,16 @@ struct CORE_API FCrc
 	static void Init();
 
 	/** generates CRC hash of the memory area */
+#if !DETECT_HW_CRC32_SUPPORT_IN_RUNTIME
 	static uint32 MemCrc32( const void* Data, int32 Length, uint32 CRC=0 );
+#else
+	typedef uint32 (*MemCrc32Functor)( const void* Data, int32 Length, uint32 CRC );
+	static MemCrc32Functor MemCrc32Func;
+	static FORCEINLINE uint32 MemCrc32(const void* Data, int32 Length, uint32 CRC = 0)
+	{
+		return MemCrc32Func(Data, Length, CRC);
+	}
+#endif
 
 	/** generates CRC hash of the element */
 	template <typename T>
