Added ToUnixTimestampDecimal which returns a double instead of int64 so our timings can be more accurate.
Various tweaks and improvements to the Analytics module. Removed some duplicate interfaces and added some defaults.
Added FNAME parameter to StartCapture instead of having the names as explicit functions ie. removed StartPlayCapture/StartRefreshCapture.
Now storing FlowName in the capture data so that scopes out of flow can be identified.

diff --git a/Engine/Source/Runtime/Core/Public/Misc/DateTime.h b/Engine/Source/Runtime/Core/Public/Misc/DateTime.h
--- a/Engine/Source/Runtime/Core/Public/Misc/DateTime.h
+++ b/Engine/Source/Runtime/Core/Public/Misc/DateTime.h
@@ -566,6 +566,17 @@ public:
 		return (Ticks - FDateTime(1970, 1, 1).Ticks) / ETimespan::TicksPerSecond;
 	}
 
+	/**
+	 * Returns this date as the number of seconds since the Unix Epoch (January 1st of 1970).
+	 *
+	 * @return Time of day.
+	 * @see FromUnixTimestamp
+	 */
+	CORE_API double ToUnixTimestampDecimal() const
+	{
+		return double(Ticks - FDateTime(1970, 1, 1).Ticks) / ETimespan::TicksPerSecond;
+	}
+
 public:
 
 	/**
