Added MemAdvice CVar change callback

diff --git a/Engine/Source/Runtime/Core/Private/Android/AndroidPlatformMemory.cpp b/Engine/Source/Runtime/Core/Private/Android/AndroidPlatformMemory.cpp
--- a/Engine/Source/Runtime/Core/Private/Android/AndroidPlatformMemory.cpp
+++ b/Engine/Source/Runtime/Core/Private/Android/AndroidPlatformMemory.cpp
@@ -17,6 +17,10 @@
 #include "Async/TaskGraphInterfaces.h"
 #include "Async/Async.h"
 
+#define JNI_CURRENT_VERSION JNI_VERSION_1_6
+extern JavaVM* GJavaVM;
+extern jobject GGameActivityThis;
+
 #if HAS_ANDROID_MEMORY_ADVICE
 #include "Containers/Ticker.h"
 #include "memory_advice/memory_advice.h"
@@ -72,11 +76,31 @@ static bool MemoryAdvisorTick(float dt)
 
 	return true;
 }
-#endif
 
-#define JNI_CURRENT_VERSION JNI_VERSION_1_6
-extern JavaVM* GJavaVM;
-extern jobject GGameActivityThis;
+static void OnCVarAndroidUseMemoryAdvisorChanged(IConsoleVariable* Var)
+{
+	if (GAndroidUseMemoryAdvisor && !GMemoryAdvisorInitialized)
+	{
+		JNIEnv* Env = NULL;
+		GJavaVM->GetEnv((void**)&Env, JNI_CURRENT_VERSION);
+		const MemoryAdvice_ErrorCode Result = MemoryAdvice_init(Env, GGameActivityThis);
+		GMemoryAdvisorInitialized = Result == MEMORYADVICE_ERROR_OK;
+		if (GMemoryAdvisorInitialized)
+		{
+			FTSTicker& Ticker = FTSTicker::GetCoreTicker();
+			Ticker.AddTicker(FTickerDelegate::CreateStatic(&MemoryAdvisorTick), 1.0f);
+		}
+		else
+		{
+			UE_LOG(LogInit, Warning, TEXT("Cannot initialize memory advice API, error code %d"), Result);
+		}
+	}
+	else if (GMemoryAdvisorInitialized && !GAndroidUseMemoryAdvisor)
+	{
+		UE_LOG(LogInit, Warning, TEXT("Cannot disable memory advisor once it has been initialized."));
+	}
+}
+#endif
 
 static int32 GAndroidAddSwapToTotalPhysical = 1;
 static FAutoConsoleVariableRef CVarAddSwapToTotalPhysical(
@@ -132,22 +156,7 @@ void FAndroidPlatformMemory::Init()
 		);
 
 #if HAS_ANDROID_MEMORY_ADVICE
-	if (GAndroidUseMemoryAdvisor)
-	{
-		JNIEnv* Env = NULL;
-		GJavaVM->GetEnv((void**)&Env, JNI_CURRENT_VERSION);
-		const MemoryAdvice_ErrorCode Result = MemoryAdvice_init(Env, GGameActivityThis);
-		GMemoryAdvisorInitialized = Result == MEMORYADVICE_ERROR_OK;
-		if (GMemoryAdvisorInitialized)
-		{
-			FTSTicker& Ticker = FTSTicker::GetCoreTicker();
-			Ticker.AddTicker(FTickerDelegate::CreateStatic(&MemoryAdvisorTick), 1.0f);
-		}
-		else
-		{
-			UE_LOG(LogInit, Warning, TEXT("Cannot initialize memory advice API, error code %d"), Result);
-		}
-	}
+	CVarAndroidUseMemoryAdvisor->SetOnChangedCallback(FConsoleVariableDelegate::CreateStatic(&OnCVarAndroidUseMemoryAdvisorChanged));
 #endif
 }
 
