Fixed compilation in Crc.cpp on a platforms that don't support signals handling

diff --git a/Engine/Source/Runtime/Core/Private/Misc/Crc.cpp b/Engine/Source/Runtime/Core/Private/Misc/Crc.cpp
--- a/Engine/Source/Runtime/Core/Private/Misc/Crc.cpp
+++ b/Engine/Source/Runtime/Core/Private/Misc/Crc.cpp
@@ -4,13 +4,15 @@
 #include "Templates/AlignmentTemplates.h"
 #include "Templates/UnrealTemplate.h"
 #include "Misc/ByteSwap.h"
-#include <signal.h>
-#include <setjmp.h>
-#if PLATFORM_ANDROID_ARM64
-#	include <cpu-features.h>
-#elif PLATFORM_LINUXARM64
-#	include <sys/auxv.h>
-#	include <asm/hwcap.h>
+#if DETECT_HW_CRC32_SUPPORT_IN_RUNTIME
+#	include <signal.h>
+#	include <setjmp.h>
+#	if PLATFORM_ANDROID_ARM64
+#		include <cpu-features.h>
+#	elif PLATFORM_LINUXARM64
+#		include <sys/auxv.h>
+#		include <asm/hwcap.h>
+#	endif
 #endif
 
 /** CRC 32 polynomial */
@@ -493,7 +495,7 @@ static uint32 MemCrc32ArmHW(const void* InData, int32 Length, uint32 CRC/*=0 */)
 
 			Length -= 64;
 		}
-		
+
 		for (uint32 Repeat = Length / 8; Repeat; --Repeat)
 		{
 			CRC = __builtin_arm_crc32d(CRC, *Data8++);
