Implement HashCombineFast similar to boost::hash_combine.
Noite: any code that was using HashCombineFast for persistent hashes should revert to HashCombine.

diff --git a/Engine/Source/Runtime/Core/Public/Templates/TypeHash.h b/Engine/Source/Runtime/Core/Public/Templates/TypeHash.h
--- a/Engine/Source/Runtime/Core/Public/Templates/TypeHash.h
+++ b/Engine/Source/Runtime/Core/Public/Templates/TypeHash.h
@@ -61,9 +61,7 @@ inline uint32 HashCombine(uint32 A, uint32 C)
  */
 inline uint32 HashCombineFast(uint32 A, uint32 B)
 {
-	// Currently call HashCombine because it is exists, but only as a placeholder until it
-	// can be replaced with something better
-	return HashCombine(A, B);
+	return A ^ (B + 0x9e3779b9 + (A << 6) + (A >> 2));
 }
 
 inline uint32 PointerHash(const void* Key)
