[Core]
* SortedMap - added so KeyType and ValueType are typedefs inside the class. This make TSortedMap more compatible with TMap

diff --git a/Engine/Source/Runtime/Core/Public/Containers/SortedMap.h b/Engine/Source/Runtime/Core/Public/Containers/SortedMap.h
--- a/Engine/Source/Runtime/Core/Public/Containers/SortedMap.h
+++ b/Engine/Source/Runtime/Core/Public/Containers/SortedMap.h
@@ -15,13 +15,15 @@
  * removing elements is O(n), and finding is O(Log n). In practice it is faster than TMap for low element
  * counts, and slower as n increases, This map is always kept sorted by the key type so cannot be sorted manually.
  */
-template <typename KeyType, typename ValueType, typename ArrayAllocator /*= FDefaultAllocator*/, typename SortPredicate /*= TLess<KeyType>*/ >
+template <typename InKeyType, typename InValueType, typename ArrayAllocator /*= FDefaultAllocator*/, typename SortPredicate /*= TLess<KeyType>*/ >
 class TSortedMap
 {
 	template <typename OtherKeyType, typename OtherValueType, typename OtherArrayAllocator, typename OtherSortPredicate>
 	friend class TSortedMap;
 
 public:
+	typedef InKeyType      KeyType;
+	typedef InValueType    ValueType;
 	typedef typename TTypeTraits<KeyType  >::ConstPointerType KeyConstPointerType;
 	typedef typename TTypeTraits<KeyType  >::ConstInitType    KeyInitType;
 	typedef typename TTypeTraits<ValueType>::ConstInitType    ValueInitType;
