[Core]
ContainersFwd.h - renamed forward declaration TSortedMap template arguments to match implementation template arguments.

Fixes hololens compile errors

diff --git a/Engine/Source/Editor/MovieSceneTools/Public/MovieSceneEventCustomization.h b/Engine/Source/Editor/MovieSceneTools/Public/MovieSceneEventCustomization.h
--- a/Engine/Source/Editor/MovieSceneTools/Public/MovieSceneEventCustomization.h
+++ b/Engine/Source/Editor/MovieSceneTools/Public/MovieSceneEventCustomization.h
@@ -27,7 +27,7 @@ class IDetailChildrenBuilder;
 
 enum class ECheckBoxState : uint8;
 
-template <typename KeyType, typename ValueType, typename ArrayAllocator, typename SortPredicate> class TSortedMap;
+template <typename InKeyType, typename InValueType, typename ArrayAllocator, typename SortPredicate> class TSortedMap;
 
 enum class EAutoCreatePayload : uint8
 {
diff --git a/Engine/Source/Runtime/Core/Public/Containers/ContainersFwd.h b/Engine/Source/Runtime/Core/Public/Containers/ContainersFwd.h
--- a/Engine/Source/Runtime/Core/Public/Containers/ContainersFwd.h
+++ b/Engine/Source/Runtime/Core/Public/Containers/ContainersFwd.h
@@ -28,7 +28,7 @@ template<typename InKeyType, typename InValueType, typename SetAllocator = FDefa
 template<typename KeyType, typename ValueType, typename SetAllocator = FDefaultSetAllocator, typename KeyFuncs = TDefaultMapHashableKeyFuncs<KeyType, ValueType, true > > class TMultiMap;
 template <typename T = void > struct TLess;
 template <typename> struct TTypeTraits;
-template<typename KeyType, typename ValueType, typename ArrayAllocator = FDefaultAllocator, typename SortPredicate = TLess<typename TTypeTraits<KeyType>::ConstPointerType> > class TSortedMap;
+template<typename InKeyType, typename InValueType, typename ArrayAllocator = FDefaultAllocator, typename SortPredicate = TLess<typename TTypeTraits<InKeyType>::ConstPointerType> > class TSortedMap;
 template<typename ElementType,bool bInAllowDuplicateKeys = false> struct DefaultKeyFuncs;
 template<typename InElementType, typename KeyFuncs = DefaultKeyFuncs<InElementType>, typename Allocator = FDefaultSetAllocator> class TSet;
 template<typename InElementType, typename InSizeType = int32> class TStridedView;
