PR #9581: Fix inability to search for OR, AND, or NOT in the content browser. (Contributed by [REDACTED_PERSON])

diff --git a/Engine/Source/Runtime/Core/Private/Misc/TextFilterExpressionEvaluator.cpp b/Engine/Source/Runtime/Core/Private/Misc/TextFilterExpressionEvaluator.cpp
--- a/Engine/Source/Runtime/Core/Private/Misc/TextFilterExpressionEvaluator.cpp
+++ b/Engine/Source/Runtime/Core/Private/Misc/TextFilterExpressionEvaluator.cpp
@@ -29,9 +29,9 @@ namespace TextFilterExpressionParser
 	const TCHAR* FGreater::Monikers[]					= { TEXT(">") };
 	const TCHAR* FGreaterOrEqual::Monikers[]			= { TEXT(">="), TEXT(">:") };
 
-	const TCHAR* FOr::Monikers[]						= { TEXT("OR"), TEXT("||"), TEXT("|") };
-	const TCHAR* FAnd::Monikers[]						= { TEXT("AND"), TEXT("&&"), TEXT("&") };
-	const TCHAR* FNot::Monikers[]						= { TEXT("NOT"), TEXT("!") };
+	const TCHAR* FOr::Monikers[]						= { TEXT("OR "), TEXT("OR("), TEXT("||"), TEXT("|") };
+	const TCHAR* FAnd::Monikers[]						= { TEXT("AND "), TEXT("AND("), TEXT("&&"), TEXT("&") };
+	const TCHAR* FNot::Monikers[]						= { TEXT("NOT "), TEXT("NOT("), TEXT("!") };
 
 	/** Consume an operator from the specified consumer's stream, if one exists at the current read position */
 	template<typename TSymbol>
diff --git a/Engine/Source/Runtime/Core/Public/Misc/TextFilterExpressionEvaluator.h b/Engine/Source/Runtime/Core/Public/Misc/TextFilterExpressionEvaluator.h
--- a/Engine/Source/Runtime/Core/Public/Misc/TextFilterExpressionEvaluator.h
+++ b/Engine/Source/Runtime/Core/Public/Misc/TextFilterExpressionEvaluator.h
@@ -182,9 +182,9 @@ DEFINE_TEXT_EXPRESSION_OPERATOR_NODE(FLessOrEqual,				2,	0x8C0A46B0, 0x8DAA4E2B,
 DEFINE_TEXT_EXPRESSION_OPERATOR_NODE(FGreater,					1,	0x6A6247F4, 0xFD78467F, 0xA6AC1244, 0x0A31E0A5)
 DEFINE_TEXT_EXPRESSION_OPERATOR_NODE(FGreaterOrEqual,			2,	0x09D75C5E, 0xA29A4194, 0x8E8E5278, 0xC84365FD)
 
-DEFINE_TEXT_EXPRESSION_OPERATOR_NODE(FOr,						3,	0xF4778B51, 0xF535414D, 0x9C0EB5F2, 0x2F2B0FD4)
-DEFINE_TEXT_EXPRESSION_OPERATOR_NODE(FAnd,						3,	0x7511397A, 0x02D24DC2, 0x86729800, 0xF454C320)
-DEFINE_TEXT_EXPRESSION_OPERATOR_NODE(FNot,						2,	0x03D78990, 0x41D04E26, 0x8E98AD2F, 0x74667868)
+DEFINE_TEXT_EXPRESSION_OPERATOR_NODE(FOr,						4,	0xF4778B51, 0xF535414D, 0x9C0EB5F2, 0x2F2B0FD4)
+DEFINE_TEXT_EXPRESSION_OPERATOR_NODE(FAnd,						4,	0x7511397A, 0x02D24DC2, 0x86729800, 0xF454C320)
+DEFINE_TEXT_EXPRESSION_OPERATOR_NODE(FNot,						3,	0x03D78990, 0x41D04E26, 0x8E98AD2F, 0x74667868)
 
 DEFINE_EXPRESSION_NODE_TYPE(TextFilterExpressionParser::FTextToken,	0x09E49538, 0x633545E3, 0x84B5644F, 0x1F11628F);
 DEFINE_EXPRESSION_NODE_TYPE(TextFilterExpressionParser::FFunction,	0x084E6214, 0x032FFA48, 0x9245ABF1, 0x9E248F1A);
