Improved deprecation comments surrounding the use of string conversion types and macros.

diff --git a/Engine/Source/Runtime/Core/Public/Containers/StringConv.h b/Engine/Source/Runtime/Core/Public/Containers/StringConv.h
--- a/Engine/Source/Runtime/Core/Public/Containers/StringConv.h
+++ b/Engine/Source/Runtime/Core/Public/Containers/StringConv.h
@@ -951,15 +951,53 @@ struct TElementType<TStringPointer<FromType, ToType>>
  *		const char* SomePointer = TCHAR_TO_ANSI(SomeUnicodeString); <--- Bad!!!
  */
 
-// These should be replaced with StringCasts when FPlatformString starts to know about UTF-8.
-typedef TStringConversion<UE::Core::Private::FTCHARToUTF8_Convert> FTCHARToUTF8;
-typedef TStringConversion<FUTF8ToTCHAR_Convert> FUTF8ToTCHAR;
-
-// Usage of these should be replaced with StringCasts.
-#define TCHAR_TO_ANSI(str) (ANSICHAR*)StringCast<ANSICHAR>(static_cast<const TCHAR*>(str)).Get()
-#define ANSI_TO_TCHAR(str) (TCHAR*)StringCast<TCHAR>(static_cast<const ANSICHAR*>(str)).Get()
-#define TCHAR_TO_UTF8(str) (ANSICHAR*)FTCHARToUTF8((const TCHAR*)str).Get()
-#define UTF8_TO_TCHAR(str) (TCHAR*)FUTF8ToTCHAR((const ANSICHAR*)str).Get()
+/*********************************************************************
+ * NOTE:
+ * The following converting types and macros should be considered
+ * deprecated and StringCast should be used instead.
+ *
+ * The reasons for deprecation include:
+ *
+ * - char is the wrong type for representing a UTF-8 encoding.  UTF8CHAR should be used now, which will become char8_t from C++20 onwards.
+ * - The macros hide a reinterpret_cast of the argument, which is both misleading and dangerous.
+ * - The const return value is silently cast away, which is also misleading and dangerous.
+ * - The macros are frequently used to initialize a local variable, which ends up pointing at deleted memory.
+ * - The names are hardcoded with the types they cast from and to, making them unusable in generic contexts.
+ *
+ * Examples of migration:
+ *   ANSICHAR* AnsiCharPtr = ...;
+ *   UTF8CHAR* Utf8CharPtr = ...;
+ *   TCHAR*    TCharPtr    = ...;
+ *
+ * Before:
+ *   FUTF8ToTCHAR Conv1(AnsiCharPtr); // Use Conv1.Get() and Conv1.Length()
+ *   FUTF8ToTCHAR Conv2(Utf8CharPtr); // Use Conv2.Get() and Conv2.Length()
+ *   FunctionTakingTCHARPtr(ANSI_TO_TCHAR(AnsiCharPtr));
+ *   Printf(TEXT("Narrow string: %s"), ANSI_TO_TCHAR(AnsiCharPtr));
+ *   Printf(TEXT("UTF-8 string: %s"), UTF8_TO_TCHAR(Utf8CharPtr));
+ *
+ * After:
+ *   auto Conv1 = StringCast<TCHAR>((const UTF8CHAR*)AnsiCharPtr); // Use Conv1.Get() and Conv1.Length()
+ *   auto Conv2 = StringCast<TCHAR>(UTF8CharPtr);                  // Use Conv2.Get() and Conv2.Length()
+ *   FunctionTakingTCHARPtr(StringCast<TCHAR>(AnsiCharPtr).Get());
+ *   Printf(TEXT("Narrow string: %hs"), AnsiCharPtr);
+ *   Printf(TEXT("Narrow string: %hs"), Utf8CharPtr);
+ *********************************************************************/
+using FTCHARToUTF8 /*UE_DEPRECATED(5.xx, "FTCHARToUTF8(PtrToTChar) is deprecated, please use StringCast<UTF8CHAR>(PtrToTChar) instead.")*/ = TStringConversion<UE::Core::Private::FTCHARToUTF8_Convert>;
+using FUTF8ToTCHAR /*UE_DEPRECATED(5.xx, "FUTF8ToTCHAR(PtrToUTF8) is deprecated, please use StringCast<TCHAR>(PtrToUTF8Char) instead.")*/ = TStringConversion<FUTF8ToTCHAR_Convert>;
+
+#define TCHAR_TO_ANSI(str) /*UE_DEPRECATED_MACRO(5.xx, "TCHAR_TO_ANSI(Ptr) is deprecated, please use StringCast<ANSICHAR>(PtrToTChar) instead.")*/ (ANSICHAR*)StringCast<ANSICHAR>(static_cast<const TCHAR*>(str)).Get()
+#define ANSI_TO_TCHAR(str) /*UE_DEPRECATED_MACRO(5.xx, "ANSI_TO_TCHAR(Ptr) is deprecated, please use StringCast<TCHAR>(PtrToAnsiChar) instead.")*/ (TCHAR*)StringCast<TCHAR>(static_cast<const ANSICHAR*>(str)).Get()
+#define TCHAR_TO_UTF8(str) /*UE_DEPRECATED_MACRO(5.xx, "TCHAR_TO_UTF8(Ptr) is deprecated, please use StringCast<UTF8CHAR>(PtrToTChar) instead.")*/ (ANSICHAR*)FTCHARToUTF8((const TCHAR*)str).Get()
+#define UTF8_TO_TCHAR(str) /*UE_DEPRECATED_MACRO(5.xx, "UTF8_TO_TCHAR(Ptr) is deprecated, please use StringCast<TCHAR>(PtrToUTF8Char) instead.")*/ (TCHAR*)FUTF8ToTCHAR((const ANSICHAR*)str).Get()
+
+
+////////////////////////////////////////////////////////////////////
+// NOTE:                                                          //
+// The following converting types and macros cannot be deprecated //
+// until FPlatformString has proper handling of UTF16CHAR and     //
+// UTF32CHAR.                                                     //
+////////////////////////////////////////////////////////////////////
 
 // special handling for platforms still using a 32-bit TCHAR
 #if PLATFORM_TCHAR_IS_4_BYTES
@@ -1047,6 +1085,21 @@ typedef TStringPointer<wchar_t, TCHAR> FWCharToTCHAR;
 /**
  * Creates an object which acts as a source of a given string type.  See example above.
  *
+ * StringCast expects correctly-typed strings.  If a cast is attempted with a char* or ANSICHAR*,
+ * and the string contains characters that are non-ASCII, including UTF-8 code units outside of
+ * the 7-bit ASCII range, then those values will fail to convert and a bogus char will be written
+ * in their place.
+ *
+ * If a conversion from UTF-8 is desired, the pointer should be cast to UTF8CHAR* before being
+ * passed to StringCast.
+ *
+ * Similarly, doing a StringCast<char> or StringCast<ANSICHAR> on a Unicode string will only
+ * successfully convert Unicode characters which already lie in the ASCII range.  For converting
+ * to UTF-8, StringCast<UTF8CHAR> should be used.
+ *
+ * The source string must not be modified or destroyed until after the result of this function
+ * has been destroyed.
+ *
  * @param Str The null-terminated source string to convert.
  */
 template <typename To, int32 DefaultConversionSize = DEFAULT_STRING_CONVERSION_SIZE, typename From>
@@ -1065,6 +1118,21 @@ FORCEINLINE auto StringCast(const From* Str)
 /**
  * Creates an object which acts as a source of a given string type.  See example above.
  *
+ * StringCast expects correctly-typed strings.  If a cast is attempted with a char* or ANSICHAR*,
+ * and the string contains characters that are non-ASCII, including UTF-8 code units outside of
+ * the 7-bit ASCII range, then those values will fail to convert and a bogus char will be written
+ * in their place.
+ *
+ * If a conversion from UTF-8 is desired, the pointer should be cast to UTF8CHAR* before being
+ * passed to StringCast.
+ *
+ * Similarly, doing a StringCast<char> or StringCast<ANSICHAR> on a Unicode string will only
+ * successfully convert Unicode characters which already lie in the ASCII range.  For converting
+ * to UTF-8, StringCast<UTF8CHAR> should be used.
+ *
+ * The source string must not be modified or destroyed until after the result of this function
+ * has been destroyed.
+ *
  * @param Str The source string to convert, not necessarily null-terminated.
  * @param Len The number of From elements in Str.
  */
