TRangeAllocator now computes usage stats

diff --git a/Engine/Source/Runtime/Core/Public/Misc/RangeAllocator.h b/Engine/Source/Runtime/Core/Public/Misc/RangeAllocator.h
--- a/Engine/Source/Runtime/Core/Public/Misc/RangeAllocator.h
+++ b/Engine/Source/Runtime/Core/Public/Misc/RangeAllocator.h
@@ -10,6 +10,8 @@
 #include "HAL/CriticalSection.h"
 #include "Misc/ScopeLock.h"
 
+#define RANGE_ALLOCATOR_RECORD_STATS !UE_BUILD_SHIPPING
+
 /*
  * This allocator has a list of chunks (fixed size committed memory regions) and each
  * chunk keeps a list of free ranges. Allocations are served by finding a large enough
@@ -260,6 +262,9 @@ protected:
 
 	void FreeChunk(uint32 ChunkIndex)
 	{
+#if RANGE_ALLOCATOR_RECORD_STATS
+		SizeTotal -= ChunkSize;
+#endif
 		const int32 InfoIndex = Chunks[ChunkIndex].InfoIndex;
 		if (InfoIndex != INDEX_NONE)
 		{
@@ -270,6 +275,11 @@ protected:
 		FreeChunkSlotHead = ChunkIndex;
 	}
 
+#if RANGE_ALLOCATOR_RECORD_STATS
+	SIZE_T SizeTotal = 0; // Total amount allocated in bytes including slack
+	SIZE_T SizeUsed = 0;  // Size of active allocations in bytes
+#endif
+
 	const uint16 MinAllocSize;
 	int32 FreeChunkSlotHead;
 	TArray<FFreeChunkInfo> FreeChunkInfos;
@@ -316,6 +326,9 @@ public:
 				FChunk& Chunk = Chunks[Info.ChunkIndex];
 				SIZE_T Ptr = Chunk.Alloc(Size, Alignment, Info.MaxFreeRangeSize, OutAllocInfo.RangeOffset, OutAllocInfo.RangeSize);
 				Check(Ptr != 0);
+#if RANGE_ALLOCATOR_RECORD_STATS
+				SizeUsed += OutAllocInfo.RangeSize * MinAlignment;
+#endif
 				OutAllocInfo.ChunkIndex = Info.ChunkIndex;
 				OutAllocInfo.InitCustomInfo(Chunk, Ptr);
 				if (Info.MaxFreeRangeSize < MinAllocSize)
@@ -340,6 +353,9 @@ public:
 			SIZE_T Ptr = Chunk.Alloc(Size, Alignment, Info.MaxFreeRangeSize, OutAllocInfo.RangeOffset, OutAllocInfo.RangeSize);
 			if (Ptr != 0)
 			{
+#if RANGE_ALLOCATOR_RECORD_STATS
+				SizeUsed += OutAllocInfo.RangeSize * MinAlignment;
+#endif
 				OutAllocInfo.ChunkIndex = Info.ChunkIndex;
 				OutAllocInfo.InitCustomInfo(Chunk, Ptr);
 				if (Info.MaxFreeRangeSize < MinAllocSize)
@@ -368,6 +384,10 @@ public:
 		new (&Chunk) FChunk(*this);
 		SIZE_T Ptr = Chunk.Alloc(Size, Alignment, MaxFreeRangeSize, OutAllocInfo.RangeOffset, OutAllocInfo.RangeSize);
 		Check(Ptr != 0);
+#if RANGE_ALLOCATOR_RECORD_STATS
+		SizeTotal += ChunkSize;
+		SizeUsed += OutAllocInfo.RangeSize * MinAlignment;
+#endif
 		OutAllocInfo.ChunkIndex = ChunkIndex;
 		OutAllocInfo.InitCustomInfo(Chunk, Ptr);
 		if (MaxFreeRangeSize >= MinAllocSize)
@@ -385,7 +405,9 @@ public:
 		uint16 NewMaxFreeRangeSize = InfoIndex == INDEX_NONE ? 0 : FreeChunkInfos[InfoIndex].MaxFreeRangeSize;
 
 		Chunks[AllocInfo.ChunkIndex].Free(AllocInfo.RangeOffset, AllocInfo.RangeSize, NewMaxFreeRangeSize);
-
+#if RANGE_ALLOCATOR_RECORD_STATS
+		SizeUsed -= AllocInfo.RangeSize * MinAlignment;
+#endif
 		if (NewMaxFreeRangeSize == ChunkSize / MinAlignment)
 		{
 			FreeChunk(AllocInfo.ChunkIndex);
