Fix FName == C string bug and add string view comparisons

FName("_1") == "_1" failed since FName("_[num]") is considered None

diff --git a/Engine/Source/Runtime/Core/Private/UObject/UnrealNames.cpp b/Engine/Source/Runtime/Core/Private/UObject/UnrealNames.cpp
--- a/Engine/Source/Runtime/Core/Private/UObject/UnrealNames.cpp
+++ b/Engine/Source/Runtime/Core/Private/UObject/UnrealNames.cpp
@@ -2651,30 +2651,6 @@ bool FName::IsWithinBounds(FNameEntryId Id)
 	FName implementation.
 -----------------------------------------------------------------------------*/
 
-template<class CharType>
-static bool NumberEqualsString(uint32 Number, const CharType* Str)
-{
-	CharType* End = nullptr;
-	return TCString<CharType>::Strtoi64(Str, &End, 10) == Number && End && *End == '\0';
-}
-
-template<class CharType1, class CharType2>
-static bool StringAndNumberEqualsString(const CharType1* Name, uint32 NameLen, int32 InternalNumber, const CharType2* Str)
-{
-	if (FPlatformString::Strnicmp(Name, Str, NameLen))
-	{
-		return false;
-	}
-
-	if (InternalNumber == NAME_NO_NUMBER_INTERNAL)
-	{
-		return Str[NameLen] == '\0';
-	}
-
-	uint32 Number = NAME_INTERNAL_TO_EXTERNAL(InternalNumber);
-	return Str[NameLen] == '_' && NumberEqualsString(Number, Str + NameLen + 1);
-}
-
 struct FNameAnsiStringView
 {
 	using CharType = ANSICHAR;
@@ -2903,22 +2879,50 @@ struct FNameHelper
 #endif
 	}
 
-	template<class CharType>
-	static bool EqualsString(FName Name, const CharType* Str)
+	template<typename CharType>
+	static bool EqualsString(FName Name, const CharType* Str, int32 Len)
 	{
-		// Make NAME_None == TEXT("") or nullptr consistent with NAME_None == FName(TEXT("")) or FName(nullptr)
-		if (Str == nullptr || Str[0] == '\0')
+		check(Len >= 0);
+
+		// Make NAME_None == TEXTVIEW("") or TEXTVIEW(nullptr) consistent with NAME_None == FName(TEXT("")) or FName(nullptr)
+		if (Len == 0)
 		{
 			return Name.IsNone();
 		}
 
-		const FNameEntry& Entry = *Name.GetComparisonNameEntry();
+		// "_[num]" is considered "None" to keep semantics of old FName implementation
+		if (Name.IsNone() & (Str[0] == '_'))
+		{
+			ParseNumber(Str, /* in-out */ Len);
+			return Len == 0;
+		}
 
-		uint32 NameLen = Entry.Header.Len;
+		const FNameEntry* Entry = FName::ResolveEntry(Name.GetComparisonIndexInternal());
+
+	#if UE_FNAME_OUTLINE_NUMBER
+		int32 Number = Entry->IsNumbered() ? Entry->GetNumber() : NAME_NO_NUMBER_INTERNAL;
+		Entry = Entry->IsNumbered() ? FName::ResolveEntry(Entry->NumberedName.Id) : Entry; 
+		check(!Entry->IsNumbered());
+	#else
+		int32 Number = Name.GetNumber();
+	#endif
+
+		// Check _[number] suffix first and shorten Len if it matches
+		if (Number != NAME_NO_NUMBER_INTERNAL && 
+			Number != ParseNumber(Str, /* in-out */ Len)) 
+		{
+			return false;
+		}
+
+		if (Len != Entry->GetNameLength())
+		{
+			return false;
+		}
+		
 		FNameBuffer Temp;
-		return Entry.IsWide()
-			? StringAndNumberEqualsString(Entry.GetUnterminatedName(Temp.WideName), NameLen, Name.GetNumber(), Str)
-			: StringAndNumberEqualsString(Entry.GetUnterminatedName(Temp.AnsiName), NameLen, Name.GetNumber(), Str);
+		return Entry->IsWide()
+			? !FPlatformString::Strnicmp(Str, Entry->GetUnterminatedName(Temp.WideName), Len)
+			: !FPlatformString::Strnicmp(Str, Entry->GetUnterminatedName(Temp.AnsiName), Len);
 	}
 
 	template<typename CharType>
@@ -3258,14 +3262,24 @@ FName::FName(const FNameEntrySerialized& LoadedEntry)
 	: FName(FNameHelper::MakeFromLoaded(LoadedEntry))
 {}
 
-bool FName::operator==(const ANSICHAR* Str) const
+bool FName::Equals(FName Name, FAnsiStringView View)
+{
+	return FNameHelper::EqualsString(Name, View.GetData(), View.Len());
+}
+
+bool FName::Equals(FName Name, FWideStringView View)
 {
-	return FNameHelper::EqualsString(*this, Str);
+	return FNameHelper::EqualsString(Name, View.GetData(), View.Len());
 }
 
-bool FName::operator==(const WIDECHAR* Str) const
+bool FName::Equals(FName Name, const ANSICHAR* Str)
 {
-	return FNameHelper::EqualsString(*this, Str);
+	return Equals(Name, FAnsiStringView(Str));
+}
+
+bool FName::Equals(FName Name, const WIDECHAR* Str)
+{
+	return Equals(Name, FWideStringView(Str));
 }
 
 int32 FName::Compare( const FName& Other ) const
@@ -3788,6 +3802,13 @@ void FName::AutoTest()
 	check(FName(FName(), NAME_EXTERNAL_TO_INTERNAL(7)) == None_7);
 	check(FName("None", NAME_EXTERNAL_TO_INTERNAL(7)) == None_7);
 	check(FName(FNameEntryId::FromEName(EName::None), FNameEntryId::FromEName(EName::None), NAME_EXTERNAL_TO_INTERNAL(7)) == None_7);
+	check(!None_7.IsNone());
+
+	// "_[num]" is considered None to keep semantics of old implementation, unlike "None_[num]"
+	check(FName("_7").IsNone()); 
+	check(FName("_0").IsNone());
+	check(FName("_2147483646").IsNone());
+	check(!FName("_2147483647").IsNone()); // invalid number
 
 	// Find existing numbered none 
 PRAGMA_DISABLE_DEPRECATION_WARNINGS
@@ -3876,35 +3897,6 @@ PRAGMA_ENABLE_DEPRECATION_WARNINGS
 		Once = false;
 	}
 
-	check(NumberEqualsString(0, "0"));
-	check(NumberEqualsString(11, "11"));
-	check(NumberEqualsString(2147483647, "2147483647"));
-	check(NumberEqualsString(4294967294, "4294967294"));
-
-	check(!NumberEqualsString(0, "1"));
-	check(!NumberEqualsString(1, "0"));
-	check(!NumberEqualsString(11, "12"));
-	check(!NumberEqualsString(12, "11"));
-	check(!NumberEqualsString(2147483647, "2147483646"));
-	check(!NumberEqualsString(2147483646, "2147483647"));
-
-	check(StringAndNumberEqualsString("abc", 3, NAME_EXTERNAL_TO_INTERNAL(10), "abc_10"));
-	check(!StringAndNumberEqualsString("aba", 3, NAME_EXTERNAL_TO_INTERNAL(10), "abc_10"));
-	check(!StringAndNumberEqualsString("abc", 2, NAME_EXTERNAL_TO_INTERNAL(10), "abc_10"));
-	check(!StringAndNumberEqualsString("abc", 2, NAME_EXTERNAL_TO_INTERNAL(11), "abc_10"));
-	check(!StringAndNumberEqualsString("abc", 3, NAME_EXTERNAL_TO_INTERNAL(10), "aba_10"));
-	check(!StringAndNumberEqualsString("abc", 3, NAME_EXTERNAL_TO_INTERNAL(10), "abc_11"));
-	check(!StringAndNumberEqualsString("abc", 3, NAME_EXTERNAL_TO_INTERNAL(10), "abc_100"));
-
-	check(StringAndNumberEqualsString("abc", 3, NAME_EXTERNAL_TO_INTERNAL(0), "abc_0"));
-	check(!StringAndNumberEqualsString("abc", 3, NAME_EXTERNAL_TO_INTERNAL(0), "abc_1"));
-
-	check(StringAndNumberEqualsString("abc", 3, NAME_NO_NUMBER_INTERNAL, "abc"));
-	check(!StringAndNumberEqualsString("abc", 2, NAME_NO_NUMBER_INTERNAL, "abc"));
-	check(!StringAndNumberEqualsString("abc", 3, NAME_NO_NUMBER_INTERNAL, "abcd"));
-	check(!StringAndNumberEqualsString("abc", 3, NAME_NO_NUMBER_INTERNAL, "abc_0"));
-	check(!StringAndNumberEqualsString("abc", 3, NAME_NO_NUMBER_INTERNAL, "abc_"));
-
 	TArray<FName> Names;
 	Names.Add("FooB");
 	Names.Add("FooABCD");
@@ -3923,6 +3915,36 @@ PRAGMA_ENABLE_DEPRECATION_WARNINGS
 	check(Names[4] == "FooB");
 	check(Names[5] == "FooC");
 	check(Names[6] == FooWide);
+		
+	check(FName() == TEXTVIEW("None"));
+	check(FName() == TEXTVIEW("none"));
+	check(FName() == TEXTVIEW(""));
+	check(FName() == FStringView());
+	check(FName("a") == TEXTVIEW("a"));
+	check(FName("a") == TEXTVIEW("A"));
+	check(FName("A") == TEXTVIEW("a"));
+	check(FName("aa") != TEXTVIEW("a"));
+	check(FName("a") != TEXTVIEW("aa"));
+	check(FName("a") == TEXTVIEW("aa").Left(1)); // Unterminated
+	check(FName("aa") != TEXTVIEW("ab"));
+	check(FName("a_0") == TEXTVIEW("a_0"));
+	check(FName("a_1") != TEXTVIEW("a_0"));
+	check(FName("a_10") == TEXTVIEW("a_10"));
+	check(FName("a")    == TEXTVIEW("a_10").Left(1)); // Unterminated number
+	check(FName("a_")   == TEXTVIEW("a_10").Left(2)); // Unterminated number
+	check(FName("a_1")  == TEXTVIEW("a_10").Left(3)); // Unterminated number
+	check(FName("a_2147483646") == TEXTVIEW("a_2147483646")); // Largest valid number
+	check(FName("a_2147483647") == TEXTVIEW("a_2147483647")); // Smallest invalid number
+	check(FName("a_") == TEXTVIEW("a_"));
+	check(FName("_") == TEXTVIEW("_"));
+	check(FName("_3") == TEXTVIEW("_3")); // Special none
+	check(FName() == TEXTVIEW("_123")); // Special none
+	check(FName("_3") == TEXTVIEW("None")); // Special none
+	check(FName("none_1") == TEXTVIEW("none_1")); // Numbered none
+	check(WideName == FStringView(Wide));
+	check(FName("a") == ANSITEXTVIEW("a"));
+	check(FName("a") == ANSITEXTVIEW("A"));
+	check(FName("A") != ANSITEXTVIEW("b"));
 
 	
 	CheckLazyName("Hej");
@@ -3933,6 +3955,7 @@ PRAGMA_ENABLE_DEPRECATION_WARNINGS
 	CheckLazyName("Hej_01");
 	CheckLazyName("Hej_-1");
 	CheckLazyName("Hej__0");
+	CheckLazyName("Hej_2147483646");
 	CheckLazyName("Hej_2147483647");
 	CheckLazyName("Hej_123");
 	CheckLazyName("None");
diff --git a/Engine/Source/Runtime/Core/Public/UObject/NameTypes.h b/Engine/Source/Runtime/Core/Public/UObject/NameTypes.h
--- a/Engine/Source/Runtime/Core/Public/UObject/NameTypes.h
+++ b/Engine/Source/Runtime/Core/Public/UObject/NameTypes.h
@@ -432,7 +432,6 @@ private:
 	{
 		return !operator==(Lhs, Rhs);
 	}
-	friend FORCEINLINE bool operator==(FName Lhs, FMinimalName Rhs);
 };
 
 /**
@@ -497,7 +496,6 @@ private:
 	{
 		return !(Lhs == Rhs);
 	}
-	friend FORCEINLINE bool operator==(FName Lhs, FScriptName Rhs);
 };
 
 struct FMemoryImageName;
@@ -568,7 +566,6 @@ private:
 	{
 		return !(Lhs == Rhs);
 	}
-	friend FORCEINLINE bool operator==(FName Lhs, FMemoryImageName Rhs);
 };
 
 /**
@@ -713,9 +710,6 @@ public:
 		return !(*this == Other);
 	}
 
-	FORCEINLINE bool operator==(EName Ename) const;
-	
-	FORCEINLINE bool operator!=(EName Ename) const;
 
 	UE_DEPRECATED(4.23, "Please use FastLess() / FNameFastLess or LexicalLess() / FNameLexicalLess instead. "
 		"Default lexical sort order is deprecated to avoid unintended expensive sorting. ")
@@ -1076,27 +1070,6 @@ PRAGMA_ENABLE_DEPRECATION_WARNINGS
 	 */
 	CORE_API FName(const FNameEntrySerialized& LoadedEntry);
 
-	/**
-	 * Equality operator.
-	 *
-	 * @param	Other	String to compare this name to
-	 * @return true if name matches the string, false otherwise
-	 */
-	CORE_API bool operator==(const ANSICHAR* Other) const;
-	CORE_API bool operator==(const WIDECHAR* Other) const;
-
-	/**
-	 * Inequality operator.
-	 *
-	 * @param	Other	String to compare this name to
-	 * @return true if name does not match the string, false otherwise
-	 */
-	template <typename CharType>
-	bool operator!=(const CharType* Other) const
-	{
-		return !operator==(Other);
-	}
-
 	CORE_API static void DisplayHash( class FOutputDevice& Ar );
 	CORE_API static FString SafeString(FNameEntryId InDisplayIndex, int32 InstanceNumber = NAME_NO_NUMBER_INTERNAL);
 
@@ -1278,106 +1251,80 @@ private:
 	{
 		return CreateNumberedNameIfNecessary(ComparisonId, ComparisonId, Number);
 	}
+	
+	
+	static bool Equals(FName A, FName B) = delete;
+	CORE_API static bool Equals(FName A, FAnsiStringView B);
+	CORE_API static bool Equals(FName A, FWideStringView B);
+	CORE_API static bool Equals(FName A, const ANSICHAR* B);
+	CORE_API static bool Equals(FName A, const WIDECHAR* B);
 
-	/**
-	 * Equality operator with CharType* on left hand side and FName on right hand side
-	 * 
-	 * @param	LHS		CharType to compare to FName
-	 * @param	RHS		FName to compare to CharType
-	 * @return True if strings match, false otherwise.
-	 */
-	template <typename CharType>
-	friend inline bool operator==(const CharType *LHS, const FName &RHS)
-	{
-		return RHS == LHS;
-	}
-
-	/**
-	 * Inequality operator with CharType* on left hand side and FName on right hand side
-	 *
-	 * @param	LHS		CharType to compare to FName
-	 * @param	RHS		FName to compare to CharType
-	 * @return True if strings don't match, false otherwise.
-	 */
-	template <typename CharType>
-	friend inline bool operator!=(const CharType *LHS, const FName &RHS)
+	FORCEINLINE static bool Equals(FName A, EName B)
 	{
-		return RHS != LHS;
+		return A.GetComparisonIndex() == B && A.GetNumber() == NAME_NO_NUMBER_INTERNAL;
 	}
 
 #if UE_FNAME_OUTLINE_NUMBER
-	friend FORCEINLINE bool operator==(FName Lhs, FMinimalName Rhs)
+	FORCEINLINE static bool Equals(FName A, FMinimalName B)
 	{
-		return Lhs.GetComparisonIndexInternal() == Rhs.Index;
+		return A.GetComparisonIndexInternal() == B.Index;
 	}
-	friend FORCEINLINE bool operator==(FName Lhs, FScriptName Rhs)
+	FORCEINLINE static bool Equals(FName A, FScriptName B)
 	{
-		return Lhs.GetComparisonIndexInternal() == Rhs.ComparisonIndex;
+		return A.GetComparisonIndexInternal() == B.ComparisonIndex;
 	}
-	friend FORCEINLINE bool operator==(FName Lhs, FMemoryImageName Rhs)
+	FORCEINLINE static bool Equals(FName A, FMemoryImageName B)
 	{
-		return Lhs.GetComparisonIndexInternal() == Rhs.ComparisonIndex;
+		return A.GetComparisonIndexInternal() == B.ComparisonIndex;
 	}
 	friend FORCEINLINE uint32 GetTypeHash(FName Name)
 	{
 		return GetTypeHash(Name.GetComparisonIndexInternal());
 	}
 #else
-	friend FORCEINLINE bool operator==(FName Lhs, FMinimalName Rhs)
+	FORCEINLINE static bool Equals(FName A, FMinimalName B)
 	{
-		return Lhs.GetComparisonIndex() == Rhs.Index && Lhs.GetNumber() == Rhs.Number;
+		return A.GetComparisonIndex() == B.Index && A.GetNumber() == B.Number;
 	}
-	friend FORCEINLINE bool operator==(FName Lhs, FScriptName Rhs)
+	FORCEINLINE static bool Equals(FName A, FScriptName B)
 	{
-		return Lhs.GetComparisonIndex() == Rhs.ComparisonIndex && Lhs.GetNumber() == Rhs.Number;
+		return A.GetComparisonIndex() == B.ComparisonIndex && A.GetNumber() == B.Number;
 	}
-	friend FORCEINLINE bool operator==(FName Lhs, FMemoryImageName Rhs)
+	FORCEINLINE static bool Equals(FName A, FMemoryImageName B)
 	{
-		return Lhs.GetComparisonIndex() == Rhs.ComparisonIndex && Lhs.GetNumber() == Rhs.Number;
+		return A.GetComparisonIndex() == B.ComparisonIndex && A.GetNumber() == B.Number;
 	}
 	friend FORCEINLINE uint32 GetTypeHash(FName Name)
 	{
 		return GetTypeHash(Name.GetComparisonIndex()) + Name.GetNumber();
 	}
 #endif
-	friend FORCEINLINE bool operator==(FMinimalName Lhs, FName Rhs)
-	{
-		return Rhs == Lhs;
-	}
-	friend FORCEINLINE bool operator!=(FMinimalName Lhs, FName Rhs)
-	{
-		return !(Lhs == Rhs);
-	}
-
-	friend FORCEINLINE bool operator!=(FName Lhs, FMinimalName Rhs)
-	{
-		return !(Lhs == Rhs);
-	}
 
-	friend FORCEINLINE bool operator==(FScriptName Lhs, FName Rhs)
-	{
-		return Rhs == Lhs;
-	}
-	friend FORCEINLINE bool operator!=(FScriptName Lhs, FName Rhs)
-	{
-		return !(Lhs == Rhs);
-	}
-	friend FORCEINLINE bool operator!=(FName Lhs, FScriptName Rhs)
-	{
-		return !(Lhs == Rhs);
-	}
-	friend FORCEINLINE bool operator==(FMemoryImageName Lhs, FName Rhs)
-	{
-		return Rhs == Lhs;
-	}
-	friend FORCEINLINE bool operator!=(FName Lhs, FMemoryImageName Rhs)
-	{
-		return !(Lhs == Rhs);
-	}
-	friend FORCEINLINE bool operator!=(FMemoryImageName Lhs, FName Rhs)
-	{
-		return !(Lhs == Rhs);
-	}
+    template <typename T>
+    friend FORCEINLINE auto operator==(FName N, T O) -> decltype(FName::Equals(N, O))
+    {
+    	return FName::Equals(N, O);
+    }
+
+    template <typename T>
+    friend FORCEINLINE auto operator==(T O, FName N) -> decltype(FName::Equals(N, O))
+    {
+    	return FName::Equals(N, O);
+    }
+
+#if !PLATFORM_COMPILER_HAS_GENERATED_COMPARISON_OPERATORS
+    template <typename T>
+    friend FORCEINLINE auto operator!=(FName N, T O) -> decltype(FName::Equals(N, O))
+    {
+    	return !FName::Equals(N, O);
+    }
+
+    template <typename T>
+    friend FORCEINLINE auto operator!=(T O, FName N) -> decltype(FName::Equals(N, O))
+    {
+    	return !FName::Equals(N, O);
+    }
+#endif
 };
 
 template<> struct TIsZeroConstructType<class FName> { enum { Value = true }; };
@@ -1503,11 +1450,6 @@ FORCEINLINE bool FMemoryImageName::IsNone() const
 	return ComparisonIndex.IsNone();
 }
 
-FORCEINLINE bool FName::operator==(EName Ename) const
-{
-	return GetComparisonIndex() == Ename && GetNumber() == NAME_NO_NUMBER_INTERNAL;
-}
-
 FORCEINLINE bool FName::IsEqual(const FName& Rhs, const ENameCase CompareMethod /*= ENameCase::IgnoreCase*/, const bool bCompareNumber /*= true*/) const
 {
 	return bCompareNumber ?
@@ -1588,10 +1530,6 @@ FORCEINLINE bool FMemoryImageName::IsNone() const
 	return ComparisonIndex.IsNone() && Number == NAME_NO_NUMBER_INTERNAL;
 }
 
-FORCEINLINE bool FName::operator==(EName Ename) const
-{
-	return GetComparisonIndex() == Ename && GetNumber() == NAME_NO_NUMBER_INTERNAL;
-}
 
 FORCEINLINE bool FName::IsEqual(const FName& Rhs, const ENameCase CompareMethod /*= ENameCase::IgnoreCase*/, const bool bCompareNumber /*= true*/) const
 {
@@ -1620,11 +1558,6 @@ FORCEINLINE FScriptName NameToScriptName(FName InName)
 	return FScriptName(InName);
 }
 
-FORCEINLINE bool FName::operator!=(EName Ename) const
-{
-	return !(*this == Ename);
-}
-
 FORCEINLINE FString LexToString(const FName& Name)
 {
 	return Name.ToString();
