Fixed usage of a deprecated iOS function

diff --git a/Engine/Source/Runtime/Core/Private/IOS/IOSPlatformMisc.cpp b/Engine/Source/Runtime/Core/Private/IOS/IOSPlatformMisc.cpp
--- a/Engine/Source/Runtime/Core/Private/IOS/IOSPlatformMisc.cpp
+++ b/Engine/Source/Runtime/Core/Private/IOS/IOSPlatformMisc.cpp
@@ -1684,35 +1684,42 @@ void FIOSPlatformMisc::MetalAssert()
 
 bool FIOSPlatformMisc::CPUHasHwCrcSupport()
 {
-	const EIOSDevice Device = GetIOSDeviceType();
-	if (Device >= IOS_IPhone6 && Device <= IOS_IPhone6SPlus)
+	// HW CRC instructions support is available on Apple A10+
+	static int HwCrcSupported = -1;
+	if (HwCrcSupported != -1)
 	{
-		//Apple A8 and A9 iPhones
-		return false;
-	}
-	else if (Device == IOS_IPadMini4 || Device == IOS_IPodTouch6 || Device == IOS_AppleTV)
-	{
-		//Apple A8
-		return false;
-	}
-	else if (Device == IOS_IPadAir2)
-	{
-		//Apple A8x
-		return false;
-	}
-	else if (Device == IOS_IPhoneSE || Device == IOS_IPad5)
-	{
-		//Apple A9
-		return false;
-	}
-	else if (Device == IOS_IPadPro_97 || Device == IOS_IPadPro_129)
-	{
-		//Apple A9x
-		return false;
+		const FString DeviceIDString = GetIOSDeviceIDString();
+		if (DeviceIDString.StartsWith(TEXT("iPod")))
+		{
+			const int Major = FCString::Atoi(&DeviceIDString[4]);
+			//iPod Touch 6 and lower don't support hw CRC32
+			HwCrcSupported = Major > 7;
+		}
+		else if (DeviceIDString.StartsWith(TEXT("iPad")))
+		{
+			// get major revision number
+			const int Major = FCString::Atoi(&DeviceIDString[4]);
+
+			//iPad 5, iPad Pro and lower
+			HwCrcSupported = Major > 6;
+		}
+		else if (DeviceIDString.StartsWith(TEXT("iPhone")))
+		{
+			const int Major = FCString::Atoi(&DeviceIDString[6]);
+			
+			// iPhone 6S, iPhone SE and below
+			HwCrcSupported = Major > 9;
+		}
+		else if (DeviceIDString.StartsWith(TEXT("AppleTV")))
+		{
+			const int Major = FCString::Atoi(&DeviceIDString[7]);
+			
+			// Apple TV
+			HwCrcSupported = Major > 5;
+		}
 	}
 
-	//We assume that we are running on a supported device that's at least A10 or newer and they support hw CRC
-	return true;
+	return HwCrcSupported == 1;
 }
 
 static FCriticalSection EnsureLock;
