Reduce code bloat from assertion macros and templates

diff --git a/Engine/Source/Runtime/Core/Private/Misc/AssertionMacros.cpp b/Engine/Source/Runtime/Core/Private/Misc/AssertionMacros.cpp
--- a/Engine/Source/Runtime/Core/Private/Misc/AssertionMacros.cpp
+++ b/Engine/Source/Runtime/Core/Private/Misc/AssertionMacros.cpp
@@ -591,6 +591,7 @@ void FORCENOINLINE FDebug::CheckVerifyFailedImpl(
 		AssertFailedImplV(Expr, File, Line, ProgramCounter, Format, Args);
 		va_end(Args);
 	}
+	PLATFORM_BREAK_IF_DESIRED_NONINLINE();
 }
 
 #endif // DO_CHECK || DO_GUARD_SLOW || DO_ENSURE
@@ -622,26 +623,39 @@ void FDebug::ProcessFatalError(void* ProgramCounter)
 
 #if DO_CHECK || DO_GUARD_SLOW || DO_ENSURE
 FORCENOINLINE bool VARARGS FDebug::OptionallyLogFormattedEnsureMessageReturningFalseImpl( bool bLog, const ANSICHAR* Expr, const ANSICHAR* File, int32 Line, void* ProgramCounter, const TCHAR* FormattedMsg, ... )
+{
+	va_list Args;
+	va_start(Args, FormattedMsg);
+	OptionallyLogFormattedEnsureMessageReturningFalseImpl(bLog, Expr, File, Line, ProgramCounter, FormattedMsg, Args);
+	va_end(Args);
+
+	return false;
+}
+
+FORCENOINLINE bool FDebug::OptionallyLogFormattedEnsureMessageReturningFalseImpl(bool bLog, const ANSICHAR* Expr, const ANSICHAR* File, int32 Line, void* ProgramCounter, const TCHAR* FormattedMsg, va_list Args)
 {
 	if (bLog)
 	{
 		const int32 TempStrSize = 4096;
-		TCHAR TempStr[ TempStrSize ];
-		GET_VARARGS( TempStr, TempStrSize, TempStrSize - 1, FormattedMsg, FormattedMsg );
+		TCHAR TempStr[TempStrSize];
+		FCString::GetVarArgs(TempStr, TempStrSize, FormattedMsg, Args);
 
-		EnsureFailed( Expr, File, Line, ProgramCounter, TempStr );
+		EnsureFailed(Expr, File, Line, ProgramCounter, TempStr);
 	}
-	
+
 	return false;
 }
 #endif
 
-FORCENOINLINE void VARARGS LowLevelFatalErrorHandler(const ANSICHAR* File, int32 Line, void* ProgramCounter, const TCHAR* Format, ...)
+FORCENOINLINE void UE_DEBUG_SECTION VARARGS LowLevelFatalErrorHandler(const ANSICHAR* File, int32 Line, void* ProgramCounter, const TCHAR* Format, ...)
 {
 	va_list Args;
 	va_start(Args, Format);
 	StaticFailDebugV(TEXT("LowLevelFatalError"), "", File, Line, /*bIsEnsure*/ false, ProgramCounter, Format, Args);
 	va_end(Args);
+
+	UE_DEBUG_BREAK_AND_PROMPT_FOR_REMOTE();
+	FDebug::ProcessFatalError(ProgramCounter);
 }
 
 void FDebug::DumpStackTraceToLog(const ELogVerbosity::Type LogVerbosity)
@@ -675,4 +689,28 @@ FORCENOINLINE void FDebug::DumpStackTraceToLog(const TCHAR* Heading, const ELogV
 #endif
 }
 
+#if DO_ENSURE && !USING_CODE_ANALYSIS
+bool UE_DEBUG_SECTION VARARGS CheckVerifyImpl(bool& InOutExecuted, bool Always, const ANSICHAR* File, int32 Line, void* ProgramCounter, const ANSICHAR* Expr, const TCHAR* Format, ...)
+{
+	if ((!InOutExecuted || Always) && FPlatformMisc::IsEnsureAllowed())
+	{
+		InOutExecuted = true;
+		va_list Args;
+		va_start(Args, Format);
+		FDebug::OptionallyLogFormattedEnsureMessageReturningFalse(true, Expr, File, Line, ProgramCounter, Format, Args);
+		va_end(Args);
+
+		if (!FPlatformMisc::IsDebuggerPresent())
+		{
+			FPlatformMisc::PromptForRemoteDebugging(true);
+			return false;
+		}
+
+		PLATFORM_BREAK_IF_DESIRED_NONINLINE();
+		return true;
+	}
+	return false;
+}
+#endif
+
 #undef FILE_LINE_DESC_ANSI
diff --git a/Engine/Source/Runtime/Core/Public/Misc/AssertionMacros.h b/Engine/Source/Runtime/Core/Public/Misc/AssertionMacros.h
--- a/Engine/Source/Runtime/Core/Public/Misc/AssertionMacros.h
+++ b/Engine/Source/Runtime/Core/Public/Misc/AssertionMacros.h
@@ -112,6 +112,7 @@ public:
 
 private:
 	static bool VARARGS OptionallyLogFormattedEnsureMessageReturningFalseImpl(bool bLog, const ANSICHAR* Expr, const ANSICHAR* File, int32 Line, void* ProgramCounter, const TCHAR* FormattedMsg, ...);
+	static bool OptionallyLogFormattedEnsureMessageReturningFalseImpl(bool bLog, const ANSICHAR* Expr, const ANSICHAR* File, int32 Line, void* ProgramCounter, const TCHAR* FormattedMsg, va_list Args);
 
 public:
 	/**
@@ -138,6 +139,11 @@ public:
 		return OptionallyLogFormattedEnsureMessageReturningFalseImpl(bLog, Expr, File, Line, ProgramCounter, (const TCHAR*)FormattedMsg, Args...);
 	}
 
+	static FORCEINLINE bool OptionallyLogFormattedEnsureMessageReturningFalse(bool bLog, const ANSICHAR* Expr, const ANSICHAR* File, int32 Line, void* ProgramCounter, const TCHAR* FormattedMsg, va_list Args)
+	{
+		return OptionallyLogFormattedEnsureMessageReturningFalseImpl(bLog, Expr, File, Line, ProgramCounter, FormattedMsg, Args);
+	}
+
 #endif // DO_CHECK || DO_GUARD_SLOW
 
 	/**
@@ -206,6 +212,17 @@ RetType FORCENOINLINE UE_DEBUG_SECTION DispatchCheckVerify(InnerType&& Inner, Ar
 	#define PLATFORM_BREAK_IF_DESIRED() PLATFORM_BREAK();
 #endif // !UE_BUILD_SHIPPING
 
+#if PLATFORM_WINDOWS
+	// We want to break inline of an assertion macro on Windows for ease of debugging
+#	define PLATFORM_BREAK_IF_DESIRED_INLINE() PLATFORM_BREAK_IF_DESIRED()
+#	define PLATFORM_BREAK_IF_DESIRED_NONINLINE()
+#else
+	// ... however on every other platform this causes significant code bloat and we want to avoid it.
+	// So stepping up through callstack in a debugger is a necessary price to pay
+#	define PLATFORM_BREAK_IF_DESIRED_INLINE()
+#	define PLATFORM_BREAK_IF_DESIRED_NONINLINE() PLATFORM_BREAK_IF_DESIRED()
+#endif
+
 
 #if DO_CHECK
 #ifndef checkCode
@@ -225,15 +242,8 @@ RetType FORCENOINLINE UE_DEBUG_SECTION DispatchCheckVerify(InnerType&& Inner, Ar
 		{ \
 			if(UNLIKELY(!(expr))) \
 			{ \
-				struct Impl \
-				{ \
-					static void FORCENOINLINE UE_DEBUG_SECTION ExecCheckImplInternal() \
-					{ \
-						FDebug::CheckVerifyFailedImpl(#expr, __FILE__, __LINE__, PLATFORM_RETURN_ADDRESS(), TEXT("")); \
-					} \
-				}; \
-				Impl::ExecCheckImplInternal(); \
-				PLATFORM_BREAK_IF_DESIRED(); \
+				FDebug::CheckVerifyFailedImpl(#expr, __FILE__, __LINE__, PLATFORM_RETURN_ADDRESS(), TEXT("")); \
+				PLATFORM_BREAK_IF_DESIRED_INLINE(); \
 				CA_ASSUME(false); \
 			} \
 		}
@@ -253,11 +263,8 @@ RetType FORCENOINLINE UE_DEBUG_SECTION DispatchCheckVerify(InnerType&& Inner, Ar
 		{ \
 			if(UNLIKELY(!(expr))) \
 			{ \
-				DispatchCheckVerify([] (const auto& LFormat, const auto&... UE_LOG_Args) UE_DEBUG_SECTION \
-				{ \
-					FDebug::CheckVerifyFailedImpl(#expr, __FILE__, __LINE__, PLATFORM_RETURN_ADDRESS(), LFormat, UE_LOG_Args...); \
-				}, format, ##__VA_ARGS__); \
-				PLATFORM_BREAK_IF_DESIRED(); \
+				FDebug::CheckVerifyFailedImpl(#expr, __FILE__, __LINE__, PLATFORM_RETURN_ADDRESS(), format, ##__VA_ARGS__); \
+				PLATFORM_BREAK_IF_DESIRED_INLINE(); \
 				CA_ASSUME(false); \
 			} \
 		}
@@ -352,29 +359,40 @@ RetType FORCENOINLINE UE_DEBUG_SECTION DispatchCheckVerify(InnerType&& Inner, Ar
  */
 
 #if DO_ENSURE && !USING_CODE_ANALYSIS // The Visual Studio 2013 analyzer doesn't understand these complex conditionals
+	struct FValidateArgsInternal
+	{
+		template <typename... Types>
+		FValidateArgsInternal(Types... Args)
+		{
+			static_assert(TAnd<TIsValidVariadicFunctionArg<Types>...>::Value, "Invalid argument(s) passed to ensureMsgf");
+		}
+	};
 
-	#define UE_ENSURE_IMPL(Capture, Always, InExpression, ...) \
+	bool UE_DEBUG_SECTION VARARGS CheckVerifyImpl(bool& InOutExecuted, bool Always, const ANSICHAR* File, int32 Line, void* ProgramCounter, const ANSICHAR* Expr, const TCHAR* Format, ...);
+
+	#define UE_ENSURE_IMPL(Capture, Always, InExpression, InFormat) \
 		(LIKELY(!!(InExpression)) || (DispatchCheckVerify<bool>([Capture] () UE_DEBUG_SECTION \
 		{ \
 			static bool bExecuted = false; \
-			if ((!bExecuted || Always) && FPlatformMisc::IsEnsureAllowed()) \
-			{ \
-				bExecuted = true; \
-				FDebug::OptionallyLogFormattedEnsureMessageReturningFalse(true, #InExpression, __FILE__, __LINE__, PLATFORM_RETURN_ADDRESS(), ##__VA_ARGS__); \
-				if (!FPlatformMisc::IsDebuggerPresent()) \
-				{ \
-					FPlatformMisc::PromptForRemoteDebugging(true); \
-					return false; \
-				} \
-				return true; \
-			} \
-			return false; \
-		}) && ([] () { PLATFORM_BREAK_IF_DESIRED(); } (), false)))
+			const bool CheckVerifyResult = CheckVerifyImpl(bExecuted, Always, __FILE__, __LINE__, PLATFORM_RETURN_ADDRESS(), #InExpression, InFormat); \
+			PLATFORM_BREAK_IF_DESIRED_INLINE(); \
+			return CheckVerifyResult; \
+		})))
+
+	#define UE_ENSURE_IMPL2(Capture, Always, InExpression, ...) \
+		(LIKELY(!!(InExpression)) || (DispatchCheckVerify<bool>([Capture] () UE_DEBUG_SECTION \
+		{ \
+			static bool bExecuted = false; \
+			FValidateArgsInternal(__VA_ARGS__); \
+			const bool CheckVerifyResult = CheckVerifyImpl(bExecuted, Always, __FILE__, __LINE__, PLATFORM_RETURN_ADDRESS(), #InExpression, ##__VA_ARGS__); \
+			PLATFORM_BREAK_IF_DESIRED_INLINE(); \
+			return CheckVerifyResult; \
+		})))
 
 	#define ensure(           InExpression                ) UE_ENSURE_IMPL( , false, InExpression, TEXT(""))
-	#define ensureMsgf(       InExpression, InFormat, ... ) UE_ENSURE_IMPL(&, false, InExpression, InFormat, ##__VA_ARGS__)
+	#define ensureMsgf(       InExpression, InFormat, ... ) UE_ENSURE_IMPL2(&, false, InExpression, InFormat, ##__VA_ARGS__)
 	#define ensureAlways(     InExpression                ) UE_ENSURE_IMPL( , true,  InExpression, TEXT(""))
-	#define ensureAlwaysMsgf( InExpression, InFormat, ... ) UE_ENSURE_IMPL(&, true,  InExpression, InFormat, ##__VA_ARGS__)
+	#define ensureAlwaysMsgf( InExpression, InFormat, ... ) UE_ENSURE_IMPL2(&, true,  InExpression, InFormat, ##__VA_ARGS__)
 
 #else	// DO_ENSURE
 
@@ -445,17 +463,11 @@ namespace UEAsserts_Private
 ----------------------------------------------------------------------------*/
 
 /** low level fatal error handler. */
-CORE_API void VARARGS LowLevelFatalErrorHandler(const ANSICHAR* File, int32 Line, void* ProgramCounter, const TCHAR* Format=TEXT(""), ... );
+CORE_API void UE_DEBUG_SECTION VARARGS LowLevelFatalErrorHandler(const ANSICHAR* File, int32 Line, void* ProgramCounter, const TCHAR* Format=TEXT(""), ... );
 
 #define LowLevelFatalError(Format, ...) \
 	{ \
 		static_assert(TIsArrayOrRefOfTypeByPredicate<decltype(Format), TIsCharEncodingCompatibleWithTCHAR>::Value, "Formatting string must be a TCHAR array."); \
-		DispatchCheckVerify([] (const auto& LFormat, const auto&... UE_LOG_Args) UE_DEBUG_SECTION \
-		{ \
-			void* ProgramCounter = PLATFORM_RETURN_ADDRESS(); \
-			LowLevelFatalErrorHandler(__FILE__, __LINE__, ProgramCounter, (const TCHAR*)LFormat, UE_LOG_Args...); \
-			UE_DEBUG_BREAK_AND_PROMPT_FOR_REMOTE(); \
-			FDebug::ProcessFatalError(ProgramCounter); \
-		}, Format, ##__VA_ARGS__); \
+		LowLevelFatalErrorHandler(__FILE__, __LINE__, PLATFORM_RETURN_ADDRESS(), (const TCHAR*)Format, ##__VA_ARGS__); \
 	}
 
