Logging: Added constexpr to most of FLogCategory[Base]

diff --git a/Engine/Source/Runtime/Core/Public/Logging/LogCategory.h b/Engine/Source/Runtime/Core/Public/Logging/LogCategory.h
--- a/Engine/Source/Runtime/Core/Public/Logging/LogCategory.h
+++ b/Engine/Source/Runtime/Core/Public/Logging/LogCategory.h
@@ -6,89 +6,84 @@
 #include "Logging/LogVerbosity.h"
 #include "UObject/NameTypes.h"
 
+#define UE_API CORE_API
+
 #if defined(UE_LOG_LAZY_CATEGORY_NAMES) && UE_LOG_LAZY_CATEGORY_NAMES
 using FLogCategoryName = FLazyName;
 #else
 using FLogCategoryName = FName;
 #endif
 
-
-/** Base class for all log categories. **/
-struct CORE_API FLogCategoryBase
+/** Base class for all log categories. */
+struct FLogCategoryBase
 {
 	/**
-	* Constructor, registers with the log suppression system and sets up the default values.
-	* @param CategoryName, name of the category
-	* @param InDefaultVerbosity, default verbosity for the category, may ignored and overrridden by many other mechanisms
-	* @param InCompileTimeVerbosity, mostly used to keep the working verbosity in bounds, macros elsewhere actually do the work of stripping compiled out things.
-	**/
-	FLogCategoryBase(const FLogCategoryName& InCategoryName, ELogVerbosity::Type InDefaultVerbosity, ELogVerbosity::Type InCompileTimeVerbosity);
-
-	/** Destructor, unregisters from the log suppression system **/
-	~FLogCategoryBase();
-	/** Should not generally be used directly. Tests the runtime verbosity and maybe triggers a debug break, etc. **/
-	FORCEINLINE bool IsSuppressed(ELogVerbosity::Type VerbosityLevel) const
-	{
-		if ((VerbosityLevel & ELogVerbosity::VerbosityMask) <= Verbosity)
-		{
-			return false;
-		}
-		return true;
-	}
-	/** Called just after a logging statement being allow to print. Checks a few things and maybe breaks into the debugger. **/
-	void PostTrigger(ELogVerbosity::Type VerbosityLevel);	
+	 * Constructor, registers with the log suppression system and sets up the default values.
+	 *
+	 * @param CategoryName           Name of the category.
+	 * @param DefaultVerbosity       Default verbosity used to filter this category at runtime.
+	 * @param CompileTimeVerbosity   Verbosity used to filter this category at compile time.
+	 */
+	UE_API FLogCategoryBase(const FLogCategoryName& CategoryName, ELogVerbosity::Type DefaultVerbosity, ELogVerbosity::Type CompileTimeVerbosity);
 
-	const FLogCategoryName& GetCategoryName() const
-	{
-		return CategoryName;
-	}
+	/** Destructor, unregisters from the log suppression system. */
+	UE_API ~FLogCategoryBase();
 
-	/** Gets the working verbosity **/
-	ELogVerbosity::Type GetVerbosity() const
+	/** Should not generally be used directly. Tests the runtime verbosity and maybe triggers a debug break, etc. */
+	FORCEINLINE constexpr bool IsSuppressed(ELogVerbosity::Type VerbosityLevel) const
 	{
-		return (ELogVerbosity::Type)Verbosity;
+		return !((VerbosityLevel & ELogVerbosity::VerbosityMask) <= Verbosity);
 	}
+
+	/** Called just after a logging statement being allow to print. Checks a few things and maybe breaks into the debugger. */
+	UE_API void PostTrigger(ELogVerbosity::Type VerbosityLevel);
+
+	inline constexpr const FLogCategoryName& GetCategoryName() const { return CategoryName; }
+
+	/** Gets the working verbosity. */
+	inline constexpr ELogVerbosity::Type GetVerbosity() const { return (ELogVerbosity::Type)Verbosity; }
 	
-	/** Sets up the working verbosity and clamps to the compile time verbosity. **/
-	void SetVerbosity(ELogVerbosity::Type Verbosity);
+	/** Sets up the working verbosity and clamps to the compile time verbosity. */
+	UE_API void SetVerbosity(ELogVerbosity::Type Verbosity);
+
+	/** Gets the compile time verbosity. */
+	inline constexpr ELogVerbosity::Type GetCompileTimeVerbosity() const { return CompileTimeVerbosity; }
 
-	/** Gets the compile time verbosity **/
-	inline ELogVerbosity::Type GetCompileTimeVerbosity() const
-	{
-		return CompileTimeVerbosity;
-	}
 private:
 	friend class FLogSuppressionImplementation;
 	friend class FLogScopedVerbosityOverride;
-	/** Internal call to set up the working verbosity from the boot time default. **/
+
+	/** Internal call to set up the working verbosity from the boot time default. */
 	void ResetFromDefault();
 
-	/** Holds the current suppression state **/
+	/** Holds the current suppression state */
 	ELogVerbosity::Type Verbosity;
-	/** Holds the break flag **/
+	/** Holds the break flag */
 	bool DebugBreakOnLog;
-	/** Holds default suppression **/
+	/** Holds default suppression */
 	uint8 DefaultVerbosity;
-	/** Holds compile time suppression **/
+	/** Holds compile time suppression */
 	const ELogVerbosity::Type CompileTimeVerbosity;
-	/** Name for this category **/
+	/** Name for this category */
 
 	const FLogCategoryName CategoryName;
 };
 
-/** Template for log categories that transfers the compile-time constant default and compile time verbosity to the FLogCategoryBase constructor. **/
-template<uint8 InDefaultVerbosity, uint8 InCompileTimeVerbosity>
+/** Template for log categories that transfers the compile-time constant default and compile time verbosity to the FLogCategoryBase constructor. */
+template <ELogVerbosity::Type InDefaultVerbosity, ELogVerbosity::Type InCompileTimeVerbosity>
 struct FLogCategory : public FLogCategoryBase
 {
 	static_assert((InDefaultVerbosity & ELogVerbosity::VerbosityMask) < ELogVerbosity::NumVerbosity, "Bogus default verbosity.");
 	static_assert(InCompileTimeVerbosity < ELogVerbosity::NumVerbosity, "Bogus compile time verbosity.");
-	enum
-	{
-		CompileTimeVerbosity = InCompileTimeVerbosity
-	};
+
+	static constexpr ELogVerbosity::Type CompileTimeVerbosity = InCompileTimeVerbosity;
+
+	inline constexpr ELogVerbosity::Type GetCompileTimeVerbosity() const { return CompileTimeVerbosity; }
 
 	FORCEINLINE FLogCategory(const FLogCategoryName& InCategoryName)
-		: FLogCategoryBase(InCategoryName, ELogVerbosity::Type(InDefaultVerbosity), ELogVerbosity::Type(CompileTimeVerbosity))
+		: FLogCategoryBase(InCategoryName, InDefaultVerbosity, CompileTimeVerbosity)
 	{
 	}
 };
+
+#undef UE_API
