StringBuilder: Added AppendV

diff --git a/Engine/Source/Runtime/Core/Private/Misc/StringBuilder.cpp b/Engine/Source/Runtime/Core/Private/Misc/StringBuilder.cpp
--- a/Engine/Source/Runtime/Core/Private/Misc/StringBuilder.cpp
+++ b/Engine/Source/Runtime/Core/Private/Misc/StringBuilder.cpp
@@ -5,6 +5,7 @@
 #include "HAL/PlatformMath.h"
 #include "HAL/UnrealMemory.h"
 #include "Misc/CString.h"
+#include "Misc/ScopeExit.h"
 
 static inline uint64_t NextPowerOfTwo(uint64_t x)
 {
@@ -62,29 +63,38 @@ void TStringBuilderBase<C>::FreeBuffer(void* Buffer, SIZE_T CharCount)
 }
 
 template <typename C>
-TStringBuilderBase<C>& TStringBuilderBase<C>::AppendfImpl(BuilderType& Self, const C* Fmt, ...)
+TStringBuilderBase<C>& TStringBuilderBase<C>::AppendV(const C* Fmt, va_list Args)
 {
 	for (;;)
 	{
 		va_list ArgPack;
-		va_start(ArgPack, Fmt);
-		const int32 RemainingSize = (int32)(Self.End - Self.CurPos);
-		const int32 Result = TCString<C>::GetVarArgs(Self.CurPos, RemainingSize, Fmt, ArgPack);
+		va_copy(ArgPack, Args);
+		const int32 RemainingSize = (int32)(End - CurPos);
+		const int32 Result = TCString<C>::GetVarArgs(CurPos, RemainingSize, Fmt, ArgPack);
 		va_end(ArgPack);
 
 		if (Result >= 0 && Result < RemainingSize)
 		{
-			Self.CurPos += Result;
-			return Self;
+			CurPos += Result;
+			return *this;
 		}
 		else
 		{
 			// Total size will be rounded up to the next power of two. Start with at least 64.
-			Self.Extend(64);
+			Extend(64);
 		}
 	}
 }
 
+template <typename C>
+TStringBuilderBase<C>& TStringBuilderBase<C>::AppendfImpl(BuilderType& Self, const C* Fmt, ...)
+{
+	va_list ArgPack;
+	va_start(ArgPack, Fmt);
+	ON_SCOPE_EXIT { va_end(ArgPack); };
+	return Self.AppendV(Fmt, ArgPack);
+}
+
 // Instantiate templates once
 
 template class TStringBuilderBase<ANSICHAR>;
diff --git a/Engine/Source/Runtime/Core/Public/Misc/StringBuilder.h b/Engine/Source/Runtime/Core/Public/Misc/StringBuilder.h
--- a/Engine/Source/Runtime/Core/Public/Misc/StringBuilder.h
+++ b/Engine/Source/Runtime/Core/Public/Misc/StringBuilder.h
@@ -332,16 +332,23 @@ public:
 	/**
 	 * Appends to the string builder similarly to how classic sprintf works.
 	 *
-	 * @param Format A format string that specifies how to format the additional arguments. Refer to standard printf format.
+	 * @param Fmt A format string that specifies how to format the additional arguments. Refer to standard printf format.
 	 */
-	template <typename FmtType, typename... Types,
-		std::enable_if_t<TIsArrayOrRefOfTypeByPredicate<FmtType, TIsCharEncodingCompatibleWithCharType>::Value>* = nullptr>
+	template <typename FmtType, typename... Types
+		UE_REQUIRES(TIsArrayOrRefOfTypeByPredicate<FmtType, TIsCharEncodingCompatibleWithCharType>::Value)>
 	BuilderType& Appendf(const FmtType& Fmt, Types... Args)
 	{
 		static_assert(TAnd<TIsValidVariadicFunctionArg<Types>...>::Value, "Invalid argument(s) passed to Appendf.");
 		return AppendfImpl(*this, (const CharType*)Fmt, Forward<Types>(Args)...);
 	}
 
+	/**
+	 * Appends to the string builder similarly to how classic vsprintf works.
+	 *
+	 * @param Fmt A format string that specifies how to format the additional arguments. Refer to standard printf format.
+	 */
+	CORE_API BuilderType& AppendV(const CharType* Fmt, va_list Args);
+
 private:
 	CORE_API static BuilderType& VARARGS AppendfImpl(BuilderType& Self, const CharType* Fmt, ...);
 
