[PlatformProcess]  Allow custom implementation of LaunchURL for Windows
#rnx

diff --git a/Engine/Source/Runtime/Core/Private/Misc/CoreDelegates.cpp b/Engine/Source/Runtime/Core/Private/Misc/CoreDelegates.cpp
--- a/Engine/Source/Runtime/Core/Private/Misc/CoreDelegates.cpp
+++ b/Engine/Source/Runtime/Core/Private/Misc/CoreDelegates.cpp
@@ -248,6 +248,7 @@ TTSMulticastDelegate<void()>& FCoreDelegates::TSConfigReadyForUse()
 TDelegate<bool()> FCoreDelegates::IsLoadingMovieCurrentlyPlaying;
 
 TDelegate<bool(const TCHAR*)> FCoreDelegates::ShouldLaunchUrl;
+TDelegate<void(const FString&, FString*)> FCoreDelegates::LaunchCustomHandlerForURL;
 
 TMulticastDelegate<void(const FString&, FPlatformUserId)> FCoreDelegates::OnActivatedByProtocol;
 
diff --git a/Engine/Source/Runtime/Core/Private/Windows/WindowsPlatformProcess.cpp b/Engine/Source/Runtime/Core/Private/Windows/WindowsPlatformProcess.cpp
--- a/Engine/Source/Runtime/Core/Private/Windows/WindowsPlatformProcess.cpp
+++ b/Engine/Source/Runtime/Core/Private/Windows/WindowsPlatformProcess.cpp
@@ -343,9 +343,12 @@ void FWindowsPlatformProcess::LaunchURL( const TCHAR* URL, const TCHAR* Parms, F
 			*Error = TEXT("");
 		}
 
+		if (FCoreDelegates::LaunchCustomHandlerForURL.IsBound())
+		{
+			FCoreDelegates::LaunchCustomHandlerForURL.Execute(URL, Error);
+		}
 		// Use the default handler if we have a URI scheme name that doesn't look like a Windows path, and is not http: or https:
-		FString SchemeName;
-		if (FParse::SchemeNameFromURI(URL, SchemeName) && SchemeName.Len() > 1 && SchemeName != TEXT("http") && SchemeName != TEXT("https"))
+		else if (FString SchemeName; FParse::SchemeNameFromURI(URL, SchemeName) && SchemeName.Len() > 1 && SchemeName != TEXT("http") && SchemeName != TEXT("https"))
 		{
 			LaunchDefaultHandlerForURL(URL, Error);
 		}
diff --git a/Engine/Source/Runtime/Core/Public/Misc/CoreDelegates.h b/Engine/Source/Runtime/Core/Public/Misc/CoreDelegates.h
--- a/Engine/Source/Runtime/Core/Public/Misc/CoreDelegates.h
+++ b/Engine/Source/Runtime/Core/Public/Misc/CoreDelegates.h
@@ -503,6 +503,9 @@ public:
 	// Return true for to launch the url
 	static TDelegate<bool(const TCHAR* URL)> ShouldLaunchUrl;
 
+	// Callback to allow user code to implement a custom implementation for FPlatformProcess::LaunchURL. (Windows Only)
+	static TDelegate<void(const FString& URL, FString* Error)> LaunchCustomHandlerForURL;
+	
 	// Callback when the application has been activated by protocol (with optional user id, depending on the platform)
 	static TMulticastDelegate<void(const FString& Parameter, FPlatformUserId UserId /*= PLATFORMUSERID_NONE*/)> OnActivatedByProtocol;
 
