Add better error handling when we crash from a SIGSYS which can be caused from sysfiltering

diff --git a/Engine/Source/Runtime/Core/Private/Unix/UnixPlatformCrashContext.cpp b/Engine/Source/Runtime/Core/Private/Unix/UnixPlatformCrashContext.cpp
--- a/Engine/Source/Runtime/Core/Private/Unix/UnixPlatformCrashContext.cpp
+++ b/Engine/Source/Runtime/Core/Private/Unix/UnixPlatformCrashContext.cpp
@@ -67,6 +67,9 @@ FString DescribeSignal(int32 Signal, siginfo_t* Info, ucontext_t *Context)
 	case SIGBUS:
 		ErrorString += FString::Printf(TEXT("SIGBUS: invalid attempt to access memory at address 0x%016llx"), (uint64)Info->si_addr);
 		break;
+	case SIGSYS:
+		ErrorString += FString::Printf(TEXT("SIGSYS: non-existent or invalid system call invoked %i"), Info->si_syscall);
+		break;
 
 		HANDLE_CASE(SIGINT, "program interrupted")
 		HANDLE_CASE(SIGQUIT, "user-requested crash")
@@ -75,7 +78,6 @@ FString DescribeSignal(int32 Signal, siginfo_t* Info, ucontext_t *Context)
 		HANDLE_CASE(SIGABRT, "abort() called")
 		HANDLE_CASE(SIGFPE, "floating-point exception")
 		HANDLE_CASE(SIGKILL, "program killed")
-		HANDLE_CASE(SIGSYS, "non-existent system call invoked")
 		HANDLE_CASE(SIGPIPE, "write on a pipe with no reader")
 		HANDLE_CASE(SIGTERM, "software termination signal")
 		HANDLE_CASE(SIGSTOP, "stop")
@@ -986,6 +988,12 @@ void PlatformCrashHandler(int32 Signal, siginfo_t* Info, void* Context)
 			DefaultErrorMessage.Append(SignalName);
 		}
 
+		if (Signal == SIGSYS)
+		{
+			DefaultErrorMessage.Append(TEXT(" from syscall "));
+			DefaultErrorMessage.Append(ItoANSI(Info->si_syscall, 10));
+		}
+
 		ErrorMessage = *DefaultErrorMessage;
 		ErrorProgramCounter = __builtin_return_address(0);
 	}
