Harden reading/writing frozen memory images.

- Prevent maliciously crafted or damaged memory images from introducing memory stomps or otherwise disrupting loading.

diff --git a/Engine/Source/Runtime/Core/Private/Serialization/MemoryImage.cpp b/Engine/Source/Runtime/Core/Private/Serialization/MemoryImage.cpp
--- a/Engine/Source/Runtime/Core/Private/Serialization/MemoryImage.cpp
+++ b/Engine/Source/Runtime/Core/Private/Serialization/MemoryImage.cpp
@@ -1335,7 +1335,11 @@ FMemoryImageObject FreezeMemoryImageObject(const void* Object, const FTypeLayout
 	check(FrozenSize > 0u);
 	void* FrozenObject = FMemory::Malloc(FrozenSize);
 	FMemory::Memcpy(FrozenObject, MemoryImageResult.Bytes.GetData(), FrozenSize);
-	MemoryImageResult.ApplyPatches(FrozenObject);
+	if (UNLIKELY(!MemoryImageResult.ApplyPatches(FrozenObject, FrozenSize)))
+	{
+		return FMemoryImageObject();
+	}
+
 	return FMemoryImageObject(TypeDesc, FrozenObject, FrozenSize);
 }
 
@@ -1508,21 +1512,38 @@ void FMemoryImageResult::SaveToArchive(FArchive& Ar) const
 	Ar.Seek(EndOffset);
 }
 
-static inline void ApplyVTablePatch(void* FrozenObject, const FTypeLayoutDesc& DerivedType, uint32 VTableOffset, uint32 Offset)
+static inline bool ApplyVTablePatch(void* FrozenObject, uint64 FrozenObjectSize, const FTypeLayoutDesc & DerivedType, uint32 VTableOffset, uint32 Offset)
 {
+	if (UNLIKELY(Offset > (FrozenObjectSize - sizeof(void*)) || VTableOffset > (DerivedType.Size - sizeof(void*))))
+	{
+		return false;
+	}
+
 	const void** VTableSrc = (void const**)((uint8*)DerivedType.GetDefaultObjectFunc() + VTableOffset);
 	void const** VTableDst = (void const**)((uint8*)FrozenObject + Offset);
 	*VTableDst = *VTableSrc;
+	return true;
 }
 
-static inline void ApplyScriptNamePatch(void* FrozenObject, const FScriptName& Name, uint32 Offset)
+static inline bool ApplyScriptNamePatch(void* FrozenObject, uint64 FrozenObjectSize, const FScriptName& Name, uint32 Offset)
 {
+	if (UNLIKELY(Offset > FrozenObjectSize - sizeof(FScriptName)))
+	{
+		return false;
+	}
+
 	void* NameDst = (uint8*)FrozenObject + Offset;
 	new(NameDst) FScriptName(Name);
+	return true;
 }
 
-static inline void ApplyMemoryImageNamePatch(void* FrozenObject, const FMemoryImageName& Name, uint32 Offset, const FPlatformTypeLayoutParameters& LayoutParameters)
+static inline bool ApplyMemoryImageNamePatch(void* FrozenObject, uint64 FrozenObjectSize, const FMemoryImageName& Name, uint32 Offset, const FPlatformTypeLayoutParameters& LayoutParameters)
 {
+	if (UNLIKELY(Offset > FrozenObjectSize - sizeof(FMemoryImageName)))
+	{
+		return false;
+	}
+
 	void* NameDst = (uint8*)FrozenObject + Offset;
 	if (LayoutParameters.IsCurrentPlatform())
 	{
@@ -1532,26 +1553,39 @@ static inline void ApplyMemoryImageNamePatch(void* FrozenObject, const FMemoryIm
 	{
 		Freeze::ApplyMemoryImageNamePatch(NameDst, Name, LayoutParameters);
 	}
+
+	return true;
 }
 
-void FMemoryImageResult::ApplyPatches(void* FrozenObject) const
+bool FMemoryImageResult::ApplyPatches(void* FrozenObject, uint64 FrozenObjectSize) const
 {
 	for (const FMemoryImageVTablePointer& Patch : VTables)
 	{
 		const FTypeLayoutDesc* DerivedType = FTypeLayoutDesc::Find(Patch.TypeNameHash);
 		check(DerivedType);
-		ApplyVTablePatch(FrozenObject, *DerivedType, Patch.VTableOffset, Patch.Offset);
+		if (!ApplyVTablePatch(FrozenObject, FrozenObjectSize, *DerivedType, Patch.VTableOffset, Patch.Offset))
+		{
+			return false;
+		}
 	}
 
 	for (const FMemoryImageNamePointer& Patch : ScriptNames)
 	{
-		ApplyScriptNamePatch(FrozenObject, NameToScriptName(Patch.Name), Patch.Offset);
+		if (!ApplyScriptNamePatch(FrozenObject, FrozenObjectSize, NameToScriptName(Patch.Name), Patch.Offset))
+		{
+			return false;
+		}
 	}
 
 	for (const FMemoryImageNamePointer& Patch : MemoryImageNames)
 	{
-		ApplyMemoryImageNamePatch(FrozenObject, FMemoryImageName(Patch.Name), Patch.Offset, TargetLayoutParameters);
+		if (!ApplyMemoryImageNamePatch(FrozenObject, FrozenObjectSize, FMemoryImageName(Patch.Name), Patch.Offset, TargetLayoutParameters))
+		{
+			return false;
+		}
 	}
+
+	return true;
 }
 
 FMemoryImageObject FMemoryImageResult::LoadFromArchive(FArchive& Ar, const FTypeLayoutDesc& TypeDesc, FPointerTableBase* PointerTable, FPlatformTypeLayoutParameters& LayoutParameters)
@@ -1565,7 +1599,7 @@ FMemoryImageObject FMemoryImageResult::LoadFromArchive(FArchive& Ar, const FType
 	void* FrozenObject = FMemory::Malloc(FrozenSize);
 	Ar.Serialize(FrozenObject, FrozenSize);
 
-	const bool bFrozenObjectValid = PointerTable->LoadFromArchive(Ar, LayoutParameters, FrozenObject);
+	bool bFrozenObjectValid = PointerTable->LoadFromArchive(Ar, LayoutParameters, FrozenObject);
 
 	uint32 NumVTables = 0u;
 	uint32 NumScriptNames = 0u;
@@ -1574,6 +1608,15 @@ FMemoryImageObject FMemoryImageResult::LoadFromArchive(FArchive& Ar, const FType
 	Ar << NumScriptNames;
 	Ar << NumMemoryImageNames;
 
+	// if image read is invalid (or if we find out when patching that offsets are invalid),
+	// we will still keep reading the image data from the archive as we cannot rewind
+
+	// clamp to reasonable numbers
+	const uint32 kReasonableMemoryImageElementCount = 16u * 1024u * 1024u;
+
+	ensure(NumVTables < kReasonableMemoryImageElementCount);
+	NumVTables = FMath::Min(NumVTables, kReasonableMemoryImageElementCount);
+
 	for (uint32 i = 0u; i < NumVTables; ++i)
 	{
 		uint64 TypeNameHash = 0u;
@@ -1581,8 +1624,11 @@ FMemoryImageObject FMemoryImageResult::LoadFromArchive(FArchive& Ar, const FType
 		Ar << TypeNameHash;
 		Ar << NumPatches;
 
-		const FTypeLayoutDesc* DerivedType = FTypeLayoutDesc::Find(TypeNameHash);
-		check(DerivedType);
+		const FTypeLayoutDesc* DerivedType = bFrozenObjectValid ? FTypeLayoutDesc::Find(TypeNameHash) : nullptr;
+		check(DerivedType || !bFrozenObjectValid);
+
+		ensure(NumPatches < kReasonableMemoryImageElementCount);
+		NumPatches = FMath::Min(NumPatches, kReasonableMemoryImageElementCount);
 
 		for(uint32 PatchIndex = 0u; PatchIndex < NumPatches; ++PatchIndex)
 		{
@@ -1590,10 +1636,16 @@ FMemoryImageObject FMemoryImageResult::LoadFromArchive(FArchive& Ar, const FType
 			uint32 Offset = 0u;
 			Ar << VTableOffset;
 			Ar << Offset;
-			ApplyVTablePatch(FrozenObject, *DerivedType, VTableOffset, Offset);
+			if (bFrozenObjectValid && !ApplyVTablePatch(FrozenObject, FrozenSize, *DerivedType, VTableOffset, Offset))
+			{
+				bFrozenObjectValid = false;
+			}
 		}
 	}
 
+	ensure(NumScriptNames < kReasonableMemoryImageElementCount);
+	NumScriptNames = FMath::Min(NumScriptNames, kReasonableMemoryImageElementCount);
+
 	for (uint32 i = 0u; i < NumScriptNames; ++i)
 	{
 		FName Name;
@@ -1601,14 +1653,23 @@ FMemoryImageObject FMemoryImageResult::LoadFromArchive(FArchive& Ar, const FType
 		Ar << Name;
 		Ar << NumPatches;
 
+		ensure(NumPatches < kReasonableMemoryImageElementCount);
+		NumPatches = FMath::Min(NumPatches, kReasonableMemoryImageElementCount);
+
 		for (uint32 PatchIndex = 0u; PatchIndex < NumPatches; ++PatchIndex)
 		{
 			uint32 Offset = 0u;
 			Ar << Offset;
-			ApplyScriptNamePatch(FrozenObject, NameToScriptName(Name), Offset);
+			if (bFrozenObjectValid && !ApplyScriptNamePatch(FrozenObject, FrozenSize, NameToScriptName(Name), Offset))
+			{
+				bFrozenObjectValid = false;
+			}
 		}
 	}
 
+	ensure(NumMemoryImageNames < kReasonableMemoryImageElementCount);
+	NumMemoryImageNames = FMath::Min(NumMemoryImageNames, kReasonableMemoryImageElementCount);
+
 	for (uint32 i = 0u; i < NumMemoryImageNames; ++i)
 	{
 		FName Name;
@@ -1616,11 +1677,17 @@ FMemoryImageObject FMemoryImageResult::LoadFromArchive(FArchive& Ar, const FType
 		Ar << Name;
 		Ar << NumPatches;
 
+		ensure(NumPatches < kReasonableMemoryImageElementCount);
+		NumPatches = FMath::Min(NumPatches, kReasonableMemoryImageElementCount);
+
 		for (uint32 PatchIndex = 0u; PatchIndex < NumPatches; ++PatchIndex)
 		{
 			uint32 Offset = 0u;
 			Ar << Offset;
-			ApplyMemoryImageNamePatch(FrozenObject, FMemoryImageName(Name), Offset, LayoutParameters);
+			if (bFrozenObjectValid && !ApplyMemoryImageNamePatch(FrozenObject, FrozenSize, FMemoryImageName(Name), Offset, LayoutParameters))
+			{
+				bFrozenObjectValid = false;
+			}
 		}
 	}
 
diff --git a/Engine/Source/Runtime/Core/Public/Serialization/MemoryImage.h b/Engine/Source/Runtime/Core/Public/Serialization/MemoryImage.h
--- a/Engine/Source/Runtime/Core/Public/Serialization/MemoryImage.h
+++ b/Engine/Source/Runtime/Core/Public/Serialization/MemoryImage.h
@@ -206,7 +206,7 @@ struct FMemoryImageResult
 	TArray<FMemoryImageNamePointer> MemoryImageNames;
 
 	CORE_API void SaveToArchive(FArchive& Ar) const;
-	CORE_API void ApplyPatches(void* FrozenObject) const;
+	CORE_API bool ApplyPatches(void* FrozenObject, uint64 FrozenObjectSize) const;
 	CORE_API static FMemoryImageObject LoadFromArchive(FArchive& Ar, const FTypeLayoutDesc& TypeDesc, FPointerTableBase* PointerTable, FPlatformTypeLayoutParameters& OutLayoutParameters);
 };
 
