Enable HW CRC on ARM platforms depending on a clang compiler flags

diff --git a/Engine/Source/Runtime/Core/Private/Misc/Crc.cpp b/Engine/Source/Runtime/Core/Private/Misc/Crc.cpp
--- a/Engine/Source/Runtime/Core/Private/Misc/Crc.cpp
+++ b/Engine/Source/Runtime/Core/Private/Misc/Crc.cpp
@@ -5,7 +5,9 @@
 #include "Templates/UnrealTemplate.h"
 #include "Misc/ByteSwap.h"
 
-#if PLATFORM_ANDROID_ARM64 || PLATFORM_LINUXARM64 || PLATFORM_IOS
+#if defined(__clang__) && defined(__ARM_FEATURE_CRYPTO)
+#	define DETECT_HW_CRC32_SUPPORT_IN_RUNTIME 0
+#elif PLATFORM_ANDROID_ARM64 || PLATFORM_LINUXARM64 || PLATFORM_IOS
 #	define DETECT_HW_CRC32_SUPPORT_IN_RUNTIME 1
 #else
 #	define DETECT_HW_CRC32_SUPPORT_IN_RUNTIME 0
@@ -347,12 +349,16 @@ uint32 FCrc::CRCTablesSB8[8][256] =
 };
 
 static uint32 MemCrc32SW(const void* InData, int32 Length, uint32 CRC/*=0 */);
-FCrc::MemCrc32Functor FCrc::MemCrc32Func = &MemCrc32SW;
-
 #if defined(__clang__) && PLATFORM_CPU_ARM_FAMILY
 static uint32 MemCrc32ArmHW(const void* InData, int32 Length, uint32 CRC/*=0 */);
 #endif
 
+#if defined(__clang__) && defined(__ARM_FEATURE_CRYPTO)
+FCrc::MemCrc32Functor FCrc::MemCrc32Func = &MemCrc32ArmHW;
+#else
+FCrc::MemCrc32Functor FCrc::MemCrc32Func = &MemCrc32SW;
+#endif
+
 #if DETECT_HW_CRC32_SUPPORT_IN_RUNTIME
 static sigjmp_buf CrcHwJump;
 static void CrcHwSignalHandler(int Sig)
diff --git a/Engine/Source/Runtime/Core/Public/HAL/Platform.h b/Engine/Source/Runtime/Core/Public/HAL/Platform.h
--- a/Engine/Source/Runtime/Core/Public/HAL/Platform.h
+++ b/Engine/Source/Runtime/Core/Public/HAL/Platform.h
@@ -479,10 +479,6 @@
 	#define PLATFORM_WEAKLY_CONSISTENT_MEMORY PLATFORM_CPU_ARM_FAMILY
 #endif
 
-#ifndef PLATFORM_HAS_CRC_INTRINSICS
-	#define PLATFORM_HAS_CRC_INTRINSICS							0
-#endif
-
 #ifndef PLATFORM_NEEDS_RHIRESOURCELIST
 	#define PLATFORM_NEEDS_RHIRESOURCELIST 1
 #endif
