Fixed AES code compilation for ARM architecture

diff --git a/Engine/Source/Runtime/Core/Private/Misc/AES.cpp b/Engine/Source/Runtime/Core/Private/Misc/AES.cpp
--- a/Engine/Source/Runtime/Core/Private/Misc/AES.cpp
+++ b/Engine/Source/Runtime/Core/Private/Misc/AES.cpp
@@ -24,6 +24,7 @@
 
 // on ARMv8 enable AES instructions unconditionally only when they are enabled with compiler
 #ifdef __ARM_FEATURE_CRYPTO
+#include <arm_neon.h>
 #define UE_PLATFORM_AES_ARMV8 1
 #define UE_PLATFORM_AES_FALLBACK 0
 #endif
