Remove unnecessary public include in Core

#rnx

diff --git a/Engine/Source/Runtime/Core/Core.Build.cs b/Engine/Source/Runtime/Core/Core.Build.cs
--- a/Engine/Source/Runtime/Core/Core.Build.cs
+++ b/Engine/Source/Runtime/Core/Core.Build.cs
@@ -30,7 +30,6 @@ public class Core : ModuleRules
 		PrivateDependencyModuleNames.Add("OodleDataCompression");
 
 		PublicDependencyModuleNames.Add("TraceLog");
-		PublicIncludePaths.Add("Runtime/TraceLog/Public");
 
 		PrivateIncludePathModuleNames.AddRange(
 			new string[] {
