Removed obsolete define PLATFORM_HAS_CRC_INTRINSICS from Mac platform

diff --git a/Engine/Source/Runtime/Core/Public/Mac/MacPlatform.h b/Engine/Source/Runtime/Core/Public/Mac/MacPlatform.h
--- a/Engine/Source/Runtime/Core/Public/Mac/MacPlatform.h
+++ b/Engine/Source/Runtime/Core/Public/Mac/MacPlatform.h
@@ -56,7 +56,6 @@ typedef FMacPlatformTypes FPlatformTypes;
 #ifndef PLATFORM_ALWAYS_HAS_FMA3
 	#define PLATFORM_ALWAYS_HAS_FMA3				0
 #endif
-#define PLATFORM_HAS_CRC_INTRINSICS					PLATFORM_MAC_ARM64
 //#define PLATFORM_USE_LS_SPEC_FOR_WIDECHAR			1
 #define PLATFORM_USE_SYSTEM_VSWPRINTF				0
 #define PLATFORM_COMPILER_DISTINGUISHES_INT_AND_LONG			1
