Fixed compilation

diff --git a/Engine/Source/Runtime/Core/Public/Misc/AssertionMacros.h b/Engine/Source/Runtime/Core/Public/Misc/AssertionMacros.h
--- a/Engine/Source/Runtime/Core/Public/Misc/AssertionMacros.h
+++ b/Engine/Source/Runtime/Core/Public/Misc/AssertionMacros.h
@@ -368,7 +368,7 @@ RetType FORCENOINLINE UE_DEBUG_SECTION DispatchCheckVerify(InnerType&& Inner, Ar
 		}
 	};
 
-	void UE_DEBUG_SECTION VARARGS CheckVerifyImpl(bool& InOutExecuted, bool Always, const ANSICHAR* File, int32 Line, void* ProgramCounter, const ANSICHAR* Expr, const TCHAR* Format, ...);
+	CORE_API void UE_DEBUG_SECTION VARARGS CheckVerifyImpl(bool& InOutExecuted, bool Always, const ANSICHAR* File, int32 Line, void* ProgramCounter, const ANSICHAR* Expr, const TCHAR* Format, ...);
 
 	#define UE_ENSURE_IMPL(Capture, Always, InExpression, InFormat) \
 		(LIKELY(!!(InExpression)) || (DispatchCheckVerify<bool>([Capture] () UE_DEBUG_SECTION \
