Make FMicrosoftPlatformString::Strupr consistent with other platform implementations and respect the DestCount parameter. This enables using it on non-null-terminated strings.

diff --git a/Engine/Source/Runtime/Core/Public/GenericPlatform/MicrosoftPlatformString.h b/Engine/Source/Runtime/Core/Public/GenericPlatform/MicrosoftPlatformString.h
--- a/Engine/Source/Runtime/Core/Public/GenericPlatform/MicrosoftPlatformString.h
+++ b/Engine/Source/Runtime/Core/Public/GenericPlatform/MicrosoftPlatformString.h
@@ -42,6 +42,18 @@ struct FMicrosoftPlatformString :
 	using FGenericPlatformString::Strncmp;
 	using FGenericPlatformString::Strnicmp;
 
+#if !PLATFORM_USE_GENERIC_STRING_IMPLEMENTATION
+	template <typename CharType>
+	static CharType* Strupr(CharType* Dest, SIZE_T DestCount)
+	{
+		for (CharType* Char = Dest; *Char && DestCount > 0; Char++, DestCount--)
+		{
+			*Char = TChar<CharType>::ToUpper(*Char);
+		}
+		return Dest;
+	}
+#endif
+
 	/** 
 	 * Wide character implementation 
 	 **/
@@ -62,11 +74,6 @@ struct FMicrosoftPlatformString :
 		return (WIDECHAR*)_tcscat(Dest, Src);
 	}
 
-	static FORCEINLINE WIDECHAR* Strupr(WIDECHAR* Dest, SIZE_T DestCount)
-	{
-		return (WIDECHAR*)_tcsupr(Dest);
-	}
-
 	static FORCEINLINE int32 Strcmp( const WIDECHAR* String1, const WIDECHAR* String2 )
 	{
 		return _tcscmp(String1, String2);
@@ -174,11 +181,6 @@ struct FMicrosoftPlatformString :
 		return (ANSICHAR*)strcat( Dest, Src );
 	}
 
-	static FORCEINLINE ANSICHAR* Strupr(ANSICHAR* Dest, SIZE_T DestCount)
-	{
-		return (ANSICHAR*)_strupr(Dest);
-	}
-
 	static FORCEINLINE int32 Strcmp( const ANSICHAR* String1, const ANSICHAR* String2 )
 	{
 		return strcmp(String1, String2);
@@ -293,11 +295,6 @@ struct FMicrosoftPlatformString :
 		return (UTF8CHAR*)Strcat((ANSICHAR*)Dest, DestCount, (const ANSICHAR*)Src);
 	}
 
-	static FORCEINLINE UTF8CHAR* Strupr(UTF8CHAR* Dest, SIZE_T DestCount)
-	{
-		return (UTF8CHAR*)Strupr((ANSICHAR*)Dest, DestCount);
-	}
-
 	static FORCEINLINE int32 Strcmp(const UTF8CHAR* String1, const UTF8CHAR* String2)
 	{
 		return Strcmp((const ANSICHAR*)String1, (const ANSICHAR*)String2);
diff --git a/Engine/Source/Runtime/Core/Tests/HAL/PlatformStringTests.cpp b/Engine/Source/Runtime/Core/Tests/HAL/PlatformStringTests.cpp
new file mode 100644
--- /dev/null
+++ b/Engine/Source/Runtime/Core/Tests/HAL/PlatformStringTests.cpp
@@ -0,0 +1,82 @@
+// Copyright Epic Games, Inc. All Rights Reserved.
+
+#if WITH_LOW_LEVEL_TESTS
+
+#include "HAL/PlatformString.h"
+#include "Containers/Array.h"
+#include "Templates/UnrealTemplate.h"
+
+#include <catch2/catch_test_macros.hpp>
+
+TEST_CASE("Core::HAL::PlatformString", "[Core]")
+{
+	SECTION("Strupr TCHAR null-termination")
+	{
+		// Whole TCHAR string
+		const TCHAR TestLiteral[] = TEXT("the quick brown fox jumps over the lazy dog");
+
+		TArray<TCHAR> TestFullString(TestLiteral, UE_ARRAY_COUNT(TestLiteral));
+
+		FPlatformString::Strupr(TestFullString.GetData(), TestFullString.Num());
+		REQUIRE(FPlatformString::Strncmp(TestFullString.GetData(), TEXT("THE QUICK BROWN FOX JUMPS OVER THE LAZY DOG"), TestFullString.Num()) == 0);
+
+		// Partial TCHAR string
+		TArray<TCHAR> TestPartialString(TestLiteral, UE_ARRAY_COUNT(TestLiteral));
+
+		FPlatformString::Strupr(TestPartialString.GetData(), 19);
+		REQUIRE(FPlatformString::Strncmp(TestPartialString.GetData(), TEXT("THE QUICK BROWN FOX jumps over the lazy dog"), TestPartialString.Num()) == 0);
+	}
+
+	SECTION("Strupr ANSICHAR null-termination")
+	{
+		// Whole ANSICHAR string
+		const ANSICHAR TestLiteral[] = "the quick brown fox jumps over the lazy dog";
+
+		TArray<ANSICHAR> TestFullString(TestLiteral, UE_ARRAY_COUNT(TestLiteral));
+
+		FPlatformString::Strupr(TestFullString.GetData(), TestFullString.Num());
+		REQUIRE(FPlatformString::Strncmp(TestFullString.GetData(), "THE QUICK BROWN FOX JUMPS OVER THE LAZY DOG", TestFullString.Num()) == 0);
+
+		// Partial ANSICHAR string
+		TArray<ANSICHAR> TestPartialString(TestLiteral, UE_ARRAY_COUNT(TestLiteral));
+
+		FPlatformString::Strupr(TestPartialString.GetData(), 19);
+		REQUIRE(FPlatformString::Strncmp(TestPartialString.GetData(), "THE QUICK BROWN FOX jumps over the lazy dog", TestPartialString.Num()) == 0);
+	}
+
+	SECTION("Strupr WIDECHAR null-termination")
+	{
+		// Whole WIDECHAR string
+		const WIDECHAR TestLiteral[] = WIDETEXT("the quick brown fox jumps over the lazy dog");
+
+		TArray<WIDECHAR> TestFullString(TestLiteral, UE_ARRAY_COUNT(TestLiteral));
+
+		FPlatformString::Strupr(TestFullString.GetData(), TestFullString.Num());
+		REQUIRE(FPlatformString::Strncmp(TestFullString.GetData(), WIDETEXT("THE QUICK BROWN FOX JUMPS OVER THE LAZY DOG"), TestFullString.Num()) == 0);
+
+		// Partial WIDECHAR string
+		TArray<WIDECHAR> TestPartialString(TestLiteral, UE_ARRAY_COUNT(TestLiteral));
+
+		FPlatformString::Strupr(TestPartialString.GetData(), 19);
+		REQUIRE(FPlatformString::Strncmp(TestPartialString.GetData(), WIDETEXT("THE QUICK BROWN FOX jumps over the lazy dog"), TestPartialString.Num()) == 0);
+	}
+
+	SECTION("Strupr UTF8CHAR null-termination")
+	{
+		// Whole UTF8CHAR string
+		const UTF8CHAR* TestLiteral = UTF8TEXT("the quick brown fox jumps over the lazy dog");
+
+		TArray<UTF8CHAR> TestFullString(TestLiteral, FPlatformString::Strlen(TestLiteral) + 1);
+
+		FPlatformString::Strupr(TestFullString.GetData(), TestFullString.Num());
+		REQUIRE(FPlatformString::Strncmp(TestFullString.GetData(), UTF8TEXT("THE QUICK BROWN FOX JUMPS OVER THE LAZY DOG"), TestFullString.Num()) == 0);
+
+		// Partial UTF8CHAR string
+		TArray<UTF8CHAR> TestPartialString(TestLiteral, FPlatformString::Strlen(TestLiteral) + 1);
+
+		FPlatformString::Strupr(TestPartialString.GetData(), 19);
+		REQUIRE(FPlatformString::Strncmp(TestPartialString.GetData(), UTF8TEXT("THE QUICK BROWN FOX jumps over the lazy dog"), TestPartialString.Num()) == 0);
+	}
+}
+
+#endif
\ No newline at end of file
