#UE Linker fix

diff --git a/Engine/Source/Runtime/Core/Public/GenericPlatform/GenericPlatformCrashContext.h b/Engine/Source/Runtime/Core/Public/GenericPlatform/GenericPlatformCrashContext.h
--- a/Engine/Source/Runtime/Core/Public/GenericPlatform/GenericPlatformCrashContext.h
+++ b/Engine/Source/Runtime/Core/Public/GenericPlatform/GenericPlatformCrashContext.h
@@ -432,7 +432,7 @@ public:
 	CORE_API static void PurgeOldCrashConfig();
 
 	/** Set or change the epic account id associated with the crash session. Will override the epic account id stored in the registry for reporting when present. */ 
-	static void SetEpicAccountId(const FString& EpicAccountId);
+	CORE_API static void SetEpicAccountId(const FString& EpicAccountId);
 
 	/** Clears the engine data dictionary */
 	CORE_API static void ResetEngineData();
