Changed FConsoleVariableRef's GetBool/GetInt/GetFloat/GetString implementations to get the value that is appropriate for the calling thread instead of always returning the main thread value.
Hoisted the GetBool/GetInt/GetFloat/GetString type conversion logic into FConsoleVariableConversionHelper so it can be shared between FConsoleVariable and FConsoleVariableRef.

diff --git a/Engine/Source/Runtime/Core/Private/HAL/ConsoleManager.cpp b/Engine/Source/Runtime/Core/Private/HAL/ConsoleManager.cpp
--- a/Engine/Source/Runtime/Core/Private/HAL/ConsoleManager.cpp
+++ b/Engine/Source/Runtime/Core/Private/HAL/ConsoleManager.cpp
@@ -653,8 +653,93 @@ protected:
 #endif
 };
 
+template<class T>
+class FConsoleVariableConversionHelper
+{
+public:
+	static bool GetBool(T Value);
+	static int32 GetInt(T Value);
+	static float GetFloat(T Value);
+	static FString GetString(T Value);
+};
 
+// specialization for bool
+template<> bool FConsoleVariableConversionHelper<bool>::GetBool(bool Value)
+{
+	return Value;
+}
+template<> int32 FConsoleVariableConversionHelper<bool>::GetInt(bool Value)
+{
+	return Value ? 1 : 0;
+}
+template<> float FConsoleVariableConversionHelper<bool>::GetFloat(bool Value)
+{
+	return Value ? 1.0f : 0.0f;
+}
+template<> FString FConsoleVariableConversionHelper<bool>::GetString(bool Value)
+{
+	return Value ? TEXT("true") : TEXT("false");
+}
 
+// specialization for int32
+template<> bool FConsoleVariableConversionHelper<int32>::GetBool(int32 Value)
+{
+	return Value != 0;
+}
+template<> int32 FConsoleVariableConversionHelper<int32>::GetInt(int32 Value)
+{
+	return Value;
+}
+template<> float FConsoleVariableConversionHelper<int32>::GetFloat(int32 Value)
+{
+	return (float)Value;
+}
+template<> FString FConsoleVariableConversionHelper<int32>::GetString(int32 Value)
+{
+	return FString::Printf(TEXT("%d"), Value);
+}
+
+// specialization for float
+template<> bool FConsoleVariableConversionHelper<float>::GetBool(float Value)
+{
+	return Value != 0.0f;
+}
+template<> int32 FConsoleVariableConversionHelper<float>::GetInt(float Value)
+{
+	return (int32)Value;
+}
+template<> float FConsoleVariableConversionHelper<float>::GetFloat(float Value)
+{
+	return Value;
+}
+template<> FString FConsoleVariableConversionHelper<float>::GetString(float Value)
+{
+	return FString::Printf(TEXT("%g"), Value);
+}
+
+// specialization for FString
+template<> bool FConsoleVariableConversionHelper<FString>::GetBool(FString Value)
+{
+	bool OutValue = false;
+	TTypeFromString<bool>::FromString(OutValue, *Value);
+	return OutValue;
+}
+template<> int32 FConsoleVariableConversionHelper<FString>::GetInt(FString Value)
+{
+	int32 OutValue = 0;
+	TTypeFromString<int32>::FromString(OutValue, *Value);
+	return OutValue;
+}
+template<> float FConsoleVariableConversionHelper<FString>::GetFloat(FString Value)
+{
+	float OutValue = 0.0f;
+	TTypeFromString<float>::FromString(OutValue, *Value);
+	return OutValue;
+}
+template<> FString FConsoleVariableConversionHelper<FString>::GetString(FString Value)
+{
+	return Value;
+}
 
 
 // T: bool, int32, float, FString
@@ -689,10 +774,10 @@ public:
 		}
 	}
 
-	virtual bool GetBool() const override;
-	virtual int32 GetInt() const override;
-	virtual float GetFloat() const override;
-	virtual FString GetString() const override;
+	virtual bool GetBool() const override { return FConsoleVariableConversionHelper<T>::GetBool(Value()); }
+	virtual int32 GetInt() const override { return FConsoleVariableConversionHelper<T>::GetInt(Value()); }
+	virtual float GetFloat() const override { return FConsoleVariableConversionHelper<T>::GetFloat(Value()); }
+	virtual FString GetString() const override { return FConsoleVariableConversionHelper<T>::GetString(Value()); }
 
 	virtual bool IsVariableBool() const override { return false; }
 	virtual bool IsVariableInt() const override { return false; }
@@ -745,22 +830,6 @@ template<> bool FConsoleVariable<FString>::IsVariableString() const
 
 // specialization for bool
 
-template<> bool FConsoleVariable<bool>::GetBool() const
-{
-	return Value();
-}
-template<> int32 FConsoleVariable<bool>::GetInt() const
-{
-	return Value() ? 1 : 0;
-}
-template<> float FConsoleVariable<bool>::GetFloat() const
-{
-	return Value() ? 1.0f : 0.0f;
-}
-template<> FString FConsoleVariable<bool>::GetString() const
-{
-	return Value() ? TEXT("true") : TEXT("false");
-}
 template<> TConsoleVariableData<bool>* FConsoleVariable<bool>::AsVariableBool()
 {
 	return &Data;
@@ -768,23 +837,6 @@ template<> TConsoleVariableData<bool>* FConsoleVariable<bool>::AsVariableBool()
 
 // specialization for int32
 
-template<> bool FConsoleVariable<int32>::GetBool() const
-{
-	return Value() != 0;
-}
-template<> int32 FConsoleVariable<int32>::GetInt() const
-{
-	return Value();
-}
-template<> float FConsoleVariable<int32>::GetFloat() const
-{
-	return (float)Value();
-}
-template<> FString FConsoleVariable<int32>::GetString() const
-{
-	return FString::Printf(TEXT("%d"), Value());
-}
-
 template<> TConsoleVariableData<int32>* FConsoleVariable<int32>::AsVariableInt()
 {
 	return &Data; 
@@ -792,22 +844,6 @@ template<> TConsoleVariableData<int32>* FConsoleVariable<int32>::AsVariableInt()
 
 // specialization for float
 
-template<> bool FConsoleVariable<float>::GetBool() const
-{
-	return Value() != 0;
-}
-template<> int32 FConsoleVariable<float>::GetInt() const
-{
-	return (int32)Value();
-}
-template<> float FConsoleVariable<float>::GetFloat() const
-{
-	return Value();
-}
-template<> FString FConsoleVariable<float>::GetString() const
-{
-	return FString::Printf(TEXT("%g"), Value());
-}
 template<> TConsoleVariableData<float>* FConsoleVariable<float>::AsVariableFloat()
 {
 	return &Data;
@@ -823,28 +859,7 @@ template<> void FConsoleVariable<FString>::Set(const TCHAR* InValue, EConsoleVar
 		OnChanged(SetBy);
 	}
 }
-template<> bool FConsoleVariable<FString>::GetBool() const
-{
-	bool OutValue = false;
-	TTypeFromString<bool>::FromString(OutValue, *Value());
-	return OutValue;
-}
-template<> int32 FConsoleVariable<FString>::GetInt() const
-{
-	int32 OutValue = 0;
-	TTypeFromString<int32>::FromString(OutValue, *Value());
-	return OutValue;
-}
-template<> float FConsoleVariable<FString>::GetFloat() const
-{
-	float OutValue = 0.0f;
-	TTypeFromString<float>::FromString(OutValue, *Value());
-	return OutValue;
-}
-template<> FString FConsoleVariable<FString>::GetString() const
-{
-	return Value();
-}
+
 template<> TConsoleVariableData<FString>* FConsoleVariable<FString>::AsVariableString()
 {
 	return &Data;
@@ -955,22 +970,10 @@ public:
 	virtual bool IsVariableFloat() const override { return false; }
 	virtual bool IsVariableString() const override { return false; }
 
-	virtual bool GetBool() const
-	{
-		return (bool)MainValue;
-	}
-	virtual int32 GetInt() const
-	{
-		return (int32)MainValue;
-	}
-	virtual float GetFloat() const
-	{
-		return (float)MainValue;
-	}
-	virtual FString GetString() const
-	{
-		return TTypeToString<T>::ToString(MainValue);
-	}
+	virtual bool GetBool() const override { return FConsoleVariableConversionHelper<T>::GetBool(Value()); }
+	virtual int32 GetInt() const override { return FConsoleVariableConversionHelper<T>::GetInt(Value()); }
+	virtual float GetFloat() const override { return FConsoleVariableConversionHelper<T>::GetFloat(Value()); }
+	virtual FString GetString() const override { return FConsoleVariableConversionHelper<T>::GetString(Value()); }
 
 private: // ----------------------------------------------------
 
@@ -1015,14 +1018,6 @@ bool FConsoleVariableRef<float>::IsVariableFloat() const
 	return true;
 }
 
-// specialization for float
-
-template <>
-FString FConsoleVariableRef<float>::GetString() const
-{
-	// otherwise we get 2.1f would become "2.100000"
-	return FString::SanitizeFloat(RefValue);
-}
 
 // string version
 
