Added "DataStream" to list of FNames UnrealNames.inl

diff --git a/Engine/Source/Runtime/Core/Public/UObject/UnrealNames.inl b/Engine/Source/Runtime/Core/Public/UObject/UnrealNames.inl
--- a/Engine/Source/Runtime/Core/Public/UObject/UnrealNames.inl
+++ b/Engine/Source/Runtime/Core/Public/UObject/UnrealNames.inl
@@ -218,6 +218,7 @@ REGISTER_NAME(291,MeshPort)
 REGISTER_NAME(292,MeshNetDriver)
 REGISTER_NAME(293,LiveStreamVoice)
 REGISTER_NAME(294,LiveStreamAnimation)
+REGISTER_NAME(295,DataStream)
 
 // Texture settings.
 REGISTER_NAME(300,Linear)
