Logging: Improved UE_LOG based on UE_LOGFMT

- Updated UE_LOG, UE_CLOG, UE_LOG_REF, UE_LOG_CLINKAGE to store constant data in a static constexpr struct.
- Moved log tracing out of the above log macros.
- Added FMsg::LogV to log with args in a va_list.

Generated Code by Field/Param Count
- i=instructions; b=bytes
- 0: 5i/30b (was 38i/169b)
- 1: 6i/33b (was 48i/211b)
- 2: 7i/39b (was 57i/253b)
- 3: 8i/46b (was 68i/299b)

Build Impact (Average of 1,000 Logs with 0-9 Fields/Params)
- Link Time: ~0ms (was ~0ms)
- Compile Time: 0.55ms (was 1.28ms)
- EXE Size: 138 bytes (was 436 bytes)
- PDB Size: 336 bytes (was 2,449 bytes)
- OBJ Size: 1,182 bytes (was 4,745 bytes)

For an example project, win64 test client exe size dropped by 2.7% and pdb size dropped by 3.1%.

diff --git a/Engine/Source/Runtime/Core/Private/Logging/LogMacros.cpp b/Engine/Source/Runtime/Core/Private/Logging/LogMacros.cpp
--- a/Engine/Source/Runtime/Core/Private/Logging/LogMacros.cpp
+++ b/Engine/Source/Runtime/Core/Private/Logging/LogMacros.cpp
@@ -49,7 +49,7 @@ void FMsg::LogfImpl(const ANSICHAR* File, int32 Line, const FLogCategoryName& Ca
 #endif
 }
 
-void FMsg::Logf_InternalImpl(const ANSICHAR* File, int32 Line, const FLogCategoryName& Category, ELogVerbosity::Type Verbosity, const TCHAR* Fmt, ...)
+void FMsg::LogV(const ANSICHAR* File, int32 Line, const FLogCategoryName& Category, ELogVerbosity::Type Verbosity, const TCHAR* Fmt, va_list Args)
 {
 #if !NO_LOGGING
 	QUICK_SCOPE_CYCLE_COUNTER(STAT_FMsgLogf);
@@ -57,33 +57,40 @@ void FMsg::Logf_InternalImpl(const ANSICHAR* File, int32 Line, const FLogCategor
 
 	if (LIKELY(Verbosity != ELogVerbosity::Fatal))
 	{
-		// SetColour is routed to GWarn just like the other verbosities and handled in the
-		// device that does the actual printing.
-		FOutputDevice* LogOverride = nullptr;
+		TStringBuilder<512> Buffer;
+		Buffer.AppendV(Fmt, Args);
+		const TCHAR* Message = *Buffer;
+		FOutputDevice* OutputDevice = nullptr;
 		switch (Verbosity)
 		{
 		case ELogVerbosity::Error:
 		case ELogVerbosity::Warning:
 		case ELogVerbosity::Display:
 		case ELogVerbosity::SetColor:
-			LogOverride = GWarn;
+			OutputDevice = GWarn;
 			break;
 		default:
 			break;
 		}
-		GROWABLE_LOGF(LogOverride	? LogOverride->Log(Category, Verbosity, Buffer)
-									: GLog->RedirectLog(Category, Verbosity, Buffer))
+		(OutputDevice ? OutputDevice : GLog)->Serialize(Message, Verbosity, Category);
 	}
 	else
 	{
-		va_list Args;
-		va_start(Args, Fmt);
 		StaticFailDebugV(TEXT("Fatal error:"), "", File, Line, /*bIsEnsure*/ false, PLATFORM_RETURN_ADDRESS(), Fmt, Args);
-		va_end(Args);
 	}
 #endif
 }
 
+void FMsg::Logf_InternalImpl(const ANSICHAR* File, int32 Line, const FLogCategoryName& Category, ELogVerbosity::Type Verbosity, const TCHAR* Fmt, ...)
+{
+#if !NO_LOGGING
+	va_list Args;
+	va_start(Args, Fmt);
+	LogV(File, Line, Category, Verbosity, Fmt, Args);
+	va_end(Args);
+#endif
+}
+
 /** Sends a formatted message to a remote tool. */
 void VARARGS FMsg::SendNotificationStringfImpl( const TCHAR *Fmt, ... )
 {
diff --git a/Engine/Source/Runtime/Core/Private/Logging/StructuredLog.cpp b/Engine/Source/Runtime/Core/Private/Logging/StructuredLog.cpp
--- a/Engine/Source/Runtime/Core/Private/Logging/StructuredLog.cpp
+++ b/Engine/Source/Runtime/Core/Private/Logging/StructuredLog.cpp
@@ -22,6 +22,16 @@
 #include "Serialization/CompactBinaryValue.h"
 #include "Serialization/VarInt.h"
 
+void StaticFailDebugV(
+	const TCHAR* Error,
+	const ANSICHAR* Expression,
+	const ANSICHAR* File,
+	int32 Line,
+	bool bIsEnsure,
+	void* ProgramCounter,
+	const TCHAR* DescriptionFormat,
+	va_list DescriptionArgs);
+
 namespace UE::Logging::Private
 {
 
@@ -726,7 +736,7 @@ FORCENOINLINE static FLogTemplate& CreateLogTemplate(const StaticLogRecordType&
 template <typename StaticLogRecordType>
 inline static FLogTemplate& EnsureLogTemplate(const StaticLogRecordType& Log, const FLogField* Fields, const int32 FieldCount)
 {
-	if (FLogTemplate* Template = Log.DynamicData.Template.load(std::memory_order_acquire))
+	if (FLogTemplate* Template = Log.DynamicData.Template.load(std::memory_order_acquire); LIKELY(Template))
 	{
 		return *Template;
 	}
@@ -760,15 +770,19 @@ inline static void DispatchLogRecord(const StaticLogRecordType& Log, const FLogR
 	}
 #endif
 
+	FOutputDevice* OutputDevice = nullptr;
 	switch (Log.Verbosity)
 	{
 	case ELogVerbosity::Error:
 	case ELogVerbosity::Warning:
 	case ELogVerbosity::Display:
-		return GWarn->SerializeRecord(Record);
+	case ELogVerbosity::SetColor:
+		OutputDevice = GWarn;
+		break;
 	default:
-		return GLog->SerializeRecord(Record);
+		break;
 	}
+	(OutputDevice ? OutputDevice : GLog)->SerializeRecord(Record);
 }
 
 void LogWithFieldArray(const FStaticLogRecord& Log, const FLogField* Fields, const int32 FieldCount)
@@ -798,4 +812,143 @@ void LogWithNoFields(const FStaticLocalizedLogRecord& Log)
 
 } // UE::Logging::Private
 
+///////////////////////////////////////////////////////////////////////////////////////////////////////////////////////
+
+namespace UE::Logging::Private
+{
+
+static constexpr const TCHAR* const GStaticBasicLogFormat = TEXT("{Message}");
+
+class FStaticBasicLogTemplateManager
+{
+public:
+	~FStaticBasicLogTemplateManager()
+	{
+		if (FLogTemplate* LocalTemplate = Template.exchange(nullptr))
+		{
+			FLogTemplate::Destroy(LocalTemplate);
+		}
+	}
+
+	inline FLogTemplate& EnsureTemplate()
+	{
+		if (FLogTemplate* LocalTemplate = Template.load(std::memory_order_acquire); LIKELY(LocalTemplate))
+		{
+			return *LocalTemplate;
+		}
+		return CreateTemplate();
+	}
+
+private:
+	FORCENOINLINE FLogTemplate& CreateTemplate()
+	{
+		FLogTemplate* NewTemplate = FLogTemplate::Create(GStaticBasicLogFormat);
+		if (FLogTemplate* ExistingTemplate = nullptr;
+			UNLIKELY(!Template.compare_exchange_strong(ExistingTemplate, NewTemplate, std::memory_order_release, std::memory_order_acquire)))
+		{
+			FLogTemplate::Destroy(NewTemplate);
+			return *ExistingTemplate;
+		}
+		return *NewTemplate;
+	}
+
+	std::atomic<FLogTemplate*> Template = nullptr;
+};
+
+static FStaticBasicLogTemplateManager GStaticBasicLogTemplateManager;
+
+FORCENOINLINE static void OutputBasicLogMessageSpec(const FLogCategoryBase& Category, const FStaticBasicLogRecord& Log)
+{
+#if LOGTRACE_ENABLED
+	FLogTrace::OutputLogMessageSpec(&Log, &Category, Log.Verbosity, Log.File, Log.Line, TEXT("%s"));
+#endif
+	Log.DynamicData.bInitialized.store(true, std::memory_order_release);
+}
+
+inline static void EnsureBasicLogMessageSpec(const FLogCategoryBase& Category, const FStaticBasicLogRecord& Log)
+{
+	if (!Log.DynamicData.bInitialized.load(std::memory_order_acquire))
+	{
+		OutputBasicLogMessageSpec(Category, Log);
+	}
+}
+
+// Serializing the log to compact binary happens in its own function because that allows stack
+// space for the writer to be returned before calling into the output devices.
+FORCENOINLINE static FCbObject SerializeBasicLogMessage(const FStaticBasicLogRecord& Log, va_list Args)
+{
+	TStringBuilder<512> Message;
+	Message.AppendV(Log.Format, Args);
+
+#if LOGTRACE_ENABLED
+	FLogTrace::OutputLogMessageSimple(&Log, TEXT("%s"), *Message);
+#endif
+
+	TCbWriter<512> Writer;
+	Writer.BeginObject();
+	Writer.AddString(ANSITEXTVIEW("Message"), Message);
+	Writer.EndObject();
+	return Writer.Save().AsObject();
+}
+
+void BasicLog(const FLogCategoryBase& Category, const FStaticBasicLogRecord* Log, ...)
+{
+#if !NO_LOGGING
+	EnsureBasicLogMessageSpec(Category, *Log);
+
+#if 0 // Enable to use structured logging for these logs. Experimental.
+	va_list Args;
+	va_start(Args, Log);
+	FCbObject Fields = SerializeBasicLogMessage(*Log, Args);
+	va_end(Args);
+
+	FLogRecord Record;
+	Record.SetFormat(GStaticBasicLogFormat);
+	Record.SetTemplate(&GStaticBasicLogTemplateManager.EnsureTemplate());
+	Record.SetFields(MoveTemp(Fields));
+	Record.SetFile(Log->File);
+	Record.SetLine(Log->Line);
+	Record.SetCategory(Category.GetCategoryName());
+	Record.SetVerbosity(Log->Verbosity);
+	Record.SetTime(FLogTime::Now());
+
+	FOutputDevice* OutputDevice = nullptr;
+	switch (Log->Verbosity)
+	{
+	case ELogVerbosity::Error:
+	case ELogVerbosity::Warning:
+	case ELogVerbosity::Display:
+	case ELogVerbosity::SetColor:
+		OutputDevice = GWarn;
+		break;
+	default:
+		break;
+	}
+	(OutputDevice ? OutputDevice : GLog)->SerializeRecord(Record);
+#else
+	va_list Args;
+	va_start(Args, Log);
+	FMsg::LogV(Log->File, Log->Line, Category.GetCategoryName(), Log->Verbosity, Log->Format, Args);
+	va_end(Args);
+#endif
+#endif
+}
+
+void BasicFatalLog(const FLogCategoryBase& Category, const FStaticBasicLogRecord* Log, ...)
+{
+#if !NO_LOGGING
+	va_list Args;
+	va_start(Args, Log);
+	StaticFailDebugV(TEXT("Fatal error:"), "", Log->File, Log->Line, /*bIsEnsure*/ false, PLATFORM_RETURN_ADDRESS(), Log->Format, Args);
+	va_end(Args);
+	va_start(Args, Log);
+	FDebug::AssertFailedV("", Log->File, Log->Line, Log->Format, Args);
+	va_end(Args);
+	UE_DEBUG_BREAK_AND_PROMPT_FOR_REMOTE();
+	FDebug::ProcessFatalError(PLATFORM_RETURN_ADDRESS());
+#endif
+}
+
+} // UE::Logging::Private
+
 #endif // !NO_LOGGING
diff --git a/Engine/Source/Runtime/Core/Public/Logging/LogMacros.h b/Engine/Source/Runtime/Core/Public/Logging/LogMacros.h
--- a/Engine/Source/Runtime/Core/Public/Logging/LogMacros.h
+++ b/Engine/Source/Runtime/Core/Public/Logging/LogMacros.h
@@ -18,6 +18,7 @@
 #include "Templates/IsValidVariadicFunctionArg.h"
 #include "Traits/IsCharEncodingCompatibleWith.h"
 
+#include <type_traits>
 
 /*----------------------------------------------------------------------------
 	Logging
@@ -30,10 +31,10 @@
  * FMsg 
  * This struct contains functions for messaging with tools or debug logs.
  **/
-struct CORE_API FMsg
+struct FMsg
 {
 	/** Sends a message to a remote tool. */
-	static void SendNotificationString( const TCHAR* Message );
+	CORE_API static void SendNotificationString( const TCHAR* Message );
 
 	/** Sends a formatted message to a remote tool. */
 	template <typename FmtType, typename... Types>
@@ -65,10 +66,12 @@ struct CORE_API FMsg
 		Logf_InternalImpl(File, Line, Category, Verbosity, (const TCHAR*)Fmt, Args...);
 	}
 
+	CORE_API static void LogV(const ANSICHAR* File, int32 Line, const FLogCategoryName& Category, ELogVerbosity::Type Verbosity, const TCHAR* Fmt, va_list Args);
+
 private:
-	static void VARARGS LogfImpl(const ANSICHAR* File, int32 Line, const FLogCategoryName& Category, ELogVerbosity::Type Verbosity, const TCHAR* Fmt, ...);
-	static void VARARGS Logf_InternalImpl(const ANSICHAR* File, int32 Line, const FLogCategoryName& Category, ELogVerbosity::Type Verbosity, const TCHAR* Fmt, ...);
-	static void VARARGS SendNotificationStringfImpl(const TCHAR* Fmt, ...);
+	CORE_API static void VARARGS LogfImpl(const ANSICHAR* File, int32 Line, const FLogCategoryName& Category, ELogVerbosity::Type Verbosity, const TCHAR* Fmt, ...);
+	CORE_API static void VARARGS Logf_InternalImpl(const ANSICHAR* File, int32 Line, const FLogCategoryName& Category, ELogVerbosity::Type Verbosity, const TCHAR* Fmt, ...);
+	CORE_API static void VARARGS SendNotificationStringfImpl(const TCHAR* Fmt, ...);
 };
 
 /*----------------------------------------------------------------------------
@@ -101,6 +104,45 @@ private:
 	#define UE_LOG_SOURCE_FILE(File) "Unknown"
 #endif
 
+namespace UE::Logging::Private
+{
+
+/** Data about a static basic log that is created on-demand. */
+struct FStaticBasicLogDynamicData
+{
+	std::atomic<bool> bInitialized = false;
+};
+
+/** Data about a static basic log that is constant for every occurrence. */
+struct FStaticBasicLogRecord
+{
+	const TCHAR* Format = nullptr;
+	const ANSICHAR* File = nullptr;
+	int32 Line = 0;
+	ELogVerbosity::Type Verbosity = ELogVerbosity::Log;
+	FStaticBasicLogDynamicData& DynamicData;
+
+	// Workaround for https://developercommunity.visualstudio.com/t/Incorrect-warning-C4700-with-unrelated-s/10285950
+	constexpr FStaticBasicLogRecord(
+		const TCHAR* InFormat,
+		const ANSICHAR* InFile,
+		int32 InLine,
+		ELogVerbosity::Type InVerbosity,
+		FStaticBasicLogDynamicData& InDynamicData)
+		: Format(InFormat)
+		, File(InFile)
+		, Line(InLine)
+		, Verbosity(InVerbosity)
+		, DynamicData(InDynamicData)
+	{
+	}
+};
+
+CORE_API void BasicLog(const FLogCategoryBase& Category, const FStaticBasicLogRecord* Log, ...);
+CORE_API void BasicFatalLog(const FLogCategoryBase& Category, const FStaticBasicLogRecord* Log, ...);
+
+} // UE::Logging::Private
+
 #if NO_LOGGING
 
 	struct FNoLoggingCategory {};
@@ -108,38 +150,38 @@ private:
 	// This will only log Fatal errors
 	#define UE_LOG(CategoryName, Verbosity, Format, ...) \
 	{ \
-		if (ELogVerbosity::Verbosity == ELogVerbosity::Fatal) \
+		if constexpr (ELogVerbosity::Verbosity == ELogVerbosity::Fatal) \
 		{ \
 			LowLevelFatalError(Format, ##__VA_ARGS__); \
-			UE_LOG_EXPAND_IS_FATAL(Verbosity, CA_ASSUME(false);, PREPROCESSOR_NOTHING) \
+			CA_ASSUME(false); \
 		} \
 	}
-	#define UE_LOG_REF(ConstLogCategoryRef, Verbosity, Format, ...) \
+	#define UE_LOG_REF(CategoryRef, Verbosity, Format, ...) \
 	{ \
-		if (ELogVerbosity::Verbosity == ELogVerbosity::Fatal) \
+		if constexpr (ELogVerbosity::Verbosity == ELogVerbosity::Fatal) \
 		{ \
 			LowLevelFatalError(Format, ##__VA_ARGS__); \
-			UE_LOG_EXPAND_IS_FATAL(Verbosity, CA_ASSUME(false);, PREPROCESSOR_NOTHING) \
+			CA_ASSUME(false); \
 		} \
 	}
 	#define UE_LOG_CLINKAGE(CategoryName, Verbosity, Format, ...) \
 	{ \
-		if (ELogVerbosity::Verbosity == ELogVerbosity::Fatal) \
+		if constexpr (ELogVerbosity::Verbosity == ELogVerbosity::Fatal) \
 		{ \
 			LowLevelFatalError(Format, ##__VA_ARGS__); \
-			UE_LOG_EXPAND_IS_FATAL(Verbosity, CA_ASSUME(false);, PREPROCESSOR_NOTHING) \
+			CA_ASSUME(false); \
 		} \
 	}
 
 	// Conditional logging (fatal errors only).
 	#define UE_CLOG(Condition, CategoryName, Verbosity, Format, ...) \
 	{ \
-		if (ELogVerbosity::Verbosity == ELogVerbosity::Fatal) \
+		if constexpr (ELogVerbosity::Verbosity == ELogVerbosity::Fatal) \
 		{ \
 			if (Condition) \
 			{ \
 				LowLevelFatalError(Format, ##__VA_ARGS__); \
-				UE_LOG_EXPAND_IS_FATAL(Verbosity, CA_ASSUME(false);, PREPROCESSOR_NOTHING) \
+				CA_ASSUME(false); \
 			} \
 		} \
 	}
@@ -194,91 +236,31 @@ private:
 		#define UE_VALIDATE_FORMAT_STRING(Format, ...)
 	#endif
 
-	/** 
-	 * INTERNAL IMPLEMENTATION. DO NOT CALL DIRECTLY!
-	**/
-	#define UE_INTERNAL_LOG_IMPL(CategoryName, Verbosity, Format, ...) \
-		static_assert(std::is_const_v<std::remove_reference_t<decltype(Format)>>, "Formatting string must be a const TCHAR array."); \
-		static_assert(TIsArrayOrRefOfTypeByPredicate<decltype(Format), TIsCharEncodingCompatibleWithTCHAR>::Value, "Formatting string must be a TCHAR array."); \
-		static_assert((ELogVerbosity::Verbosity & ELogVerbosity::VerbosityMask) < ELogVerbosity::NumVerbosity && ELogVerbosity::Verbosity > 0, "Verbosity must be constant and in range."); \
-		UE_LOG_EXPAND_IS_FATAL(Verbosity, PREPROCESSOR_NOTHING, if (!CategoryName.IsSuppressed(ELogVerbosity::Verbosity))) \
-			{ \
-				UE_VALIDATE_FORMAT_STRING(Format, ##__VA_ARGS__); \
-				DispatchCheckVerify([] (const auto& LCategoryName, const auto& LFormat, const auto&... UE_LOG_Args) UE_DEBUG_SECTION \
-				{ \
-					TRACE_LOG_MESSAGE(LCategoryName, Verbosity, LFormat, UE_LOG_Args...) \
-					UE_LOG_EXPAND_IS_FATAL(Verbosity, \
-						{ \
-							FMsg::Logf_Internal(UE_LOG_SOURCE_FILE(__FILE__), __LINE__, LCategoryName.GetCategoryName(), ELogVerbosity::Verbosity, LFormat, UE_LOG_Args...); \
-							UE_DEBUG_BREAK_AND_PROMPT_FOR_REMOTE(); \
-							FDebug::ProcessFatalError(PLATFORM_RETURN_ADDRESS()); \
-						}, \
-						{ \
-							FMsg::Logf_Internal(nullptr, 0, LCategoryName.GetCategoryName(), ELogVerbosity::Verbosity, LFormat, UE_LOG_Args...); \
-						} \
-					) \
-				}, CategoryName, Format, ##__VA_ARGS__); \
-				UE_LOG_EXPAND_IS_FATAL(Verbosity, CA_ASSUME(false);, PREPROCESSOR_NOTHING) \
-			} \
-
 	/**
-	 * A  macro that outputs a formatted message to log if a given logging category is active at a given verbosity level
-	 * @param Category  A reference to a FLogCategoryBase instance
-	 * @param Verbosity Verbosity level to test against
-	 * @param Format    Format text
-	 ***/
-	#define UE_LOG_REF(ConstLogCategoryRef, Verbosity, Format, ...) \
-	{ \
-		CA_CONSTANT_IF((ELogVerbosity::Verbosity & ELogVerbosity::VerbosityMask) <= ELogVerbosity::COMPILED_IN_MINIMUM_VERBOSITY && (ELogVerbosity::Warning & ELogVerbosity::VerbosityMask) <= ConstLogCategoryRef.GetCompileTimeVerbosity()) \
-		{ \
-			UE_INTERNAL_LOG_IMPL(ConstLogCategoryRef, Verbosity, Format, ##__VA_ARGS__); \
-		} \
-	}
-
-	/** 
-	 * A  macro that outputs a formatted message to log if a given logging category is active at a given verbosity level
-	 * @param CategoryName name of the logging category
-	 * @param Verbosity, verbosity level to test against
-	 * @param Format, format text
-	 ***/
+	 * A macro that logs a formatted message if the log category is active at the requested verbosity level.
+	 *
+	 * @param CategoryName   Name of the log category as provided to DEFINE_LOG_CATEGORY.
+	 * @param Verbosity      Verbosity level of this message. See ELogVerbosity.
+	 * @param Format         Format string literal in the style of printf.
+	 */
 	#define UE_LOG(CategoryName, Verbosity, Format, ...) \
-	{ \
-		CA_CONSTANT_IF((ELogVerbosity::Verbosity & ELogVerbosity::VerbosityMask) <= ELogVerbosity::COMPILED_IN_MINIMUM_VERBOSITY && (ELogVerbosity::Warning & ELogVerbosity::VerbosityMask) <= FLogCategory##CategoryName::CompileTimeVerbosity) \
-		{ \
-			UE_INTERNAL_LOG_IMPL(CategoryName, Verbosity, Format, ##__VA_ARGS__); \
-		} \
-	}
+		UE_PRIVATE_LOG(PREPROCESSOR_NOTHING, constexpr, CategoryName, Verbosity, Format, ##__VA_ARGS__)
+
+	/**
+	 * DO NOT USE. A macro that logs a formatted message if the log category is active at the requested verbosity level.
+	 *
+	 * @note This does not trace the category correctly and will be deprecated in a future release.
+	 *
+	 * @param CategoryRef   A reference to an instance of FLogCategoryBase.
+	 * @param Verbosity     Verbosity level of this message. See ELogVerbosity.
+	 * @param Format        Format string literal in the style of printf.
+	 */
+	#define UE_LOG_REF(CategoryRef, Verbosity, Format, ...) \
+		UE_PRIVATE_LOG(PREPROCESSOR_NOTHING, PREPROCESSOR_NOTHING, CategoryRef, Verbosity, Format, ##__VA_ARGS__)
+
+	// DO NOT USE. Use UE_LOG because this will be deprecated in a future release.
+	#define UE_LOG_CLINKAGE UE_LOG
 
-	/** 
-	 * A  macro that outputs a formatted message to log if a given logging category is active at a given verbosity level
-	 * @param CategoryName name of the logging category
-	 * @param Verbosity, verbosity level to test against
-	 * @param Format, format text
-	 ***/
-	#define UE_LOG_CLINKAGE(CategoryName, Verbosity, Format, ...) \
-	{ \
-		static_assert(std::is_const_v<std::remove_reference_t<decltype(Format)>>, "Formatting string must be a const TCHAR array."); \
-		static_assert(TIsArrayOrRefOfTypeByPredicate<decltype(Format), TIsCharEncodingCompatibleWithTCHAR>::Value, "Formatting string must be a TCHAR array."); \
-		static_assert((ELogVerbosity::Verbosity & ELogVerbosity::VerbosityMask) < ELogVerbosity::NumVerbosity && ELogVerbosity::Verbosity > 0, "Verbosity must be constant and in range."); \
-		CA_CONSTANT_IF((ELogVerbosity::Verbosity & ELogVerbosity::VerbosityMask) <= ELogVerbosity::COMPILED_IN_MINIMUM_VERBOSITY && (ELogVerbosity::Warning & ELogVerbosity::VerbosityMask) <= FLogCategory##CategoryName::CompileTimeVerbosity) \
-		{ \
-			UE_LOG_EXPAND_IS_FATAL(Verbosity, PREPROCESSOR_NOTHING, if (!CategoryName.IsSuppressed(ELogVerbosity::Verbosity))) \
-			{ \
-				TRACE_LOG_MESSAGE(CategoryName, Verbosity, Format, ##__VA_ARGS__) \
-				UE_LOG_EXPAND_IS_FATAL(Verbosity, \
-					{ \
-						FMsg::Logf_Internal(UE_LOG_SOURCE_FILE(__FILE__), __LINE__, CategoryName.GetCategoryName(), ELogVerbosity::Verbosity, Format,  ##__VA_ARGS__); \
-						UE_DEBUG_BREAK_AND_PROMPT_FOR_REMOTE(); \
-						FDebug::ProcessFatalError(PLATFORM_RETURN_ADDRESS()); \
-						CA_ASSUME(false); \
-					}, \
-					{ \
-						FMsg::Logf_Internal(nullptr, 0, CategoryName.GetCategoryName(), ELogVerbosity::Verbosity, Format,  ##__VA_ARGS__); \
-					} \
-				) \
-			} \
-		} \
-	}
 	/**
 	* A  macro that outputs a formatted message to the log specifically used for security events
 	* @param NetConnection, a valid UNetConnection
@@ -299,34 +281,46 @@ private:
 		} \
 	}
 
-	// Conditional logging. Will only log if Condition is met.
+	/**
+	 * A macro that conditionally logs a formatted message if the log category is active at the requested verbosity level.
+	 *
+	 * @note The condition is not evaluated unless the log category is active at the requested verbosity level.
+	 *
+	 * @param Condition      Condition that must evaluate to true in order for the message to be logged.
+	 * @param CategoryName   Name of the log category as provided to DEFINE_LOG_CATEGORY.
+	 * @param Verbosity      Verbosity level of this message. See ELogVerbosity.
+	 * @param Format         Format string literal in the style of printf.
+	 */
 	#define UE_CLOG(Condition, CategoryName, Verbosity, Format, ...) \
+		UE_PRIVATE_LOG(if (Condition), constexpr, CategoryName, Verbosity, Format, ##__VA_ARGS__)
+
+	/** Private macro used to implement the public log macros. DO NOT CALL DIRECTLY! */
+	#define UE_PRIVATE_LOG(Condition, CategoryConst, Category, Verbosity, Format, ...) \
 	{ \
 		static_assert(std::is_const_v<std::remove_reference_t<decltype(Format)>>, "Formatting string must be a const TCHAR array."); \
 		static_assert(TIsArrayOrRefOfTypeByPredicate<decltype(Format), TIsCharEncodingCompatibleWithTCHAR>::Value, "Formatting string must be a TCHAR array."); \
-		static_assert((ELogVerbosity::Verbosity & ELogVerbosity::VerbosityMask) < ELogVerbosity::NumVerbosity && ELogVerbosity::Verbosity > 0, "Verbosity must be constant and in range."); \
-		CA_CONSTANT_IF((ELogVerbosity::Verbosity & ELogVerbosity::VerbosityMask) <= ELogVerbosity::COMPILED_IN_MINIMUM_VERBOSITY && (ELogVerbosity::Warning & ELogVerbosity::VerbosityMask) <= FLogCategory##CategoryName::CompileTimeVerbosity) \
+		UE_VALIDATE_FORMAT_STRING(Format, ##__VA_ARGS__); \
+		static ::UE::Logging::Private::FStaticBasicLogDynamicData LOG_Dynamic; \
+		static constexpr ::UE::Logging::Private::FStaticBasicLogRecord LOG_Static(Format, __FILE__, __LINE__, ::ELogVerbosity::Verbosity, LOG_Dynamic); \
+		static_assert((::ELogVerbosity::Verbosity & ::ELogVerbosity::VerbosityMask) < ::ELogVerbosity::NumVerbosity && ::ELogVerbosity::Verbosity > 0, "Verbosity must be constant and in range."); \
+		if constexpr (::ELogVerbosity::Verbosity == ::ELogVerbosity::Fatal) \
 		{ \
-			UE_LOG_EXPAND_IS_FATAL(Verbosity, PREPROCESSOR_NOTHING, if (!CategoryName.IsSuppressed(ELogVerbosity::Verbosity))) \
+			Condition \
 			{ \
-				if (Condition) \
+				::UE::Logging::Private::BasicFatalLog(Category, &LOG_Static, ##__VA_ARGS__); \
+				CA_ASSUME(false); \
+			} \
+		} \
+		else if constexpr ((::ELogVerbosity::Verbosity & ::ELogVerbosity::VerbosityMask) <= ::ELogVerbosity::COMPILED_IN_MINIMUM_VERBOSITY) \
+		{ \
+			if CategoryConst ((::ELogVerbosity::Verbosity & ::ELogVerbosity::VerbosityMask) <= Category.GetCompileTimeVerbosity()) \
+			{ \
+				if (!Category.IsSuppressed(::ELogVerbosity::Verbosity)) \
 				{ \
-					DispatchCheckVerify([] (const auto& LCategoryName, const auto& LFormat, const auto&... UE_LOG_Args) UE_DEBUG_SECTION \
+					Condition \
 					{ \
-						TRACE_LOG_MESSAGE(LCategoryName, Verbosity, LFormat, UE_LOG_Args...) \
-						UE_LOG_EXPAND_IS_FATAL(Verbosity, \
-							{ \
-								FMsg::Logf_Internal(UE_LOG_SOURCE_FILE(__FILE__), __LINE__, LCategoryName.GetCategoryName(), ELogVerbosity::Verbosity, LFormat, UE_LOG_Args...); \
-								UE_DEBUG_BREAK_AND_PROMPT_FOR_REMOTE(); \
-								FDebug::ProcessFatalError(PLATFORM_RETURN_ADDRESS()); \
-							}, \
-							{ \
-								FMsg::Logf_Internal(nullptr, 0, LCategoryName.GetCategoryName(), ELogVerbosity::Verbosity, LFormat, UE_LOG_Args...); \
-							} \
-						) \
-						CA_ASSUME(true); \
-					}, CategoryName, Format, ##__VA_ARGS__); \
-					UE_LOG_EXPAND_IS_FATAL(Verbosity, CA_ASSUME(false);, PREPROCESSOR_NOTHING) \
+						::UE::Logging::Private::BasicLog(Category, &LOG_Static, ##__VA_ARGS__); \
+					} \
 				} \
 			} \
 		} \
@@ -343,7 +337,7 @@ private:
 	#define UE_SUPPRESS(CategoryName, Verbosity, ExecuteIfUnsuppressed) \
 	{ \
 		static_assert((ELogVerbosity::Verbosity & ELogVerbosity::VerbosityMask) < ELogVerbosity::NumVerbosity && ELogVerbosity::Verbosity > 0, "Verbosity must be constant and in range."); \
-		CA_CONSTANT_IF((ELogVerbosity::Verbosity & ELogVerbosity::VerbosityMask) <= ELogVerbosity::COMPILED_IN_MINIMUM_VERBOSITY && (ELogVerbosity::Warning & ELogVerbosity::VerbosityMask) <= FLogCategory##CategoryName::CompileTimeVerbosity) \
+		CA_CONSTANT_IF((ELogVerbosity::Verbosity & ELogVerbosity::VerbosityMask) <= ELogVerbosity::COMPILED_IN_MINIMUM_VERBOSITY && (ELogVerbosity::Verbosity & ELogVerbosity::VerbosityMask) <= FLogCategory##CategoryName::CompileTimeVerbosity) \
 		{ \
 			if (!CategoryName.IsSuppressed(ELogVerbosity::Verbosity)) \
 			{ \
diff --git a/Engine/Source/Runtime/Experimental/Chaos/Public/Chaos/Math/Krylov.h b/Engine/Source/Runtime/Experimental/Chaos/Public/Chaos/Math/Krylov.h
--- a/Engine/Source/Runtime/Experimental/Chaos/Public/Chaos/Math/Krylov.h
+++ b/Engine/Source/Runtime/Experimental/Chaos/Public/Chaos/Math/Krylov.h
@@ -236,7 +236,7 @@ void LanczosCG(
 	set(x, c);
 	scale(x, p);
 	AXPY(v, -alpha, q);
-	T residual;
+	T residual{};
 	TArray<TV> y; y.SetNum(x.Num());
 	for (int it = 1; it < max_it; it++) 
 	{
