fix several issues with MatrixInverse
de-copy-paste the scalar implementation of MatrixInverse that was in 3 places
remove scalar implementations from VectorMatrixInverse
make sure Determinant is only computed once, not twice
do the validity check on Determinant on the actual value that will be used for inverse
make MatrixInverse return bool directly

diff --git a/Engine/Source/Runtime/Core/Private/Math/UnrealMath.cpp b/Engine/Source/Runtime/Core/Private/Math/UnrealMath.cpp
--- a/Engine/Source/Runtime/Core/Private/Math/UnrealMath.cpp
+++ b/Engine/Source/Runtime/Core/Private/Math/UnrealMath.cpp
@@ -816,7 +816,203 @@ namespace Math
 		T XYZ = (T)TVector<T>::DotProduct(TwistAxis, TVector<T>(X, Y, Z));
 		return FMath::UnwindRadians((T)2.0f * FMath::Atan2(XYZ, W));
 	}
+	
+} // namespace UE::Math
+}
+
+	
+/**
+ * Calculate the inverse of an FMatrix44.
+ *
+ * @param DstMatrix		FMatrix44 pointer to where the result should be stored
+ * @param SrcMatrix		FMatrix44 pointer to the Matrix to be inversed
+ * @return bool			returns false if matrix is not invertable and stores identity 
+ *
+ * Src == Dst is allowed
+ */
+ // scalar implementation of MatrixInverse is here, used as the fallback for VectorMatrixInverse
+bool FMath::MatrixInverse(FMatrix44d* DstMatrix, const FMatrix44d* SrcMatrix)
+{
+	typedef double Double4x4[4][4];
+	const Double4x4& M = *((const Double4x4*)SrcMatrix);
+	Double4x4 Result;
+	double Det[4];
+	Double4x4 Tmp;
+
+	Tmp[0][0] = M[2][2] * M[3][3] - M[2][3] * M[3][2];
+	Tmp[0][1] = M[1][2] * M[3][3] - M[1][3] * M[3][2];
+	Tmp[0][2] = M[1][2] * M[2][3] - M[1][3] * M[2][2];
+
+	Tmp[1][0] = M[2][2] * M[3][3] - M[2][3] * M[3][2];
+	Tmp[1][1] = M[0][2] * M[3][3] - M[0][3] * M[3][2];
+	Tmp[1][2] = M[0][2] * M[2][3] - M[0][3] * M[2][2];
+
+	Tmp[2][0] = M[1][2] * M[3][3] - M[1][3] * M[3][2];
+	Tmp[2][1] = M[0][2] * M[3][3] - M[0][3] * M[3][2];
+	Tmp[2][2] = M[0][2] * M[1][3] - M[0][3] * M[1][2];
+
+	Tmp[3][0] = M[1][2] * M[2][3] - M[1][3] * M[2][2];
+	Tmp[3][1] = M[0][2] * M[2][3] - M[0][3] * M[2][2];
+	Tmp[3][2] = M[0][2] * M[1][3] - M[0][3] * M[1][2];
+
+	Det[0] = M[1][1] * Tmp[0][0] - M[2][1] * Tmp[0][1] + M[3][1] * Tmp[0][2];
+	Det[1] = M[0][1] * Tmp[1][0] - M[2][1] * Tmp[1][1] + M[3][1] * Tmp[1][2];
+	Det[2] = M[0][1] * Tmp[2][0] - M[1][1] * Tmp[2][1] + M[3][1] * Tmp[2][2];
+	Det[3] = M[0][1] * Tmp[3][0] - M[1][1] * Tmp[3][1] + M[2][1] * Tmp[3][2];
+
+	const double Determinant = M[0][0] * Det[0] - M[1][0] * Det[1] + M[2][0] * Det[2] - M[3][0] * Det[3];
+
+	if ( Determinant == 0.0 || !FMath::IsFinite(Determinant) )
+	{
+		*DstMatrix = FMatrix44d::Identity;
+		return false;
+	}
+
+	const double	RDet = 1.0 / Determinant;
+
+	Result[0][0] = RDet * Det[0];
+	Result[0][1] = -RDet * Det[1];
+	Result[0][2] = RDet * Det[2];
+	Result[0][3] = -RDet * Det[3];
+	Result[1][0] = -RDet * (M[1][0] * Tmp[0][0] - M[2][0] * Tmp[0][1] + M[3][0] * Tmp[0][2]);
+	Result[1][1] = RDet * (M[0][0] * Tmp[1][0] - M[2][0] * Tmp[1][1] + M[3][0] * Tmp[1][2]);
+	Result[1][2] = -RDet * (M[0][0] * Tmp[2][0] - M[1][0] * Tmp[2][1] + M[3][0] * Tmp[2][2]);
+	Result[1][3] = RDet * (M[0][0] * Tmp[3][0] - M[1][0] * Tmp[3][1] + M[2][0] * Tmp[3][2]);
+	Result[2][0] = RDet * (
+		M[1][0] * (M[2][1] * M[3][3] - M[2][3] * M[3][1]) -
+		M[2][0] * (M[1][1] * M[3][3] - M[1][3] * M[3][1]) +
+		M[3][0] * (M[1][1] * M[2][3] - M[1][3] * M[2][1])
+		);
+	Result[2][1] = -RDet * (
+		M[0][0] * (M[2][1] * M[3][3] - M[2][3] * M[3][1]) -
+		M[2][0] * (M[0][1] * M[3][3] - M[0][3] * M[3][1]) +
+		M[3][0] * (M[0][1] * M[2][3] - M[0][3] * M[2][1])
+		);
+	Result[2][2] = RDet * (
+		M[0][0] * (M[1][1] * M[3][3] - M[1][3] * M[3][1]) -
+		M[1][0] * (M[0][1] * M[3][3] - M[0][3] * M[3][1]) +
+		M[3][0] * (M[0][1] * M[1][3] - M[0][3] * M[1][1])
+		);
+	Result[2][3] = -RDet * (
+		M[0][0] * (M[1][1] * M[2][3] - M[1][3] * M[2][1]) -
+		M[1][0] * (M[0][1] * M[2][3] - M[0][3] * M[2][1]) +
+		M[2][0] * (M[0][1] * M[1][3] - M[0][3] * M[1][1])
+		);
+	Result[3][0] = -RDet * (
+		M[1][0] * (M[2][1] * M[3][2] - M[2][2] * M[3][1]) -
+		M[2][0] * (M[1][1] * M[3][2] - M[1][2] * M[3][1]) +
+		M[3][0] * (M[1][1] * M[2][2] - M[1][2] * M[2][1])
+		);
+	Result[3][1] = RDet * (
+		M[0][0] * (M[2][1] * M[3][2] - M[2][2] * M[3][1]) -
+		M[2][0] * (M[0][1] * M[3][2] - M[0][2] * M[3][1]) +
+		M[3][0] * (M[0][1] * M[2][2] - M[0][2] * M[2][1])
+		);
+	Result[3][2] = -RDet * (
+		M[0][0] * (M[1][1] * M[3][2] - M[1][2] * M[3][1]) -
+		M[1][0] * (M[0][1] * M[3][2] - M[0][2] * M[3][1]) +
+		M[3][0] * (M[0][1] * M[1][2] - M[0][2] * M[1][1])
+		);
+	Result[3][3] = RDet * (
+		M[0][0] * (M[1][1] * M[2][2] - M[1][2] * M[2][1]) -
+		M[1][0] * (M[0][1] * M[2][2] - M[0][2] * M[2][1]) +
+		M[2][0] * (M[0][1] * M[1][2] - M[0][2] * M[1][1])
+		);
+
+	memcpy(DstMatrix, &Result, 16 * sizeof(double));
+	return true;
 }
+
+bool FMath::MatrixInverse(FMatrix44f* DstMatrix, const FMatrix44f* SrcMatrix)
+{
+	typedef float Float4x4[4][4];
+	const Float4x4& M = *((const Float4x4*)SrcMatrix);
+	Float4x4 Result;
+	float Det[4];
+	Float4x4 Tmp;
+
+	Tmp[0][0] = M[2][2] * M[3][3] - M[2][3] * M[3][2];
+	Tmp[0][1] = M[1][2] * M[3][3] - M[1][3] * M[3][2];
+	Tmp[0][2] = M[1][2] * M[2][3] - M[1][3] * M[2][2];
+
+	Tmp[1][0] = M[2][2] * M[3][3] - M[2][3] * M[3][2];
+	Tmp[1][1] = M[0][2] * M[3][3] - M[0][3] * M[3][2];
+	Tmp[1][2] = M[0][2] * M[2][3] - M[0][3] * M[2][2];
+
+	Tmp[2][0] = M[1][2] * M[3][3] - M[1][3] * M[3][2];
+	Tmp[2][1] = M[0][2] * M[3][3] - M[0][3] * M[3][2];
+	Tmp[2][2] = M[0][2] * M[1][3] - M[0][3] * M[1][2];
+
+	Tmp[3][0] = M[1][2] * M[2][3] - M[1][3] * M[2][2];
+	Tmp[3][1] = M[0][2] * M[2][3] - M[0][3] * M[2][2];
+	Tmp[3][2] = M[0][2] * M[1][3] - M[0][3] * M[1][2];
+
+	Det[0] = M[1][1] * Tmp[0][0] - M[2][1] * Tmp[0][1] + M[3][1] * Tmp[0][2];
+	Det[1] = M[0][1] * Tmp[1][0] - M[2][1] * Tmp[1][1] + M[3][1] * Tmp[1][2];
+	Det[2] = M[0][1] * Tmp[2][0] - M[1][1] * Tmp[2][1] + M[3][1] * Tmp[2][2];
+	Det[3] = M[0][1] * Tmp[3][0] - M[1][1] * Tmp[3][1] + M[2][1] * Tmp[3][2];
+
+	const float Determinant = M[0][0] * Det[0] - M[1][0] * Det[1] + M[2][0] * Det[2] - M[3][0] * Det[3];
+	
+	if ( Determinant == 0.0f || !FMath::IsFinite(Determinant) )
+	{
+		*DstMatrix = FMatrix44f::Identity;
+		return false;
+	}
+
+	const float	RDet = 1.0f / Determinant;
+
+	Result[0][0] = RDet * Det[0];
+	Result[0][1] = -RDet * Det[1];
+	Result[0][2] = RDet * Det[2];
+	Result[0][3] = -RDet * Det[3];
+	Result[1][0] = -RDet * (M[1][0] * Tmp[0][0] - M[2][0] * Tmp[0][1] + M[3][0] * Tmp[0][2]);
+	Result[1][1] = RDet * (M[0][0] * Tmp[1][0] - M[2][0] * Tmp[1][1] + M[3][0] * Tmp[1][2]);
+	Result[1][2] = -RDet * (M[0][0] * Tmp[2][0] - M[1][0] * Tmp[2][1] + M[3][0] * Tmp[2][2]);
+	Result[1][3] = RDet * (M[0][0] * Tmp[3][0] - M[1][0] * Tmp[3][1] + M[2][0] * Tmp[3][2]);
+	Result[2][0] = RDet * (
+		M[1][0] * (M[2][1] * M[3][3] - M[2][3] * M[3][1]) -
+		M[2][0] * (M[1][1] * M[3][3] - M[1][3] * M[3][1]) +
+		M[3][0] * (M[1][1] * M[2][3] - M[1][3] * M[2][1])
+		);
+	Result[2][1] = -RDet * (
+		M[0][0] * (M[2][1] * M[3][3] - M[2][3] * M[3][1]) -
+		M[2][0] * (M[0][1] * M[3][3] - M[0][3] * M[3][1]) +
+		M[3][0] * (M[0][1] * M[2][3] - M[0][3] * M[2][1])
+		);
+	Result[2][2] = RDet * (
+		M[0][0] * (M[1][1] * M[3][3] - M[1][3] * M[3][1]) -
+		M[1][0] * (M[0][1] * M[3][3] - M[0][3] * M[3][1]) +
+		M[3][0] * (M[0][1] * M[1][3] - M[0][3] * M[1][1])
+		);
+	Result[2][3] = -RDet * (
+		M[0][0] * (M[1][1] * M[2][3] - M[1][3] * M[2][1]) -
+		M[1][0] * (M[0][1] * M[2][3] - M[0][3] * M[2][1]) +
+		M[2][0] * (M[0][1] * M[1][3] - M[0][3] * M[1][1])
+		);
+	Result[3][0] = -RDet * (
+		M[1][0] * (M[2][1] * M[3][2] - M[2][2] * M[3][1]) -
+		M[2][0] * (M[1][1] * M[3][2] - M[1][2] * M[3][1]) +
+		M[3][0] * (M[1][1] * M[2][2] - M[1][2] * M[2][1])
+		);
+	Result[3][1] = RDet * (
+		M[0][0] * (M[2][1] * M[3][2] - M[2][2] * M[3][1]) -
+		M[2][0] * (M[0][1] * M[3][2] - M[0][2] * M[3][1]) +
+		M[3][0] * (M[0][1] * M[2][2] - M[0][2] * M[2][1])
+		);
+	Result[3][2] = -RDet * (
+		M[0][0] * (M[1][1] * M[3][2] - M[1][2] * M[3][1]) -
+		M[1][0] * (M[0][1] * M[3][2] - M[0][2] * M[3][1]) +
+		M[3][0] * (M[0][1] * M[1][2] - M[0][2] * M[1][1])
+		);
+	Result[3][3] = RDet * (
+		M[0][0] * (M[1][1] * M[2][2] - M[1][2] * M[2][1]) -
+		M[1][0] * (M[0][1] * M[2][2] - M[0][2] * M[2][1]) +
+		M[2][0] * (M[0][1] * M[1][2] - M[0][2] * M[1][1])
+		);
+
+	memcpy(DstMatrix, &Result, 16 * sizeof(float));
+	return true;
 }
 
 
diff --git a/Engine/Source/Runtime/Core/Private/Math/UnrealMathSSE.cpp b/Engine/Source/Runtime/Core/Private/Math/UnrealMathSSE.cpp
--- a/Engine/Source/Runtime/Core/Private/Math/UnrealMathSSE.cpp
+++ b/Engine/Source/Runtime/Core/Private/Math/UnrealMathSSE.cpp
@@ -157,178 +157,4 @@ void VectorMatrixMultiply(FMatrix44d* Result, const FMatrix44d* Matrix1, const F
 }
 
 
-/**
- * Calculate the inverse of an FMatrix.
- *
- * @param DstMatrix		FMatrix pointer to where the result should be stored
- * @param SrcMatrix		FMatrix pointer to the Matrix to be inversed
- */
-void VectorMatrixInverse(FMatrix44f* DstMatrix, const FMatrix44f* SrcMatrix)
-{
-	typedef float Float4x4[4][4];
-	const Float4x4& M = *((const Float4x4*)SrcMatrix);
-	Float4x4 Result;
-	float Det[4];
-	Float4x4 Tmp;
-
-	Tmp[0][0] = M[2][2] * M[3][3] - M[2][3] * M[3][2];
-	Tmp[0][1] = M[1][2] * M[3][3] - M[1][3] * M[3][2];
-	Tmp[0][2] = M[1][2] * M[2][3] - M[1][3] * M[2][2];
-
-	Tmp[1][0] = M[2][2] * M[3][3] - M[2][3] * M[3][2];
-	Tmp[1][1] = M[0][2] * M[3][3] - M[0][3] * M[3][2];
-	Tmp[1][2] = M[0][2] * M[2][3] - M[0][3] * M[2][2];
-
-	Tmp[2][0] = M[1][2] * M[3][3] - M[1][3] * M[3][2];
-	Tmp[2][1] = M[0][2] * M[3][3] - M[0][3] * M[3][2];
-	Tmp[2][2] = M[0][2] * M[1][3] - M[0][3] * M[1][2];
-
-	Tmp[3][0] = M[1][2] * M[2][3] - M[1][3] * M[2][2];
-	Tmp[3][1] = M[0][2] * M[2][3] - M[0][3] * M[2][2];
-	Tmp[3][2] = M[0][2] * M[1][3] - M[0][3] * M[1][2];
-
-	Det[0] = M[1][1] * Tmp[0][0] - M[2][1] * Tmp[0][1] + M[3][1] * Tmp[0][2];
-	Det[1] = M[0][1] * Tmp[1][0] - M[2][1] * Tmp[1][1] + M[3][1] * Tmp[1][2];
-	Det[2] = M[0][1] * Tmp[2][0] - M[1][1] * Tmp[2][1] + M[3][1] * Tmp[2][2];
-	Det[3] = M[0][1] * Tmp[3][0] - M[1][1] * Tmp[3][1] + M[2][1] * Tmp[3][2];
-
-	const float Determinant = M[0][0] * Det[0] - M[1][0] * Det[1] + M[2][0] * Det[2] - M[3][0] * Det[3];
-	const float	RDet = 1.0f / Determinant;
-
-	Result[0][0] = RDet * Det[0];
-	Result[0][1] = -RDet * Det[1];
-	Result[0][2] = RDet * Det[2];
-	Result[0][3] = -RDet * Det[3];
-	Result[1][0] = -RDet * (M[1][0] * Tmp[0][0] - M[2][0] * Tmp[0][1] + M[3][0] * Tmp[0][2]);
-	Result[1][1] = RDet * (M[0][0] * Tmp[1][0] - M[2][0] * Tmp[1][1] + M[3][0] * Tmp[1][2]);
-	Result[1][2] = -RDet * (M[0][0] * Tmp[2][0] - M[1][0] * Tmp[2][1] + M[3][0] * Tmp[2][2]);
-	Result[1][3] = RDet * (M[0][0] * Tmp[3][0] - M[1][0] * Tmp[3][1] + M[2][0] * Tmp[3][2]);
-	Result[2][0] = RDet * (
-		M[1][0] * (M[2][1] * M[3][3] - M[2][3] * M[3][1]) -
-		M[2][0] * (M[1][1] * M[3][3] - M[1][3] * M[3][1]) +
-		M[3][0] * (M[1][1] * M[2][3] - M[1][3] * M[2][1])
-		);
-	Result[2][1] = -RDet * (
-		M[0][0] * (M[2][1] * M[3][3] - M[2][3] * M[3][1]) -
-		M[2][0] * (M[0][1] * M[3][3] - M[0][3] * M[3][1]) +
-		M[3][0] * (M[0][1] * M[2][3] - M[0][3] * M[2][1])
-		);
-	Result[2][2] = RDet * (
-		M[0][0] * (M[1][1] * M[3][3] - M[1][3] * M[3][1]) -
-		M[1][0] * (M[0][1] * M[3][3] - M[0][3] * M[3][1]) +
-		M[3][0] * (M[0][1] * M[1][3] - M[0][3] * M[1][1])
-		);
-	Result[2][3] = -RDet * (
-		M[0][0] * (M[1][1] * M[2][3] - M[1][3] * M[2][1]) -
-		M[1][0] * (M[0][1] * M[2][3] - M[0][3] * M[2][1]) +
-		M[2][0] * (M[0][1] * M[1][3] - M[0][3] * M[1][1])
-		);
-	Result[3][0] = -RDet * (
-		M[1][0] * (M[2][1] * M[3][2] - M[2][2] * M[3][1]) -
-		M[2][0] * (M[1][1] * M[3][2] - M[1][2] * M[3][1]) +
-		M[3][0] * (M[1][1] * M[2][2] - M[1][2] * M[2][1])
-		);
-	Result[3][1] = RDet * (
-		M[0][0] * (M[2][1] * M[3][2] - M[2][2] * M[3][1]) -
-		M[2][0] * (M[0][1] * M[3][2] - M[0][2] * M[3][1]) +
-		M[3][0] * (M[0][1] * M[2][2] - M[0][2] * M[2][1])
-		);
-	Result[3][2] = -RDet * (
-		M[0][0] * (M[1][1] * M[3][2] - M[1][2] * M[3][1]) -
-		M[1][0] * (M[0][1] * M[3][2] - M[0][2] * M[3][1]) +
-		M[3][0] * (M[0][1] * M[1][2] - M[0][2] * M[1][1])
-		);
-	Result[3][3] = RDet * (
-		M[0][0] * (M[1][1] * M[2][2] - M[1][2] * M[2][1]) -
-		M[1][0] * (M[0][1] * M[2][2] - M[0][2] * M[2][1]) +
-		M[2][0] * (M[0][1] * M[1][2] - M[0][2] * M[1][1])
-		);
-
-	memcpy(DstMatrix, &Result, sizeof(Result));
-}
-
-void VectorMatrixInverse(FMatrix44d* DstMatrix, const FMatrix44d* SrcMatrix)
-{
-	typedef double Double4x4[4][4];
-	const Double4x4& M = *((const Double4x4*)SrcMatrix->M);
-	Double4x4 Result;
-	double Det[4];
-	Double4x4 Tmp;
-
-	Tmp[0][0] = M[2][2] * M[3][3] - M[2][3] * M[3][2];
-	Tmp[0][1] = M[1][2] * M[3][3] - M[1][3] * M[3][2];
-	Tmp[0][2] = M[1][2] * M[2][3] - M[1][3] * M[2][2];
-
-	Tmp[1][0] = M[2][2] * M[3][3] - M[2][3] * M[3][2];
-	Tmp[1][1] = M[0][2] * M[3][3] - M[0][3] * M[3][2];
-	Tmp[1][2] = M[0][2] * M[2][3] - M[0][3] * M[2][2];
-
-	Tmp[2][0] = M[1][2] * M[3][3] - M[1][3] * M[3][2];
-	Tmp[2][1] = M[0][2] * M[3][3] - M[0][3] * M[3][2];
-	Tmp[2][2] = M[0][2] * M[1][3] - M[0][3] * M[1][2];
-
-	Tmp[3][0] = M[1][2] * M[2][3] - M[1][3] * M[2][2];
-	Tmp[3][1] = M[0][2] * M[2][3] - M[0][3] * M[2][2];
-	Tmp[3][2] = M[0][2] * M[1][3] - M[0][3] * M[1][2];
-
-	Det[0] = M[1][1] * Tmp[0][0] - M[2][1] * Tmp[0][1] + M[3][1] * Tmp[0][2];
-	Det[1] = M[0][1] * Tmp[1][0] - M[2][1] * Tmp[1][1] + M[3][1] * Tmp[1][2];
-	Det[2] = M[0][1] * Tmp[2][0] - M[1][1] * Tmp[2][1] + M[3][1] * Tmp[2][2];
-	Det[3] = M[0][1] * Tmp[3][0] - M[1][1] * Tmp[3][1] + M[2][1] * Tmp[3][2];
-
-	const double Determinant = M[0][0] * Det[0] - M[1][0] * Det[1] + M[2][0] * Det[2] - M[3][0] * Det[3];
-	const double	RDet = 1.0 / Determinant;
-
-	Result[0][0] = RDet * Det[0];
-	Result[0][1] = -RDet * Det[1];
-	Result[0][2] = RDet * Det[2];
-	Result[0][3] = -RDet * Det[3];
-	Result[1][0] = -RDet * (M[1][0] * Tmp[0][0] - M[2][0] * Tmp[0][1] + M[3][0] * Tmp[0][2]);
-	Result[1][1] = RDet * (M[0][0] * Tmp[1][0] - M[2][0] * Tmp[1][1] + M[3][0] * Tmp[1][2]);
-	Result[1][2] = -RDet * (M[0][0] * Tmp[2][0] - M[1][0] * Tmp[2][1] + M[3][0] * Tmp[2][2]);
-	Result[1][3] = RDet * (M[0][0] * Tmp[3][0] - M[1][0] * Tmp[3][1] + M[2][0] * Tmp[3][2]);
-	Result[2][0] = RDet * (
-		M[1][0] * (M[2][1] * M[3][3] - M[2][3] * M[3][1]) -
-		M[2][0] * (M[1][1] * M[3][3] - M[1][3] * M[3][1]) +
-		M[3][0] * (M[1][1] * M[2][3] - M[1][3] * M[2][1])
-		);
-	Result[2][1] = -RDet * (
-		M[0][0] * (M[2][1] * M[3][3] - M[2][3] * M[3][1]) -
-		M[2][0] * (M[0][1] * M[3][3] - M[0][3] * M[3][1]) +
-		M[3][0] * (M[0][1] * M[2][3] - M[0][3] * M[2][1])
-		);
-	Result[2][2] = RDet * (
-		M[0][0] * (M[1][1] * M[3][3] - M[1][3] * M[3][1]) -
-		M[1][0] * (M[0][1] * M[3][3] - M[0][3] * M[3][1]) +
-		M[3][0] * (M[0][1] * M[1][3] - M[0][3] * M[1][1])
-		);
-	Result[2][3] = -RDet * (
-		M[0][0] * (M[1][1] * M[2][3] - M[1][3] * M[2][1]) -
-		M[1][0] * (M[0][1] * M[2][3] - M[0][3] * M[2][1]) +
-		M[2][0] * (M[0][1] * M[1][3] - M[0][3] * M[1][1])
-		);
-	Result[3][0] = -RDet * (
-		M[1][0] * (M[2][1] * M[3][2] - M[2][2] * M[3][1]) -
-		M[2][0] * (M[1][1] * M[3][2] - M[1][2] * M[3][1]) +
-		M[3][0] * (M[1][1] * M[2][2] - M[1][2] * M[2][1])
-		);
-	Result[3][1] = RDet * (
-		M[0][0] * (M[2][1] * M[3][2] - M[2][2] * M[3][1]) -
-		M[2][0] * (M[0][1] * M[3][2] - M[0][2] * M[3][1]) +
-		M[3][0] * (M[0][1] * M[2][2] - M[0][2] * M[2][1])
-		);
-	Result[3][2] = -RDet * (
-		M[0][0] * (M[1][1] * M[3][2] - M[1][2] * M[3][1]) -
-		M[1][0] * (M[0][1] * M[3][2] - M[0][2] * M[3][1]) +
-		M[3][0] * (M[0][1] * M[1][2] - M[0][2] * M[1][1])
-		);
-	Result[3][3] = RDet * (
-		M[0][0] * (M[1][1] * M[2][2] - M[1][2] * M[2][1]) -
-		M[1][0] * (M[0][1] * M[2][2] - M[0][2] * M[2][1]) +
-		M[2][0] * (M[0][1] * M[1][2] - M[0][2] * M[1][1])
-		);
-
-	memcpy(DstMatrix, &Result, sizeof(Result));
-}
-
 #endif // UE_USING_UNREALMATH_SSE
diff --git a/Engine/Source/Runtime/Core/Public/Math/Matrix.h b/Engine/Source/Runtime/Core/Public/Math/Matrix.h
--- a/Engine/Source/Runtime/Core/Public/Math/Matrix.h
+++ b/Engine/Source/Runtime/Core/Public/Math/Matrix.h
@@ -201,10 +201,10 @@ public:
 	/** @return the determinant of rotation 3x3 matrix */
 	inline T RotDeterminant() const;
 
-	/** Fast path, doesn't check for nil matrices in final release builds */
+	/** Get the inverse of this matrix.  Will ensure on nil matrices in non-final builds.  Not faster than Inverse. */
 	inline TMatrix<T> InverseFast() const;
 
-	/** Fast path, and handles nil matrices. */
+	/** Get the inverse of this matrix.  Will silently change nil/nan matrices to identity. */
 	inline TMatrix<T> Inverse() const;
 
 	inline TMatrix<T> TransposeAdjoint() const;
@@ -558,8 +558,7 @@ struct TLookAtMatrix : public TLookFromMatrix<T>
 	explicit TLookAtMatrix(const TLookAtMatrix<FArg>& From) : TLookFromMatrix<T>(From) {}
 };
 
-
-} // namespace UE::Core
+} // namespace UE::Math
 } // namespace UE
 
 UE_DECLARE_LWC_TYPE(Matrix, 44);
@@ -586,8 +585,9 @@ template<> CORE_API FQuat4d FMatrix44d::ToQuat() const;
 
 
 // very high quality 4x4 matrix inverse
+// @todo: this is redundant with FMatrix44d::Inverse and should be removed ; seems to be unused
 template<typename FArg, TEMPLATE_REQUIRES(std::is_floating_point<FArg>::value)>
-static inline void Inverse4x4( double* dst, const FArg* src )
+static inline bool Inverse4x4( double* dst, const FArg* src )
 {
 	const double s0  = (double)(src[ 0]); const double s1  = (double)(src[ 1]); const double s2  = (double)(src[ 2]); const double s3  = (double)(src[ 3]);
 	const double s4  = (double)(src[ 4]); const double s5  = (double)(src[ 5]); const double s6  = (double)(src[ 6]); const double s7  = (double)(src[ 7]);
@@ -613,14 +613,18 @@ static inline void Inverse4x4( double* dst, const FArg* src )
 	inv[15] =  s0 * s5  * s10 - s0 * s6  * s9  - s4 * s1 * s10 + s4 * s2 * s9  + s8  * s1 * s6  - s8  * s2 * s5;
 
 	double det = s0 * inv[0] + s1 * inv[4] + s2 * inv[8] + s3 * inv[12];
-	if( det != 0.0 )
+	if( det == 0.0 || !FMath::IsFinite(det) )
 	{
-		det = 1.0 / det;
+		memcpy(dst,&FMatrix44d::Identity,sizeof(FMatrix44d));
+		return false;
 	}
+
+	det = 1.0 / det;
 	for( int i = 0; i < 16; i++ )
 	{
 		dst[i] = inv[i] * det;
 	}
+	return true;
 }
 
 #include "Math/Matrix.inl" // IWYU pragma: export
diff --git a/Engine/Source/Runtime/Core/Public/Math/Matrix.inl b/Engine/Source/Runtime/Core/Public/Math/Matrix.inl
--- a/Engine/Source/Runtime/Core/Public/Math/Matrix.inl
+++ b/Engine/Source/Runtime/Core/Public/Math/Matrix.inl
@@ -287,69 +287,54 @@ inline T TMatrix<T>::RotDeterminant() const
 		M[2][0] * (M[0][1] * M[1][2] - M[0][2] * M[1][1]);
 }
 
-// Inverse.
-/** Fast path, doesn't check for nil matrices in final release builds */
+// InverseFast.
+// Unlike Inverse, this may ensure if used on nil/nan matrix in non-final builds
+//   should not be used on matrices that may be nil/nan
+// note: not actually faster than Inverse
 template<typename T>
 inline TMatrix<T> TMatrix<T>::InverseFast() const
 {
 	// If we're in non final release, then make sure we're not creating NaNs
 #if !(UE_BUILD_SHIPPING || UE_BUILD_TEST)
 	// Check for zero scale matrix to invert
+	// @todo : remove this axis length check
 	if (GetScaledAxis(EAxis::X).IsNearlyZero(UE_SMALL_NUMBER) &&
 		GetScaledAxis(EAxis::Y).IsNearlyZero(UE_SMALL_NUMBER) &&
 		GetScaledAxis(EAxis::Z).IsNearlyZero(UE_SMALL_NUMBER))
 	{
 		ErrorEnsure(TEXT("TMatrix<T>::InverseFast(), trying to invert a NIL matrix, this results in NaNs! Use Inverse() instead."));
 	}
-	else
-	{
-		const T	Det = Determinant();
-
-		if (Det == 0.0f || !FMath::IsFinite(Det))
-		{
-			ErrorEnsure(TEXT("TMatrix<T>::InverseFast(), trying to invert a non-invertible matrix, this results in NaNs! Use Inverse() instead."));
-		}
-	}
 #endif
+
 	TMatrix<T> Result;
-	VectorMatrixInverse(&Result, this);
-	Result.DiagnosticCheckNaN();
+	if ( ! VectorMatrixInverse(&Result, this) )
+	{
+#if !(UE_BUILD_SHIPPING || UE_BUILD_TEST)
+		ErrorEnsure(TEXT("TMatrix<T>::InverseFast(), trying to invert a non-invertible matrix, this results in NaNs! Use Inverse() instead."));
+#endif
+	}
+	Result.DiagnosticCheckNaN(); // <- pointless, VectorMatrixInverse ensures output is valid
 	return Result;
 }
 
-// Inverse.
+// Inverse.   zero & NaN matrices are silently changed to identity
 template<typename T>
 inline TMatrix<T> TMatrix<T>::Inverse() const
 {
-	TMatrix<T> Result;
-
 	// Check for zero scale matrix to invert
+	// @todo : remove this axis length check
 	if (GetScaledAxis(EAxis::X).IsNearlyZero(UE_SMALL_NUMBER) &&
 		GetScaledAxis(EAxis::Y).IsNearlyZero(UE_SMALL_NUMBER) &&
 		GetScaledAxis(EAxis::Z).IsNearlyZero(UE_SMALL_NUMBER))
 	{
 		// just set to zero - avoids unsafe inverse of zero and duplicates what QNANs were resulting in before (scaling away all children)
-		Result = Identity;
-	}
-	else
-	{
-		const T	Det = Determinant();
-
-		// Under fp:fast this comparison can (and has been) getting inverted such that NaNs
-		// get an Identity matrix out. Under fp:precise any NaNs (correctly) get treated as != 0 
-		// and carry through to the inverse and crash, when this function is supposed to handle 
-		// that (vs InverseFast which expects well-formed matrices)
-		if (FMath::IsFinite(Det) == false || Det == 0.0f)
-		{
-			Result = Identity;
-		}
-		else
-		{
-			VectorMatrixInverse(&Result, this);
-		}
+		return Identity;
 	}
-
-	Result.DiagnosticCheckNaN();
+	
+	TMatrix<T> Result;
+	// VectorMatrixInverse will return false and fill identity for non-invertible matrices
+	VectorMatrixInverse(&Result, this);
+	Result.DiagnosticCheckNaN(); // <- pointless, VectorMatrixInverse ensures output is valid
 	return Result;
 }
 
diff --git a/Engine/Source/Runtime/Core/Public/Math/UnrealMathDirectX.h b/Engine/Source/Runtime/Core/Public/Math/UnrealMathDirectX.h
--- a/Engine/Source/Runtime/Core/Public/Math/UnrealMathDirectX.h
+++ b/Engine/Source/Runtime/Core/Public/Math/UnrealMathDirectX.h
@@ -593,12 +593,26 @@ FORCEINLINE void VectorMatrixMultiply( FMatrix* Result, const FMatrix* Matrix1,
  * @param DstMatrix		FMatrix pointer to where the result should be stored
  * @param SrcMatrix		FMatrix pointer to the Matrix to be inversed
  */
-FORCEINLINE void VectorMatrixInverse( FMatrix* DstMatrix, const FMatrix* SrcMatrix )
+FORCEINLINE bool VectorMatrixInverse( FMatrix* DstMatrix, const FMatrix* SrcMatrix )
 {
 	using namespace DirectX;
 	XMMATRIX XMSrcMatrix = XMLoadFloat4x4A((const XMFLOAT4X4A*)(SrcMatrix));
-	XMMATRIX XMDstMatrix = XMMatrixInverse( nullptr, XMSrcMatrix );
+	XMVECTOR XMDeterminant;
+	XMMATRIX XMDstMatrix = XMMatrixInverse( &XMDeterminant, XMSrcMatrix );
+	
+	//XMMatrixInverse Returns the matrix inverse of M. If there is no inverse (that is, if the determinant is 0), 
+	//	then XMMatrixInverse returns an infinite matrix.
+
+	float Determinant;
+	XMStoreFloat(&Determinant,XMDeterminant);
+	if ( Determinant == 0.f || !FMath::IsFinite(Determinant) )
+	{
+		*DstMatrix = FMatrix44d::Identity;
+		return false;
+	}
+	
 	XMStoreFloat4x4A( (XMFLOAT4X4A*)(DstMatrix), XMDstMatrix);
+	return true;
 }
 
 /**
diff --git a/Engine/Source/Runtime/Core/Public/Math/UnrealMathFPU.h b/Engine/Source/Runtime/Core/Public/Math/UnrealMathFPU.h
--- a/Engine/Source/Runtime/Core/Public/Math/UnrealMathFPU.h
+++ b/Engine/Source/Runtime/Core/Public/Math/UnrealMathFPU.h
@@ -1551,177 +1551,20 @@ FORCEINLINE void VectorMatrixMultiply(FMatrix44f* Result, const FMatrix44f* Matr
 
 
 /**
- * Calculate the inverse of an FMatrix44d.
+ * Calculate the inverse of an FMatrix44.  Src == Dst is allowed
+ *
+ * @param DstMatrix		FMatrix44 pointer to where the result should be stored
+ * @param SrcMatrix		FMatrix44 pointer to the Matrix to be inversed
+ * @return bool			returns false if matrix is not invertable and stores identity 
  *
- * @param DstMatrix		FMatrix44d pointer to where the result should be stored
- * @param SrcMatrix		FMatrix44d pointer to the Matrix to be inversed
  */
-FORCEINLINE void VectorMatrixInverse(FMatrix44d* DstMatrix, const FMatrix44d* SrcMatrix)
+FORCEINLINE bool VectorMatrixInverse(FMatrix44d* DstMatrix, const FMatrix44d* SrcMatrix)
 {
-	typedef double Double4x4[4][4];
-	const Double4x4& M = *((const Double4x4*)SrcMatrix);
-	Double4x4 Result;
-	double Det[4];
-	Double4x4 Tmp;
-
-	Tmp[0][0] = M[2][2] * M[3][3] - M[2][3] * M[3][2];
-	Tmp[0][1] = M[1][2] * M[3][3] - M[1][3] * M[3][2];
-	Tmp[0][2] = M[1][2] * M[2][3] - M[1][3] * M[2][2];
-
-	Tmp[1][0] = M[2][2] * M[3][3] - M[2][3] * M[3][2];
-	Tmp[1][1] = M[0][2] * M[3][3] - M[0][3] * M[3][2];
-	Tmp[1][2] = M[0][2] * M[2][3] - M[0][3] * M[2][2];
-
-	Tmp[2][0] = M[1][2] * M[3][3] - M[1][3] * M[3][2];
-	Tmp[2][1] = M[0][2] * M[3][3] - M[0][3] * M[3][2];
-	Tmp[2][2] = M[0][2] * M[1][3] - M[0][3] * M[1][2];
-
-	Tmp[3][0] = M[1][2] * M[2][3] - M[1][3] * M[2][2];
-	Tmp[3][1] = M[0][2] * M[2][3] - M[0][3] * M[2][2];
-	Tmp[3][2] = M[0][2] * M[1][3] - M[0][3] * M[1][2];
-
-	Det[0] = M[1][1] * Tmp[0][0] - M[2][1] * Tmp[0][1] + M[3][1] * Tmp[0][2];
-	Det[1] = M[0][1] * Tmp[1][0] - M[2][1] * Tmp[1][1] + M[3][1] * Tmp[1][2];
-	Det[2] = M[0][1] * Tmp[2][0] - M[1][1] * Tmp[2][1] + M[3][1] * Tmp[2][2];
-	Det[3] = M[0][1] * Tmp[3][0] - M[1][1] * Tmp[3][1] + M[2][1] * Tmp[3][2];
-
-	const double Determinant = M[0][0] * Det[0] - M[1][0] * Det[1] + M[2][0] * Det[2] - M[3][0] * Det[3];
-	const double	RDet = 1.0 / Determinant;
-
-	Result[0][0] = RDet * Det[0];
-	Result[0][1] = -RDet * Det[1];
-	Result[0][2] = RDet * Det[2];
-	Result[0][3] = -RDet * Det[3];
-	Result[1][0] = -RDet * (M[1][0] * Tmp[0][0] - M[2][0] * Tmp[0][1] + M[3][0] * Tmp[0][2]);
-	Result[1][1] = RDet * (M[0][0] * Tmp[1][0] - M[2][0] * Tmp[1][1] + M[3][0] * Tmp[1][2]);
-	Result[1][2] = -RDet * (M[0][0] * Tmp[2][0] - M[1][0] * Tmp[2][1] + M[3][0] * Tmp[2][2]);
-	Result[1][3] = RDet * (M[0][0] * Tmp[3][0] - M[1][0] * Tmp[3][1] + M[2][0] * Tmp[3][2]);
-	Result[2][0] = RDet * (
-		M[1][0] * (M[2][1] * M[3][3] - M[2][3] * M[3][1]) -
-		M[2][0] * (M[1][1] * M[3][3] - M[1][3] * M[3][1]) +
-		M[3][0] * (M[1][1] * M[2][3] - M[1][3] * M[2][1])
-		);
-	Result[2][1] = -RDet * (
-		M[0][0] * (M[2][1] * M[3][3] - M[2][3] * M[3][1]) -
-		M[2][0] * (M[0][1] * M[3][3] - M[0][3] * M[3][1]) +
-		M[3][0] * (M[0][1] * M[2][3] - M[0][3] * M[2][1])
-		);
-	Result[2][2] = RDet * (
-		M[0][0] * (M[1][1] * M[3][3] - M[1][3] * M[3][1]) -
-		M[1][0] * (M[0][1] * M[3][3] - M[0][3] * M[3][1]) +
-		M[3][0] * (M[0][1] * M[1][3] - M[0][3] * M[1][1])
-		);
-	Result[2][3] = -RDet * (
-		M[0][0] * (M[1][1] * M[2][3] - M[1][3] * M[2][1]) -
-		M[1][0] * (M[0][1] * M[2][3] - M[0][3] * M[2][1]) +
-		M[2][0] * (M[0][1] * M[1][3] - M[0][3] * M[1][1])
-		);
-	Result[3][0] = -RDet * (
-		M[1][0] * (M[2][1] * M[3][2] - M[2][2] * M[3][1]) -
-		M[2][0] * (M[1][1] * M[3][2] - M[1][2] * M[3][1]) +
-		M[3][0] * (M[1][1] * M[2][2] - M[1][2] * M[2][1])
-		);
-	Result[3][1] = RDet * (
-		M[0][0] * (M[2][1] * M[3][2] - M[2][2] * M[3][1]) -
-		M[2][0] * (M[0][1] * M[3][2] - M[0][2] * M[3][1]) +
-		M[3][0] * (M[0][1] * M[2][2] - M[0][2] * M[2][1])
-		);
-	Result[3][2] = -RDet * (
-		M[0][0] * (M[1][1] * M[3][2] - M[1][2] * M[3][1]) -
-		M[1][0] * (M[0][1] * M[3][2] - M[0][2] * M[3][1]) +
-		M[3][0] * (M[0][1] * M[1][2] - M[0][2] * M[1][1])
-		);
-	Result[3][3] = RDet * (
-		M[0][0] * (M[1][1] * M[2][2] - M[1][2] * M[2][1]) -
-		M[1][0] * (M[0][1] * M[2][2] - M[0][2] * M[2][1]) +
-		M[2][0] * (M[0][1] * M[1][2] - M[0][2] * M[1][1])
-		);
-
-	memcpy(DstMatrix, &Result, 16 * sizeof(double));
-}
-
-FORCEINLINE void VectorMatrixInverse(FMatrix44f* DstMatrix, const FMatrix44f* SrcMatrix)
+	return FMath::MatrixInverse(DstMatrix,SrcMatrix);
+}
+FORCEINLINE bool VectorMatrixInverse(FMatrix44f* DstMatrix, const FMatrix44f* SrcMatrix)
 {
-	typedef float Float4x4[4][4];
-	const Float4x4& M = *((const Float4x4*)SrcMatrix);
-	Float4x4 Result;
-	float Det[4];
-	Float4x4 Tmp;
-
-	Tmp[0][0] = M[2][2] * M[3][3] - M[2][3] * M[3][2];
-	Tmp[0][1] = M[1][2] * M[3][3] - M[1][3] * M[3][2];
-	Tmp[0][2] = M[1][2] * M[2][3] - M[1][3] * M[2][2];
-
-	Tmp[1][0] = M[2][2] * M[3][3] - M[2][3] * M[3][2];
-	Tmp[1][1] = M[0][2] * M[3][3] - M[0][3] * M[3][2];
-	Tmp[1][2] = M[0][2] * M[2][3] - M[0][3] * M[2][2];
-
-	Tmp[2][0] = M[1][2] * M[3][3] - M[1][3] * M[3][2];
-	Tmp[2][1] = M[0][2] * M[3][3] - M[0][3] * M[3][2];
-	Tmp[2][2] = M[0][2] * M[1][3] - M[0][3] * M[1][2];
-
-	Tmp[3][0] = M[1][2] * M[2][3] - M[1][3] * M[2][2];
-	Tmp[3][1] = M[0][2] * M[2][3] - M[0][3] * M[2][2];
-	Tmp[3][2] = M[0][2] * M[1][3] - M[0][3] * M[1][2];
-
-	Det[0] = M[1][1] * Tmp[0][0] - M[2][1] * Tmp[0][1] + M[3][1] * Tmp[0][2];
-	Det[1] = M[0][1] * Tmp[1][0] - M[2][1] * Tmp[1][1] + M[3][1] * Tmp[1][2];
-	Det[2] = M[0][1] * Tmp[2][0] - M[1][1] * Tmp[2][1] + M[3][1] * Tmp[2][2];
-	Det[3] = M[0][1] * Tmp[3][0] - M[1][1] * Tmp[3][1] + M[2][1] * Tmp[3][2];
-
-	const float Determinant = M[0][0] * Det[0] - M[1][0] * Det[1] + M[2][0] * Det[2] - M[3][0] * Det[3];
-	const float	RDet = 1.0f / Determinant;
-
-	Result[0][0] = RDet * Det[0];
-	Result[0][1] = -RDet * Det[1];
-	Result[0][2] = RDet * Det[2];
-	Result[0][3] = -RDet * Det[3];
-	Result[1][0] = -RDet * (M[1][0] * Tmp[0][0] - M[2][0] * Tmp[0][1] + M[3][0] * Tmp[0][2]);
-	Result[1][1] = RDet * (M[0][0] * Tmp[1][0] - M[2][0] * Tmp[1][1] + M[3][0] * Tmp[1][2]);
-	Result[1][2] = -RDet * (M[0][0] * Tmp[2][0] - M[1][0] * Tmp[2][1] + M[3][0] * Tmp[2][2]);
-	Result[1][3] = RDet * (M[0][0] * Tmp[3][0] - M[1][0] * Tmp[3][1] + M[2][0] * Tmp[3][2]);
-	Result[2][0] = RDet * (
-		M[1][0] * (M[2][1] * M[3][3] - M[2][3] * M[3][1]) -
-		M[2][0] * (M[1][1] * M[3][3] - M[1][3] * M[3][1]) +
-		M[3][0] * (M[1][1] * M[2][3] - M[1][3] * M[2][1])
-		);
-	Result[2][1] = -RDet * (
-		M[0][0] * (M[2][1] * M[3][3] - M[2][3] * M[3][1]) -
-		M[2][0] * (M[0][1] * M[3][3] - M[0][3] * M[3][1]) +
-		M[3][0] * (M[0][1] * M[2][3] - M[0][3] * M[2][1])
-		);
-	Result[2][2] = RDet * (
-		M[0][0] * (M[1][1] * M[3][3] - M[1][3] * M[3][1]) -
-		M[1][0] * (M[0][1] * M[3][3] - M[0][3] * M[3][1]) +
-		M[3][0] * (M[0][1] * M[1][3] - M[0][3] * M[1][1])
-		);
-	Result[2][3] = -RDet * (
-		M[0][0] * (M[1][1] * M[2][3] - M[1][3] * M[2][1]) -
-		M[1][0] * (M[0][1] * M[2][3] - M[0][3] * M[2][1]) +
-		M[2][0] * (M[0][1] * M[1][3] - M[0][3] * M[1][1])
-		);
-	Result[3][0] = -RDet * (
-		M[1][0] * (M[2][1] * M[3][2] - M[2][2] * M[3][1]) -
-		M[2][0] * (M[1][1] * M[3][2] - M[1][2] * M[3][1]) +
-		M[3][0] * (M[1][1] * M[2][2] - M[1][2] * M[2][1])
-		);
-	Result[3][1] = RDet * (
-		M[0][0] * (M[2][1] * M[3][2] - M[2][2] * M[3][1]) -
-		M[2][0] * (M[0][1] * M[3][2] - M[0][2] * M[3][1]) +
-		M[3][0] * (M[0][1] * M[2][2] - M[0][2] * M[2][1])
-		);
-	Result[3][2] = -RDet * (
-		M[0][0] * (M[1][1] * M[3][2] - M[1][2] * M[3][1]) -
-		M[1][0] * (M[0][1] * M[3][2] - M[0][2] * M[3][1]) +
-		M[3][0] * (M[0][1] * M[1][2] - M[0][2] * M[1][1])
-		);
-	Result[3][3] = RDet * (
-		M[0][0] * (M[1][1] * M[2][2] - M[1][2] * M[2][1]) -
-		M[1][0] * (M[0][1] * M[2][2] - M[0][2] * M[2][1]) +
-		M[2][0] * (M[0][1] * M[1][2] - M[0][2] * M[1][1])
-		);
-
-	memcpy(DstMatrix, &Result, 16 * sizeof(float));
+	return FMath::MatrixInverse(DstMatrix,SrcMatrix);
 }
 
 /**
diff --git a/Engine/Source/Runtime/Core/Public/Math/UnrealMathNeon.h b/Engine/Source/Runtime/Core/Public/Math/UnrealMathNeon.h
--- a/Engine/Source/Runtime/Core/Public/Math/UnrealMathNeon.h
+++ b/Engine/Source/Runtime/Core/Public/Math/UnrealMathNeon.h
@@ -1923,178 +1923,20 @@ FORCEINLINE void VectorMatrixMultiply(FMatrix44d* Result, const FMatrix44d* Matr
 }
 
 /**
- * Calculate the inverse of an FMatrix.
+ * Calculate the inverse of an FMatrix44.  Src == Dst is allowed
+ *
+ * @param DstMatrix		FMatrix44 pointer to where the result should be stored
+ * @param SrcMatrix		FMatrix44 pointer to the Matrix to be inversed
+ * @return bool			returns false if matrix is not invertable and stores identity 
  *
- * @param DstMatrix		FMatrix pointer to where the result should be stored
- * @param SrcMatrix		FMatrix pointer to the Matrix to be inversed
  */
-// OPTIMIZE ME: stolen from UnMathFpu.h
-FORCEINLINE void VectorMatrixInverse(FMatrix44f* DstMatrix, const FMatrix44f* SrcMatrix )
-{
-	typedef float Float4x4[4][4];
-	const Float4x4& M = *((const Float4x4*) SrcMatrix);
-	Float4x4 Result;
-	float Det[4];
-	Float4x4 Tmp;
-	
-	Tmp[0][0]       = M[2][2] * M[3][3] - M[2][3] * M[3][2];
-	Tmp[0][1]       = M[1][2] * M[3][3] - M[1][3] * M[3][2];
-	Tmp[0][2]       = M[1][2] * M[2][3] - M[1][3] * M[2][2];
-	
-	Tmp[1][0]       = M[2][2] * M[3][3] - M[2][3] * M[3][2];
-	Tmp[1][1]       = M[0][2] * M[3][3] - M[0][3] * M[3][2];
-	Tmp[1][2]       = M[0][2] * M[2][3] - M[0][3] * M[2][2];
-	
-	Tmp[2][0]       = M[1][2] * M[3][3] - M[1][3] * M[3][2];
-	Tmp[2][1]       = M[0][2] * M[3][3] - M[0][3] * M[3][2];
-	Tmp[2][2]       = M[0][2] * M[1][3] - M[0][3] * M[1][2];
-	
-	Tmp[3][0]       = M[1][2] * M[2][3] - M[1][3] * M[2][2];
-	Tmp[3][1]       = M[0][2] * M[2][3] - M[0][3] * M[2][2];
-	Tmp[3][2]       = M[0][2] * M[1][3] - M[0][3] * M[1][2];
-	
-	Det[0]          = M[1][1]*Tmp[0][0] - M[2][1]*Tmp[0][1] + M[3][1]*Tmp[0][2];
-	Det[1]          = M[0][1]*Tmp[1][0] - M[2][1]*Tmp[1][1] + M[3][1]*Tmp[1][2];
-	Det[2]          = M[0][1]*Tmp[2][0] - M[1][1]*Tmp[2][1] + M[3][1]*Tmp[2][2];
-	Det[3]          = M[0][1]*Tmp[3][0] - M[1][1]*Tmp[3][1] + M[2][1]*Tmp[3][2];
-	
-	float Determinant = M[0][0]*Det[0] - M[1][0]*Det[1] + M[2][0]*Det[2] - M[3][0]*Det[3];
-	const float     RDet = 1.0f / Determinant;
-	
-	Result[0][0] =  RDet * Det[0];
-	Result[0][1] = -RDet * Det[1];
-	Result[0][2] =  RDet * Det[2];
-	Result[0][3] = -RDet * Det[3];
-	Result[1][0] = -RDet * (M[1][0]*Tmp[0][0] - M[2][0]*Tmp[0][1] + M[3][0]*Tmp[0][2]);
-	Result[1][1] =  RDet * (M[0][0]*Tmp[1][0] - M[2][0]*Tmp[1][1] + M[3][0]*Tmp[1][2]);
-	Result[1][2] = -RDet * (M[0][0]*Tmp[2][0] - M[1][0]*Tmp[2][1] + M[3][0]*Tmp[2][2]);
-	Result[1][3] =  RDet * (M[0][0]*Tmp[3][0] - M[1][0]*Tmp[3][1] + M[2][0]*Tmp[3][2]);
-	Result[2][0] =  RDet * (
-							M[1][0] * (M[2][1] * M[3][3] - M[2][3] * M[3][1]) -
-							M[2][0] * (M[1][1] * M[3][3] - M[1][3] * M[3][1]) +
-							M[3][0] * (M[1][1] * M[2][3] - M[1][3] * M[2][1])
-							);
-	Result[2][1] = -RDet * (
-							M[0][0] * (M[2][1] * M[3][3] - M[2][3] * M[3][1]) -
-							M[2][0] * (M[0][1] * M[3][3] - M[0][3] * M[3][1]) +
-							M[3][0] * (M[0][1] * M[2][3] - M[0][3] * M[2][1])
-							);
-	Result[2][2] =  RDet * (
-							M[0][0] * (M[1][1] * M[3][3] - M[1][3] * M[3][1]) -
-							M[1][0] * (M[0][1] * M[3][3] - M[0][3] * M[3][1]) +
-							M[3][0] * (M[0][1] * M[1][3] - M[0][3] * M[1][1])
-							);
-	Result[2][3] = -RDet * (
-							M[0][0] * (M[1][1] * M[2][3] - M[1][3] * M[2][1]) -
-							M[1][0] * (M[0][1] * M[2][3] - M[0][3] * M[2][1]) +
-							M[2][0] * (M[0][1] * M[1][3] - M[0][3] * M[1][1])
-							);
-	Result[3][0] = -RDet * (
-							M[1][0] * (M[2][1] * M[3][2] - M[2][2] * M[3][1]) -
-							M[2][0] * (M[1][1] * M[3][2] - M[1][2] * M[3][1]) +
-							M[3][0] * (M[1][1] * M[2][2] - M[1][2] * M[2][1])
-							);
-	Result[3][1] =  RDet * (
-							M[0][0] * (M[2][1] * M[3][2] - M[2][2] * M[3][1]) -
-							M[2][0] * (M[0][1] * M[3][2] - M[0][2] * M[3][1]) +
-							M[3][0] * (M[0][1] * M[2][2] - M[0][2] * M[2][1])
-							);
-	Result[3][2] = -RDet * (
-							M[0][0] * (M[1][1] * M[3][2] - M[1][2] * M[3][1]) -
-							M[1][0] * (M[0][1] * M[3][2] - M[0][2] * M[3][1]) +
-							M[3][0] * (M[0][1] * M[1][2] - M[0][2] * M[1][1])
-							);
-	Result[3][3] =  RDet * (
-							M[0][0] * (M[1][1] * M[2][2] - M[1][2] * M[2][1]) -
-							M[1][0] * (M[0][1] * M[2][2] - M[0][2] * M[2][1]) +
-							M[2][0] * (M[0][1] * M[1][2] - M[0][2] * M[1][1])
-							);
-	
-	memcpy( DstMatrix, &Result, sizeof(Result) );
-}
-
-FORCEINLINE void VectorMatrixInverse(FMatrix44d* DstMatrix, const FMatrix44d* SrcMatrix)
-{
-	typedef double Double4x4[4][4];
-	const Double4x4& M = *((const Double4x4*)SrcMatrix);
-	Double4x4 Result;
-	double Det[4];
-	Double4x4 Tmp;
-
-	Tmp[0][0] = M[2][2] * M[3][3] - M[2][3] * M[3][2];
-	Tmp[0][1] = M[1][2] * M[3][3] - M[1][3] * M[3][2];
-	Tmp[0][2] = M[1][2] * M[2][3] - M[1][3] * M[2][2];
-
-	Tmp[1][0] = M[2][2] * M[3][3] - M[2][3] * M[3][2];
-	Tmp[1][1] = M[0][2] * M[3][3] - M[0][3] * M[3][2];
-	Tmp[1][2] = M[0][2] * M[2][3] - M[0][3] * M[2][2];
-
-	Tmp[2][0] = M[1][2] * M[3][3] - M[1][3] * M[3][2];
-	Tmp[2][1] = M[0][2] * M[3][3] - M[0][3] * M[3][2];
-	Tmp[2][2] = M[0][2] * M[1][3] - M[0][3] * M[1][2];
-
-	Tmp[3][0] = M[1][2] * M[2][3] - M[1][3] * M[2][2];
-	Tmp[3][1] = M[0][2] * M[2][3] - M[0][3] * M[2][2];
-	Tmp[3][2] = M[0][2] * M[1][3] - M[0][3] * M[1][2];
-
-	Det[0] = M[1][1] * Tmp[0][0] - M[2][1] * Tmp[0][1] + M[3][1] * Tmp[0][2];
-	Det[1] = M[0][1] * Tmp[1][0] - M[2][1] * Tmp[1][1] + M[3][1] * Tmp[1][2];
-	Det[2] = M[0][1] * Tmp[2][0] - M[1][1] * Tmp[2][1] + M[3][1] * Tmp[2][2];
-	Det[3] = M[0][1] * Tmp[3][0] - M[1][1] * Tmp[3][1] + M[2][1] * Tmp[3][2];
-
-	double Determinant = M[0][0] * Det[0] - M[1][0] * Det[1] + M[2][0] * Det[2] - M[3][0] * Det[3];
-	const double RDet = 1.0 / Determinant;
-
-	Result[0][0] = RDet * Det[0];
-	Result[0][1] = -RDet * Det[1];
-	Result[0][2] = RDet * Det[2];
-	Result[0][3] = -RDet * Det[3];
-	Result[1][0] = -RDet * (M[1][0] * Tmp[0][0] - M[2][0] * Tmp[0][1] + M[3][0] * Tmp[0][2]);
-	Result[1][1] = RDet * (M[0][0] * Tmp[1][0] - M[2][0] * Tmp[1][1] + M[3][0] * Tmp[1][2]);
-	Result[1][2] = -RDet * (M[0][0] * Tmp[2][0] - M[1][0] * Tmp[2][1] + M[3][0] * Tmp[2][2]);
-	Result[1][3] = RDet * (M[0][0] * Tmp[3][0] - M[1][0] * Tmp[3][1] + M[2][0] * Tmp[3][2]);
-	Result[2][0] = RDet * (
-		M[1][0] * (M[2][1] * M[3][3] - M[2][3] * M[3][1]) -
-		M[2][0] * (M[1][1] * M[3][3] - M[1][3] * M[3][1]) +
-		M[3][0] * (M[1][1] * M[2][3] - M[1][3] * M[2][1])
-		);
-	Result[2][1] = -RDet * (
-		M[0][0] * (M[2][1] * M[3][3] - M[2][3] * M[3][1]) -
-		M[2][0] * (M[0][1] * M[3][3] - M[0][3] * M[3][1]) +
-		M[3][0] * (M[0][1] * M[2][3] - M[0][3] * M[2][1])
-		);
-	Result[2][2] = RDet * (
-		M[0][0] * (M[1][1] * M[3][3] - M[1][3] * M[3][1]) -
-		M[1][0] * (M[0][1] * M[3][3] - M[0][3] * M[3][1]) +
-		M[3][0] * (M[0][1] * M[1][3] - M[0][3] * M[1][1])
-		);
-	Result[2][3] = -RDet * (
-		M[0][0] * (M[1][1] * M[2][3] - M[1][3] * M[2][1]) -
-		M[1][0] * (M[0][1] * M[2][3] - M[0][3] * M[2][1]) +
-		M[2][0] * (M[0][1] * M[1][3] - M[0][3] * M[1][1])
-		);
-	Result[3][0] = -RDet * (
-		M[1][0] * (M[2][1] * M[3][2] - M[2][2] * M[3][1]) -
-		M[2][0] * (M[1][1] * M[3][2] - M[1][2] * M[3][1]) +
-		M[3][0] * (M[1][1] * M[2][2] - M[1][2] * M[2][1])
-		);
-	Result[3][1] = RDet * (
-		M[0][0] * (M[2][1] * M[3][2] - M[2][2] * M[3][1]) -
-		M[2][0] * (M[0][1] * M[3][2] - M[0][2] * M[3][1]) +
-		M[3][0] * (M[0][1] * M[2][2] - M[0][2] * M[2][1])
-		);
-	Result[3][2] = -RDet * (
-		M[0][0] * (M[1][1] * M[3][2] - M[1][2] * M[3][1]) -
-		M[1][0] * (M[0][1] * M[3][2] - M[0][2] * M[3][1]) +
-		M[3][0] * (M[0][1] * M[1][2] - M[0][2] * M[1][1])
-		);
-	Result[3][3] = RDet * (
-		M[0][0] * (M[1][1] * M[2][2] - M[1][2] * M[2][1]) -
-		M[1][0] * (M[0][1] * M[2][2] - M[0][2] * M[2][1]) +
-		M[2][0] * (M[0][1] * M[1][2] - M[0][2] * M[1][1])
-		);
-
-	memcpy(DstMatrix, &Result, sizeof(Result));
+FORCEINLINE bool VectorMatrixInverse(FMatrix44d* DstMatrix, const FMatrix44d* SrcMatrix)
+{
+	return FMath::MatrixInverse(DstMatrix,SrcMatrix);
+}
+FORCEINLINE bool VectorMatrixInverse(FMatrix44f* DstMatrix, const FMatrix44f* SrcMatrix)
+{
+	return FMath::MatrixInverse(DstMatrix,SrcMatrix);
 }
 
 /**
diff --git a/Engine/Source/Runtime/Core/Public/Math/UnrealMathSSE.h b/Engine/Source/Runtime/Core/Public/Math/UnrealMathSSE.h
--- a/Engine/Source/Runtime/Core/Public/Math/UnrealMathSSE.h
+++ b/Engine/Source/Runtime/Core/Public/Math/UnrealMathSSE.h
@@ -2203,15 +2203,21 @@ CORE_API void VectorMatrixMultiply(FMatrix44f* Result, const FMatrix44f* Matrix1
 CORE_API void VectorMatrixMultiply(FMatrix44d* Result, const FMatrix44d* Matrix1, const FMatrix44d* Matrix2);
 
 /**
- * Calculate the inverse of an FMatrix.
+ * Calculate the inverse of an FMatrix44.  Src == Dst is allowed
+ *
+ * @param DstMatrix		FMatrix44 pointer to where the result should be stored
+ * @param SrcMatrix		FMatrix44 pointer to the Matrix to be inversed
+ * @return bool			returns false if matrix is not invertable and stores identity 
  *
- * @param DstMatrix		FMatrix pointer to where the result should be stored
- * @param SrcMatrix		FMatrix pointer to the Matrix to be inversed
  */
-CORE_API void VectorMatrixInverse(FMatrix44f* DstMatrix, const FMatrix44f* SrcMatrix);
-CORE_API void VectorMatrixInverse(FMatrix44d* DstMatrix, const FMatrix44d* SrcMatrix);
-
-
+FORCEINLINE bool VectorMatrixInverse(FMatrix44d* DstMatrix, const FMatrix44d* SrcMatrix)
+{
+	return FMath::MatrixInverse(DstMatrix,SrcMatrix);
+}
+FORCEINLINE bool VectorMatrixInverse(FMatrix44f* DstMatrix, const FMatrix44f* SrcMatrix)
+{
+	return FMath::MatrixInverse(DstMatrix,SrcMatrix);
+}
 
 /**
  * Calculate Homogeneous transform.
diff --git a/Engine/Source/Runtime/Core/Public/Math/UnrealMathUtility.h b/Engine/Source/Runtime/Core/Public/Math/UnrealMathUtility.h
--- a/Engine/Source/Runtime/Core/Public/Math/UnrealMathUtility.h
+++ b/Engine/Source/Runtime/Core/Public/Math/UnrealMathUtility.h
@@ -2488,6 +2488,20 @@ public:
 		}
 		return WeightedMovingAverage(CurrentSample, PreviousSample, Weight);
 	}
+	
+	/**
+	 * Calculate the inverse of an FMatrix44.  Src == Dst is allowed
+	 *
+	 * @param DstMatrix		FMatrix44 pointer to where the result should be stored
+	 * @param SrcMatrix		FMatrix44 pointer to the Matrix to be inversed
+	 * @return bool			returns false if matrix is not invertable and stores identity 
+	 *
+	 * Do not call this directly, use VectorMatrixInverse or Matrix::Inverse
+	 * this is the fallback scalar implementation used by VectorMatrixInverse
+	 */
+	UE_NODISCARD static CORE_API bool MatrixInverse(FMatrix44f* DstMatrix, const FMatrix44f* SrcMatrix);
+	UE_NODISCARD static CORE_API bool MatrixInverse(FMatrix44d* DstMatrix, const FMatrix44d* SrcMatrix);
+
 };
 
 // LWC Conversion helpers
