note in UnrealMathDirectX.h that it is no longer used or supported

diff --git a/Engine/Source/Runtime/Core/Public/Math/UnrealMathDirectX.h b/Engine/Source/Runtime/Core/Public/Math/UnrealMathDirectX.h
--- a/Engine/Source/Runtime/Core/Public/Math/UnrealMathDirectX.h
+++ b/Engine/Source/Runtime/Core/Public/Math/UnrealMathDirectX.h
@@ -4,6 +4,9 @@
 
 // HEADER_UNIT_SKIP - Not included directly
 
+// enabled by bWithDirectXMath
+// LWC_TODO: No longer supported. Needs removing.
+
 #if defined __cplusplus_cli && !PLATFORM_HOLOLENS
 // there are compile issues with this file in managed mode, so use the FPU version
 #include "Math/UnrealMathFPU.h"
@@ -607,7 +610,7 @@ FORCEINLINE bool VectorMatrixInverse( FMatrix* DstMatrix, const FMatrix* SrcMatr
 	XMStoreFloat(&Determinant,XMDeterminant);
 	if ( Determinant == 0.f || !FMath::IsFinite(Determinant) )
 	{
-		*DstMatrix = FMatrix44d::Identity;
+		*DstMatrix = FMatrix::Identity;
 		return false;
 	}
 	
