In refresh, free all slacked memory and reallocate pool if needed

#rnx

diff --git a/Engine/Source/Runtime/Core/Private/HAL/Allocators/CachedOSVeryLargePageAllocator.cpp b/Engine/Source/Runtime/Core/Private/HAL/Allocators/CachedOSVeryLargePageAllocator.cpp
--- a/Engine/Source/Runtime/Core/Private/HAL/Allocators/CachedOSVeryLargePageAllocator.cpp
+++ b/Engine/Source/Runtime/Core/Private/HAL/Allocators/CachedOSVeryLargePageAllocator.cpp
@@ -7,6 +7,7 @@
 #include "CoreGlobals.h"
 #include "HAL/LowLevelMemTracker.h"
 #include "HAL/IConsoleManager.h"
+#include "HAL/PlatformTime.h"
 #include "ProfilingDebugging/CsvProfiler.h"
 
 #if UE_USE_VERYLARGEPAGEALLOCATOR
@@ -104,45 +105,63 @@ void FCachedOSVeryLargePageAllocator::Refresh()
 	}
 
 	// Function is not thread safe 
-	UE_LOG(LogMemory, Log, TEXT("Refreshing LargePageAllocator"));
-
-	for (int32 i = 0; i < FMemory::AllocationHints::Max; i++)
+	UE_LOG(LogMemory, Log, TEXT("LargePageAllocator - Refresh"));
+	
 	{
-		FMemory::AllocationHints AllocationHint = FMemory::AllocationHints(i);
+		// Shrink Empty back store
+		uint64 StartTime = FPlatformTime::Cycles64();
 
-		int32 LargePageCount = CommitedLargePagesCount[AllocationHint] + EmptyBackStoreCount[AllocationHint];
+		for (int i = 0; i < FMemory::AllocationHints::Max; i++)
+		{
+			FMemory::AllocationHints AllocationHint = FMemory::AllocationHints(i);
+			ShrinkEmptyBackStore(0, AllocationHint);
+		}
 
-		if (LargePageCount < GVeryLargePageAllocatorMaxEmptyBackStoreCount[AllocationHint] && GVeryLargePageAllocatorPreAllocatePools)
+		uint64 LenghtCycles64 = (FPlatformTime::Cycles64() - StartTime);
+		UE_LOG(LogMemory, Log, TEXT("LargePageAllocator - Shrink empty backstore duration: %.1f ms"), float(FPlatformTime::ToMilliseconds64(LenghtCycles64)));
+	}
+
+	{
+		// Preallocate
+		uint64 StartTime = FPlatformTime::Cycles64();
+
+		for (int32 i = 0; i < FMemory::AllocationHints::Max; i++)
 		{
-			// Preallocate large pages
-			for ( ; LargePageCount < GVeryLargePageAllocatorMaxEmptyBackStoreCount[AllocationHint]; LargePageCount++)
+			FMemory::AllocationHints AllocationHint = FMemory::AllocationHints(i);
+
+			int32 LargePageCount = CommitedLargePagesCount[AllocationHint] + EmptyBackStoreCount[AllocationHint];
+
+			if (LargePageCount < GVeryLargePageAllocatorMaxEmptyBackStoreCount[AllocationHint] && GVeryLargePageAllocatorPreAllocatePools)
 			{
-				LLM_PLATFORM_SCOPE(ELLMTag::FMalloc);
+				// Preallocate large pages
+				for (; LargePageCount < GVeryLargePageAllocatorMaxEmptyBackStoreCount[AllocationHint]; LargePageCount++)
+				{
+					LLM_PLATFORM_SCOPE(ELLMTag::FMalloc);
 
-				FLargePage* LargePage = FreeLargePagesHead[AllocationHint];
+					FLargePage* LargePage = FreeLargePagesHead[AllocationHint];
 
-				check(LargePage != nullptr); // Can't happen
-				LargePage->AllocationHint = AllocationHint;
-				LargePage->Unlink();
+					check(LargePage != nullptr); // Can't happen
+					LargePage->AllocationHint = AllocationHint;
+					LargePage->Unlink();
 
-				if (!Block.Commit(LargePage->BaseAddress - AddressSpaceReserved, SizeOfLargePage, false))
-				{
-					// Cant commit 2MB pages, stop preallocating and return page to FreeLargePagesHead list
-					LargePage->LinkHead(FreeLargePagesHead[AllocationHint]);
-					GVeryLargePageAllocatorPreAllocatePools = false;
-					break;
-				}
+					if (!Block.Commit(LargePage->BaseAddress - AddressSpaceReserved, SizeOfLargePage, false))
+					{
+						// Cant commit 2MB pages, stop preallocating and return page to FreeLargePagesHead list
+						LargePage->LinkHead(FreeLargePagesHead[AllocationHint]);
+						GVeryLargePageAllocatorPreAllocatePools = false;
+						break;
+					}
 
-				LargePage->LinkHead(EmptyButAvailableLargePagesHead[AllocationHint]);
-				CachedFree += SizeOfLargePage;
-				CommitedLargePagesCount[AllocationHint] += 1;
-				EmptyBackStoreCount[AllocationHint] += 1;
+					LargePage->LinkHead(EmptyButAvailableLargePagesHead[AllocationHint]);
+					CachedFree += SizeOfLargePage;
+					CommitedLargePagesCount[AllocationHint] += 1;
+					EmptyBackStoreCount[AllocationHint] += 1;
+				}
 			}
 		}
-		else
-		{
-			ShrinkEmptyBackStore(GVeryLargePageAllocatorMaxEmptyBackStoreCount[AllocationHint], AllocationHint);
-		}
+
+		uint64 LenghtCycles64 = (FPlatformTime::Cycles64() - StartTime);
+		UE_LOG(LogMemory, Log, TEXT("LargePageAllocator - Preallocate memory duration: %.1f ms"), float(FPlatformTime::ToMilliseconds64(LenghtCycles64)));
 	}
 }
 
