Added Tanh and Cosh Math functions.

diff --git a/Engine/Source/Runtime/Core/Public/GenericPlatform/GenericPlatformMath.h b/Engine/Source/Runtime/Core/Public/GenericPlatform/GenericPlatformMath.h
--- a/Engine/Source/Runtime/Core/Public/GenericPlatform/GenericPlatformMath.h
+++ b/Engine/Source/Runtime/Core/Public/GenericPlatform/GenericPlatformMath.h
@@ -525,12 +525,18 @@ struct FGenericPlatformMath
 	static FORCEINLINE float Acos( float Value ) { return acosf( (Value<-1.f) ? -1.f : ((Value<1.f) ? Value : 1.f) ); }
 	static FORCEINLINE double Acos( double Value ) { return acos( (Value<-1.0) ? -1.0 : ((Value<1.0) ? Value : 1.0) ); }
 
+	static FORCEINLINE float Cosh(float Value) { return coshf(Value); }
+	static FORCEINLINE double Cosh(double Value) { return cosh(Value); }
+
 	static FORCEINLINE float Tan( float Value ) { return tanf(Value); }
 	static FORCEINLINE double Tan( double Value ) { return tan(Value); }
 
 	static FORCEINLINE float Atan( float Value ) { return atanf(Value); }
 	static FORCEINLINE double Atan( double Value ) { return atan(Value); }
 
+	static FORCEINLINE float Tanh(float Value) { return tanhf(Value); }
+	static FORCEINLINE double Tanh(double Value) { return tanh(Value); }
+
 	static CORE_API float Atan2( float Y, float X );
 	static CORE_API double Atan2( double Y, double X );
 	RESOLVE_FLOAT_AMBIGUITY_2_ARGS(Atan2);
