refactored FMRSWRecursiveAccessDetector to support partially overlapping concurrent read accesses

diff --git a/Engine/Source/Runtime/Core/Private/Misc/MTAccessDetector.cpp b/Engine/Source/Runtime/Core/Private/Misc/MTAccessDetector.cpp
--- a/Engine/Source/Runtime/Core/Private/Misc/MTAccessDetector.cpp
+++ b/Engine/Source/Runtime/Core/Private/Misc/MTAccessDetector.cpp
@@ -403,17 +403,14 @@ bool FRWRecursiveAccessDetectorTest::RunTest(const FString& Parameters)
 	{	// destroying detector from inside access scope
 		auto* AD = new FMRSWRecursiveAccessDetector; //-V774: "The 'AD' pointer was used after the memory was released." - no, it wasn't because of the `if`
 
-		AD->AcquireReadAccess();
-
 		FMRSWRecursiveAccessDetector::FDestructionSentinel DestructionSentinel{ FMRSWRecursiveAccessDetector::EAccessType::Reader };
-		FMRSWRecursiveAccessDetector::FDestructionSentinel* PrevDestructionSentinel = AD->SetDestructionSentinel(&DestructionSentinel);
+		AD->AcquireReadAccess(DestructionSentinel);
 
 		delete AD;
 
 		if (!DestructionSentinel.bDestroyed)
 		{
-			AD->SetDestructionSentinel(PrevDestructionSentinel);  //-V774
-			AD->ReleaseReadAccess();  //-V774
+			AD->ReleaseReadAccess(DestructionSentinel);  //-V774
 		}
 	}
 
diff --git a/Engine/Source/Runtime/Core/Public/Misc/MTAccessDetector.h b/Engine/Source/Runtime/Core/Public/Misc/MTAccessDetector.h
--- a/Engine/Source/Runtime/Core/Public/Misc/MTAccessDetector.h
+++ b/Engine/Source/Runtime/Core/Public/Misc/MTAccessDetector.h
@@ -290,6 +290,7 @@ private:
 		constexpr FState()
 			: Value(0)
 		{
+			static_assert(sizeof(FState) == sizeof(uint64)); // `FState` is stored in `std::atomic<uint64>`
 		}
 
 		constexpr FState(uint64 InValue)
@@ -344,7 +345,13 @@ private:
 		}
 	}
 
-public: // DestructionSentinel
+//////////////////////////////////////////////////////////////////////
+// DestructionSentinel
+// This access detector supports its destruction (along with its owner) while being "accessed". If this is expected, acquire/release access 
+// overloads must be used with a destruction sentinel stored on the callstack. If destroyed, `~FMRSWRecursiveAccessDetector()` will 
+// clean up after the destroyed instance and notify the user (`FDestructionSentinel::bDestroyed`) that the release method can't be called as 
+// the instance was destroyed.
+public:
 	enum class EAccessType { Reader, Writer };
 
 	struct FDestructionSentinel
@@ -355,21 +362,15 @@ public: // DestructionSentinel
 		}
 
 		EAccessType AccessType;
+		const FMRSWRecursiveAccessDetector* Accessor;
 		bool bDestroyed = false;
 	};
 
 private:
-	mutable FDestructionSentinel* DestructionSentinel = nullptr;
-
-	// if access detector was destroyed while holding read accessed (same thread) it can't release the access, but is still registered in TLS.
-	// use this static method to clear it
-	FORCEINLINE static void RemoveReader(FMRSWRecursiveAccessDetector* Reader)
-	{
-		uint32 ReaderIndex = ReadersTls.IndexOfByPredicate([Reader](FReaderNum ReaderNum) { return ReaderNum.Reader == Reader; });
-		checkfSlow(ReaderIndex != INDEX_NONE, TEXT("Can't find the reader in the TLS table"));
-		checkfSlow(ReadersTls[ReaderIndex].Num == 1, TEXT("More than one accesses detected while destoying access detector: %d"), ReadersTls[ReaderIndex].Num);
-		ReadersTls.RemoveAtSwap(ReaderIndex, 1, /* bAllowShrinking = */ false);
-	}
+	// a TLS stack of destruction sentinels shared by all access detector instances on the callstack. the same access detector instance can have
+	// multiple destruction sentinels on the stack. the next "release access" with destruction sentinel always matches the top of the stack
+	inline static thread_local TArray<FDestructionSentinel*, TInlineAllocator<4>> DestructionSentinelStackTls;
+//////////////////////////////////////////////////////////////////////
 
 private:
 	///////////////////////////////////////////////
@@ -383,6 +384,19 @@ private:
 	};
 
 	inline static thread_local TArray<FReaderNum, TInlineAllocator<4>> ReadersTls;
+
+	void RemoveReaderFromTls() const
+	{
+		uint32 ReaderIndex = ReadersTls.IndexOfByPredicate([this](FReaderNum ReaderNum) { return ReaderNum.Reader == this; });
+		checkfSlow(ReaderIndex != INDEX_NONE,
+			TEXT("Invalid usage of the race detector! No matching AcquireReadAccess(): %u readers, %u writers on thread %d"),
+			LoadState().ReaderNum, LoadState().WriterNum, LoadState().WriterThreadId - 1);
+		uint32 ReaderNum = --ReadersTls[ReaderIndex].Num;
+		if (ReaderNum == 0)
+		{
+			ReadersTls.RemoveAtSwap(ReaderIndex, 1, /* bAllowShrinking = */ false);
+		}
+	}
 	///////////////////////////////////////////////
 
 public:
@@ -405,30 +419,37 @@ public:
 
 	FORCEINLINE ~FMRSWRecursiveAccessDetector()
 	{
-		if (DestructionSentinel)
+		// search for all destruction sentinels for this instance and remove them from the stack, while building an expected correct state of
+		// the access detector.
+		FState ExpectedState;
+		for (int DestructionSentinelIndex = 0; DestructionSentinelIndex != DestructionSentinelStackTls.Num(); ++DestructionSentinelIndex)
 		{
-			DestructionSentinel->bDestroyed = true;
-			if (DestructionSentinel->AccessType == EAccessType::Reader)
-			{	// readers are registered in TLS table and must be cleared on destruction
-				RemoveReader(this);
+			FDestructionSentinel& DestructionSentinel = *DestructionSentinelStackTls[DestructionSentinelIndex];
+			if (DestructionSentinel.Accessor == this)
+			{
+				DestructionSentinel.bDestroyed = true;
+				if (DestructionSentinel.AccessType == EAccessType::Reader)
+				{
+					++ExpectedState.ReaderNum;
+
+					RemoveReaderFromTls();
+				}
+				else
+				{
+					++ExpectedState.WriterNum;
+					ExpectedState.WriterThreadId = FPlatformTLS::GetCurrentThreadId() + 1; // to shift from 0 TID that is considered "invalid"
+				}
+
+				DestructionSentinelStackTls.RemoveAtSwap(DestructionSentinelIndex);
+				--DestructionSentinelIndex;
 			}
 		}
-		else
-		{
-			checkf(LoadState().Value == 0,
-				TEXT("Race detector destroyed while being accessed: %u readers, %u writers on thread %d"),
-				LoadState().ReaderNum, LoadState().WriterNum, LoadState().WriterThreadId - 1);
-		}
-	}
 
-	// supports destroying an instance while it's being accessed by notifying the user about destruction so the user don't release access to
-	// an already destroyed instance. Should be used only from inside an access scope.
-	// @return previously set destruction sentinel that should be set back before leaving an access scope
-	FORCEINLINE FDestructionSentinel* SetDestructionSentinel(FDestructionSentinel* InDestructionSentinel) const
-	{
-		FDestructionSentinel* PrevDestructionSentinel = DestructionSentinel;
-		DestructionSentinel = InDestructionSentinel;
-		return PrevDestructionSentinel;
+		checkf(LoadState().Value == ExpectedState.Value,
+			TEXT("Race detector destroyed while being accessed on another thread: %d readers, %d writers on thread %d"),
+			LoadState().ReaderNum - ExpectedState.ReaderNum, 
+			LoadState().WriterNum - ExpectedState.WriterNum, 
+			LoadState().WriterThreadId - 1);
 	}
 
 	FORCEINLINE void AcquireReadAccess() const
@@ -451,23 +472,36 @@ public:
 		}
 	}
 
-	FORCEINLINE void ReleaseReadAccess() const
+	// an overload that handles access detector destruction from inside a read access, must be used along with the corresponding
+	// overload of `ReleaseReadAcess`
+	FORCEINLINE void AcquireReadAccess(FDestructionSentinel& DestructionSentinel) const
 	{
-		// unregister the reader from TLS
-		uint32 ReaderIndex = ReadersTls.IndexOfByPredicate([this](FReaderNum ReaderNum) { return ReaderNum.Reader == this; });
-		checkfSlow(ReaderIndex != INDEX_NONE, 
-			TEXT("Invalid usage of the race detector! No matching AcquireReadAccess(): %u readers, %u writers on thread %d"),
-			LoadState().ReaderNum, LoadState().WriterNum, LoadState().WriterThreadId - 1);
-		uint32 ReaderNum = --ReadersTls[ReaderIndex].Num;
-		if (ReaderNum == 0)
-		{
-			ReadersTls.RemoveAtSwap(ReaderIndex, 1, /* bAllowShrinking = */ false);
-		}
+		DestructionSentinel.Accessor = this;
+		DestructionSentinelStackTls.Add(&DestructionSentinel);
 
+		AcquireReadAccess();
+	}
+
+	FORCEINLINE void ReleaseReadAccess() const
+	{
+		RemoveReaderFromTls();
 		DecrementReaderNum();
 		// no need to check for writers
 	}
 
+	// an overload that handles access detector destruction from inside a read access, must be used along with the corresponding
+	// overload of `AcquireReadAcess`
+	FORCEINLINE void ReleaseReadAccess(FDestructionSentinel& DestructionSentinel) const
+	{
+		ReleaseReadAccess();
+
+		checkSlow(DestructionSentinel.Accessor == this);
+		checkfSlow(DestructionSentinelStackTls.Num() != 0, TEXT("An attempt to remove a not registered destruction sentinel"));
+		checkfSlow(DestructionSentinelStackTls.Last() == &DestructionSentinel, TEXT("Mismatched destruction sentinel: %p != %p"), DestructionSentinelStackTls.Last(), &DestructionSentinel);
+
+		DestructionSentinelStackTls.RemoveAtSwap(DestructionSentinelStackTls.Num() - 1);
+	}
+
 	FORCEINLINE void AcquireWriteAccess()
 	{
 		FState LocalState = LoadState();
@@ -499,6 +533,16 @@ public:
 			LocalState.WriterThreadId - 1, PrevState.WriterThreadId - 1);
 	}
 
+	// an overload that handles access detector destruction from inside a write access, must be used along with the corresponding
+	// overload of `ReleaseWriteAcess`
+	FORCEINLINE void AcquireWriteAccess(FDestructionSentinel& DestructionSentinel)
+	{
+		DestructionSentinel.Accessor = this;
+		DestructionSentinelStackTls.Add(&DestructionSentinel);
+
+		AcquireWriteAccess();
+	}
+
 	FORCEINLINE void ReleaseWriteAccess()
 	{
 		FState LocalState = LoadState();
@@ -513,6 +557,19 @@ public:
 			LocalState.WriterNum, PrevState.WriterNum,
 			LocalState.WriterThreadId - 1, PrevState.WriterThreadId - 1);
 	}
+
+	// an overload that handles access detector destruction from inside a write access, must be used along with the corresponding
+	// overload of `AcquireWriteAcess`
+	FORCEINLINE void ReleaseWriteAccess(FDestructionSentinel& DestructionSentinel)
+	{
+		ReleaseWriteAccess();
+
+		checkSlow(DestructionSentinel.Accessor == this);
+		checkfSlow(DestructionSentinelStackTls.Num() != 0, TEXT("An attempt to remove a not registered destruction sentinel"));
+		checkfSlow(DestructionSentinelStackTls.Last() == &DestructionSentinel, TEXT("Mismatched destruction sentinel: %p != %p"), DestructionSentinelStackTls.Last(), &DestructionSentinel);
+
+		DestructionSentinelStackTls.RemoveAtSwap(DestructionSentinelStackTls.Num() - 1);
+	}
 };
 
 ///////////////////////////////////////////////////////////////////
