Added a missing include to ContainerAllocationPolicies.cpp

#rnx

diff --git a/Engine/Source/Runtime/Core/Private/Containers/ContainerAllocationPolicies.cpp b/Engine/Source/Runtime/Core/Private/Containers/ContainerAllocationPolicies.cpp
--- a/Engine/Source/Runtime/Core/Private/Containers/ContainerAllocationPolicies.cpp
+++ b/Engine/Source/Runtime/Core/Private/Containers/ContainerAllocationPolicies.cpp
@@ -1,6 +1,7 @@
 // Copyright Epic Games, Inc. All Rights Reserved.
 
 #include "Containers/ContainerAllocationPolicies.h"
+#include "CoreGlobals.h"
 
 FORCENOINLINE void UE::Core::Private::OnInvalidAlignedHeapAllocatorNum(int32 NewNum, SIZE_T NumBytesPerElement)
 {
