Added a missing includes of CoreGlobals.h

#rnx

diff --git a/Engine/Source/Runtime/Core/Private/HAL/Allocators/AnsiAllocator.cpp b/Engine/Source/Runtime/Core/Private/HAL/Allocators/AnsiAllocator.cpp
--- a/Engine/Source/Runtime/Core/Private/HAL/Allocators/AnsiAllocator.cpp
+++ b/Engine/Source/Runtime/Core/Private/HAL/Allocators/AnsiAllocator.cpp
@@ -1,6 +1,7 @@
 // Copyright Epic Games, Inc. All Rights Reserved.
 
 #include "HAL/Allocators/AnsiAllocator.h"
+#include "CoreGlobals.h"
 
 FORCENOINLINE void UE::Core::Private::OnInvalidAnsiAllocatorNum(int32 NewNum, SIZE_T NumBytesPerElement)
 {
diff --git a/Engine/Source/Runtime/Core/Private/HAL/LowLevelMemoryUtils.cpp b/Engine/Source/Runtime/Core/Private/HAL/LowLevelMemoryUtils.cpp
--- a/Engine/Source/Runtime/Core/Private/HAL/LowLevelMemoryUtils.cpp
+++ b/Engine/Source/Runtime/Core/Private/HAL/LowLevelMemoryUtils.cpp
@@ -1,6 +1,7 @@
 // Copyright Epic Games, Inc. All Rights Reserved.
 
 #include "LowLevelMemoryUtils.h"
+#include "CoreGlobals.h"
 
 #if ENABLE_LOW_LEVEL_MEM_TRACKER
 
