Logging: Adjusted Fatal verbosity comparisons to satisfy static analysis

#rnx

diff --git a/Engine/Source/Runtime/Core/Public/Logging/LogMacros.h b/Engine/Source/Runtime/Core/Public/Logging/LogMacros.h
--- a/Engine/Source/Runtime/Core/Public/Logging/LogMacros.h
+++ b/Engine/Source/Runtime/Core/Public/Logging/LogMacros.h
@@ -150,7 +150,7 @@ CORE_API void BasicFatalLog(const FLogCategoryBase& Category, const FStaticBasic
 	// This will only log Fatal errors
 	#define UE_LOG(CategoryName, Verbosity, Format, ...) \
 	{ \
-		if constexpr (ELogVerbosity::Verbosity == ELogVerbosity::Fatal) \
+		if constexpr ((ELogVerbosity::Verbosity & ELogVerbosity::VerbosityMask) == ELogVerbosity::Fatal) \
 		{ \
 			LowLevelFatalError(Format, ##__VA_ARGS__); \
 			CA_ASSUME(false); \
@@ -158,7 +158,7 @@ CORE_API void BasicFatalLog(const FLogCategoryBase& Category, const FStaticBasic
 	}
 	#define UE_LOG_REF(CategoryRef, Verbosity, Format, ...) \
 	{ \
-		if constexpr (ELogVerbosity::Verbosity == ELogVerbosity::Fatal) \
+		if constexpr ((ELogVerbosity::Verbosity & ELogVerbosity::VerbosityMask) == ELogVerbosity::Fatal) \
 		{ \
 			LowLevelFatalError(Format, ##__VA_ARGS__); \
 			CA_ASSUME(false); \
@@ -166,7 +166,7 @@ CORE_API void BasicFatalLog(const FLogCategoryBase& Category, const FStaticBasic
 	}
 	#define UE_LOG_CLINKAGE(CategoryName, Verbosity, Format, ...) \
 	{ \
-		if constexpr (ELogVerbosity::Verbosity == ELogVerbosity::Fatal) \
+		if constexpr ((ELogVerbosity::Verbosity & ELogVerbosity::VerbosityMask) == ELogVerbosity::Fatal) \
 		{ \
 			LowLevelFatalError(Format, ##__VA_ARGS__); \
 			CA_ASSUME(false); \
@@ -176,7 +176,7 @@ CORE_API void BasicFatalLog(const FLogCategoryBase& Category, const FStaticBasic
 	// Conditional logging (fatal errors only).
 	#define UE_CLOG(Condition, CategoryName, Verbosity, Format, ...) \
 	{ \
-		if constexpr (ELogVerbosity::Verbosity == ELogVerbosity::Fatal) \
+		if constexpr ((ELogVerbosity::Verbosity & ELogVerbosity::VerbosityMask) == ELogVerbosity::Fatal) \
 		{ \
 			if (Condition) \
 			{ \
@@ -303,7 +303,7 @@ CORE_API void BasicFatalLog(const FLogCategoryBase& Category, const FStaticBasic
 		static ::UE::Logging::Private::FStaticBasicLogDynamicData LOG_Dynamic; \
 		static constexpr ::UE::Logging::Private::FStaticBasicLogRecord LOG_Static(Format, __FILE__, __LINE__, ::ELogVerbosity::Verbosity, LOG_Dynamic); \
 		static_assert((::ELogVerbosity::Verbosity & ::ELogVerbosity::VerbosityMask) < ::ELogVerbosity::NumVerbosity && ::ELogVerbosity::Verbosity > 0, "Verbosity must be constant and in range."); \
-		if constexpr (::ELogVerbosity::Verbosity == ::ELogVerbosity::Fatal) \
+		if constexpr ((::ELogVerbosity::Verbosity & ELogVerbosity::VerbosityMask) == ::ELogVerbosity::Fatal) \
 		{ \
 			Condition \
 			{ \
diff --git a/Engine/Source/Runtime/Core/Public/Logging/StructuredLog.h b/Engine/Source/Runtime/Core/Public/Logging/StructuredLog.h
--- a/Engine/Source/Runtime/Core/Public/Logging/StructuredLog.h
+++ b/Engine/Source/Runtime/Core/Public/Logging/StructuredLog.h
@@ -356,14 +356,14 @@ FORCENOINLINE UE_DEBUG_SECTION void LogWithFields(typename TFieldArgType<FieldAr
 template <const auto& Log, typename LogCategoryType, ELogVerbosity::Type Verbosity, typename... FieldArgTypes>
 inline void LogIfActive(FieldArgTypes&&... FieldArgs)
 {
-	static_assert(Verbosity != ELogVerbosity::Fatal, "Fatal verbosity is not supported by this API at this time.");
+	static_assert((Verbosity & ELogVerbosity::VerbosityMask) != ELogVerbosity::Fatal, "Fatal verbosity is not supported by this API at this time.");
 	static_assert((Verbosity & ELogVerbosity::VerbosityMask) < ELogVerbosity::NumVerbosity && Verbosity > 0, "Verbosity must be constant and in range.");
 
 	if constexpr (
 		(Verbosity & ELogVerbosity::VerbosityMask) <= ELogVerbosity::COMPILED_IN_MINIMUM_VERBOSITY &&
 		(Verbosity & ELogVerbosity::VerbosityMask) <= LogCategoryType::CompileTimeVerbosity)
 	{
-		if (Verbosity == ELogVerbosity::Fatal || !Log.Category.IsSuppressed(Verbosity))
+		if ((Verbosity & ELogVerbosity::VerbosityMask) == ELogVerbosity::Fatal || !Log.Category.IsSuppressed(Verbosity))
 		{
 			if constexpr (sizeof...(FieldArgTypes) == 0)
 			{
@@ -376,7 +376,7 @@ inline void LogIfActive(FieldArgTypes&&... FieldArgs)
 		}
 	}
 
-	if constexpr (Verbosity == ELogVerbosity::Fatal)
+	if constexpr ((Verbosity & ELogVerbosity::VerbosityMask) == ELogVerbosity::Fatal)
 	{
 		UE_DEBUG_BREAK_AND_PROMPT_FOR_REMOTE();
 		FDebug::ProcessFatalError(PLATFORM_RETURN_ADDRESS());
