Fixed UE_REQUIRES for enum arguments.

diff --git a/Engine/Source/Runtime/Core/Public/Templates/UnrealTypeTraits.h b/Engine/Source/Runtime/Core/Public/Templates/UnrealTypeTraits.h
--- a/Engine/Source/Runtime/Core/Public/Templates/UnrealTypeTraits.h
+++ b/Engine/Source/Runtime/Core/Public/Templates/UnrealTypeTraits.h
@@ -42,7 +42,7 @@
 		concept BoolIdentityConcept = B;
 	}
 
-	#define UE_REQUIRES(...) > requires (__VA_ARGS__) && UE::Core::Private::BoolIdentityConcept<true
+	#define UE_REQUIRES(...) > requires (!!(__VA_ARGS__)) && UE::Core::Private::BoolIdentityConcept<true
 #endif
 
 #define TEMPLATE_REQUIRES(...) typename TEnableIf<__VA_ARGS__, int>::type = 0
