Fix missing channel check for metadata restore scopes.

When creating the custom scope type, a check if the channel is active was missing causing the scopes to emit events regardless if the metadata channel was active.

diff --git a/Engine/Source/Runtime/Core/Private/ProfilingDebugging/MetadataTrace.cpp b/Engine/Source/Runtime/Core/Private/ProfilingDebugging/MetadataTrace.cpp
--- a/Engine/Source/Runtime/Core/Private/ProfilingDebugging/MetadataTrace.cpp
+++ b/Engine/Source/Runtime/Core/Private/ProfilingDebugging/MetadataTrace.cpp
@@ -22,7 +22,7 @@ uint32 FMetadataTrace::SaveStack()
 ////////////////////////////////////////////////////////////////////////////////////////////////////
 FMetadataRestoreScope::FMetadataRestoreScope(uint32 SavedMetadataIdentifier)
 {
-	if (SavedMetadataIdentifier)
+	if (UE_TRACE_CHANNELEXPR_IS_ENABLED(MetadataChannel) & (SavedMetadataIdentifier != 0))
 	{
 		ActivateScope(SavedMetadataIdentifier);
 	}
