Added 501 warning ingore for CHECK_MESSAGE macros

diff --git a/Engine/Source/Runtime/Core/Public/Misc/LowLevelTestAdapter.h b/Engine/Source/Runtime/Core/Public/Misc/LowLevelTestAdapter.h
--- a/Engine/Source/Runtime/Core/Public/Misc/LowLevelTestAdapter.h
+++ b/Engine/Source/Runtime/Core/Public/Misc/LowLevelTestAdapter.h
@@ -70,8 +70,9 @@
 #define INFO(What)
 #define CHECK(Expr) if (!(Expr)) { FAutomationTestFramework::Get().GetCurrentTest()->AddError(TEXT("Condition failed")); }
 #define CHECK_FALSE(Expr) if (Expr) { FAutomationTestFramework::Get().GetCurrentTest()->AddError(TEXT("Condition expected to return false but returned true")); }
-//-V:CHECK_MESSAGE:571
+//-V:CHECK_MESSAGE:571,501
 #define CHECK_MESSAGE(Message, Expr) if (!(Expr)) { FAutomationTestFramework::Get().GetCurrentTest()->AddError(Message); }
+//-V:CHECK_FALSE_MESSAGE:571,501
 #define CHECK_FALSE_MESSAGE(Message, Expr) if (Expr) { FAutomationTestFramework::Get().GetCurrentTest()->AddError(Message); }
 #define REQUIRE(Expr) if (!(Expr)) { FAutomationTestFramework::Get().GetCurrentTest()->AddError(TEXT("Required condition failed, interrupting test")); return; }
 //-V:REQUIRE_MESSAGE:571
