Fixed potential invalid quote count in CSV parser.
Added overloaded constructors for FCsvParser which move/copy rather than (move/copy)+move.
Fixed a rather weird const_cast construct.

diff --git a/Engine/Source/Runtime/Core/Private/Serialization/Csv/CsvParser.cpp b/Engine/Source/Runtime/Core/Private/Serialization/Csv/CsvParser.cpp
--- a/Engine/Source/Runtime/Core/Private/Serialization/Csv/CsvParser.cpp
+++ b/Engine/Source/Runtime/Core/Private/Serialization/Csv/CsvParser.cpp
@@ -4,15 +4,24 @@
 
 #include "Templates/UnrealTemplate.h"
 
-FCsvParser::FCsvParser(FString InSourceString)
-	: SourceString(MoveTemp(InSourceString))		
+FCsvParser::FCsvParser(FString&& InSourceString)
+	: SourceString(MoveTemp(InSourceString))
 {
 	if (!SourceString.IsEmpty())
 	{
 		ParseRows();
 	}
 }
-		
+
+FCsvParser::FCsvParser(const FString& InSourceString)
+	: SourceString(InSourceString)
+{
+	if (!SourceString.IsEmpty())
+	{
+		ParseRows();
+	}
+}
+
 void FCsvParser::ParseRows()
 {
 	BufferStart = &SourceString[0];
@@ -49,7 +58,7 @@ FCsvParser::EParseResult FCsvParser::ParseRow()
 
 FCsvParser::EParseResult FCsvParser::ParseCell()
 {
-	TCHAR* WriteAt = BufferStart + (ReadAt - BufferStart);
+	TCHAR* WriteAt = const_cast<TCHAR*>(ReadAt);
 
 	// Check if this cell is quoted. Whitespace between cell opening and quote is invalid.
 	bool bQuoted = *ReadAt == '"';
@@ -71,22 +80,20 @@ FCsvParser::EParseResult FCsvParser::ParseCell()
 			{
 				// RFC 4180 specifies that double quotes are escaped as ""
 
-				int32 NumQuotes = 0;
-				while (*(ReadAt + NumQuotes) == '"') ++NumQuotes;
-
-				if (NumQuotes % 2 != 0)
+				while (*ReadAt == '"')
 				{
-					// No longer quoted if there are an even number of quotes before this one
-					bQuoted = false;
+					++ReadAt;
+					bQuoted = !bQuoted;
+
+					// Unescape the double quotes
+					if (bQuoted)
+					{
+						*WriteAt++ = TEXT('"');
+					}
 				}
 
-				ReadAt += NumQuotes;
-
-				// Unescape the double quotes
 				// We null terminate and leave the write pos pointing at the trailing closing quote 
 				// if present so it gets overwritten by any subsequent text in the cell
-				NumQuotes /= 2;
-				while(NumQuotes-- > 0) *(WriteAt++) = TEXT('"');
 				*WriteAt = TEXT('\0');
 
 				continue;
diff --git a/Engine/Source/Runtime/Core/Public/Serialization/Csv/CsvParser.h b/Engine/Source/Runtime/Core/Public/Serialization/Csv/CsvParser.h
--- a/Engine/Source/Runtime/Core/Public/Serialization/Csv/CsvParser.h
+++ b/Engine/Source/Runtime/Core/Public/Serialization/Csv/CsvParser.h
@@ -16,7 +16,8 @@ struct FCsvParser
 	 *	Takes a copy because the parser tramples over the file to create null-terminated cell strings
 	 *	Avoid a copy by passing overship of the source string with MoveTemp()
 	 */
-	CORE_API FCsvParser(FString InSourceString);
+	CORE_API FCsvParser(FString&& InSourceString);
+	CORE_API FCsvParser(const FString& InSourceString);
 
 	/**
 	 *	Const access to the parsed rows
@@ -28,11 +29,13 @@ struct FCsvParser
 		return Rows;
 	}
 
-private:
 	/** Non copyable - has pointers to itself */
-	FCsvParser(const FCsvParser&);
-	FCsvParser& operator=(const FCsvParser&);
+	FCsvParser(FCsvParser&&) = delete;
+	FCsvParser(const FCsvParser&) = delete;
+	FCsvParser& operator=(FCsvParser&&) = delete;
+	FCsvParser& operator=(const FCsvParser&) = delete;
 
+private:
 	/** Internal enum to control parsing progress */
 	enum struct EParseResult { EndOfCell, EndOfRow, EndOfString };
 
