UE::Tasks::Prerequisites() helper function supports containers as an input parameter, and returns TStaticArray instead of a custom wrapper over it

diff --git a/Engine/Source/Runtime/Core/Private/Tests/Tasks/TasksTest.cpp b/Engine/Source/Runtime/Core/Private/Tests/Tasks/TasksTest.cpp
--- a/Engine/Source/Runtime/Core/Private/Tests/Tasks/TasksTest.cpp
+++ b/Engine/Source/Runtime/Core/Private/Tests/Tasks/TasksTest.cpp
@@ -354,6 +354,11 @@ namespace UE { namespace TasksTests
 			Outer.Wait();
 		}
 
+		{ // check that `Prerequisites()` support containers
+			Prerequisites(TArray<FTask>{});
+			Prerequisites(TArrayView<FTask>{});
+		}
+
 #if TASKGRAPH_NEW_FRONTEND
 		{	// a basic test for a named thread task
 			FTask GTTask = Launch
diff --git a/Engine/Source/Runtime/Core/Public/Tasks/Task.h b/Engine/Source/Runtime/Core/Public/Tasks/Task.h
--- a/Engine/Source/Runtime/Core/Public/Tasks/Task.h
+++ b/Engine/Source/Runtime/Core/Public/Tasks/Task.h
@@ -14,9 +14,6 @@ namespace UE::Tasks
 	template<typename ResultType>
 	class TTask;
 
-	template<typename... TaskTypes>
-	class TPrerequisites;
-
 	template<typename TaskCollectionType>
 	void Wait(const TaskCollectionType& Tasks);
 
@@ -39,8 +36,11 @@ namespace UE::Tasks
 			// friends to get access to `Pimpl`
 			friend FTaskBase;
 
-			template<typename... TaskTypes>
-			friend class UE::Tasks::TPrerequisites;
+			template<int32 Index, typename ArrayType, typename FirstTaskType, typename... OtherTasksTypes>
+			friend void PrerequisitesUnpacker(ArrayType& Array, FirstTaskType& FirstTask, OtherTasksTypes&... OtherTasks);
+
+			template<int32 Index, typename ArrayType, typename TaskType>
+			friend void PrerequisitesUnpacker(ArrayType& Array, TaskType& FirstTask);
 
 			template<typename TaskCollectionType>
 			friend bool TryRetractAndExecute(const TaskCollectionType& Tasks);
@@ -69,6 +69,8 @@ namespace UE::Tasks
 			{}
 
 		public:
+			using FTaskHandleId = void;
+
 			FTaskHandle() = default;
 
 			bool IsValid() const
@@ -431,36 +433,35 @@ namespace UE::Tasks
 
 	using FTask = Private::FTaskHandle;
 
-	// a convenient proxy collection for specifying task prerequisites that can include both tasks and task events
-	// usage: Launch(UE_SOURCE_LOCATION, [] {}, Prerequisites(Task1, Task2, TaskEvent1, ...));
-	template<typename... TaskTypes>
-	class TPrerequisites : public TStaticArray<Private::FTaskBase*, sizeof...(TaskTypes)>
+	namespace Private
 	{
-	public:
-		TPrerequisites(TaskTypes&&... Tasks)
+		template<int32 Index, typename ArrayType, typename FirstTaskType, typename... OtherTasksTypes>
+		void PrerequisitesUnpacker(ArrayType& Array, FirstTaskType& FirstTask, OtherTasksTypes&... OtherTasks)
 		{
-			Fill(0, Forward<TaskTypes>(Tasks)...);
+			Array[Index] = FirstTask.Pimpl.GetReference();
+			PrerequisitesUnpacker<Index + 1>(Array, OtherTasks...);
 		}
 
-	private:
-		template<typename FirstTaskType, typename... OtherTaskTypes>
-		void Fill(uint32 Index, FirstTaskType&& FirstTask, OtherTaskTypes&&... OtherTasks)
+		template<int32 Index, typename ArrayType, typename TaskType>
+		void PrerequisitesUnpacker(ArrayType& Array, TaskType& Task)
 		{
-			(*this)[Index] = FirstTask.Pimpl.GetReference();
-			Fill(Index + 1, Forward<OtherTaskTypes>(OtherTasks)...);
+			Array[Index] = Task.Pimpl.GetReference();
 		}
+	}
 
-		template<typename TaskType>
-		void Fill(uint32 Index, TaskType&& Task)
-		{
-			(*this)[Index] = Task.Pimpl.GetReference();
-		}
-	};
+	template<typename... TaskTypes, 
+		typename std::decay_t<decltype(std::declval<TTuple<TaskTypes...>>().template Get<0>())>::FTaskHandleId* = nullptr>
+	TStaticArray<Private::FTaskBase*, sizeof...(TaskTypes)> Prerequisites(TaskTypes&... Tasks)
+	{
+		TStaticArray<Private::FTaskBase*, sizeof...(TaskTypes)> Res;
+		Private::PrerequisitesUnpacker<0>(Res, Tasks...);
+		return Res;
+	}
 
-	template<typename... TaskTypes>
-	TPrerequisites<TaskTypes...> Prerequisites(TaskTypes&... Tasks)
+	template<typename TaskCollectionType>
+	const TaskCollectionType& Prerequisites(const TaskCollectionType& Tasks)
 	{
-		return TPrerequisites<TaskTypes...>{ Forward<TaskTypes>(Tasks)... };
+		return Tasks;
 	}
 
 	// Adds the nested task to the task that is being currently executed by the current thread. A parent task is not flagged completed
@@ -513,9 +514,3 @@ namespace UE::Tasks
 		EExtendedTaskPriority ExtendedPriority;
 	};
 }
-
-template <typename... TaskTypes>
-struct TIsContiguousContainer<UE::Tasks::TPrerequisites<TaskTypes...>>
-{
-	static constexpr bool Value = true;
-};
