Memory: Improved memory tracking on some platform to avoid having a very large (hundreds on MB) of Untracked memory.

- On some platforms we set aside a large pool of memory for frame buffers, LLM and possible (though unlikely) OS use. This memory is not availabe for normal use for the engine, so modified the stats to explicitly treat is as allocated.

- Previously, TotalPhysical != UsedPhysical + AvailablePhysical, which led to unclear results in case of some memory calculations. Now, this basic equation is true with UsedPhysical being the sum of the buffer set aside and the normal physical pages requested by the engine.

- Introduced InternalHeap memory tag to account for this special memory.

- In subsequent changes we might add InternalHeapUsed to show how much of this buffer is reallly used, but it might be inaccurate because we can't account for the OS using this pool as it's outside of our control. This case is rare/unlikely on the affected platforms, so we could be of some use.

- Added FPlatformMemory::TracksLLMAllocations() which allows platforms to explicitly say if they account for LLM overhead or not.

- Separated physical memory counter for the Main engine memory and FrameBuffers. Since they form completely seprate pools, it didn't make sense to combine them. It only really matters for the non-default mode withuot the pages cache.

#rn fix ps5 The memory pool reserved for special use (e.g. frame buffers) is no longer treated as 'untracked' by the memory stats. It's accounted for as InternalHeap memory tag now.

diff --git a/Engine/Source/Runtime/Core/Private/HAL/LowLevelMemTracker.cpp b/Engine/Source/Runtime/Core/Private/HAL/LowLevelMemTracker.cpp
--- a/Engine/Source/Runtime/Core/Private/HAL/LowLevelMemTracker.cpp
+++ b/Engine/Source/Runtime/Core/Private/HAL/LowLevelMemTracker.cpp
@@ -955,7 +955,7 @@ void FLowLevelMemTracker::TickInternal()
 
 	int64 StaticOverhead = sizeof(FLowLevelMemTracker);
 	MemoryUsageCurrentOverhead = StaticOverhead + Allocator.GetTotal();
-	PlatformTracker.SetTagAmountInUpdate(FindOrAddTagData(ELLMTag::PlatformOverhead), MemoryUsageCurrentOverhead, true);
+	PlatformTracker.SetTagAmountInUpdate(FindOrAddTagData(ELLMTag::PlatformOverhead), MemoryUsageCurrentOverhead, !FPlatformMemory::TracksLLMAllocations());
 
 	// Compare memory the platform thinks we have allocated to what we have tracked, including the program memory
 	FPlatformMemoryStats PlatformStats = FPlatformMemory::GetStats();
@@ -980,7 +980,7 @@ void FLowLevelMemTracker::TickInternal()
 	int64 TrackedTotal = DefaultTracker.GetTrackedTotal();
 	// remove the MemoryUsageCurrentOverhead from the "Total" for the default LLM as it's not something anyone needs
 	// to investigate when finding what to reduce the platform LLM will have the info 
-	int64 DefaultProcessMemory = PlatformProcessMemory - MemoryUsageCurrentOverhead;
+	int64 DefaultProcessMemory = PlatformProcessMemory - (!FPlatformMemory::TracksLLMAllocations() ? MemoryUsageCurrentOverhead : 0);
 	int64 DefaultUntracked = FMath::Max<int64>(0, DefaultProcessMemory - TrackedTotal);
 	DefaultTracker.SetTagAmountInUpdate(FindOrAddTagData(ELLMTag::Total), DefaultProcessMemory, false);
 	DefaultTracker.SetTagAmountInUpdate(FindOrAddTagData(ELLMTag::Untracked), DefaultUntracked, false);
diff --git a/Engine/Source/Runtime/Core/Public/GenericPlatform/GenericPlatformMemory.h b/Engine/Source/Runtime/Core/Public/GenericPlatform/GenericPlatformMemory.h
--- a/Engine/Source/Runtime/Core/Public/GenericPlatform/GenericPlatformMemory.h
+++ b/Engine/Source/Runtime/Core/Public/GenericPlatform/GenericPlatformMemory.h
@@ -774,6 +774,15 @@ public:
 	*/
 	static void UpdateCustomLLMTags() { };
 
+	/**
+	 * Indicates whether LLM allocations are already accounted for with tracking and in GetStats.
+	 * Returns true if LLM allocations come from a memory pool separate from the main engine's memory or
+	 * are already tagged with tracking by the platform memory system.
+	 * Returns false if LLM uses the regular memory shared with the engine and allocations are not tracked.
+	 * @see GetLLMAllocFunctions
+	 */
+	static bool TracksLLMAllocations() { return false; }
+
 	/**
 	* Returns true if Protecting the parent processes pages has been enabled
 	* Only supported on platforms that support forking
