UnrealMathTest small fixes

diff --git a/Engine/Source/Runtime/Core/Tests/Math/UnrealMathTest.cpp b/Engine/Source/Runtime/Core/Tests/Math/UnrealMathTest.cpp
--- a/Engine/Source/Runtime/Core/Tests/Math/UnrealMathTest.cpp
+++ b/Engine/Source/Runtime/Core/Tests/Math/UnrealMathTest.cpp
@@ -3726,12 +3726,16 @@ TEST_CASE_NAMED(FInterpolationFunctionTests, "System::Core::Math::Interpolation
 			int32 K = 0;
 			for (int32 J = 1; J < Functions.Num(); ++J)
 			{
-				for (int32 L = 0; L < Values.Num(); ++L)
+				if (!FMath::IsNearlyEqual(Values[K], Values[J], 0.0001f))
 				{
-					INFO(FString::Printf(TEXT("%s: %f"), *(Functions[L].Value), Values[L]));
+					for (int32 L = 0; L < Values.Num(); ++L)
+					{
+						INFO(FString::Printf(TEXT("%s: %f"), *(Functions[L].Value), Values[L]));
+					}
+					REQUIRE_MESSAGE(FString::Printf(TEXT("Easing Function tests failed at index %d!"), I), false);
 				}
-				REQUIRE_MESSAGE(FString::Printf(TEXT("Easing Function tests failed at index %d!"), I), FMath::IsNearlyEqual(Values[K], Values[J], 0.0001f));
 			}
+			
 		}
 	};
 
@@ -4067,16 +4071,14 @@ TEST_CASE_NAMED(FMathRoundHalfFromZeroTests, "System::Core::Math::Round HalfFrom
 
 class FIsNearlyEqualByULPTestClass {
 public:
-	bool CheckMessage(const FString& What, bool Value)
+	void CheckMessage(const FString& What, bool Value)
 	{
 		CHECK_MESSAGE(*What, Value);
-		return Value;
 	}
 
-	bool CheckFalseMessage(const FString& What, bool Value)
+	void CheckFalseMessage(const FString& What, bool Value)
 	{
 		CHECK_FALSE_MESSAGE(*What, Value);
-		return !Value;
 	}
 
 	bool IsNearlyEqualByULPTest() {
@@ -4157,8 +4159,8 @@ public:
 
 		};
 
-		bool(FIsNearlyEqualByULPTestClass::*FuncTrue)(const FString&, bool) = &FIsNearlyEqualByULPTestClass::CheckMessage;
-		bool(FIsNearlyEqualByULPTestClass::*FuncFalse)(const FString&, bool) = &FIsNearlyEqualByULPTestClass::CheckFalseMessage;
+		void(FIsNearlyEqualByULPTestClass::*FuncTrue)(const FString&, bool) = &FIsNearlyEqualByULPTestClass::CheckMessage;
+		void(FIsNearlyEqualByULPTestClass::*FuncFalse)(const FString&, bool) = &FIsNearlyEqualByULPTestClass::CheckFalseMessage;
 
 		for (const TestItem& Item : TestItems)
 		{
@@ -4695,7 +4697,7 @@ namespace UnrealMathTestInternal {
 
 }
 
-TEST_CASE_NAMED(FColorConversionHeavyTest, "System::Core::Math::ColorConversionHeavy", "[EAutomationTestFlags][EngineFilter]")
+TEST_CASE_NAMED(FColorConversionHeavyTest, "System::Core::Math::ColorConversionHeavy", "[ApplicationContextMask][EngineFilter]")
 {
 	// WARNING: This test runs for a good while (at time of writing, ~90s when using
 	// a single core, proportionately less on many-core machines) with no user
