Some missed clamps from a code review to protect against divide by zero.

diff --git a/Engine/Source/Runtime/Core/Private/IO/IoStore.cpp b/Engine/Source/Runtime/Core/Private/IO/IoStore.cpp
--- a/Engine/Source/Runtime/Core/Private/IO/IoStore.cpp
+++ b/Engine/Source/Runtime/Core/Private/IO/IoStore.cpp
@@ -1934,7 +1934,7 @@ void FIoStoreWriterContextImpl::Flush()
 		WritesEnd - WritesStart,
 		WritesSeconds,
 		*FText::AsNumber(WriteByteCount.Load()).ToString(),
-		*FText::AsNumber((int64)((double)WriteByteCount.Load() / WritesSeconds)).ToString()
+		*FText::AsNumber((int64)((double)WriteByteCount.Load() / FMath::Max(.0001f, WritesSeconds))).ToString()
 		);
 
 	AllEntries.Empty();
@@ -1960,7 +1960,7 @@ void FIoStoreWriterContextImpl::Flush()
 		FinalizeEnd - FinalizeStart, 
 		IoStoreWriters.Num(), 
 		*FText::AsNumber(TotalTocSize).ToString(), 
-		*FText::AsNumber((int64)((double)TotalTocSize / (FinalizeEnd - FinalizeStart))).ToString()
+		*FText::AsNumber((int64)((double)TotalTocSize / FMath::Max(.0001f, FinalizeEnd - FinalizeStart))).ToString()
 		);
 }
 
