Removed TWeakPtr parameter from delegate macros.

diff --git a/Engine/Source/Runtime/Core/Public/Delegates/Delegate.h b/Engine/Source/Runtime/Core/Public/Delegates/Delegate.h
--- a/Engine/Source/Runtime/Core/Public/Delegates/Delegate.h
+++ b/Engine/Source/Runtime/Core/Public/Delegates/Delegate.h
@@ -232,8 +232,8 @@
 	class EventName : public BaseTypeEvent { friend class OwningType; };
 
 /** Declare user's dynamic delegate, with wrapper proxy method for executing the delegate */
-#define FUNC_DECLARE_DYNAMIC_DELEGATE( TWeakPtr, DynamicDelegateClassName, ExecFunction, FuncParamList, FuncParamPassThru, ... ) \
-	class DynamicDelegateClassName : public TBaseDynamicDelegate<TWeakPtr, __VA_ARGS__> \
+#define FUNC_DECLARE_DYNAMIC_DELEGATE( DynamicDelegateClassName, ExecFunction, FuncParamList, FuncParamPassThru, ... ) \
+	class DynamicDelegateClassName : public TBaseDynamicDelegate<FWeakObjectPtr, __VA_ARGS__> \
 	{ \
 	public: \
 		/** Default constructor */ \
@@ -243,7 +243,7 @@
 		\
 		/** Construction from an FScriptDelegate must be explicit.  This is really only used by UObject system internals. */ \
 		explicit DynamicDelegateClassName( const TScriptDelegate<>& InScriptDelegate ) \
-			: TBaseDynamicDelegate<TWeakPtr, __VA_ARGS__>( InScriptDelegate ) \
+			: TBaseDynamicDelegate<FWeakObjectPtr, __VA_ARGS__>( InScriptDelegate ) \
 		{ \
 		} \
 		\
@@ -267,8 +267,8 @@
 	};
 
 /** Declare user's dynamic delegate with return value, with wrapper proxy method for executing the delegate */
-#define FUNC_DECLARE_DYNAMIC_DELEGATE_RETVAL(TWeakPtr, DynamicDelegateRetValClassName, ExecFunction, RetValType, FuncParamList, FuncParamPassThru, ...) \
-	class DynamicDelegateRetValClassName : public TBaseDynamicDelegate<TWeakPtr, __VA_ARGS__> \
+#define FUNC_DECLARE_DYNAMIC_DELEGATE_RETVAL(DynamicDelegateRetValClassName, ExecFunction, RetValType, FuncParamList, FuncParamPassThru, ...) \
+	class DynamicDelegateRetValClassName : public TBaseDynamicDelegate<FWeakObjectPtr, __VA_ARGS__> \
 	{ \
 	public: \
 		/** Default constructor */ \
@@ -278,7 +278,7 @@
 		\
 		/** Construction from an FScriptDelegate must be explicit.  This is really only used by UObject system internals. */ \
 		explicit DynamicDelegateRetValClassName( const TScriptDelegate<>& InScriptDelegate ) \
-			: TBaseDynamicDelegate<TWeakPtr, __VA_ARGS__>( InScriptDelegate ) \
+			: TBaseDynamicDelegate<FWeakObjectPtr, __VA_ARGS__>( InScriptDelegate ) \
 		{ \
 		} \
 		\
@@ -293,8 +293,8 @@
 
 
 /** Declare user's dynamic multi-cast delegate, with wrapper proxy method for executing the delegate */
-#define FUNC_DECLARE_DYNAMIC_MULTICAST_DELEGATE(TWeakPtr, DynamicMulticastDelegateClassName, ExecFunction, FuncParamList, FuncParamPassThru, ...) \
-class DynamicMulticastDelegateClassName : public TBaseDynamicMulticastDelegate<TWeakPtr, __VA_ARGS__> \
+#define FUNC_DECLARE_DYNAMIC_MULTICAST_DELEGATE(DynamicMulticastDelegateClassName, ExecFunction, FuncParamList, FuncParamPassThru, ...) \
+class DynamicMulticastDelegateClassName : public TBaseDynamicMulticastDelegate<FWeakObjectPtr, __VA_ARGS__> \
 	{ \
 	public: \
 		/** Default constructor */ \
@@ -304,7 +304,7 @@ class DynamicMulticastDelegateClassName : public TBaseDynamicMulticastDelegate<T
 		\
 		/** Construction from an FMulticastScriptDelegate must be explicit.  This is really only used by UObject system internals. */ \
 		explicit DynamicMulticastDelegateClassName( const TMulticastScriptDelegate<>& InMulticastScriptDelegate ) \
-			: TBaseDynamicMulticastDelegate<TWeakPtr, __VA_ARGS__>( InMulticastScriptDelegate ) \
+			: TBaseDynamicMulticastDelegate<FWeakObjectPtr, __VA_ARGS__>( InMulticastScriptDelegate ) \
 		{ \
 		} \
 		\
diff --git a/Engine/Source/Runtime/Core/Public/Delegates/DelegateCombinations.h b/Engine/Source/Runtime/Core/Public/Delegates/DelegateCombinations.h
--- a/Engine/Source/Runtime/Core/Public/Delegates/DelegateCombinations.h
+++ b/Engine/Source/Runtime/Core/Public/Delegates/DelegateCombinations.h
@@ -32,16 +32,16 @@
 #define DECLARE_EVENT( OwningType, EventName ) FUNC_DECLARE_EVENT( OwningType, EventName, void )
 
 /** Declares a blueprint-accessible delegate that can only bind to one UFUNCTION at a time */
-#define DECLARE_DYNAMIC_DELEGATE( DelegateName ) BODY_MACRO_COMBINE(CURRENT_FILE_ID,_,__LINE__,_DELEGATE) FUNC_DECLARE_DYNAMIC_DELEGATE( FWeakObjectPtr, DelegateName, DelegateName##_DelegateWrapper, , FUNC_CONCAT( *this ), void )
+#define DECLARE_DYNAMIC_DELEGATE( DelegateName ) BODY_MACRO_COMBINE(CURRENT_FILE_ID,_,__LINE__,_DELEGATE) FUNC_DECLARE_DYNAMIC_DELEGATE( DelegateName, DelegateName##_DelegateWrapper, , FUNC_CONCAT( *this ), void )
 
 /** Declares a blueprint-accessible broadcast delegate that can bind to multiple native UFUNCTIONs simultaneously */
-#define DECLARE_DYNAMIC_MULTICAST_DELEGATE( DelegateName ) BODY_MACRO_COMBINE(CURRENT_FILE_ID,_,__LINE__,_DELEGATE) FUNC_DECLARE_DYNAMIC_MULTICAST_DELEGATE( FWeakObjectPtr, DelegateName, DelegateName##_DelegateWrapper, , FUNC_CONCAT( *this ), void )
+#define DECLARE_DYNAMIC_MULTICAST_DELEGATE( DelegateName ) BODY_MACRO_COMBINE(CURRENT_FILE_ID,_,__LINE__,_DELEGATE) FUNC_DECLARE_DYNAMIC_MULTICAST_DELEGATE( DelegateName, DelegateName##_DelegateWrapper, , FUNC_CONCAT( *this ), void )
 
 /** Declares a delegate with return value that can only bind to one native function at a time */
 #define DECLARE_DELEGATE_RetVal( ReturnValueType, DelegateName ) FUNC_DECLARE_DELEGATE( DelegateName, ReturnValueType )
 
 /** Declares a blueprint-accessible delegate with return value that can only bind to one UNFUNCTION at a time */
-#define DECLARE_DYNAMIC_DELEGATE_RetVal( ReturnValueType, DelegateName ) BODY_MACRO_COMBINE(CURRENT_FILE_ID,_,__LINE__,_DELEGATE) FUNC_DECLARE_DYNAMIC_DELEGATE_RETVAL( FWeakObjectPtr, DelegateName, DelegateName##_DelegateWrapper, ReturnValueType, , FUNC_CONCAT( *this ), ReturnValueType )
+#define DECLARE_DYNAMIC_DELEGATE_RetVal( ReturnValueType, DelegateName ) BODY_MACRO_COMBINE(CURRENT_FILE_ID,_,__LINE__,_DELEGATE) FUNC_DECLARE_DYNAMIC_DELEGATE_RETVAL( DelegateName, DelegateName##_DelegateWrapper, ReturnValueType, , FUNC_CONCAT( *this ), ReturnValueType )
 
 
 // Multiple-parameter versions of above delegate types:
@@ -49,79 +49,79 @@
 #define DECLARE_MULTICAST_DELEGATE_OneParam( DelegateName, Param1Type ) FUNC_DECLARE_MULTICAST_DELEGATE( DelegateName, void, Param1Type )
 #define DECLARE_TS_MULTICAST_DELEGATE_OneParam( DelegateName, Param1Type ) FUNC_DECLARE_TS_MULTICAST_DELEGATE( DelegateName, void, Param1Type )
 #define DECLARE_EVENT_OneParam( OwningType, EventName, Param1Type ) FUNC_DECLARE_EVENT( OwningType, EventName, void, Param1Type )
-#define DECLARE_DYNAMIC_DELEGATE_OneParam( DelegateName, Param1Type, Param1Name ) BODY_MACRO_COMBINE(CURRENT_FILE_ID,_,__LINE__,_DELEGATE) FUNC_DECLARE_DYNAMIC_DELEGATE( FWeakObjectPtr, DelegateName, DelegateName##_DelegateWrapper, FUNC_CONCAT( Param1Type InParam1 ), FUNC_CONCAT( *this, InParam1 ), void, Param1Type )
-#define DECLARE_DYNAMIC_MULTICAST_DELEGATE_OneParam( DelegateName, Param1Type, Param1Name ) BODY_MACRO_COMBINE(CURRENT_FILE_ID,_,__LINE__,_DELEGATE) FUNC_DECLARE_DYNAMIC_MULTICAST_DELEGATE( FWeakObjectPtr, DelegateName, DelegateName##_DelegateWrapper, FUNC_CONCAT( Param1Type InParam1 ), FUNC_CONCAT( *this, InParam1 ), void, Param1Type )
+#define DECLARE_DYNAMIC_DELEGATE_OneParam( DelegateName, Param1Type, Param1Name ) BODY_MACRO_COMBINE(CURRENT_FILE_ID,_,__LINE__,_DELEGATE) FUNC_DECLARE_DYNAMIC_DELEGATE( DelegateName, DelegateName##_DelegateWrapper, FUNC_CONCAT( Param1Type InParam1 ), FUNC_CONCAT( *this, InParam1 ), void, Param1Type )
+#define DECLARE_DYNAMIC_MULTICAST_DELEGATE_OneParam( DelegateName, Param1Type, Param1Name ) BODY_MACRO_COMBINE(CURRENT_FILE_ID,_,__LINE__,_DELEGATE) FUNC_DECLARE_DYNAMIC_MULTICAST_DELEGATE( DelegateName, DelegateName##_DelegateWrapper, FUNC_CONCAT( Param1Type InParam1 ), FUNC_CONCAT( *this, InParam1 ), void, Param1Type )
 #define DECLARE_DELEGATE_RetVal_OneParam( ReturnValueType, DelegateName, Param1Type ) FUNC_DECLARE_DELEGATE( DelegateName, ReturnValueType, Param1Type )
-#define DECLARE_DYNAMIC_DELEGATE_RetVal_OneParam( ReturnValueType, DelegateName, Param1Type, Param1Name ) BODY_MACRO_COMBINE(CURRENT_FILE_ID,_,__LINE__,_DELEGATE) FUNC_DECLARE_DYNAMIC_DELEGATE_RETVAL( FWeakObjectPtr, DelegateName, DelegateName##_DelegateWrapper, ReturnValueType, FUNC_CONCAT( Param1Type InParam1 ), FUNC_CONCAT( *this, InParam1 ), ReturnValueType, Param1Type )
+#define DECLARE_DYNAMIC_DELEGATE_RetVal_OneParam( ReturnValueType, DelegateName, Param1Type, Param1Name ) BODY_MACRO_COMBINE(CURRENT_FILE_ID,_,__LINE__,_DELEGATE) FUNC_DECLARE_DYNAMIC_DELEGATE_RETVAL( DelegateName, DelegateName##_DelegateWrapper, ReturnValueType, FUNC_CONCAT( Param1Type InParam1 ), FUNC_CONCAT( *this, InParam1 ), ReturnValueType, Param1Type )
 
 #define DECLARE_DELEGATE_TwoParams( DelegateName, Param1Type, Param2Type ) FUNC_DECLARE_DELEGATE( DelegateName, void, Param1Type, Param2Type )
 #define DECLARE_MULTICAST_DELEGATE_TwoParams( DelegateName, Param1Type, Param2Type ) FUNC_DECLARE_MULTICAST_DELEGATE( DelegateName, void, Param1Type, Param2Type )
 #define DECLARE_TS_MULTICAST_DELEGATE_TwoParams( DelegateName, Param1Type, Param2Type ) FUNC_DECLARE_TS_MULTICAST_DELEGATE( DelegateName, void, Param1Type, Param2Type )
 #define DECLARE_EVENT_TwoParams( OwningType, EventName, Param1Type, Param2Type ) FUNC_DECLARE_EVENT( OwningType, EventName, void, Param1Type, Param2Type )
-#define DECLARE_DYNAMIC_DELEGATE_TwoParams( DelegateName, Param1Type, Param1Name, Param2Type, Param2Name ) BODY_MACRO_COMBINE(CURRENT_FILE_ID,_,__LINE__,_DELEGATE) FUNC_DECLARE_DYNAMIC_DELEGATE( FWeakObjectPtr, DelegateName, DelegateName##_DelegateWrapper, FUNC_CONCAT( Param1Type InParam1, Param2Type InParam2 ), FUNC_CONCAT( *this, InParam1, InParam2 ), void, Param1Type, Param2Type )
-#define DECLARE_DYNAMIC_MULTICAST_DELEGATE_TwoParams( DelegateName, Param1Type, Param1Name, Param2Type, Param2Name ) BODY_MACRO_COMBINE(CURRENT_FILE_ID,_,__LINE__,_DELEGATE) FUNC_DECLARE_DYNAMIC_MULTICAST_DELEGATE( FWeakObjectPtr, DelegateName, DelegateName##_DelegateWrapper, FUNC_CONCAT( Param1Type InParam1, Param2Type InParam2 ), FUNC_CONCAT( *this, InParam1, InParam2 ), void, Param1Type, Param2Type )
+#define DECLARE_DYNAMIC_DELEGATE_TwoParams( DelegateName, Param1Type, Param1Name, Param2Type, Param2Name ) BODY_MACRO_COMBINE(CURRENT_FILE_ID,_,__LINE__,_DELEGATE) FUNC_DECLARE_DYNAMIC_DELEGATE( DelegateName, DelegateName##_DelegateWrapper, FUNC_CONCAT( Param1Type InParam1, Param2Type InParam2 ), FUNC_CONCAT( *this, InParam1, InParam2 ), void, Param1Type, Param2Type )
+#define DECLARE_DYNAMIC_MULTICAST_DELEGATE_TwoParams( DelegateName, Param1Type, Param1Name, Param2Type, Param2Name ) BODY_MACRO_COMBINE(CURRENT_FILE_ID,_,__LINE__,_DELEGATE) FUNC_DECLARE_DYNAMIC_MULTICAST_DELEGATE( DelegateName, DelegateName##_DelegateWrapper, FUNC_CONCAT( Param1Type InParam1, Param2Type InParam2 ), FUNC_CONCAT( *this, InParam1, InParam2 ), void, Param1Type, Param2Type )
 #define DECLARE_DELEGATE_RetVal_TwoParams( ReturnValueType, DelegateName, Param1Type, Param2Type ) FUNC_DECLARE_DELEGATE( DelegateName, ReturnValueType, Param1Type, Param2Type )
-#define DECLARE_DYNAMIC_DELEGATE_RetVal_TwoParams( ReturnValueType, DelegateName, Param1Type, Param1Name, Param2Type, Param2Name ) BODY_MACRO_COMBINE(CURRENT_FILE_ID,_,__LINE__,_DELEGATE) FUNC_DECLARE_DYNAMIC_DELEGATE_RETVAL( FWeakObjectPtr, DelegateName, DelegateName##_DelegateWrapper, ReturnValueType, FUNC_CONCAT( Param1Type InParam1, Param2Type InParam2 ), FUNC_CONCAT( *this, InParam1, InParam2 ), ReturnValueType, Param1Type, Param2Type )
+#define DECLARE_DYNAMIC_DELEGATE_RetVal_TwoParams( ReturnValueType, DelegateName, Param1Type, Param1Name, Param2Type, Param2Name ) BODY_MACRO_COMBINE(CURRENT_FILE_ID,_,__LINE__,_DELEGATE) FUNC_DECLARE_DYNAMIC_DELEGATE_RETVAL( DelegateName, DelegateName##_DelegateWrapper, ReturnValueType, FUNC_CONCAT( Param1Type InParam1, Param2Type InParam2 ), FUNC_CONCAT( *this, InParam1, InParam2 ), ReturnValueType, Param1Type, Param2Type )
 
 #define DECLARE_DELEGATE_ThreeParams( DelegateName, Param1Type, Param2Type, Param3Type ) FUNC_DECLARE_DELEGATE( DelegateName, void, Param1Type, Param2Type, Param3Type )
 #define DECLARE_MULTICAST_DELEGATE_ThreeParams( DelegateName, Param1Type, Param2Type, Param3Type ) FUNC_DECLARE_MULTICAST_DELEGATE( DelegateName, void, Param1Type, Param2Type, Param3Type )
 #define DECLARE_TS_MULTICAST_DELEGATE_ThreeParams( DelegateName, Param1Type, Param2Type, Param3Type ) FUNC_DECLARE_TS_MULTICAST_DELEGATE( DelegateName, void, Param1Type, Param2Type, Param3Type )
 #define DECLARE_EVENT_ThreeParams( OwningType, EventName, Param1Type, Param2Type, Param3Type ) FUNC_DECLARE_EVENT( OwningType, EventName, void, Param1Type, Param2Type, Param3Type )
-#define DECLARE_DYNAMIC_DELEGATE_ThreeParams( DelegateName, Param1Type, Param1Name, Param2Type, Param2Name, Param3Type, Param3Name ) BODY_MACRO_COMBINE(CURRENT_FILE_ID,_,__LINE__,_DELEGATE) FUNC_DECLARE_DYNAMIC_DELEGATE( FWeakObjectPtr, DelegateName, DelegateName##_DelegateWrapper, FUNC_CONCAT( Param1Type InParam1, Param2Type InParam2, Param3Type InParam3 ), FUNC_CONCAT( *this, InParam1, InParam2, InParam3 ), void, Param1Type, Param2Type, Param3Type )
-#define DECLARE_DYNAMIC_MULTICAST_DELEGATE_ThreeParams( DelegateName, Param1Type, Param1Name, Param2Type, Param2Name, Param3Type, Param3Name ) BODY_MACRO_COMBINE(CURRENT_FILE_ID,_,__LINE__,_DELEGATE) FUNC_DECLARE_DYNAMIC_MULTICAST_DELEGATE( FWeakObjectPtr, DelegateName, DelegateName##_DelegateWrapper, FUNC_CONCAT( Param1Type InParam1, Param2Type InParam2, Param3Type InParam3 ), FUNC_CONCAT( *this, InParam1, InParam2, InParam3 ), void, Param1Type, Param2Type, Param3Type )
+#define DECLARE_DYNAMIC_DELEGATE_ThreeParams( DelegateName, Param1Type, Param1Name, Param2Type, Param2Name, Param3Type, Param3Name ) BODY_MACRO_COMBINE(CURRENT_FILE_ID,_,__LINE__,_DELEGATE) FUNC_DECLARE_DYNAMIC_DELEGATE( DelegateName, DelegateName##_DelegateWrapper, FUNC_CONCAT( Param1Type InParam1, Param2Type InParam2, Param3Type InParam3 ), FUNC_CONCAT( *this, InParam1, InParam2, InParam3 ), void, Param1Type, Param2Type, Param3Type )
+#define DECLARE_DYNAMIC_MULTICAST_DELEGATE_ThreeParams( DelegateName, Param1Type, Param1Name, Param2Type, Param2Name, Param3Type, Param3Name ) BODY_MACRO_COMBINE(CURRENT_FILE_ID,_,__LINE__,_DELEGATE) FUNC_DECLARE_DYNAMIC_MULTICAST_DELEGATE( DelegateName, DelegateName##_DelegateWrapper, FUNC_CONCAT( Param1Type InParam1, Param2Type InParam2, Param3Type InParam3 ), FUNC_CONCAT( *this, InParam1, InParam2, InParam3 ), void, Param1Type, Param2Type, Param3Type )
 #define DECLARE_DELEGATE_RetVal_ThreeParams( ReturnValueType, DelegateName, Param1Type, Param2Type, Param3Type ) FUNC_DECLARE_DELEGATE( DelegateName, ReturnValueType, Param1Type, Param2Type, Param3Type )
-#define DECLARE_DYNAMIC_DELEGATE_RetVal_ThreeParams( ReturnValueType, DelegateName, Param1Type, Param1Name, Param2Type, Param2Name, Param3Type, Param3Name ) BODY_MACRO_COMBINE(CURRENT_FILE_ID,_,__LINE__,_DELEGATE) FUNC_DECLARE_DYNAMIC_DELEGATE_RETVAL( FWeakObjectPtr, DelegateName, DelegateName##_DelegateWrapper, ReturnValueType, FUNC_CONCAT( Param1Type InParam1, Param2Type InParam2, Param3Type InParam3 ), FUNC_CONCAT( *this, InParam1, InParam2, InParam3 ), ReturnValueType, Param1Type, Param2Type, Param3Type )
+#define DECLARE_DYNAMIC_DELEGATE_RetVal_ThreeParams( ReturnValueType, DelegateName, Param1Type, Param1Name, Param2Type, Param2Name, Param3Type, Param3Name ) BODY_MACRO_COMBINE(CURRENT_FILE_ID,_,__LINE__,_DELEGATE) FUNC_DECLARE_DYNAMIC_DELEGATE_RETVAL( DelegateName, DelegateName##_DelegateWrapper, ReturnValueType, FUNC_CONCAT( Param1Type InParam1, Param2Type InParam2, Param3Type InParam3 ), FUNC_CONCAT( *this, InParam1, InParam2, InParam3 ), ReturnValueType, Param1Type, Param2Type, Param3Type )
 
 #define DECLARE_DELEGATE_FourParams( DelegateName, Param1Type, Param2Type, Param3Type, Param4Type ) FUNC_DECLARE_DELEGATE( DelegateName, void, Param1Type, Param2Type, Param3Type, Param4Type )
 #define DECLARE_MULTICAST_DELEGATE_FourParams( DelegateName, Param1Type, Param2Type, Param3Type, Param4Type ) FUNC_DECLARE_MULTICAST_DELEGATE( DelegateName, void, Param1Type, Param2Type, Param3Type, Param4Type )
 #define DECLARE_TS_MULTICAST_DELEGATE_FourParams( DelegateName, Param1Type, Param2Type, Param3Type, Param4Type ) FUNC_DECLARE_TS_MULTICAST_DELEGATE( DelegateName, void, Param1Type, Param2Type, Param3Type, Param4Type )
 #define DECLARE_EVENT_FourParams( OwningType, EventName, Param1Type, Param2Type, Param3Type, Param4Type ) FUNC_DECLARE_EVENT( OwningType, EventName, void, Param1Type, Param2Type, Param3Type, Param4Type )
-#define DECLARE_DYNAMIC_DELEGATE_FourParams( DelegateName, Param1Type, Param1Name, Param2Type, Param2Name, Param3Type, Param3Name, Param4Type, Param4Name ) BODY_MACRO_COMBINE(CURRENT_FILE_ID,_,__LINE__,_DELEGATE) FUNC_DECLARE_DYNAMIC_DELEGATE( FWeakObjectPtr, DelegateName, DelegateName##_DelegateWrapper, FUNC_CONCAT( Param1Type InParam1, Param2Type InParam2, Param3Type InParam3, Param4Type InParam4 ), FUNC_CONCAT( *this, InParam1, InParam2, InParam3, InParam4 ), void, Param1Type, Param2Type, Param3Type, Param4Type )
-#define DECLARE_DYNAMIC_MULTICAST_DELEGATE_FourParams( DelegateName, Param1Type, Param1Name, Param2Type, Param2Name, Param3Type, Param3Name, Param4Type, Param4Name ) BODY_MACRO_COMBINE(CURRENT_FILE_ID,_,__LINE__,_DELEGATE) FUNC_DECLARE_DYNAMIC_MULTICAST_DELEGATE( FWeakObjectPtr, DelegateName, DelegateName##_DelegateWrapper, FUNC_CONCAT( Param1Type InParam1, Param2Type InParam2, Param3Type InParam3, Param4Type InParam4 ), FUNC_CONCAT( *this, InParam1, InParam2, InParam3, InParam4 ), void, Param1Type, Param2Type, Param3Type, Param4Type )
+#define DECLARE_DYNAMIC_DELEGATE_FourParams( DelegateName, Param1Type, Param1Name, Param2Type, Param2Name, Param3Type, Param3Name, Param4Type, Param4Name ) BODY_MACRO_COMBINE(CURRENT_FILE_ID,_,__LINE__,_DELEGATE) FUNC_DECLARE_DYNAMIC_DELEGATE( DelegateName, DelegateName##_DelegateWrapper, FUNC_CONCAT( Param1Type InParam1, Param2Type InParam2, Param3Type InParam3, Param4Type InParam4 ), FUNC_CONCAT( *this, InParam1, InParam2, InParam3, InParam4 ), void, Param1Type, Param2Type, Param3Type, Param4Type )
+#define DECLARE_DYNAMIC_MULTICAST_DELEGATE_FourParams( DelegateName, Param1Type, Param1Name, Param2Type, Param2Name, Param3Type, Param3Name, Param4Type, Param4Name ) BODY_MACRO_COMBINE(CURRENT_FILE_ID,_,__LINE__,_DELEGATE) FUNC_DECLARE_DYNAMIC_MULTICAST_DELEGATE( DelegateName, DelegateName##_DelegateWrapper, FUNC_CONCAT( Param1Type InParam1, Param2Type InParam2, Param3Type InParam3, Param4Type InParam4 ), FUNC_CONCAT( *this, InParam1, InParam2, InParam3, InParam4 ), void, Param1Type, Param2Type, Param3Type, Param4Type )
 #define DECLARE_DELEGATE_RetVal_FourParams( ReturnValueType, DelegateName, Param1Type, Param2Type, Param3Type, Param4Type ) FUNC_DECLARE_DELEGATE( DelegateName, ReturnValueType, Param1Type, Param2Type, Param3Type, Param4Type )
-#define DECLARE_DYNAMIC_DELEGATE_RetVal_FourParams( ReturnValueType, DelegateName, Param1Type, Param1Name, Param2Type, Param2Name, Param3Type, Param3Name, Param4Type, Param4Name ) BODY_MACRO_COMBINE(CURRENT_FILE_ID,_,__LINE__,_DELEGATE) FUNC_DECLARE_DYNAMIC_DELEGATE_RETVAL( FWeakObjectPtr, DelegateName, DelegateName##_DelegateWrapper, ReturnValueType, FUNC_CONCAT( Param1Type InParam1, Param2Type InParam2, Param3Type InParam3, Param4Type InParam4 ), FUNC_CONCAT( *this, InParam1, InParam2, InParam3, InParam4 ), ReturnValueType, Param1Type, Param2Type, Param3Type, Param4Type )
+#define DECLARE_DYNAMIC_DELEGATE_RetVal_FourParams( ReturnValueType, DelegateName, Param1Type, Param1Name, Param2Type, Param2Name, Param3Type, Param3Name, Param4Type, Param4Name ) BODY_MACRO_COMBINE(CURRENT_FILE_ID,_,__LINE__,_DELEGATE) FUNC_DECLARE_DYNAMIC_DELEGATE_RETVAL( DelegateName, DelegateName##_DelegateWrapper, ReturnValueType, FUNC_CONCAT( Param1Type InParam1, Param2Type InParam2, Param3Type InParam3, Param4Type InParam4 ), FUNC_CONCAT( *this, InParam1, InParam2, InParam3, InParam4 ), ReturnValueType, Param1Type, Param2Type, Param3Type, Param4Type )
 
 #define DECLARE_DELEGATE_FiveParams( DelegateName, Param1Type, Param2Type, Param3Type, Param4Type, Param5Type ) FUNC_DECLARE_DELEGATE( DelegateName, void, Param1Type, Param2Type, Param3Type, Param4Type, Param5Type )
 #define DECLARE_MULTICAST_DELEGATE_FiveParams( DelegateName, Param1Type, Param2Type, Param3Type, Param4Type, Param5Type ) FUNC_DECLARE_MULTICAST_DELEGATE( DelegateName, void, Param1Type, Param2Type, Param3Type, Param4Type, Param5Type )
 #define DECLARE_TS_MULTICAST_DELEGATE_FiveParams( DelegateName, Param1Type, Param2Type, Param3Type, Param4Type, Param5Type ) FUNC_DECLARE_TS_MULTICAST_DELEGATE( DelegateName, void, Param1Type, Param2Type, Param3Type, Param4Type, Param5Type )
 #define DECLARE_EVENT_FiveParams( OwningType, EventName, Param1Type, Param2Type, Param3Type, Param4Type, Param5Type ) FUNC_DECLARE_EVENT( OwningType, EventName, void, Param1Type, Param2Type, Param3Type, Param4Type, Param5Type )
-#define DECLARE_DYNAMIC_DELEGATE_FiveParams( DelegateName, Param1Type, Param1Name, Param2Type, Param2Name, Param3Type, Param3Name, Param4Type, Param4Name, Param5Type, Param5Name ) BODY_MACRO_COMBINE(CURRENT_FILE_ID,_,__LINE__,_DELEGATE) FUNC_DECLARE_DYNAMIC_DELEGATE( FWeakObjectPtr, DelegateName, DelegateName##_DelegateWrapper, FUNC_CONCAT( Param1Type InParam1, Param2Type InParam2, Param3Type InParam3, Param4Type InParam4, Param5Type InParam5 ), FUNC_CONCAT( *this, InParam1, InParam2, InParam3, InParam4, InParam5 ), void, Param1Type, Param2Type, Param3Type, Param4Type, Param5Type )
-#define DECLARE_DYNAMIC_MULTICAST_DELEGATE_FiveParams( DelegateName, Param1Type, Param1Name, Param2Type, Param2Name, Param3Type, Param3Name, Param4Type, Param4Name, Param5Type, Param5Name ) BODY_MACRO_COMBINE(CURRENT_FILE_ID,_,__LINE__,_DELEGATE) FUNC_DECLARE_DYNAMIC_MULTICAST_DELEGATE( FWeakObjectPtr, DelegateName, DelegateName##_DelegateWrapper, FUNC_CONCAT( Param1Type InParam1, Param2Type InParam2, Param3Type InParam3, Param4Type InParam4, Param5Type InParam5 ), FUNC_CONCAT( *this, InParam1, InParam2, InParam3, InParam4, InParam5 ), void, Param1Type, Param2Type, Param3Type, Param4Type, Param5Type )
+#define DECLARE_DYNAMIC_DELEGATE_FiveParams( DelegateName, Param1Type, Param1Name, Param2Type, Param2Name, Param3Type, Param3Name, Param4Type, Param4Name, Param5Type, Param5Name ) BODY_MACRO_COMBINE(CURRENT_FILE_ID,_,__LINE__,_DELEGATE) FUNC_DECLARE_DYNAMIC_DELEGATE( DelegateName, DelegateName##_DelegateWrapper, FUNC_CONCAT( Param1Type InParam1, Param2Type InParam2, Param3Type InParam3, Param4Type InParam4, Param5Type InParam5 ), FUNC_CONCAT( *this, InParam1, InParam2, InParam3, InParam4, InParam5 ), void, Param1Type, Param2Type, Param3Type, Param4Type, Param5Type )
+#define DECLARE_DYNAMIC_MULTICAST_DELEGATE_FiveParams( DelegateName, Param1Type, Param1Name, Param2Type, Param2Name, Param3Type, Param3Name, Param4Type, Param4Name, Param5Type, Param5Name ) BODY_MACRO_COMBINE(CURRENT_FILE_ID,_,__LINE__,_DELEGATE) FUNC_DECLARE_DYNAMIC_MULTICAST_DELEGATE( DelegateName, DelegateName##_DelegateWrapper, FUNC_CONCAT( Param1Type InParam1, Param2Type InParam2, Param3Type InParam3, Param4Type InParam4, Param5Type InParam5 ), FUNC_CONCAT( *this, InParam1, InParam2, InParam3, InParam4, InParam5 ), void, Param1Type, Param2Type, Param3Type, Param4Type, Param5Type )
 #define DECLARE_DELEGATE_RetVal_FiveParams( ReturnValueType, DelegateName, Param1Type, Param2Type, Param3Type, Param4Type, Param5Type ) FUNC_DECLARE_DELEGATE( DelegateName, ReturnValueType, Param1Type, Param2Type, Param3Type, Param4Type, Param5Type )
-#define DECLARE_DYNAMIC_DELEGATE_RetVal_FiveParams( ReturnValueType, DelegateName, Param1Type, Param1Name, Param2Type, Param2Name, Param3Type, Param3Name, Param4Type, Param4Name, Param5Type, Param5Name ) BODY_MACRO_COMBINE(CURRENT_FILE_ID,_,__LINE__,_DELEGATE) FUNC_DECLARE_DYNAMIC_DELEGATE_RETVAL( FWeakObjectPtr, DelegateName, DelegateName##_DelegateWrapper, ReturnValueType, FUNC_CONCAT( Param1Type InParam1, Param2Type InParam2, Param3Type InParam3, Param4Type InParam4, Param5Type InParam5 ), FUNC_CONCAT( *this, InParam1, InParam2, InParam3, InParam4, InParam5 ), ReturnValueType, Param1Type, Param2Type, Param3Type, Param4Type, Param5Type )
+#define DECLARE_DYNAMIC_DELEGATE_RetVal_FiveParams( ReturnValueType, DelegateName, Param1Type, Param1Name, Param2Type, Param2Name, Param3Type, Param3Name, Param4Type, Param4Name, Param5Type, Param5Name ) BODY_MACRO_COMBINE(CURRENT_FILE_ID,_,__LINE__,_DELEGATE) FUNC_DECLARE_DYNAMIC_DELEGATE_RETVAL( DelegateName, DelegateName##_DelegateWrapper, ReturnValueType, FUNC_CONCAT( Param1Type InParam1, Param2Type InParam2, Param3Type InParam3, Param4Type InParam4, Param5Type InParam5 ), FUNC_CONCAT( *this, InParam1, InParam2, InParam3, InParam4, InParam5 ), ReturnValueType, Param1Type, Param2Type, Param3Type, Param4Type, Param5Type )
 
 #define DECLARE_DELEGATE_SixParams( DelegateName, Param1Type, Param2Type, Param3Type, Param4Type, Param5Type, Param6Type ) FUNC_DECLARE_DELEGATE( DelegateName, void, Param1Type, Param2Type, Param3Type, Param4Type, Param5Type, Param6Type )
 #define DECLARE_MULTICAST_DELEGATE_SixParams( DelegateName, Param1Type, Param2Type, Param3Type, Param4Type, Param5Type, Param6Type ) FUNC_DECLARE_MULTICAST_DELEGATE( DelegateName, void, Param1Type, Param2Type, Param3Type, Param4Type, Param5Type, Param6Type )
 #define DECLARE_TS_MULTICAST_DELEGATE_SixParams( DelegateName, Param1Type, Param2Type, Param3Type, Param4Type, Param5Type, Param6Type ) FUNC_DECLARE_TS_MULTICAST_DELEGATE( DelegateName, void, Param1Type, Param2Type, Param3Type, Param4Type, Param5Type, Param6Type )
 #define DECLARE_EVENT_SixParams( OwningType, EventName, Param1Type, Param2Type, Param3Type, Param4Type, Param5Type, Param6Type ) FUNC_DECLARE_EVENT( OwningType, EventName, void, Param1Type, Param2Type, Param3Type, Param4Type, Param5Type, Param6Type )
-#define DECLARE_DYNAMIC_DELEGATE_SixParams( DelegateName, Param1Type, Param1Name, Param2Type, Param2Name, Param3Type, Param3Name, Param4Type, Param4Name, Param5Type, Param5Name, Param6Type, Param6Name ) BODY_MACRO_COMBINE(CURRENT_FILE_ID,_,__LINE__,_DELEGATE) FUNC_DECLARE_DYNAMIC_DELEGATE( FWeakObjectPtr, DelegateName, DelegateName##_DelegateWrapper, FUNC_CONCAT( Param1Type InParam1, Param2Type InParam2, Param3Type InParam3, Param4Type InParam4, Param5Type InParam5, Param6Type InParam6 ), FUNC_CONCAT( *this, InParam1, InParam2, InParam3, InParam4, InParam5, InParam6 ), void, Param1Type, Param2Type, Param3Type, Param4Type, Param5Type, Param6Type )
-#define DECLARE_DYNAMIC_MULTICAST_DELEGATE_SixParams( DelegateName, Param1Type, Param1Name, Param2Type, Param2Name, Param3Type, Param3Name, Param4Type, Param4Name, Param5Type, Param5Name, Param6Type, Param6Name ) BODY_MACRO_COMBINE(CURRENT_FILE_ID,_,__LINE__,_DELEGATE) FUNC_DECLARE_DYNAMIC_MULTICAST_DELEGATE( FWeakObjectPtr, DelegateName, DelegateName##_DelegateWrapper, FUNC_CONCAT( Param1Type InParam1, Param2Type InParam2, Param3Type InParam3, Param4Type InParam4, Param5Type InParam5, Param6Type InParam6 ), FUNC_CONCAT( *this, InParam1, InParam2, InParam3, InParam4, InParam5, InParam6 ), void, Param1Type, Param2Type, Param3Type, Param4Type, Param5Type, Param6Type )
+#define DECLARE_DYNAMIC_DELEGATE_SixParams( DelegateName, Param1Type, Param1Name, Param2Type, Param2Name, Param3Type, Param3Name, Param4Type, Param4Name, Param5Type, Param5Name, Param6Type, Param6Name ) BODY_MACRO_COMBINE(CURRENT_FILE_ID,_,__LINE__,_DELEGATE) FUNC_DECLARE_DYNAMIC_DELEGATE( DelegateName, DelegateName##_DelegateWrapper, FUNC_CONCAT( Param1Type InParam1, Param2Type InParam2, Param3Type InParam3, Param4Type InParam4, Param5Type InParam5, Param6Type InParam6 ), FUNC_CONCAT( *this, InParam1, InParam2, InParam3, InParam4, InParam5, InParam6 ), void, Param1Type, Param2Type, Param3Type, Param4Type, Param5Type, Param6Type )
+#define DECLARE_DYNAMIC_MULTICAST_DELEGATE_SixParams( DelegateName, Param1Type, Param1Name, Param2Type, Param2Name, Param3Type, Param3Name, Param4Type, Param4Name, Param5Type, Param5Name, Param6Type, Param6Name ) BODY_MACRO_COMBINE(CURRENT_FILE_ID,_,__LINE__,_DELEGATE) FUNC_DECLARE_DYNAMIC_MULTICAST_DELEGATE( DelegateName, DelegateName##_DelegateWrapper, FUNC_CONCAT( Param1Type InParam1, Param2Type InParam2, Param3Type InParam3, Param4Type InParam4, Param5Type InParam5, Param6Type InParam6 ), FUNC_CONCAT( *this, InParam1, InParam2, InParam3, InParam4, InParam5, InParam6 ), void, Param1Type, Param2Type, Param3Type, Param4Type, Param5Type, Param6Type )
 #define DECLARE_DELEGATE_RetVal_SixParams( ReturnValueType, DelegateName, Param1Type, Param2Type, Param3Type, Param4Type, Param5Type, Param6Type ) FUNC_DECLARE_DELEGATE( DelegateName, ReturnValueType, Param1Type, Param2Type, Param3Type, Param4Type, Param5Type, Param6Type )
-#define DECLARE_DYNAMIC_DELEGATE_RetVal_SixParams( ReturnValueType, DelegateName, Param1Type, Param1Name, Param2Type, Param2Name, Param3Type, Param3Name, Param4Type, Param4Name, Param5Type, Param5Name, Param6Type, Param6Name ) BODY_MACRO_COMBINE(CURRENT_FILE_ID,_,__LINE__,_DELEGATE) FUNC_DECLARE_DYNAMIC_DELEGATE_RETVAL( FWeakObjectPtr, DelegateName, DelegateName##_DelegateWrapper, ReturnValueType, FUNC_CONCAT( Param1Type InParam1, Param2Type InParam2, Param3Type InParam3, Param4Type InParam4, Param5Type InParam5, Param6Type InParam6 ), FUNC_CONCAT( *this, InParam1, InParam2, InParam3, InParam4, InParam5, InParam6 ), ReturnValueType, Param1Type, Param2Type, Param3Type, Param4Type, Param5Type, Param6Type )
+#define DECLARE_DYNAMIC_DELEGATE_RetVal_SixParams( ReturnValueType, DelegateName, Param1Type, Param1Name, Param2Type, Param2Name, Param3Type, Param3Name, Param4Type, Param4Name, Param5Type, Param5Name, Param6Type, Param6Name ) BODY_MACRO_COMBINE(CURRENT_FILE_ID,_,__LINE__,_DELEGATE) FUNC_DECLARE_DYNAMIC_DELEGATE_RETVAL( DelegateName, DelegateName##_DelegateWrapper, ReturnValueType, FUNC_CONCAT( Param1Type InParam1, Param2Type InParam2, Param3Type InParam3, Param4Type InParam4, Param5Type InParam5, Param6Type InParam6 ), FUNC_CONCAT( *this, InParam1, InParam2, InParam3, InParam4, InParam5, InParam6 ), ReturnValueType, Param1Type, Param2Type, Param3Type, Param4Type, Param5Type, Param6Type )
 
 #define DECLARE_DELEGATE_SevenParams( DelegateName, Param1Type, Param2Type, Param3Type, Param4Type, Param5Type, Param6Type, Param7Type ) FUNC_DECLARE_DELEGATE( DelegateName, void, Param1Type, Param2Type, Param3Type, Param4Type, Param5Type, Param6Type, Param7Type )
 #define DECLARE_MULTICAST_DELEGATE_SevenParams( DelegateName, Param1Type, Param2Type, Param3Type, Param4Type, Param5Type, Param6Type, Param7Type ) FUNC_DECLARE_MULTICAST_DELEGATE( DelegateName, void, Param1Type, Param2Type, Param3Type, Param4Type, Param5Type, Param6Type, Param7Type )
 #define DECLARE_TS_MULTICAST_DELEGATE_SevenParams( DelegateName, Param1Type, Param2Type, Param3Type, Param4Type, Param5Type, Param6Type, Param7Type ) FUNC_DECLARE_TS_MULTICAST_DELEGATE( DelegateName, void, Param1Type, Param2Type, Param3Type, Param4Type, Param5Type, Param6Type, Param7Type )
 #define DECLARE_EVENT_SevenParams( OwningType, EventName, Param1Type, Param2Type, Param3Type, Param4Type, Param5Type, Param6Type, Param7Type ) FUNC_DECLARE_EVENT( OwningType, EventName, void, Param1Type, Param2Type, Param3Type, Param4Type, Param5Type, Param6Type, Param7Type )
-#define DECLARE_DYNAMIC_DELEGATE_SevenParams( DelegateName, Param1Type, Param1Name, Param2Type, Param2Name, Param3Type, Param3Name, Param4Type, Param4Name, Param5Type, Param5Name, Param6Type, Param6Name, Param7Type, Param7Name ) BODY_MACRO_COMBINE(CURRENT_FILE_ID,_,__LINE__,_DELEGATE) FUNC_DECLARE_DYNAMIC_DELEGATE( FWeakObjectPtr, DelegateName, DelegateName##_DelegateWrapper, FUNC_CONCAT( Param1Type InParam1, Param2Type InParam2, Param3Type InParam3, Param4Type InParam4, Param5Type InParam5, Param6Type InParam6, Param7Type InParam7 ), FUNC_CONCAT( *this, InParam1, InParam2, InParam3, InParam4, InParam5, InParam6, InParam7 ), void, Param1Type, Param2Type, Param3Type, Param4Type, Param5Type, Param6Type, Param7Type )
-#define DECLARE_DYNAMIC_MULTICAST_DELEGATE_SevenParams( DelegateName, Param1Type, Param1Name, Param2Type, Param2Name, Param3Type, Param3Name, Param4Type, Param4Name, Param5Type, Param5Name, Param6Type, Param6Name, Param7Type, Param7Name ) BODY_MACRO_COMBINE(CURRENT_FILE_ID,_,__LINE__,_DELEGATE) FUNC_DECLARE_DYNAMIC_MULTICAST_DELEGATE( FWeakObjectPtr, DelegateName, DelegateName##_DelegateWrapper, FUNC_CONCAT( Param1Type InParam1, Param2Type InParam2, Param3Type InParam3, Param4Type InParam4, Param5Type InParam5, Param6Type InParam6, Param7Type InParam7 ), FUNC_CONCAT( *this, InParam1, InParam2, InParam3, InParam4, InParam5, InParam6, InParam7 ), void, Param1Type, Param2Type, Param3Type, Param4Type, Param5Type, Param6Type, Param7Type )
+#define DECLARE_DYNAMIC_DELEGATE_SevenParams( DelegateName, Param1Type, Param1Name, Param2Type, Param2Name, Param3Type, Param3Name, Param4Type, Param4Name, Param5Type, Param5Name, Param6Type, Param6Name, Param7Type, Param7Name ) BODY_MACRO_COMBINE(CURRENT_FILE_ID,_,__LINE__,_DELEGATE) FUNC_DECLARE_DYNAMIC_DELEGATE( DelegateName, DelegateName##_DelegateWrapper, FUNC_CONCAT( Param1Type InParam1, Param2Type InParam2, Param3Type InParam3, Param4Type InParam4, Param5Type InParam5, Param6Type InParam6, Param7Type InParam7 ), FUNC_CONCAT( *this, InParam1, InParam2, InParam3, InParam4, InParam5, InParam6, InParam7 ), void, Param1Type, Param2Type, Param3Type, Param4Type, Param5Type, Param6Type, Param7Type )
+#define DECLARE_DYNAMIC_MULTICAST_DELEGATE_SevenParams( DelegateName, Param1Type, Param1Name, Param2Type, Param2Name, Param3Type, Param3Name, Param4Type, Param4Name, Param5Type, Param5Name, Param6Type, Param6Name, Param7Type, Param7Name ) BODY_MACRO_COMBINE(CURRENT_FILE_ID,_,__LINE__,_DELEGATE) FUNC_DECLARE_DYNAMIC_MULTICAST_DELEGATE( DelegateName, DelegateName##_DelegateWrapper, FUNC_CONCAT( Param1Type InParam1, Param2Type InParam2, Param3Type InParam3, Param4Type InParam4, Param5Type InParam5, Param6Type InParam6, Param7Type InParam7 ), FUNC_CONCAT( *this, InParam1, InParam2, InParam3, InParam4, InParam5, InParam6, InParam7 ), void, Param1Type, Param2Type, Param3Type, Param4Type, Param5Type, Param6Type, Param7Type )
 #define DECLARE_DELEGATE_RetVal_SevenParams( ReturnValueType, DelegateName, Param1Type, Param2Type, Param3Type, Param4Type, Param5Type, Param6Type, Param7Type ) FUNC_DECLARE_DELEGATE( DelegateName, ReturnValueType, Param1Type, Param2Type, Param3Type, Param4Type, Param5Type, Param6Type, Param7Type )
-#define DECLARE_DYNAMIC_DELEGATE_RetVal_SevenParams( ReturnValueType, DelegateName, Param1Type, Param1Name, Param2Type, Param2Name, Param3Type, Param3Name, Param4Type, Param4Name, Param5Type, Param5Name, Param6Type, Param6Name, Param7Type, Param7Name ) BODY_MACRO_COMBINE(CURRENT_FILE_ID,_,__LINE__,_DELEGATE) FUNC_DECLARE_DYNAMIC_DELEGATE_RETVAL( FWeakObjectPtr, DelegateName, DelegateName##_DelegateWrapper, ReturnValueType, FUNC_CONCAT( Param1Type InParam1, Param2Type InParam2, Param3Type InParam3, Param4Type InParam4, Param5Type InParam5, Param6Type InParam6, Param7Type InParam7 ), FUNC_CONCAT( *this, InParam1, InParam2, InParam3, InParam4, InParam5, InParam6, InParam7 ), ReturnValueType, Param1Type, Param2Type, Param3Type, Param4Type, Param5Type, Param6Type, Param7Type )
+#define DECLARE_DYNAMIC_DELEGATE_RetVal_SevenParams( ReturnValueType, DelegateName, Param1Type, Param1Name, Param2Type, Param2Name, Param3Type, Param3Name, Param4Type, Param4Name, Param5Type, Param5Name, Param6Type, Param6Name, Param7Type, Param7Name ) BODY_MACRO_COMBINE(CURRENT_FILE_ID,_,__LINE__,_DELEGATE) FUNC_DECLARE_DYNAMIC_DELEGATE_RETVAL( DelegateName, DelegateName##_DelegateWrapper, ReturnValueType, FUNC_CONCAT( Param1Type InParam1, Param2Type InParam2, Param3Type InParam3, Param4Type InParam4, Param5Type InParam5, Param6Type InParam6, Param7Type InParam7 ), FUNC_CONCAT( *this, InParam1, InParam2, InParam3, InParam4, InParam5, InParam6, InParam7 ), ReturnValueType, Param1Type, Param2Type, Param3Type, Param4Type, Param5Type, Param6Type, Param7Type )
 
 #define DECLARE_DELEGATE_EightParams( DelegateName, Param1Type, Param2Type, Param3Type, Param4Type, Param5Type, Param6Type, Param7Type, Param8Type ) FUNC_DECLARE_DELEGATE( DelegateName, void, Param1Type, Param2Type, Param3Type, Param4Type, Param5Type, Param6Type, Param7Type, Param8Type )
 #define DECLARE_MULTICAST_DELEGATE_EightParams( DelegateName, Param1Type, Param2Type, Param3Type, Param4Type, Param5Type, Param6Type, Param7Type, Param8Type ) FUNC_DECLARE_MULTICAST_DELEGATE( DelegateName, void, Param1Type, Param2Type, Param3Type, Param4Type, Param5Type, Param6Type, Param7Type, Param8Type )
 #define DECLARE_TS_MULTICAST_DELEGATE_EightParams( DelegateName, Param1Type, Param2Type, Param3Type, Param4Type, Param5Type, Param6Type, Param7Type, Param8Type ) FUNC_DECLARE_TS_MULTICAST_DELEGATE( DelegateName, void, Param1Type, Param2Type, Param3Type, Param4Type, Param5Type, Param6Type, Param7Type, Param8Type )
 #define DECLARE_EVENT_EightParams( OwningType, EventName, Param1Type, Param2Type, Param3Type, Param4Type, Param5Type, Param6Type, Param7Type, Param8Type ) FUNC_DECLARE_EVENT( OwningType, EventName, void, Param1Type, Param2Type, Param3Type, Param4Type, Param5Type, Param6Type, Param7Type, Param8Type )
-#define DECLARE_DYNAMIC_DELEGATE_EightParams( DelegateName, Param1Type, Param1Name, Param2Type, Param2Name, Param3Type, Param3Name, Param4Type, Param4Name, Param5Type, Param5Name, Param6Type, Param6Name, Param7Type, Param7Name, Param8Type, Param8Name ) BODY_MACRO_COMBINE(CURRENT_FILE_ID,_,__LINE__,_DELEGATE) FUNC_DECLARE_DYNAMIC_DELEGATE( FWeakObjectPtr, DelegateName, DelegateName##_DelegateWrapper, FUNC_CONCAT( Param1Type InParam1, Param2Type InParam2, Param3Type InParam3, Param4Type InParam4, Param5Type InParam5, Param6Type InParam6, Param7Type InParam7, Param8Type InParam8 ), FUNC_CONCAT( *this, InParam1, InParam2, InParam3, InParam4, InParam5, InParam6, InParam7, InParam8 ), void, Param1Type, Param2Type, Param3Type, Param4Type, Param5Type, Param6Type, Param7Type, Param8Type )
-#define DECLARE_DYNAMIC_MULTICAST_DELEGATE_EightParams( DelegateName, Param1Type, Param1Name, Param2Type, Param2Name, Param3Type, Param3Name, Param4Type, Param4Name, Param5Type, Param5Name, Param6Type, Param6Name, Param7Type, Param7Name, Param8Type, Param8Name ) BODY_MACRO_COMBINE(CURRENT_FILE_ID,_,__LINE__,_DELEGATE) FUNC_DECLARE_DYNAMIC_MULTICAST_DELEGATE( FWeakObjectPtr, DelegateName, DelegateName##_DelegateWrapper, FUNC_CONCAT( Param1Type InParam1, Param2Type InParam2, Param3Type InParam3, Param4Type InParam4, Param5Type InParam5, Param6Type InParam6, Param7Type InParam7, Param8Type InParam8 ), FUNC_CONCAT( *this, InParam1, InParam2, InParam3, InParam4, InParam5, InParam6, InParam7, InParam8 ), void, Param1Type, Param2Type, Param3Type, Param4Type, Param5Type, Param6Type, Param7Type, Param8Type )
+#define DECLARE_DYNAMIC_DELEGATE_EightParams( DelegateName, Param1Type, Param1Name, Param2Type, Param2Name, Param3Type, Param3Name, Param4Type, Param4Name, Param5Type, Param5Name, Param6Type, Param6Name, Param7Type, Param7Name, Param8Type, Param8Name ) BODY_MACRO_COMBINE(CURRENT_FILE_ID,_,__LINE__,_DELEGATE) FUNC_DECLARE_DYNAMIC_DELEGATE( DelegateName, DelegateName##_DelegateWrapper, FUNC_CONCAT( Param1Type InParam1, Param2Type InParam2, Param3Type InParam3, Param4Type InParam4, Param5Type InParam5, Param6Type InParam6, Param7Type InParam7, Param8Type InParam8 ), FUNC_CONCAT( *this, InParam1, InParam2, InParam3, InParam4, InParam5, InParam6, InParam7, InParam8 ), void, Param1Type, Param2Type, Param3Type, Param4Type, Param5Type, Param6Type, Param7Type, Param8Type )
+#define DECLARE_DYNAMIC_MULTICAST_DELEGATE_EightParams( DelegateName, Param1Type, Param1Name, Param2Type, Param2Name, Param3Type, Param3Name, Param4Type, Param4Name, Param5Type, Param5Name, Param6Type, Param6Name, Param7Type, Param7Name, Param8Type, Param8Name ) BODY_MACRO_COMBINE(CURRENT_FILE_ID,_,__LINE__,_DELEGATE) FUNC_DECLARE_DYNAMIC_MULTICAST_DELEGATE( DelegateName, DelegateName##_DelegateWrapper, FUNC_CONCAT( Param1Type InParam1, Param2Type InParam2, Param3Type InParam3, Param4Type InParam4, Param5Type InParam5, Param6Type InParam6, Param7Type InParam7, Param8Type InParam8 ), FUNC_CONCAT( *this, InParam1, InParam2, InParam3, InParam4, InParam5, InParam6, InParam7, InParam8 ), void, Param1Type, Param2Type, Param3Type, Param4Type, Param5Type, Param6Type, Param7Type, Param8Type )
 #define DECLARE_DELEGATE_RetVal_EightParams( ReturnValueType, DelegateName, Param1Type, Param2Type, Param3Type, Param4Type, Param5Type, Param6Type, Param7Type, Param8Type ) FUNC_DECLARE_DELEGATE( DelegateName, ReturnValueType, Param1Type, Param2Type, Param3Type, Param4Type, Param5Type, Param6Type, Param7Type, Param8Type )
-#define DECLARE_DYNAMIC_DELEGATE_RetVal_EightParams( ReturnValueType, DelegateName, Param1Type, Param1Name, Param2Type, Param2Name, Param3Type, Param3Name, Param4Type, Param4Name, Param5Type, Param5Name, Param6Type, Param6Name, Param7Type, Param7Name, Param8Type, Param8Name ) BODY_MACRO_COMBINE(CURRENT_FILE_ID,_,__LINE__,_DELEGATE) FUNC_DECLARE_DYNAMIC_DELEGATE_RETVAL( FWeakObjectPtr, DelegateName, DelegateName##_DelegateWrapper, ReturnValueType, FUNC_CONCAT( Param1Type InParam1, Param2Type InParam2, Param3Type InParam3, Param4Type InParam4, Param5Type InParam5, Param6Type InParam6, Param7Type InParam7, Param8Type InParam8 ), FUNC_CONCAT( *this, InParam1, InParam2, InParam3, InParam4, InParam5, InParam6, InParam7, InParam8 ), ReturnValueType, Param1Type, Param2Type, Param3Type, Param4Type, Param5Type, Param6Type, Param7Type, Param8Type )
+#define DECLARE_DYNAMIC_DELEGATE_RetVal_EightParams( ReturnValueType, DelegateName, Param1Type, Param1Name, Param2Type, Param2Name, Param3Type, Param3Name, Param4Type, Param4Name, Param5Type, Param5Name, Param6Type, Param6Name, Param7Type, Param7Name, Param8Type, Param8Name ) BODY_MACRO_COMBINE(CURRENT_FILE_ID,_,__LINE__,_DELEGATE) FUNC_DECLARE_DYNAMIC_DELEGATE_RETVAL( DelegateName, DelegateName##_DelegateWrapper, ReturnValueType, FUNC_CONCAT( Param1Type InParam1, Param2Type InParam2, Param3Type InParam3, Param4Type InParam4, Param5Type InParam5, Param6Type InParam6, Param7Type InParam7, Param8Type InParam8 ), FUNC_CONCAT( *this, InParam1, InParam2, InParam3, InParam4, InParam5, InParam6, InParam7, InParam8 ), ReturnValueType, Param1Type, Param2Type, Param3Type, Param4Type, Param5Type, Param6Type, Param7Type, Param8Type )
 
 #define DECLARE_DELEGATE_NineParams( DelegateName, Param1Type, Param2Type, Param3Type, Param4Type, Param5Type, Param6Type, Param7Type, Param8Type, Param9Type ) FUNC_DECLARE_DELEGATE( DelegateName, void, Param1Type, Param2Type, Param3Type, Param4Type, Param5Type, Param6Type, Param7Type, Param8Type, Param9Type )
 #define DECLARE_MULTICAST_DELEGATE_NineParams( DelegateName, Param1Type, Param2Type, Param3Type, Param4Type, Param5Type, Param6Type, Param7Type, Param8Type, Param9Type ) FUNC_DECLARE_MULTICAST_DELEGATE( DelegateName, void, Param1Type, Param2Type, Param3Type, Param4Type, Param5Type, Param6Type, Param7Type, Param8Type, Param9Type )
 #define DECLARE_TS_MULTICAST_DELEGATE_NineParams( DelegateName, Param1Type, Param2Type, Param3Type, Param4Type, Param5Type, Param6Type, Param7Type, Param8Type, Param9Type ) FUNC_DECLARE_TS_MULTICAST_DELEGATE( DelegateName, void, Param1Type, Param2Type, Param3Type, Param4Type, Param5Type, Param6Type, Param7Type, Param8Type, Param9Type )
 #define DECLARE_EVENT_NineParams( OwningType, EventName, Param1Type, Param2Type, Param3Type, Param4Type, Param5Type, Param6Type, Param7Type, Param8Type, Param9Type ) FUNC_DECLARE_EVENT( OwningType, EventName, void, Param1Type, Param2Type, Param3Type, Param4Type, Param5Type, Param6Type, Param7Type, Param8Type, Param9Type )
-#define DECLARE_DYNAMIC_DELEGATE_NineParams( DelegateName, Param1Type, Param1Name, Param2Type, Param2Name, Param3Type, Param3Name, Param4Type, Param4Name, Param5Type, Param5Name, Param6Type, Param6Name, Param7Type, Param7Name, Param8Type, Param8Name, Param9Type, Param9Name ) BODY_MACRO_COMBINE(CURRENT_FILE_ID,_,__LINE__,_DELEGATE) FUNC_DECLARE_DYNAMIC_DELEGATE( FWeakObjectPtr, DelegateName, DelegateName##_DelegateWrapper, FUNC_CONCAT( Param1Type InParam1, Param2Type InParam2, Param3Type InParam3, Param4Type InParam4, Param5Type InParam5, Param6Type InParam6, Param7Type InParam7, Param8Type InParam8, Param9Type InParam9 ), FUNC_CONCAT( *this, InParam1, InParam2, InParam3, InParam4, InParam5, InParam6, InParam7, InParam8, InParam9 ), void, Param1Type, Param2Type, Param3Type, Param4Type, Param5Type, Param6Type, Param7Type, Param8Type, Param9Type )
-#define DECLARE_DYNAMIC_MULTICAST_DELEGATE_NineParams( DelegateName, Param1Type, Param1Name, Param2Type, Param2Name, Param3Type, Param3Name, Param4Type, Param4Name, Param5Type, Param5Name, Param6Type, Param6Name, Param7Type, Param7Name, Param8Type, Param8Name, Param9Type, Param9Name ) BODY_MACRO_COMBINE(CURRENT_FILE_ID,_,__LINE__,_DELEGATE) FUNC_DECLARE_DYNAMIC_MULTICAST_DELEGATE( FWeakObjectPtr, DelegateName, DelegateName##_DelegateWrapper, FUNC_CONCAT( Param1Type InParam1, Param2Type InParam2, Param3Type InParam3, Param4Type InParam4, Param5Type InParam5, Param6Type InParam6, Param7Type InParam7, Param8Type InParam8, Param9Type InParam9 ), FUNC_CONCAT( *this, InParam1, InParam2, InParam3, InParam4, InParam5, InParam6, InParam7, InParam8, InParam9 ), void, Param1Type, Param2Type, Param3Type, Param4Type, Param5Type, Param6Type, Param7Type, Param8Type, Param9Type )
+#define DECLARE_DYNAMIC_DELEGATE_NineParams( DelegateName, Param1Type, Param1Name, Param2Type, Param2Name, Param3Type, Param3Name, Param4Type, Param4Name, Param5Type, Param5Name, Param6Type, Param6Name, Param7Type, Param7Name, Param8Type, Param8Name, Param9Type, Param9Name ) BODY_MACRO_COMBINE(CURRENT_FILE_ID,_,__LINE__,_DELEGATE) FUNC_DECLARE_DYNAMIC_DELEGATE( DelegateName, DelegateName##_DelegateWrapper, FUNC_CONCAT( Param1Type InParam1, Param2Type InParam2, Param3Type InParam3, Param4Type InParam4, Param5Type InParam5, Param6Type InParam6, Param7Type InParam7, Param8Type InParam8, Param9Type InParam9 ), FUNC_CONCAT( *this, InParam1, InParam2, InParam3, InParam4, InParam5, InParam6, InParam7, InParam8, InParam9 ), void, Param1Type, Param2Type, Param3Type, Param4Type, Param5Type, Param6Type, Param7Type, Param8Type, Param9Type )
+#define DECLARE_DYNAMIC_MULTICAST_DELEGATE_NineParams( DelegateName, Param1Type, Param1Name, Param2Type, Param2Name, Param3Type, Param3Name, Param4Type, Param4Name, Param5Type, Param5Name, Param6Type, Param6Name, Param7Type, Param7Name, Param8Type, Param8Name, Param9Type, Param9Name ) BODY_MACRO_COMBINE(CURRENT_FILE_ID,_,__LINE__,_DELEGATE) FUNC_DECLARE_DYNAMIC_MULTICAST_DELEGATE( DelegateName, DelegateName##_DelegateWrapper, FUNC_CONCAT( Param1Type InParam1, Param2Type InParam2, Param3Type InParam3, Param4Type InParam4, Param5Type InParam5, Param6Type InParam6, Param7Type InParam7, Param8Type InParam8, Param9Type InParam9 ), FUNC_CONCAT( *this, InParam1, InParam2, InParam3, InParam4, InParam5, InParam6, InParam7, InParam8, InParam9 ), void, Param1Type, Param2Type, Param3Type, Param4Type, Param5Type, Param6Type, Param7Type, Param8Type, Param9Type )
 #define DECLARE_DELEGATE_RetVal_NineParams( ReturnValueType, DelegateName, Param1Type, Param2Type, Param3Type, Param4Type, Param5Type, Param6Type, Param7Type, Param8Type, Param9Type ) FUNC_DECLARE_DELEGATE( DelegateName, ReturnValueType, Param1Type, Param2Type, Param3Type, Param4Type, Param5Type, Param6Type, Param7Type, Param8Type, Param9Type )
-#define DECLARE_DYNAMIC_DELEGATE_RetVal_NineParams( ReturnValueType, DelegateName, Param1Type, Param1Name, Param2Type, Param2Name, Param3Type, Param3Name, Param4Type, Param4Name, Param5Type, Param5Name, Param6Type, Param6Name, Param7Type, Param7Name, Param8Type, Param8Name, Param9Type, Param9Name ) BODY_MACRO_COMBINE(CURRENT_FILE_ID,_,__LINE__,_DELEGATE) FUNC_DECLARE_DYNAMIC_DELEGATE_RETVAL( FWeakObjectPtr, DelegateName, DelegateName##_DelegateWrapper, ReturnValueType, FUNC_CONCAT( Param1Type InParam1, Param2Type InParam2, Param3Type InParam3, Param4Type InParam4, Param5Type InParam5, Param6Type InParam6, Param7Type InParam7, Param8Type InParam8, Param9Type InParam9 ), FUNC_CONCAT( *this, InParam1, InParam2, InParam3, InParam4, InParam5, InParam6, InParam7, InParam8, InParam9 ), ReturnValueType, Param1Type, Param2Type, Param3Type, Param4Type, Param5Type, Param6Type, Param7Type, Param8Type, Param9Type )
+#define DECLARE_DYNAMIC_DELEGATE_RetVal_NineParams( ReturnValueType, DelegateName, Param1Type, Param1Name, Param2Type, Param2Name, Param3Type, Param3Name, Param4Type, Param4Name, Param5Type, Param5Name, Param6Type, Param6Name, Param7Type, Param7Name, Param8Type, Param8Name, Param9Type, Param9Name ) BODY_MACRO_COMBINE(CURRENT_FILE_ID,_,__LINE__,_DELEGATE) FUNC_DECLARE_DYNAMIC_DELEGATE_RETVAL( DelegateName, DelegateName##_DelegateWrapper, ReturnValueType, FUNC_CONCAT( Param1Type InParam1, Param2Type InParam2, Param3Type InParam3, Param4Type InParam4, Param5Type InParam5, Param6Type InParam6, Param7Type InParam7, Param8Type InParam8, Param9Type InParam9 ), FUNC_CONCAT( *this, InParam1, InParam2, InParam3, InParam4, InParam5, InParam6, InParam7, InParam8, InParam9 ), ReturnValueType, Param1Type, Param2Type, Param3Type, Param4Type, Param5Type, Param6Type, Param7Type, Param8Type, Param9Type )
diff --git a/Engine/Source/Runtime/CoreUObject/Public/UObject/SparseDelegate.h b/Engine/Source/Runtime/CoreUObject/Public/UObject/SparseDelegate.h
--- a/Engine/Source/Runtime/CoreUObject/Public/UObject/SparseDelegate.h
+++ b/Engine/Source/Runtime/CoreUObject/Public/UObject/SparseDelegate.h
@@ -437,7 +437,7 @@ private:
 };
 
 #define FUNC_DECLARE_DYNAMIC_MULTICAST_SPARSE_DELEGATE(SparseDelegateClassName, OwningClass, DelegateName, FuncParamList, FuncParamPassThru, ...) \
-	FUNC_DECLARE_DYNAMIC_MULTICAST_DELEGATE(FWeakObjectPtr, SparseDelegateClassName##_MCSignature, SparseDelegateClassName##_DelegateWrapper, FUNC_CONCAT(FuncParamList), FUNC_CONCAT(FuncParamPassThru), __VA_ARGS__) \
+	FUNC_DECLARE_DYNAMIC_MULTICAST_DELEGATE(SparseDelegateClassName##_MCSignature, SparseDelegateClassName##_DelegateWrapper, FUNC_CONCAT(FuncParamList), FUNC_CONCAT(FuncParamPassThru), __VA_ARGS__) \
 	struct SparseDelegateClassName##InfoGetter \
 	{ \
 		static const char* GetDelegateName() { return #DelegateName; } \
