Rename FInterpCurve<T>::InaccurateFindNearest into FInterpCurve<T>::FindNearest

https://docs.google.com/presentation/d/1ktMLMfY-pgjqwLMBERA1NJkEAB_jlyMt9ni99QdTdl0/edit?usp=sharing
This document is an analysis of the current algorithm, a presentation of the improvements and some tests

diff --git a/Engine/Plugins/Experimental/PCG/Source/PCG/Private/Data/PCGSplineData.cpp b/Engine/Plugins/Experimental/PCG/Source/PCG/Private/Data/PCGSplineData.cpp
--- a/Engine/Plugins/Experimental/PCG/Source/PCG/Private/Data/PCGSplineData.cpp
+++ b/Engine/Plugins/Experimental/PCG/Source/PCG/Private/Data/PCGSplineData.cpp
@@ -204,7 +204,7 @@ bool UPCGSplineProjectionData::SamplePoint(const FTransform& InTransform, const
 	FVector2D LocalPosition2D = Project(LocalTransform.GetLocation());
 	float Dummy;
 	// Find nearest key on 2D spline
-	float NearestInputKey = ProjectedPosition.InaccurateFindNearest(LocalPosition2D, Dummy);
+	float NearestInputKey = ProjectedPosition.FindNearest(LocalPosition2D, Dummy);
 	// TODO: if we didn't want to hand off density computation to the spline and do it here instead, we could do it in 2D space.
 	// Find point on original spline using the previously found key. Note this is an approximation that might not hold true since
 	// we are changing the curve length. Also, to support surface orientations that are not axis aligned, the project function
diff --git a/Engine/Plugins/Experimental/PCG/Source/PCG/Private/Elements/PCGSplineSampler.cpp b/Engine/Plugins/Experimental/PCG/Source/PCG/Private/Elements/PCGSplineSampler.cpp
--- a/Engine/Plugins/Experimental/PCG/Source/PCG/Private/Elements/PCGSplineSampler.cpp
+++ b/Engine/Plugins/Experimental/PCG/Source/PCG/Private/Elements/PCGSplineSampler.cpp
@@ -816,7 +816,7 @@ namespace PCGSplineSampler
 
 			{
 				TRACE_CPUPROFILER_EVENT_SCOPE(FPCGSplineSamplerElement::Execute::GetVoronoiEdges);
-				GetVoronoiEdges(PolygonPoints, FBox(MinPoint, FVector(MaxPoint.X, MaxPoint.Y, MinPoint.Z)), VoronoiEdges, CellMember);
+			GetVoronoiEdges(PolygonPoints, FBox(MinPoint, FVector(MaxPoint.X, MaxPoint.Y, MinPoint.Z)), VoronoiEdges, CellMember);
 			}
 
 			// Find the subset of the Voronoi Diagram which composes the Medial Axis
@@ -873,102 +873,119 @@ namespace PCGSplineSampler
 
 			TArray<FVector::FReal> DistancesSquaredOnXY;
 
-			// Point sampling
+		// Point sampling
 			for(FVector::FReal Y = LocalMinY; Y < LocalMaxY; Y += Params.InteriorSampleSpacing)
+		{
+			const FVector2D RayMin(MinPoint.X - RayPadding, Y);
+			const FVector2D RayMax(MaxPoint.X + RayPadding, Y);
+
+			// Get the intersections along this ray, sorted by Y value
+			TArray<FVector2D> Intersections;
+			PCGSplineSamplerHelpers::SegmentPolygonIntersection2D(RayMin, RayMax, SplineSamplePoints2D, Intersections);
+
+			if (Intersections.Num() % 2 != 0)
 			{
-				const FVector2D RayMin(MinPoint.X - RayPadding, Y);
-				const FVector2D RayMax(MaxPoint.X + RayPadding, Y);
+				UE_LOG(LogPCG, Error, TEXT("Intersection test failed, skipping samples for this row"));
+				continue;
+			}
+
+			Intersections.Sort([](const FVector2D& LHS, const FVector2D& RHS) { return LHS.X < RHS.X; });
 
-				// Get the intersections along this ray, sorted by Y value
-				TArray<FVector2D> Intersections;
-				PCGSplineSamplerHelpers::SegmentPolygonIntersection2D(RayMin, RayMax, SplineSamplePoints2D, Intersections);
+			// TODO: async processing
+			// Each pair of intersections defines a range in which point samples may lie
+			for (int32 RangeIndex = 0; RangeIndex < Intersections.Num(); RangeIndex += 2)
+			{
+				const FVector::FReal MinX = FMath::CeilToDouble(Intersections[RangeIndex].X / Params.InteriorSampleSpacing) * Params.InteriorSampleSpacing;
+				const FVector::FReal MaxX = FMath::FloorToDouble(Intersections[RangeIndex + 1].X / Params.InteriorSampleSpacing) * Params.InteriorSampleSpacing;
 
-				if (Intersections.Num() % 2 != 0)
+				if (MaxX - MinX < Params.InteriorSampleSpacing)
 				{
-					UE_LOG(LogPCG, Error, TEXT("Intersection test failed, skipping samples for this row"));
 					continue;
 				}
 
-				Intersections.Sort([](const FVector2D& LHS, const FVector2D& RHS) { return LHS.X < RHS.X; });
-
-				// TODO: async processing
-				// Each pair of intersections defines a range in which point samples may lie
-				for (int32 RangeIndex = 0; RangeIndex < Intersections.Num(); RangeIndex += 2)
+				for (FVector::FReal X = MinX; X < MaxX + UE_KINDA_SMALL_NUMBER; X += Params.InteriorSampleSpacing)
 				{
-					const FVector::FReal MinX = FMath::CeilToDouble(Intersections[RangeIndex].X / Params.InteriorSampleSpacing) * Params.InteriorSampleSpacing;
-					const FVector::FReal MaxX = FMath::FloorToDouble(Intersections[RangeIndex + 1].X / Params.InteriorSampleSpacing) * Params.InteriorSampleSpacing;
-
-					if (MaxX - MinX < Params.InteriorSampleSpacing)
-					{
-						continue;
-					}
+					const FVector2D SampleLocation = FVector2D(X, Y);
+					FVector SurfaceLocation = FVector(SampleLocation, MinPoint.Z);
 
-					for (FVector::FReal X = MinX; X < MaxX + UE_KINDA_SMALL_NUMBER; X += Params.InteriorSampleSpacing)
+					if (bProjectOntoSurface)
 					{
-						const FVector2D SampleLocation = FVector2D(X, Y);
-						FVector SurfaceLocation = FVector(SampleLocation, MinPoint.Z);
+						// Precompute 2D distance to every spline point
+						DistancesSquaredOnXY.Reset(SplineSamplePoints.Num());
+						for (const FVector& SplinePoint : SplineSamplePoints)
+						{
+							FVector::FReal DistanceSquared = FVector::DistSquaredXY(SurfaceLocation, SplinePoint);
+							DistancesSquaredOnXY.Add(DistanceSquared);
+						}
 
-						if (bProjectOntoSurface)
+						// Compute average Z value weighted by 1 / Distance^2
+						FVector::FReal SumZ = 0.f;
+						FVector::FReal SumWeights = 0.f;
+						for (int32 PointIndex = 0; PointIndex < SplineSamplePoints.Num(); ++PointIndex)
 						{
-							// Precompute 2D distance to every spline point
-							DistancesSquaredOnXY.Reset(SplineSamplePoints.Num());
-							for (const FVector& SplinePoint : SplineSamplePoints)
+							const FVector::FReal DistanceSquared = DistancesSquaredOnXY[PointIndex];
+
+							// If sample point overlaps exactly with a border point, then that must be the height
+							if (FMath::IsNearlyZero(DistanceSquared))
 							{
-								FVector::FReal DistanceSquared = FVector::DistSquaredXY(SurfaceLocation, SplinePoint);
-								DistancesSquaredOnXY.Add(DistanceSquared);
+								SumZ = SplineSamplePoints[PointIndex].Z;
+								SumWeights = 1.f;
+								break;
 							}
 
-							// Compute average Z value weighted by 1 / Distance^2
-							FVector::FReal SumZ = 0.f;
-							FVector::FReal SumWeights = 0.f;
-							for (int32 PointIndex = 0; PointIndex < SplineSamplePoints.Num(); ++PointIndex)
-							{
-								const FVector::FReal DistanceSquared = DistancesSquaredOnXY[PointIndex];
+							// TODO: it would be more accurate to use distance to the polyline instead of distance to the polyline points,
+							// however it would also be much more expensive. Perhaps worth investigating when Params.bTreatSplineAsPolyline is true
+							const FVector::FReal Weight = 1.f / DistanceSquared;
 
-								// If sample point overlaps exactly with a border point, then that must be the height
-								if (FMath::IsNearlyZero(DistanceSquared))
-								{
-									SumZ = SplineSamplePoints[PointIndex].Z;
-									SumWeights = 1.f;
-									break;
-								}
+							SumWeights += Weight;
+							SumZ += SplineSamplePoints[PointIndex].Z * Weight;
+						}
 
-								// TODO: it would be more accurate to use distance to the polyline instead of distance to the polyline points,
-								// however it would also be much more expensive. Perhaps worth investigating when Params.bTreatSplineAsPolyline is true
-								const FVector::FReal Weight = 1.f / DistanceSquared;
+						SurfaceLocation.Z = SumZ / SumWeights;
+					}
 
-								SumWeights += Weight;
-								SumZ += SplineSamplePoints[PointIndex].Z * Weight;
-							}
+					// if bTreatAsPolyline, then we shouldnt use this, we should use nearest point on the polygon line segments
+					float Dummy;
+					const float NearestSplineKey = bFindNearestSplineKey ? Spline->SplineCurves.Position.FindNearest(SurfaceLocation, Dummy) : 0.f;
 
-							SurfaceLocation.Z = SumZ / SumWeights;
-						}
+					const FVector PointLocationLS = (Params.bProjectOntoSurface ? FVector(SurfaceLocation) : FVector(SampleLocation, MinPoint.Z));
+
+					FTransform TransformLS = FTransform::Identity;
+					TransformLS.SetLocation(PointLocationLS);
 
-						// if bTreatAsPolyline, then we shouldnt use this, we should use nearest point on the polygon line segments
-						float Dummy;
-						const float NearestSplineKey = bFindNearestSplineKey ? Spline->SplineCurves.Position.InaccurateFindNearest(SurfaceLocation, Dummy) : 0.f;
+					if (Params.InteriorOrientation == EPCGSplineSamplingInteriorOrientation::FollowCurvature)
+					{
+						TransformLS.SetRotation(Spline->GetQuaternionAtSplineInputKey(NearestSplineKey, ESplineCoordinateSpace::Local));
+					}
 
-						const FVector PointLocationLS = (Params.bProjectOntoSurface ? FVector(SurfaceLocation) : FVector(SampleLocation, MinPoint.Z));
+					// Calculate density fall off
+						float Density = 1.0f;
 
-						FTransform TransformLS = FTransform::Identity;
-						TransformLS.SetLocation(PointLocationLS);
+					if (bGenerateMedialAxis && MedialAxisEdges.Num())
+					{
+						FVector::FReal SmallestDistSquared = MaxDimensionSquared;
 
-						if (Params.InteriorOrientation == EPCGSplineSamplingInteriorOrientation::FollowCurvature)
+						// Find distance from SampleLocation to MedialAxis
+						for (const TTuple<FVector2D, FVector2D>& Edge : MedialAxisEdges)
 						{
-							TransformLS.SetRotation(Spline->GetQuaternionAtSplineInputKey(NearestSplineKey, ESplineCoordinateSpace::Local));
+							const FVector::FReal DistSquared = FMath::PointDistToSegmentSquared(FVector(SampleLocation, 0), FVector(Edge.Get<0>(), 0), FVector(Edge.Get<1>(), 0));
+
+							if (DistSquared < SmallestDistSquared)
+							{
+								SmallestDistSquared = DistSquared;
+							}
 						}
 
-						// Calculate density fall off
-						float Density = 1.0f;
+						const FVector::FReal SmallestDist = FMath::Sqrt(SmallestDistSquared);
 
-						if (bGenerateMedialAxis && MedialAxisEdges.Num())
-						{
-							FVector::FReal SmallestDistSquared = MaxDimensionSquared;
+						FVector::FReal PointToSplineDist = 0.f;
 
-							// Find distance from SampleLocation to MedialAxis
-							for (const TTuple<FVector2D, FVector2D>& Edge : MedialAxisEdges)
+						if (Params.bTreatSplineAsPolyline)
+						{
+							SmallestDistSquared = MaxDimensionSquared;
+							for (int32 PointIndex = 0; PointIndex < SplineSamplePoints.Num(); ++PointIndex)
 							{
-								const FVector::FReal DistSquared = FMath::PointDistToSegmentSquared(FVector(SampleLocation, 0), FVector(Edge.Get<0>(), 0), FVector(Edge.Get<1>(), 0));
+								const FVector::FReal DistSquared = FMath::PointDistToSegmentSquared(SurfaceLocation, SplineSamplePoints[PointIndex], SplineSamplePoints[(PointIndex + 1) % SplineSamplePoints.Num()]);
 
 								if (DistSquared < SmallestDistSquared)
 								{
@@ -976,48 +993,31 @@ namespace PCGSplineSampler
 								}
 							}
 
-							const FVector::FReal SmallestDist = FMath::Sqrt(SmallestDistSquared);
-
-							FVector::FReal PointToSplineDist = 0.f;
-
-							if (Params.bTreatSplineAsPolyline)
-							{
-								SmallestDistSquared = MaxDimensionSquared;
-								for (int32 PointIndex = 0; PointIndex < SplineSamplePoints.Num(); ++PointIndex)
-								{
-									const FVector::FReal DistSquared = FMath::PointDistToSegmentSquared(SurfaceLocation, SplineSamplePoints[PointIndex], SplineSamplePoints[(PointIndex + 1) % SplineSamplePoints.Num()]);
-
-									if (DistSquared < SmallestDistSquared)
-									{
-										SmallestDistSquared = DistSquared;
-									}
-								}
-
-								PointToSplineDist = FMath::Sqrt(SmallestDistSquared);
-							}
-							else
-							{
+							PointToSplineDist = FMath::Sqrt(SmallestDistSquared);
+						}
+						else
+						{
 								const FVector NearestSplineLocation = Spline->GetLocationAtSplineInputKey(NearestSplineKey, ESplineCoordinateSpace::Local);
-								PointToSplineDist = FVector2D::Distance(SampleLocation, FVector2D(NearestSplineLocation.X, NearestSplineLocation.Y));
-							}
+							PointToSplineDist = FVector2D::Distance(SampleLocation, FVector2D(NearestSplineLocation.X, NearestSplineLocation.Y));
+						}
 
-							// Linear fall off in the range [0, 1]
-							const FVector::FReal T = SmallestDist / (SmallestDist + PointToSplineDist + UE_KINDA_SMALL_NUMBER);
+						// Linear fall off in the range [0, 1]
+						const FVector::FReal T = SmallestDist / (SmallestDist + PointToSplineDist + UE_KINDA_SMALL_NUMBER);
 							Density = DensityFalloffCurve->Eval(T);
 						}
 
 						InteriorSplinePointData[DispatchIndex].Emplace(TransformLS * LineData->GetTransform(), PointLocationLS, Density);
 					}
-				}
+					}
 			}
 		});
 
 		// Finally, gather the data and push to the points
 		int32 PointCount = 0;
 		for (const TArray<TTuple<FTransform, FVector, float>>& InteriorData : InteriorSplinePointData)
-		{
+					{
 			PointCount += InteriorData.Num();
-		}
+					}
 
 		// TODO: should we parallel for this too?
 		OutPoints.Reserve(PointCount);
@@ -1029,8 +1029,8 @@ namespace PCGSplineSampler
 				Point.Transform = InteriorPoint.Get<0>();
 				SetSeed(Point, InteriorPoint.Get<1>(), Params);
 				Point.Density = InteriorPoint.Get<2>();
-				Point.BoundsMin = BoundsMin;
-				Point.BoundsMax = BoundsMax;
+					Point.BoundsMin = BoundsMin;
+					Point.BoundsMax = BoundsMax;
 			}
 		}
 	}
diff --git a/Engine/Plugins/Experimental/Water/Source/Runtime/Private/BuoyancyComponent.cpp b/Engine/Plugins/Experimental/Water/Source/Runtime/Private/BuoyancyComponent.cpp
--- a/Engine/Plugins/Experimental/Water/Source/Runtime/Private/BuoyancyComponent.cpp
+++ b/Engine/Plugins/Experimental/Water/Source/Runtime/Private/BuoyancyComponent.cpp
@@ -538,10 +538,10 @@ float GetWaterSplineKeyFast(FVector Location, const UWaterBodyComponent* WaterBo
 	if (Segment == -1)
 	{
 		float DummyDistance;
-		return InterpCurve.InaccurateFindNearest(LocalLocation, DummyDistance, Segment);
+		return InterpCurve.FindNearest(LocalLocation, DummyDistance, Segment);
 	}
 
-	//We have the cached segment, so search for the best point as in FInterpCurve<T>::InaccurateFindNearest
+	//We have the cached segment, so search for the best point as in FInterpCurve<T>::FindNearest
 	//but only in the current segment and the two immediate neighbors
 
 	//River splines aren't looped, so we don't have to handle that case
@@ -556,7 +556,7 @@ float GetWaterSplineKeyFast(FVector Location, const UWaterBodyComponent* WaterBo
 		{
 			const int32 SegmentIdx = FMath::Clamp(i, 0, LastSegmentIdx);
 			float LocalDistanceSq;
-			float LocalResult = InterpCurve.InaccurateFindNearestOnSegment(LocalLocation, SegmentIdx, LocalDistanceSq);
+			float LocalResult = InterpCurve.FindNearestOnSegment(LocalLocation, SegmentIdx, LocalDistanceSq);
 			if (LocalDistanceSq < BestDistanceSq)
 			{
 				BestDistanceSq = LocalDistanceSq;
@@ -569,7 +569,7 @@ float GetWaterSplineKeyFast(FVector Location, const UWaterBodyComponent* WaterBo
 		{
 			//We're at either end of the search - it's possible we skipped a segment so just do a full lookup in this case
 			float DummyDistance;
-			return InterpCurve.InaccurateFindNearest(LocalLocation, DummyDistance, Segment);
+			return InterpCurve.FindNearest(LocalLocation, DummyDistance, Segment);
 		}
 
 		Segment = BestSegment;
diff --git a/Engine/Plugins/Experimental/Water/Source/Runtime/Private/WaterBodyTypes.cpp b/Engine/Plugins/Experimental/Water/Source/Runtime/Private/WaterBodyTypes.cpp
--- a/Engine/Plugins/Experimental/Water/Source/Runtime/Private/WaterBodyTypes.cpp
+++ b/Engine/Plugins/Experimental/Water/Source/Runtime/Private/WaterBodyTypes.cpp
@@ -33,7 +33,7 @@ float FWaterSplineDataPhysics::FindInputKeyClosestToWorldLocation(const FVector&
 {
 	const FVector LocalLocation = ComponentTransform.InverseTransformPosition(WorldLocation);
 	float Dummy;
-	return SplineCurves.Position.InaccurateFindNearest(LocalLocation, Dummy);
+	return SplineCurves.Position.FindNearest(LocalLocation, Dummy);
 }
 
 FVector FWaterSplineDataPhysics::GetUpVectorAtSplineInputKey(float InKey, ESplineCoordinateSpace::Type CoordinateSpace /*= ESplineCoordinateSpace::World*/) const
diff --git a/Engine/Plugins/FX/Niagara/Source/Niagara/Private/NiagaraDataInterfaceSpline.cpp b/Engine/Plugins/FX/Niagara/Source/Niagara/Private/NiagaraDataInterfaceSpline.cpp
--- a/Engine/Plugins/FX/Niagara/Source/Niagara/Private/NiagaraDataInterfaceSpline.cpp
+++ b/Engine/Plugins/FX/Niagara/Source/Niagara/Private/NiagaraDataInterfaceSpline.cpp
@@ -1023,7 +1023,7 @@ template <>
 float FNDISpline_InstanceData::EvaluateFindNearestPosition<TIntegralConstant<bool, false>>(FVector InPosition) const
 {
 	float Dummy;
-	return SplineCurves.Position.InaccurateFindNearest(InPosition, Dummy);
+	return SplineCurves.Position.FindNearest(InPosition, Dummy);
 }
 
 template <>
diff --git a/Engine/Source/Editor/ComponentVisualizers/Private/SplineComponentVisualizer.cpp b/Engine/Source/Editor/ComponentVisualizers/Private/SplineComponentVisualizer.cpp
--- a/Engine/Source/Editor/ComponentVisualizers/Private/SplineComponentVisualizer.cpp
+++ b/Engine/Source/Editor/ComponentVisualizers/Private/SplineComponentVisualizer.cpp
@@ -2391,7 +2391,7 @@ float FSplineComponentVisualizer::FindNearest(const FVector& InLocalPos, int32 I
 
 	FInterpCurveVector& SplinePosition = SplineComp->GetSplinePointsPosition();
 	float OutSquaredDistance = 0.0f;
-	float t = SplinePosition.InaccurateFindNearestOnSegment(InLocalPos, InSegmentIndex, OutSquaredDistance);
+	float t = SplinePosition.FindNearestOnSegment(InLocalPos, InSegmentIndex, OutSquaredDistance);
 	OutSplinePos = SplinePosition.Eval(t, FVector::ZeroVector);
 	OutSplineTangent = SplinePosition.EvalDerivative(t, FVector::ZeroVector);
 
diff --git a/Engine/Source/Runtime/Core/Public/Math/InterpCurve.h b/Engine/Source/Runtime/Core/Public/Math/InterpCurve.h
--- a/Engine/Source/Runtime/Core/Public/Math/InterpCurve.h
+++ b/Engine/Source/Runtime/Core/Public/Math/InterpCurve.h
@@ -4,6 +4,7 @@
 
 #include "CoreTypes.h"
 #include "Misc/AssertionMacros.h"
+#include "Algo/MinElement.h"
 #include "Containers/Array.h"
 #include "Math/UnrealMathUtility.h"
 #include "Math/Color.h"
@@ -96,7 +97,12 @@ public:
 	 * @param OutDistanceSq - output - the squared distance between the given point and the closest found point.
 	 * @return The key (the 't' parameter) of the nearest point. 
 	 */
-	float InaccurateFindNearest( const T &PointInSpace, float& OutDistanceSq ) const;
+	float FindNearest( const T &PointInSpace, float& OutDistanceSq ) const;
+
+	/**
+	 * @deprecated Use FindNearest instead.
+	 */
+	float InaccurateFindNearest(const T& PointInSpace, float& OutDistanceSq) const;
 
 	/**
 	* Find the nearest point on spline to the given point.
@@ -106,7 +112,12 @@ public:
 	* @param OutSegment - output - the nearest segment to the given point.
 	* @return The key (the 't' parameter) of the nearest point.
 	*/
-	float InaccurateFindNearest( const T &PointInSpace, float& OutDistanceSq, float& OutSegment ) const;
+	float FindNearest( const T &PointInSpace, float& OutDistanceSq, float& OutSegment ) const;
+
+	/**
+	 * @deprecated Use FindNearest instead.
+	 */
+	float InaccurateFindNearest(const T& PointInSpace, float& OutDistanceSq, float& OutSegment) const;
 
 	/** 
 	 * Find the nearest point (to the given point) on segment between Points[PtIdx] and Points[PtIdx+1]
@@ -114,7 +125,12 @@ public:
 	 * @param PointInSpace - the given point
 	 * @return The key (the 't' parameter) of the found point. 
 	 */
-	float InaccurateFindNearestOnSegment( const T &PointInSpace, int32 PtIdx, float& OutSquaredDistance ) const;
+	float FindNearestOnSegment( const T &PointInSpace, int32 PtIdx, float& OutSquaredDistance ) const;
+
+	/**
+	 * @deprecated Use FindNearestOnSegment instead.
+	 */
+	float InaccurateFindNearestOnSegment(const T& PointInSpace, int32 PtIdx, float& OutSquaredDistance) const;
 
 	/** Automatically set the tangents on the curve based on surrounding points */
 	void AutoSetTangents(float Tension = 0.0f, bool bStationaryEndpoints = true);
@@ -417,7 +433,6 @@ T FInterpCurve<T>::EvalDerivative(const float InVal, const T& Default) const
 	}
 }
 
-
 template< class T >
 T FInterpCurve<T>::EvalSecondDerivative(const float InVal, const T& Default) const
 {
@@ -480,16 +495,22 @@ T FInterpCurve<T>::EvalSecondDerivative(const float InVal, const T& Default) con
 	}
 }
 
+template< class T >
+float FInterpCurve<T>::FindNearest(const T& PointInSpace, float& OutDistanceSq) const
+{
+	float OutSegment;
+	return FindNearest(PointInSpace, OutDistanceSq, OutSegment);
+}
 
 template< class T >
-float FInterpCurve<T>::InaccurateFindNearest(const T &PointInSpace, float& OutDistanceSq) const							// LWC_TODO: Precision loss
+float FInterpCurve<T>::InaccurateFindNearest(const T &PointInSpace, float& OutDistanceSq) const
 {
 	float OutSegment;
-	return InaccurateFindNearest(PointInSpace, OutDistanceSq, OutSegment);
+	return FindNearest(PointInSpace, OutDistanceSq, OutSegment);
 }
 
 template< class T >
-float FInterpCurve<T>::InaccurateFindNearest(const T &PointInSpace, float& OutDistanceSq, float& OutSegment) const		// LWC_TODO: Precision loss
+float FInterpCurve<T>::FindNearest(const T &PointInSpace, float& OutDistanceSq, float& OutSegment) const		// LWC_TODO: Precision loss
 {
 	const int32 NumPoints = Points.Num();
 	const int32 NumSegments = bIsLooped ? NumPoints : NumPoints - 1;
@@ -497,12 +518,12 @@ float FInterpCurve<T>::InaccurateFindNearest(const T &PointInSpace, float& OutDi
 	if (NumPoints > 1)
 	{
 		float BestDistanceSq;
-		float BestResult = InaccurateFindNearestOnSegment(PointInSpace, 0, BestDistanceSq);
+		float BestResult = FindNearestOnSegment(PointInSpace, 0, BestDistanceSq);
 		float BestSegment = 0;
 		for (int32 Segment = 1; Segment < NumSegments; ++Segment)
 		{
 			float LocalDistanceSq;
-			float LocalResult = InaccurateFindNearestOnSegment(PointInSpace, Segment, LocalDistanceSq);
+			float LocalResult = FindNearestOnSegment(PointInSpace, Segment, LocalDistanceSq);
 			if (LocalDistanceSq < BestDistanceSq)
 			{
 				BestDistanceSq = LocalDistanceSq;
@@ -525,9 +546,14 @@ float FInterpCurve<T>::InaccurateFindNearest(const T &PointInSpace, float& OutDi
 	return 0.0f;
 }
 
+template< class T >
+float FInterpCurve<T>::InaccurateFindNearest(const T& PointInSpace, float& OutDistanceSq, float& OutSegment) const
+{
+	return FindNearest(PointInSpace, OutDistanceSq, OutSegment);
+}
 
 template< class T >
-float FInterpCurve<T>::InaccurateFindNearestOnSegment(const T& PointInSpace, int32 PtIdx, float& OutSquaredDistance) const		// LWC_TODO: Precision loss
+float FInterpCurve<T>::FindNearestOnSegment(const T& PointInSpace, int32 PtIdx, float& OutSquaredDistance) const
 {
 	const int32 NumPoints = Points.Num();
 	const int32 LastPoint = NumPoints - 1;
@@ -561,58 +587,138 @@ float FInterpCurve<T>::InaccurateFindNearestOnSegment(const T& PointInSpace, int
 	}
 
 	{
-		const int32 PointsChecked = 3;
-		const int32 IterationNum = 3;
-		const float Scale = 0.75;
+		// To be accurate and fast, the algorithm is divided into 2 steps:
+		// A first step starting from the 4 start coordinates with a large tolerance to find a good approximation in few iterations
+		// if needed, a last step starting from the best approximation with a small tolerance to be accurate is performed
+		// Few small differences in the Newton algorithms:
+		// - in the first step, the break is done before the last evaluation (not needed) (BreakBeforeEvaluateStep1)
+		// - in the second step, the last evaluation is done before the break as the result point is the final point needed to evaluate the error (BreakAfterEvaluateStep2)
+		// - in the first step, to avoid duplicated test, if the new coordinate is:
+		//    - before the range ("start value" - 1/3), the process is break
+		//    - after the range ("start value" + 1/3), the next "start value" is not test
+
+		constexpr float ToleranceStep1 = 1.e-2f;
+		constexpr float ToleranceStep2 = 1.e-5f;
+		constexpr float ToleranceToLaunchStep2 = 1.e-6f; // Step1 break before the evaluation. So if the last move of step1 is not so small, a last evaluation has to be done
+
+		constexpr int32 MaxIteration = 20;
+		constexpr float InvThree = 1.f / 3.f;
+		constexpr float TwoInvThree = 2.f / 3.f;
 
-		// Newton's methods is repeated 3 times, starting with t = 0, 0.5, 1.
-		float ValuesT[PointsChecked];
-		ValuesT[0] = 0.0f;
-		ValuesT[1] = 0.5f;
-		ValuesT[2] = 1.0f;
+		// The parametric space is divided into 3 areas. The search starts from the four limits of the areas
+		// So Newton's methods is repeated 4 times
+		constexpr int32 NewtonIterationCount = 4;
 
-		T InitialPoints[PointsChecked];
-		InitialPoints[0] = Points[PtIdx].OutVal;
-		InitialPoints[1] = FMath::CubicInterp(Points[PtIdx].OutVal, Points[PtIdx].LeaveTangent * Diff, Points[NextPtIdx].OutVal, Points[NextPtIdx].ArriveTangent * Diff, ValuesT[1]);
-		InitialPoints[2] = Points[NextPtIdx].OutVal;
+		constexpr float StartValuesT[] = { 0.f, InvThree, TwoInvThree, 1.f };
+		float ValuesT[] = {0.f, InvThree, TwoInvThree, 1.f};
+		float NextMovesT[] = { 0.f, 0.f, 0.f, 0.f };
+		float DistancesSq[NewtonIterationCount];
 
-		float DistancesSq[PointsChecked];
+		float LastValue = UE_BIG_NUMBER;
+		float MinValue = 0.f;
+		float MaxValue = UE_BIG_NUMBER;
+		int32 NextValue = 1;
 
-		for (int32 point = 0; point < PointsChecked; ++point)
+		const TArray<FInterpCurvePoint<T>>& PointsT = Points;
+
+		TFunction<bool(float&, const float)> BreakIfConverged = [&LastValue, &MinValue, &MaxValue, &NextValue, &StartValuesT](float& Value, const float Tolerance) -> bool
 		{
-			//Algorithm explanation: http://permalink.gmane.org/gmane.games.devel.sweng/8285
-			T FoundPoint = InitialPoints[point];
-			float LastMove = 1.0f;
-			for (int32 iter = 0; iter < IterationNum; ++iter)
+			if (FMath::IsNearlyEqual(LastValue, Value, Tolerance))
+			{
+				// Has converged
+				Value = LastValue;
+				return true;
+			}
+
+			if (Value > MaxValue)
 			{
-				const T LastBestTangent = FMath::CubicInterpDerivative(Points[PtIdx].OutVal, Points[PtIdx].LeaveTangent * Diff, Points[NextPtIdx].OutVal, Points[NextPtIdx].ArriveTangent * Diff, ValuesT[point]);
+				// The new coordinate is after the range ("start value" + 1/3), the next "start value" is not test
+				// Value is clamp into [0.,1.], so value cannot be > to 1.
+				++NextValue;
+				MaxValue = StartValuesT[NextValue >= NewtonIterationCount ? NewtonIterationCount - 1 : NextValue];
+			}
+
+			if (Value < MinValue)
+			{
+				// The new coordinate is before the range, the process is break as the area has been proceed 
+				return true;
+			}
+
+			LastValue = Value;
+
+			return false;
+		};
+
+		TFunction<bool(float&, const float)> NeverBreak = [](float& Value, const float Tolerance) -> bool
+		{
+			return false;
+		};
+
+		TFunction<float(float&, float&, TFunction<bool(float&, const float)>, TFunction<bool(float&, const float)>, const float)> Newton = [&PointsT, &PtIdx, &NextPtIdx, &PointInSpace, &Diff, MaxIteration, InvThree](float& Value, float& Move, TFunction<bool(float&, const float)> BreakBeforeEvaluate, TFunction<bool(float&, const float)> BreakAfterEvaluate, const float Tolerance) -> float
+		{
+			T FoundPoint;
+
+			for (int32 Iter = 0; Iter < MaxIteration; ++Iter)
+			{
+				Value += Move;
+				Value = FMath::Clamp(Value, 0.0f, 1.0f);
+
+				if (BreakBeforeEvaluate(Value, Tolerance))
+				{
+					break;
+				}
+
+				FoundPoint = FMath::CubicInterp(PointsT[PtIdx].OutVal, PointsT[PtIdx].LeaveTangent * Diff, PointsT[NextPtIdx].OutVal, PointsT[NextPtIdx].ArriveTangent * Diff, Value);
+
+				if (BreakAfterEvaluate(Value, Tolerance))
+				{
+					break;
+				}
+
+				const T Tangent = FMath::CubicInterpDerivative(PointsT[PtIdx].OutVal, PointsT[PtIdx].LeaveTangent * Diff, PointsT[NextPtIdx].OutVal, PointsT[NextPtIdx].ArriveTangent * Diff, Value);
 				const T Delta = (PointInSpace - FoundPoint);
-				float Move = static_cast<float>((LastBestTangent | Delta) / LastBestTangent.SizeSquared());
-				Move = FMath::Clamp(Move, -LastMove*Scale, LastMove*Scale);
-				ValuesT[point] += Move;
-				ValuesT[point] = FMath::Clamp(ValuesT[point], 0.0f, 1.0f);
-				LastMove = FMath::Abs(Move);
-				FoundPoint = FMath::CubicInterp(Points[PtIdx].OutVal, Points[PtIdx].LeaveTangent * Diff, Points[NextPtIdx].OutVal, Points[NextPtIdx].ArriveTangent * Diff, ValuesT[point]);
+				Move = static_cast<float>(Tangent.Dot(Delta) / Tangent.SizeSquared());
 			}
-			DistancesSq[point] = static_cast<float>((FoundPoint - PointInSpace).SizeSquared());
-			ValuesT[point] = ValuesT[point] * Diff + Points[PtIdx].InVal;
-		}
+			return static_cast<float>((FoundPoint - PointInSpace).SizeSquared());
+		};
 
-		if (DistancesSq[0] <= DistancesSq[1] && DistancesSq[0] <= DistancesSq[2])
+		// A first step starting from the 4 start coordinates with a large tolerance to find a good approximation in few iterations
+		for (int32 Index = 0; Index < NewtonIterationCount; ++Index)
 		{
-			OutSquaredDistance = DistancesSq[0];
-			return ValuesT[0];
+			MinValue = StartValuesT[Index == 0 ? 0 : Index - 1];
+			MaxValue = StartValuesT[NextValue >= NewtonIterationCount ? Index : NextValue];
+
+			float& Value = ValuesT[Index];
+			float& Move = NextMovesT[Index];
+
+			DistancesSq[Index] = Newton(Value, Move, BreakIfConverged, NeverBreak, ToleranceStep1);
 		}
-		if (DistancesSq[1] <= DistancesSq[2])
+
+		// Find the index of the best approximation
+		const float* MinValuePtr = Algo::MinElement(DistancesSq);
+		const int32 IndexOfTheSmallest = MinValuePtr - DistancesSq;
+
+		float& Move = NextMovesT[IndexOfTheSmallest];
+
+		// if needed, a last step starting from the best approximation with a small tolerance to be accurate is performed
+		if (!FMath::IsNearlyZero(Move, ToleranceToLaunchStep2))
 		{
-			OutSquaredDistance = DistancesSq[1];
-			return ValuesT[1];
+			MinValue = 0.f;
+			MaxValue = 1.f;
+			float& Value = ValuesT[IndexOfTheSmallest];
+			DistancesSq[IndexOfTheSmallest] = Newton(Value, Move, NeverBreak, BreakIfConverged, ToleranceStep2);
 		}
-		OutSquaredDistance = DistancesSq[2];
-		return ValuesT[2];
+
+		OutSquaredDistance = DistancesSq[IndexOfTheSmallest];
+		return ValuesT[IndexOfTheSmallest] * Diff + Points[PtIdx].InVal;
 	}
 }
 
+template< class T >
+float FInterpCurve<T>::InaccurateFindNearestOnSegment(const T& PointInSpace, int32 PtIdx, float& OutSquaredDistance) const
+{
+	return FindNearestOnSegment(PointInSpace, PtIdx, OutSquaredDistance);
+}
 
 template< class T >
 void FInterpCurve<T>::AutoSetTangents(float Tension, bool bStationaryEndpoints)
diff --git a/Engine/Source/Runtime/Engine/Private/Components/SplineComponent.cpp b/Engine/Source/Runtime/Engine/Private/Components/SplineComponent.cpp
--- a/Engine/Source/Runtime/Engine/Private/Components/SplineComponent.cpp
+++ b/Engine/Source/Runtime/Engine/Private/Components/SplineComponent.cpp
@@ -546,7 +546,7 @@ float USplineComponent::GetDistanceAlongSplineAtLocation(const FVector& InLocati
 {
 	const FVector LocalLocation = (CoordinateSpace == ESplineCoordinateSpace::World) ? GetComponentTransform().InverseTransformPosition(InLocation) : InLocation;
 	float Dummy;
-	float Key = SplineCurves.Position.InaccurateFindNearest(LocalLocation, Dummy);
+	float Key = SplineCurves.Position.FindNearest(LocalLocation, Dummy);
 	return GetDistanceAlongSplineAtSplineInputKey(Key);
 }
 
@@ -1543,7 +1543,7 @@ float USplineComponent::FindInputKeyClosestToWorldLocation(const FVector& WorldL
 {
 	const FVector LocalLocation = GetComponentTransform().InverseTransformPosition(WorldLocation);
 	float Dummy;
-	return SplineCurves.Position.InaccurateFindNearest(LocalLocation, Dummy);
+	return SplineCurves.Position.FindNearest(LocalLocation, Dummy);
 }
 
 
diff --git a/Engine/Source/Runtime/Landscape/Private/LandscapeSplines.cpp b/Engine/Source/Runtime/Landscape/Private/LandscapeSplines.cpp
--- a/Engine/Source/Runtime/Landscape/Private/LandscapeSplines.cpp
+++ b/Engine/Source/Runtime/Landscape/Private/LandscapeSplines.cpp
@@ -3331,7 +3331,7 @@ FName ULandscapeSplineSegment::GetCollisionProfileName() const
 void ULandscapeSplineSegment::FindNearest( const FVector& InLocation, float& t, FVector& OutLocation, FVector& OutTangent )
 {
 	float TempOutDistanceSq;
-	t = SplineInfo.InaccurateFindNearest(InLocation, TempOutDistanceSq);
+	t = SplineInfo.FindNearest(InLocation, TempOutDistanceSq);
 	OutLocation = SplineInfo.Eval(t, FVector::ZeroVector);
 	OutTangent = SplineInfo.EvalDerivative(t, FVector::ZeroVector);
 }
