Fix Horde Issue 269744

diff --git a/Engine/Source/Runtime/Core/Public/Math/InterpCurve.h b/Engine/Source/Runtime/Core/Public/Math/InterpCurve.h
--- a/Engine/Source/Runtime/Core/Public/Math/InterpCurve.h
+++ b/Engine/Source/Runtime/Core/Public/Math/InterpCurve.h
@@ -621,7 +621,7 @@ float FInterpCurve<T>::FindNearestOnSegment(const T& PointInSpace, int32 PtIdx,
 
 		const TArray<FInterpCurvePoint<T>>& PointsT = Points;
 
-		TFunction<bool(float&, const float)> BreakIfConverged = [&LastValue, &MinValue, &MaxValue, &NextValue, &StartValuesT](float& Value, const float Tolerance) -> bool
+		TFunction<bool(float&, const float)> BreakIfConverged = [&LastValue, &MinValue, &MaxValue, &NextValue, &StartValuesT, NewtonIterationCount](float& Value, const float Tolerance) -> bool
 		{
 			if (FMath::IsNearlyEqual(LastValue, Value, Tolerance))
 			{
