TDeque test - potential memory issue when AGGRESSIVE_MEMORY_SAVING=1

diff --git a/Engine/Source/Runtime/Core/Tests/Containers/DequeTest.cpp b/Engine/Source/Runtime/Core/Tests/Containers/DequeTest.cpp
--- a/Engine/Source/Runtime/Core/Tests/Containers/DequeTest.cpp
+++ b/Engine/Source/Runtime/Core/Tests/Containers/DequeTest.cpp
@@ -736,7 +736,7 @@ TEST_CASE_NAMED(FDequeMoveSimpleTest, "System::Core::Containers::TDeque::Move si
 		CHECK(Deque.IsEmpty());
 		int32Token Result;
 		int CheckValue = 0;
-		while (Deque.TryPopFirst(Result))
+		while (DestQueue.TryPopFirst(Result))
 		{
 			CHECK(*Result == CheckValue++);
 		}
@@ -749,10 +749,10 @@ TEST_CASE_NAMED(FDequeMoveWithHEadTailWrapTest, "System::Core::Containers::TDequ
 	int32Token::Reset();
 	{
 		TDeque<int32Token> Deque;
-		Deque.Reserve(Test::DefaultCapacity);
-		CHECK(Deque.Max() == Test::DefaultCapacity);
 		for (int32 Iteration = 0; Iteration < Test::DefaultCapacity; ++Iteration)
 		{
+			Deque.Reserve(Test::DefaultCapacity);
+			CHECK(Deque.Max() == Test::DefaultCapacity);
 			CHECK(Test::EmplaceLastPopFirst(Deque, 1));  // Rotates head and tail
 			for (int32 i = 0; i < Test::DefaultCapacity; ++i)
 			{
@@ -763,7 +763,7 @@ TEST_CASE_NAMED(FDequeMoveWithHEadTailWrapTest, "System::Core::Containers::TDequ
 			CHECK(Deque.IsEmpty());
 			int32Token Result;
 			int32 CheckValue = 0;
-			while (Deque.TryPopFirst(Result))
+			while (DestQueue.TryPopFirst(Result))
 			{
 				CHECK(*Result == CheckValue++);
 			}
@@ -777,12 +777,12 @@ TEST_CASE_NAMED(FDequeMoveVariableSizeWithHeadTailWrapTest, "System::Core::Conta
 	int32Token::Reset();
 	{
 		TDeque<int32Token> Deque;
-		Deque.Reserve(Test::DefaultCapacity);
-		CHECK(Deque.Max() == Test::DefaultCapacity);
 		for (int32 Iteration = 0; Iteration < Test::DefaultCapacity; ++Iteration)
 		{
 			for (int32 Size = 1; Size <= Test::DefaultCapacity; ++Size)
 			{
+				Deque.Reserve(Test::DefaultCapacity);
+				CHECK(Deque.Max() == Test::DefaultCapacity);
 				CHECK(Test::EmplaceLastPopFirst(Deque, 1));  // Rotates head and tail
 				const int32 SeedValue = FMath::RandRange(1, 999);
 				for (int32 i = 0; i < Size; ++i)
@@ -794,7 +794,7 @@ TEST_CASE_NAMED(FDequeMoveVariableSizeWithHeadTailWrapTest, "System::Core::Conta
 				CHECK(Deque.IsEmpty());
 				int32Token Result;
 				int32 CheckValue = SeedValue;
-				while (Deque.TryPopFirst(Result))
+				while (DestQueue.TryPopFirst(Result))
 				{
 					CHECK(*Result == CheckValue++);
 				}
