Remove CORE_API from IVirtualizationSystem and selectively apply it to the methods that need it.

#rnx

diff --git a/Engine/Source/Runtime/Core/Public/Virtualization/VirtualizationSystem.h b/Engine/Source/Runtime/Core/Public/Virtualization/VirtualizationSystem.h
--- a/Engine/Source/Runtime/Core/Public/Virtualization/VirtualizationSystem.h
+++ b/Engine/Source/Runtime/Core/Public/Virtualization/VirtualizationSystem.h
@@ -525,7 +525,7 @@ CORE_API void Shutdown();
  *							in failure and IsEnabled will always return false.
  * SystemName=Default	-	This will cause the default Epic implementation to be used @see VirtualizationManager
  */
-class CORE_API IVirtualizationSystem
+class IVirtualizationSystem
 {
 public:
 	IVirtualizationSystem() = default;
@@ -547,13 +547,13 @@ public:
 	virtual bool Initialize(const FInitParams& InitParams) = 0;
 
 	/** Returns true if a virtualization system has been initialized and false if not */
-	static bool IsInitialized();
+	CORE_API static bool IsInitialized();
 
 	/** 
 	 * Gain access to the current virtualization system active for the project. If the system has not yet been 
 	 * initialized then calling this method will initialize it.
 	 */
-	static IVirtualizationSystem& Get();
+	CORE_API static IVirtualizationSystem& Get();
 
 	/** Poll to see if content virtualization is enabled or not. */
 	virtual bool IsEnabled() const = 0;
