Skip ColorConversionHeavy
#rnx

diff --git a/Engine/Source/Runtime/Core/Tests/Math/UnrealMathTest.cpp b/Engine/Source/Runtime/Core/Tests/Math/UnrealMathTest.cpp
--- a/Engine/Source/Runtime/Core/Tests/Math/UnrealMathTest.cpp
+++ b/Engine/Source/Runtime/Core/Tests/Math/UnrealMathTest.cpp
@@ -4697,7 +4697,7 @@ namespace UnrealMathTestInternal {
 
 }
 
-TEST_CASE_NAMED(FColorConversionHeavyTest, "System::Core::Math::ColorConversionHeavy", "[ApplicationContextMask][EngineFilter]")
+TEST_CASE_NAMED(FColorConversionHeavyTest, "System::Core::Math::ColorConversionHeavy", "[.][ApplicationContextMask][EngineFilter]")
 {
 	// WARNING: This test runs for a good while (at time of writing, ~90s when using
 	// a single core, proportionately less on many-core machines) with no user
