Commented-out deprecatation of TScriptDelegate<FWeakObjectPtr> - will actually deprecate when it verified clean in other streams.

diff --git a/Engine/Source/Runtime/Core/Public/Delegates/Delegate.h b/Engine/Source/Runtime/Core/Public/Delegates/Delegate.h
--- a/Engine/Source/Runtime/Core/Public/Delegates/Delegate.h
+++ b/Engine/Source/Runtime/Core/Public/Delegates/Delegate.h
@@ -233,7 +233,7 @@
 
 /** Declare user's dynamic delegate, with wrapper proxy method for executing the delegate */
 #define FUNC_DECLARE_DYNAMIC_DELEGATE( DynamicDelegateClassName, ExecFunction, FuncParamList, FuncParamPassThru, ... ) \
-	class DynamicDelegateClassName : public TBaseDynamicDelegate<FWeakObjectPtr, __VA_ARGS__> \
+	class DynamicDelegateClassName : public TBaseDynamicDelegate<UE::Core::Private::TScriptDelegateDefault, __VA_ARGS__> \
 	{ \
 	public: \
 		/** Default constructor */ \
@@ -243,7 +243,7 @@
 		\
 		/** Construction from an FScriptDelegate must be explicit.  This is really only used by UObject system internals. */ \
 		explicit DynamicDelegateClassName( const TScriptDelegate<>& InScriptDelegate ) \
-			: TBaseDynamicDelegate<FWeakObjectPtr, __VA_ARGS__>( InScriptDelegate ) \
+			: TBaseDynamicDelegate<UE::Core::Private::TScriptDelegateDefault, __VA_ARGS__>( InScriptDelegate ) \
 		{ \
 		} \
 		\
@@ -268,7 +268,7 @@
 
 /** Declare user's dynamic delegate with return value, with wrapper proxy method for executing the delegate */
 #define FUNC_DECLARE_DYNAMIC_DELEGATE_RETVAL(DynamicDelegateRetValClassName, ExecFunction, RetValType, FuncParamList, FuncParamPassThru, ...) \
-	class DynamicDelegateRetValClassName : public TBaseDynamicDelegate<FWeakObjectPtr, __VA_ARGS__> \
+	class DynamicDelegateRetValClassName : public TBaseDynamicDelegate<UE::Core::Private::TScriptDelegateDefault, __VA_ARGS__> \
 	{ \
 	public: \
 		/** Default constructor */ \
@@ -278,7 +278,7 @@
 		\
 		/** Construction from an FScriptDelegate must be explicit.  This is really only used by UObject system internals. */ \
 		explicit DynamicDelegateRetValClassName( const TScriptDelegate<>& InScriptDelegate ) \
-			: TBaseDynamicDelegate<FWeakObjectPtr, __VA_ARGS__>( InScriptDelegate ) \
+			: TBaseDynamicDelegate<UE::Core::Private::TScriptDelegateDefault, __VA_ARGS__>( InScriptDelegate ) \
 		{ \
 		} \
 		\
@@ -294,7 +294,7 @@
 
 /** Declare user's dynamic multi-cast delegate, with wrapper proxy method for executing the delegate */
 #define FUNC_DECLARE_DYNAMIC_MULTICAST_DELEGATE(DynamicMulticastDelegateClassName, ExecFunction, FuncParamList, FuncParamPassThru, ...) \
-class DynamicMulticastDelegateClassName : public TBaseDynamicMulticastDelegate<FWeakObjectPtr, __VA_ARGS__> \
+class DynamicMulticastDelegateClassName : public TBaseDynamicMulticastDelegate<UE::Core::Private::TScriptDelegateDefault, __VA_ARGS__> \
 	{ \
 	public: \
 		/** Default constructor */ \
@@ -304,7 +304,7 @@ class DynamicMulticastDelegateClassName : public TBaseDynamicMulticastDelegate<F
 		\
 		/** Construction from an FMulticastScriptDelegate must be explicit.  This is really only used by UObject system internals. */ \
 		explicit DynamicMulticastDelegateClassName( const TMulticastScriptDelegate<>& InMulticastScriptDelegate ) \
-			: TBaseDynamicMulticastDelegate<FWeakObjectPtr, __VA_ARGS__>( InMulticastScriptDelegate ) \
+			: TBaseDynamicMulticastDelegate<UE::Core::Private::TScriptDelegateDefault, __VA_ARGS__>( InMulticastScriptDelegate ) \
 		{ \
 		} \
 		\
diff --git a/Engine/Source/Runtime/Core/Public/Delegates/DelegateSignatureImpl.inl b/Engine/Source/Runtime/Core/Public/Delegates/DelegateSignatureImpl.inl
--- a/Engine/Source/Runtime/Core/Public/Delegates/DelegateSignatureImpl.inl
+++ b/Engine/Source/Runtime/Core/Public/Delegates/DelegateSignatureImpl.inl
@@ -998,8 +998,8 @@ using TTSMulticastDelegate = TMulticastDelegate<DelegateSignature, FDefaultTSDel
  * macros to create the actual delegate type, templated to the function signature the delegate is compatible with.
  * Then, you can create an instance of that class when you want to assign functions to the delegate.
  */
-template <typename TWeakPtr, typename RetValType, typename... ParamTypes>
-class TBaseDynamicDelegate : public TScriptDelegate<TWeakPtr>
+template <typename ScriptDelegateDummyType, typename RetValType, typename... ParamTypes>
+class TBaseDynamicDelegate : public TScriptDelegate<ScriptDelegateDummyType>
 {
 public:
 	/**
@@ -1012,8 +1012,8 @@ public:
 	 *
 	 * @param	InScriptDelegate	The delegate to construct from by copying
 	 */
-	explicit TBaseDynamicDelegate( const TScriptDelegate<TWeakPtr>& InScriptDelegate )
-		: TScriptDelegate<TWeakPtr>( InScriptDelegate )
+	explicit TBaseDynamicDelegate( const TScriptDelegate<ScriptDelegateDummyType>& InScriptDelegate )
+		: TScriptDelegate<ScriptDelegateDummyType>( InScriptDelegate )
 	{ }
 
 	/**
@@ -1075,12 +1075,12 @@ public:
  * signature the delegate is compatible with.   Then, you can create an instance of that class when you
  * want to assign functions to the delegate.
  */
-template <typename TWeakPtr, typename RetValType, typename... ParamTypes>
-class TBaseDynamicMulticastDelegate : public TMulticastScriptDelegate<TWeakPtr>
+template <typename ScriptDelegateDummyType, typename RetValType, typename... ParamTypes>
+class TBaseDynamicMulticastDelegate : public TMulticastScriptDelegate<ScriptDelegateDummyType>
 {
 public:
 	/** The actual single-cast delegate class for this multi-cast delegate */
-	typedef TBaseDynamicDelegate<FWeakObjectPtr, RetValType, ParamTypes...> FDelegate;
+	typedef TBaseDynamicDelegate<ScriptDelegateDummyType, RetValType, ParamTypes...> FDelegate;
 
 	/**
 	 * Default constructor
@@ -1092,8 +1092,8 @@ public:
 	 *
 	 * @param	InScriptDelegate	The delegate to construct from by copying
 	 */
-	explicit TBaseDynamicMulticastDelegate( const TMulticastScriptDelegate<TWeakPtr>& InMulticastScriptDelegate )
-		: TMulticastScriptDelegate<TWeakPtr>( InMulticastScriptDelegate )
+	explicit TBaseDynamicMulticastDelegate( const TMulticastScriptDelegate<ScriptDelegateDummyType>& InMulticastScriptDelegate )
+		: TMulticastScriptDelegate<ScriptDelegateDummyType>( InMulticastScriptDelegate )
 	{ }
 
 	/**
diff --git a/Engine/Source/Runtime/Core/Public/UObject/ScriptDelegates.h b/Engine/Source/Runtime/Core/Public/UObject/ScriptDelegates.h
--- a/Engine/Source/Runtime/Core/Public/UObject/ScriptDelegates.h
+++ b/Engine/Source/Runtime/Core/Public/UObject/ScriptDelegates.h
@@ -12,15 +12,32 @@
 #include "UObject/NameTypes.h"
 #include "Templates/SharedPointer.h"
 
+namespace UE::Core::Private
+{
+	struct TScriptDelegateDefault;
+
+	template <typename ScriptDelegateParameterType>
+	struct TScriptDelegateTraits
+	{
+		// Although templated, WeakPtrType is not intended to be anything other than FWeakObjectPtr,
+		// and is only a template for module organization reasons.
+		using WeakPtrType = FWeakObjectPtr;
+	};
+
+	template <>
+	struct /*UE_DEPRECATED(5.3, "TScriptDelegate<FWeakObjectPtr> and TMulticastScriptDelegate<FWeakObjectPtr> have been deprecated, please use FScriptDelegate or FMulticastScriptDelegate respectively.")*/ TScriptDelegateTraits<FWeakObjectPtr>
+	{
+		using WeakPtrType = FWeakObjectPtr;
+	};
+}
+
 /**
  * Script delegate base class.
  */
-template <typename TWeakPtr /*= FWeakObjectPtr*/>
+template <typename Dummy>
 class TScriptDelegate
 {
-	// Although templated, the parameter is not intended to be anything other than the default,
-	// and is only a template for module organization reasons.
-	static_assert(std::is_same_v<TWeakPtr, FWeakObjectPtr>, "TWeakPtr should not be overridden");
+	using WeakPtrType = typename UE::Core::Private::TScriptDelegateTraits<Dummy>::WeakPtrType;
 
 public:
 
@@ -269,17 +286,17 @@ public:
 			// by access protection scope
 			UE_DELEGATES_MT_SCOPED_WRITE_ACCESS(AccessDetector);
 
-			checkf(Object.IsValid() != false, TEXT("ProcessDelegate() called with no object bound to delegate!"));
-			checkf(FunctionName != NAME_None, TEXT("ProcessDelegate() called with no function name set!"));
+		checkf( Object.IsValid() != false, TEXT( "ProcessDelegate() called with no object bound to delegate!" ) );
+		checkf( FunctionName != NAME_None, TEXT( "ProcessDelegate() called with no function name set!" ) );
 
-			// Object was pending kill, so we cannot execute the delegate.  Note that it's important to assert
-			// here and not simply continue execution, as memory may be left uninitialized if the delegate is
-			// not able to execute, resulting in much harder-to-detect code errors.  Users should always make
-			// sure IsBound() returns true before calling ProcessDelegate()!
+		// Object was pending kill, so we cannot execute the delegate.  Note that it's important to assert
+		// here and not simply continue execution, as memory may be left uninitialized if the delegate is
+		// not able to execute, resulting in much harder-to-detect code errors.  Users should always make
+		// sure IsBound() returns true before calling ProcessDelegate()!
 			ObjectPtr = static_cast<UObjectTemplate*>(Object.Get());	// Down-cast
-			checkSlow(IsValid(ObjectPtr));
+		checkSlow( IsValid(ObjectPtr) );
 
-			// Object *must* implement the specified function
+		// Object *must* implement the specified function
 			Function = ObjectPtr->FindFunctionChecked(FunctionName);
 		}
 
@@ -297,7 +314,7 @@ public:
 protected:
 
 	/** The object bound to this delegate, or nullptr if no object is bound */
-	TWeakPtr Object;
+	WeakPtrType Object;
 
 	/** Name of the function to call on the bound object */
 	FName FunctionName;
@@ -306,18 +323,26 @@ protected:
 
 	// 
 	friend class FCallDelegateHelper;
+
+	friend struct TIsZeroConstructType<TScriptDelegate>;
 };
 
 
-template<typename TWeakPtr> struct TIsZeroConstructType<TScriptDelegate<TWeakPtr> > { enum { Value = TIsZeroConstructType<TWeakPtr>::Value }; };
+template <typename Dummy>
+struct TIsZeroConstructType<TScriptDelegate<Dummy>>
+{
+	static constexpr bool Value = TIsZeroConstructType<typename TScriptDelegate<Dummy>::WeakPtrType>::Value;
+};
 
 
 /**
  * Script multi-cast delegate base class
  */
-template <typename TWeakPtr/* = FWeakObjectPtr*/>
+template <typename Dummy>
 class TMulticastScriptDelegate
 {
+	using UnicastDelegateType = TScriptDelegate<Dummy>;
+
 public:
 
 	/**
@@ -345,7 +370,7 @@ public:
 	 * @param	InDelegate	Delegate to check
 	 * @return	True if the delegate is already in the list.
 	 */
-	bool Contains( const TScriptDelegate<TWeakPtr>& InDelegate ) const
+	bool Contains( const UnicastDelegateType& InDelegate ) const
 	{
 		UE_DELEGATES_MT_SCOPED_READ_ACCESS(AccessDetector);
 
@@ -363,7 +388,7 @@ public:
 	{
 		UE_DELEGATES_MT_SCOPED_READ_ACCESS(AccessDetector);
 
-		return InvocationList.ContainsByPredicate( [=]( const TScriptDelegate<TWeakPtr>& Delegate ){
+		return InvocationList.ContainsByPredicate( [=]( const UnicastDelegateType& Delegate ){
 			return Delegate.GetFunctionName() == InFunctionName && Delegate.IsBoundToObjectEvenIfUnreachable(InObject);
 		} );
 	}
@@ -373,7 +398,7 @@ public:
 	 *
 	 * @param	InDelegate	Delegate to add
 	 */
-	void Add( const TScriptDelegate<TWeakPtr>& InDelegate )
+	void Add( const UnicastDelegateType& InDelegate )
 	{
 		UE_DELEGATES_MT_SCOPED_WRITE_ACCESS(AccessDetector);
 
@@ -390,7 +415,7 @@ public:
 	 *
 	 * @param	InDelegate	Delegate to add
 	 */
-	void AddUnique( const TScriptDelegate<TWeakPtr>& InDelegate )
+	void AddUnique( const UnicastDelegateType& InDelegate )
 	{
 		UE_DELEGATES_MT_SCOPED_WRITE_ACCESS(AccessDetector);
 
@@ -407,7 +432,7 @@ public:
 	 *
 	 * @param	InDelegate	Delegate to remove
 	 */
-	void Remove( const TScriptDelegate<TWeakPtr>& InDelegate )
+	void Remove( const UnicastDelegateType& InDelegate )
 	{
 		UE_DELEGATES_MT_SCOPED_WRITE_ACCESS(AccessDetector);
 
@@ -450,7 +475,7 @@ public:
 
 		for (int32 BindingIndex = InvocationList.Num() - 1; BindingIndex >= 0; --BindingIndex)
 		{
-			const TScriptDelegate<TWeakPtr>& Binding = InvocationList[BindingIndex];
+			const UnicastDelegateType& Binding = InvocationList[BindingIndex];
 
 			if (Binding.IsBoundToObject(Object) || Binding.IsCompactable())
 			{
@@ -499,7 +524,7 @@ public:
 	}
 
 	/** Multi-cast delegate serialization */
-	friend FArchive& operator<<( FArchive& Ar, TMulticastScriptDelegate<TWeakPtr>& D )
+	friend FArchive& operator<<( FArchive& Ar, TMulticastScriptDelegate& D )
 	{
 		UE_DELEGATES_MT_SCOPED_WRITE_ACCESS(D.AccessDetector);
 
@@ -520,7 +545,7 @@ public:
 		return Ar;
 	}
 
-	friend void operator<<(FStructuredArchive::FSlot Slot, TMulticastScriptDelegate<TWeakPtr>& D)
+	friend void operator<<(FStructuredArchive::FSlot Slot, TMulticastScriptDelegate& D)
 	{
 		UE_DELEGATES_MT_SCOPED_WRITE_ACCESS(D.AccessDetector);
 
@@ -556,7 +581,7 @@ public:
 		if( InvocationList.Num() > 0 )
 		{
 			// Create a copy of the invocation list, just in case the list is modified by one of the callbacks during the broadcast
-			typedef TArray< TScriptDelegate<TWeakPtr>, TInlineAllocator< 4 > > FInlineInvocationList;
+			typedef TArray< UnicastDelegateType, TInlineAllocator< 4 > > FInlineInvocationList;
 			FInlineInvocationList InvocationListCopy = FInlineInvocationList(InvocationList);
 	
 			// Invoke each bound function
@@ -633,7 +658,7 @@ protected:
 	 *
 	 * @param	InDelegate	Delegate to add
 	*/
-	void AddInternal( const TScriptDelegate<TWeakPtr>& InDelegate )
+	void AddInternal( const UnicastDelegateType& InDelegate )
 	{
 #if !(UE_BUILD_SHIPPING || UE_BUILD_TEST)
 		// Verify same function isn't already bound
@@ -652,7 +677,7 @@ protected:
 	 *
 	 * @param	InDelegate	Delegate to add
 	 */
-	void AddUniqueInternal( const TScriptDelegate<TWeakPtr>& InDelegate )
+	void AddUniqueInternal( const UnicastDelegateType& InDelegate )
 	{
 		// Add the item to the invocation list only if it is unique
 		InvocationList.AddUnique( InDelegate );
@@ -664,7 +689,7 @@ protected:
 	 *
 	 * @param	InDelegate	Delegate to remove
 	*/
-	void RemoveInternal( const TScriptDelegate<TWeakPtr>& InDelegate ) const
+	void RemoveInternal( const UnicastDelegateType& InDelegate ) const
 	{
 		InvocationList.RemoveSingleSwap(InDelegate);
 	}
@@ -678,7 +703,7 @@ protected:
 	*/
 	void RemoveInternal( const UObject* InObject, FName InFunctionName ) const
 	{
-		int32 FoundDelegate = InvocationList.IndexOfByPredicate([=](const TScriptDelegate<TWeakPtr>& Delegate) {
+		int32 FoundDelegate = InvocationList.IndexOfByPredicate([=](const UnicastDelegateType& Delegate) {
 			return Delegate.GetFunctionName() == InFunctionName && Delegate.IsBoundToObjectEvenIfUnreachable(InObject);
 		});
 
@@ -691,13 +716,13 @@ protected:
 	/** Cleans up any delegates in our invocation list that have expired (performance is O(N)) */
 	void CompactInvocationList() const
 	{
-		InvocationList.RemoveAllSwap([](const TScriptDelegate<TWeakPtr>& Delegate){
+		InvocationList.RemoveAllSwap([](const UnicastDelegateType& Delegate){
 			return Delegate.IsCompactable();
 		});
 	}
 
 public:
-	typedef TArray< TScriptDelegate<TWeakPtr> > FInvocationList;
+	typedef TArray<UnicastDelegateType> FInvocationList;
 
 protected:
 
@@ -714,8 +739,12 @@ protected:
 	// 
 	friend class FCallDelegateHelper;
 
-	friend struct TIsZeroConstructType<TMulticastScriptDelegate<TWeakPtr> >;
+	friend struct TIsZeroConstructType<TMulticastScriptDelegate>;
 };
 
 
-template<typename TWeakPtr> struct TIsZeroConstructType<TMulticastScriptDelegate<TWeakPtr> > { enum { Value = TIsZeroConstructType<typename TMulticastScriptDelegate<TWeakPtr>::FInvocationList>::Value }; };
+template <typename TWeakPtr>
+struct TIsZeroConstructType<TMulticastScriptDelegate<TWeakPtr>>
+{
+	static constexpr bool Value = TIsZeroConstructType<typename TMulticastScriptDelegate<TWeakPtr>::FInvocationList>::Value;
+};
diff --git a/Engine/Source/Runtime/CoreUObject/Public/UObject/ScriptDelegateFwd.h b/Engine/Source/Runtime/CoreUObject/Public/UObject/ScriptDelegateFwd.h
--- a/Engine/Source/Runtime/CoreUObject/Public/UObject/ScriptDelegateFwd.h
+++ b/Engine/Source/Runtime/CoreUObject/Public/UObject/ScriptDelegateFwd.h
@@ -8,8 +8,13 @@
 
 #include "UObject/WeakObjectPtrFwd.h"
 
-template <typename TWeakPtr = FWeakObjectPtr> class TScriptDelegate;
-template <typename TWeakPtr = FWeakObjectPtr> class TMulticastScriptDelegate;
+namespace UE::Core::Private
+{
+	struct TScriptDelegateDefault;
+}
+
+template <typename Dummy = UE::Core::Private::TScriptDelegateDefault> class TScriptDelegate;
+template <typename Dummy = UE::Core::Private::TScriptDelegateDefault> class TMulticastScriptDelegate;
 
 // Typedef script delegates for convenience.
 typedef TScriptDelegate<> FScriptDelegate;
