Add GetInt64 to GConfig.  This already existed for FConfigFile.

#localization none
#tests verfied via debugger that correct sizes are now returned

diff --git a/Engine/Source/Runtime/Core/Private/Misc/ConfigCacheIni.cpp b/Engine/Source/Runtime/Core/Private/Misc/ConfigCacheIni.cpp
--- a/Engine/Source/Runtime/Core/Private/Misc/ConfigCacheIni.cpp
+++ b/Engine/Source/Runtime/Core/Private/Misc/ConfigCacheIni.cpp
@@ -2916,9 +2916,25 @@ bool FConfigCacheIni::GetInt
 	if( GetString( Section, Key, Text, Filename ) )
 	{
 		Value = FCString::Atoi(*Text);
-		return 1;
+		return true;
 	}
-	return 0;
+	return false;
+}
+bool FConfigCacheIni::GetInt64
+(
+	const TCHAR* Section,
+	const TCHAR* Key,
+	int64& Value,
+	const FString& Filename
+)
+{
+	FString Text;
+	if (GetString(Section, Key, Text, Filename))
+	{
+		Value = FCString::Atoi64(*Text);
+		return true;
+	}
+	return false;
 }
 bool FConfigCacheIni::GetFloat
 (
@@ -2932,9 +2948,9 @@ bool FConfigCacheIni::GetFloat
 	if( GetString( Section, Key, Text, Filename ) )
 	{
 		Value = FCString::Atof(*Text);
-		return 1;
+		return true;
 	}
-	return 0;
+	return false;
 }
 bool FConfigCacheIni::GetDouble
 	(
@@ -2948,12 +2964,10 @@ bool FConfigCacheIni::GetDouble
 	if( GetString( Section, Key, Text, Filename ) )
 	{
 		Value = FCString::Atod(*Text);
-		return 1;
+		return true;
 	}
-	return 0;
+	return false;
 }
-
-
 bool FConfigCacheIni::GetBool
 (
 	const TCHAR*		Section,
@@ -2966,9 +2980,9 @@ bool FConfigCacheIni::GetBool
 	if( GetString( Section, Key, Text, Filename ) )
 	{
 		Value = FCString::ToBool(*Text);
-		return 1;
+		return true;
 	}
-	return 0;
+	return false;
 }
 int32 FConfigCacheIni::GetArray
 (
diff --git a/Engine/Source/Runtime/Core/Public/Misc/ConfigCacheIni.h b/Engine/Source/Runtime/Core/Public/Misc/ConfigCacheIni.h
--- a/Engine/Source/Runtime/Core/Public/Misc/ConfigCacheIni.h
+++ b/Engine/Source/Runtime/Core/Public/Misc/ConfigCacheIni.h
@@ -857,6 +857,13 @@ public:
 		int32&				Value,
 		const FString&	Filename
 	);
+	bool GetInt64
+	(
+		const TCHAR* Section,
+		const TCHAR* Key,
+		int64& Value,
+		const FString& Filename
+	);
 	bool GetFloat
 	(
 		const TCHAR*		Section,
