Skip local validation for materials compiled for preview shader platforms. The main editor shader platform will still fail validation so it's not possible to create materials that don't pass validation.

diff --git a/Engine/Source/Runtime/Core/Private/Misc/DataDrivenPlatformInfoRegistry.cpp b/Engine/Source/Runtime/Core/Private/Misc/DataDrivenPlatformInfoRegistry.cpp
--- a/Engine/Source/Runtime/Core/Private/Misc/DataDrivenPlatformInfoRegistry.cpp
+++ b/Engine/Source/Runtime/Core/Private/Misc/DataDrivenPlatformInfoRegistry.cpp
@@ -246,51 +246,54 @@ static FString GetSectionString(const FConfigSection& Section, FName Key)
 #if DDPI_HAS_EXTENDED_PLATFORMINFO_DATA
 static void ParsePreviewPlatforms(const FConfigFile& IniFile)
 {
-	// walk over the file looking for PreviewPlatform sections
-	for (auto Section : IniFile)
+	if (!FParse::Param(FCommandLine::Get(), TEXT("NoPreviewPlatforms")))
 	{
-		if (Section.Key.StartsWith(TEXT("PreviewPlatform ")))
+		// walk over the file looking for PreviewPlatform sections
+		for (auto Section : IniFile)
 		{
-			const FString& SectionName = Section.Key;
-			FName PreviewPlatformName = *(SectionName.Mid(16) + TEXT("_Preview"));
-
-			// Early-out if enabled cvar is specified and not set
-			TArray<FString> Tokens;
-			GetSectionString(Section.Value, FName("EnabledCVar")).ParseIntoArray(Tokens, TEXT(":"));
-			if (Tokens.Num() == 5)
+			if (Section.Key.StartsWith(TEXT("PreviewPlatform ")))
 			{
-				// now load a local version of the ini hierarchy
-				FConfigFile LocalIni;
-				FConfigCacheIni::LoadLocalIniFile(LocalIni, *Tokens[1], true, *Tokens[2]);
-
-				// and get the enabled cvar's value
-				bool bEnabled = false;
-				LocalIni.GetBool(*Tokens[3], *Tokens[4], bEnabled);
-				if (!bEnabled)
+				const FString& SectionName = Section.Key;
+				FName PreviewPlatformName = *(SectionName.Mid(16) + TEXT("_Preview"));
+
+				// Early-out if enabled cvar is specified and not set
+				TArray<FString> Tokens;
+				GetSectionString(Section.Value, FName("EnabledCVar")).ParseIntoArray(Tokens, TEXT(":"));
+				if (Tokens.Num() == 5)
 				{
-					continue;
+					// now load a local version of the ini hierarchy
+					FConfigFile LocalIni;
+					FConfigCacheIni::LoadLocalIniFile(LocalIni, *Tokens[1], true, *Tokens[2]);
+
+					// and get the enabled cvar's value
+					bool bEnabled = false;
+					LocalIni.GetBool(*Tokens[3], *Tokens[4], bEnabled);
+					if (!bEnabled)
+					{
+						continue;
+					}
 				}
-			}
 
-			FName PlatformName = *GetSectionString(Section.Value, FName("PlatformName"));
-			checkf(PlatformName != NAME_None, TEXT("DataDrivenPlatformInfo section [%s] must specify a PlatformName"), *SectionName);
-
-			FPreviewPlatformMenuItem Item;
-			Item.PlatformName = PlatformName;
-			Item.PreviewShaderPlatformName = PreviewPlatformName;
-			Item.ShaderFormat = *GetSectionString(Section.Value, FName("ShaderFormat"));
-			checkf(Item.ShaderFormat != NAME_None, TEXT("DataDrivenPlatformInfo section [PreviewPlatform %s] must specify a ShaderFormat"), *SectionName);
-			Item.ActiveIconPath = GetSectionString(Section.Value, FName("ActiveIconPath"));
-			Item.ActiveIconName = *GetSectionString(Section.Value, FName("ActiveIconName"));
-			Item.InactiveIconPath = GetSectionString(Section.Value, FName("InactiveIconPath"));
-			Item.InactiveIconName = *GetSectionString(Section.Value, FName("InactiveIconName"));
-			Item.DeviceProfileName = *GetSectionString(Section.Value, FName("DeviceProfileName"));
-			Item.ShaderPlatformToPreview = *GetSectionString(Section.Value, FName("ShaderPlatform"));
-			checkf(Item.ShaderPlatformToPreview != NAME_None, TEXT("DataDrivenPlatformInfo section [PreviewPlatform %s] must specify a ShaderPlatform"), *SectionName);
-			FTextStringHelper::ReadFromBuffer(*GetSectionString(Section.Value, FName("FriendlyName")), Item.OptionalFriendlyNameOverride);
-			FTextStringHelper::ReadFromBuffer(*GetSectionString(Section.Value, FName("MenuTooltip")), Item.MenuTooltip);
-			FTextStringHelper::ReadFromBuffer(*GetSectionString(Section.Value, FName("IconText")), Item.IconText);
-			PreviewPlatformMenuItems.Add(Item);
+				FName PlatformName = *GetSectionString(Section.Value, FName("PlatformName"));
+				checkf(PlatformName != NAME_None, TEXT("DataDrivenPlatformInfo section [%s] must specify a PlatformName"), *SectionName);
+
+				FPreviewPlatformMenuItem Item;
+				Item.PlatformName = PlatformName;
+				Item.PreviewShaderPlatformName = PreviewPlatformName;
+				Item.ShaderFormat = *GetSectionString(Section.Value, FName("ShaderFormat"));
+				checkf(Item.ShaderFormat != NAME_None, TEXT("DataDrivenPlatformInfo section [PreviewPlatform %s] must specify a ShaderFormat"), *SectionName);
+				Item.ActiveIconPath = GetSectionString(Section.Value, FName("ActiveIconPath"));
+				Item.ActiveIconName = *GetSectionString(Section.Value, FName("ActiveIconName"));
+				Item.InactiveIconPath = GetSectionString(Section.Value, FName("InactiveIconPath"));
+				Item.InactiveIconName = *GetSectionString(Section.Value, FName("InactiveIconName"));
+				Item.DeviceProfileName = *GetSectionString(Section.Value, FName("DeviceProfileName"));
+				Item.ShaderPlatformToPreview = *GetSectionString(Section.Value, FName("ShaderPlatform"));
+				checkf(Item.ShaderPlatformToPreview != NAME_None, TEXT("DataDrivenPlatformInfo section [PreviewPlatform %s] must specify a ShaderPlatform"), *SectionName);
+				FTextStringHelper::ReadFromBuffer(*GetSectionString(Section.Value, FName("FriendlyName")), Item.OptionalFriendlyNameOverride);
+				FTextStringHelper::ReadFromBuffer(*GetSectionString(Section.Value, FName("MenuTooltip")), Item.MenuTooltip);
+				FTextStringHelper::ReadFromBuffer(*GetSectionString(Section.Value, FName("IconText")), Item.IconText);
+				PreviewPlatformMenuItems.Add(Item);
+			}
 		}
 	}
 }
diff --git a/Engine/Source/Runtime/Engine/Private/Materials/HLSLMaterialTranslator.cpp b/Engine/Source/Runtime/Engine/Private/Materials/HLSLMaterialTranslator.cpp
--- a/Engine/Source/Runtime/Engine/Private/Materials/HLSLMaterialTranslator.cpp
+++ b/Engine/Source/Runtime/Engine/Private/Materials/HLSLMaterialTranslator.cpp
@@ -11872,8 +11872,17 @@ static bool IsMaterialPathExemptFromRestrictiveMode(FStringBuilderBase& PathName
 	return false;
 }
 
-static bool IsMaterialExpressionExemptFromRestrictiveMode(const UMaterialExpression* Expression, FStringBuilderBase& PathName)
+static bool IsMaterialExpressionExemptFromRestrictiveMode(const UMaterialExpression* Expression, FStringBuilderBase& PathName, EShaderPlatform ShaderPlatform)
 {
+	static bool bNoPreviewPlatforms = FParse::Param(FCommandLine::Get(), TEXT("NoPreviewPlatforms"));
+	if (!bNoPreviewPlatforms && FDataDrivenShaderPlatformInfo::GetIsPreviewPlatform(ShaderPlatform))
+	{
+		// Preview shader platforms do not need to have their expressions validated.
+		// The validation for the main editor platform will prevent disallowed expressions
+		// from being used in the material editor.
+		return true;
+	}
+
 	if (Expression)
 	{
 		if (Expression->Material)
@@ -11893,7 +11902,7 @@ static bool IsMaterialExpressionExemptFromRestrictiveMode(const UMaterialExpress
 bool FHLSLMaterialTranslator::ValidateMaterialExpressionPermission(const UMaterialExpression* Expression)
 {
 	FNameBuilder MaterialPathName;
-	if (!IsExpressionClassPermitted(Expression->GetClass()) && !IsMaterialExpressionExemptFromRestrictiveMode(Expression, MaterialPathName))
+	if (!IsExpressionClassPermitted(Expression->GetClass()) && !IsMaterialExpressionExemptFromRestrictiveMode(Expression, MaterialPathName, Platform))
 	{
 		if (MaterialPathName.Len() > 0)
 		{
diff --git a/Engine/Source/Runtime/RHI/Private/DataDrivenShaderPlatformInfo.cpp b/Engine/Source/Runtime/RHI/Private/DataDrivenShaderPlatformInfo.cpp
--- a/Engine/Source/Runtime/RHI/Private/DataDrivenShaderPlatformInfo.cpp
+++ b/Engine/Source/Runtime/RHI/Private/DataDrivenShaderPlatformInfo.cpp
@@ -335,23 +335,26 @@ void FGenericDataDrivenShaderPlatformInfo::Initialize()
 				Infos[ShaderPlatform].bContainsValidPlatformInfo = true;
 
 #if WITH_EDITOR
-				for (const FPreviewPlatformMenuItem& Item : FDataDrivenPlatformInfoRegistry::GetAllPreviewPlatformMenuItems())
+				if (!FParse::Param(FCommandLine::Get(), TEXT("NoPreviewPlatforms")))
 				{
-					FName PreviewPlatformName = *(Infos[ShaderPlatform].Name).ToString();
-					if (Item.ShaderPlatformToPreview == PreviewPlatformName)
+					for (const FPreviewPlatformMenuItem& Item : FDataDrivenPlatformInfoRegistry::GetAllPreviewPlatformMenuItems())
 					{
-						EShaderPlatform PreviewShaderPlatform = EShaderPlatform(CustomShaderPlatform++);
-
-						ParseDataDrivenShaderInfo(Section.Value, Infos[PreviewShaderPlatform]);
-						if (!Item.OptionalFriendlyNameOverride.IsEmpty())
+						FName PreviewPlatformName = *(Infos[ShaderPlatform].Name).ToString();
+						if (Item.ShaderPlatformToPreview == PreviewPlatformName)
 						{
-							Infos[PreviewShaderPlatform].FriendlyName = Item.OptionalFriendlyNameOverride;
+							EShaderPlatform PreviewShaderPlatform = EShaderPlatform(CustomShaderPlatform++);
+
+							ParseDataDrivenShaderInfo(Section.Value, Infos[PreviewShaderPlatform]);
+							if (!Item.OptionalFriendlyNameOverride.IsEmpty())
+							{
+								Infos[PreviewShaderPlatform].FriendlyName = Item.OptionalFriendlyNameOverride;
+							}
+							Infos[PreviewShaderPlatform].Name = Item.PreviewShaderPlatformName;
+							PlatformNameToShaderPlatformMap.FindOrAdd(Infos[PreviewShaderPlatform].Name) = PreviewShaderPlatform;
+							Infos[PreviewShaderPlatform].PreviewShaderPlatformParent = ShaderPlatform;
+							Infos[PreviewShaderPlatform].bIsPreviewPlatform = true;
+							Infos[PreviewShaderPlatform].bContainsValidPlatformInfo = true;
 						}
-						Infos[PreviewShaderPlatform].Name = Item.PreviewShaderPlatformName;
-						PlatformNameToShaderPlatformMap.FindOrAdd(Infos[PreviewShaderPlatform].Name) = PreviewShaderPlatform;
-						Infos[PreviewShaderPlatform].PreviewShaderPlatformParent = ShaderPlatform;
-						Infos[PreviewShaderPlatform].bIsPreviewPlatform = true;
-						Infos[PreviewShaderPlatform].bContainsValidPlatformInfo = true;
 					}
 				}
 #endif
