Translated tests which needed in new macro implementation
Translated Core\Tests\Containers\StringViewTest.cpp
Translated all files under Core\Tests\Templates
Translated Core\Tests\Compression\CompressedBufferTest.cpp
Translated two files under Core\Tests\Serialization
Translate all files under Core\Tests\Memory

diff --git a/Engine/Source/Runtime/Core/Public/Misc/LowLevelTestAdapter.h b/Engine/Source/Runtime/Core/Public/Misc/LowLevelTestAdapter.h
--- a/Engine/Source/Runtime/Core/Public/Misc/LowLevelTestAdapter.h
+++ b/Engine/Source/Runtime/Core/Public/Misc/LowLevelTestAdapter.h
@@ -70,6 +70,8 @@
 #define INFO(What)
 #define CHECK(Expr) if (!(Expr)) { FAutomationTestFramework::Get().GetCurrentTest()->AddError(TEXT("Condition failed")); }
 #define CHECK_FALSE(Expr) if (Expr) { FAutomationTestFramework::Get().GetCurrentTest()->AddError(TEXT("Condition expected to return false but returned true")); }
+#define CHECKED_IF(Expr) if (Expr)
+#define CHECKED_ELSE(Expr) if (!(Expr))
 //-V:CHECK_MESSAGE:571,501
 #define CHECK_MESSAGE(Message, Expr) if (!(Expr)) { FAutomationTestFramework::Get().GetCurrentTest()->AddError(Message); }
 //-V:CHECK_FALSE_MESSAGE:571,501
@@ -77,7 +79,12 @@
 #define REQUIRE(Expr) if (!(Expr)) { FAutomationTestFramework::Get().GetCurrentTest()->AddError(TEXT("Required condition failed, interrupting test")); return; }
 //-V:REQUIRE_MESSAGE:571
 #define REQUIRE_MESSAGE(Message, Expr) if (!(Expr)) { FAutomationTestFramework::Get().GetCurrentTest()->AddError(Message); return; }
+#define STATIC_REQUIRE( ... ) static_assert(   __VA_ARGS__,  #__VA_ARGS__ );
+#define STATIC_CHECK( ... ) static_assert(   __VA_ARGS__,  #__VA_ARGS__ );
 
 #define SECTION(Text) AddInfo(TEXT(Text));
+#define FAIL_CHECK(Message) { FAutomationTestFramework::Get().GetCurrentTest()->AddError(Message); }
+
+#define CAPTURE(...) { FAutomationTestFramework::Get().GetCurrentTest()->AddInfo(FString::Printf(TEXT("%s := %d"), TEXT(#__VA_ARGS__), __VA_ARGS__)); }
 
 #endif // !WITH_LOW_LEVEL_TESTS
\ No newline at end of file
diff --git a/Engine/Source/Runtime/Core/Tests/Containers/StringViewTest.cpp b/Engine/Source/Runtime/Core/Tests/Containers/StringViewTest.cpp
--- a/Engine/Source/Runtime/Core/Tests/Containers/StringViewTest.cpp
+++ b/Engine/Source/Runtime/Core/Tests/Containers/StringViewTest.cpp
@@ -1,12 +1,12 @@
 // Copyright Epic Games, Inc. All Rights Reserved.
 
-#if WITH_LOW_LEVEL_TESTS
+#if WITH_TESTS
 
 #include "Containers/StringView.h"
 
 #include "Containers/UnrealString.h"
 #include "Misc/StringBuilder.h"
-#include "TestHarness.h"
+#include "Tests/TestHarnessAdapter.h"
 
 static_assert(std::is_same_v<typename FStringView::ElementType, TCHAR>, "FStringView must use TCHAR.");
 static_assert(std::is_same_v<typename FAnsiStringView::ElementType, ANSICHAR>, "FAnsiStringView must use ANSICHAR.");
@@ -34,7 +34,7 @@ int32 GetNum(const FTestType&) { return 3; }
 
 template <> struct TIsContiguousContainer<UE::String::Private::TestArgumentDependentLookup::FTestType> { static constexpr bool Value = true; };
 
-TEST_CASE("Core::String::StringView::Constructor", "[Core][String][Smoke]")
+TEST_CASE_NAMED(FStringViewConstructorTest, "System::Core::String::StringView::Constructor", "[Core][String][SmokeFilter]")
 {
 	SECTION("Default View")
 	{
@@ -145,7 +145,7 @@ TEST_CASE("Core::String::StringView::Constructor", "[Core][String][Smoke]")
 	}
 }
 
-TEST_CASE("Core::String::StringView::Iterators", "[Core][String][Smoke]")
+TEST_CASE_NAMED(FStringViewIteratorsTest, "System::Core::String::StringView::Iterators", "[Core][String][SmokeFilter]")
 {
 	SECTION("View Iterator")
 	{
@@ -170,7 +170,7 @@ TEST_CASE("Core::String::StringView::Iterators", "[Core][String][Smoke]")
 	}
 }
 
-TEST_CASE("Core::String::StringView::Equality", "[Core][String][Smoke]")
+TEST_CASE_NAMED(FStringViewEqualityTest, "System::Core::String::StringView::Equality", "[Core][String][SmokeFilter]")
 {
 	const ANSICHAR* AnsiStringLiteralSrc = "String To Test!";
 	const ANSICHAR* AnsiStringLiteralLower = "string to test!";
@@ -349,7 +349,7 @@ TEST_CASE("Core::String::StringView::Equality", "[Core][String][Smoke]")
 	}
 }
 
-TEST_CASE("Core::String::StringView::ComparisonCaseSensitive", "[Core][String][Smoke]")
+TEST_CASE_NAMED(FStringViewComparisonCaseSensitiveTest, "System::Core::String::StringView::ComparisonCaseSensitive", "[Core][String][SmokeFilter]")
 {
 	SECTION("Basic comparisons involving case")
 	{
@@ -437,7 +437,7 @@ TEST_CASE("Core::String::StringView::ComparisonCaseSensitive", "[Core][String][S
 	}
 }
 
-TEST_CASE("Core::String::StringView::ComparisonCaseInsensitive", "[Core][String][Smoke]")
+TEST_CASE_NAMED(FStringViewComparisonCaseInsensitiveTest, "System::Core::String::StringView::ComparisonCaseInsensitive", "[Core][String][SmokeFilter]")
 {
 	SECTION("Basic comparisons involving case")
 	{
@@ -501,7 +501,7 @@ TEST_CASE("Core::String::StringView::ComparisonCaseInsensitive", "[Core][String]
 	}
 }
 
-TEST_CASE("Core::String::StringView::ArrayAccessor", "[Core][String][Smoke]")
+TEST_CASE_NAMED(FStringViewArrayAccessorTest, "System::Core::String::StringView::ArrayAccessor", "[Core][String][SmokeFilter]")
 {
 	const TCHAR* SrcString = TEXT("String To Test");
 	FStringView View(SrcString);
@@ -512,7 +512,7 @@ TEST_CASE("Core::String::StringView::ArrayAccessor", "[Core][String][Smoke]")
 	}
 }
 
-TEST_CASE("Core::String::StringView::Modifiers", "[Core][String][Smoke]")
+TEST_CASE_NAMED(FStringViewModifiersTest, "System::Core::String::StringView::Modifiers", "[Core][String][SmokeFilter]")
 {
 	const TCHAR* FullText = TEXT("PrefixSuffix");
 	const TCHAR* Prefix = TEXT("Prefix");
@@ -537,7 +537,7 @@ TEST_CASE("Core::String::StringView::Modifiers", "[Core][String][Smoke]")
 	}
 }
 
-TEST_CASE("Core::String::StringView::StartsWith", "[Core][String][Smoke]")
+TEST_CASE_NAMED(FStringViewStartsWithTest, "System::Core::String::StringView::StartsWith", "[Core][String][SmokeFilter]")
 {
 	SECTION("Test an empty view")
 	{
@@ -571,7 +571,7 @@ TEST_CASE("Core::String::StringView::StartsWith", "[Core][String][Smoke]")
 	}
 }
 
-TEST_CASE("Core::String::StringView::EndsWith", "[Core][String][Smoke]")
+TEST_CASE_NAMED(FStringViewEndsWithTest, "System::Core::String::StringView::EndsWith", "[Core][String][SmokeFilter]")
 {
 	SECTION("Test an empty view")
 	{
@@ -605,7 +605,7 @@ TEST_CASE("Core::String::StringView::EndsWith", "[Core][String][Smoke]")
 	}
 }
 
-TEST_CASE("Core::String::StringView::SubStr", "[Core][String][Smoke]")
+TEST_CASE_NAMED(FStringViewStringSubStrTest, "System::Core::String::StringView::SubStr", "[Core][String][SmokeFilter]")
 {
 	{
 		FStringView EmptyView;
@@ -635,7 +635,7 @@ TEST_CASE("Core::String::StringView::SubStr", "[Core][String][Smoke]")
 	}
 }
 
-TEST_CASE("Core::String::StringView::Left", "[Core][String][Smoke]")
+TEST_CASE_NAMED(FStringViewLeftTest, "System::Core::String::StringView::Left", "[Core][String][SmokeFilter]")
 {
 	{
 		FStringView EmptyView;
@@ -657,7 +657,7 @@ TEST_CASE("Core::String::StringView::Left", "[Core][String][Smoke]")
 	}
 }
 
-TEST_CASE("Core::String::StringView::LeftChop", "[Core][String][Smoke]")
+TEST_CASE_NAMED(FStringViewLeftChopTest, "System::Core::String::StringView::LeftChop", "[Core][String][SmokeFilter]")
 {
 	{
 		FStringView EmptyView;
@@ -679,7 +679,7 @@ TEST_CASE("Core::String::StringView::LeftChop", "[Core][String][Smoke]")
 	}
 }
 
-TEST_CASE("Core::String::StringView::Right", "[Core][String][Smoke]")
+TEST_CASE_NAMED(FStringViewRightTest, "System::Core::String::StringView::Right", "[Core][String][SmokeFilter]")
 {
 	{
 		FStringView EmptyView;
@@ -701,7 +701,7 @@ TEST_CASE("Core::String::StringView::Right", "[Core][String][Smoke]")
 	}
 }
 
-TEST_CASE("Core::String::StringView::RightChop", "[Core][String][Smoke]")
+TEST_CASE_NAMED(FStringViewRightChopTest, "System::Core::String::StringView::RightChop", "[Core][String][SmokeFilter]")
 {
 	{
 		FStringView EmptyView;
@@ -723,7 +723,7 @@ TEST_CASE("Core::String::StringView::RightChop", "[Core][String][Smoke]")
 	}
 }
 
-TEST_CASE("Core::String::StringView::Mid", "[Core][String][Smoke]")
+TEST_CASE_NAMED(FStringViewMidTest, "System::Core::String::StringView::Mid", "[Core][String][SmokeFilter]")
 {
 	{
 		FStringView EmptyView;
@@ -755,7 +755,7 @@ TEST_CASE("Core::String::StringView::Mid", "[Core][String][Smoke]")
 	}
 }
 
-TEST_CASE("Core::String::StringView::TrimStartAndEnd", "[Core][String][Smoke]")
+TEST_CASE_NAMED(FStringViewTrimStartAndEndTest, "System::Core::String::StringView::TrimStartAndEnd", "[Core][String][SmokeFilter]")
 {
 	CHECK(TEXTVIEW("").TrimStartAndEnd().IsEmpty());
 	CHECK(TEXTVIEW(" ").TrimStartAndEnd().IsEmpty());
@@ -767,7 +767,7 @@ TEST_CASE("Core::String::StringView::TrimStartAndEnd", "[Core][String][Smoke]")
 	CHECK(TEXTVIEW(" \t\r\nABC123\n\r\t ").TrimStartAndEnd() == TEXTVIEW("ABC123"));
 }
 
-TEST_CASE("Core::String::StringView::TrimStart", "[Core][String][Smoke]")
+TEST_CASE_NAMED(FStringViewTrimStartTest, "System::Core::String::StringView::TrimStart", "[Core][String][SmokeFilter]")
 {
 	CHECK(TEXTVIEW("").TrimStart().IsEmpty());
 	CHECK(TEXTVIEW(" ").TrimStart().IsEmpty());
@@ -779,7 +779,7 @@ TEST_CASE("Core::String::StringView::TrimStart", "[Core][String][Smoke]")
 	CHECK(TEXTVIEW(" \t\r\nABC123\n\r\t ").TrimStart() == TEXTVIEW("ABC123\n\r\t "));
 }
 
-TEST_CASE("Core::String::StringView::TrimEnd", "[Core][String][Smoke]")
+TEST_CASE_NAMED(FStringViewTrimEndTest, "System::Core::String::StringView::TrimEnd", "[Core][String][SmokeFilter]")
 {
 	CHECK(TEXTVIEW("").TrimEnd().IsEmpty());
 	CHECK(TEXTVIEW(" ").TrimEnd().IsEmpty());
@@ -791,7 +791,7 @@ TEST_CASE("Core::String::StringView::TrimEnd", "[Core][String][Smoke]")
 	CHECK(TEXTVIEW(" \t\r\nABC123\n\r\t ").TrimEnd() == TEXTVIEW(" \t\r\nABC123"));
 }
 
-TEST_CASE("Core::String::StringView::FindChar", "[Core][String][Smoke]")
+TEST_CASE_NAMED(FStringViewFindCharTest, "System::Core::String::StringView::FindChar", "[Core][String][SmokeFilter]")
 {
 	FStringView EmptyView;
 	FStringView View = TEXT("aBce Fga");
@@ -833,7 +833,7 @@ TEST_CASE("Core::String::StringView::FindChar", "[Core][String][Smoke]")
 	}
 }
 
-TEST_CASE("Core::String::StringView::FindLastChar", "[Core][String][Smoke]")
+TEST_CASE_NAMED(FStringViewFindLastCharTest, "System::Core::String::StringView::FindLastChar", "[Core][String][SmokeFilter]")
 {
 	FStringView EmptyView;
 	FStringView View = TEXT("aBce Fga");
@@ -959,7 +959,7 @@ static void TestSlicing(const FString& Str)
 	}
 }
 
-TEST_CASE("Core::String::StringView::Slice", "[Core][String][Smoke]")
+TEST_CASE_NAMED(FStringViewSliceTest, "System::Core::String::StringView::Slice", "[Core][String][SmokeFilter]")
 {
 	// We assume that FString has already passed its tests, and we just want views to be consistent with it
 
@@ -975,4 +975,4 @@ TEST_CASE("Core::String::StringView::Slice", "[Core][String][Smoke]")
 	TestSlicing(TerminatorOnly);
 }
 
-#endif // WITH_LOW_LEVEL_TESTS
+#endif // WITH_TESTS
diff --git a/Engine/Source/Runtime/Core/Tests/Memory/CompositeBufferTest.cpp b/Engine/Source/Runtime/Core/Tests/Memory/CompositeBufferTest.cpp
--- a/Engine/Source/Runtime/Core/Tests/Memory/CompositeBufferTest.cpp
+++ b/Engine/Source/Runtime/Core/Tests/Memory/CompositeBufferTest.cpp
@@ -1,13 +1,13 @@
 // Copyright Epic Games, Inc. All Rights Reserved.
 
-#if WITH_LOW_LEVEL_TESTS
+#if WITH_TESTS
 
 #include "Memory/CompositeBuffer.h"
 
-#include "TestHarness.h"
+#include "Tests/TestHarnessAdapter.h"
 #include <type_traits>
 
-TEST_CASE("Core::Memory::CompositeBuffer", "[Core][Memory][Smoke]")
+TEST_CASE_NAMED(FMemoryCompositeBufferTest, "System::Core::Memory::CompositeBuffer", "[Core][Memory][SmokeFilter]")
 {
 	SECTION("Static")
 	{
@@ -214,4 +214,4 @@ TEST_CASE("Core::Memory::CompositeBuffer", "[Core][Memory][Smoke]")
 	}
 }
 
-#endif // WITH_LOW_LEVEL_TESTS
+#endif // WITH_TESTS
diff --git a/Engine/Source/Runtime/Core/Tests/Memory/MemoryViewTest.cpp b/Engine/Source/Runtime/Core/Tests/Memory/MemoryViewTest.cpp
--- a/Engine/Source/Runtime/Core/Tests/Memory/MemoryViewTest.cpp
+++ b/Engine/Source/Runtime/Core/Tests/Memory/MemoryViewTest.cpp
@@ -1,15 +1,15 @@
 // Copyright Epic Games, Inc. All Rights Reserved.
 
-#if WITH_LOW_LEVEL_TESTS
+#if WITH_TESTS
 
 #include "Memory/MemoryView.h"
 
 #include "Containers/ArrayView.h"
-#include "TestHarness.h"
+#include "Tests/TestHarnessAdapter.h"
 
 #include <type_traits>
 
-TEST_CASE("Core::Memory::MemoryView", "[Core][Memory][Smoke]")
+TEST_CASE_NAMED(FMemoryMemoryViewTest, "System::Core::Memory::MemoryView", "[Core][Memory][SmokeFilter]")
 {
 	auto TestMemoryView = [](const FMemoryView& View, const void* Data, uint64 Size)
 	{
@@ -270,4 +270,4 @@ TEST_CASE("Core::Memory::MemoryView", "[Core][Memory][Smoke]")
 	}
 }
 
-#endif // WITH_LOW_LEVEL_TESTS
+#endif // WITH_TESTS
diff --git a/Engine/Source/Runtime/Core/Tests/Serialization/Base64Test.cpp b/Engine/Source/Runtime/Core/Tests/Serialization/Base64Test.cpp
--- a/Engine/Source/Runtime/Core/Tests/Serialization/Base64Test.cpp
+++ b/Engine/Source/Runtime/Core/Tests/Serialization/Base64Test.cpp
@@ -1,9 +1,9 @@
 // Copyright Epic Games, Inc. All Rights Reserved.
 
-#if WITH_LOW_LEVEL_TESTS
+#if WITH_TESTS
 
 #include "Misc/Base64.h"
-#include "TestHarness.h"
+#include "Tests/TestHarnessAdapter.h"
 
 template <size_t InputSize>
 struct TBase64Test
@@ -54,7 +54,7 @@ void RunRoundTripSuccessfulBase64Tests(const TestType (&TestVectors)[NumVectors]
 }
 
 // Base64 encode/decode works for the test vectors provided in the RFC
-TEST_CASE("Core::Serialization::Base64::encode/decode RFC test vectors", "[Core][base64][Smoke]")
+TEST_CASE_NAMED(FBase64EncodeDecodeRFCTestVectorsTest, "System::Core::Serialization::Base64::encode/decode RFC test vectors", "[Core][base64][SmokeFilter]")
 {
 	// This test does a good job of verifying the correct padding of various lengths by going through 2 cycles of the padding modulus.
 	const FBase64StringTest TestVectors[] =
@@ -74,7 +74,7 @@ TEST_CASE("Core::Serialization::Base64::encode/decode RFC test vectors", "[Core]
 }
 
 // Standard base64 encode/decode every possible 6 bit input
-TEST_CASE("Core::Serialization::Base64::Standard encode/decode every 6 bit input", "[Core][base64][Smoke]")
+TEST_CASE_NAMED(FBase64StandartEncodeDecodeEverySixBitInputTest, "System::Core::Serialization::Base64::Standard encode/decode every 6 bit input", "[Core][base64][SmokeFilter]")
 {
 	const TBase64Test<48> TestVectors[] =
 	{
@@ -85,7 +85,7 @@ TEST_CASE("Core::Serialization::Base64::Standard encode/decode every 6 bit input
 }
 
 // UrlSafe base64 encode/decode every possible 6 bit input
-TEST_CASE("Core::Serialization::Base64::URLSafe encode/decode every 6 bit input", "[Core][base64][Smoke]")
+TEST_CASE_NAMED(FBase64URLSafeTest, "System::Core::Serialization::Base64::URLSafe encode/decode every 6 bit input", "[Core][base64][SmokeFilter]")
 {
 	const TBase64Test<48> TestVectors[] = 
 	{
@@ -135,7 +135,7 @@ void RunNegativeDecodeBase64Test(EBase64Mode Mode, TFunction<bool(int)>&& IsInDe
 }
 
 // Base64 decode should fail on input that is not in the decoding alphabet
-TEST_CASE("Core::Serialization::Base64::Negative", "[Core][base64][Smoke]")
+TEST_CASE_NAMED(FBase64NegativeTest, "System::Core::Serialization::Base64::Negative", "[Core][base64][SmokeFilter]")
 {
 	auto IsInStandardAlphabet = [](int Ix)
 	{
@@ -159,4 +159,4 @@ TEST_CASE("Core::Serialization::Base64::Negative", "[Core][base64][Smoke]")
 	RunNegativeDecodeBase64Test(EBase64Mode::UrlSafe, MoveTemp(IsInUrlSafeAlphabet));
 }
 
-#endif
+#endif //WITH_TESTS
diff --git a/Engine/Source/Runtime/Core/Tests/Serialization/CompactBinaryPackageTest.cpp b/Engine/Source/Runtime/Core/Tests/Serialization/CompactBinaryPackageTest.cpp
--- a/Engine/Source/Runtime/Core/Tests/Serialization/CompactBinaryPackageTest.cpp
+++ b/Engine/Source/Runtime/Core/Tests/Serialization/CompactBinaryPackageTest.cpp
@@ -1,6 +1,6 @@
 // Copyright Epic Games, Inc. All Rights Reserved.
 
-#if WITH_LOW_LEVEL_TESTS
+#if WITH_TESTS
 
 #include "Serialization/CompactBinaryPackage.h"
 
@@ -9,9 +9,9 @@
 #include "Serialization/CompactBinaryValidation.h"
 #include "Serialization/CompactBinaryWriter.h"
 #include "Serialization/MemoryReader.h"
-#include "TestHarness.h"
+#include "Tests/TestHarnessAdapter.h"
 
-TEST_CASE("Core::Serialization::CompactBinary::Attachment", "[Core][CompactBinary][Smoke]")
+TEST_CASE_NAMED(FCompactBinaryAttachmentTest, "System::Core::Serialization::CompactBinary::Attachment", "[Core][CompactBinary][SmokeFilter]")
 {
 	const auto TestSaveLoadValidate = [](const FCbAttachment& Attachment)
 	{
@@ -247,7 +247,7 @@ TEST_CASE("Core::Serialization::CompactBinary::Attachment", "[Core][CompactBinar
 	}
 }
 
-TEST_CASE("Core::Serialization::CompactBinary::Package", "[Core][CompactBinary][Smoke]")
+TEST_CASE_NAMED(FCompactBinaryPackageTest, "System::Core::Serialization::CompactBinary::Package", "[Core][CompactBinary][SmokeFilter]")
 {
 	const auto TestSaveLoadValidate = [](const FCbPackage& Package)
 	{
@@ -548,7 +548,7 @@ TEST_CASE("Core::Serialization::CompactBinary::Package", "[Core][CompactBinary][
 	}
 }
 
-TEST_CASE("Core::Serialization::CompactBinary::InvalidPackage", "[Core][CompactBinary][Smoke]")
+TEST_CASE_NAMED(FCompactBinaryInvalidPackageTest, "System::Core::Serialization::CompactBinary::InvalidPackage", "[Core][CompactBinary][SmokeFilter]")
 {
 	const auto TestLoad = [](std::initializer_list<uint8> RawData, FCbBufferAllocator Allocator = FUniqueBuffer::Alloc)
 	{
@@ -599,4 +599,4 @@ TEST_CASE("Core::Serialization::CompactBinary::InvalidPackage", "[Core][CompactB
 	}
 }
 
-#endif // WITH_LOW_LEVEL_TESTS
+#endif // WITH_TESTS
diff --git a/Engine/Source/Runtime/Core/Tests/Templates/ValueOrErrorTest.cpp b/Engine/Source/Runtime/Core/Tests/Templates/ValueOrErrorTest.cpp
--- a/Engine/Source/Runtime/Core/Tests/Templates/ValueOrErrorTest.cpp
+++ b/Engine/Source/Runtime/Core/Tests/Templates/ValueOrErrorTest.cpp
@@ -1,14 +1,14 @@
 // Copyright Epic Games, Inc. All Rights Reserved.
 
-#if WITH_LOW_LEVEL_TESTS
+#if WITH_TESTS
 
 #include "Templates/ValueOrError.h"
 
-#include "TestHarness.h"
+#include "Tests/TestHarnessAdapter.h"
 
 #include <type_traits>
 
-TEST_CASE("Core::Templates::ValueOrError", "[Core][Templates][Smoke]")
+TEST_CASE_NAMED(FTemplatesValueOrErrorTest, "System::Core::Templates::ValueOrError", "[Core][Templates][SmokeFilter]")
 {
 	SECTION("Static")
 	{
@@ -232,4 +232,4 @@ TEST_CASE("Core::Templates::ValueOrError", "[Core][Templates][Smoke]")
 	CHECK(ErrorCount == 0);
 }
 
-#endif
\ No newline at end of file
+#endif //WITH_TESTS
\ No newline at end of file
