Fixed warning V547: Expression 'VisitCount == 1' is always false.

diff --git a/Engine/Source/Runtime/Core/Public/Misc/LowLevelTestAdapter.h b/Engine/Source/Runtime/Core/Public/Misc/LowLevelTestAdapter.h
--- a/Engine/Source/Runtime/Core/Public/Misc/LowLevelTestAdapter.h
+++ b/Engine/Source/Runtime/Core/Public/Misc/LowLevelTestAdapter.h
@@ -68,16 +68,19 @@
 #define TEST_CASE_NAMED(ClassName, PrettyName, TFlags) TEST_CASE_NAMED_STR(ClassName, #ClassName, PrettyName, TFlags)
 
 #define INFO(What)
+//-V:CHECK:571,501,547
 #define CHECK(Expr) if (!(Expr)) { FAutomationTestFramework::Get().GetCurrentTest()->AddError(TEXT("Condition failed")); }
+//-V:CHECK_FALSE:571,501,547
 #define CHECK_FALSE(Expr) if (Expr) { FAutomationTestFramework::Get().GetCurrentTest()->AddError(TEXT("Condition expected to return false but returned true")); }
 #define CHECKED_IF(Expr) if (Expr)
 #define CHECKED_ELSE(Expr) if (!(Expr))
-//-V:CHECK_MESSAGE:571,501
+//-V:CHECK_MESSAGE:571,501,547
 #define CHECK_MESSAGE(Message, Expr) if (!(Expr)) { FAutomationTestFramework::Get().GetCurrentTest()->AddError(Message); }
-//-V:CHECK_FALSE_MESSAGE:571,501
+//-V:CHECK_FALSE_MESSAGE:571,501,547
 #define CHECK_FALSE_MESSAGE(Message, Expr) if (Expr) { FAutomationTestFramework::Get().GetCurrentTest()->AddError(Message); }
+//-V:REQUIRE:571,501,547
 #define REQUIRE(Expr) if (!(Expr)) { FAutomationTestFramework::Get().GetCurrentTest()->AddError(TEXT("Required condition failed, interrupting test")); return; }
-//-V:REQUIRE_MESSAGE:571
+//-V:REQUIRE_MESSAGE:571,501,547
 #define REQUIRE_MESSAGE(Message, Expr) if (!(Expr)) { FAutomationTestFramework::Get().GetCurrentTest()->AddError(Message); return; }
 #define STATIC_REQUIRE( ... ) static_assert(   __VA_ARGS__,  #__VA_ARGS__ );
 #define STATIC_CHECK( ... ) static_assert(   __VA_ARGS__,  #__VA_ARGS__ );
