Add GetAllowedBytesToWriteThrottledStorage to support platforms with write throttled file systems.

#rnx

diff --git a/Engine/Source/Runtime/Core/Public/GenericPlatform/GenericPlatformFile.h b/Engine/Source/Runtime/Core/Public/GenericPlatform/GenericPlatformFile.h
--- a/Engine/Source/Runtime/Core/Public/GenericPlatform/GenericPlatformFile.h
+++ b/Engine/Source/Runtime/Core/Public/GenericPlatform/GenericPlatformFile.h
@@ -647,6 +647,21 @@ public:
 	virtual void SetCreatePublicFiles(bool bCreatePublicFiles)
 	{
 	}
+
+	/**
+	 * Returns the number of bytes that are currently allowed to be written to throttled write storage (if the platform 
+	 * has such restrictions)
+	 * 
+	 * @param DestinationPath		If specified, the file system can optionally take into account the destination of 
+	 *								the file to determine the current limit
+	 * 
+	 * @returns						The number of bytes that are allowed to be written to write throttled storage.  
+	 *								If there is no limit, INT64_MAX is returned
+	 */
+	virtual int64 GetAllowedBytesToWriteThrottledStorage(const TCHAR* DestinationPath = nullptr)
+	{
+		return INT64_MAX;
+	}
 };
 
 /**
