Translated tests which not needed in new macro implementation
Translated Core\Tests\Math\ColorTests.cpp
Translated Core\Tests\Containers\TokenTest.cpp
Translated Core\Tests\Hash\Blake3Test.cpp
Translated Core\Tests\HAL\PlatformStringTests.cpp
Translated 3 files from String directory

diff --git a/Engine/Source/Runtime/Core/Tests/Containers/TokenTest.cpp b/Engine/Source/Runtime/Core/Tests/Containers/TokenTest.cpp
--- a/Engine/Source/Runtime/Core/Tests/Containers/TokenTest.cpp
+++ b/Engine/Source/Runtime/Core/Tests/Containers/TokenTest.cpp
@@ -1,11 +1,11 @@
 // Copyright Epic Games, Inc. All Rights Reserved.
 
-#if WITH_LOW_LEVEL_TESTS
+#if WITH_TESTS
 
 #include "TokenTest.h"
-#include "TestHarness.h"
+#include "Tests/TestHarnessAdapter.h"
 
-TEST_CASE("Core::Containers::TToken::Default", "[Core][Containers][TToken]")
+TEST_CASE_NAMED(FTokenDefaultTest, "System::Core::Containers::TToken::Default", "[Core][Containers][TToken]")
 {
 	int32Token::Reset();
 	{
@@ -20,7 +20,7 @@ TEST_CASE("Core::Containers::TToken::Default", "[Core][Containers][TToken]")
 	CHECK(int32Token::NumDestructionCalls() == 1);
 }
 
-TEST_CASE("Core::Containers::TToken::Explicit constructor", "[Core][Containers][TToken]")
+TEST_CASE_NAMED(FTokenExplicitConstructorTest, "System::Core::Containers::TToken::Explicit constructor", "[Core][Containers][TToken]")
 {
 	int32Token::Reset();
 	{
@@ -37,7 +37,7 @@ TEST_CASE("Core::Containers::TToken::Explicit constructor", "[Core][Containers][
 	CHECK(int32Token::NumMoveAssignmentCalls() == 0);
 }
 
-TEST_CASE("Core::Containers::TToken::Copy constructor", "[Core][Containers][TToken]")
+TEST_CASE_NAMED(FTokenCopyConstructorTest, "System::Core::Containers::TToken::Copy constructor", "[Core][Containers][TToken]")
 {
 	int32Token::Reset();
 	{
@@ -55,7 +55,7 @@ TEST_CASE("Core::Containers::TToken::Copy constructor", "[Core][Containers][TTok
 	CHECK(int32Token::NumMoveAssignmentCalls() == 0);
 }
 
-TEST_CASE("Core::Containers::TToken::Move constructor", "[Core][Containers][TToken]")
+TEST_CASE_NAMED(FTokenMoveConstructorTest, "System::Core::Containers::TToken::Move constructor", "[Core][Containers][TToken]")
 {
 	int32Token::Reset();
 	{
@@ -73,7 +73,7 @@ TEST_CASE("Core::Containers::TToken::Move constructor", "[Core][Containers][TTok
 	CHECK(int32Token::NumMoveAssignmentCalls() == 0);
 }
 
-TEST_CASE("Core::Containers::TToken::Copy assignment", "[Core][Containers][TToken]")
+TEST_CASE_NAMED(FTokenCopyAssignmentTest, "System::Core::Containers::TToken::Copy assignment", "[Core][Containers][TToken]")
 {
 	int32Token::Reset();
 	{
@@ -92,7 +92,7 @@ TEST_CASE("Core::Containers::TToken::Copy assignment", "[Core][Containers][TToke
 	CHECK(int32Token::NumMoveAssignmentCalls() == 0);
 }
 
-TEST_CASE("Core::Containers::TToken::Move assignment", "[Core][Containers][TToken]")
+TEST_CASE_NAMED(FTokenMoveAssignmentTest, "System::Core::Containers::TToken::Move assignment", "[Core][Containers][TToken]")
 {
 	int32Token::Reset();
 	{
@@ -111,4 +111,4 @@ TEST_CASE("Core::Containers::TToken::Move assignment", "[Core][Containers][TToke
 	CHECK(int32Token::NumMoveAssignmentCalls() == 1);
 }
 
-#endif
\ No newline at end of file
+#endif //WITH_TESTS
\ No newline at end of file
diff --git a/Engine/Source/Runtime/Core/Tests/HAL/PlatformStringTests.cpp b/Engine/Source/Runtime/Core/Tests/HAL/PlatformStringTests.cpp
--- a/Engine/Source/Runtime/Core/Tests/HAL/PlatformStringTests.cpp
+++ b/Engine/Source/Runtime/Core/Tests/HAL/PlatformStringTests.cpp
@@ -1,14 +1,14 @@
 // Copyright Epic Games, Inc. All Rights Reserved.
 
-#if WITH_LOW_LEVEL_TESTS
+#if WITH_TESTS
 
 #include "HAL/PlatformString.h"
 #include "Containers/Array.h"
 #include "Templates/UnrealTemplate.h"
 
-#include <catch2/catch_test_macros.hpp>
+#include "Tests/TestHarnessAdapter.h"
 
-TEST_CASE("Core::HAL::PlatformString", "[Core]")
+TEST_CASE_NAMED(FHALPlatformStringTest, "System::Core::HAL::PlatformString", "[Core]")
 {
 	SECTION("Strupr TCHAR null-termination")
 	{
@@ -79,4 +79,4 @@ TEST_CASE("Core::HAL::PlatformString", "[Core]")
 	}
 }
 
-#endif
\ No newline at end of file
+#endif //WITH_TESTS
\ No newline at end of file
diff --git a/Engine/Source/Runtime/Core/Tests/Hash/Blake3Test.cpp b/Engine/Source/Runtime/Core/Tests/Hash/Blake3Test.cpp
--- a/Engine/Source/Runtime/Core/Tests/Hash/Blake3Test.cpp
+++ b/Engine/Source/Runtime/Core/Tests/Hash/Blake3Test.cpp
@@ -1,13 +1,13 @@
 // Copyright Epic Games, Inc. All Rights Reserved.
 
-#if WITH_LOW_LEVEL_TESTS
+#if WITH_TESTS
 
 #include "Hash/Blake3.h"
 
 #include "Misc/StringBuilder.h"
-#include "TestHarness.h"
+#include "Tests/TestHarnessAdapter.h"
 
-TEST_CASE("Core::Hash::Blake3", "[Core][Hash][Smoke]")
+TEST_CASE_NAMED(FHashBlake3Test, "System::Core::Hash::Blake3", "[Core][Hash][SmokeFilter]")
 {
 	struct FBlake3TestCase
 	{
@@ -62,4 +62,4 @@ TEST_CASE("Core::Hash::Blake3", "[Core][Hash][Smoke]")
 	}
 }
 
-#endif
\ No newline at end of file
+#endif //WITH_TESTS
\ No newline at end of file
diff --git a/Engine/Source/Runtime/Core/Tests/Math/ColorTests.cpp b/Engine/Source/Runtime/Core/Tests/Math/ColorTests.cpp
--- a/Engine/Source/Runtime/Core/Tests/Math/ColorTests.cpp
+++ b/Engine/Source/Runtime/Core/Tests/Math/ColorTests.cpp
@@ -1,12 +1,12 @@
 // Copyright Epic Games, Inc. All Rights Reserved.
 
-#if WITH_LOW_LEVEL_TESTS
+#if WITH_TESTS
 
 #include "CoreMinimal.h"
 #include "Math/Color.h"
-#include "TestHarness.h"
+#include "Tests/TestHarnessAdapter.h"
 
-TEST_CASE("Core::Math::FLinearColor::Smoke Test", "[Core][Math][Smoke]")
+TEST_CASE_NAMED(FFLinearColorSmokeTest, "System::Core::Math::FLinearColor::Smoke Test", "[Core][Math][SmokeFilter]")
 {
 	SECTION("FLinearColor")
 	{
@@ -37,7 +37,7 @@ TEST_CASE("Core::Math::FLinearColor::Smoke Test", "[Core][Math][Smoke]")
 	}
 }
 
-TEST_CASE("Core::Math::FColor::Smoke Test", "[Core][Math][Smoke]")
+TEST_CASE_NAMED(FFColorSmokeTest, "System::Core::Math::FColor::Smoke Test", "[Core][Math][SmokeFilter]")
 {
 	SECTION("FColor")
 	{
@@ -67,4 +67,4 @@ TEST_CASE("Core::Math::FColor::Smoke Test", "[Core][Math][Smoke]")
 	}
 }
 
-#endif
\ No newline at end of file
+#endif //WITH_TESTS
\ No newline at end of file
diff --git a/Engine/Source/Runtime/Core/Tests/Misc/AESTest.cpp b/Engine/Source/Runtime/Core/Tests/Misc/AESTest.cpp
--- a/Engine/Source/Runtime/Core/Tests/Misc/AESTest.cpp
+++ b/Engine/Source/Runtime/Core/Tests/Misc/AESTest.cpp
@@ -1,6 +1,6 @@
 // Copyright Epic Games, Inc. All Rights Reserved.
 
-#if WITH_LOW_LEVEL_TESTS
+#if WITH_TESTS
 
 #include "Misc/AES.h"
 
@@ -11,9 +11,9 @@
 #include "Containers/Array.h"
 #include "Misc/StringBuilder.h"
 
-#include "TestHarness.h"
+#include "Tests/TestHarnessAdapter.h"
 
-TEST_CASE("Core::Misc::AES", "[Core][AES][Smoke]")
+TEST_CASE_NAMED(FMiscAESTest, "System::Core::Misc::AES", "[Core][AES][SmokeFilter]")
 {
 	auto RunTest = [](FAnsiStringView Key, FAnsiStringView PlainText, FAnsiStringView CipherText, int IterCount = 1)
 	{
@@ -855,4 +855,4 @@ TEST_CASE("Core::Misc::AES", "[Core][AES][Smoke]")
 	}
 }
 
-#endif
+#endif //WITH_TESTS
diff --git a/Engine/Source/Runtime/Core/Tests/Misc/SecureHashTest.cpp b/Engine/Source/Runtime/Core/Tests/Misc/SecureHashTest.cpp
--- a/Engine/Source/Runtime/Core/Tests/Misc/SecureHashTest.cpp
+++ b/Engine/Source/Runtime/Core/Tests/Misc/SecureHashTest.cpp
@@ -1,6 +1,6 @@
 // Copyright Epic Games, Inc. All Rights Reserved.
 
-#if WITH_LOW_LEVEL_TESTS
+#if WITH_TESTS
 
 #include "Misc/SecureHash.h"
 
@@ -11,9 +11,9 @@
 #include "Containers/Array.h"
 #include "Misc/StringBuilder.h"
 
-#include "TestHarness.h"
+#include "Tests/TestHarnessAdapter.h"
 
-TEST_CASE("Core::Hash::SecureHash", "[Core][SHA1][Smoke]")
+TEST_CASE_NAMED(FHashSecureHashTest, "System::Core::Hash::SecureHash", "[Core][SHA1][SmokeFilter]")
 {
 	struct FSHA1TestCase
 	{
@@ -208,4 +208,4 @@ TEST_CASE("Core::Hash::SecureHash", "[Core][SHA1][Smoke]")
 	}
 }
 
-#endif
+#endif //WITH_TESTS
diff --git a/Engine/Source/Runtime/Core/Tests/Misc/URLRequestFilterTest.cpp b/Engine/Source/Runtime/Core/Tests/Misc/URLRequestFilterTest.cpp
--- a/Engine/Source/Runtime/Core/Tests/Misc/URLRequestFilterTest.cpp
+++ b/Engine/Source/Runtime/Core/Tests/Misc/URLRequestFilterTest.cpp
@@ -1,6 +1,6 @@
 // Copyright Epic Games, Inc. All Rights Reserved.
 
-#if WITH_LOW_LEVEL_TESTS
+#if WITH_TESTS
 
 #include "Misc/URLRequestFilter.h"
 #include "Misc/ConfigCacheIni.h"
@@ -10,7 +10,7 @@
 #include "Containers/Map.h"
 #include "Containers/UnrealString.h"
 
-#include <catch2/catch_test_macros.hpp>
+#include "Tests/TestHarnessAdapter.h"
 
 namespace UE::Core::Private
 {
@@ -63,7 +63,7 @@ private:
 };
 
 
-TEST_CASE("Core::Misc::URLRequestFilter", "[Core]")
+TEST_CASE_NAMED(FMiscURLRequestFilterTest, "System::Core::Misc::URLRequestFilter", "[Core]")
 {
 	SECTION("Empty filter")
 	{
@@ -525,4 +525,4 @@ TEST_CASE("Core::Misc::URLRequestFilter", "[Core]")
 
 }
 
-#endif
+#endif //WITH_TESTS
diff --git a/Engine/Source/Runtime/Core/Tests/String/EscapeTest.cpp b/Engine/Source/Runtime/Core/Tests/String/EscapeTest.cpp
--- a/Engine/Source/Runtime/Core/Tests/String/EscapeTest.cpp
+++ b/Engine/Source/Runtime/Core/Tests/String/EscapeTest.cpp
@@ -1,17 +1,17 @@
 // Copyright Epic Games, Inc. All Rights Reserved.
 
-#if WITH_LOW_LEVEL_TESTS
+#if WITH_TESTS
 
 #include "String/Escape.h"
 
 #include "Containers/StringView.h"
 #include "Misc/StringBuilder.h"
-#include "TestHarness.h"
+#include "Tests/TestHarnessAdapter.h"
 
 namespace UE::String
 {
 
-TEST_CASE("Core::String::Escape", "[Core][String][Smoke]")
+TEST_CASE_NAMED(FStringEscapeTest, "System::Core::String::Escape", "[Core][String][SmokeFilter]")
 {
 	SECTION("Escape")
 	{
@@ -38,4 +38,4 @@ TEST_CASE("Core::String::Escape", "[Core][String][Smoke]")
 
 } // UE::String
 
-#endif
\ No newline at end of file
+#endif //WITH_TESTS
\ No newline at end of file
diff --git a/Engine/Source/Runtime/Core/Tests/String/FindTest.cpp b/Engine/Source/Runtime/Core/Tests/String/FindTest.cpp
--- a/Engine/Source/Runtime/Core/Tests/String/FindTest.cpp
+++ b/Engine/Source/Runtime/Core/Tests/String/FindTest.cpp
@@ -1,16 +1,16 @@
 // Copyright Epic Games, Inc. All Rights Reserved.
 
-#if WITH_LOW_LEVEL_TESTS
+#if WITH_TESTS
 
 #include "String/Find.h"
 
 #include "Containers/StringView.h"
-#include "TestHarness.h"
+#include "Tests/TestHarnessAdapter.h"
 
 namespace UE::String
 {
 
-TEST_CASE("Core::String::Find", "[Core][String][Smoke]")
+TEST_CASE_NAMED(FStringFindTest, "System::Core::String::Find", "[Core][String][SmokeFilter]")
 {
 	SECTION("FindFirst")
 	{
@@ -225,4 +225,4 @@ TEST_CASE("Core::String::Find", "[Core][String][Smoke]")
 
 } // UE::String
 
-#endif
\ No newline at end of file
+#endif //WITH_TESTS
\ No newline at end of file
diff --git a/Engine/Source/Runtime/Core/Tests/String/JoinTest.cpp b/Engine/Source/Runtime/Core/Tests/String/JoinTest.cpp
--- a/Engine/Source/Runtime/Core/Tests/String/JoinTest.cpp
+++ b/Engine/Source/Runtime/Core/Tests/String/JoinTest.cpp
@@ -1,18 +1,18 @@
 // Copyright Epic Games, Inc. All Rights Reserved.
 
-#if WITH_LOW_LEVEL_TESTS
+#if WITH_TESTS
 
 #include "String/Join.h"
 
 #include "Containers/StringView.h"
 #include "Containers/UnrealString.h"
 #include "Misc/StringBuilder.h"
-#include "TestHarness.h"
+#include "Tests/TestHarnessAdapter.h"
 
 namespace UE::String
 {
 
-TEST_CASE("Core::String::Join", "[Core][String][Smoke]")
+TEST_CASE_NAMED(FStringJoinTest, "System::Core::String::Join", "[Core][String][SmokeFilter]")
 {
 	SECTION("Join")
 	{
@@ -39,4 +39,4 @@ TEST_CASE("Core::String::Join", "[Core][String][Smoke]")
 
 } // UE::String
 
-#endif
\ No newline at end of file
+#endif //WITH_TESTS
\ No newline at end of file
diff --git a/Engine/Source/Runtime/Core/Tests/String/RemoveFromTest.cpp b/Engine/Source/Runtime/Core/Tests/String/RemoveFromTest.cpp
--- a/Engine/Source/Runtime/Core/Tests/String/RemoveFromTest.cpp
+++ b/Engine/Source/Runtime/Core/Tests/String/RemoveFromTest.cpp
@@ -1,13 +1,13 @@
 // Copyright Epic Games, Inc. All Rights Reserved.
 
-#if WITH_LOW_LEVEL_TESTS
+#if WITH_TESTS
 
 #include "String/RemoveFrom.h"
 
 #include "Containers/StringView.h"
-#include "TestHarness.h"
+#include "Tests/TestHarnessAdapter.h"
 
-TEST_CASE("Core::String::RemoveFrom", "[Core][String][Smoke]")
+TEST_CASE_NAMED(FStringRemoveFromTest, "System::Core::String::RemoveFrom", "[Core][String][SmokeFilter]")
 {
 	SECTION("RemoveFromStart")
 	{
@@ -76,4 +76,4 @@ TEST_CASE("Core::String::RemoveFrom", "[Core][String][Smoke]")
 	}
 }
 
-#endif // WITH_LOW_LEVEL_TESTS
+#endif //WITH_TESTS
