Updated version numbers in the example comments for UE_DEPRECATED and added examples for class templates and function templates.

diff --git a/Engine/Source/Runtime/Core/Public/Misc/CoreMiscDefines.h b/Engine/Source/Runtime/Core/Public/Misc/CoreMiscDefines.h
--- a/Engine/Source/Runtime/Core/Public/Misc/CoreMiscDefines.h
+++ b/Engine/Source/Runtime/Core/Public/Misc/CoreMiscDefines.h
@@ -179,18 +179,18 @@ enum EInPlace {InPlace};
  *
  * Sample usage (note the slightly different syntax for classes and structures):
  *
- *		UE_DEPRECATED(4.xx, "Message")
+ *		UE_DEPRECATED(5.xx, "Message")
  *		void MyFunction();
  *
- *		UE_DEPRECATED(4.xx, "Message")
+ *		UE_DEPRECATED(5.xx, "Message")
  *		typedef FThing MyType;
  *
- *		using MyAlias UE_DEPRECATED(4.xx, "Message") = FThing;
+ *		using MyAlias UE_DEPRECATED(5.xx, "Message") = FThing;
  *
- *		UE_DEPRECATED(4.xx, "Message")
+ *		UE_DEPRECATED(5.xx, "Message")
  *		int32 MyVariable;
  *
- *		namespace UE_DEPRECATED(4.xx, "Message") MyNamespace
+ *		namespace UE_DEPRECATED(5.xx, "Message") MyNamespace
  *		{
  *		}
  *
@@ -199,27 +199,27 @@ enum EInPlace {InPlace};
  *		deprecated, then declare the type with the visibility macro.  Note that macros like
  *		USTRUCT must immediately precede the the type declaration, not the forward declaration.
  *
- *		struct UE_DEPRECATED(4.xx, "Message") FMyStruct;
+ *		struct UE_DEPRECATED(5.xx, "Message") FMyStruct;
  *		USTRUCT()
  *		struct MODULE_API FMyStruct
  *		{
  *		};
  *
- *		class UE_DEPRECATED(4.xx, "Message") FMyClass;
+ *		class UE_DEPRECATED(5.xx, "Message") FMyClass;
  *		class MODULE_API FMyClass
  *		{
  *		};
  *
- *		enum class UE_DEPRECATED(4.xx, "Message") EMyEnumeration
+ *		enum class UE_DEPRECATED(5.xx, "Message") EMyEnumeration
  *		{
  *			Zero = 0,
- *			One UE_DEPRECATED(4.xx, "Message") = 1,
+ *			One UE_DEPRECATED(5.xx, "Message") = 1,
  *			Two = 2
  *		};
  *
  *		Unfortunately, VC++ will complain about using member functions and fields from deprecated
  *		class/structs even for class/struct implementation e.g.:
- *		class UE_DEPRECATED(4.xx, "") DeprecatedClass
+ *		class UE_DEPRECATED(5.xx, "") DeprecatedClass
  *		{
  *		public:
  *			DeprecatedClass() {}
@@ -242,7 +242,7 @@ enum EInPlace {InPlace};
  *			float MyFloat;
  *		};
  *
- *		class UE_DEPRECATED(4.xx, "") DeprecatedClass : DeprecatedClass_Base_DEPRECATED
+ *		class UE_DEPRECATED(5.xx, "") DeprecatedClass : DeprecatedClass_Base_DEPRECATED
  *		{
  *		public:
  *			DeprecatedClass() {}
@@ -253,6 +253,17 @@ enum EInPlace {InPlace};
  *			}
  *		};
  *
+ *		template <typename T>
+ *		class UE_DEPRECATED(5.xx, "") DeprecatedClassTemplate
+ *		{
+ *		};
+ *
+ *		template <typename T>
+ *		UE_DEPRECATED(5.xx, "")
+ *		void DeprecatedFunctionTemplate()
+ *		{
+ *		}
+ *
  * @param VERSION The release number in which the feature was marked deprecated.
  * @param MESSAGE A message containing upgrade notes.
  */
