- Removed the final "NoSave=true" line (we had removed it for Editors, but games are unable to have, say, the console history saved for development if any plugin had an Input.ini file)

diff --git a/Engine/Source/Runtime/Core/Private/Misc/ConfigCacheIni.cpp b/Engine/Source/Runtime/Core/Private/Misc/ConfigCacheIni.cpp
--- a/Engine/Source/Runtime/Core/Private/Misc/ConfigCacheIni.cpp
+++ b/Engine/Source/Runtime/Core/Private/Misc/ConfigCacheIni.cpp
@@ -1330,13 +1330,6 @@ void FConfigFile::AddDynamicLayerToHierarchy(const FString& Filename)
 
 	SourceIniHierarchy.AddDynamicLayer(Filename);
 	CombineFromBuffer(ConfigContent, Filename);
-
-	// Removing NoSave flag that was being set to true when adding a dynamic layered to a config hierarchy 
-	// (usually from plugins), which was causing settings in the editor to not be saved out
-	if (!GIsEditor)
-	{
-		NoSave = true;
-	}
 }
 
 
