Actually deprecate TScriptDelegate<FWeakObjectPtr>.

diff --git a/Engine/Source/Runtime/Core/Public/UObject/ScriptDelegates.h b/Engine/Source/Runtime/Core/Public/UObject/ScriptDelegates.h
--- a/Engine/Source/Runtime/Core/Public/UObject/ScriptDelegates.h
+++ b/Engine/Source/Runtime/Core/Public/UObject/ScriptDelegates.h
@@ -25,8 +25,12 @@ namespace UE::Core::Private
 	};
 
 	template <>
-	struct /*UE_DEPRECATED(5.3, "TScriptDelegate<FWeakObjectPtr> and TMulticastScriptDelegate<FWeakObjectPtr> have been deprecated, please use FScriptDelegate or FMulticastScriptDelegate respectively.")*/ TScriptDelegateTraits<FWeakObjectPtr>
+	struct UE_DEPRECATED(5.3, "TScriptDelegate<FWeakObjectPtr> and TMulticastScriptDelegate<FWeakObjectPtr> have been deprecated, please use FScriptDelegate or FMulticastScriptDelegate respectively.") TScriptDelegateTraits<FWeakObjectPtr>
 	{
+		// After this deprecated specialization has been removed, all of the functions inside
+		// TMulticastScriptDelegate which take OtherDummy parameters should also be removed,
+		// and also the TScriptDelegate(const TScriptDelegate<FWeakObjectPtr>&) constructor.
+
 		using WeakPtrType = FWeakObjectPtr;
 	};
 }
@@ -39,6 +43,9 @@ class TScriptDelegate
 {
 	using WeakPtrType = typename UE::Core::Private::TScriptDelegateTraits<Dummy>::WeakPtrType;
 
+	template <typename>
+	friend class TScriptDelegate;
+
 public:
 
 	/** Default constructor. */
@@ -47,6 +54,17 @@ public:
 		  FunctionName( NAME_None )
 	{ }
 
+	/* UE_DEPRECATED(5.3, "Deprecated - remove after TScriptDelegateTraits<FWeakObjectPtr> is removed") */
+	template <
+		typename OtherDummy
+		UE_REQUIRES(std::is_same_v<OtherDummy, FWeakObjectPtr>)
+	>
+	TScriptDelegate(const TScriptDelegate<OtherDummy>& Other)
+		: Object(Other.Object)
+		, FunctionName(Other.FunctionName)
+	{
+	}
+
 private:
 
 	template <class UObjectTemplate>
@@ -376,6 +394,15 @@ public:
 
 		return InvocationList.Contains( InDelegate );
 	}
+	template <
+		typename OtherDummy
+		UE_REQUIRES(std::is_same_v<OtherDummy, FWeakObjectPtr>)
+	>
+	/* UE_DEPRECATED(5.3, "Deprecated - remove after TScriptDelegateTraits<FWeakObjectPtr> is removed") */
+	FORCEINLINE bool Contains(const TScriptDelegate<OtherDummy>& InDelegate) const
+	{
+		return Contains((UnicastDelegateType)InDelegate);
+	}
 
 	/**
 	 * Checks whether a function delegate is already a member of this multi-cast delegate's invocation list
@@ -408,6 +435,15 @@ public:
 		// Add the delegate
 		AddInternal( InDelegate );
 	}
+	template <
+		typename OtherDummy
+		UE_REQUIRES(std::is_same_v<OtherDummy, FWeakObjectPtr>)
+	>
+	/* UE_DEPRECATED(5.3, "Deprecated - remove after TScriptDelegateTraits<FWeakObjectPtr> is removed") */
+	FORCEINLINE void Add(const TScriptDelegate<OtherDummy>& InDelegate)
+	{
+		Add((UnicastDelegateType)InDelegate);
+	}
 
 	/**
 	 * Adds a function delegate to this multi-cast delegate's invocation list if a delegate with the same signature
@@ -425,6 +461,15 @@ public:
 		// Then check for any objects that may have expired
 		CompactInvocationList();
 	}
+	template <
+		typename OtherDummy
+		UE_REQUIRES(std::is_same_v<OtherDummy, FWeakObjectPtr>)
+	>
+	/* UE_DEPRECATED(5.3, "Deprecated - remove after TScriptDelegateTraits<FWeakObjectPtr> is removed") */
+	FORCEINLINE void AddUnique(const TScriptDelegate<OtherDummy>& InDelegate)
+	{
+		AddUnique((UnicastDelegateType)InDelegate);
+	}
 
 	/**
 	 * Removes a function from this multi-cast delegate's invocation list (performance is O(N)).  Note that the
@@ -442,6 +487,15 @@ public:
 		// Check for any delegates that may have expired
 		CompactInvocationList();
 	}
+	template <
+		typename OtherDummy
+		UE_REQUIRES(std::is_same_v<OtherDummy, FWeakObjectPtr>)
+	>
+	/* UE_DEPRECATED(5.3, "Deprecated - remove after TScriptDelegateTraits<FWeakObjectPtr> is removed") */
+	FORCEINLINE void Remove(const TScriptDelegate<OtherDummy>& InDelegate)
+	{
+		Remove((UnicastDelegateType)InDelegate);
+	}
 
 	/**
 	 * Removes a function from this multi-cast delegate's invocation list (performance is O(N)).  Note that the
