Simple expression capturer for Functional Tests

diff --git a/Engine/Source/Runtime/Core/Public/Misc/LowLevelTestAdapter.h b/Engine/Source/Runtime/Core/Public/Misc/LowLevelTestAdapter.h
--- a/Engine/Source/Runtime/Core/Public/Misc/LowLevelTestAdapter.h
+++ b/Engine/Source/Runtime/Core/Public/Misc/LowLevelTestAdapter.h
@@ -5,7 +5,48 @@
 #if !WITH_LOW_LEVEL_TESTS
 
 #include "CoreTypes.h"
-#include "Misc/AssertionMacros.h"
+
+#include <sstream>
+#include <vector>
+
+/**
+* @brief Captures expressions and their evaluated values.
+* Internal use only for the low level tests adapter.
+*
+* @param	InExpressions	Comma-separated list of expressions
+* @param	InExperssionsValues	The list of evaluated expressions values
+*/
+template<typename... ArgTypes>
+FString CaptureExpressionsAndValues(const FString& InExpressions, ArgTypes&&... InExpressionsValues)
+{
+	std::ostringstream Result;
+	auto Args = { InExpressionsValues... };
+	auto Iter = Args.begin();
+
+	FString RemainingExpressions = InExpressions;
+	FString Expression;
+	while (RemainingExpressions.Split(TEXT(","), &Expression, &RemainingExpressions))
+	{
+		if (Iter == Args.end())
+		{
+			break;
+		}
+		else if (Iter != Args.begin())
+		{
+			Result << ", ";
+		}
+		Result << std::string(TCHAR_TO_UTF8(*Expression.TrimStartAndEnd())) << " = "  << *Iter++;
+	}
+
+	if (Iter != Args.end() && !RemainingExpressions.TrimStartAndEnd().IsEmpty())
+	{
+		Result << ", " << std::string(TCHAR_TO_UTF8(*RemainingExpressions)) << " = " << *Iter;
+	}
+
+	Result << std::endl;
+	return FString(Result.str().c_str());
+}
+
 
 #define IMPLEMENT_SIMPLE_AUTOMATION_TEST_PRIVATE_LLT( TClass, PrettyName, TFlags, FileName, LineNumber ) \
 		class TClass : public FAutomationTestBase \
@@ -82,12 +123,12 @@
 #define REQUIRE(Expr) if (!(Expr)) { FAutomationTestFramework::Get().GetCurrentTest()->AddError(TEXT("Required condition failed, interrupting test")); return; }
 //-V:REQUIRE_MESSAGE:571,501,547
 #define REQUIRE_MESSAGE(Message, Expr) if (!(Expr)) { FAutomationTestFramework::Get().GetCurrentTest()->AddError(Message); return; }
-#define STATIC_REQUIRE( ... ) static_assert(   __VA_ARGS__,  #__VA_ARGS__ );
-#define STATIC_CHECK( ... ) static_assert(   __VA_ARGS__,  #__VA_ARGS__ );
+#define STATIC_REQUIRE(...) static_assert(__VA_ARGS__, #__VA_ARGS__);
+#define STATIC_CHECK(...) static_assert(__VA_ARGS__, #__VA_ARGS__);
 
-#define SECTION(Text) AddInfo(TEXT(Text));
-#define FAIL_CHECK(Message) { FAutomationTestFramework::Get().GetCurrentTest()->AddError(Message); }
+#define SECTION(Text) FAutomationTestFramework::Get().GetCurrentTest()->AddInfo(TEXT(Text));
+#define FAIL_CHECK(Message) FAutomationTestFramework::Get().GetCurrentTest()->AddError(Message);
 
-#define CAPTURE(...) { FAutomationTestFramework::Get().GetCurrentTest()->AddInfo(FString::Printf(TEXT("%s := %d"), TEXT(#__VA_ARGS__), __VA_ARGS__)); }
+#define CAPTURE(...) FAutomationTestFramework::Get().GetCurrentTest()->AddInfo(CaptureExpressionsAndValues(#__VA_ARGS__, __VA_ARGS__));
 
 #endif // !WITH_LOW_LEVEL_TESTS
\ No newline at end of file
