Replace -[NSProcess operatingSystemVersion] with sysctlbyname()

diff --git a/Engine/Source/Runtime/Core/Private/Apple/ApplePlatformMisc.cpp b/Engine/Source/Runtime/Core/Private/Apple/ApplePlatformMisc.cpp
--- a/Engine/Source/Runtime/Core/Private/Apple/ApplePlatformMisc.cpp
+++ b/Engine/Source/Runtime/Core/Private/Apple/ApplePlatformMisc.cpp
@@ -330,7 +330,14 @@ FString FApplePlatformMisc::GetLocalCurrencySymbol()
 
 bool FApplePlatformMisc::IsOSAtLeastVersion(const uint32 MacOSVersion[3], const uint32 IOSVersion[3], const uint32 TVOSVersion[3])
 {
-	static const uint32 OSVersion[3] = { (uint32)[NSProcessInfo processInfo].operatingSystemVersion.majorVersion, (uint32)[NSProcessInfo processInfo].operatingSystemVersion.minorVersion, (uint32)[NSProcessInfo processInfo].operatingSystemVersion.patchVersion };
+	NSOperatingSystemVersion CurrentSystemVersion =
+#if PLATFORM_MAC
+	FMacPlatformMisc::GetNSOperatingSystemVersion();
+#else
+	[NSProcessInfo processInfo].operatingSystemVersion;
+#endif
+
+	static const uint32 OSVersion[3] = { (uint32)CurrentSystemVersion.majorVersion, (uint32)CurrentSystemVersion.minorVersion, (uint32)CurrentSystemVersion.patchVersion };
 	const uint32* VersionToCompare = PLATFORM_MAC ? MacOSVersion : (PLATFORM_IOS ? IOSVersion : TVOSVersion);
 
 	for (uint32 Index = 0; Index < 3; Index++)
diff --git a/Engine/Source/Runtime/Core/Private/Mac/MacPlatformMisc.cpp b/Engine/Source/Runtime/Core/Private/Mac/MacPlatformMisc.cpp
--- a/Engine/Source/Runtime/Core/Private/Mac/MacPlatformMisc.cpp
+++ b/Engine/Source/Runtime/Core/Private/Mac/MacPlatformMisc.cpp
@@ -161,32 +161,44 @@ static const char* GetClassCode()
 
  CORE_API FMacMallocCrashHandler* GCrashMalloc = nullptr;
 
+// e.g. 13.2.1
 static void MacPlatformGetOSProductVersion(FString &OutOSVersion)
 {
-	SIZE_T OSProductVersionStringSize = 0;
-	if (0 == (sysctlbyname("kern.osproductversion", NULL, &OSProductVersionStringSize, NULL, 0)))
+	static FString OSVersion = TEXT("");
+	if (OSVersion.IsEmpty())
 	{
-		ANSICHAR *OSProductVersionCString = new ANSICHAR[OSProductVersionStringSize];
-		if (0 == (sysctlbyname("kern.osproductversion", OSProductVersionCString, &OSProductVersionStringSize, NULL, 0)))
+		SIZE_T OSProductVersionStringSize = 0;
+		if (0 == (sysctlbyname("kern.osproductversion", NULL, &OSProductVersionStringSize, NULL, 0)))
 		{
-			OutOSVersion = ANSI_TO_TCHAR(OSProductVersionCString);
+			ANSICHAR *OSProductVersionCString = new ANSICHAR[OSProductVersionStringSize];
+			if (0 == (sysctlbyname("kern.osproductversion", OSProductVersionCString, &OSProductVersionStringSize, NULL, 0)))
+			{
+				OSVersion = ANSI_TO_TCHAR(OSProductVersionCString);
+			}
+			delete [] OSProductVersionCString;
 		}
-		delete [] OSProductVersionCString;
 	}
+	OutOSVersion = OSVersion;
 }
 
+// e.g. 22D68
 static void MacPlatformGetOSVersion(FString &OutOSBuild)
 {
-	SIZE_T OSVersionStringSize = 0;
-	if (0 == (sysctlbyname("kern.osversion", NULL, &OSVersionStringSize, NULL, 0)))
+	static FString OSBuild = TEXT("");
+	if (OSBuild.IsEmpty())
 	{
-		ANSICHAR *OSVersionCString = new ANSICHAR[OSVersionStringSize];
-		if (0 == (sysctlbyname("kern.osversion", OSVersionCString, &OSVersionStringSize, NULL, 0)))
+		SIZE_T OSVersionStringSize = 0;
+		if (0 == (sysctlbyname("kern.osversion", NULL, &OSVersionStringSize, NULL, 0)))
 		{
-			OutOSBuild = ANSI_TO_TCHAR(OSVersionCString);
+			ANSICHAR *OSVersionCString = new ANSICHAR[OSVersionStringSize];
+			if (0 == (sysctlbyname("kern.osversion", OSVersionCString, &OSVersionStringSize, NULL, 0)))
+			{
+				OSBuild = ANSI_TO_TCHAR(OSVersionCString);
+			}
+			delete [] OSVersionCString;
 		}
-		delete [] OSVersionCString;
 	}
+	OutOSBuild = OSBuild;
 }
 
 /**
@@ -236,7 +248,7 @@ struct FMacApplicationInfo
 		
 		MacPlatformGetOSVersion(OSBuild);
 
-		OSXVersion = [[NSProcessInfo processInfo] operatingSystemVersion];
+		OSXVersion = FMacPlatformMisc::GetNSOperatingSystemVersion();
 		RunningOnMavericks = OSXVersion.majorVersion == 10 && OSXVersion.minorVersion == 9;
 
 		XcodeVersion.majorVersion = XcodeVersion.minorVersion = XcodeVersion.patchVersion = 0;
@@ -1617,6 +1629,21 @@ FString FMacPlatformMisc::GetOSVersion()
 	return OSVersion;
 }
 
+NSOperatingSystemVersion FMacPlatformMisc::GetNSOperatingSystemVersion()
+{
+	FString OSVersionString;
+	NSOperatingSystemVersion SystemVersion = {};
+	MacPlatformGetOSProductVersion(OSVersionString);
+	if (sscanf(TCHAR_TO_ANSI(*OSVersionString), "%ld.%ld.%ld", &SystemVersion.majorVersion, &SystemVersion.minorVersion, &SystemVersion.patchVersion))
+	{
+		return SystemVersion;
+	}
+	else
+	{
+		return {0,0,0};
+	}
+}
+
 bool FMacPlatformMisc::GetDiskTotalAndFreeSpace(const FString& InPath, uint64& TotalNumberOfBytes, uint64& NumberOfFreeBytes)
 {
 	struct statfs FSStat = { 0 };
diff --git a/Engine/Source/Runtime/Core/Public/Mac/MacPlatformMisc.h b/Engine/Source/Runtime/Core/Public/Mac/MacPlatformMisc.h
--- a/Engine/Source/Runtime/Core/Public/Mac/MacPlatformMisc.h
+++ b/Engine/Source/Runtime/Core/Public/Mac/MacPlatformMisc.h
@@ -42,6 +42,7 @@ struct CORE_API FMacPlatformMisc : public FApplePlatformMisc
 	static struct FGPUDriverInfo GetGPUDriverInfo(const FString& DeviceDescription);
 	static void GetOSVersions( FString& out_OSVersionLabel, FString& out_OSSubVersionLabel );
 	static FString GetOSVersion();
+	static NSOperatingSystemVersion GetNSOperatingSystemVersion();
 	static bool HasPlatformFeature(const TCHAR* FeatureName);
 	static bool GetDiskTotalAndFreeSpace(const FString& InPath, uint64& TotalNumberOfBytes, uint64& NumberOfFreeBytes);
 	static bool HasSeparateChannelForDebugOutput();
diff --git a/Engine/Source/Runtime/Launch/Private/Mac/LaunchMac.cpp b/Engine/Source/Runtime/Launch/Private/Mac/LaunchMac.cpp
--- a/Engine/Source/Runtime/Launch/Private/Mac/LaunchMac.cpp
+++ b/Engine/Source/Runtime/Launch/Private/Mac/LaunchMac.cpp
@@ -265,7 +265,7 @@ static int32 MacOSVersionCompare(const NSOperatingSystemVersion& VersionA, const
 	NSDictionary* InfoDictionary = [[NSBundle mainBundle] infoDictionary];
 	NSString* MinimumSystemVersionString = (NSString*)InfoDictionary[@"LSMinimumSystemVersion"];
 	NSOperatingSystemVersion MinimumSystemVersion = { 0 };
-	NSOperatingSystemVersion CurrentSystemVersion = [[NSProcessInfo processInfo] operatingSystemVersion];
+	NSOperatingSystemVersion CurrentSystemVersion = FMacPlatformMisc::GetNSOperatingSystemVersion();
 	NSOperatingSystemVersion MinSupportedMacOSVersion = { 12, 0, 0 };
 	NSString* MinSupportedMacOSVersionString = @"12.0.0";
 
