[Insights] Update Regions code to use custom trace events instead of hijacking regular bookmarks

diff --git a/Engine/Plugins/TraceUtilities/Source/TraceUtilities/Private/TraceUtilLibrary.cpp b/Engine/Plugins/TraceUtilities/Source/TraceUtilities/Private/TraceUtilLibrary.cpp
--- a/Engine/Plugins/TraceUtilities/Source/TraceUtilities/Private/TraceUtilLibrary.cpp
+++ b/Engine/Plugins/TraceUtilities/Source/TraceUtilities/Private/TraceUtilLibrary.cpp
@@ -80,10 +80,10 @@ void UTraceUtilLibrary::TraceBookmark(const FString& Name)
 
 void UTraceUtilLibrary::TraceMarkRegionStart(const FString& Name)
 {
-	TRACE_BOOKMARK(TEXT("RegionStart:%s"), *Name);
+	TRACE_BEGIN_REGION(*Name);
 }
 
 void UTraceUtilLibrary::TraceMarkRegionEnd(const FString& Name)
 {
-	TRACE_BOOKMARK(TEXT("RegionEnd:%s"), *Name);
+	TRACE_END_REGION(*Name);
 }
diff --git a/Engine/Source/Runtime/Core/Private/ProfilingDebugging/MiscTrace.cpp b/Engine/Source/Runtime/Core/Private/ProfilingDebugging/MiscTrace.cpp
--- a/Engine/Source/Runtime/Core/Private/ProfilingDebugging/MiscTrace.cpp
+++ b/Engine/Source/Runtime/Core/Private/ProfilingDebugging/MiscTrace.cpp
@@ -10,6 +10,7 @@
 
 UE_TRACE_CHANNEL(FrameChannel)
 UE_TRACE_CHANNEL(BookmarkChannel)
+UE_TRACE_CHANNEL(RegionChannel)
 UE_TRACE_CHANNEL(ScreenshotChannel)
 
 UE_TRACE_EVENT_BEGIN(Misc, BookmarkSpec, NoSync|Important)
@@ -25,6 +26,16 @@ UE_TRACE_EVENT_BEGIN(Misc, Bookmark)
 	UE_TRACE_EVENT_FIELD(uint8[], FormatArgs)
 UE_TRACE_EVENT_END()
 
+UE_TRACE_EVENT_BEGIN(Misc, RegionBegin)
+	UE_TRACE_EVENT_FIELD(uint64, Cycle)
+	UE_TRACE_EVENT_FIELD(UE::Trace::WideString, RegionName)
+UE_TRACE_EVENT_END()
+
+UE_TRACE_EVENT_BEGIN(Misc, RegionEnd)
+	UE_TRACE_EVENT_FIELD(uint64, Cycle)
+	UE_TRACE_EVENT_FIELD(UE::Trace::WideString, RegionName)
+UE_TRACE_EVENT_END()
+
 UE_TRACE_EVENT_BEGIN(Misc, BeginFrame)
 	UE_TRACE_EVENT_FIELD(uint64, Cycle)
 	UE_TRACE_EVENT_FIELD(uint8, FrameType)
@@ -65,6 +76,20 @@ void FMiscTrace::OutputBookmarkSpec(const void* BookmarkPoint, const ANSICHAR* F
 		<< BookmarkSpec.FileName(File, FileNameLen);
 }
 
+void FMiscTrace::OutputBeginRegion(const TCHAR* RegionName)
+{
+	UE_TRACE_LOG(Misc, RegionBegin, RegionChannel)
+		<< RegionBegin.Cycle(FPlatformTime::Cycles64())
+		<< RegionBegin.RegionName(RegionName);
+}
+
+void FMiscTrace::OutputEndRegion(const TCHAR* RegionName)
+{
+	UE_TRACE_LOG(Misc, RegionEnd, RegionChannel)
+		<< RegionEnd.Cycle(FPlatformTime::Cycles64())
+		<< RegionEnd.RegionName(RegionName);
+}
+
 void FMiscTrace::OutputBookmarkInternal(const void* BookmarkPoint, uint16 EncodedFormatArgsSize, uint8* EncodedFormatArgs)
 {
 	UE_TRACE_LOG(Misc, Bookmark, BookmarkChannel)
diff --git a/Engine/Source/Runtime/Core/Private/ProfilingDebugging/TraceAuxiliary.cpp b/Engine/Source/Runtime/Core/Private/ProfilingDebugging/TraceAuxiliary.cpp
--- a/Engine/Source/Runtime/Core/Private/ProfilingDebugging/TraceAuxiliary.cpp
+++ b/Engine/Source/Runtime/Core/Private/ProfilingDebugging/TraceAuxiliary.cpp
@@ -55,7 +55,7 @@
 #include "Trace/Trace.inl"
 
 ////////////////////////////////////////////////////////////////////////////////
-const TCHAR* GDefaultChannels = TEXT("cpu,gpu,frame,log,bookmark,screenshot");
+const TCHAR* GDefaultChannels = TEXT("cpu,gpu,frame,log,bookmark,screenshot,region");
 const TCHAR* GMemoryChannels = TEXT("memtag,memalloc,callstack,module");
 const TCHAR* GTraceConfigSection = TEXT("Trace.Config");
 
diff --git a/Engine/Source/Runtime/Core/Public/ProfilingDebugging/MiscTrace.h b/Engine/Source/Runtime/Core/Public/ProfilingDebugging/MiscTrace.h
--- a/Engine/Source/Runtime/Core/Public/ProfilingDebugging/MiscTrace.h
+++ b/Engine/Source/Runtime/Core/Public/ProfilingDebugging/MiscTrace.h
@@ -61,6 +61,9 @@ struct FMiscTrace
 		}
 	}
 
+	CORE_API static void OutputBeginRegion(const TCHAR* RegionName);
+	CORE_API static void OutputEndRegion(const TCHAR* RegionName);
+	
 	CORE_API static void OutputBeginFrame(ETraceFrameType FrameType);
 	CORE_API static void OutputEndFrame(ETraceFrameType FrameType);
 
@@ -70,6 +73,7 @@ struct FMiscTrace
 
 private:
 	CORE_API static void OutputBookmarkInternal(const void* BookmarkPoint, uint16 EncodedFormatArgsSize, uint8* EncodedFormatArgs);
+	
 };
 
 #define TRACE_BOOKMARK(Format, ...) \
@@ -81,6 +85,12 @@ private:
 	} \
 	FMiscTrace::OutputBookmark(&PREPROCESSOR_JOIN(__BookmarkPoint, __LINE__), ##__VA_ARGS__);
 
+#define TRACE_BEGIN_REGION(RegionName) \
+	FMiscTrace::OutputBeginRegion(RegionName);
+
+#define TRACE_END_REGION(RegionName) \
+	FMiscTrace::OutputEndRegion(RegionName);
+
 #define TRACE_BEGIN_FRAME(FrameType) \
 	FMiscTrace::OutputBeginFrame(FrameType);
 
@@ -99,6 +109,8 @@ private:
 #else
 
 #define TRACE_BOOKMARK(...)
+#define TRACE_BEGIN_REGION(...)
+#define TRACE_END_REGION(...)
 #define TRACE_BEGIN_FRAME(...)
 #define TRACE_END_FRAME(...)
 #define TRACE_SCREENSHOT(...)
