Updated TPagedArray:
* Fixed TPagedArray::Emplace to properly construct the element in the first available page.
* Updated API to match TArray's:
** Added Emplace_GetRef
** Added Add_GetRef
** Emplace now returns the insertion index.
** Add now returns the insertion index.
** Added a Reset overload with a parameter capacity.
** Added an Empty overload with a parameter capacity.
** Updated code to use UE_NODISCARD

* Updated unit tests
** Added Reserve Emplace multiple pages
** Added Emplace variants
** Added Empty with target capacity
** Added Empty vs Reset with target capacity

#tests passed unit tests

diff --git a/Engine/Source/Runtime/Core/Public/Containers/PagedArray.h b/Engine/Source/Runtime/Core/Public/Containers/PagedArray.h
--- a/Engine/Source/Runtime/Core/Public/Containers/PagedArray.h
+++ b/Engine/Source/Runtime/Core/Public/Containers/PagedArray.h
@@ -200,7 +200,7 @@ public:
 	 *
 	 * @returns Number of bytes allocated by this container.
 	 */
-	FORCEINLINE SIZE_T GetAllocatedSize() const
+	[[nodiscard]] FORCEINLINE SIZE_T GetAllocatedSize() const
 	{
 		SIZE_T Size = 0;
 		Size += Pages.GetAllocatedSize();
@@ -211,34 +211,34 @@ public:
 		return Size;
 	}
 
-	FORCEINLINE SizeType Max() const
+	[[nodiscard]] FORCEINLINE SizeType Max() const
 	{
 		return Pages.Num() * PageTraits::Capacity;
 	}
 
-	FORCEINLINE SizeType Num() const
+	[[nodiscard]] FORCEINLINE SizeType Num() const
 	{
 		return Count;
 	}
 
-	FORCEINLINE SizeType NumPages() const
+	[[nodiscard]] FORCEINLINE SizeType NumPages() const
 	{
 		return Pages.Num();
 	}
 
-	FORCEINLINE bool IsEmpty() const
+	[[nodiscard]] FORCEINLINE bool IsEmpty() const
 	{
 		return !Count;
 	}
 
-	FORCEINLINE const ElementType& Last() const
+	[[nodiscard]] FORCEINLINE const ElementType& Last() const
 	{
 		CheckValidIndex(0);
 		const SizeType LastIndex = Num() - 1;
 		return Pages[GetPageIndex(LastIndex)][GetPageOffset(LastIndex)];
 	}
 
-	FORCEINLINE ElementType& Last()
+	[[nodiscard]] FORCEINLINE ElementType& Last()
 	{
 		CheckValidIndex(0);
 		const SizeType LastIndex = Num() - 1;
@@ -380,10 +380,26 @@ public:
 	 * This method returns a reference to the constructed element.
 	 */
 	template <typename... ArgsType>
-	ElementType& Emplace(ArgsType&&... Args)
+	SizeType Emplace(ArgsType&&... Args)
 	{
 		GrowIfRequired();
-		ElementType& Result = Pages.Last().Emplace_GetRef(Forward<ArgsType>(Args)...);
+		const SizeType ElementIndex = Num();
+		const SizeType PageIndex = GetPageIndex(ElementIndex);
+		Pages[PageIndex].Emplace(Forward<ArgsType>(Args)...);
+		++Count;
+		return ElementIndex;
+	}
+
+	/*
+	 * Constructs an element in place using the parameter arguments and adds it at the back of the container.
+	 * This method returns a reference to the constructed element.
+	 */
+	template <typename... ArgsType>
+	[[nodiscard]] ElementType& Emplace_GetRef(ArgsType&&... Args)
+	{
+		GrowIfRequired();
+		const SizeType PageIndex = GetPageIndex(Num());
+		ElementType& Result = Pages[PageIndex].Emplace_GetRef(Forward<ArgsType>(Args)...);
 		++Count;
 		return Result;
 	}
@@ -391,14 +407,24 @@ public:
 	/*
 	 * Adds the parameter element at the back of the container.
 	 */
-	FORCEINLINE void Add(const ElementType& Element)
+	FORCEINLINE SizeType Add(const ElementType& Element)
 	{
-		Emplace(Element);
+		return Emplace(Element);
 	}
 
-	FORCEINLINE void Add(ElementType&& Element)
+	FORCEINLINE SizeType Add(ElementType&& Element)
 	{
-		Emplace(MoveTempIfPossible(Element));
+		return Emplace(MoveTempIfPossible(Element));
+	}
+
+	[[nodiscard]] FORCEINLINE ElementType& Add_GetRef(const ElementType& Element)
+	{
+		return Emplace_GetRef(Element);
+	}
+
+	[[nodiscard]] FORCEINLINE ElementType& Add_GetRef(ElementType&& Element)
+	{
+		return Emplace_GetRef(MoveTempIfPossible(Element));
 	}
 
 	FORCEINLINE void Push(const ElementType& Element)
@@ -466,6 +492,23 @@ public:
 		Count = 0;
 	}
 
+	/*
+	 * Empties the container effectively destroying all contained elements and ensures storage for the parameter
+	 * capacity. Storage is adjusted to meet the parameter capacity value.
+	 */
+	void Empty(SizeType InCapacity)
+	{
+		Reset();
+
+		const SizeType PageCount = Pages.Num();
+		const SizeType RequiredPageCount = NumRequiredPages(InCapacity);
+		Pages.SetNum(RequiredPageCount);
+		for (SizeType Index = PageCount; Index < RequiredPageCount; ++Index)
+		{
+			Pages[Index].Reserve(PageTraits::Capacity);
+		}
+	}
+
 	/*
 	 * Destroys all contained elements but doesn't release the container's storage.
 	 */
@@ -478,6 +521,16 @@ public:
 		Count = 0;
 	}
 
+	/*
+	 * Destroys all contained elements and ensures storage for the parameter capacity. Storage is only acquired
+	 * if the current container's storage cannot meet the parameter capacity value.
+	 */
+	void Reset(SizeType InCapacity)
+	{
+		Reset();
+		Reserve(InCapacity);
+	}
+
 	/*
 	 * Reserves storage for the parameter element count.
 	 */
diff --git a/Engine/Source/Runtime/Core/Tests/Containers/PagedArrayTest.cpp b/Engine/Source/Runtime/Core/Tests/Containers/PagedArrayTest.cpp
--- a/Engine/Source/Runtime/Core/Tests/Containers/PagedArrayTest.cpp
+++ b/Engine/Source/Runtime/Core/Tests/Containers/PagedArrayTest.cpp
@@ -74,6 +74,38 @@ TEST_CASE(
 	CHECK(int32Token::EvenConstructionDestructionCalls(1));
 }
 
+TEST_CASE(
+	"System::Core::Containers::TPagedArray::Reserve Emplace multiple pages",
+	"[SmokeFilter][Core][Containers][PagedArray]")
+{
+	int32Token::Reset();
+	{
+		using ArrayType = TPagedArray<int32Token, 128>;
+		ArrayType PagedArray;
+		CHECK(!PagedArray.Max());
+		CHECK(!PagedArray.Num());
+
+		PagedArray.Reserve(0);	// innocuous
+		PagedArray.Reserve(1);
+		PagedArray.Emplace(0);
+		const int32Token* ElementAddress = &PagedArray[0];
+	
+		CHECK(PagedArray.Max() >= 1);
+		CHECK(PagedArray.Max() == ArrayType::MaxPerPage());
+		CHECK(PagedArray.Num() == 1);
+
+		PagedArray.Reset();
+		PagedArray.Reserve(ArrayType::MaxPerPage() + 1);
+		PagedArray.Emplace(0);
+		const int32Token* NewElementAddress = &PagedArray[0];
+
+		CHECK(NewElementAddress == ElementAddress);
+		CHECK(PagedArray.Max() == ArrayType::MaxPerPage() * 2);
+		CHECK(PagedArray.Num() == 1);
+	}
+	CHECK(int32Token::EvenConstructionDestructionCalls(2));
+}
+
 TEST_CASE("System::Core::Containers::TPagedArray::Add", "[SmokeFilter][Core][Containers][PagedArray]")
 {
 	int32Token::Reset();
@@ -121,6 +153,44 @@ TEST_CASE(
 	CHECK(int32Token::EvenConstructionDestructionCalls(5));
 }
 
+TEST_CASE(
+	"System::Core::Containers::TPagedArray::Emplace variants",
+	"[SmokeFilter][Core][Containers][PagedArray]")
+{
+	using ArrayType = TPagedArray<int32Token, 128>;
+	int32Token::Reset();
+	{
+		ArrayType PagedArray;
+		int32Token* CheckAddress = nullptr;
+		int32Token* ElementAddress = nullptr;
+		int32 CheckIndex = INDEX_NONE;
+
+		// These are compile-error checks: uncomment to test
+		// PagedArray.Emplace_GetRef(); //nodiscard compile-error
+		// PagedArray.Add_GetRef(); //nodiscard compile-error
+
+		ElementAddress = &PagedArray.Emplace_GetRef(0);
+		CheckAddress = &PagedArray[0];
+		CHECK(ElementAddress == CheckAddress);
+		CheckIndex = PagedArray.Emplace(13);
+		CHECK(CheckIndex == PagedArray.Num() - 1);
+		CHECK(PagedArray[CheckIndex] == 13);
+
+		PagedArray.Reset(PagedArray.Max() * 3);
+		ElementAddress = &PagedArray.Emplace_GetRef(0);
+		CheckAddress = &PagedArray[0];
+		CHECK(ElementAddress == CheckAddress);
+
+		for (int32 i = 0; i < ArrayType::MaxPerPage(); ++i)
+		{
+			CheckIndex = PagedArray.Emplace(i + 1);
+			CHECK(CheckIndex == PagedArray.Num() - 1);
+			CHECK(PagedArray[CheckIndex] == i + 1);
+		}
+	}
+	CHECK(int32Token::EvenConstructionDestructionCalls(ArrayType::MaxPerPage() + 3));
+}
+
 TEST_CASE("System::Core::Containers::TPagedArray::Reset", "[SmokeFilter][Core][Containers][PagedArray]")
 {
 	int32Token::Reset();
@@ -219,6 +289,89 @@ TEST_CASE(
 	CHECK(int32Token::EvenConstructionDestructionCalls(1));
 }
 
+TEST_CASE(
+	"System::Core::Containers::TPagedArray::Empty with target capacity",
+	"[SmokeFilter][Core][Containers][PagedArray]")
+{
+	using ArrayType = TPagedArray<int32Token, 128>;
+
+	int32Token::Reset();
+	{
+		ArrayType PagedArray;
+		PagedArray.Emplace(0);
+		CHECK(PagedArray.Max());
+		CHECK(PagedArray.Num() == 1);
+		CHECK(!PagedArray.IsEmpty());
+		
+		PagedArray.Empty(ArrayType::MaxPerPage() + 1);
+		CHECK(int32Token::EvenConstructionDestructionCalls(1));
+
+		CHECK(PagedArray.Max() == ArrayType::MaxPerPage() * 2);
+		CHECK(!PagedArray.Num());
+		CHECK(PagedArray.IsEmpty());
+
+		PagedArray.SetNum(ArrayType::MaxPerPage() * 3 + 1);
+		CHECK(PagedArray.Max() == ArrayType::MaxPerPage() * 4);
+		
+		PagedArray.Empty(ArrayType::MaxPerPage() + 1);
+		CHECK(PagedArray.Max() == ArrayType::MaxPerPage() * 2);
+		CHECK(!PagedArray.Num());
+		CHECK(PagedArray.IsEmpty());
+
+		PagedArray.Empty(ArrayType::MaxPerPage());
+		CHECK(PagedArray.Max() == ArrayType::MaxPerPage());
+		CHECK(!PagedArray.Num());
+		CHECK(PagedArray.IsEmpty());
+
+		PagedArray.Empty(0);
+		CHECK(!PagedArray.Max());
+		CHECK(!PagedArray.Num());
+		CHECK(PagedArray.IsEmpty());
+	}
+	CHECK(int32Token::EvenConstructionDestructionCalls(ArrayType::MaxPerPage() * 3 + 2));
+}
+
+TEST_CASE(
+	"System::Core::Containers::TPagedArray::Empty vs Reset with target capacity",
+	"[SmokeFilter][Core][Containers][PagedArray]")
+{
+	using ArrayType = TPagedArray<int32Token, 128>;
+
+	int32Token::Reset();
+	{
+		ArrayType PagedArray;
+		
+		// Requesting higher capacity through Reset or Empty have the same effect
+		PagedArray.Emplace(0);
+		CHECK(PagedArray.Max() == ArrayType::MaxPerPage());
+
+		PagedArray.Reset(ArrayType::MaxPerPage() + 1);
+		CHECK(PagedArray.Max() == ArrayType::MaxPerPage() * 2);
+		CHECK(!PagedArray.Num());
+		CHECK(PagedArray.IsEmpty());
+
+		PagedArray.Empty(ArrayType::MaxPerPage() + 1);
+		CHECK(PagedArray.Max() == ArrayType::MaxPerPage() * 2);
+		CHECK(!PagedArray.Num());
+		CHECK(PagedArray.IsEmpty());
+
+		// Requesting lower capacity through Reset or Empty DO NOT have the same effect: Empty shrinks
+		PagedArray.SetNum(ArrayType::MaxPerPage() * 4);
+		CHECK(PagedArray.Max() == ArrayType::MaxPerPage() * 4);
+
+		PagedArray.Reset(ArrayType::MaxPerPage() + 1);
+		CHECK(PagedArray.Max() == ArrayType::MaxPerPage() * 4);
+		CHECK(!PagedArray.Num());
+		CHECK(PagedArray.IsEmpty());
+
+		PagedArray.Empty(ArrayType::MaxPerPage() + 1);
+		CHECK(PagedArray.Max() == ArrayType::MaxPerPage() * 2);
+		CHECK(!PagedArray.Num());
+		CHECK(PagedArray.IsEmpty());
+	}
+	CHECK(int32Token::EvenConstructionDestructionCalls(ArrayType::MaxPerPage() * 4 + 1));
+}
+
 TEST_CASE("System::Core::Containers::TPagedArray::Indexing", "[SmokeFilter][Core][Containers][PagedArray]")
 {
 	using ArrayType = TPagedArray<int32Token, 128>;
