removed the old FTicker that was deprecated in 5.0

diff --git a/Engine/Source/Runtime/Core/Private/Containers/BackgroundableTicker.cpp b/Engine/Source/Runtime/Core/Private/Containers/BackgroundableTicker.cpp
--- a/Engine/Source/Runtime/Core/Private/Containers/BackgroundableTicker.cpp
+++ b/Engine/Source/Runtime/Core/Private/Containers/BackgroundableTicker.cpp
@@ -41,44 +41,3 @@ FTSBackgroundableTicker::~FTSBackgroundableTicker()
 	FTSTicker::GetCoreTicker().RemoveTicker(CoreTickerHandle);
 	FCoreDelegates::MobileBackgroundTickDelegate.Remove(BackgroundTickerHandle);
 }
-
-PRAGMA_DISABLE_DEPRECATION_WARNINGS
-
-//////////////////////////////////////////////////////////////////////////////////////////////////////////////////////////////////////
-// deprecated non thread-safe version
-
-FBackgroundableTicker& FBackgroundableTicker::GetCoreTicker()
-{
-	static FBackgroundableTicker Singleton;
-	return Singleton;
-}
-
-FBackgroundableTicker::FBackgroundableTicker()
-{
-	CoreTickerHandle = FTicker::GetCoreTicker().AddTicker(TEXT("FBackgroundableTicker"), 0.0f, [this](float DeltaTime) -> bool
-	{
-		QUICK_SCOPE_CYCLE_COUNTER(STAT_FBackgroundableTicker_ForegroundTick);
-		if (bWasBackgrounded)
-		{
-			// When returning to the foreground, ensure we do not report enormous delta time values coming from the foreground ticker.
-			bWasBackgrounded = false;
-			DeltaTime = FMath::Clamp(DeltaTime, 0.0f, 1.0f / 60);
-		}
-		Tick(DeltaTime);
-		return true;
-	});
-	BackgroundTickerHandle = FCoreDelegates::MobileBackgroundTickDelegate.AddLambda([this](float DeltaTime)
-	{
-		QUICK_SCOPE_CYCLE_COUNTER(STAT_FBackgroundableTicker_BackgroundTick);
-		bWasBackgrounded = true;
-		Tick(DeltaTime);
-	});
-}
-
-FBackgroundableTicker::~FBackgroundableTicker()
-{
-	FTicker::GetCoreTicker().RemoveTicker(CoreTickerHandle);
-	FCoreDelegates::MobileBackgroundTickDelegate.Remove(BackgroundTickerHandle);
-}
-
-PRAGMA_ENABLE_DEPRECATION_WARNINGS
diff --git a/Engine/Source/Runtime/Core/Private/Containers/Ticker.cpp b/Engine/Source/Runtime/Core/Private/Containers/Ticker.cpp
--- a/Engine/Source/Runtime/Core/Private/Containers/Ticker.cpp
+++ b/Engine/Source/Runtime/Core/Private/Containers/Ticker.cpp
@@ -130,14 +130,6 @@ void FTSTicker::Tick(float DeltaTime)
 
 	Elements.Reset();
 	Exchange(TickedElements, Elements);
-
-	if (this == &GetCoreTicker())
-	{
-		// tick also deprecated FTicker for backward compatibility
-		PRAGMA_DISABLE_DEPRECATION_WARNINGS
-		FTicker::GetCoreTicker().Tick(DeltaTime);
-		PRAGMA_ENABLE_DEPRECATION_WARNINGS
-	}
 }
 
 void FTSTicker::Reset()
@@ -176,166 +168,3 @@ FTSTickerObjectBase::~FTSTickerObjectBase()
 {
 	FTSTicker::RemoveTicker(TickHandle);
 }
-
-//////////////////////////////////////////////////////////////////////////////////////////////////////////////////////////////////////
-// deprecated non thread-safe version
-
-PRAGMA_DISABLE_DEPRECATION_WARNINGS
-
-FTicker::FTicker()
-	: CurrentTime(0.0)
-	, bInTick(false)
-	, bCurrentElementRemoved(false)
-{
-}
-
-FTicker::~FTicker()
-{
-}
-
-FDelegateHandle FTicker::AddTicker(const FTickerDelegate& InDelegate, float InDelay)
-{
-	// We can add elements safely even during tick.
-	Elements.Emplace(CurrentTime + InDelay, InDelay, InDelegate);
-	// @todo this needs a unique handle for each add call to allow you to register the same delegate twice with a different delay safely.
-	// because of this, RemoveTicker removes all elements that use this handle.
-	return InDelegate.GetHandle();
-}
-
-
-FDelegateHandle FTicker::AddTicker(const TCHAR* InName, float InDelay, TFunction<bool(float)> Function)
-{
-	checkf(IsInGameThread(), TEXT("FTicker is not thread-safe and will be deprecated, please use FTSTicker"));
-
-	// todo - use InName for profiling. Added in sig to be forward looking..
-
-	FTickerDelegate Delegate = FTickerDelegate::CreateLambda(Function);
-
-	// We can add elements safely even during tick.
-	Elements.Emplace(CurrentTime + InDelay, InDelay, Delegate);
-
-	return Delegate.GetHandle();
-}
-
-void FTicker::RemoveTicker(FDelegateHandle Handle)
-{
-	checkf(IsInGameThread(), TEXT("FTicker is not thread-safe and will be deprecated, please use FTSTicker"));
-
-	// must remove the handle from both arrays because we could be in the middle of a tick, 
-	// and may be removing ourselves or something else we've already considered this Tick
-	auto CompareHandle = [=](const FElement& Element){ return Element.Delegate.GetHandle() == Handle; };
-	// We can remove elements safely even during tick.
-	Elements.RemoveAllSwap(CompareHandle);
-	TickedElements.RemoveAllSwap(CompareHandle);
-	// if we are ticking, we must check for the edge case of the CurrentElement removing itself.
-	if (bInTick && CompareHandle(CurrentElement))
-	{
-		// Technically it's possible for someone to try to remove CurrentDelegate multiple times, so make sure we never set this value to false in here.
-		bCurrentElementRemoved = true;
-	}
-}
-
-void FTicker::Tick(float DeltaTime)
-{
-	SCOPE_TIME_GUARD(TEXT("FTicker::Tick"));
-
-	QUICK_SCOPE_CYCLE_COUNTER(STAT_FTicker_Tick);
-	if (!Elements.Num())
-	{
-		return;
-	}
-
-	// make sure we scope the "InTick" state
-	TGuardValue<bool> TickGuard(bInTick, true);
-
-	CurrentTime += DeltaTime;
-	// we will keep popping off elements until the array is empty.
-	// We cannot keep references or iterators into any of these containers because the act of firing a delegate could
-	// add or remove any other delegate, including itself. 
-	// As we pop off elements, we track the CurrentElement in case it tries to remove itself.
-	// When finished, we put it onto a TickedElements list. This allows another delegate in the list to remove it,
-	// while still keeping track of elements we've ticked by popping off the Elements array.
-	// See RemoveTicker for details on how this state is used.
-	while (Elements.Num())
-	{
-		// as an optimization, if the Element is not ready to fire, move it immediately to the TickedElements list.
-		// This is safe because no one can mutate the arrays while we do this.
-		if (Elements.Last().FireTime > CurrentTime)
-		{
-			TickedElements.Add(Elements.Pop(false));
-		}
-		else
-		{
-			// we now know the element is ready to fire, so pop it off the list.
-			CurrentElement = Elements.Pop(false);
-			// reset this state every time we reassign to CurrentElement.
-			bCurrentElementRemoved = false;
-			// fire the delegate. It can return false to tell us to remove it immediately.
-			bool bRemoveElement = !CurrentElement.Fire(DeltaTime);
-			// The act of firing the delegate could also have caused it to remove itself.
-			// if either happens, we just won't add this to the TickedElements array.
-			if (!bRemoveElement && !bCurrentElementRemoved)
-			{
-				// update the fire time.
-				// Note this is where Timer skew occurs. Use FireTime += DelayTime if skew is not wanted.
-				CurrentElement.FireTime = CurrentTime + CurrentElement.DelayTime;
-				// push this element onto the TickedElement list so we know we already ticked it.
-				TickedElements.Push(CurrentElement);
-			}
-		}
-	}
-
-	// Now that we've considered all the delegates, we swap it back into the Elements array.
-	Exchange(TickedElements, Elements);
-	// Also clear the CurrentElement delegate as our tick is done
-	CurrentElement.Delegate.Unbind();
-}
-
-FTicker::FElement::FElement() 
-	: FireTime(0.0)
-	, DelayTime(0.0f)
-{}
-
-FTicker::FElement::FElement(double InFireTime, float InDelayTime, const FTickerDelegate& InDelegate) 
-	: FireTime(InFireTime)
-	, DelayTime(InDelayTime)
-	, Delegate(InDelegate)
-{}
-
-bool FTicker::FElement::Fire(float DeltaTime)
-{
-	if (Delegate.IsBound())
-	{
-		return Delegate.Execute(DeltaTime);
-	}
-	return false;
-}
-
-#if PLATFORM_WINDOWS && PLATFORM_32BITS
-// Workaround for ICE on VC++ 2017 14.13.26128 for UnrealGame Win32
-PRAGMA_DISABLE_OPTIMIZATION
-#endif
-
-FTickerObjectBase::FTickerObjectBase(float InDelay, FTicker& InTicker)
-	: Ticker(InTicker)
-{
-	// Register delegate for ticker callback
-	FTickerDelegate TickDelegate = FTickerDelegate::CreateRaw(this, &FTickerObjectBase::Tick);
-	TickHandle = Ticker.AddTicker(TickDelegate, InDelay);
-}
-
-#if PLATFORM_WINDOWS && PLATFORM_32BITS
-PRAGMA_ENABLE_OPTIMIZATION
-#endif
-
-FTickerObjectBase::~FTickerObjectBase()
-{
-	// Unregister ticker delegate
-	if (TickHandle != FDelegateHandle())
-	{
-		Ticker.RemoveTicker(TickHandle);
-		TickHandle = FDelegateHandle();
-	}
-}
-
-PRAGMA_ENABLE_DEPRECATION_WARNINGS
diff --git a/Engine/Source/Runtime/Core/Private/Misc/CoreMisc.cpp b/Engine/Source/Runtime/Core/Private/Misc/CoreMisc.cpp
--- a/Engine/Source/Runtime/Core/Private/Misc/CoreMisc.cpp
+++ b/Engine/Source/Runtime/Core/Private/Misc/CoreMisc.cpp
@@ -221,24 +221,6 @@ bool WillNeedAudioVisualData()
 	return true;
 }
 
-//-----------------------------------------------------------------------------
-
-PRAGMA_DISABLE_DEPRECATION_WARNINGS
-
-class FCoreTicker : public FTicker {};
-
-FTicker& FTicker::GetCoreTicker()
-{
-	return TLazySingleton<FCoreTicker>::Get();
-}
-
-void FTicker::TearDownCoreTicker()
-{
-	TLazySingleton<FCoreTicker>::TearDown();
-}
-
-PRAGMA_ENABLE_DEPRECATION_WARNINGS
-
 /*----------------------------------------------------------------------------
 	Runtime functions.
 ----------------------------------------------------------------------------*/
diff --git a/Engine/Source/Runtime/Core/Public/Containers/BackgroundableTicker.h b/Engine/Source/Runtime/Core/Public/Containers/BackgroundableTicker.h
--- a/Engine/Source/Runtime/Core/Public/Containers/BackgroundableTicker.h
+++ b/Engine/Source/Runtime/Core/Public/Containers/BackgroundableTicker.h
@@ -24,31 +24,3 @@ private:
 	::FDelegateHandle BackgroundTickerHandle;
 	bool bWasBackgrounded = false;
 };
-
-PRAGMA_DISABLE_DEPRECATION_WARNINGS
-
-/**
- * DEPRECATED not thread-safe version
- * For migration guide see `FTSTicker`
- * 
- * This works the same as the core FTicker, but on supported mobile platforms 
- * it continues ticking while the app is running in the background.
- */
-class UE_DEPRECATED(5.0, "Please use thread-safe FTSTickerObjectBase, see the migration guide in the code comment above FTSTicker") FBackgroundableTicker
-	: public FTicker
-{
-public:
-
-	CORE_API static FBackgroundableTicker& GetCoreTicker();
-
-	CORE_API FBackgroundableTicker();
-	CORE_API ~FBackgroundableTicker();
-
-private:
-	
-	FDelegateHandle CoreTickerHandle;
-	FDelegateHandle BackgroundTickerHandle;
-	bool bWasBackgrounded = false;
-};
-
-PRAGMA_ENABLE_DEPRECATION_WARNINGS
diff --git a/Engine/Source/Runtime/Core/Public/Containers/Ticker.h b/Engine/Source/Runtime/Core/Public/Containers/Ticker.h
--- a/Engine/Source/Runtime/Core/Public/Containers/Ticker.h
+++ b/Engine/Source/Runtime/Core/Public/Containers/Ticker.h
@@ -159,156 +159,3 @@ private:
 	/** Delegate for callbacks to Tick */
 	FTSTicker::FDelegateHandle TickHandle;
 };
-
-/**
- * DEPRECATED
- * Not thread-safe ticker class. Fires delegates after a delay.
- * 
- * Note: Do not try to add the same delegate instance twice, as there is no way to remove only a single instance (see member RemoveTicker).
- * 
- * Migration guide:
- * As FTicker is deprecated and replaced by thread-safe FTSTicker class, if you see a deprecation warning please migrate your code. Change
- * `FTicker` to `FTSTicker` and change the type of the delegate handle returned by `AddTicker` methods from `FDelegateHandle` to
- * `FTSTicker::FDelegateHandle`
- */
-class UE_DEPRECATED(5.0, "Please use thread-safe FTSTicker, see the migration guide in the code comment above") FTicker
-{
-public:
-
-	/** Singleton used for the ticker in Core / Launch. If you add a new ticker for a different subsystem, do not put that singleton here! **/
-	CORE_API static FTicker& GetCoreTicker();
-	CORE_API static void TearDownCoreTicker();
-
-	CORE_API FTicker();
-	CORE_API ~FTicker();
-
-	/**
-	 * Add a new ticker with a given delay / interval
-	 * 
-	 * @param InDelegate Delegate to fire after the delay
-	 * @param InDelay Delay until next fire; 0 means "next frame"
-	 */
-	CORE_API FDelegateHandle AddTicker(const FTickerDelegate& InDelegate, float InDelay = 0.0f);	
-
-	/**
-	* Add a new ticker with a given delay / interval.
-	*
-	* @param InName Name of this ticker for profiling
-	* @param InDelay Delay until next fire; 0 means "next frame"
-	* @param Function Function to execute. Should return true to fire after another InDelay time
-	*/
-	CORE_API FDelegateHandle AddTicker(const TCHAR* InName, float InDelay, TFunction<bool(float)> Function);
-
-	/**
-	 * Removes a previously added ticker delegate.
-	 * 
-	 * Note: will remove ALL tickers that use this handle, as there's no way to uniquely identify which one you are trying to remove.
-	 *
-	 * @param Handle The handle of the ticker to remove.
-	 */
-	CORE_API void RemoveTicker(FDelegateHandle Handle);
-
-	/**
-	 * Fire all tickers who have passed their delay and reschedule the ones that return true
-	 * 
-	 * Note: This reschedule has timer skew, meaning we always wait a full Delay period after 
-	 * each invocation even if we missed the last interval by a bit. For instance, setting a 
-	 * delay of 0.1 seconds will not necessarily fire the delegate 10 times per second, depending 
-	 * on the Tick() rate. What this DOES guarantee is that a delegate will never be called MORE
-	 * FREQUENTLY than the Delay interval, regardless of how far we miss the scheduled interval.
-	 * 
-	 * @param DeltaTime	time that has passed since the last tick call
-	 */
-	CORE_API void Tick(float DeltaTime);
-
-
-protected:
-	/**
-	 * Element of the priority queue
-	 */
-	struct FElement
-	{
-		/** Time that this delegate must not fire before **/
-		double FireTime;
-		/** Delay that this delegate was scheduled with. Kept here so that if the delegate returns true, we will reschedule it. **/
-		float DelayTime;
-		/** Delegate to fire **/
-		FTickerDelegate Delegate;
-
-		/** Default ctor is only required to implement CurrentElement handling without making it a pointer. */
-		CORE_API FElement();
-		/** This is the ctor that the code will generally use. */
-		CORE_API FElement(double InFireTime, float InDelayTime, const FTickerDelegate& InDelegate);
-		
-		/** Fire the delegate if it is fireable **/
-		CORE_API bool Fire(float DeltaTime);
-	};
-
-	/** Current time of the ticker **/
-	double CurrentTime;
-	/** Future delegates to fire **/
-	TArray<FElement, TInlineAllocator<1>> Elements;
-	/** List of delegates that have been considered during Tick. Either they been fired or are not yet ready.*/
-	TArray<FElement, TInlineAllocator<1>> TickedElements;
-	/** Current element being ticked (only valid during tick). */
-	FElement CurrentElement;
-	/** State to track whether CurrentElement is valid. */
-	bool bInTick;
-	/** State to track whether the CurrentElement removed itself during its own delegate execution. */
-	bool bCurrentElementRemoved;
-};
-
-/**
- * Base class for ticker objects
- */
-class UE_DEPRECATED(5.0, "Please use thread-safe FTSTickerObjectBase, see the migration guide in the code comment above FTSTicker") FTickerObjectBase
-{
-	PRAGMA_DISABLE_DEPRECATION_WARNINGS // FTicker
-
-public:
-
-	/**
-	 * Constructor
-	 *
-	 * @param InDelay Delay until next fire; 0 means "next frame"
-	 * @param Ticker the ticker to register with. Defaults to FTicker::GetCoreTicker().
-	*/
-#if ( PLATFORM_WINDOWS && defined(__clang__) )
-	CORE_API FTickerObjectBase()							// @todo clang: non-default argument constructors needed to prevent an ICE in clang
-		: FTickerObjectBase(0.0f, FTicker::GetCoreTicker())	// https://llvm.org/bugs/show_bug.cgi?id=28137
-	{
-	}
-
-	CORE_API FTickerObjectBase(float InDelay)
-		: FTickerObjectBase(InDelay, FTicker::GetCoreTicker())
-	{
-	}
-
-	CORE_API FTickerObjectBase(float InDelay, FTicker& Ticker);
-#else
-	CORE_API FTickerObjectBase(float InDelay = 0.0f, FTicker& Ticker = FTicker::GetCoreTicker());
-#endif
-
-	/** Virtual destructor. */
-	CORE_API virtual ~FTickerObjectBase();
-
-	/**
-	 * Pure virtual that must be overloaded by the inheriting class.
-	 *
-	 * @param DeltaTime	time passed since the last call.
-	 * @return true if should continue ticking
-	 */
-	virtual bool Tick(float DeltaTime) = 0;
-
-private:
-	// noncopyable idiom
-	FTickerObjectBase(const FTickerObjectBase&);
-	FTickerObjectBase& operator=(const FTickerObjectBase&);
-
-	/** Ticker to register with */
-	FTicker& Ticker;
-	/** Delegate for callbacks to Tick */
-	FDelegateHandle TickHandle;
-
-	PRAGMA_ENABLE_DEPRECATION_WARNINGS
-};
diff --git a/Engine/Source/Runtime/Core/Tests/Containers/TickerTests.cpp b/Engine/Source/Runtime/Core/Tests/Containers/TickerTests.cpp
--- a/Engine/Source/Runtime/Core/Tests/Containers/TickerTests.cpp
+++ b/Engine/Source/Runtime/Core/Tests/Containers/TickerTests.cpp
@@ -11,33 +11,6 @@
 #include "Tests/TestHarnessAdapter.h"
 
 
-template<uint32 NumDelegates, uint32 NumTicks>
-void TickerPerfTest()
-{
-	PRAGMA_DISABLE_DEPRECATION_WARNINGS
-
-		FTicker Ticker;
-
-	TArray<FDelegateHandle> DelegateHandles;
-	DelegateHandles.Reserve(NumDelegates);
-	for (uint32 i = 0; i != NumDelegates; ++i)
-	{
-		DelegateHandles.Add(Ticker.AddTicker(UE_SOURCE_LOCATION, 0.0f, [](float DeltaTime) { return true; }));
-	}
-
-	for (uint32 i = 0; i != NumTicks; ++i)
-	{
-		Ticker.Tick(0.0f);
-	}
-
-	for (FDelegateHandle& DelegateHandle : DelegateHandles)
-	{
-		Ticker.RemoveTicker(DelegateHandle);
-	}
-
-	PRAGMA_ENABLE_DEPRECATION_WARNINGS
-}
-
 template<uint32 NumDelegates, uint32 NumTicks>
 void TSTickerPerfTest()
 {
@@ -150,30 +123,6 @@ TEST_CASE_NAMED(FTSTickerTest,"System::Core::Containers::TSTicker", "[Applicatio
 		FTSTicker::RemoveTicker(DelegateHandle);
 	}
 
-	{	// demonstrate that the old ticker calls a delegate in the same tick that it was added
-		PRAGMA_DISABLE_DEPRECATION_WARNINGS;
-
-		FTicker Ticker;
-		bool bTicked = false;
-		Ticker.AddTicker(UE_SOURCE_LOCATION, 0.0f,
-			[&Ticker, &bTicked](float)
-			{
-				Ticker.AddTicker(UE_SOURCE_LOCATION, 0.0f,
-				[&bTicked](float)
-					{
-						bTicked = true;
-		return false;
-					}
-		);
-		return false;
-			}
-		);
-		Ticker.Tick(0.0f);
-		check(bTicked);
-
-		PRAGMA_ENABLE_DEPRECATION_WARNINGS;
-	}
-
 	{	// check that delegate is called in the same tick that it was added, for backward compatibility with the previous implementation
 		FTSTicker Ticker;
 		bool bTicked = false;
@@ -245,7 +194,6 @@ TEST_CASE_NAMED(FTSTickerTest,"System::Core::Containers::TSTicker", "[Applicatio
 		verify(Wait(Tasks, FTimespan::FromSeconds(5)));
 	}
 
-	UE_BENCHMARK(5, TickerPerfTest<100, 100>);
 	UE_BENCHMARK(5, TSTickerPerfTest<100, 100>);
 
 }
