Added FParkingLot, FMutex, TUniqueLock, TDynamicUniqueLock

FParkingLot is a framework for building locks and other synchronization primitives, based on WTF::ParkingLot from WebKit.

FMutex is a one-byte mutex that is neither fair nor recursive, based on WTF::Lock from WebKit.

TUniqueLock and TDynamicUniqueLock are generic types to manage mutex locking, based on std::unique_lock.

diff --git a/Engine/Source/Programs/UnrealBuildTool/Platform/Windows/UEBuildWindows.cs b/Engine/Source/Programs/UnrealBuildTool/Platform/Windows/UEBuildWindows.cs
--- a/Engine/Source/Programs/UnrealBuildTool/Platform/Windows/UEBuildWindows.cs
+++ b/Engine/Source/Programs/UnrealBuildTool/Platform/Windows/UEBuildWindows.cs
@@ -1612,6 +1612,7 @@ namespace UnrealBuildTool
 			LinkEnvironment.SystemLibraries.Add("netapi32.lib");
 			LinkEnvironment.SystemLibraries.Add("iphlpapi.lib");
 			LinkEnvironment.SystemLibraries.Add("setupapi.lib"); //  Required for access monitor device enumeration
+			LinkEnvironment.SystemLibraries.Add("synchronization.lib"); // Required for WaitOnAddress and WakeByAddressSingle
 
 			// Windows 7 Desktop Windows Manager API for Slate Windows Compliance
 			LinkEnvironment.SystemLibraries.Add("dwmapi.lib");
diff --git a/Engine/Source/Runtime/Core/Private/Async/Mutex.cpp b/Engine/Source/Runtime/Core/Private/Async/Mutex.cpp
new file mode 100644
--- /dev/null
+++ b/Engine/Source/Runtime/Core/Private/Async/Mutex.cpp
@@ -0,0 +1,87 @@
+// Copyright Epic Games, Inc. All Rights Reserved.
+
+#include "Async/Mutex.h"
+
+#include "Async/ParkingLot.h"
+#include "HAL/PlatformProcess.h"
+
+namespace UE
+{
+
+void FMutex::LockSlow()
+{
+	constexpr int32 SpinLimit = 40;
+	int32 SpinCount = 0;
+	for (;;)
+	{
+		uint8 CurrentState = State.load(std::memory_order_acquire);
+
+		// Try to acquire the lock if it was unlocked, even if there are waiting threads.
+		// Acquiring the lock despite the waiting threads means that this lock is not FIFO and thus not fair.
+		if (!(CurrentState & IsLockedFlag))
+		{
+			if (State.compare_exchange_weak(CurrentState, CurrentState | IsLockedFlag, std::memory_order_acquire))
+			{
+				return;
+			}
+			continue;
+		}
+
+		// Spin up to the spin limit while there are no waiting threads.
+		if (!(CurrentState & HasWaitingThreadsFlag) && SpinCount < SpinLimit)
+		{
+			FPlatformProcess::YieldThread();
+			++SpinCount;
+			continue;
+		}
+
+		// Store that there are waiting threads. Restart if the state has changed since it was loaded.
+		if (!(CurrentState & HasWaitingThreadsFlag))
+		{
+			if (!State.compare_exchange_weak(CurrentState, CurrentState | HasWaitingThreadsFlag, std::memory_order_acquire))
+			{
+				continue;
+			}
+			CurrentState |= HasWaitingThreadsFlag;
+		}
+
+		// Wait if the state has not changed. Either way, loop back and try to acquire the lock trying to wait.
+		FParkingLot::Wait(&State, [this, CurrentState] { return State.load(std::memory_order_acquire) == CurrentState; }, []{});
+	}
+}
+
+void FMutex::UnlockSlow()
+{
+	uint8 CurrentState;
+
+	for (;;)
+	{
+		CurrentState = State.load(std::memory_order_acquire);
+
+		checkSlow(CurrentState & IsLockedFlag);
+
+		// Spin on the fast path because there may be spurious failures.
+		if (CurrentState == IsLockedFlag)
+		{
+			if (State.compare_exchange_weak(CurrentState, 0, std::memory_order_release))
+			{
+				return;
+			}
+			continue;
+		}
+
+		// There is at least one thread waiting. Wake one thread and return.
+		FParkingLot::WakeOne(&State, [this](FParkingLotWakeState WakeState) -> uint64
+		{
+			constexpr uint8 ExpectedState = IsLockedFlag | HasWaitingThreadsFlag;
+			const uint8 NewState = WakeState.bHasWaitingThreads ? HasWaitingThreadsFlag : 0;
+			uint8 OldState = ExpectedState;
+			State.compare_exchange_strong(OldState, NewState, std::memory_order_release);
+			checkSlow(OldState == ExpectedState);
+			return 0;
+		});
+		return;
+	}
+}
+
+} // UE
diff --git a/Engine/Source/Runtime/Core/Private/Async/ParkingLot.cpp b/Engine/Source/Runtime/Core/Private/Async/ParkingLot.cpp
new file mode 100644
--- /dev/null
+++ b/Engine/Source/Runtime/Core/Private/Async/ParkingLot.cpp
@@ -0,0 +1,724 @@
+// Copyright Epic Games, Inc. All Rights Reserved.
+
+#include "Async/ParkingLot.h"
+
+#include "Algo/Sort.h"
+#include "Async/UniqueLock.h"
+#include "Containers/Array.h"
+#include "Containers/ArrayView.h"
+#include "Templates/RefCounting.h"
+#include <atomic>
+
+#define UE_PARKINGLOT_USE_WAITONADDRESS PLATFORM_WINDOWS
+#define UE_PARKINGLOT_USE_WORDMUTEX 0
+
+#if UE_PARKINGLOT_USE_WORDMUTEX
+#include "Async/WordMutex.h"
+using FParkingLotLock = UE::Private::FWordMutex;
+#else
+#include "HAL/CriticalSection.h"
+using FParkingLotLock = FCriticalSection;
+#endif
+
+#if UE_PARKINGLOT_USE_WAITONADDRESS
+#include "Microsoft/WindowsHWrapper.h"
+#else
+#include <condition_variable>
+#include <mutex>
+#endif
+
+namespace UE::Private
+{
+
+///////////////////////////////////////////////////////////////////////////////////////////////////////////////////////
+
+/**
+ * A thread as stored in the wait queue.
+ */
+struct FParkingLotThread final
+{
+	FParkingLotThread* Next = nullptr;
+
+#if UE_PARKINGLOT_USE_WAITONADDRESS
+	const void* volatile WaitAddress = nullptr;
+#else
+	const void* WaitAddress = nullptr;
+	std::mutex Lock;
+	std::condition_variable Condition;
+#endif
+
+	uint64 WakeToken = 0;
+	mutable std::atomic<uint32> ReferenceCount = 0;
+
+	inline void AddRef() const
+	{
+		ReferenceCount.fetch_add(1, std::memory_order_relaxed);
+	}
+
+	inline void Release() const
+	{
+		if (ReferenceCount.fetch_sub(1, std::memory_order_acq_rel) == 1)
+		{
+			delete this;
+		}
+	}
+
+	FParkingLotThread() = default;
+
+private:
+	~FParkingLotThread() = default;
+
+	FParkingLotThread(const FParkingLotThread&) = delete;
+	FParkingLotThread& operator=(const FParkingLotThread&) = delete;
+};
+
+class FParkingLotThreadLocalData
+{
+public:
+	// TODO: This version reserves the table as threads spawn. The function-local thread_local
+	//       reserves the tables as threads wait for the first time. Not sure which is better.
+	//static thread_local FParkingLotThreadLocalData ThreadLocalData;
+
+	static FParkingLotThread& Get()
+	{
+		static thread_local FParkingLotThreadLocalData ThreadLocalData;
+		FParkingLotThreadLocalData& LocalThreadLocalData = ThreadLocalData;
+		if (!LocalThreadLocalData.Thread)
+		{
+			LocalThreadLocalData.Thread = new FParkingLotThread;
+		}
+		return *LocalThreadLocalData.Thread;
+	}
+
+private:
+	inline static std::atomic<uint32> ThreadCount = 0;
+
+	TRefCountPtr<FParkingLotThread> Thread;
+
+	FParkingLotThreadLocalData();
+	~FParkingLotThreadLocalData();
+
+	FParkingLotThreadLocalData(const FParkingLotThreadLocalData&) = delete;
+	FParkingLotThreadLocalData& operator=(const FParkingLotThreadLocalData&) = delete;
+};
+
+//thread_local FParkingLotThreadLocalData FParkingLotThreadLocalData::ThreadLocalData;
+
+///////////////////////////////////////////////////////////////////////////////////////////////////////////////////////
+
+enum class EParkingLotAction
+{
+	Stop,
+	Continue,
+	RemoveAndStop,
+	RemoveAndContinue,
+};
+
+/**
+ * A bucket in the hash table keyed by memory address.
+ *
+ * Buckets must be locked to access the list of waiting threads.
+ * Buckets are aligned to a cache line to reduce false sharing.
+ */
+class alignas(PLATFORM_CACHE_LINE_SIZE) FParkingLotBucket final
+{
+public:
+	FParkingLotBucket() = default;
+	FParkingLotBucket(const FParkingLotBucket&) = delete;
+	FParkingLotBucket& operator=(const FParkingLotBucket&) = delete;
+
+	[[nodiscard]] inline TDynamicUniqueLock<FParkingLotLock> LockDynamic() { return TDynamicUniqueLock(Mutex); }
+
+	inline void Lock() { Mutex.Lock(); }
+	inline void Unlock() { Mutex.Unlock(); }
+
+	inline bool IsEmpty() const { return !Head; }
+
+	void Enqueue(FParkingLotThread* Thread);
+	FParkingLotThread* Dequeue();
+
+	/**
+	 * Dequeues threads based on a visitor.
+	 *
+	 * Visitor signature is EParkingLotAction (FParkingLotThread&).
+	 * Visitor is called for every thread in the bucket, from head to tail.
+	 * Threads are dequeued if the returned action contains Remove.
+	 * Visiting stops if the returned action contains Stop.
+	 */
+	template <typename VisitorType>
+	void DequeueIf(VisitorType&& Visitor);
+
+private:
+	FParkingLotLock Mutex;
+	FParkingLotThread* Head = nullptr;
+	FParkingLotThread* Tail = nullptr;
+};
+
+void FParkingLotBucket::Enqueue(FParkingLotThread* Thread)
+{
+	checkSlow(Thread);
+	checkSlow(!Thread->Next);
+	if (Tail)
+	{
+		Tail->Next = Thread;
+		Tail = Thread;
+	}
+	else
+	{
+		Head = Thread;
+		Tail = Thread;
+	}
+}
+
+FParkingLotThread* FParkingLotBucket::Dequeue()
+{
+	FParkingLotThread* Thread = Head;
+	if (Thread)
+	{
+		Head = Thread->Next;
+		Thread->Next = nullptr;
+		if (Tail == Thread)
+		{
+			Tail = nullptr;
+		}
+	}
+	return Thread;
+}
+
+template <typename VisitorType>
+void FParkingLotBucket::DequeueIf(VisitorType&& Visitor)
+{
+	FParkingLotThread** Next = &Head;
+	FParkingLotThread* Prev = nullptr;
+	while (FParkingLotThread* Thread = *Next)
+	{
+		switch (Visitor(Thread))
+		{
+		case EParkingLotAction::Stop:
+			return;
+		case EParkingLotAction::Continue:
+			Prev = Thread;
+			Next = &Thread->Next;
+			break;
+		case EParkingLotAction::RemoveAndStop:
+			if (Tail == Thread)
+			{
+				Tail = Prev;
+			}
+			*Next = Thread->Next;
+			Thread->Next = nullptr;
+			return;
+		case EParkingLotAction::RemoveAndContinue:
+			if (Tail == Thread)
+			{
+				Tail = Prev;
+			}
+			*Next = Thread->Next;
+			Thread->Next = nullptr;
+			break;
+		}
+	}
+}
+
+///////////////////////////////////////////////////////////////////////////////////////////////////////////////////////
+
+/**
+ * A hash table of queues of waiting threads keyed by memory address.
+ *
+ * Tables are never freed. The size of the table is bounded by the maximum number of threads that
+ * exist concurrently and have used the wait queue. The table grows by powers of two, which means
+ * that the maximum size leaked is less than the maximum size that the table ever grows to. Table
+ * leaks are also limited in size because a table is an array of bucket pointers, and the buckets
+ * are reused when the table grows.
+ */
+class FParkingLotTable final
+{
+	/** Minimum bucket count to create a table with. */
+	constexpr static uint32 MinSize = 32;
+
+public:
+	/** Find or create, and lock, the bucket for the memory address. */
+	static FParkingLotBucket& FindOrCreateBucket(const void* Address, TDynamicUniqueLock<FParkingLotLock>& OutLock);
+
+	/** Find and lock the bucket for the memory address. Returns null if the bucket was not created. */
+	static FParkingLotBucket* FindBucket(const void* Address, TDynamicUniqueLock<FParkingLotLock>& OutLock);
+
+	/** Reserve memory for the table to handle at least ThreadCount waiting threads. */
+	static void Reserve(uint32 ThreadCount);
+
+private:
+	/** Returns the current global table, creating it if it does not exist. */
+	static FParkingLotTable& CreateOrGet();
+
+	/** Create a new table with the specified bucket count. */
+	static FParkingLotTable& Create(uint32 Size);
+	/** Destroy a table. Must not be called on any table that has been globally visible. */
+	static void Destroy(FParkingLotTable& Table);
+
+	/** Try to lock the whole table by locking every one of its buckets. */
+	static bool TryLock(FParkingLotTable& Table, TArray<FParkingLotBucket*>& OutBuckets);
+	/** Unlock an array of buckets that was filled by TryLock. */
+	static void Unlock(TConstArrayView<FParkingLotBucket*> LockedBuckets);
+
+	/** Calculate a 32-bit hash from the memory address. */
+	static uint32 HashAddress(const void* Address);
+
+	/** Pointer to the current global table. Previous global tables are leaked. */
+	inline static std::atomic<FParkingLotTable*> GlobalTable;
+
+	FParkingLotTable() = default;
+	~FParkingLotTable() = default;
+	FParkingLotTable(const FParkingLotTable&) = delete;
+	FParkingLotTable& operator=(const FParkingLotTable&) = delete;
+
+	/** Find or create the bucket at the index. */
+	template <typename AllocatorType>
+	FParkingLotBucket& FindOrCreateBucket(uint32 Index, AllocatorType&& Allocator);
+
+	uint32 BucketCount = 0;
+	std::atomic<FParkingLotBucket*> Buckets[0];
+};
+
+FParkingLotBucket& FParkingLotTable::FindOrCreateBucket(const void* Address, TDynamicUniqueLock<FParkingLotLock>& OutLock)
+{
+	const uint32 Hash = HashAddress(Address);
+
+	for (;;)
+	{
+		FParkingLotTable& Table = CreateOrGet();
+		const uint32 Index = Hash % Table.BucketCount;
+		FParkingLotBucket& Bucket = Table.FindOrCreateBucket(Index, [] { return new FParkingLotBucket; });
+		OutLock = Bucket.LockDynamic();
+
+		if (LIKELY(&Table == GlobalTable.load(std::memory_order_acquire)))
+		{
+			return Bucket;
+		}
+
+		// Restart because the table was resized since it was loaded above.
+		OutLock.Unlock();
+	}
+}
+
+FParkingLotBucket* FParkingLotTable::FindBucket(const void* Address, TDynamicUniqueLock<FParkingLotLock>& OutLock)
+{
+	const uint32 Hash = HashAddress(Address);
+
+	for (;;)
+	{
+		if (FParkingLotTable* Table = GlobalTable.load(std::memory_order_acquire))
+		{
+			const uint32 Index = Hash % Table->BucketCount;
+			if (FParkingLotBucket* Bucket = Table->Buckets[Index].load(std::memory_order_acquire))
+			{
+				OutLock = Bucket->LockDynamic();
+
+				if (LIKELY(Table == GlobalTable.load(std::memory_order_acquire)))
+				{
+					return Bucket;
+				}
+
+				// Restart because the table was resized since it was loaded above.
+				OutLock = {};
+				continue;
+			}
+		}
+		return nullptr;
+	}
+}
+
+template <typename AllocatorType>
+FParkingLotBucket& FParkingLotTable::FindOrCreateBucket(uint32 Index, AllocatorType&& Allocator)
+{
+	std::atomic<FParkingLotBucket*>& BucketPtr = Buckets[Index];
+	FParkingLotBucket* Bucket = BucketPtr.load(std::memory_order_acquire);
+	if (UNLIKELY(!Bucket))
+	{
+		FParkingLotBucket* NewBucket = Allocator();
+		if (BucketPtr.compare_exchange_strong(Bucket, NewBucket, std::memory_order_release, std::memory_order_acquire))
+		{
+			Bucket = NewBucket;
+		}
+		else
+		{
+			delete NewBucket;
+		}
+		checkSlow(Bucket);
+	}
+	return *Bucket;
+}
+
+FParkingLotTable& FParkingLotTable::CreateOrGet()
+{
+	FParkingLotTable* Table = GlobalTable.load(std::memory_order_acquire);
+
+	if (LIKELY(Table))
+	{
+		return *Table;
+	}
+
+	FParkingLotTable& NewTable = Create(MinSize);
+
+	if (LIKELY(GlobalTable.compare_exchange_strong(Table, &NewTable, std::memory_order_release, std::memory_order_acquire)))
+	{
+		return NewTable;
+	}
+
+	Destroy(NewTable);
+
+	checkSlow(Table);
+	return *Table;
+}
+
+FParkingLotTable& FParkingLotTable::Create(const uint32 Size)
+{
+	const SIZE_T MemorySize = sizeof(FParkingLotTable) + sizeof(FParkingLotBucket*) * SIZE_T(Size);
+	void* const Memory = FMemory::Malloc(MemorySize, alignof(FParkingLotTable));
+	FMemory::Memzero(Memory, MemorySize);
+	FParkingLotTable& Table = *new(Memory) FParkingLotTable;
+	Table.BucketCount = Size;
+	return Table;
+}
+
+void FParkingLotTable::Destroy(FParkingLotTable& Table)
+{
+	Table.~FParkingLotTable();
+	FMemory::Free(&Table);
+}
+
+void FParkingLotTable::Reserve(const uint32 ThreadCount)
+{
+	const uint32 TargetBucketCount = FMath::RoundUpToPowerOfTwo(ThreadCount);
+	TArray<FParkingLotBucket*> ExistingBuckets;
+
+	for (;;)
+	{
+		FParkingLotTable& ExistingTable = CreateOrGet();
+
+		if (LIKELY(ExistingTable.BucketCount >= TargetBucketCount))
+		{
+			// Reserve is called every time a thread is created and has amortized constant time
+			// because of its power-of-two table growth. Most calls return here without locking.
+			return;
+		}
+
+		if (!TryLock(ExistingTable, ExistingBuckets))
+		{
+			continue;
+		}
+
+		// Gather waiting threads to be redistributed into the buckets of the new table.
+		// Threads with the same address remain in the same relative order as they were queued.
+		TArray<FParkingLotThread*> Threads;
+		for (FParkingLotBucket* Bucket : ExistingBuckets)
+		{
+			while (FParkingLotThread* Thread = Bucket->Dequeue())
+			{
+				Threads.Add(Thread);
+			}
+		}
+
+		FParkingLotTable& NewTable = Create(TargetBucketCount);
+
+		// Reuse existing now-empty buckets when populating the new table.
+		TArray<FParkingLotBucket*> AvailableBuckets = ExistingBuckets;
+		const auto AllocateBucket = [&AvailableBuckets]() -> FParkingLotBucket*
+		{
+			return !AvailableBuckets.IsEmpty() ? AvailableBuckets.Pop() : new FParkingLotBucket;
+		};
+
+		// Add waiting threads to the new table.
+		for (FParkingLotThread* Thread : Threads)
+		{
+			const uint32 Hash = HashAddress(Thread->WaitAddress);
+			const uint32 Index = Hash % NewTable.BucketCount;
+			FParkingLotBucket& Bucket = NewTable.FindOrCreateBucket(Index, AllocateBucket);
+			Bucket.Enqueue(Thread);
+		}
+
+		// Assign any available buckets to the table to avoid having to free them.
+		for (uint32 Index = 0; !AvailableBuckets.IsEmpty() && Index < NewTable.BucketCount; ++Index)
+		{
+			NewTable.FindOrCreateBucket(Index, AllocateBucket);
+		}
+		checkSlow(AvailableBuckets.IsEmpty());
+
+		// Make the new table visible to other threads.
+		FParkingLotTable* CompareTable = GlobalTable.exchange(&NewTable, std::memory_order_release);
+		checkSlow(CompareTable == &ExistingTable);
+
+		// Unlock buckets that came from the existing table now that the new table is visible.
+		Unlock(ExistingBuckets);
+		return;
+	}
+}
+
+bool FParkingLotTable::TryLock(FParkingLotTable& Table, TArray<FParkingLotBucket*>& OutBuckets)
+{
+	OutBuckets.Reset();
+
+	// Gather buckets from the table, creating them as needed because the lock is on the bucket.
+	for (uint32 Index = 0; Index < Table.BucketCount; ++Index)
+	{
+		OutBuckets.Add(&Table.FindOrCreateBucket(Index, [] { return new FParkingLotBucket; }));
+	}
+
+	// Lock the buckets in order by address to ensure consistent ordering regardless of the table being locked.
+	Algo::Sort(OutBuckets);
+	for (FParkingLotBucket* Bucket : OutBuckets)
+	{
+		Bucket->Lock();
+	}
+
+	// Table is locked if the global table pointer still points to it, otherwise it has grown.
+	if (&Table == GlobalTable)
+	{
+		return true;
+	}
+
+	// Unlock and return that the table could not be locked.
+	Unlock(OutBuckets);
+	OutBuckets.Reset();
+	return false;
+}
+
+void FParkingLotTable::Unlock(TConstArrayView<FParkingLotBucket*> LockedBuckets)
+{
+	for (FParkingLotBucket* Bucket : LockedBuckets)
+	{
+		Bucket->Unlock();
+	}
+}
+
+uint32 FParkingLotTable::HashAddress(const void* Address)
+{
+	constexpr uint64 A = 0xdc2b17dc9d2fbc29;
+	constexpr uint64 B = 0xcb1014192cb2c5fc;
+	constexpr uint64 C = 0x5b12db9242bd7ce7;
+	const UPTRINT Value = UPTRINT(Address);
+	return uint32(((A * (Value >> 32)) + (B * (Value & 0xffffffff)) + C) >> 32);
+}
+
+///////////////////////////////////////////////////////////////////////////////////////////////////////////////////////
+
+FParkingLotThreadLocalData::FParkingLotThreadLocalData()
+{
+	FParkingLotTable::Reserve(ThreadCount.fetch_add(1, std::memory_order_relaxed) + 1);
+}
+
+FParkingLotThreadLocalData::~FParkingLotThreadLocalData()
+{
+	ThreadCount.fetch_sub(1, std::memory_order_relaxed);
+}
+
+///////////////////////////////////////////////////////////////////////////////////////////////////////////////////////
+
+} // UE::Private
+
+namespace UE
+{
+
+FParkingLotWaitState FParkingLot::PrivateWaitFor(const void* Address, TFunctionRef<bool()> CanWait, TFunctionRef<void()> BeforeWait, uint32 WaitMs)
+{
+	using namespace UE::Private;
+
+	FParkingLotThread& Self = FParkingLotThreadLocalData::Get();
+
+	checkfSlow(!Self.WaitAddress, TEXT("WaitAddress must be null. This can happen if Wait is called by BeforeWait."));
+	checkfSlow(Self.WakeToken == 0, TEXT("WakeToken must be 0. This is an error in FParkingLot."));
+
+	FParkingLotWaitState State;
+
+	// Enqueue the thread if CanWait returns true while the bucket is locked.
+	{
+		TDynamicUniqueLock<FParkingLotLock> BucketLock;
+		FParkingLotBucket& Bucket = FParkingLotTable::FindOrCreateBucket(Address, BucketLock);
+		State.bDidWait = CanWait();
+		if (!State.bDidWait)
+		{
+			return State;
+		}
+		Self.WaitAddress = Address;
+		Bucket.Enqueue(&Self);
+	}
+
+	// BeforeWait must be invoked after the bucket is unlocked.
+	BeforeWait();
+
+#if UE_PARKINGLOT_USE_WAITONADDRESS
+	for (;;)
+	{
+		const DWORD SystemWaitMs = (WaitMs == MAX_uint32) ? INFINITE : WaitMs;
+		if (!WaitOnAddress(&Self.WaitAddress, &Address, sizeof(Address), SystemWaitMs))
+		{
+			if (GetLastError() == ERROR_TIMEOUT)
+			{
+				break;
+			}
+		}
+		if (!Self.WaitAddress)
+		{
+			State.bDidWake = true;
+			State.WakeToken = Self.WakeToken;
+			Self.WakeToken = 0;
+			return State;
+		}
+	}
+#else
+	{
+		// Wait until the timeout or until the thread has been dequeued.
+		std::unique_lock SelfLock(Self.Lock);
+		if (WaitMs == MAX_uint32)
+		{
+			Self.Condition.wait(SelfLock, [&Self] { return !Self.WaitAddress; });
+		}
+		else
+		{
+			Self.Condition.wait_for(SelfLock, std::chrono::milliseconds(WaitMs), [&Self] { return !Self.WaitAddress; });
+		}
+
+		// WaitAddress is reset when the thread is dequeued.
+		if (!Self.WaitAddress)
+		{
+			State.bDidWake = true;
+			State.WakeToken = Self.WakeToken;
+			Self.WakeToken = 0;
+			return State;
+		}
+	}
+#endif
+
+	// The timeout was reached and the thread needs to dequeue itself.
+	// This can race with a call to wake a thread, which means Self is unsafe to access outside of the lock.
+	bool bDequeued = false;
+	if (TDynamicUniqueLock<FParkingLotLock> BucketLock; FParkingLotBucket* Bucket = FParkingLotTable::FindBucket(Address, BucketLock))
+	{
+		Bucket->DequeueIf([Self = &Self, &bDequeued](FParkingLotThread* Thread)
+		{
+			if (Thread == Self)
+			{
+				bDequeued = true;
+				Thread->WaitAddress = nullptr;
+				return EParkingLotAction::RemoveAndStop;
+			}
+			return EParkingLotAction::Continue;
+		});
+	}
+
+	// The thread did not dequeue itself, which means that we need to wait until the other thread
+	// has finished waking this thread by setting its wait address to null.
+	if (!bDequeued)
+	{
+	#if UE_PARKINGLOT_USE_WAITONADDRESS
+		while (Self.WaitAddress)
+		{
+			WaitOnAddress(&Self.WaitAddress, &Address, sizeof(Address), INFINITE);
+		}
+	#else
+		std::unique_lock SelfLock(Self.Lock);
+		Self.Condition.wait(SelfLock, [&Self] { return !Self.WaitAddress; });
+	#endif
+		State.bDidWake = true;
+		State.WakeToken = Self.WakeToken;
+		Self.WakeToken = 0;
+	}
+
+	return State;
+}
+
+void FParkingLot::PrivateWakeOne(const void* Address, TFunctionRef<uint64(FParkingLotWakeState)> OnWakeState)
+{
+	using namespace UE::Private;
+
+	TRefCountPtr<FParkingLotThread> WakeThread;
+	uint64 WakeToken = 0;
+
+	{
+		TDynamicUniqueLock<FParkingLotLock> BucketLock;
+		FParkingLotBucket& Bucket = FParkingLotTable::FindOrCreateBucket(Address, BucketLock);
+		Bucket.DequeueIf([Address, &WakeThread](FParkingLotThread* Thread)
+		{
+			if (Thread->WaitAddress == Address)
+			{
+				WakeThread = Thread;
+				return EParkingLotAction::RemoveAndStop;
+			}
+			return EParkingLotAction::Continue;
+		});
+		FParkingLotWakeState WakeState;
+		WakeState.bDidWake = !!WakeThread;
+		WakeState.bHasWaitingThreads = !Bucket.IsEmpty();
+		WakeToken = OnWakeState(WakeState);
+	}
+
+	if (WakeThread)
+	{
+		checkSlow(WakeThread->WaitAddress == Address);
+		WakeThread->WakeToken = WakeToken;
+	#if UE_PARKINGLOT_USE_WAITONADDRESS
+		WakeThread->WaitAddress = nullptr;
+		WakeByAddressSingle((void*)&WakeThread->WaitAddress);
+	#else
+		{
+			std::unique_lock WakeThreadLock(WakeThread->Lock);
+			WakeThread->WaitAddress = nullptr;
+		}
+		WakeThread->Condition.notify_one();
+	#endif
+	}
+}
+
+FParkingLotWakeState FParkingLot::WakeOne(const void* Address)
+{
+	FParkingLotWakeState OutState;
+	PrivateWakeOne(Address, [&OutState](FParkingLotWakeState State) -> uint64
+	{
+		OutState = State;
+		return 0;
+	});
+	return OutState;
+}
+
+uint32 FParkingLot::WakeMultiple(const void* const Address, const uint32 WakeCount)
+{
+	using namespace UE::Private;
+
+	TArray<TRefCountPtr<FParkingLotThread>> WakeThreads;
+
+	if (TDynamicUniqueLock<FParkingLotLock> BucketLock; FParkingLotBucket* Bucket = FParkingLotTable::FindBucket(Address, BucketLock))
+	{
+		Bucket->DequeueIf([Address, &WakeThreads, WakeCount](FParkingLotThread* Thread)
+		{
+			if (Thread->WaitAddress == Address)
+			{
+				const uint32 Count = uint32(WakeThreads.Add(Thread) + 1);
+				return Count == WakeCount ? EParkingLotAction::RemoveAndStop : EParkingLotAction::RemoveAndContinue;
+			}
+			return EParkingLotAction::Continue;
+		});
+	}
+
+	for (FParkingLotThread* WakeThread : WakeThreads)
+	{
+		checkSlow(WakeThread->WaitAddress == Address);
+	#if UE_PARKINGLOT_USE_WAITONADDRESS
+		WakeThread->WaitAddress = nullptr;
+		WakeByAddressSingle((void*)&WakeThread->WaitAddress);
+	#else
+		{
+			std::unique_lock WakeThreadLock(WakeThread->Lock);
+			WakeThread->WaitAddress = nullptr;
+		}
+		WakeThread->Condition.notify_one();
+	#endif
+	}
+
+	return uint32(WakeThreads.Num());
+}
+
+void FParkingLot::WakeAll(const void* Address)
+{
+	WakeMultiple(Address, MAX_uint32);
+}
+
+} // UE
diff --git a/Engine/Source/Runtime/Core/Public/Async/Mutex.h b/Engine/Source/Runtime/Core/Public/Async/Mutex.h
new file mode 100644
--- /dev/null
+++ b/Engine/Source/Runtime/Core/Public/Async/Mutex.h
@@ -0,0 +1,71 @@
+// Copyright Epic Games, Inc. All Rights Reserved.
+
+#pragma once
+
+#include "CoreTypes.h"
+#include <atomic>
+
+#define UE_API CORE_API
+
+namespace UE
+{
+
+/**
+ * A one-byte mutex that is not fair and does not support recursive locking.
+ */
+class FMutex final
+{
+public:
+	constexpr FMutex() = default;
+
+	FMutex(const FMutex&) = delete;
+	FMutex& operator=(const FMutex&) = delete;
+
+	inline bool IsLocked() const
+	{
+		return (State.load(std::memory_order_acquire) & IsLockedFlag);
+	}
+
+	inline bool TryLock()
+	{
+		uint8 Expected = State.load(std::memory_order_acquire);
+		if (LIKELY(!(Expected & IsLockedFlag) && State.compare_exchange_strong(Expected, Expected | IsLockedFlag, std::memory_order_acquire)))
+		{
+			return true;
+		}
+		return false;
+	}
+
+	inline void Lock()
+	{
+		uint8 Expected = 0;
+		if (LIKELY(State.compare_exchange_weak(Expected, IsLockedFlag, std::memory_order_acquire)))
+		{
+			return;
+		}
+		LockSlow();
+	}
+
+	inline void Unlock()
+	{
+		uint8 Expected = IsLockedFlag;
+		if (LIKELY(State.compare_exchange_weak(Expected, 0, std::memory_order_release)))
+		{
+			return;
+		}
+		UnlockSlow();
+	}
+
+private:
+	UE_API void LockSlow();
+	UE_API void UnlockSlow();
+
+	static constexpr uint8 IsLockedFlag = 1 << 0;
+	static constexpr uint8 HasWaitingThreadsFlag = 1 << 1;
+
+	std::atomic<uint8> State = 0;
+};
+
+} // UE
+
+#undef UE_API
diff --git a/Engine/Source/Runtime/Core/Public/Async/ParkingLot.h b/Engine/Source/Runtime/Core/Public/Async/ParkingLot.h
new file mode 100644
--- /dev/null
+++ b/Engine/Source/Runtime/Core/Public/Async/ParkingLot.h
@@ -0,0 +1,111 @@
+// Copyright Epic Games, Inc. All Rights Reserved.
+
+#pragma once
+
+#include "CoreTypes.h"
+#include "Math/NumericLimits.h"
+#include "Templates/Function.h"
+
+#define UE_API CORE_API
+
+namespace UE
+{
+
+struct FParkingLotWaitState final
+{
+	/** Did this thread wait? True only if CanWait returned true. */
+	bool bDidWait = false;
+	/** Did this wake from a wait? True only if a Wake call woke the thread, false for timeouts. */
+	bool bDidWake = false;
+	/** Optional value that was provided by the callback in WakeOne. */
+	uint64 WakeToken = 0;
+};
+
+struct FParkingLotWakeState final
+{
+	/** Did a thread wake up? */
+	bool bDidWake = false;
+	/** Does the queue have another thread waiting? */
+	bool bHasWaitingThreads = false;
+};
+
+/**
+ * A global table of queues of waiting threads keyed by memory address.
+ */
+class FParkingLot final
+{
+public:
+	/**
+	 * Queue the calling thread to wait if CanWait returns true. BeforeWait is only called if CanWait returns true.
+	 *
+	 * @param Address      Address to use as the key for the queue. The same address is used to wake the thread.
+	 * @param CanWait      Function called while the queue is locked. A return of false cancels the wait.
+	 * @param BeforeWait   Function called after the queue is unlocked and before the thread waits.
+	 */
+	template <typename CanWaitType, typename BeforeWaitType>
+	inline static FParkingLotWaitState Wait(const void* Address, CanWaitType&& CanWait, BeforeWaitType&& BeforeWait)
+	{
+		return PrivateWaitFor(Address, Forward<CanWaitType>(CanWait), Forward<BeforeWaitType>(BeforeWait), MAX_uint32);
+	}
+
+	/**
+	 * Queue the calling thread to wait if CanWait returns true. BeforeWait is only called if CanWait returns true.
+	 *
+	 * @param Address      Address to use as the key for the queue. The same address is used to wake the thread.
+	 * @param CanWait      Function called while the queue is locked. A return of false cancels the wait.
+	 * @param BeforeWait   Function called after the queue is unlocked and before the thread waits.
+	 * @param WaitMs       Minimum wait time after which waiting is automatically canceled and the thread wakes.
+	 */
+	template <typename CanWaitType, typename BeforeWaitType>
+	inline static FParkingLotWaitState WaitFor(const void* Address, CanWaitType&& CanWait, BeforeWaitType&& BeforeWait, uint32 WaitMs)
+	{
+		return PrivateWaitFor(Address, Forward<CanWaitType>(CanWait), Forward<BeforeWaitType>(BeforeWait), WaitMs);
+	}
+
+	/**
+	 * Wake one thread from the queue of threads waiting on the address.
+	 *
+	 * @param Address       Address to use as the key for the queue. Must match the address used in Wait.
+	 * @param OnWakeState   Function called while the queue is locked. Receives the wake state. Returns WakeToken.
+	 */
+	template <typename OnWakeStateType>
+	inline static void WakeOne(const void* Address, OnWakeStateType&& OnWakeState)
+	{
+		return PrivateWakeOne(Address, Forward<OnWakeStateType>(OnWakeState));
+	}
+
+	/**
+	 * Wake one thread from the queue of threads waiting on the address.
+	 *
+	 * @param Address   Address to use as the key for the queue. Must match the address used in Wait.
+	 * @return The wake state, which includes whether a thread woke up and whether there are more queued.
+	 */
+	UE_API static FParkingLotWakeState WakeOne(const void* Address);
+
+	/**
+	 * Wake up to WakeCount threads from the queue of threads waiting on the address.
+	 *
+	 * @param Address     Address to use as the key for the queue. Must match the address used in Wait.
+	 * @param WakeCount   The maximum number of threads to wake.
+	 * @return The number of threads that this call woke up.
+	 */
+	UE_API static uint32 WakeMultiple(const void* Address, uint32 WakeCount);
+
+	/**
+	 * Wake all threads from the queue of threads waiting on the address.
+	 *
+	 * @param Address   Address to use as the key for the queue. Must match the address used in Wait.
+	 */
+	UE_API static void WakeAll(const void* Address);
+
+private:
+	UE_API static FParkingLotWaitState PrivateWaitFor(const void* Address, TFunctionRef<bool()> CanWait, TFunctionRef<void()> BeforeWait, uint32 WaitMs);
+	UE_API static void PrivateWakeOne(const void* Address, TFunctionRef<uint64(FParkingLotWakeState)> OnWakeState);
+
+	FParkingLot() = delete;
+	FParkingLot(const FParkingLot&) = delete;
+};
+
+} // UE
+
+#undef UE_API
diff --git a/Engine/Source/Runtime/Core/Public/Async/UniqueLock.h b/Engine/Source/Runtime/Core/Public/Async/UniqueLock.h
new file mode 100644
--- /dev/null
+++ b/Engine/Source/Runtime/Core/Public/Async/UniqueLock.h
@@ -0,0 +1,117 @@
+// Copyright Epic Games, Inc. All Rights Reserved.
+
+#pragma once
+
+#include "Misc/AssertionMacros.h"
+
+namespace UE
+{
+
+struct FDeferLock final
+{
+	explicit FDeferLock() = default;
+};
+
+template <typename LockType>
+class TUniqueLock final
+{
+public:
+	TUniqueLock(const TUniqueLock&) = delete;
+	TUniqueLock& operator=(const TUniqueLock&) = delete;
+
+	[[nodiscard]] inline explicit TUniqueLock(LockType& Lock)
+		: Mutex(Lock)
+	{
+		Mutex.Lock();
+	}
+
+	inline ~TUniqueLock()
+	{
+		Mutex.Unlock();
+	}
+
+private:
+	LockType& Mutex;
+};
+
+template <typename LockType>
+class TDynamicUniqueLock final
+{
+public:
+	TDynamicUniqueLock() = default;
+
+	TDynamicUniqueLock(const TDynamicUniqueLock&) = delete;
+	TDynamicUniqueLock& operator=(const TDynamicUniqueLock&) = delete;
+
+	[[nodiscard]] inline explicit TDynamicUniqueLock(LockType& Lock)
+		: Mutex(&Lock)
+	{
+		Mutex->Lock();
+		bLocked = true;
+	}
+
+	[[nodiscard]] inline explicit TDynamicUniqueLock(LockType& Lock, FDeferLock)
+		: Mutex(&Lock)
+	{
+	}
+
+	[[nodiscard]] inline TDynamicUniqueLock(TDynamicUniqueLock&& Other)
+		: Mutex(Other.Mutex)
+		, bLocked(Other.bLocked)
+	{
+		Other.Mutex = nullptr;
+		Other.bLocked = false;
+	}
+
+	inline TDynamicUniqueLock& operator=(TDynamicUniqueLock&& Other)
+	{
+		if (bLocked)
+		{
+			Mutex->Unlock();
+		}
+		Mutex = Other.Mutex;
+		bLocked = Other.bLocked;
+		Other.Mutex = nullptr;
+		Other.bLocked = false;
+		return *this;
+	}
+
+	inline ~TDynamicUniqueLock()
+	{
+		if (bLocked)
+		{
+			Mutex->Unlock();
+		}
+	}
+
+	void Lock()
+	{
+		check(!bLocked);
+		check(Mutex);
+		Mutex->Lock();
+		bLocked = true;
+	}
+
+	void Unlock()
+	{
+		check(bLocked);
+		bLocked = false;
+		Mutex->Unlock();
+	}
+
+	inline bool OwnsLock() const
+	{
+		return bLocked;
+	}
+
+	inline explicit operator bool() const
+	{
+		return OwnsLock();
+	}
+
+private:
+	LockType* Mutex = nullptr;
+	bool bLocked = false;
+};
+
+} // UE
