[BugFix] "Diff against Depot" produces very noisy output. Export paths are included in text diffs which can show different paths depending on where the temp file was saved. This causes diffs that should be clean to have false positives. As a solution I added a PPF_ForDiff flag that ommits the ExportPath
#udn 00537121

diff --git a/Engine/Source/Developer/AssetTools/Private/AssetTools.cpp b/Engine/Source/Developer/AssetTools/Private/AssetTools.cpp
--- a/Engine/Source/Developer/AssetTools/Private/AssetTools.cpp
+++ b/Engine/Source/Developer/AssetTools/Private/AssetTools.cpp
@@ -2860,7 +2860,10 @@ FString UAssetToolsImpl::DumpAssetToTempFile(UObject* Asset) const
 	// Export asset to archive
 	if (Asset)
 	{
-		UExporter::ExportToOutputDevice(&Context, Asset, nullptr, Archive, TEXT("copy"), 0, PPF_ExportsNotFullyQualified|PPF_Copy|PPF_Delimited, false, Asset->GetOuter());
+		UExporter::ExportToOutputDevice(&Context, Asset, nullptr,
+			Archive, TEXT("copy"), 0,
+			PPF_ExportsNotFullyQualified|PPF_Copy|PPF_Delimited|PPF_ForDiff,
+			false, Asset->GetOuter());
 	}
 	// Used to generate unique file names during a run
 	static int TempFileNum = 0;
diff --git a/Engine/Source/Runtime/Core/Public/UObject/PropertyPortFlags.h b/Engine/Source/Runtime/Core/Public/UObject/PropertyPortFlags.h
--- a/Engine/Source/Runtime/Core/Public/UObject/PropertyPortFlags.h
+++ b/Engine/Source/Runtime/Core/Public/UObject/PropertyPortFlags.h
@@ -71,9 +71,13 @@ enum EPropertyPortFlags
 	/** Indicates that object property values should be exported without the package or class information */
 	PPF_SimpleObjectText			= 0x00002000,
 
+	PPF_ForDiff						= 0x00004000,
+	
 	/** parsing default properties - allow text for transient properties to be imported - also modifies ObjectProperty importing slightly for subobjects */
 	PPF_ParsingDefaultProperties	= 0x00008000,
 
+	//								= 0x00010000,
+	
 	/** indicates that non-categorized transient properties should be exported (by default, they would not be) */
 	PPF_IncludeTransient			= 0x00020000,
 
@@ -113,4 +117,6 @@ enum EPropertyPortFlags
 
 	/** Set when duplicating objects verbatim (doesn't reset unique IDs) */
 	PPF_DuplicateVerbatim			= 0x40000000,
+
+	//								= 0x80000000,
 };
diff --git a/Engine/Source/Runtime/Engine/Private/UnrealExporter.cpp b/Engine/Source/Runtime/Engine/Private/UnrealExporter.cpp
--- a/Engine/Source/Runtime/Engine/Private/UnrealExporter.cpp
+++ b/Engine/Source/Runtime/Engine/Private/UnrealExporter.cpp
@@ -472,9 +472,13 @@ void UExporter::EmitBeginObject( FOutputDevice& Ar, UObject* Obj, uint32 PortFla
 		}
 	}
 
-	// Emit the object path
-	Ar.Logf(TEXT(" ExportPath=%s"), *FObjectPropertyBase::GetExportPath(Obj, nullptr, nullptr, PortFlags | (PPF_Delimited  & ~PPF_ExportsNotFullyQualified)));
-
+	// When exporting for diffs, export paths can cause false positives. since diff files don't get imported, we can
+	// skip adding this info the file.
+	if (!(PortFlags & PPF_ForDiff))
+	{
+		// Emit the object path
+		Ar.Logf(TEXT(" ExportPath=%s"), *FObjectPropertyBase::GetExportPath(Obj, nullptr, nullptr, PortFlags | (PPF_Delimited  & ~PPF_ExportsNotFullyQualified)));
+	}
 	// end in a return
 	Ar.Logf(TEXT("\r\n"));
 
