Added lifetime-bound annotation to TOptional::Get()'s parameter, to catch dangling references when an rvalue argument is bound to an lvalue reference.

diff --git a/Engine/Source/Runtime/Core/Public/Misc/Optional.h b/Engine/Source/Runtime/Core/Public/Misc/Optional.h
--- a/Engine/Source/Runtime/Core/Public/Misc/Optional.h
+++ b/Engine/Source/Runtime/Core/Public/Misc/Optional.h
@@ -206,7 +206,7 @@ public:
 		  OptionalType& operator*()		  { return GetValue(); }
 
 	/** @return The optional value when set; DefaultValue otherwise. */
-	const OptionalType& Get(const OptionalType& DefaultValue) const { return IsSet() ? *(OptionalType*)&Value : DefaultValue; }
+	const OptionalType& Get(const OptionalType& DefaultValue UE_LIFETIMEBOUND) const UE_LIFETIMEBOUND { return IsSet() ? *(OptionalType*)&Value : DefaultValue; }
 
 	/** @return A pointer to the optional value when set, nullptr otherwise. */
 	OptionalType* GetPtrOrNull() { return IsSet() ? (OptionalType*)&Value : nullptr; }
