Add common metric conversion unit for km/m/cm/m2/cm2

diff --git a/Engine/Source/Runtime/Core/Public/Math/UnrealMathUtility.h b/Engine/Source/Runtime/Core/Public/Math/UnrealMathUtility.h
--- a/Engine/Source/Runtime/Core/Public/Math/UnrealMathUtility.h
+++ b/Engine/Source/Runtime/Core/Public/Math/UnrealMathUtility.h
@@ -172,6 +172,13 @@ class TRange;
 #define UE_DOUBLE_HALF_SQRT_2	(0.70710678118654752440084436210485)
 #define UE_DOUBLE_HALF_SQRT_3	(0.86602540378443864676372317075294)
 
+// Common metric unit conversion
+#define UE_KM_TO_M   (1000.f)
+#define UE_M_TO_KM   (0.001f)
+#define UE_CM_TO_M   (0.01f)
+#define UE_M_TO_CM   (100.f)
+#define UE_CM2_TO_M2 (0.0001f)
+#define UE_M2_TO_CM2 (10000.f)
 
 // Magic numbers for numerical precision.
 #define UE_DELTA		(0.00001f)
