Made ensure conditions likely when !DO_ENSURE, to match the DO_ENSURE case.

diff --git a/Engine/Source/Runtime/Core/Public/Misc/AssertionMacros.h b/Engine/Source/Runtime/Core/Public/Misc/AssertionMacros.h
--- a/Engine/Source/Runtime/Core/Public/Misc/AssertionMacros.h
+++ b/Engine/Source/Runtime/Core/Public/Misc/AssertionMacros.h
@@ -400,10 +400,10 @@ RetType FORCENOINLINE UE_DEBUG_SECTION DispatchCheckVerify(InnerType&& Inner, Ar
 
 #else	// DO_ENSURE
 
-	#define ensure(           InExpression                ) (!!(InExpression))
-	#define ensureMsgf(       InExpression, InFormat, ... ) (!!(InExpression))
-	#define ensureAlways(     InExpression                ) (!!(InExpression))
-	#define ensureAlwaysMsgf( InExpression, InFormat, ... ) (!!(InExpression))
+	#define ensure(           InExpression                ) (LIKELY(!!(InExpression)))
+	#define ensureMsgf(       InExpression, InFormat, ... ) (LIKELY(!!(InExpression)))
+	#define ensureAlways(     InExpression                ) (LIKELY(!!(InExpression)))
+	#define ensureAlwaysMsgf( InExpression, InFormat, ... ) (LIKELY(!!(InExpression)))
 
 #endif	// DO_CHECK
 
