Add serializers for TPagedArray

diff --git a/Engine/Source/Runtime/Core/Public/Containers/PagedArray.h b/Engine/Source/Runtime/Core/Public/Containers/PagedArray.h
--- a/Engine/Source/Runtime/Core/Public/Containers/PagedArray.h
+++ b/Engine/Source/Runtime/Core/Public/Containers/PagedArray.h
@@ -8,6 +8,7 @@
 #include "IteratorAdapter.h"
 #include "Math/UnrealMathUtility.h"
 #include "Misc/AssertionMacros.h"
+#include "Serialization/StructuredArchive.h"
 #include "Templates/UnrealTemplate.h"  // For GetData, GetNum
 
 #include <cstring>
@@ -17,7 +18,7 @@
 template <typename InElementType, int32 InPageSizeInBytes = 16384, typename InAllocatorType = FDefaultAllocator>
 class TPagedArray;
 
-namespace UE::PagedArray::Private
+namespace UE::Core::PagedArray::Private
 {
 
 /**
@@ -85,7 +86,63 @@ private:
 template <typename InElementType, typename InPageType, typename InPageTraits>
 using TIterator = TIteratorAdapter<TIteratorBase<InElementType, InPageType, InPageTraits>>;
 
-}  // namespace UE::PagedArray::Private
+/** Serializer. */
+template <typename ElementType, int32 PageSizeInBytes, typename AllocatorType>
+FArchive& Serialize(FArchive& Ar, TPagedArray<ElementType, PageSizeInBytes, AllocatorType>& InOutPagedArray)
+{
+	using SizeType = typename AllocatorType::SizeType;
+
+	InOutPagedArray.CountBytes(Ar);
+
+	SizeType NumElements = InOutPagedArray.Num();
+	Ar << NumElements;
+
+	if (Ar.IsLoading())
+	{
+		InOutPagedArray.Empty(NumElements);
+
+		for (SizeType ElementIndex = 0; ElementIndex < NumElements; ++ElementIndex)
+		{
+			Ar << InOutPagedArray.Emplace_GetRef();
+		}
+	}
+	else
+	{
+		for (ElementType& Element : InOutPagedArray)
+		{
+			Ar << Element;
+		}
+	}
+	return Ar;
+}
+
+/** Structured archive serializer. */
+template <typename ElementType, int32 PageSizeInBytes, typename AllocatorType>
+void SerializeStructured(FStructuredArchive::FSlot Slot, TPagedArray<ElementType, PageSizeInBytes, AllocatorType>& InOutPagedArray)
+{
+	int32 NumElements = InOutPagedArray.Num();
+	FStructuredArchive::FArray Array = Slot.EnterArray(NumElements);
+	if (Slot.GetUnderlyingArchive().IsLoading())
+	{
+		InOutPagedArray.Empty(NumElements);
+
+		for (int32 ElementIndex = 0; ElementIndex < NumElements; ++ElementIndex)
+		{
+			FStructuredArchive::FSlot ElementSlot = Array.EnterElement();
+			ElementSlot << InOutPagedArray.Emplace_GetRef();
+		}
+	}
+	else
+	{
+		for (ElementType& Element : InOutPagedArray)
+		{
+			FStructuredArchive::FSlot ElementSlot = Array.EnterElement();
+			ElementSlot << Element;
+		}
+	}
+}
+
+}  // namespace UE::Core::PagedArray::Private
 
 /**
  * Fixed size block allocated container class.
@@ -104,7 +161,7 @@ class TPagedArray
 	friend class TPagedArray;
 
 	using PageType = TArray<InElementType, InAllocatorType>;
-	using PageTraits = UE::PagedArray::Private::TPageTraits<InElementType, InPageSizeInBytes>;
+	using PageTraits = UE::Core::PagedArray::Private::TPageTraits<InElementType, InPageSizeInBytes>;
 
 public:
 	using AllocatorType = InAllocatorType;
@@ -133,8 +190,8 @@ private:
 	}
 
 public:
-	using ConstIteratorType = UE::PagedArray::Private::TIterator<const ElementType, const PageType, PageTraits>;
-	using IteratorType = UE::PagedArray::Private::TIterator<ElementType, PageType, PageTraits>;
+	using ConstIteratorType = UE::Core::PagedArray::Private::TIterator<const ElementType, const PageType, PageTraits>;
+	using IteratorType = UE::Core::PagedArray::Private::TIterator<ElementType, PageType, PageTraits>;
 
 	static constexpr SizeType MaxPerPage()
 	{
@@ -211,6 +268,16 @@ public:
 		return Size;
 	}
 
+	/**
+	 * Count bytes needed to serialize this paged array.
+	 *
+	 * @param Ar Archive to count for.
+	 */
+	void CountBytes(FArchive& Ar) const
+	{
+		Ar.CountBytes(Num() * sizeof(ElementType), Max() * sizeof(ElementType));
+	}
+
 	[[nodiscard]] FORCEINLINE SizeType Max() const
 	{
 		return Pages.Num() * PageTraits::Capacity;
@@ -731,3 +798,17 @@ private:
 		checkf((Index >= 0) & (Index < Count), TEXT("Parameter index %d exceeds container size %d"), Index, Count);
 	}
 };
+
+/** Serializer. */
+template <typename ElementType, int32 PageSizeInBytes, typename AllocatorType>
+FORCEINLINE FArchive& operator<<(FArchive& Ar, TPagedArray<ElementType, PageSizeInBytes, AllocatorType>& InOutPagedArray)
+{
+	return UE::Core::PagedArray::Private::Serialize(Ar, InOutPagedArray);
+}
+
+/** Structured archive serializer. */
+template <typename ElementType, int32 PageSizeInBytes, typename AllocatorType>
+FORCEINLINE void operator<<(FStructuredArchive::FSlot Slot, TPagedArray<ElementType, PageSizeInBytes, AllocatorType>& InOutPagedArray)
+{
+	UE::Core::PagedArray::Private::SerializeStructured(Slot, InOutPagedArray);
+}
\ No newline at end of file
diff --git a/Engine/Source/Runtime/Core/Tests/Containers/PagedArrayTest.cpp b/Engine/Source/Runtime/Core/Tests/Containers/PagedArrayTest.cpp
--- a/Engine/Source/Runtime/Core/Tests/Containers/PagedArrayTest.cpp
+++ b/Engine/Source/Runtime/Core/Tests/Containers/PagedArrayTest.cpp
@@ -4,6 +4,9 @@
 
 #include "Containers/PagedArray.h"
 #include "CoreMinimal.h"
+#include "Serialization/Formatters/BinaryArchiveFormatter.h"
+#include "Serialization/MemoryReader.h"
+#include "Serialization/MemoryWriter.h"
 #include "TestUtils.h"
 #include "TokenTest.h"
 
@@ -1394,6 +1397,228 @@ TEST_CASE("System::Core::Containers::TPagedArray::Iterator arithmetic", "[SmokeF
 	CHECK(int32Token::EvenConstructionDestructionCalls(5 + 3));
 }
 
+TEST_CASE("System::Core::Containers::TPagedArray::Serialize", "[SmokeFilter][Core][Containers][PagedArray]")
+{
+	using ArrayType = TPagedArray<int32Token, 16>;
+	int32Token::Reset();
+	// Container with no pages
+	{
+		ArrayType SourcePagedArray;
+		ArrayType DestPagedArray;
+
+		CHECK(SourcePagedArray.IsEmpty());
+		CHECK(DestPagedArray.IsEmpty());
+
+		TArray<uint8> Bytes;
+		{
+			FMemoryWriter MemoryWriter(Bytes);
+			MemoryWriter << SourcePagedArray;
+		}
+
+		// Just storing the count - no elements
+		CHECK(Bytes.Num() == sizeof(ArrayType::SizeType));
+
+		{
+			FMemoryReader MemoryReader(Bytes);
+			MemoryReader << DestPagedArray;
+		}
+
+		CHECK(SourcePagedArray == DestPagedArray);
+	}
+	CHECK(int32Token::EvenConstructionDestructionCalls(0));
+	int32Token::Reset();
+	// Container with single page
+	{
+		ArrayType SourcePagedArray;
+		ArrayType DestPagedArray;
+
+		SourcePagedArray.SetNum(ArrayType::MaxPerPage());
+		int32 CheckValue = 0;
+		for (int32Token& Value : SourcePagedArray)
+		{
+			Value = CheckValue++;
+		}
+
+		CHECK(SourcePagedArray.Num() == ArrayType::MaxPerPage());
+		CHECK(DestPagedArray.IsEmpty());
+
+		TArray<uint8> Bytes;
+		{
+			FMemoryWriter MemoryWriter(Bytes);
+			MemoryWriter << SourcePagedArray;
+		}
+
+		// Storing the count and the elements
+		CHECK(Bytes.Num() == (sizeof(ArrayType::SizeType) + (sizeof(ArrayType::ElementType) * ArrayType::MaxPerPage())));
+
+		{
+			FMemoryReader MemoryReader(Bytes);
+			MemoryReader << DestPagedArray;
+		}
+
+		CHECK(SourcePagedArray == DestPagedArray);
+	}
+	// We construct 3 times over - SetNum, assigning CheckValue and then reading into the dest array
+	CHECK(int32Token::EvenConstructionDestructionCalls(ArrayType::MaxPerPage() * 3));
+	int32Token::Reset();
+	// Container with multiple pages
+	{
+		ArrayType SourcePagedArray;
+		ArrayType DestPagedArray;
+
+		SourcePagedArray.SetNum(ArrayType::MaxPerPage() * 3);
+		int32 CheckValue = 0;
+		for (int32Token& Value : SourcePagedArray)
+		{
+			Value = CheckValue++;
+		}
+
+		CHECK(SourcePagedArray.Num() == ArrayType::MaxPerPage() * 3);
+		CHECK(DestPagedArray.IsEmpty());
+
+		TArray<uint8> Bytes;
+		{
+			FMemoryWriter MemoryWriter(Bytes);
+			MemoryWriter << SourcePagedArray;
+		}
+
+		// Storing the count and the elements
+		CHECK(Bytes.Num() == (sizeof(ArrayType::SizeType) + (sizeof(ArrayType::ElementType) * ArrayType::MaxPerPage() * 3)));
+
+		{
+			FMemoryReader MemoryReader(Bytes);
+			MemoryReader << DestPagedArray;
+		}
+
+		CHECK(SourcePagedArray == DestPagedArray);
+	}
+	// We construct 3 times over - SetNum, assigning CheckValue and then reading into the dest array
+	CHECK(int32Token::EvenConstructionDestructionCalls(ArrayType::MaxPerPage() * 3 * 3));
+}
+
+TEST_CASE("System::Core::Containers::TPagedArray::Structured Serialize", "[SmokeFilter][Core][Containers][PagedArray]")
+{
+	using ArrayType = TPagedArray<int32Token, 16>;
+	int32Token::Reset();
+	// Container with no pages
+	{
+		ArrayType SourcePagedArray;
+		ArrayType DestPagedArray;
+
+		CHECK(SourcePagedArray.IsEmpty());
+		CHECK(DestPagedArray.IsEmpty());
+
+		TArray<uint8> Bytes;
+		{
+			FMemoryWriter MemoryWriter(Bytes);
+			FBinaryArchiveFormatter BinaryFormatter(MemoryWriter);
+			FStructuredArchive StructuredArchive(BinaryFormatter);
+			FStructuredArchive::FSlot RootSlot = StructuredArchive.Open();
+			RootSlot << SourcePagedArray;
+			StructuredArchive.Close();
+		}
+
+		// Just storing the count - no elements
+		CHECK(Bytes.Num() == sizeof(ArrayType::SizeType));
+
+		{
+			FMemoryReader MemoryReader(Bytes);
+			FBinaryArchiveFormatter BinaryFormatter(MemoryReader);
+			FStructuredArchive StructuredArchive(BinaryFormatter);
+			FStructuredArchive::FSlot RootSlot = StructuredArchive.Open();
+			RootSlot << DestPagedArray;
+			StructuredArchive.Close();
+		}
+
+		CHECK(SourcePagedArray == DestPagedArray);
+	}
+	CHECK(int32Token::EvenConstructionDestructionCalls(0));
+	int32Token::Reset();
+	// Container with single page
+	{
+		ArrayType SourcePagedArray;
+		ArrayType DestPagedArray;
+
+		SourcePagedArray.SetNum(ArrayType::MaxPerPage());
+		int32 CheckValue = 0;
+		for (int32Token& Value : SourcePagedArray)
+		{
+			Value = CheckValue++;
+		}
+
+		CHECK(SourcePagedArray.Num() == ArrayType::MaxPerPage());
+		CHECK(DestPagedArray.IsEmpty());
+
+		TArray<uint8> Bytes;
+		{
+			FMemoryWriter MemoryWriter(Bytes);
+			FBinaryArchiveFormatter BinaryFormatter(MemoryWriter);
+			FStructuredArchive StructuredArchive(BinaryFormatter);
+			FStructuredArchive::FSlot RootSlot = StructuredArchive.Open();
+			RootSlot << SourcePagedArray;
+			StructuredArchive.Close();
+		}
+
+		// Storing the count and the elements
+		CHECK(Bytes.Num() == (sizeof(ArrayType::SizeType) + (sizeof(ArrayType::ElementType) * ArrayType::MaxPerPage())));
+
+		{
+			FMemoryReader MemoryReader(Bytes);
+			FBinaryArchiveFormatter BinaryFormatter(MemoryReader);
+			FStructuredArchive StructuredArchive(BinaryFormatter);
+			FStructuredArchive::FSlot RootSlot = StructuredArchive.Open();
+			RootSlot << DestPagedArray;
+			StructuredArchive.Close();
+		}
+
+		CHECK(SourcePagedArray == DestPagedArray);
+	}
+	// We construct 3 times over - SetNum, assigning CheckValue and then reading into the dest array
+	CHECK(int32Token::EvenConstructionDestructionCalls(ArrayType::MaxPerPage() * 3));
+	int32Token::Reset();
+	// Container with multiple pages
+	{
+		ArrayType SourcePagedArray;
+		ArrayType DestPagedArray;
+
+		SourcePagedArray.SetNum(ArrayType::MaxPerPage() * 3);
+		int32 CheckValue = 0;
+		for (int32Token& Value : SourcePagedArray)
+		{
+			Value = CheckValue++;
+		}
+
+		CHECK(SourcePagedArray.Num() == ArrayType::MaxPerPage() * 3);
+		CHECK(DestPagedArray.IsEmpty());
+
+		TArray<uint8> Bytes;
+		{
+			FMemoryWriter MemoryWriter(Bytes);
+			FBinaryArchiveFormatter BinaryFormatter(MemoryWriter);
+			FStructuredArchive StructuredArchive(BinaryFormatter);
+			FStructuredArchive::FSlot RootSlot = StructuredArchive.Open();
+			RootSlot << SourcePagedArray;
+			StructuredArchive.Close();
+		}
+
+		// Storing the count and the elements
+		CHECK(Bytes.Num() == (sizeof(ArrayType::SizeType) + (sizeof(ArrayType::ElementType) * ArrayType::MaxPerPage() * 3)));
+
+		{
+			FMemoryReader MemoryReader(Bytes);
+			FBinaryArchiveFormatter BinaryFormatter(MemoryReader);
+			FStructuredArchive StructuredArchive(BinaryFormatter);
+			FStructuredArchive::FSlot RootSlot = StructuredArchive.Open();
+			RootSlot << DestPagedArray;
+			StructuredArchive.Close();
+		}
+
+		CHECK(SourcePagedArray == DestPagedArray);
+	}
+	// We construct 3 times over - SetNum, assigning CheckValue and then reading into the dest array
+	CHECK(int32Token::EvenConstructionDestructionCalls(ArrayType::MaxPerPage() * 3 * 3));
+}
+
 }  // namespace PagedArray
 
 #endif
diff --git a/Engine/Source/Runtime/Core/Tests/Containers/TokenTest.h b/Engine/Source/Runtime/Core/Tests/Containers/TokenTest.h
--- a/Engine/Source/Runtime/Core/Tests/Containers/TokenTest.h
+++ b/Engine/Source/Runtime/Core/Tests/Containers/TokenTest.h
@@ -152,3 +152,15 @@ inline bool operator!=(const TTestToken<InValueType>& left, const InValueType& r
 {
 	return !operator==(left, right);
 }
+
+template <typename InValueType>
+inline FArchive& operator<<(FArchive& Ar, TTestToken<InValueType>& TestToken)
+{
+	return Ar << TestToken.Value;
+}
+
+template <typename InValueType>
+inline void operator<<(FStructuredArchive::FSlot Slot, TTestToken<InValueType>& TestToken)
+{
+	Slot.EnterRecord() << SA_VALUE(TEXT("Value"), TestToken.Value);
+}
\ No newline at end of file
