StallDetector: Fixed shadowing error caused by a poorly scoped global

#rnx

diff --git a/Engine/Source/Runtime/Core/Private/ProfilingDebugging/StallDetector.cpp b/Engine/Source/Runtime/Core/Private/ProfilingDebugging/StallDetector.cpp
--- a/Engine/Source/Runtime/Core/Private/ProfilingDebugging/StallDetector.cpp
+++ b/Engine/Source/Runtime/Core/Private/ProfilingDebugging/StallDetector.cpp
@@ -41,7 +41,7 @@ static const double InvalidSeconds = -1.0;
 * Stall Detector Thread
 **/
 
-namespace UE
+namespace UE::StallDetector::Private
 {
 	class FStallDetectorRunnable : public FRunnable
 	{
@@ -84,7 +84,7 @@ namespace UE
 	static FRunnableThread* StallDetectorThread = nullptr;
 }
 
-UE::FStallDetectorRunnable::FStallDetectorRunnable()
+UE::StallDetector::Private::FStallDetectorRunnable::FStallDetectorRunnable()
 	: StartedThread(false)
 	, StopThread(false)
 #if STALL_DETECTOR_HEART_BEAT_CLOCK
@@ -93,7 +93,7 @@ UE::FStallDetectorRunnable::FStallDetectorRunnable()
 {
 }
 
-uint32 UE::FStallDetectorRunnable::Run()
+uint32 UE::StallDetector::Private::FStallDetectorRunnable::Run()
 {
 	while (!StopThread)
 	{
