ParkingLot: Moved into a UE::ParkingLot namespace

#rnx

diff --git a/Engine/Source/Runtime/Core/Private/Async/Mutex.cpp b/Engine/Source/Runtime/Core/Private/Async/Mutex.cpp
--- a/Engine/Source/Runtime/Core/Private/Async/Mutex.cpp
+++ b/Engine/Source/Runtime/Core/Private/Async/Mutex.cpp
@@ -46,7 +46,7 @@ void FMutex::LockSlow()
 		}
 
 		// Wait if the state has not changed. Either way, loop back and try to acquire the lock trying to wait.
-		FParkingLot::Wait(&State, [this, CurrentState] { return State.load(std::memory_order_acquire) == CurrentState; }, []{});
+		ParkingLot::Wait(&State, [this, CurrentState] { return State.load(std::memory_order_acquire) == CurrentState; }, []{});
 	}
 }
 
@@ -71,7 +71,7 @@ void FMutex::UnlockSlow()
 		}
 
 		// There is at least one thread waiting. Wake one thread and return.
-		FParkingLot::WakeOne(&State, [this](FParkingLotWakeState WakeState) -> uint64
+		ParkingLot::WakeOne(&State, [this](ParkingLot::FWakeState WakeState) -> uint64
 		{
 			constexpr uint8 ExpectedState = IsLockedFlag | HasWaitingThreadsFlag;
 			const uint8 NewState = WakeState.bHasWaitingThreads ? HasWaitingThreadsFlag : 0;
diff --git a/Engine/Source/Runtime/Core/Private/Async/ParkingLot.cpp b/Engine/Source/Runtime/Core/Private/Async/ParkingLot.cpp
--- a/Engine/Source/Runtime/Core/Private/Async/ParkingLot.cpp
+++ b/Engine/Source/Runtime/Core/Private/Async/ParkingLot.cpp
@@ -14,10 +14,8 @@
 
 #if UE_PARKINGLOT_USE_WORDMUTEX
 #include "Async/WordMutex.h"
-using FParkingLotLock = UE::Private::FWordMutex;
 #else
 #include "HAL/CriticalSection.h"
-using FParkingLotLock = FCriticalSection;
 #endif
 
 #if UE_PARKINGLOT_USE_WAITONADDRESS
@@ -27,17 +25,23 @@ using FParkingLotLock = FCriticalSection;
 #include <mutex>
 #endif
 
-namespace UE::Private
+namespace UE::ParkingLot::Private
 {
 
+#if UE_PARKINGLOT_USE_WORDMUTEX
+using FBucketMutex = UE::Private::FWordMutex;
+#else
+using FBucketMutex = FCriticalSection;
+#endif
+
 ///////////////////////////////////////////////////////////////////////////////////////////////////////////////////////
 
 /**
  * A thread as stored in the wait queue.
  */
-struct FParkingLotThread final
+struct FThread final
 {
-	FParkingLotThread* Next = nullptr;
+	FThread* Next = nullptr;
 
 #if UE_PARKINGLOT_USE_WAITONADDRESS
 	const void* volatile WaitAddress = nullptr;
@@ -63,29 +67,29 @@ struct FParkingLotThread final
 		}
 	}
 
-	FParkingLotThread() = default;
+	FThread() = default;
 
 private:
-	~FParkingLotThread() = default;
+	~FThread() = default;
 
-	FParkingLotThread(const FParkingLotThread&) = delete;
-	FParkingLotThread& operator=(const FParkingLotThread&) = delete;
+	FThread(const FThread&) = delete;
+	FThread& operator=(const FThread&) = delete;
 };
 
-class FParkingLotThreadLocalData
+class FThreadLocalData
 {
 public:
 	// TODO: This version reserves the table as threads spawn. The function-local thread_local
 	//       reserves the tables as threads wait for the first time. Not sure which is better.
-	//static thread_local FParkingLotThreadLocalData ThreadLocalData;
+	//static thread_local FThreadLocalData ThreadLocalData;
 
-	static FParkingLotThread& Get()
+	static FThread& Get()
 	{
-		static thread_local FParkingLotThreadLocalData ThreadLocalData;
-		FParkingLotThreadLocalData& LocalThreadLocalData = ThreadLocalData;
+		static thread_local FThreadLocalData ThreadLocalData;
+		FThreadLocalData& LocalThreadLocalData = ThreadLocalData;
 		if (!LocalThreadLocalData.Thread)
 		{
-			LocalThreadLocalData.Thread = new FParkingLotThread;
+			LocalThreadLocalData.Thread = new FThread;
 		}
 		return *LocalThreadLocalData.Thread;
 	}
@@ -93,20 +97,20 @@ public:
 private:
 	inline static std::atomic<uint32> ThreadCount = 0;
 
-	TRefCountPtr<FParkingLotThread> Thread;
+	TRefCountPtr<FThread> Thread;
 
-	FParkingLotThreadLocalData();
-	~FParkingLotThreadLocalData();
+	FThreadLocalData();
+	~FThreadLocalData();
 
-	FParkingLotThreadLocalData(const FParkingLotThreadLocalData&) = delete;
-	FParkingLotThreadLocalData& operator=(const FParkingLotThreadLocalData&) = delete;
+	FThreadLocalData(const FThreadLocalData&) = delete;
+	FThreadLocalData& operator=(const FThreadLocalData&) = delete;
 };
 
-//thread_local FParkingLotThreadLocalData FParkingLotThreadLocalData::ThreadLocalData;
+//thread_local FThreadLocalData FThreadLocalData::ThreadLocalData;
 
 ///////////////////////////////////////////////////////////////////////////////////////////////////////////////////////
 
-enum class EParkingLotAction
+enum class EQueueAction
 {
 	Stop,
 	Continue,
@@ -120,27 +124,27 @@ enum class EParkingLotAction
  * Buckets must be locked to access the list of waiting threads.
  * Buckets are aligned to a cache line to reduce false sharing.
  */
-class alignas(PLATFORM_CACHE_LINE_SIZE) FParkingLotBucket final
+class alignas(PLATFORM_CACHE_LINE_SIZE) FBucket final
 {
 public:
-	FParkingLotBucket() = default;
-	FParkingLotBucket(const FParkingLotBucket&) = delete;
-	FParkingLotBucket& operator=(const FParkingLotBucket&) = delete;
+	FBucket() = default;
+	FBucket(const FBucket&) = delete;
+	FBucket& operator=(const FBucket&) = delete;
 
-	[[nodiscard]] inline TDynamicUniqueLock<FParkingLotLock> LockDynamic() { return TDynamicUniqueLock(Mutex); }
+	[[nodiscard]] inline TDynamicUniqueLock<FBucketMutex> LockDynamic() { return TDynamicUniqueLock(Mutex); }
 
 	inline void Lock() { Mutex.Lock(); }
 	inline void Unlock() { Mutex.Unlock(); }
 
 	inline bool IsEmpty() const { return !Head; }
 
-	void Enqueue(FParkingLotThread* Thread);
-	FParkingLotThread* Dequeue();
+	void Enqueue(FThread* Thread);
+	FThread* Dequeue();
 
 	/**
 	 * Dequeues threads based on a visitor.
 	 *
-	 * Visitor signature is EParkingLotAction (FParkingLotThread&).
+	 * Visitor signature is EQueueAction (FThread&).
 	 * Visitor is called for every thread in the bucket, from head to tail.
 	 * Threads are dequeued if the returned action contains Remove.
 	 * Visiting stops if the returned action contains Stop.
@@ -149,12 +153,12 @@ public:
 	void DequeueIf(VisitorType&& Visitor);
 
 private:
-	FParkingLotLock Mutex;
-	FParkingLotThread* Head = nullptr;
-	FParkingLotThread* Tail = nullptr;
+	FBucketMutex Mutex;
+	FThread* Head = nullptr;
+	FThread* Tail = nullptr;
 };
 
-void FParkingLotBucket::Enqueue(FParkingLotThread* Thread)
+void FBucket::Enqueue(FThread* Thread)
 {
 	checkSlow(Thread);
 	checkSlow(!Thread->Next);
@@ -170,9 +174,9 @@ void FParkingLotBucket::Enqueue(FParkingLotThread* Thread)
 	}
 }
 
-FParkingLotThread* FParkingLotBucket::Dequeue()
+FThread* FBucket::Dequeue()
 {
-	FParkingLotThread* Thread = Head;
+	FThread* Thread = Head;
 	if (Thread)
 	{
 		Head = Thread->Next;
@@ -186,21 +190,21 @@ FParkingLotThread* FParkingLotBucket::Dequeue()
 }
 
 template <typename VisitorType>
-void FParkingLotBucket::DequeueIf(VisitorType&& Visitor)
+void FBucket::DequeueIf(VisitorType&& Visitor)
 {
-	FParkingLotThread** Next = &Head;
-	FParkingLotThread* Prev = nullptr;
-	while (FParkingLotThread* Thread = *Next)
+	FThread** Next = &Head;
+	FThread* Prev = nullptr;
+	while (FThread* Thread = *Next)
 	{
 		switch (Visitor(Thread))
 		{
-		case EParkingLotAction::Stop:
+		case EQueueAction::Stop:
 			return;
-		case EParkingLotAction::Continue:
+		case EQueueAction::Continue:
 			Prev = Thread;
 			Next = &Thread->Next;
 			break;
-		case EParkingLotAction::RemoveAndStop:
+		case EQueueAction::RemoveAndStop:
 			if (Tail == Thread)
 			{
 				Tail = Prev;
@@ -208,7 +212,7 @@ void FParkingLotBucket::DequeueIf(VisitorType&& Visitor)
 			*Next = Thread->Next;
 			Thread->Next = nullptr;
 			return;
-		case EParkingLotAction::RemoveAndContinue:
+		case EQueueAction::RemoveAndContinue:
 			if (Tail == Thread)
 			{
 				Tail = Prev;
@@ -231,63 +235,63 @@ void FParkingLotBucket::DequeueIf(VisitorType&& Visitor)
  * leaks are also limited in size because a table is an array of bucket pointers, and the buckets
  * are reused when the table grows.
  */
-class FParkingLotTable final
+class FTable final
 {
 	/** Minimum bucket count to create a table with. */
 	constexpr static uint32 MinSize = 32;
 
 public:
 	/** Find or create, and lock, the bucket for the memory address. */
-	static FParkingLotBucket& FindOrCreateBucket(const void* Address, TDynamicUniqueLock<FParkingLotLock>& OutLock);
+	static FBucket& FindOrCreateBucket(const void* Address, TDynamicUniqueLock<FBucketMutex>& OutLock);
 
 	/** Find and lock the bucket for the memory address. Returns null if the bucket was not created. */
-	static FParkingLotBucket* FindBucket(const void* Address, TDynamicUniqueLock<FParkingLotLock>& OutLock);
+	static FBucket* FindBucket(const void* Address, TDynamicUniqueLock<FBucketMutex>& OutLock);
 
 	/** Reserve memory for the table to handle at least ThreadCount waiting threads. */
 	static void Reserve(uint32 ThreadCount);
 
 private:
 	/** Returns the current global table, creating it if it does not exist. */
-	static FParkingLotTable& CreateOrGet();
+	static FTable& CreateOrGet();
 
 	/** Create a new table with the specified bucket count. */
-	static FParkingLotTable& Create(uint32 Size);
+	static FTable& Create(uint32 Size);
 	/** Destroy a table. Must not be called on any table that has been globally visible. */
-	static void Destroy(FParkingLotTable& Table);
+	static void Destroy(FTable& Table);
 
 	/** Try to lock the whole table by locking every one of its buckets. */
-	static bool TryLock(FParkingLotTable& Table, TArray<FParkingLotBucket*>& OutBuckets);
+	static bool TryLock(FTable& Table, TArray<FBucket*>& OutBuckets);
 	/** Unlock an array of buckets that was filled by TryLock. */
-	static void Unlock(TConstArrayView<FParkingLotBucket*> LockedBuckets);
+	static void Unlock(TConstArrayView<FBucket*> LockedBuckets);
 
 	/** Calculate a 32-bit hash from the memory address. */
 	static uint32 HashAddress(const void* Address);
 
 	/** Pointer to the current global table. Previous global tables are leaked. */
-	inline static std::atomic<FParkingLotTable*> GlobalTable;
+	inline static std::atomic<FTable*> GlobalTable;
 
-	FParkingLotTable() = default;
-	~FParkingLotTable() = default;
-	FParkingLotTable(const FParkingLotTable&) = delete;
-	FParkingLotTable& operator=(const FParkingLotTable&) = delete;
+	FTable() = default;
+	~FTable() = default;
+	FTable(const FTable&) = delete;
+	FTable& operator=(const FTable&) = delete;
 
 	/** Find or create the bucket at the index. */
 	template <typename AllocatorType>
-	FParkingLotBucket& FindOrCreateBucket(uint32 Index, AllocatorType&& Allocator);
+	FBucket& FindOrCreateBucket(uint32 Index, AllocatorType&& Allocator);
 
 	uint32 BucketCount = 0;
-	std::atomic<FParkingLotBucket*> Buckets[0];
+	std::atomic<FBucket*> Buckets[0];
 };
 
-FParkingLotBucket& FParkingLotTable::FindOrCreateBucket(const void* Address, TDynamicUniqueLock<FParkingLotLock>& OutLock)
+FBucket& FTable::FindOrCreateBucket(const void* Address, TDynamicUniqueLock<FBucketMutex>& OutLock)
 {
 	const uint32 Hash = HashAddress(Address);
 
 	for (;;)
 	{
-		FParkingLotTable& Table = CreateOrGet();
+		FTable& Table = CreateOrGet();
 		const uint32 Index = Hash % Table.BucketCount;
-		FParkingLotBucket& Bucket = Table.FindOrCreateBucket(Index, [] { return new FParkingLotBucket; });
+		FBucket& Bucket = Table.FindOrCreateBucket(Index, [] { return new FBucket; });
 		OutLock = Bucket.LockDynamic();
 
 		if (LIKELY(&Table == GlobalTable.load(std::memory_order_acquire)))
@@ -300,16 +304,16 @@ FParkingLotBucket& FParkingLotTable::FindOrCreateBucket(const void* Address, TDy
 	}
 }
 
-FParkingLotBucket* FParkingLotTable::FindBucket(const void* Address, TDynamicUniqueLock<FParkingLotLock>& OutLock)
+FBucket* FTable::FindBucket(const void* Address, TDynamicUniqueLock<FBucketMutex>& OutLock)
 {
 	const uint32 Hash = HashAddress(Address);
 
 	for (;;)
 	{
-		if (FParkingLotTable* Table = GlobalTable.load(std::memory_order_acquire))
+		if (FTable* Table = GlobalTable.load(std::memory_order_acquire))
 		{
 			const uint32 Index = Hash % Table->BucketCount;
-			if (FParkingLotBucket* Bucket = Table->Buckets[Index].load(std::memory_order_acquire))
+			if (FBucket* Bucket = Table->Buckets[Index].load(std::memory_order_acquire))
 			{
 				OutLock = Bucket->LockDynamic();
 
@@ -328,13 +332,13 @@ FParkingLotBucket* FParkingLotTable::FindBucket(const void* Address, TDynamicUni
 }
 
 template <typename AllocatorType>
-FParkingLotBucket& FParkingLotTable::FindOrCreateBucket(uint32 Index, AllocatorType&& Allocator)
+FBucket& FTable::FindOrCreateBucket(uint32 Index, AllocatorType&& Allocator)
 {
-	std::atomic<FParkingLotBucket*>& BucketPtr = Buckets[Index];
-	FParkingLotBucket* Bucket = BucketPtr.load(std::memory_order_acquire);
+	std::atomic<FBucket*>& BucketPtr = Buckets[Index];
+	FBucket* Bucket = BucketPtr.load(std::memory_order_acquire);
 	if (UNLIKELY(!Bucket))
 	{
-		FParkingLotBucket* NewBucket = Allocator();
+		FBucket* NewBucket = Allocator();
 		if (BucketPtr.compare_exchange_strong(Bucket, NewBucket, std::memory_order_release, std::memory_order_acquire))
 		{
 			Bucket = NewBucket;
@@ -348,16 +352,16 @@ FParkingLotBucket& FParkingLotTable::FindOrCreateBucket(uint32 Index, AllocatorT
 	return *Bucket;
 }
 
-FParkingLotTable& FParkingLotTable::CreateOrGet()
+FTable& FTable::CreateOrGet()
 {
-	FParkingLotTable* Table = GlobalTable.load(std::memory_order_acquire);
+	FTable* Table = GlobalTable.load(std::memory_order_acquire);
 
 	if (LIKELY(Table))
 	{
 		return *Table;
 	}
 
-	FParkingLotTable& NewTable = Create(MinSize);
+	FTable& NewTable = Create(MinSize);
 
 	if (LIKELY(GlobalTable.compare_exchange_strong(Table, &NewTable, std::memory_order_release, std::memory_order_acquire)))
 	{
@@ -370,30 +374,30 @@ FParkingLotTable& FParkingLotTable::CreateOrGet()
 	return *Table;
 }
 
-FParkingLotTable& FParkingLotTable::Create(const uint32 Size)
+FTable& FTable::Create(const uint32 Size)
 {
-	const SIZE_T MemorySize = sizeof(FParkingLotTable) + sizeof(FParkingLotBucket*) * SIZE_T(Size);
-	void* const Memory = FMemory::Malloc(MemorySize, alignof(FParkingLotTable));
+	const SIZE_T MemorySize = sizeof(FTable) + sizeof(FBucket*) * SIZE_T(Size);
+	void* const Memory = FMemory::Malloc(MemorySize, alignof(FTable));
 	FMemory::Memzero(Memory, MemorySize);
-	FParkingLotTable& Table = *new(Memory) FParkingLotTable;
+	FTable& Table = *new(Memory) FTable;
 	Table.BucketCount = Size;
 	return Table;
 }
 
-void FParkingLotTable::Destroy(FParkingLotTable& Table)
+void FTable::Destroy(FTable& Table)
 {
-	Table.~FParkingLotTable();
+	Table.~FTable();
 	FMemory::Free(&Table);
 }
 
-void FParkingLotTable::Reserve(const uint32 ThreadCount)
+void FTable::Reserve(const uint32 ThreadCount)
 {
 	const uint32 TargetBucketCount = FMath::RoundUpToPowerOfTwo(ThreadCount);
-	TArray<FParkingLotBucket*> ExistingBuckets;
+	TArray<FBucket*> ExistingBuckets;
 
 	for (;;)
 	{
-		FParkingLotTable& ExistingTable = CreateOrGet();
+		FTable& ExistingTable = CreateOrGet();
 
 		if (LIKELY(ExistingTable.BucketCount >= TargetBucketCount))
 		{
@@ -409,30 +413,30 @@ void FParkingLotTable::Reserve(const uint32 ThreadCount)
 
 		// Gather waiting threads to be redistributed into the buckets of the new table.
 		// Threads with the same address remain in the same relative order as they were queued.
-		TArray<FParkingLotThread*> Threads;
-		for (FParkingLotBucket* Bucket : ExistingBuckets)
+		TArray<FThread*> Threads;
+		for (FBucket* Bucket : ExistingBuckets)
 		{
-			while (FParkingLotThread* Thread = Bucket->Dequeue())
+			while (FThread* Thread = Bucket->Dequeue())
 			{
 				Threads.Add(Thread);
 			}
 		}
 
-		FParkingLotTable& NewTable = Create(TargetBucketCount);
+		FTable& NewTable = Create(TargetBucketCount);
 
 		// Reuse existing now-empty buckets when populating the new table.
-		TArray<FParkingLotBucket*> AvailableBuckets = ExistingBuckets;
-		const auto AllocateBucket = [&AvailableBuckets]() -> FParkingLotBucket*
+		TArray<FBucket*> AvailableBuckets = ExistingBuckets;
+		const auto AllocateBucket = [&AvailableBuckets]() -> FBucket*
 		{
-			return !AvailableBuckets.IsEmpty() ? AvailableBuckets.Pop() : new FParkingLotBucket;
+			return !AvailableBuckets.IsEmpty() ? AvailableBuckets.Pop() : new FBucket;
 		};
 
 		// Add waiting threads to the new table.
-		for (FParkingLotThread* Thread : Threads)
+		for (FThread* Thread : Threads)
 		{
 			const uint32 Hash = HashAddress(Thread->WaitAddress);
 			const uint32 Index = Hash % NewTable.BucketCount;
-			FParkingLotBucket& Bucket = NewTable.FindOrCreateBucket(Index, AllocateBucket);
+			FBucket& Bucket = NewTable.FindOrCreateBucket(Index, AllocateBucket);
 			Bucket.Enqueue(Thread);
 		}
 
@@ -444,7 +448,7 @@ void FParkingLotTable::Reserve(const uint32 ThreadCount)
 		checkSlow(AvailableBuckets.IsEmpty());
 
 		// Make the new table visible to other threads.
-		FParkingLotTable* CompareTable = GlobalTable.exchange(&NewTable, std::memory_order_release);
+		FTable* CompareTable = GlobalTable.exchange(&NewTable, std::memory_order_release);
 		checkSlow(CompareTable == &ExistingTable);
 
 		// Unlock buckets that came from the existing table now that the new table is visible.
@@ -453,19 +457,19 @@ void FParkingLotTable::Reserve(const uint32 ThreadCount)
 	}
 }
 
-bool FParkingLotTable::TryLock(FParkingLotTable& Table, TArray<FParkingLotBucket*>& OutBuckets)
+bool FTable::TryLock(FTable& Table, TArray<FBucket*>& OutBuckets)
 {
 	OutBuckets.Reset();
 
 	// Gather buckets from the table, creating them as needed because the lock is on the bucket.
 	for (uint32 Index = 0; Index < Table.BucketCount; ++Index)
 	{
-		OutBuckets.Add(&Table.FindOrCreateBucket(Index, [] { return new FParkingLotBucket; }));
+		OutBuckets.Add(&Table.FindOrCreateBucket(Index, [] { return new FBucket; }));
 	}
 
 	// Lock the buckets in order by address to ensure consistent ordering regardless of the table being locked.
 	Algo::Sort(OutBuckets);
-	for (FParkingLotBucket* Bucket : OutBuckets)
+	for (FBucket* Bucket : OutBuckets)
 	{
 		Bucket->Lock();
 	}
@@ -482,15 +486,15 @@ bool FParkingLotTable::TryLock(FParkingLotTable& Table, TArray<FParkingLotBucket
 	return false;
 }
 
-void FParkingLotTable::Unlock(TConstArrayView<FParkingLotBucket*> LockedBuckets)
+void FTable::Unlock(TConstArrayView<FBucket*> LockedBuckets)
 {
-	for (FParkingLotBucket* Bucket : LockedBuckets)
+	for (FBucket* Bucket : LockedBuckets)
 	{
 		Bucket->Unlock();
 	}
 }
 
-uint32 FParkingLotTable::HashAddress(const void* Address)
+uint32 FTable::HashAddress(const void* Address)
 {
 	constexpr uint64 A = 0xdc2b17dc9d2fbc29;
 	constexpr uint64 B = 0xcb1014192cb2c5fc;
@@ -501,38 +505,33 @@ uint32 FParkingLotTable::HashAddress(const void* Address)
 
 ///////////////////////////////////////////////////////////////////////////////////////////////////////////////////////
 
-FParkingLotThreadLocalData::FParkingLotThreadLocalData()
+FThreadLocalData::FThreadLocalData()
 {
-	FParkingLotTable::Reserve(ThreadCount.fetch_add(1, std::memory_order_relaxed) + 1);
+	FTable::Reserve(ThreadCount.fetch_add(1, std::memory_order_relaxed) + 1);
 }
 
-FParkingLotThreadLocalData::~FParkingLotThreadLocalData()
+FThreadLocalData::~FThreadLocalData()
 {
 	ThreadCount.fetch_sub(1, std::memory_order_relaxed);
 }
 
 ///////////////////////////////////////////////////////////////////////////////////////////////////////////////////////
 
-} // UE::Private
-
-namespace UE
+FWaitState WaitFor(const void* Address, TFunctionRef<bool()> CanWait, TFunctionRef<void()> BeforeWait, uint32 WaitMs)
 {
+	using namespace Private;
 
-FParkingLotWaitState FParkingLot::PrivateWaitFor(const void* Address, TFunctionRef<bool()> CanWait, TFunctionRef<void()> BeforeWait, uint32 WaitMs)
-{
-	using namespace UE::Private;
-
-	FParkingLotThread& Self = FParkingLotThreadLocalData::Get();
+	FThread& Self = FThreadLocalData::Get();
 
 	checkfSlow(!Self.WaitAddress, TEXT("WaitAddress must be null. This can happen if Wait is called by BeforeWait."));
-	checkfSlow(Self.WakeToken == 0, TEXT("WakeToken must be 0. This is an error in FParkingLot."));
+	checkfSlow(Self.WakeToken == 0, TEXT("WakeToken must be 0. This is an error in ParkingLot."));
 
-	FParkingLotWaitState State;
+	FWaitState State;
 
 	// Enqueue the thread if CanWait returns true while the bucket is locked.
 	{
-		TDynamicUniqueLock<FParkingLotLock> BucketLock;
-		FParkingLotBucket& Bucket = FParkingLotTable::FindOrCreateBucket(Address, BucketLock);
+		TDynamicUniqueLock<FBucketMutex> BucketLock;
+		FBucket& Bucket = FTable::FindOrCreateBucket(Address, BucketLock);
 		State.bDidWait = CanWait();
 		if (!State.bDidWait)
 		{
@@ -591,17 +590,17 @@ FParkingLotWaitState FParkingLot::PrivateWaitFor(const void* Address, TFunctionR
 	// The timeout was reached and the thread needs to dequeue itself.
 	// This can race with a call to wake a thread, which means Self is unsafe to access outside of the lock.
 	bool bDequeued = false;
-	if (TDynamicUniqueLock<FParkingLotLock> BucketLock; FParkingLotBucket* Bucket = FParkingLotTable::FindBucket(Address, BucketLock))
+	if (TDynamicUniqueLock<FBucketMutex> BucketLock; FBucket* Bucket = FTable::FindBucket(Address, BucketLock))
 	{
-		Bucket->DequeueIf([Self = &Self, &bDequeued](FParkingLotThread* Thread)
+		Bucket->DequeueIf([Self = &Self, &bDequeued](FThread* Thread)
 		{
 			if (Thread == Self)
 			{
 				bDequeued = true;
 				Thread->WaitAddress = nullptr;
-				return EParkingLotAction::RemoveAndStop;
+				return EQueueAction::RemoveAndStop;
 			}
-			return EParkingLotAction::Continue;
+			return EQueueAction::Continue;
 		});
 	}
 
@@ -626,26 +625,26 @@ FParkingLotWaitState FParkingLot::PrivateWaitFor(const void* Address, TFunctionR
 	return State;
 }
 
-void FParkingLot::PrivateWakeOne(const void* Address, TFunctionRef<uint64(FParkingLotWakeState)> OnWakeState)
+void WakeOne(const void* Address, TFunctionRef<uint64(FWakeState)> OnWakeState)
 {
-	using namespace UE::Private;
+	using namespace Private;
 
-	TRefCountPtr<FParkingLotThread> WakeThread;
+	TRefCountPtr<FThread> WakeThread;
 	uint64 WakeToken = 0;
 
 	{
-		TDynamicUniqueLock<FParkingLotLock> BucketLock;
-		FParkingLotBucket& Bucket = FParkingLotTable::FindOrCreateBucket(Address, BucketLock);
-		Bucket.DequeueIf([Address, &WakeThread](FParkingLotThread* Thread)
+		TDynamicUniqueLock<FBucketMutex> BucketLock;
+		FBucket& Bucket = FTable::FindOrCreateBucket(Address, BucketLock);
+		Bucket.DequeueIf([Address, &WakeThread](FThread* Thread)
 		{
 			if (Thread->WaitAddress == Address)
 			{
 				WakeThread = Thread;
-				return EParkingLotAction::RemoveAndStop;
+				return EQueueAction::RemoveAndStop;
 			}
-			return EParkingLotAction::Continue;
+			return EQueueAction::Continue;
 		});
-		FParkingLotWakeState WakeState;
+		FWakeState WakeState;
 		WakeState.bDidWake = !!WakeThread;
 		WakeState.bHasWaitingThreads = !Bucket.IsEmpty();
 		WakeToken = OnWakeState(WakeState);
@@ -668,10 +667,15 @@ void FParkingLot::PrivateWakeOne(const void* Address, TFunctionRef<uint64(FParki
 	}
 }
 
-FParkingLotWakeState FParkingLot::WakeOne(const void* Address)
+} // UE::ParkingLot::Private
+
+namespace UE::ParkingLot
+{
+
+FWakeState WakeOne(const void* Address)
 {
-	FParkingLotWakeState OutState;
-	PrivateWakeOne(Address, [&OutState](FParkingLotWakeState State) -> uint64
+	FWakeState OutState;
+	Private::WakeOne(Address, [&OutState](FWakeState State) -> uint64
 	{
 		OutState = State;
 		return 0;
@@ -679,26 +683,26 @@ FParkingLotWakeState FParkingLot::WakeOne(const void* Address)
 	return OutState;
 }
 
-uint32 FParkingLot::WakeMultiple(const void* const Address, const uint32 WakeCount)
+uint32 WakeMultiple(const void* const Address, const uint32 WakeCount)
 {
-	using namespace UE::Private;
+	using namespace Private;
 
-	TArray<TRefCountPtr<FParkingLotThread>> WakeThreads;
+	TArray<TRefCountPtr<FThread>> WakeThreads;
 
-	if (TDynamicUniqueLock<FParkingLotLock> BucketLock; FParkingLotBucket* Bucket = FParkingLotTable::FindBucket(Address, BucketLock))
+	if (TDynamicUniqueLock<FBucketMutex> BucketLock; FBucket* Bucket = FTable::FindBucket(Address, BucketLock))
 	{
-		Bucket->DequeueIf([Address, &WakeThreads, WakeCount](FParkingLotThread* Thread)
+		Bucket->DequeueIf([Address, &WakeThreads, WakeCount](FThread* Thread)
 		{
 			if (Thread->WaitAddress == Address)
 			{
 				const uint32 Count = uint32(WakeThreads.Add(Thread) + 1);
-				return Count == WakeCount ? EParkingLotAction::RemoveAndStop : EParkingLotAction::RemoveAndContinue;
+				return Count == WakeCount ? EQueueAction::RemoveAndStop : EQueueAction::RemoveAndContinue;
 			}
-			return EParkingLotAction::Continue;
+			return EQueueAction::Continue;
 		});
 	}
 
-	for (FParkingLotThread* WakeThread : WakeThreads)
+	for (FThread* WakeThread : WakeThreads)
 	{
 		checkSlow(WakeThread->WaitAddress == Address);
 	#if UE_PARKINGLOT_USE_WAITONADDRESS
@@ -716,9 +720,9 @@ uint32 FParkingLot::WakeMultiple(const void* const Address, const uint32 WakeCou
 	return uint32(WakeThreads.Num());
 }
 
-void FParkingLot::WakeAll(const void* Address)
+void WakeAll(const void* Address)
 {
 	WakeMultiple(Address, MAX_uint32);
 }
 
-} // UE
+} // UE::ParkingLot
diff --git a/Engine/Source/Runtime/Core/Public/Async/ParkingLot.h b/Engine/Source/Runtime/Core/Public/Async/ParkingLot.h
--- a/Engine/Source/Runtime/Core/Public/Async/ParkingLot.h
+++ b/Engine/Source/Runtime/Core/Public/Async/ParkingLot.h
@@ -8,10 +8,13 @@
 
 #define UE_API CORE_API
 
-namespace UE
+/**
+ * ParkingLot is a global table of queues of waiting threads keyed by memory address.
+ */
+namespace UE::ParkingLot
 {
 
-struct FParkingLotWaitState final
+struct FWaitState final
 {
 	/** Did this thread wait? True only if CanWait returned true. */
 	bool bDidWait = false;
@@ -21,7 +24,7 @@ struct FParkingLotWaitState final
 	uint64 WakeToken = 0;
 };
 
-struct FParkingLotWakeState final
+struct FWakeState final
 {
 	/** Did a thread wake up? */
 	bool bDidWake = false;
@@ -29,83 +32,75 @@ struct FParkingLotWakeState final
 	bool bHasWaitingThreads = false;
 };
 
+namespace Private
+{
+UE_API FWaitState WaitFor(const void* Address, TFunctionRef<bool()> CanWait, TFunctionRef<void()> BeforeWait, uint32 WaitMs);
+UE_API void WakeOne(const void* Address, TFunctionRef<uint64(FWakeState)> OnWakeState);
+}
+
 /**
- * A global table of queues of waiting threads keyed by memory address.
+ * Queue the calling thread to wait if CanWait returns true. BeforeWait is only called if CanWait returns true.
+ *
+ * @param Address      Address to use as the key for the queue. The same address is used to wake the thread.
+ * @param CanWait      Function called while the queue is locked. A return of false cancels the wait.
+ * @param BeforeWait   Function called after the queue is unlocked and before the thread waits.
  */
-class FParkingLot final
+template <typename CanWaitType, typename BeforeWaitType>
+inline FWaitState Wait(const void* Address, CanWaitType&& CanWait, BeforeWaitType&& BeforeWait)
 {
-public:
-	/**
-	 * Queue the calling thread to wait if CanWait returns true. BeforeWait is only called if CanWait returns true.
-	 *
-	 * @param Address      Address to use as the key for the queue. The same address is used to wake the thread.
-	 * @param CanWait      Function called while the queue is locked. A return of false cancels the wait.
-	 * @param BeforeWait   Function called after the queue is unlocked and before the thread waits.
-	 */
-	template <typename CanWaitType, typename BeforeWaitType>
-	inline static FParkingLotWaitState Wait(const void* Address, CanWaitType&& CanWait, BeforeWaitType&& BeforeWait)
-	{
-		return PrivateWaitFor(Address, Forward<CanWaitType>(CanWait), Forward<BeforeWaitType>(BeforeWait), MAX_uint32);
-	}
-
-	/**
-	 * Queue the calling thread to wait if CanWait returns true. BeforeWait is only called if CanWait returns true.
-	 *
-	 * @param Address      Address to use as the key for the queue. The same address is used to wake the thread.
-	 * @param CanWait      Function called while the queue is locked. A return of false cancels the wait.
-	 * @param BeforeWait   Function called after the queue is unlocked and before the thread waits.
-	 * @param WaitMs       Minimum wait time after which waiting is automatically canceled and the thread wakes.
-	 */
-	template <typename CanWaitType, typename BeforeWaitType>
-	inline static FParkingLotWaitState WaitFor(const void* Address, CanWaitType&& CanWait, BeforeWaitType&& BeforeWait, uint32 WaitMs)
-	{
-		return PrivateWaitFor(Address, Forward<CanWaitType>(CanWait), Forward<BeforeWaitType>(BeforeWait), WaitMs);
-	}
+	return Private::WaitFor(Address, Forward<CanWaitType>(CanWait), Forward<BeforeWaitType>(BeforeWait), MAX_uint32);
+}
 
-	/**
-	 * Wake one thread from the queue of threads waiting on the address.
-	 *
-	 * @param Address       Address to use as the key for the queue. Must match the address used in Wait.
-	 * @param OnWakeState   Function called while the queue is locked. Receives the wake state. Returns WakeToken.
-	 */
-	template <typename OnWakeStateType>
-	inline static void WakeOne(const void* Address, OnWakeStateType&& OnWakeState)
-	{
-		return PrivateWakeOne(Address, Forward<OnWakeStateType>(OnWakeState));
-	}
-
-	/**
-	 * Wake one thread from the queue of threads waiting on the address.
-	 *
-	 * @param Address   Address to use as the key for the queue. Must match the address used in Wait.
-	 * @return The wake state, which includes whether a thread woke up and whether there are more queued.
-	 */
-	UE_API static FParkingLotWakeState WakeOne(const void* Address);
+/**
+ * Queue the calling thread to wait if CanWait returns true. BeforeWait is only called if CanWait returns true.
+ *
+ * @param Address      Address to use as the key for the queue. The same address is used to wake the thread.
+ * @param CanWait      Function called while the queue is locked. A return of false cancels the wait.
+ * @param BeforeWait   Function called after the queue is unlocked and before the thread waits.
+ * @param WaitMs       Minimum wait time after which waiting is automatically canceled and the thread wakes.
+ */
+template <typename CanWaitType, typename BeforeWaitType>
+inline FWaitState WaitFor(const void* Address, CanWaitType&& CanWait, BeforeWaitType&& BeforeWait, uint32 WaitMs)
+{
+	return Private::WaitFor(Address, Forward<CanWaitType>(CanWait), Forward<BeforeWaitType>(BeforeWait), WaitMs);
+}
 
-	/**
-	 * Wake up to WakeCount threads from the queue of threads waiting on the address.
-	 *
-	 * @param Address     Address to use as the key for the queue. Must match the address used in Wait.
-	 * @param WakeCount   The maximum number of threads to wake.
-	 * @return The number of threads that this call woke up.
-	 */
-	UE_API static uint32 WakeMultiple(const void* Address, uint32 WakeCount);
+/**
+ * Wake one thread from the queue of threads waiting on the address.
+ *
+ * @param Address       Address to use as the key for the queue. Must match the address used in Wait.
+ * @param OnWakeState   Function called while the queue is locked. Receives the wake state. Returns WakeToken.
+ */
+template <typename OnWakeStateType>
+inline void WakeOne(const void* Address, OnWakeStateType&& OnWakeState)
+{
+	return Private::WakeOne(Address, Forward<OnWakeStateType>(OnWakeState));
+}
 
-	/**
-	 * Wake all threads from the queue of threads waiting on the address.
-	 *
-	 * @param Address   Address to use as the key for the queue. Must match the address used in Wait.
-	 */
-	UE_API static void WakeAll(const void* Address);
+/**
+ * Wake one thread from the queue of threads waiting on the address.
+ *
+ * @param Address   Address to use as the key for the queue. Must match the address used in Wait.
+ * @return The wake state, which includes whether a thread woke up and whether there are more queued.
+ */
+UE_API FWakeState WakeOne(const void* Address);
 
-private:
-	UE_API static FParkingLotWaitState PrivateWaitFor(const void* Address, TFunctionRef<bool()> CanWait, TFunctionRef<void()> BeforeWait, uint32 WaitMs);
-	UE_API static void PrivateWakeOne(const void* Address, TFunctionRef<uint64(FParkingLotWakeState)> OnWakeState);
+/**
+ * Wake up to WakeCount threads from the queue of threads waiting on the address.
+ *
+ * @param Address     Address to use as the key for the queue. Must match the address used in Wait.
+ * @param WakeCount   The maximum number of threads to wake.
+ * @return The number of threads that this call woke up.
+ */
+UE_API uint32 WakeMultiple(const void* Address, uint32 WakeCount);
 
-	FParkingLot() = delete;
-	FParkingLot(const FParkingLot&) = delete;
-};
+/**
+ * Wake all threads from the queue of threads waiting on the address.
+ *
+ * @param Address   Address to use as the key for the queue. Must match the address used in Wait.
+ */
+UE_API void WakeAll(const void* Address);
 
-} // UE
+} // UE::ParkingLot
 
 #undef UE_API
