Add missing include to TokenTest.h

diff --git a/Engine/Source/Runtime/Core/Tests/Containers/TokenTest.h b/Engine/Source/Runtime/Core/Tests/Containers/TokenTest.h
--- a/Engine/Source/Runtime/Core/Tests/Containers/TokenTest.h
+++ b/Engine/Source/Runtime/Core/Tests/Containers/TokenTest.h
@@ -3,6 +3,7 @@
 #pragma once
 
 #include "CoreTypes.h"
+#include "Serialization/StructuredArchive.h"
 #include "Templates/UnrealTemplate.h"
 
 template <typename InValueType>
