Added ReverseIterate() range adapter for reverse-iteration over ranges and some sample usage.

diff --git a/Engine/Source/Runtime/AssetRegistry/Private/AssetRegistry.cpp b/Engine/Source/Runtime/AssetRegistry/Private/AssetRegistry.cpp
--- a/Engine/Source/Runtime/AssetRegistry/Private/AssetRegistry.cpp
+++ b/Engine/Source/Runtime/AssetRegistry/Private/AssetRegistry.cpp
@@ -5797,9 +5797,9 @@ void FAssetRegistryImpl::PushProcessLoadedAssetsBatch(Impl::FEventContext& Event
 	}
 
 	// Push back any objects from the batch that were not processed due to timing out
-	for (int32 Index = UnprocessedFromBatch.Num() - 1; Index >= 0; --Index)
+	for (const UObject* Obj : ReverseIterate(UnprocessedFromBatch))
 	{
-		LoadedAssetsToProcess.EmplaceFront(UnprocessedFromBatch[Index]);
+		LoadedAssetsToProcess.EmplaceFront(Obj);
 	}
 }
 
diff --git a/Engine/Source/Runtime/Core/Private/Async/TaskGraph.cpp b/Engine/Source/Runtime/Core/Private/Async/TaskGraph.cpp
--- a/Engine/Source/Runtime/Core/Private/Async/TaskGraph.cpp
+++ b/Engine/Source/Runtime/Core/Private/Async/TaskGraph.cpp
@@ -2494,9 +2494,8 @@ void FGraphEvent::DispatchSubsequents(TArray<FBaseGraphTask*>& NewTasks, ENamedT
 	bool bWakeUpWorker = false;
 	TArray<FBaseGraphTask*> PoppedTasks;
 	SubsequentList.PopAllAndClose(PoppedTasks);
-	for (int32 Index = PoppedTasks.Num() - 1; Index >= 0 ; Index--) // reverse the order since PopAll is implicitly backwards
+	for (FBaseGraphTask* NewTask : ReverseIterate(PoppedTasks)) // reverse the order since PopAll is implicitly backwards
 	{
-		FBaseGraphTask* NewTask = PoppedTasks[Index];
 		checkThreadGraph(NewTask);
 		NewTask->ConditionalQueueTask(CurrentThreadIfKnown, bWakeUpWorker);
 	}
diff --git a/Engine/Source/Runtime/Core/Public/Containers/Array.h b/Engine/Source/Runtime/Core/Public/Containers/Array.h
--- a/Engine/Source/Runtime/Core/Public/Containers/Array.h
+++ b/Engine/Source/Runtime/Core/Public/Containers/Array.h
@@ -4,6 +4,7 @@
 
 #include "CoreTypes.h"
 #include "Misc/AssertionMacros.h"
+#include "Misc/ReverseIterate.h"
 #include "HAL/UnrealMemory.h"
 #include "Templates/IsSigned.h"
 #include "Templates/UnrealTypeTraits.h"
@@ -191,7 +192,7 @@ FORCEINLINE TIndexedContainerIterator<ContainerType, ElementType, SizeType> oper
 	 * Pointer-like iterator type for ranged-for loops which checks that the
 	 * container hasn't been resized during iteration.
 	 */
-	template <typename ElementType, typename SizeType>
+	template <typename ElementType, typename SizeType, bool bReverse = false>
 	struct TCheckedPointerIterator
 	{
 		// This iterator type only supports the minimal functionality needed to support
@@ -208,18 +209,39 @@ FORCEINLINE TIndexedContainerIterator<ContainerType, ElementType, SizeType> oper
 
 		FORCEINLINE ElementType& operator*() const
 		{
-			return *Ptr;
+			if constexpr (bReverse)
+			{
+				return *(Ptr - 1);
+			}
+			else
+			{
+				return *Ptr;
+			}
 		}
 
 		FORCEINLINE TCheckedPointerIterator& operator++()
 		{
-			++Ptr;
+			if constexpr (bReverse)
+			{
+				--Ptr;
+			}
+			else
+			{
+				++Ptr;
+			}
 			return *this;
 		}
 
 		FORCEINLINE TCheckedPointerIterator& operator--()
 		{
-			--Ptr;
+			if constexpr (bReverse)
+			{
+				++Ptr;
+			}
+			else
+			{
+				--Ptr;
+			}
 			return *this;
 		}
 
@@ -2774,11 +2796,15 @@ public:
 	}
 
 	#if TARRAY_RANGED_FOR_CHECKS
-		typedef TCheckedPointerIterator<      ElementType, SizeType> RangedForIteratorType;
-		typedef TCheckedPointerIterator<const ElementType, SizeType> RangedForConstIteratorType;
+		typedef TCheckedPointerIterator<      ElementType, SizeType, false> RangedForIteratorType;
+		typedef TCheckedPointerIterator<const ElementType, SizeType, false> RangedForConstIteratorType;
+		typedef TCheckedPointerIterator<      ElementType, SizeType, true>  RangedForReverseIteratorType;
+		typedef TCheckedPointerIterator<const ElementType, SizeType, true>  RangedForConstReverseIteratorType;
 	#else
-		typedef       ElementType* RangedForIteratorType;
-		typedef const ElementType* RangedForConstIteratorType;
+		typedef                               ElementType* RangedForIteratorType;
+		typedef                         const ElementType* RangedForConstIteratorType;
+		typedef TReversePointerIterator<      ElementType> RangedForReverseIteratorType;
+		typedef TReversePointerIterator<const ElementType> RangedForConstReverseIteratorType;
 	#endif
 
 public:
@@ -2788,15 +2814,23 @@ public:
 	 * STL-like iterators to enable range-based for loop support.
 	 */
 	#if TARRAY_RANGED_FOR_CHECKS
-		FORCEINLINE RangedForIteratorType      begin()       { return RangedForIteratorType     (ArrayNum, GetData()); }
-		FORCEINLINE RangedForConstIteratorType begin() const { return RangedForConstIteratorType(ArrayNum, GetData()); }
-		FORCEINLINE RangedForIteratorType      end  ()       { return RangedForIteratorType     (ArrayNum, GetData() + Num()); }
-		FORCEINLINE RangedForConstIteratorType end  () const { return RangedForConstIteratorType(ArrayNum, GetData() + Num()); }
+		FORCEINLINE RangedForIteratorType             begin ()       { return RangedForIteratorType            (ArrayNum, GetData()); }
+		FORCEINLINE RangedForConstIteratorType        begin () const { return RangedForConstIteratorType       (ArrayNum, GetData()); }
+		FORCEINLINE RangedForIteratorType             end   ()       { return RangedForIteratorType            (ArrayNum, GetData() + Num()); }
+		FORCEINLINE RangedForConstIteratorType        end   () const { return RangedForConstIteratorType       (ArrayNum, GetData() + Num()); }
+		FORCEINLINE RangedForReverseIteratorType      rbegin()       { return RangedForReverseIteratorType     (ArrayNum, GetData() + Num()); }
+		FORCEINLINE RangedForConstReverseIteratorType rbegin() const { return RangedForConstReverseIteratorType(ArrayNum, GetData() + Num()); }
+		FORCEINLINE RangedForReverseIteratorType      rend  ()       { return RangedForReverseIteratorType     (ArrayNum, GetData()); }
+		FORCEINLINE RangedForConstReverseIteratorType rend  () const { return RangedForConstReverseIteratorType(ArrayNum, GetData()); }
 	#else
-		FORCEINLINE RangedForIteratorType      begin()       { return GetData(); }
-		FORCEINLINE RangedForConstIteratorType begin() const { return GetData(); }
-		FORCEINLINE RangedForIteratorType      end()         { return GetData() + Num(); }
-		FORCEINLINE RangedForConstIteratorType end() const   { return GetData() + Num(); }
+		FORCEINLINE RangedForIteratorType             begin ()       { return                                   GetData(); }
+		FORCEINLINE RangedForConstIteratorType        begin () const { return                                   GetData(); }
+		FORCEINLINE RangedForIteratorType             end   ()       { return                                   GetData() + Num(); }
+		FORCEINLINE RangedForConstIteratorType        end   () const { return                                   GetData() + Num(); }
+		FORCEINLINE RangedForReverseIteratorType      rbegin()       { return RangedForReverseIteratorType     (GetData() + Num()); }
+		FORCEINLINE RangedForConstReverseIteratorType rbegin() const { return RangedForConstReverseIteratorType(GetData() + Num()); }
+		FORCEINLINE RangedForReverseIteratorType      rend  ()       { return RangedForReverseIteratorType     (GetData()); }
+		FORCEINLINE RangedForConstReverseIteratorType rend  () const { return RangedForConstReverseIteratorType(GetData()); }
 	#endif
 
 public:
diff --git a/Engine/Source/Runtime/Core/Public/Containers/ArrayView.h b/Engine/Source/Runtime/Core/Public/Containers/ArrayView.h
--- a/Engine/Source/Runtime/Core/Public/Containers/ArrayView.h
+++ b/Engine/Source/Runtime/Core/Public/Containers/ArrayView.h
@@ -7,6 +7,7 @@
 #include "Templates/IsSigned.h"
 #include "Templates/PointerIsConvertibleFromTo.h"
 #include "Misc/AssertionMacros.h"
+#include "Misc/ReverseIterate.h"
 #include "Templates/Invoke.h"
 #include "Templates/UnrealTypeTraits.h"
 #include "Traits/ElementType.h"
@@ -683,8 +684,10 @@ public:
 	 * DO NOT USE DIRECTLY
 	 * STL-like iterators to enable range-based for loop support.
 	 */
-	FORCEINLINE ElementType* begin() const { return GetData(); }
-	FORCEINLINE ElementType* end  () const { return GetData() + Num(); }
+	FORCEINLINE ElementType*                         begin () const { return GetData(); }
+	FORCEINLINE ElementType*                         end   () const { return GetData() + Num(); }
+	FORCEINLINE TReversePointerIterator<ElementType> rbegin() const { return TReversePointerIterator<ElementType>(GetData() + Num()); }
+	FORCEINLINE TReversePointerIterator<ElementType> rend  () const { return TReversePointerIterator<ElementType>(GetData()); }
 
 public:
 	/**
diff --git a/Engine/Source/Runtime/Core/Public/Containers/StaticArray.h b/Engine/Source/Runtime/Core/Public/Containers/StaticArray.h
--- a/Engine/Source/Runtime/Core/Public/Containers/StaticArray.h
+++ b/Engine/Source/Runtime/Core/Public/Containers/StaticArray.h
@@ -4,6 +4,7 @@
 
 #include "CoreTypes.h"
 #include "Misc/AssertionMacros.h"
+#include "Misc/ReverseIterate.h"
 #include "Templates/UnrealTemplate.h"
 #include "Templates/UnrealTypeTraits.h"
 #include "Delegates/IntegerSequence.h"
@@ -137,7 +138,7 @@ private:
 
 public:
 
-	template <typename StorageElementType>
+	template <typename StorageElementType, bool bReverse = false>
 	struct FRangedForIterator
 	{
 		explicit FRangedForIterator(StorageElementType* InPtr)
@@ -146,12 +147,26 @@ public:
 
 		auto& operator*() const
 		{
-			return Ptr->Element;
+			if constexpr (bReverse)
+			{
+				return (Ptr - 1)->Element;
+			}
+			else
+			{
+				return Ptr->Element;
+			}
 		}
 
 		FRangedForIterator& operator++()
 		{
-			++Ptr;
+			if constexpr (bReverse)
+			{
+				--Ptr;
+			}
+			else
+			{
+				++Ptr;
+			}
 			return *this;
 		}
 
@@ -164,15 +179,20 @@ public:
 		StorageElementType* Ptr;
 	};
 
-	using RangedForIteratorType = FRangedForIterator<TArrayStorageElementAligned>;
-	using RangedForConstIteratorType = FRangedForIterator<const TArrayStorageElementAligned>;
+	using RangedForIteratorType             = FRangedForIterator<      TArrayStorageElementAligned>;
+	using RangedForConstIteratorType        = FRangedForIterator<const TArrayStorageElementAligned>;
+	using RangedForReverseIteratorType      = FRangedForIterator<      TArrayStorageElementAligned, true>;
+	using RangedForConstReverseIteratorType = FRangedForIterator<const TArrayStorageElementAligned, true>;
 
 	/** STL-like iterators to enable range-based for loop support. */
-	FORCEINLINE RangedForIteratorType		begin()			{ return RangedForIteratorType(Storage.Elements); }
-	FORCEINLINE RangedForConstIteratorType	begin() const	{ return RangedForConstIteratorType(Storage.Elements); }
-	FORCEINLINE RangedForIteratorType		end()			{ return RangedForIteratorType(Storage.Elements + NumElements); }
-	FORCEINLINE RangedForConstIteratorType	end() const		{ return RangedForConstIteratorType(Storage.Elements + NumElements); }
-
+	FORCEINLINE RangedForIteratorType				begin()			{ return RangedForIteratorType(Storage.Elements); }
+	FORCEINLINE RangedForConstIteratorType			begin() const	{ return RangedForConstIteratorType(Storage.Elements); }
+	FORCEINLINE RangedForIteratorType				end()			{ return RangedForIteratorType(Storage.Elements + NumElements); }
+	FORCEINLINE RangedForConstIteratorType			end() const		{ return RangedForConstIteratorType(Storage.Elements + NumElements); }
+	FORCEINLINE RangedForReverseIteratorType		rbegin()		{ return RangedForReverseIteratorType(Storage.Elements + NumElements); }
+	FORCEINLINE RangedForConstReverseIteratorType	rbegin() const	{ return RangedForConstReverseIteratorType(Storage.Elements + NumElements); }
+	FORCEINLINE RangedForReverseIteratorType		rend()			{ return RangedForReverseIteratorType(Storage.Elements); }
+	FORCEINLINE RangedForConstReverseIteratorType	rend() const	{ return RangedForConstReverseIteratorType(Storage.Elements); }
 };
 
 /** Creates a static array filled with the specified value. */
diff --git a/Engine/Source/Runtime/Core/Public/Containers/StringView.h b/Engine/Source/Runtime/Core/Public/Containers/StringView.h
--- a/Engine/Source/Runtime/Core/Public/Containers/StringView.h
+++ b/Engine/Source/Runtime/Core/Public/Containers/StringView.h
@@ -9,6 +9,7 @@
 #include "Math/UnrealMathUtility.h"
 #include "Misc/Crc.h"
 #include "Misc/CString.h"
+#include "Misc/ReverseIterate.h"
 #include "String/Find.h"
 #include "Templates/UnrealTemplate.h"
 #include "Traits/ElementType.h"
@@ -444,6 +445,8 @@ public:
 	 */
 	constexpr inline const CharType* begin() const { return DataPtr; }
 	constexpr inline const CharType* end() const { return DataPtr + Size; }
+	constexpr inline TReversePointerIterator<const CharType> rbegin() const { return TReversePointerIterator<const CharType>(DataPtr + Size); }
+	constexpr inline TReversePointerIterator<const CharType> rend() const { return TReversePointerIterator<const CharType>(DataPtr); }
 
 protected:
 	const CharType* DataPtr = nullptr;
diff --git a/Engine/Source/Runtime/Core/Public/Containers/UnrealString.h b/Engine/Source/Runtime/Core/Public/Containers/UnrealString.h
--- a/Engine/Source/Runtime/Core/Public/Containers/UnrealString.h
+++ b/Engine/Source/Runtime/Core/Public/Containers/UnrealString.h
@@ -253,10 +253,14 @@ public:
 	 * DO NOT USE DIRECTLY
 	 * STL-like iterators to enable range-based for loop support.
 	 */
-	FORCEINLINE DataType::RangedForIteratorType      begin()       { auto Result = Data.begin();                                   return Result; }
-	FORCEINLINE DataType::RangedForConstIteratorType begin() const { auto Result = Data.begin();                                   return Result; }
-	FORCEINLINE DataType::RangedForIteratorType      end  ()       { auto Result = Data.end();   if (Data.Num()) { --Result; }     return Result; }
-	FORCEINLINE DataType::RangedForConstIteratorType end  () const { auto Result = Data.end();   if (Data.Num()) { --Result; }     return Result; }
+	FORCEINLINE DataType::RangedForIteratorType             begin ()       { auto Result = Data.begin();                                return Result; }
+	FORCEINLINE DataType::RangedForConstIteratorType        begin () const { auto Result = Data.begin();                                return Result; }
+	FORCEINLINE DataType::RangedForIteratorType             end   ()       { auto Result = Data.end();    if (Data.Num()) { --Result; } return Result; }
+	FORCEINLINE DataType::RangedForConstIteratorType        end   () const { auto Result = Data.end();    if (Data.Num()) { --Result; } return Result; }
+	FORCEINLINE DataType::RangedForReverseIteratorType      rbegin()       { auto Result = Data.rbegin(); if (Data.Num()) { ++Result; } return Result; }
+	FORCEINLINE DataType::RangedForConstReverseIteratorType rbegin() const { auto Result = Data.rbegin(); if (Data.Num()) { ++Result; } return Result; }
+	FORCEINLINE DataType::RangedForReverseIteratorType      rend  ()       { auto Result = Data.rend();                                 return Result; }
+	FORCEINLINE DataType::RangedForConstReverseIteratorType rend  () const { auto Result = Data.rend();                                 return Result; }
 
 	FORCEINLINE SIZE_T GetAllocatedSize() const
 	{
diff --git a/Engine/Source/Runtime/Core/Public/Misc/ReverseIterate.h b/Engine/Source/Runtime/Core/Public/Misc/ReverseIterate.h
new file mode 100644
--- /dev/null
+++ b/Engine/Source/Runtime/Core/Public/Misc/ReverseIterate.h
@@ -0,0 +1,138 @@
+// Copyright Epic Games, Inc. All Rights Reserved.
+
+#pragma once
+
+#include <iterator>
+
+/**
+ * Pointer-like reverse iterator type.
+ */
+template <typename T>
+struct TReversePointerIterator
+{
+	// This iterator type only supports the minimal functionality needed to support
+	// C++ ranged-for syntax.  For example, it does not provide post-increment ++ nor ==.
+
+	/**
+	 * Constructor for TReversePointerIterator.
+	 *
+	 * @param  InPtr  A pointer to the location after the element being referenced.
+	 *
+	 * @note  Like std::reverse_iterator, this points to one past the element being referenced.
+	 *        Therefore, for an array of size N starting at P, the begin iterator should be
+	 *        constructed at (P+N) and the end iterator should be constructed at P.
+	 */
+	constexpr explicit TReversePointerIterator(T* InPtr UE_LIFETIMEBOUND)
+		: Ptr(InPtr)
+	{
+	}
+
+	constexpr inline T& operator*() const
+	{
+		return *(Ptr - 1);
+	}
+
+	constexpr inline TReversePointerIterator& operator++()
+	{
+		--Ptr;
+		return *this;
+	}
+
+	constexpr inline bool operator!=(const TReversePointerIterator& Rhs) const
+	{
+		return Ptr != Rhs.Ptr;
+	}
+
+private:
+	T* Ptr;
+};
+
+namespace UE::Core::Private
+{
+	template <typename RangeType>
+	struct TReverseIterationAdapter
+	{
+		constexpr explicit TReverseIterationAdapter(RangeType& InRange UE_LIFETIMEBOUND)
+			: Range(InRange)
+		{
+		}
+
+		TReverseIterationAdapter(TReverseIterationAdapter&&) = delete;
+		TReverseIterationAdapter& operator=(TReverseIterationAdapter&&) = delete;
+		TReverseIterationAdapter& operator=(const TReverseIterationAdapter&) = delete;
+		TReverseIterationAdapter(const TReverseIterationAdapter&) = delete;
+		~TReverseIterationAdapter() = default;
+
+		FORCEINLINE constexpr auto begin() const UE_LIFETIMEBOUND
+		{
+			using std::rbegin;
+			return rbegin(Range);
+		}
+
+		FORCEINLINE constexpr auto end() const UE_LIFETIMEBOUND
+		{
+			using std::rend;
+			return rend(Range);
+		}
+
+	private:
+		RangeType& Range;
+	};
+
+	template <typename ElementType, std::size_t N>
+	struct TReverseIterationAdapter<ElementType(&)[N]>
+	{
+		constexpr explicit TReverseIterationAdapter(ElementType* InArray UE_LIFETIMEBOUND)
+			: Array(InArray)
+		{
+		}
+
+		TReverseIterationAdapter(TReverseIterationAdapter&&) = delete;
+		TReverseIterationAdapter& operator=(TReverseIterationAdapter&&) = delete;
+		TReverseIterationAdapter& operator=(const TReverseIterationAdapter&) = delete;
+		TReverseIterationAdapter(const TReverseIterationAdapter&) = delete;
+		~TReverseIterationAdapter() = default;
+
+		inline constexpr TReversePointerIterator<ElementType> begin() const UE_LIFETIMEBOUND
+		{
+			return TReversePointerIterator<ElementType>(Array + N);
+		}
+
+		inline constexpr TReversePointerIterator<ElementType> end() const UE_LIFETIMEBOUND
+		{
+			return TReversePointerIterator<ElementType>(Array);
+		}
+
+	private:
+		ElementType* Array;
+	};
+}
+
+/**
+ * Allows a range to be iterated backwards.  The container must not be modified during iteration,
+ * but elements may be.  The container must outlive the adapter.
+ *
+ * @param  Range  The range to iterate over backwards.
+ *
+ * @return  An adapter which allows reverse iteration over the range.  The adapter is not intended
+ *          to be handled directly, but rather as a parameter to an algorithm or ranged-for loop.
+ *
+ * Example:
+ *
+ * TArray<int> Array = { 1, 2, 3, 4, 5 };
+ *
+ * for (int& Element : ReverseIterate(Array))
+ * {
+ *     UE_LOG(LogTemp, Log, TEXT("%d"), Element);
+ * }
+ * // Logs: 5 4 3 2 1
+ *
+ * TArray<FString> Output;
+ * Algo::Transform(ReverseIterate(Array), Output, UE_PROJECTION(LexToString));
+ * // Output: [ "5", "4", "3", "2", "1" ]
+ */
+template <typename RangeType>
+inline constexpr UE::Core::Private::TReverseIterationAdapter<RangeType> ReverseIterate(RangeType&& Range UE_LIFETIMEBOUND)
+{
+	return UE::Core::Private::TReverseIterationAdapter<RangeType>(Range);
+}
diff --git a/Engine/Source/Runtime/CoreUObject/Private/UObject/ScriptCore.cpp b/Engine/Source/Runtime/CoreUObject/Private/UObject/ScriptCore.cpp
--- a/Engine/Source/Runtime/CoreUObject/Private/UObject/ScriptCore.cpp
+++ b/Engine/Source/Runtime/CoreUObject/Private/UObject/ScriptCore.cpp
@@ -737,10 +737,10 @@ void FFrame::GetStackTrace(FStringBuilderBase& Result) const
 	if (FrameStack.Num() > 0)
 	{
 		Result << TEXT("Script call stack:\n");
-		for (int32 Index = FrameStack.Num() - 1; Index >= 0; Index--)
+		for (const FFrame* Frame : ReverseIterate(FrameStack))
 		{
 			Result << TEXT("\t");
-			FrameStack[Index]->Node->GetFullName(Result);
+			Frame->Node->GetFullName(Result);
 			Result << TEXT("\n");
 		}
 	}
diff --git a/Engine/Source/Runtime/Engine/Private/Actor.cpp b/Engine/Source/Runtime/Engine/Private/Actor.cpp
--- a/Engine/Source/Runtime/Engine/Private/Actor.cpp
+++ b/Engine/Source/Runtime/Engine/Private/Actor.cpp
@@ -3406,11 +3406,11 @@ void AActor::ClearInstanceComponents(const bool bDestroyComponents)
 		TArray<UActorComponent*> CachedComponents(InstanceComponents);
 
 		// Run in reverse to reduce memory churn when the components are removed from InstanceComponents
-		for (int32 Index=CachedComponents.Num()-1; Index >= 0; --Index)
+		for (UActorComponent* CachedComponent : ReverseIterate(CachedComponents))
 		{
-			if (CachedComponents[Index])
+			if (CachedComponent)
 			{
-				CachedComponents[Index]->DestroyComponent();
+				CachedComponent->DestroyComponent();
 			}
 		}
 	}
