StringBuilder: Converted WriteToString from a type to a function

This was not possible on C++14 and is possible now due to guaranteed copy elision in C++17.
This will break backward compatibility when WriteToString was used as a type or in a context where it would be copied.

diff --git a/Engine/Source/Developer/DerivedDataCache/Private/FileSystemCacheStore.cpp b/Engine/Source/Developer/DerivedDataCache/Private/FileSystemCacheStore.cpp
--- a/Engine/Source/Developer/DerivedDataCache/Private/FileSystemCacheStore.cpp
+++ b/Engine/Source/Developer/DerivedDataCache/Private/FileSystemCacheStore.cpp
@@ -2492,9 +2492,9 @@ TUniquePtr<FArchive> FFileSystemCacheStore::OpenFileRead(FStringBuilderBase& Pat
 	if (IsInGameThread() && FStringView(CachePath).StartsWith(TEXTVIEW("//"), ESearchCase::CaseSensitive))
 	{
 		FRequestOwner AsyncOwner(EPriority::Normal);
-		Private::LaunchTaskInCacheThreadPool(AsyncOwner, [this, Path = WriteToString<256>(Path)]() mutable
+		Private::LaunchTaskInCacheThreadPool(AsyncOwner, [this, Path = MakeShared<TStringBuilder<256>>(InPlace, Path)]() mutable
 		{
-			(void)FileExists(Path);
+			(void)FileExists(*Path);
 		});
 		AsyncOwner.KeepAlive();
 	}
diff --git a/Engine/Source/Editor/UnrealEd/Private/Commandlets/AssetRegistryGenerator.cpp b/Engine/Source/Editor/UnrealEd/Private/Commandlets/AssetRegistryGenerator.cpp
--- a/Engine/Source/Editor/UnrealEd/Private/Commandlets/AssetRegistryGenerator.cpp
+++ b/Engine/Source/Editor/UnrealEd/Private/Commandlets/AssetRegistryGenerator.cpp
@@ -1429,14 +1429,13 @@ void FAssetRegistryGenerator::ComputePackageDifferences(const FComputeDifference
 				// mark them as modified even if their map package is modified. Doing so would mark all actors in the map as
 				// modified anytime one of them changed.
 				// Workaround: Detect external actors by naming convention and suppress their reference to the map package.
-				WriteToString<256> ReferencerPackageNameStr(ReferencerPackageName);
+				TStringBuilder<256> ReferencerPackageNameStr(InPlace, ReferencerPackageName);
 				int32 ExternalActorsFolderIndex = ReferencerPackageNameStr.ToView().Find(ExternalActorsFolderName);
 				if (ExternalActorsFolderIndex != INDEX_NONE)
 				{
 					if (ModifiedPackageLeafName.IsEmpty())
 					{
-						WriteToString<256> ModifiedPackageNameStr(ModifiedPackage);
-						ModifiedPackageLeafName = FPathViews::GetCleanFilename(ModifiedPackageNameStr);
+						ModifiedPackageLeafName = FPathViews::GetCleanFilename(WriteToString<256>(ModifiedPackage));
 					}
 					FStringView GeneratedRelativePath = ReferencerPackageNameStr.ToView().RightChop(ExternalActorsFolderIndex + ExternalActorsFolderName.Len());
 					if (GeneratedRelativePath.Contains(ModifiedPackageLeafName))
diff --git a/Engine/Source/Editor/UnrealEd/Private/CookOnTheFlyServer.cpp b/Engine/Source/Editor/UnrealEd/Private/CookOnTheFlyServer.cpp
--- a/Engine/Source/Editor/UnrealEd/Private/CookOnTheFlyServer.cpp
+++ b/Engine/Source/Editor/UnrealEd/Private/CookOnTheFlyServer.cpp
@@ -2502,7 +2502,7 @@ void UCookOnTheFlyServer::DemoteToIdle(UE::Cook::FPackageData& PackageData, UE::
 			// Suppress the message in cases that cause large spam like NotInCurrentPlugin for DLC cooks.
 			if (Reason != ESuppressCookReason::NotInCurrentPlugin)
 			{
-				WriteToString<256> PackageNameStr(PackageData.GetPackageName());
+				TStringBuilder<256> PackageNameStr(InPlace, PackageData.GetPackageName());
 
 				// ExternalActors: Do not send a message for every NeverCook external Actor package; too much spam
 				if (!(Reason == ESuppressCookReason::NeverCook &&
@@ -11529,7 +11529,7 @@ void UCookOnTheFlyServer::ReportHiddenDependency(FName Referencer, FName Depende
 static
 void ConditionalWaitOnCommandFile(FStringView GateName, TFunctionRef<void (FStringView)> CommandHandler)
 {
-	WriteToString<128> ArgPrefix(TEXTVIEW("-"), GateName, TEXTVIEW("WaitOnCommandFile="));
+	TStringBuilder<128> ArgPrefix(InPlace, TEXTVIEW("-"), GateName, TEXTVIEW("WaitOnCommandFile="));
 
 	FString WaitOnCommandFile;
 	if (!FParse::Value(FCommandLine::Get(), *ArgPrefix, WaitOnCommandFile))
diff --git a/Engine/Source/Runtime/AssetRegistry/Private/AssetDataGatherer.cpp b/Engine/Source/Runtime/AssetRegistry/Private/AssetDataGatherer.cpp
--- a/Engine/Source/Runtime/AssetRegistry/Private/AssetDataGatherer.cpp
+++ b/Engine/Source/Runtime/AssetRegistry/Private/AssetDataGatherer.cpp
@@ -3468,8 +3468,7 @@ void FAssetDataGatherer::TickInternal(bool& bOutIsTickInterrupt, double& TickSta
 		if (AssetFileData.Type == EGatherableFileType::VerseFile)
 		{
 			// Store Verse results in a hybrid format using the LongPackageName but keeping the extension
-			WriteToString<256> MappedPath(AssetFileData.LongPackageName, FPathViews::GetExtension(AssetFileData.LocalAbsPath, true));
-			LocalVerseResults.Add(*MappedPath);
+			LocalVerseResults.Emplace(WriteToString<256>(AssetFileData.LongPackageName, FPathViews::GetExtension(AssetFileData.LocalAbsPath, true)));
 			continue;
 		}
 
diff --git a/Engine/Source/Runtime/AssetRegistry/Private/AssetRegistry.cpp b/Engine/Source/Runtime/AssetRegistry/Private/AssetRegistry.cpp
--- a/Engine/Source/Runtime/AssetRegistry/Private/AssetRegistry.cpp
+++ b/Engine/Source/Runtime/AssetRegistry/Private/AssetRegistry.cpp
@@ -5069,8 +5069,7 @@ void FAssetRegistryImpl::VerseFilesGathered(const double TickStartTime, TRingBuf
 		CachedVerseFiles.Add(VerseFilePath, &bAlreadyExists);
 		if (!bAlreadyExists)
 		{
-			WriteToString<256> VerseFilePathString(VerseFilePath);
-			FName VerseDirectoryPath(FPathViews::GetPath(VerseFilePathString.ToView()));
+			FName VerseDirectoryPath(FPathViews::GetPath(WriteToString<256>(VerseFilePath)));
 			TArray<FName>& FilePathsArray = CachedVerseFilesByPath.FindOrAdd(VerseDirectoryPath);
 			FilePathsArray.Add(VerseFilePath);
 		}
@@ -5332,8 +5331,7 @@ void FAssetRegistryImpl::RemoveVerseFile(FName VerseFilePathToRemove)
 {
 	if (CachedVerseFiles.Remove(VerseFilePathToRemove))
 	{
-		WriteToString<256> VerseFilePathString(VerseFilePathToRemove);
-		FName VerseDirectoryPath(FPathViews::GetPath(VerseFilePathString.ToView()));
+		FName VerseDirectoryPath(FPathViews::GetPath(WriteToString<256>(VerseFilePathToRemove)));
 		TArray<FName>* FilePathsArray = CachedVerseFilesByPath.Find(VerseDirectoryPath);
 		if (ensure(FilePathsArray)) // We found it in CachedVerseFiles, so we must also find it here
 		{
@@ -5471,8 +5469,7 @@ void FAssetRegistryImpl::OnDirectoryChanged(Impl::FEventContext& EventContext,
 					break;
 
 				case FFileChangeData::FCA_Removed:
-					WriteToString<256> VerseFilePathString(LongPackageName, FPathViews::GetExtension(File, true));
-					RemoveVerseFile(*VerseFilePathString);
+					RemoveVerseFile(FName(WriteToString<256>(LongPackageName, FPathViews::GetExtension(File, /*bIncludeDot*/ true))));
 					UE_LOG(LogAssetRegistry, Verbose, TEXT("Verse file was removed from content directory: %s"), *File);
 					break;
 				}
diff --git a/Engine/Source/Runtime/Core/Public/Misc/StringBuilder.h b/Engine/Source/Runtime/Core/Public/Misc/StringBuilder.h
--- a/Engine/Source/Runtime/Core/Public/Misc/StringBuilder.h
+++ b/Engine/Source/Runtime/Core/Public/Misc/StringBuilder.h
@@ -412,6 +412,16 @@ public:
 	{
 	}
 
+	/**
+	 * Construct a string builder by appending the arguments using operator<<.
+	 */
+	template <typename... ArgTypes>
+	explicit TStringBuilderWithBuffer(EInPlace, ArgTypes&&... Args)
+		: TStringBuilderBase<CharType>(StringBuffer, BufferSize)
+	{
+		(*this << ... << (ArgTypes&&)Args);
+	}
+
 	using TStringBuilderBase<CharType>::operator=;
 
 private:
@@ -483,40 +493,59 @@ inline FWideStringBuilderBase&		operator<<(FWideStringBuilderBase& Builder, uint
 inline FUtf8StringBuilderBase&		operator<<(FUtf8StringBuilderBase& Builder, int16 Value)							{ return Builder << int32(Value); }
 inline FUtf8StringBuilderBase&		operator<<(FUtf8StringBuilderBase& Builder, uint16 Value)							{ return Builder << uint32(Value); }
 
-/**
- * A function-like type that creates a TStringBuilder by appending its arguments.
- *
- * Example Use Cases:
- *
- * For void Action(FStringView) -> Action(WriteToString<64>(Arg1, Arg2));
- * For UE_LOG or checkf -> checkf(Condition, TEXT("%s"), *WriteToString<32>(Arg));
- */
 template <typename CharType, int32 BufferSize>
-class TWriteToString : public TStringBuilderWithBuffer<CharType, BufferSize>
+class UE_DEPRECATED(5.3, "Use WriteToString<N>(...) or TStringBuilder<N>(InPlace, ...).") TWriteToString : public TStringBuilderWithBuffer<CharType, BufferSize>
 {
 public:
 	template <typename... ArgTypes>
 	explicit TWriteToString(ArgTypes&&... Args)
 	{
-	#if PLATFORM_COMPILER_HAS_FOLD_EXPRESSIONS
-		(*this << ... << Forward<ArgTypes>(Args));
-	#else
-		using Fold = int[];
-		void(Fold{0, (void(*this << Forward<ArgTypes>(Args)), 0)...});
-	#endif
+		(*this << ... << (ArgTypes&&)Args);
 	}
 };
 
+PRAGMA_DISABLE_DEPRECATION_WARNINGS
 template <typename CharType, int32 BufferSize>
 struct TIsContiguousContainer<TWriteToString<CharType, BufferSize>>
 {
 	static constexpr inline bool Value = true;
 };
+PRAGMA_ENABLE_DEPRECATION_WARNINGS
+
+/**
+ * A function to create and append to a temporary string builder.
+ *
+ * Example Use Cases:
+ *
+ * For void Action(FStringView) -> Action(WriteToString<64>(Arg1, Arg2));
+ * For UE_LOG or checkf -> checkf(Condition, TEXT("%s"), *WriteToString<32>(Arg));
+ */
+template <int32 BufferSize, typename... ArgTypes>
+TStringBuilderWithBuffer<TCHAR, BufferSize> WriteToString(ArgTypes&&... Args)
+{
+	return TStringBuilderWithBuffer<TCHAR, BufferSize>(InPlace, (ArgTypes&&)Args...);
+}
 
-template <int32 BufferSize> using WriteToString = TWriteToString<TCHAR, BufferSize>;
-template <int32 BufferSize> using WriteToAnsiString = TWriteToString<ANSICHAR, BufferSize>;
-template <int32 BufferSize> using WriteToWideString = TWriteToString<WIDECHAR, BufferSize>;
-template <int32 BufferSize> using WriteToUtf8String = TWriteToString<UTF8CHAR, BufferSize>;
+/** A function to create and append to a temporary string builder. See WriteToString. */
+template <int32 BufferSize, typename... ArgTypes>
+TStringBuilderWithBuffer<ANSICHAR, BufferSize> WriteToAnsiString(ArgTypes&&... Args)
+{
+	return TStringBuilderWithBuffer<ANSICHAR, BufferSize>(InPlace, (ArgTypes&&)Args...);
+}
+
+/** A function to create and append to a temporary string builder. See WriteToString. */
+template <int32 BufferSize, typename... ArgTypes>
+TStringBuilderWithBuffer<WIDECHAR, BufferSize> WriteToWideString(ArgTypes&&... Args)
+{
+	return TStringBuilderWithBuffer<WIDECHAR, BufferSize>(InPlace, (ArgTypes&&)Args...);
+}
+
+/** A function to create and append to a temporary string builder. See WriteToString. */
+template <int32 BufferSize, typename... ArgTypes>
+TStringBuilderWithBuffer<UTF8CHAR, BufferSize> WriteToUtf8String(ArgTypes&&... Args)
+{
+	return TStringBuilderWithBuffer<UTF8CHAR, BufferSize>(InPlace, (ArgTypes&&)Args...);
+}
 
 /**
  * Returns an object that can be used as the output container for algorithms by appending to the builder.
diff --git a/Engine/Source/Runtime/Core/Tests/Containers/StringViewTest.cpp b/Engine/Source/Runtime/Core/Tests/Containers/StringViewTest.cpp
--- a/Engine/Source/Runtime/Core/Tests/Containers/StringViewTest.cpp
+++ b/Engine/Source/Runtime/Core/Tests/Containers/StringViewTest.cpp
@@ -122,9 +122,9 @@ TEST_CASE_NAMED(FStringViewConstructorTest, "System::Core::String::StringView::C
 		STATIC_CHECK(std::is_same_v<decltype(ViewWideConstLiteral)::ElementType, WIDECHAR>);
 		STATIC_CHECK(std::is_same_v<decltype(ViewUtf8ConstLiteral)::ElementType, UTF8CHAR>);
 
-		TWriteToString<ANSICHAR, 16> WriteAnsiString = WriteToAnsiString<16>("Test");
- 		TWriteToString<WIDECHAR, 16> WriteWideString = WriteToWideString<16>(WIDETEXT("Test"));
- 		TWriteToString<UTF8CHAR, 16> WriteUft8String = WriteToUtf8String<16>(UTF8TEXT("Test"));
+		auto WriteAnsiString = WriteToAnsiString<16>("Test");
+ 		auto WriteWideString = WriteToWideString<16>(WIDETEXT("Test"));
+ 		auto WriteUft8String = WriteToUtf8String<16>(UTF8TEXT("Test"));
 		TStringView ViewAnsiStringBuilder(WriteAnsiString);
 		TStringView ViewWideStringBuilder(WriteWideString);
 		TStringView ViewUtf8StringBuilder(WriteUft8String);
diff --git a/Engine/Source/Runtime/CoreUObject/Private/AssetRegistry/AssetData.cpp b/Engine/Source/Runtime/CoreUObject/Private/AssetRegistry/AssetData.cpp
--- a/Engine/Source/Runtime/CoreUObject/Private/AssetRegistry/AssetData.cpp
+++ b/Engine/Source/Runtime/CoreUObject/Private/AssetRegistry/AssetData.cpp
@@ -722,8 +722,7 @@ bool FAssetData::TryNetworkRead(FCbFieldView Field, bool bReadPackageName, FName
 	{
 		if (bHasAssetName)
 		{
-			WriteToString<256> PackageNameStr(InPackageName);
-			PackagePath = FName(FPathViews::GetPath(PackageNameStr.ToView()));
+			PackagePath = FName(FPathViews::GetPath(WriteToString<256>(InPackageName)));
 		}
 		else
 		{
diff --git a/Engine/Source/Runtime/CoreUObject/Private/UObject/LinkerLoad.cpp b/Engine/Source/Runtime/CoreUObject/Private/UObject/LinkerLoad.cpp
--- a/Engine/Source/Runtime/CoreUObject/Private/UObject/LinkerLoad.cpp
+++ b/Engine/Source/Runtime/CoreUObject/Private/UObject/LinkerLoad.cpp
@@ -93,15 +93,14 @@ void TrackPackageAssetClass(UPackage* Package, FLinkerLoad& LinkerLoad, const TA
 	}
 
 	FName PackageName = Package->GetFName();
-	WriteToString<256> PackageNameStr(PackageName);
-	FStringView PackageLeafName = FPathViews::GetCleanFilename(PackageNameStr.ToView());
+	TStringBuilder<256> PackageNameStr(InPlace, PackageName);
+	FStringView PackageLeafName = FPathViews::GetCleanFilename(PackageNameStr);
 	const FObjectExport* MostImportant = nullptr;
 	for (const FObjectExport& Export : Exports)
 	{
 		if (Export.bIsAsset && Export.ClassIndex.IsImport())
 		{
-			WriteToString<256> ObjectName(Export.ObjectName);
-			if (ObjectName.ToView() == PackageLeafName)
+			if (WriteToString<256>(Export.ObjectName) == PackageLeafName)
 			{
 				MostImportant = &Export;
 				break;
