Logging: Implemented Fatal verbosity for structured logging and fixed use cases broken due to compiler issues

- Fixed use of a category imported from another module
- Fixed logging in an inline function on an imported type
- Fixed logging in a virtual function in a class template
- Fixed logging in a function template

#rnx

diff --git a/Engine/Source/Runtime/Core/Private/Logging/StructuredLog.cpp b/Engine/Source/Runtime/Core/Private/Logging/StructuredLog.cpp
--- a/Engine/Source/Runtime/Core/Private/Logging/StructuredLog.cpp
+++ b/Engine/Source/Runtime/Core/Private/Logging/StructuredLog.cpp
@@ -22,15 +22,8 @@
 #include "Serialization/CompactBinaryValue.h"
 #include "Serialization/VarInt.h"
 
-void StaticFailDebugV(
-	const TCHAR* Error,
-	const ANSICHAR* Expression,
-	const ANSICHAR* File,
-	int32 Line,
-	bool bIsEnsure,
-	void* ProgramCounter,
-	const TCHAR* DescriptionFormat,
-	va_list DescriptionArgs);
+void VARARGS StaticFailDebug (const TCHAR* Error, const ANSICHAR* Expression, const ANSICHAR* File, int32 Line, bool bIsEnsure, void* ProgramCounter, const TCHAR* DescriptionFormat, ...);
+void         StaticFailDebugV(const TCHAR* Error, const ANSICHAR* Expression, const ANSICHAR* File, int32 Line, bool bIsEnsure, void* ProgramCounter, const TCHAR* DescriptionFormat, va_list DescriptionArgs);
 
 namespace UE::Logging::Private
 {
@@ -706,10 +699,10 @@ FORCENOINLINE static FCbObject SerializeLogFields(
 }
 
 template <typename StaticLogRecordType>
-FORCENOINLINE static FLogTemplate& CreateLogTemplate(const StaticLogRecordType& Log, const FLogField* Fields, const int32 FieldCount)
+FORCENOINLINE static FLogTemplate& CreateLogTemplate(const FLogCategoryBase& Category, const StaticLogRecordType& Log, const FLogField* Fields, const int32 FieldCount)
 {
 #if LOGTRACE_ENABLED
-	FLogTrace::OutputLogMessageSpec(&Log, &Log.Category, Log.Verbosity, Log.File, Log.Line, TEXT("%s"));
+	FLogTrace::OutputLogMessageSpec(&Log, &Category, Log.Verbosity, Log.File, Log.Line, TEXT("%s"));
 #endif
 
 	FLogTemplate* NewTemplate = FLogTemplate::CreateStatic(Log, Fields, FieldCount);
@@ -734,19 +727,19 @@ FORCENOINLINE static FLogTemplate& CreateLogTemplate(const StaticLogRecordType&
 }
 
 template <typename StaticLogRecordType>
-inline static FLogTemplate& EnsureLogTemplate(const StaticLogRecordType& Log, const FLogField* Fields, const int32 FieldCount)
+inline static FLogTemplate& EnsureLogTemplate(const FLogCategoryBase& Category, const StaticLogRecordType& Log, const FLogField* Fields, const int32 FieldCount)
 {
 	if (FLogTemplate* Template = Log.DynamicData.Template.load(std::memory_order_acquire); LIKELY(Template))
 	{
 		return *Template;
 	}
-	return CreateLogTemplate(Log, Fields, FieldCount);
+	return CreateLogTemplate(Category, Log, Fields, FieldCount);
 }
 
 template <typename StaticLogRecordType>
-inline static FLogRecord CreateLogRecord(const StaticLogRecordType& Log, const FLogField* Fields, const int32 FieldCount)
+inline static FLogRecord CreateLogRecord(const FLogCategoryBase& Category, const StaticLogRecordType& Log, const FLogField* Fields, const int32 FieldCount)
 {
-	FLogTemplate& Template = EnsureLogTemplate(Log, Fields, FieldCount);
+	FLogTemplate& Template = EnsureLogTemplate(Category, Log, Fields, FieldCount);
 
 	FLogRecord Record;
 	Record.SetFormat(Log.Format);
@@ -754,7 +747,7 @@ inline static FLogRecord CreateLogRecord(const StaticLogRecordType& Log, const F
 	Record.SetFields(SerializeLogFields(Log, Template, Fields, FieldCount));
 	Record.SetFile(Log.File);
 	Record.SetLine(Log.Line);
-	Record.SetCategory(Log.Category.GetCategoryName());
+	Record.SetCategory(Category.GetCategoryName());
 	Record.SetVerbosity(Log.Verbosity);
 	Record.SetTime(FLogTime::Now());
 	return Record;
@@ -785,29 +778,75 @@ inline static void DispatchLogRecord(const StaticLogRecordType& Log, const FLogR
 	(OutputDevice ? OutputDevice : GLog)->SerializeRecord(Record);
 }
 
-void LogWithFieldArray(const FStaticLogRecord& Log, const FLogField* Fields, const int32 FieldCount)
+void LogWithFieldArray(const FLogCategoryBase& Category, const FStaticLogRecord& Log, const FLogField* Fields, const int32 FieldCount)
 {
-	DispatchLogRecord(Log, CreateLogRecord(Log, Fields, FieldCount));
+#if !NO_LOGGING
+	DispatchLogRecord(Log, CreateLogRecord(Category, Log, Fields, FieldCount));
+#endif
 }
 
-void LogWithNoFields(const FStaticLogRecord& Log)
+void LogWithNoFields(const FLogCategoryBase& Category, const FStaticLogRecord& Log)
 {
 	// A non-null field pointer enables field validation in FLogTemplate::Create.
 	static constexpr FLogField EmptyField{};
-	LogWithFieldArray(Log, &EmptyField, 0);
+	LogWithFieldArray(Category, Log, &EmptyField, 0);
 }
 
-void LogWithFieldArray(const FStaticLocalizedLogRecord& Log, const FLogField* Fields, const int32 FieldCount)
+void LogWithFieldArray(const FLogCategoryBase& Category, const FStaticLocalizedLogRecord& Log, const FLogField* Fields, const int32 FieldCount)
 {
-	FLogRecord Record = CreateLogRecord(Log, Fields, FieldCount);
+	FLogRecord Record = CreateLogRecord(Category, Log, Fields, FieldCount);
 	Record.SetTextNamespace(Log.TextNamespace);
 	Record.SetTextKey(Log.TextKey);
 	DispatchLogRecord(Log, Record);
 }
 
-void LogWithNoFields(const FStaticLocalizedLogRecord& Log)
+void LogWithNoFields(const FLogCategoryBase& Category, const FStaticLocalizedLogRecord& Log)
 {
-	LogWithFieldArray(Log, nullptr, 0);
+	LogWithFieldArray(Category, Log, nullptr, 0);
+}
+
+void FatalLogWithFieldArray(const FLogCategoryBase& Category, const FStaticLogRecord& Log, const FLogField* Fields, const int32 FieldCount)
+{
+#if !NO_LOGGING
+	FLogRecord Record = CreateLogRecord(Category, Log, Fields, FieldCount);
+	TStringBuilder<512> Message;
+	Record.FormatMessageTo(Message);
+
+	StaticFailDebug(TEXT("Fatal error:"), "", Log.File, Log.Line, /*bIsEnsure*/ false, PLATFORM_RETURN_ADDRESS(), TEXT("%s"), *Message);
+	FDebug::AssertFailed("", Log.File, Log.Line, TEXT("%s"), *Message);
+
+	UE_DEBUG_BREAK_AND_PROMPT_FOR_REMOTE();
+	FDebug::ProcessFatalError(PLATFORM_RETURN_ADDRESS());
+#endif
+}
+
+void FatalLogWithNoFields(const FLogCategoryBase& Category, const FStaticLogRecord& Log)
+{
+	// A non-null field pointer enables field validation in FLogTemplate::Create.
+	static constexpr FLogField EmptyField{};
+	FatalLogWithFieldArray(Category, Log, &EmptyField, 0);
+}
+
+void FatalLogWithFieldArray(const FLogCategoryBase& Category, const FStaticLocalizedLogRecord& Log, const FLogField* Fields, const int32 FieldCount)
+{
+#if !NO_LOGGING
+	FLogRecord Record = CreateLogRecord(Category, Log, Fields, FieldCount);
+	TStringBuilder<512> Message;
+	Record.FormatMessageTo(Message);
+
+	StaticFailDebug(TEXT("Fatal error:"), "", Log.File, Log.Line, /*bIsEnsure*/ false, PLATFORM_RETURN_ADDRESS(), TEXT("%s"), *Message);
+	FDebug::AssertFailed("", Log.File, Log.Line, TEXT("%s"), *Message);
+
+	UE_DEBUG_BREAK_AND_PROMPT_FOR_REMOTE();
+	FDebug::ProcessFatalError(PLATFORM_RETURN_ADDRESS());
+#endif
+}
+
+void FatalLogWithNoFields(const FLogCategoryBase& Category, const FStaticLocalizedLogRecord& Log)
+{
+	// A non-null field pointer enables field validation in FLogTemplate::Create.
+	static constexpr FLogField EmptyField{};
+	FatalLogWithFieldArray(Category, Log, &EmptyField, 0);
 }
 
 } // UE::Logging::Private
@@ -944,6 +983,7 @@ void BasicFatalLog(const FLogCategoryBase& Category, const FStaticBasicLogRecord
 	va_start(Args, Log);
 	FDebug::AssertFailedV("", Log->File, Log->Line, Log->Format, Args);
 	va_end(Args);
+
 	UE_DEBUG_BREAK_AND_PROMPT_FOR_REMOTE();
 	FDebug::ProcessFatalError(PLATFORM_RETURN_ADDRESS());
 #endif
diff --git a/Engine/Source/Runtime/Core/Public/Logging/StructuredLog.h b/Engine/Source/Runtime/Core/Public/Logging/StructuredLog.h
--- a/Engine/Source/Runtime/Core/Public/Logging/StructuredLog.h
+++ b/Engine/Source/Runtime/Core/Public/Logging/StructuredLog.h
@@ -234,18 +234,31 @@ struct FStaticLogDynamicData
 /** Data about a static log that is constant for every occurrence. */
 struct FStaticLogRecord
 {
-	const FLogCategoryBase& Category;
 	const TCHAR* Format = nullptr;
 	const ANSICHAR* File = nullptr;
 	int32 Line = 0;
 	ELogVerbosity::Type Verbosity = ELogVerbosity::Log;
 	FStaticLogDynamicData& DynamicData;
+
+	// Workaround for https://developercommunity.visualstudio.com/t/Incorrect-warning-C4700-with-unrelated-s/10285950
+	constexpr FStaticLogRecord(
+		const TCHAR* InFormat,
+		const ANSICHAR* InFile,
+		int32 InLine,
+		ELogVerbosity::Type InVerbosity,
+		FStaticLogDynamicData& InDynamicData)
+		: Format(InFormat)
+		, File(InFile)
+		, Line(InLine)
+		, Verbosity(InVerbosity)
+		, DynamicData(InDynamicData)
+	{
+	}
 };
 
 /** Data about a static localized log that is constant for every occurrence. */
 struct FStaticLocalizedLogRecord
 {
-	const FLogCategoryBase& Category;
 	const TCHAR* TextNamespace = nullptr;
 	const TCHAR* TextKey = nullptr;
 	const TCHAR* Format = nullptr;
@@ -253,6 +266,25 @@ struct FStaticLocalizedLogRecord
 	int32 Line = 0;
 	ELogVerbosity::Type Verbosity = ELogVerbosity::Log;
 	FStaticLogDynamicData& DynamicData;
+
+	// Workaround for https://developercommunity.visualstudio.com/t/Incorrect-warning-C4700-with-unrelated-s/10285950
+	constexpr FStaticLocalizedLogRecord(
+		const TCHAR* InTextNamespace,
+		const TCHAR* InTextKey,
+		const TCHAR* InFormat,
+		const ANSICHAR* InFile,
+		int32 InLine,
+		ELogVerbosity::Type InVerbosity,
+		FStaticLogDynamicData& InDynamicData)
+		: TextNamespace(InTextNamespace)
+		, TextKey(InTextKey)
+		, Format(InFormat)
+		, File(InFile)
+		, Line(InLine)
+		, Verbosity(InVerbosity)
+		, DynamicData(InDynamicData)
+	{
+	}
 };
 
 struct FLogField
@@ -277,11 +309,17 @@ struct FLogField
 namespace UE::Logging::Private
 {
 
-UE_API void LogWithNoFields(const FStaticLogRecord& Log);
-UE_API void LogWithFieldArray(const FStaticLogRecord& Log, const FLogField* Fields, int32 FieldCount);
+UE_API void LogWithNoFields(const FLogCategoryBase& Category, const FStaticLogRecord& Log);
+UE_API void LogWithFieldArray(const FLogCategoryBase& Category, const FStaticLogRecord& Log, const FLogField* Fields, int32 FieldCount);
+
+UE_API void FatalLogWithNoFields(const FLogCategoryBase& Category, const FStaticLogRecord& Log);
+UE_API void FatalLogWithFieldArray(const FLogCategoryBase& Category, const FStaticLogRecord& Log, const FLogField* Fields, int32 FieldCount);
+
+UE_API void LogWithNoFields(const FLogCategoryBase& Category, const FStaticLocalizedLogRecord& Log);
+UE_API void LogWithFieldArray(const FLogCategoryBase& Category, const FStaticLocalizedLogRecord& Log, const FLogField* Fields, int32 FieldCount);
 
-UE_API void LogWithNoFields(const FStaticLocalizedLogRecord& Log);
-UE_API void LogWithFieldArray(const FStaticLocalizedLogRecord& Log, const FLogField* Fields, int32 FieldCount);
+UE_API void FatalLogWithNoFields(const FLogCategoryBase& Category, const FStaticLocalizedLogRecord& Log);
+UE_API void FatalLogWithFieldArray(const FLogCategoryBase& Category, const FStaticLocalizedLogRecord& Log, const FLogField* Fields, int32 FieldCount);
 
 /** Wrapper to identify field names interleaved with field values. */
 template <typename NameType>
@@ -342,46 +380,60 @@ struct TFieldArgType<TLogFieldName<NameType>>
 };
 
 /** Log with fields created from the arguments, which may be values or pairs of name/value. */
-template <const auto& Log, typename... FieldArgTypes>
-FORCENOINLINE UE_DEBUG_SECTION void LogWithFields(typename TFieldArgType<FieldArgTypes>::Type... FieldArgs)
+template <typename... FieldArgTypes, typename LogType>
+FORCENOINLINE UE_DEBUG_SECTION void LogWithFields(const FLogCategoryBase& Category, const LogType& Log, typename TFieldArgType<FieldArgTypes>::Type... FieldArgs)
+{
+	constexpr int32 FieldCount = FLogFieldCreator::template GetCount<FieldArgTypes...>();
+	static_assert(FieldCount > 0);
+	FLogField Fields[FieldCount];
+	FLogFieldCreator::Create(Fields, (FieldArgTypes&&)FieldArgs...);
+	LogWithFieldArray(Category, Log, Fields, FieldCount);
+}
+
+/** Fatal log with fields created from the arguments, which may be values or pairs of name/value. */
+template <typename... FieldArgTypes, typename LogType>
+inline void FatalLogWithFields(const FLogCategoryBase& Category, const LogType& Log, typename TFieldArgType<FieldArgTypes>::Type... FieldArgs)
 {
 	constexpr int32 FieldCount = FLogFieldCreator::template GetCount<FieldArgTypes...>();
 	static_assert(FieldCount > 0);
 	FLogField Fields[FieldCount];
 	FLogFieldCreator::Create(Fields, (FieldArgTypes&&)FieldArgs...);
-	LogWithFieldArray(Log, Fields, FieldCount);
+	FatalLogWithFieldArray(Category, Log, Fields, FieldCount);
 }
 
 /** Log if the category is active at this level of verbosity. */
-template <const auto& Log, typename LogCategoryType, ELogVerbosity::Type Verbosity, typename... FieldArgTypes>
-inline void LogIfActive(FieldArgTypes&&... FieldArgs)
+template <typename LogCategoryType, ELogVerbosity::Type Verbosity, typename LogRecordType, typename... FieldArgTypes>
+inline void LogIfActive(const LogCategoryType& Category, const LogRecordType& Log, FieldArgTypes&&... FieldArgs)
 {
-	static_assert((Verbosity & ELogVerbosity::VerbosityMask) != ELogVerbosity::Fatal, "Fatal verbosity is not supported by this API at this time.");
 	static_assert((Verbosity & ELogVerbosity::VerbosityMask) < ELogVerbosity::NumVerbosity && Verbosity > 0, "Verbosity must be constant and in range.");
-
-	if constexpr (
+	if constexpr ((Verbosity & ELogVerbosity::VerbosityMask) == ELogVerbosity::Fatal)
+	{
+		if constexpr (sizeof...(FieldArgTypes) == 0)
+		{
+			FatalLogWithNoFields(Category, Log);
+		}
+		else
+		{
+			FatalLogWithFields<FieldArgTypes...>(Category, Log, (FieldArgTypes&&)FieldArgs...);
+		}
+		CA_ASSUME(false);
+	}
+	else if constexpr (
 		(Verbosity & ELogVerbosity::VerbosityMask) <= ELogVerbosity::COMPILED_IN_MINIMUM_VERBOSITY &&
 		(Verbosity & ELogVerbosity::VerbosityMask) <= LogCategoryType::CompileTimeVerbosity)
 	{
-		if ((Verbosity & ELogVerbosity::VerbosityMask) == ELogVerbosity::Fatal || !Log.Category.IsSuppressed(Verbosity))
+		if (!Category.IsSuppressed(Verbosity))
 		{
 			if constexpr (sizeof...(FieldArgTypes) == 0)
 			{
-				LogWithNoFields(Log);
+				LogWithNoFields(Category, Log);
 			}
 			else
 			{
-				LogWithFields<Log, FieldArgTypes...>((FieldArgTypes&&)FieldArgs...);
+				LogWithFields<FieldArgTypes...>(Category, Log, (FieldArgTypes&&)FieldArgs...);
 			}
 		}
 	}
-
-	if constexpr ((Verbosity & ELogVerbosity::VerbosityMask) == ELogVerbosity::Fatal)
-	{
-		UE_DEBUG_BREAK_AND_PROMPT_FOR_REMOTE();
-		FDebug::ProcessFatalError(PLATFORM_RETURN_ADDRESS());
-		CA_ASSUME(false);
-	}
 }
 
 } // UE::Logging::Private
@@ -390,8 +442,8 @@ inline void LogIfActive(FieldArgTypes&&... FieldArgs)
 	do \
 	{ \
 		static ::UE::Logging::Private::FStaticLogDynamicData LOG_Dynamic; \
-		static constexpr ::UE::Logging::Private::FStaticLogRecord LOG_Static{CategoryName, TEXT(Format), __FILE__, __LINE__, ::ELogVerbosity::Verbosity, LOG_Dynamic}; \
-		::UE::Logging::Private::LogIfActive<LOG_Static, FLogCategory##CategoryName, ::ELogVerbosity::Verbosity>(__VA_ARGS__); \
+		static constexpr ::UE::Logging::Private::FStaticLogRecord LOG_Static UE_PRIVATE_LOGFMT_AGGREGATE(TEXT(Format), __FILE__, __LINE__, ::ELogVerbosity::Verbosity, LOG_Dynamic); \
+		::UE::Logging::Private::LogIfActive<FLogCategory##CategoryName, ::ELogVerbosity::Verbosity>(CategoryName, LOG_Static, ##__VA_ARGS__); \
 	} \
 	while (false)
 
@@ -399,11 +451,15 @@ inline void LogIfActive(FieldArgTypes&&... FieldArgs)
 	do \
 	{ \
 		static ::UE::Logging::Private::FStaticLogDynamicData LOG_Dynamic; \
-		static constexpr ::UE::Logging::Private::FStaticLocalizedLogRecord LOG_Static{CategoryName, TEXT(Namespace), TEXT(Key), TEXT(Format), __FILE__, __LINE__, ::ELogVerbosity::Verbosity, LOG_Dynamic}; \
-		::UE::Logging::Private::LogIfActive<LOG_Static, FLogCategory##CategoryName, ::ELogVerbosity::Verbosity>(__VA_ARGS__); \
+		static constexpr ::UE::Logging::Private::FStaticLocalizedLogRecord LOG_Static UE_PRIVATE_LOGFMT_AGGREGATE(TEXT(Namespace), TEXT(Key), TEXT(Format), __FILE__, __LINE__, ::ELogVerbosity::Verbosity, LOG_Dynamic); \
+		::UE::Logging::Private::LogIfActive<FLogCategory##CategoryName, ::ELogVerbosity::Verbosity>(CategoryName, LOG_Static, ##__VA_ARGS__); \
 	} \
 	while (false)
 
+// This macro avoids having non-parenthesized commas in the log macros.
+// Such commas can cause issues when a log macro is wrapped by another macro.
+#define UE_PRIVATE_LOGFMT_AGGREGATE(...) {__VA_ARGS__}
+
 // This macro expands a field from either `(Name, Value)` or `Value`
 // A `(Name, Value)` field is converted to `CheckFieldName(Name), Value`
 // A `Value` field is passed through as `Value`
