Added warn macro in LowLevelTestAdapter.h

diff --git a/Engine/Source/Runtime/Core/Public/Misc/LowLevelTestAdapter.h b/Engine/Source/Runtime/Core/Public/Misc/LowLevelTestAdapter.h
--- a/Engine/Source/Runtime/Core/Public/Misc/LowLevelTestAdapter.h
+++ b/Engine/Source/Runtime/Core/Public/Misc/LowLevelTestAdapter.h
@@ -130,5 +130,6 @@ FString CaptureExpressionsAndValues(const FString& InExpressions, ArgTypes&&...
 #define FAIL_CHECK(Message) FAutomationTestFramework::Get().GetCurrentTest()->AddError(Message);
 
 #define CAPTURE(...) FAutomationTestFramework::Get().GetCurrentTest()->AddInfo(CaptureExpressionsAndValues(#__VA_ARGS__, __VA_ARGS__));
+#define WARN(Message) { FAutomationTestFramework::Get().GetCurrentTest()->AddWarning(Message); }
 
 #endif // !WITH_LOW_LEVEL_TESTS
\ No newline at end of file
