Fix TSAN warning because memory was always overwritten even if unchanged

#rnx

diff --git a/Engine/Source/Runtime/Core/Public/Templates/UnrealTemplate.h b/Engine/Source/Runtime/Core/Public/Templates/UnrealTemplate.h
--- a/Engine/Source/Runtime/Core/Public/Templates/UnrealTemplate.h
+++ b/Engine/Source/Runtime/Core/Public/Templates/UnrealTemplate.h
@@ -350,6 +350,49 @@ private:
 	AssignedType OldValue;
 };
 
+
+/**
+ * exception-safe guard around saving/restoring a value.
+ * Commonly used to make sure a value is restored
+ * even if the code early outs in the future.
+ * Usage:
+ *  	TOptionalGuardValue<bool> GuardSomeBool(bSomeBool, false); // Sets bSomeBool to false, and restores it in dtor.
+ */
+template <typename RefType, typename AssignedType = RefType>
+struct TOptionalGuardValue : private FNoncopyable
+{
+	TOptionalGuardValue(RefType& ReferenceValue, const AssignedType& NewValue)
+		: RefValue(ReferenceValue), OldValue(ReferenceValue)
+	{
+		if (RefValue != NewValue)
+		{
+			RefValue = NewValue;
+		}
+	}
+	~TOptionalGuardValue()
+	{
+		if (RefValue != OldValue)
+		{
+			RefValue = OldValue;
+		}
+	}
+
+	/**
+	 * Overloaded dereference operator.
+	 * Provides read-only access to the original value of the data being tracked by this struct
+	 *
+	 * @return	a const reference to the original data value
+	 */
+	FORCEINLINE const AssignedType& operator*() const
+	{
+		return OldValue;
+	}
+
+private:
+	RefType& RefValue;
+	AssignedType OldValue;
+};
+
 template <typename FuncType>
 struct TGuardValue_Bitfield_Cleanup : public FNoncopyable
 {
diff --git a/Engine/Source/Runtime/CoreUObject/Private/UObject/Obj.cpp b/Engine/Source/Runtime/CoreUObject/Private/UObject/Obj.cpp
--- a/Engine/Source/Runtime/CoreUObject/Private/UObject/Obj.cpp
+++ b/Engine/Source/Runtime/CoreUObject/Private/UObject/Obj.cpp
@@ -1354,7 +1354,9 @@ void UObject::Serialize(FStructuredArchive::FRecord Record)
 	bool bReportSoftObjectPathRedirects = false;
 
 	{
-		TGuardValue<bool*> GuardValue(
+		// TOptionalGuardValue will not overwrite the value if it remains the same.
+		// This is important for TSAN as we only want warnings if this unprotected value is changing.
+		TOptionalGuardValue<bool*> GuardValue(
 			GReportSoftObjectPathRedirects,
 			  GReportSoftObjectPathRedirects
 			? GReportSoftObjectPathRedirects
