Windows platform extensions no longer need customized Windows version checking logic.

#rnx

diff --git a/Engine/Source/Runtime/Core/Private/Windows/WindowsPlatformMisc.cpp b/Engine/Source/Runtime/Core/Private/Windows/WindowsPlatformMisc.cpp
--- a/Engine/Source/Runtime/Core/Private/Windows/WindowsPlatformMisc.cpp
+++ b/Engine/Source/Runtime/Core/Private/Windows/WindowsPlatformMisc.cpp
@@ -336,7 +336,6 @@ int32 GetOSVersionsHelper( TCHAR* OutOSVersionLabel, int32 OSVersionLabelLength,
 	return ErrorCode;
 }
 
-#if WINDOWS_USE_DEFAULT_OSVERSIONHELPER
 int32 FWindowsOSVersionHelper::GetOSVersions( FString& OutOSVersionLabel, FString& OutOSSubVersionLabel )
 {
 	TCHAR OSVersionLabel[128];
@@ -352,7 +351,6 @@ int32 FWindowsOSVersionHelper::GetOSVersions( FString& OutOSVersionLabel, FStrin
 
 	return Result;
 }
-#endif //WINDOWS_USE_DEFAULT_OSVERSIONHELPER
 
 namespace
 {
diff --git a/Engine/Source/Runtime/Core/Public/Windows/WindowsPlatform.h b/Engine/Source/Runtime/Core/Public/Windows/WindowsPlatform.h
--- a/Engine/Source/Runtime/Core/Public/Windows/WindowsPlatform.h
+++ b/Engine/Source/Runtime/Core/Public/Windows/WindowsPlatform.h
@@ -115,7 +115,6 @@ typedef FWindowsPlatformTypes FPlatformTypes;
 #define WINDOWS_USE_FEATURE_PLATFORMHTTP_CLASS				1
 #define WINDOWS_USE_FEATURE_LAUNCH							1
 #define WINDOWS_USE_FEATURE_DYNAMIC_RHI						1
-#define WINDOWS_USE_DEFAULT_OSVERSIONHELPER					1
 
 // Q: Why is there a __nop() before __debugbreak()?
 // A: VS' debug engine has a bug where it will silently swallow explicit
