Add FConditionVariable. This new FConditionVariable is philosophically like the WebKit WTF::Condition, but I wrote it from first principles rather than copying the code, so it doesn't have all of the magics (for example, no timeouts for now).

diff --git a/Engine/Source/Runtime/Core/Public/Experimental/Async/ConditionVariable.h b/Engine/Source/Runtime/Core/Public/Experimental/Async/ConditionVariable.h
new file mode 100644
--- /dev/null
+++ b/Engine/Source/Runtime/Core/Public/Experimental/Async/ConditionVariable.h
@@ -0,0 +1,70 @@
+// Copyright Epic Games, Inc. All Rights Reserved.
+
+#pragma once
+
+#include "Async/ParkingLot.h"
+
+#define UE_API CORE_API
+
+namespace UE
+{
+
+/**
+ * A one-byte portable condition variable. Gives the same decent performance everywhere.
+ */
+class FConditionVariable final
+{
+public:
+    constexpr FConditionVariable() = default;
+
+    FConditionVariable(const FConditionVariable&) = delete;
+    FConditionVariable& operator=(const FConditionVariable&) = delete;
+
+    void NotifyOne()
+    {
+        if (bHasWaiters)
+        {
+            FParkingLot::WakeOne(
+                &bHasWaiters,
+                [this] (FParkingLotWakeState WakeState) -> uint64
+                {
+                    if (!WakeState.bHasWaitingThreads)
+                        bHasWaiters = false;
+                    return 0;
+                });
+        }
+    }
+    
+    void NotifyAll()
+    {
+        if (bHasWaiters)
+        {
+            bHasWaiters = false;
+            FParkingLot::WakeAll(&bHasWaiters);
+        }
+    }
+
+    template<typename TLock>
+    void Wait(TLock& Lock)
+    {
+        FParkingLot::Wait(
+            &bHasWaiters,
+            [] () -> bool
+            {
+                bHasWaiters = true;
+                return true;
+            },
+            [&Lock] ()
+            {
+                Lock.Unlock();
+            });
+        Lock.Lock();
+    }
+    
+private:
+    std::atomic<bool> bHasWaiters = false;
+};
+
+} // namespace UE
+
+
