StringBuilder: Added UE_LIFETIMEBOUND to catch some cases where pointers are kept into temporary string builders

diff --git a/Engine/Source/Editor/AnimGraph/Private/AnimGraphNode_LinkedAnimLayer.cpp b/Engine/Source/Editor/AnimGraph/Private/AnimGraphNode_LinkedAnimLayer.cpp
--- a/Engine/Source/Editor/AnimGraph/Private/AnimGraphNode_LinkedAnimLayer.cpp
+++ b/Engine/Source/Editor/AnimGraph/Private/AnimGraphNode_LinkedAnimLayer.cpp
@@ -314,7 +314,7 @@ void UAnimGraphNode_LinkedAnimLayer::ValidateAnimNodeDuringCompilation(USkeleton
 void UAnimGraphNode_LinkedAnimLayer::ValidateCircularRefAndNesting(const UEdGraph* CurrentGraph, const TArray<UEdGraph*>& AllGraphs, TArray<const UEdGraph*> GraphStack, bool bWithinLinkedLayerGraph, FCompilerResultsLog& MessageLog)
 {
 	// Build graph chain string
-	auto BuildGraphChainString = [](const TArray<const UEdGraph*>& GraphStack)
+	auto BuildGraphChainString = [](const TArray<const UEdGraph*>& GraphStack) -> FString
 	{
 		TStringBuilder<1024> GraphChain;
 		bool bFirst = true;
@@ -327,7 +327,7 @@ void UAnimGraphNode_LinkedAnimLayer::ValidateCircularRefAndNesting(const UEdGrap
 			GraphChain << *Graph->GetName();
 			bFirst = false;
 		}
-		return GraphChain.ToString();
+		return FString(GraphChain);
 	};
 
 	// If a graph is already in the stack we have a circular reference
diff --git a/Engine/Source/Runtime/Core/Public/Misc/StringBuilder.h b/Engine/Source/Runtime/Core/Public/Misc/StringBuilder.h
--- a/Engine/Source/Runtime/Core/Public/Misc/StringBuilder.h
+++ b/Engine/Source/Runtime/Core/Public/Misc/StringBuilder.h
@@ -99,14 +99,37 @@ public:
 		Initialize(BufferPointer, BufferCapacity);
 	}
 
-	inline int32 Len() const					{ return int32(CurPos - Base); }
-	inline CharType* GetData()					{ return Base; }
-	inline const CharType* GetData() const		{ return Base; }
-	inline const CharType* ToString() const		{ EnsureNulTerminated(); return Base; }
-	inline ViewType ToView() const				{ return ViewType(Base, Len()); }
-	inline const CharType* operator*() const	{ EnsureNulTerminated(); return Base; }
-
-	inline const CharType	LastChar() const	{ return *(CurPos - 1); }
+	inline int32 Len() const { return int32(CurPos - Base); }
+
+	/** Returns a pointer to Len() code units that are not necessarily null-terminated. */
+	inline CharType* GetData() UE_LIFETIMEBOUND { return Base; }
+	inline const CharType* GetData() const UE_LIFETIMEBOUND { return Base; }
+
+	/** Returns a pointer to a null-terminated string that is valid until the builder is mutated. */
+	inline const CharType* ToString() const UE_LIFETIMEBOUND
+	{
+		EnsureNulTerminated();
+		return Base;
+	}
+
+	/** Returns a pointer to a null-terminated string that is valid until the builder is mutated. */
+	inline const CharType* operator*() const UE_LIFETIMEBOUND
+	{
+		EnsureNulTerminated();
+		return Base;
+	}
+
+	/** Returns a view of the string that is valid until the builder is mutated. */
+	inline ViewType ToView() const UE_LIFETIMEBOUND
+	{
+		return ViewType(Base, Len());
+	}
+
+	/** Returns the last character, technically the last code unit. */
+	inline const CharType LastChar() const
+	{
+		return *(CurPos - 1);
+	}
 
 	/**
 	 * Helper function to return the amount of memory allocated by this container.
