Remove totally unnecessary UE_API define from ConditionVariable.h.

diff --git a/Engine/Source/Runtime/Core/Public/Experimental/Async/ConditionVariable.h b/Engine/Source/Runtime/Core/Public/Experimental/Async/ConditionVariable.h
--- a/Engine/Source/Runtime/Core/Public/Experimental/Async/ConditionVariable.h
+++ b/Engine/Source/Runtime/Core/Public/Experimental/Async/ConditionVariable.h
@@ -4,8 +4,6 @@
 
 #include "Async/ParkingLot.h"
 
-#define UE_API CORE_API
-
 namespace UE
 {
 
