OodleDataCompression parallel encode and decode helpers
can be swapped in with existing OodleDataCompression calls
does not add any extra header

diff --git a/Engine/Source/Runtime/Core/Private/Compression/OodleDataCompression.cpp b/Engine/Source/Runtime/Core/Private/Compression/OodleDataCompression.cpp
--- a/Engine/Source/Runtime/Core/Private/Compression/OodleDataCompression.cpp
+++ b/Engine/Source/Runtime/Core/Private/Compression/OodleDataCompression.cpp
@@ -13,6 +13,7 @@
 #include "Misc/ConfigCacheIni.h"
 #include "ProfilingDebugging/CpuProfilerTrace.h"
 #include "Templates/UnrealTemplate.h"
+#include "Async/ParallelFor.h"
 #include "oodle2.h"
 #include "oodle2base.h"
 
@@ -426,7 +427,9 @@ struct OodleScratchBuffers
 	int64 OodleEncode( void * OutCompressedData, int64 InCompressedBufferSize,
 								const void * InUncompressedData, int64 InUncompressedSize,
 								ECompressor Compressor,
-								ECompressionLevel Level)
+								ECompressionLevel Level,
+								bool CompressIndependentChunks,
+								int64 DictionaryBackup)
 	{
 		TRACE_CPUPROFILER_EVENT_SCOPE(Oodle.Encode);
 
@@ -444,7 +447,18 @@ struct OodleScratchBuffers
 			return OODLELZ_FAILED;
 		}
 		
-		// try to take a mutex for one of the pre-allocated decode buffers
+		OodleLZ_CompressOptions LZOptions = * OodleLZ_CompressOptions_GetDefault(LZCompressor,LZLevel);
+		if ( CompressIndependentChunks )
+		{
+			LZOptions.seekChunkReset = true;
+			LZOptions.seekChunkLen = OODLELZ_BLOCK_LEN;
+
+			check( DictionaryBackup == 0 ); // seek chunk and DictionaryBackup are mutually exclusive
+		}
+
+		const uint8 * DictionaryStart = (const uint8 *)InUncompressedData - DictionaryBackup;
+
+		// try to take a mutex for one of the pre-allocated scratch buffers
 		for (int i = 0; i < OodleScratchBufferCount; ++i)
 		{
 			if (OodleScratches[i].OodleScratchMemoryMutex.TryLock()) 
@@ -465,7 +479,7 @@ struct OodleScratchBuffers
 				OO_SINTa scratchSize = OodleScratchMemorySize;
 
 				OO_SINTa Result = OodleLZ_Compress(LZCompressor,InUncompressedData,InUncompressedSize,OutCompressedData,LZLevel,
-															NULL,NULL,NULL,
+															&LZOptions,DictionaryStart,nullptr,
 															scratchMem,scratchSize);
 
 				OodleScratches[i].OodleScratchMemoryMutex.Unlock();
@@ -484,7 +498,7 @@ struct OodleScratchBuffers
 		OO_SINTa scratchSize = 0;
 
 		OO_SINTa Result = OodleLZ_Compress(LZCompressor,InUncompressedData,InUncompressedSize,OutCompressedData,LZLevel,
-													NULL,NULL,NULL,
+													&LZOptions,DictionaryStart,nullptr,
 													scratchMem,scratchSize);
 													
 		UE_LOG(OodleDataCompression, VeryVerbose, TEXT("OodleEncode: %s (%s): %lld -> %lld"), \
@@ -519,13 +533,15 @@ int64 CORE_API Compress(
 							void * OutCompressedData, int64 InCompressedBufferSize,
 							const void * InUncompressedData, int64 InUncompressedSize,
 							ECompressor Compressor,
-							ECompressionLevel Level)
+							ECompressionLevel Level,
+							bool CompressIndependentChunks,
+							int64 DictionaryBackup)
 {
 	OodleScratchBuffers * Scratches = GetGlobalOodleScratchBuffers();
 
 	int64 EncodedSize = Scratches->OodleEncode(OutCompressedData,InCompressedBufferSize,
 							InUncompressedData,InUncompressedSize,
-							Compressor,Level);
+							Compressor,Level,CompressIndependentChunks,DictionaryBackup);
 						
 	return EncodedSize;
 }
@@ -576,4 +592,236 @@ void CORE_API StartupPreInit(void)
 	OodleCore_Plugins_SetAllocators(OodleAlloc, OodleFree);
 }
 
-};
+
+
+int64 CompressParallelSub(
+							TArray64< TArray64<uint8> > & OutChunkCompressedData,
+							const void * InUncompressedData, int64 UncompressedSize,
+							ECompressor Compressor,	ECompressionLevel Level,
+							bool CompressIndependentChunks)
+{
+	int32 NumWorkers = FTaskGraphInterface::Get().GetNumWorkerThreads();
+	int64 ChunkLen = OodleLZ_MakeSeekChunkLen(UncompressedSize,NumWorkers);
+	int64 NumChunks = (UncompressedSize + ChunkLen-1)/ChunkLen;
+
+	OutChunkCompressedData.SetNum( NumChunks );
+	
+	// heuristic; don't backup more than chunklen ,or 16 MB max :
+	//	this does hurt ratio in some cases (vs non-parallel encoding, or unlimited dictionary backup)
+	const int64 MaxDictionaryBackup = FMath::Min( 16*1024*1024 , ChunkLen );
+
+	// ParallelFor will run one of the chunks on the calling thread
+	//	and will also fast path for NumChunks == 1
+
+	ParallelFor( TEXT("OodleCompressParallel.PF"),(int32)NumChunks,1, [&](int32 Index)
+	{
+		TArray64<uint8> & OutCompressedArray = OutChunkCompressedData[Index];
+		int64 ChunkStartPos = Index * ChunkLen;
+		int64 ChunkUncompressedSize = FMath::Min( ChunkLen, (UncompressedSize - ChunkStartPos) );
+		
+		int64 Reserve = CompressedBufferSizeNeeded(ChunkUncompressedSize);
+		OutCompressedArray.SetNum( Reserve );
+
+		// if ! CompressIndependentChunks, do dictionary backup
+		int64 DictionaryBackup = CompressIndependentChunks ? 0 : FMath::Min( ChunkStartPos , MaxDictionaryBackup );
+
+		const uint8 * ChunkUncompressedData = (const uint8 *)InUncompressedData + ChunkStartPos;
+
+		int64 Ret = Compress(&OutCompressedArray[0],Reserve,
+			ChunkUncompressedData,ChunkUncompressedSize,
+			Compressor,Level,CompressIndependentChunks,DictionaryBackup);
+
+		// zero size for error
+		check( OODLELZ_FAILED == 0 );
+		check( Ret >= 0 );
+
+		OutCompressedArray.SetNum(Ret,false);
+	} );
+
+	int64 Total = 0;
+	for(int64 i=0;i<NumChunks;i++)
+	{
+		int64 Len = OutChunkCompressedData[i].Num();
+		if ( Len == 0 )
+		{
+			return OODLELZ_FAILED;
+		}
+		Total += Len;
+	}
+
+	return Total;
+}
+
+int64 CompressParallel(
+							void * OutCompressedData, int64 CompressedBufferSize,
+							const void * InUncompressedData, int64 UncompressedSize,
+							ECompressor Compressor,	ECompressionLevel Level,
+							bool CompressIndependentChunks)
+{
+	if ( UncompressedSize <= OODLELZ_BLOCK_LEN )
+	{
+		return Compress(OutCompressedData,CompressedBufferSize,
+			InUncompressedData,UncompressedSize,Compressor,Level,CompressIndependentChunks);
+	}
+
+	TArray64< TArray64<uint8> > Chunks;
+	int64 TotalCompLen = CompressParallelSub(Chunks,InUncompressedData,UncompressedSize,Compressor,Level,CompressIndependentChunks);
+	if ( TotalCompLen <= 0 )
+	{
+		UE_LOG(OodleDataCompression,Error,TEXT("CompressParallelSub failed\n"));
+
+		return TotalCompLen;
+	}
+
+	if ( TotalCompLen > CompressedBufferSize )
+	{
+		UE_LOG(OodleDataCompression,Error,TEXT("CompressParallel buffer size insufficient\n"));
+		return OODLELZ_FAILED;
+	}
+
+	uint8 * StartPtr = (uint8 *)OutCompressedData;
+	uint8 * OutPtr = StartPtr; 
+	for(int64 i=0;i<Chunks.Num();i++)
+	{
+		memcpy(OutPtr, &Chunks[i][0] , Chunks[i].Num() );
+		OutPtr += Chunks[i].Num();
+	}
+	check( (OutPtr - StartPtr) == TotalCompLen );
+
+	#if 0
+	// for testing : verify decode
+	{
+		TArray64<uint8> Decomp;
+		Decomp.SetNum(UncompressedSize);
+		DecompressParallel(&Decomp[0],Decomp.Num(),OutCompressedData,TotalCompLen);
+
+		check( memcmp(InUncompressedData,&Decomp[0],UncompressedSize) == 0 );
+	}
+	#endif
+
+	return TotalCompLen;
+}
+							
+int64 CompressParallel(
+							TArray64<uint8> & OutCompressedArray,
+							const void * InUncompressedData, int64 UncompressedSize,
+							ECompressor Compressor,	ECompressionLevel Level,
+							bool CompressIndependentChunks)
+{
+	// appends to OutCompressedArray
+	int64 StartNum = OutCompressedArray.Num();
+
+	if ( UncompressedSize <= OODLELZ_BLOCK_LEN )
+	{
+		int64 Reserve = CompressedBufferSizeNeeded(UncompressedSize);
+		OutCompressedArray.SetNum( StartNum + Reserve );
+
+		int64 TotalCompLen =  Compress(&OutCompressedArray[StartNum],Reserve,
+			InUncompressedData,UncompressedSize,Compressor,Level,CompressIndependentChunks);
+
+		if ( TotalCompLen <= 0 )
+		{
+			return TotalCompLen;
+		}
+	
+		check( TotalCompLen <= Reserve );
+		OutCompressedArray.SetNum( StartNum + TotalCompLen , false);
+		return TotalCompLen;
+	}
+
+	TArray64< TArray64<uint8> > Chunks;
+	int64 TotalCompLen = CompressParallelSub(Chunks,InUncompressedData,UncompressedSize,Compressor,Level,CompressIndependentChunks);
+	if ( TotalCompLen <= 0 )
+	{
+		UE_LOG(OodleDataCompression,Error,TEXT("CompressParallelSub failed\n"));
+
+		return TotalCompLen;
+	}
+
+	OutCompressedArray.SetNum( StartNum + TotalCompLen );
+
+	uint8 * StartPtr = &OutCompressedArray[StartNum];
+	uint8 * OutPtr = StartPtr; 
+	for(int64 i=0;i<Chunks.Num();i++)
+	{
+		memcpy(OutPtr, &Chunks[i][0] , Chunks[i].Num() );
+		OutPtr += Chunks[i].Num();
+	}
+	check( (OutPtr - StartPtr) == TotalCompLen );
+
+	return TotalCompLen;
+}
+
+bool DecompressParallel(
+						void * OutUncompressedData, int64 UncompressedSize,
+						const void * InCompressedData, int64 CompressedSize
+						)
+{
+	if ( UncompressedSize <= 2*OODLELZ_BLOCK_LEN )
+	{
+		// small buffer, just early out to a synchronous decode
+
+		return Decompress(OutUncompressedData,UncompressedSize,InCompressedData,CompressedSize);
+	}
+	
+	// this function assumes you either have no  seek resets at all
+	//	or seek resets at every OODLELZ_BLOCK_LEN
+	// therefore we are free to choose our decode seek chunk len here freely
+	// if that was no the case
+	//	(eg. say if you did seek resets at OODLELZ_BLOCK_LEN*4)
+	// then we would have to change this logic to instead scan the comp data
+	//	and find the seek resets, using OodleLZ_GetCompressedStepForRawStep
+
+	int32 NumWorkers = FTaskGraphInterface::Get().GetNumWorkerThreads();
+	int32 ChunkLen = OodleLZ_MakeSeekChunkLen(UncompressedSize,NumWorkers*4);
+	int64 NumChunks = (UncompressedSize + ChunkLen-1)/ChunkLen;
+
+	int64 SeekTableMemSize = OodleLZ_GetSeekTableMemorySizeNeeded((int32)NumChunks,OodleLZSeekTable_Flags_None);
+	void * SeekTableMem = FMemory::Malloc(SeekTableMemSize);
+	OodleLZ_SeekTable * SeekTable = (OodleLZ_SeekTable *) SeekTableMem;
+
+	if ( ! OodleLZ_FillSeekTable(SeekTable,OodleLZSeekTable_Flags_None,ChunkLen,OutUncompressedData,UncompressedSize,InCompressedData,CompressedSize) )
+	{
+		UE_LOG(OodleDataCompression,Error,TEXT("OodleLZ_FillSeekTable failed\n"));
+		return false;
+	}
+
+	if ( ! SeekTable->seekChunksIndependent )
+	{
+		// CompressIndependentChunks was not used in the encoder, cannot decode in parallel
+		
+		FMemory::Free(SeekTableMem);
+
+		return Decompress(OutUncompressedData,UncompressedSize,InCompressedData,CompressedSize);
+	}
+
+	TArray64<int64> SeekCompressedPos;
+	SeekCompressedPos.SetNum(NumChunks+1);
+	SeekCompressedPos[0] = 0;
+	for(int64 i=0;i<NumChunks;i++)
+	{
+		SeekCompressedPos[i+1] = SeekCompressedPos[i] + SeekTable->seekChunkCompLens[i];
+	}
+	
+	bool AnyFailed = false;
+
+	ParallelFor( TEXT("OodleDecompressParallel.PF"),(int32)NumChunks,1, [&](int32 Index)
+	{
+		int64 RawPos = Index * ChunkLen;
+		int64 RawLen = FMath::Min( ChunkLen, (UncompressedSize - RawPos) );
+		int64 CompPos = SeekCompressedPos[Index];
+		int64 CompLen = SeekTable->seekChunkCompLens[Index];
+
+		if ( ! Decompress( (uint8 *)OutUncompressedData + RawPos,RawLen,(const uint8 *)InCompressedData + CompPos,CompLen) )
+		{
+			// no need for atomics :
+			AnyFailed = true;
+		}
+	} );
+
+	FMemory::Free(SeekTableMem);
+
+	return !AnyFailed;
+}
+
+}; // FOodleDataCompression
diff --git a/Engine/Source/Runtime/Core/Private/Compression/OodleDataCompressionUtil.cpp b/Engine/Source/Runtime/Core/Private/Compression/OodleDataCompressionUtil.cpp
--- a/Engine/Source/Runtime/Core/Private/Compression/OodleDataCompressionUtil.cpp
+++ b/Engine/Source/Runtime/Core/Private/Compression/OodleDataCompressionUtil.cpp
@@ -23,7 +23,7 @@ namespace FOodleCompressedArray
 
 		// We compress in to the buffer past our header, then write the header below.
 		void* CompPtr = OutCompressed.GetData() + HeaderSizeNeeded;
-		int32 CompressedSize = IntCastChecked<int32>(FOodleDataCompression::Compress(CompPtr, CompSizeNeeded, InData, InDataSize, InCompressor, InLevel));
+		int32 CompressedSize = IntCastChecked<int32>(FOodleDataCompression::CompressParallel(CompPtr, CompSizeNeeded, InData, InDataSize, InCompressor, InLevel));
 		if ( CompressedSize <= 0 )
 		{
 			// Probably a bad parameter.
@@ -65,7 +65,7 @@ namespace FOodleCompressedArray
 
 		// We compress in to the buffer past our header, then write the header below.
 		void* CompPtr = OutCompressed.GetData() + HeaderSizeNeeded;
-		int64 CompressedSize = FOodleDataCompression::Compress(CompPtr, WorkingSizeNeeded, InData, InDataSize, InCompressor, InLevel);
+		int64 CompressedSize = FOodleDataCompression::CompressParallel(CompPtr, WorkingSizeNeeded, InData, InDataSize, InCompressor, InLevel);
 		if (CompressedSize <= 0)
 		{
 			// Probably a bad parameter.
@@ -129,7 +129,7 @@ namespace FOodleCompressedArray
 			return false;
 		}
 
-		return FOodleDataCompression::Decompress(InDestinationBuffer, DecompressedSize, InCompressed.GetData() + OffsetToCompressedData, CompressedSize);
+		return FOodleDataCompression::DecompressParallel(InDestinationBuffer, DecompressedSize, InCompressed.GetData() + OffsetToCompressedData, CompressedSize);
 	}
 	bool CORE_API DecompressToExistingBuffer64(void* InDestinationBuffer, int64 InDestinationBufferSize, TArray64<uint8> const& InCompressed)
 	{
@@ -151,7 +151,7 @@ namespace FOodleCompressedArray
 			return false;
 		}
 
-		return FOodleDataCompression::Decompress(InDestinationBuffer, DecompressedSize, InCompressed.GetData() + OffsetToCompressedData, CompressedSize);
+		return FOodleDataCompression::DecompressParallel(InDestinationBuffer, DecompressedSize, InCompressed.GetData() + OffsetToCompressedData, CompressedSize);
 	}
 
 	bool CORE_API DecompressToAllocatedBuffer(void*& OutDestinationBuffer, int32& OutDestinationBufferSize, TArray<uint8> const& InCompressed)
@@ -175,7 +175,7 @@ namespace FOodleCompressedArray
 		OutDestinationBufferSize = DecompressedSize;
 		OutDestinationBuffer = DestinationBuffer;
 
-		if (FOodleDataCompression::Decompress(DestinationBuffer, DecompressedSize, InCompressed.GetData() + OffsetToCompressedData, CompressedSize) == false)
+		if (FOodleDataCompression::DecompressParallel(DestinationBuffer, DecompressedSize, InCompressed.GetData() + OffsetToCompressedData, CompressedSize) == false)
 		{
 			FMemory::Free(DestinationBuffer);
 			OutDestinationBuffer = 0;
@@ -205,7 +205,7 @@ namespace FOodleCompressedArray
 		OutDestinationBufferSize = DecompressedSize;
 		OutDestinationBuffer = DestinationBuffer;
 
-		if (FOodleDataCompression::Decompress(DestinationBuffer, DecompressedSize, InCompressed.GetData() + OffsetToCompressedData, CompressedSize) == false)
+		if (FOodleDataCompression::DecompressParallel(DestinationBuffer, DecompressedSize, InCompressed.GetData() + OffsetToCompressedData, CompressedSize) == false)
 		{
 			FMemory::Free(DestinationBuffer);
 			OutDestinationBuffer = 0;
diff --git a/Engine/Source/Runtime/Core/Public/Compression/OodleDataCompression.h b/Engine/Source/Runtime/Core/Public/Compression/OodleDataCompression.h
--- a/Engine/Source/Runtime/Core/Public/Compression/OodleDataCompression.h
+++ b/Engine/Source/Runtime/Core/Public/Compression/OodleDataCompression.h
@@ -162,6 +162,7 @@ int64 CORE_API GetMaximumCompressedSize(int64 UncompressedSize);
 /**
 * Encode provided data with chosen Compressor and Level
 * CompressedBufferSize must be >= CompressedBufferSizeNeeded(UncompressedSize)
+* use CompressParallel instead if buffer can ever be large
 *
 * @param	OutCompressedData			output buffer where compressed data is written
 * @param	CompressedBufferSize		bytes available to write in OutCompressedData
@@ -169,19 +170,68 @@ int64 CORE_API GetMaximumCompressedSize(int64 UncompressedSize);
 * @param	UncompressedSize			number of bytes in InUncompressedData to read
 * @param	Compressor					ECompressor to encode with (this is saved in the stream)
 * @param	Level						ECompressionLevel to encode with (this is not saved in the stream)
+* @param	CompressIndependentChunks	(optional) should chunks be made independent (allows for parallel decode)
+* @param	DictionaryBackup			(optional) number of bytes preceding InUncompressedData which should be used for dictionary preload
 * @return								Compressed size written or zero for failure
 */
 int64 CORE_API Compress(
 							void * OutCompressedData, int64 CompressedBufferSize,
 							const void * InUncompressedData, int64 UncompressedSize,
 							ECompressor Compressor,
-							ECompressionLevel Level);
+							ECompressionLevel Level,
+							bool CompressIndependentChunks = false,
+							int64 DictionaryBackup = 0);
+
+/**
+* Encode provided data with chosen Compressor and Level, using multiple threads for large buffers
+* CompressedBufferSize must be >= CompressedBufferSizeNeeded(UncompressedSize)
+* CompressParallel can be used in all cases instead of Compress
+* it is fast even on small buffers (will just run synchronously on the calling thread)
+* the output compressed data is the same as Compress
+*
+* @param	OutCompressedData			output buffer where compressed data is written
+* @param	CompressedBufferSize		bytes available to write in OutCompressedData
+* @param	InUncompressedData			input buffer containing data to compress
+* @param	UncompressedSize			number of bytes in InUncompressedData to read
+* @param	Compressor					ECompressor to encode with (this is saved in the stream)
+* @param	Level						ECompressionLevel to encode with (this is not saved in the stream)
+* @param	CompressIndependentChunks				should chunks be made independent (allows for parallel decode)
+* @return								Compressed size written or zero for failure
+*/
+int64 CORE_API CompressParallel(
+							void * OutCompressedData, int64 CompressedBufferSize,
+							const void * InUncompressedData, int64 UncompressedSize,
+							ECompressor Compressor,	ECompressionLevel Level,
+							bool CompressIndependentChunks = false);
+							
+/**
+* Encode provided data with chosen Compressor and Level, using multiple threads for large buffers
+* CompressedBufferSize must be >= CompressedBufferSizeNeeded(UncompressedSize)
+* CompressParallel can be used in all cases instead of Compress
+* it is fast even on small buffers (will just run synchronously on the calling thread)
+* the output compressed data is the same as Compress
+*
+* @param	OutCompressedData			output buffer where compressed data is written; array is appended to
+* @param	InUncompressedData			input buffer containing data to compress
+* @param	UncompressedSize			number of bytes in InUncompressedData to read
+* @param	Compressor					ECompressor to encode with (this is saved in the stream)
+* @param	Level						ECompressionLevel to encode with (this is not saved in the stream)
+* @param	CompressIndependentChunks				should chunks be made independent (allows for parallel decode)
+* @return								Compressed size written or zero for failure
+*/
+int64 CORE_API CompressParallel(
+							TArray64<uint8> & OutCompressedData,
+							const void * InUncompressedData, int64 UncompressedSize,
+							ECompressor Compressor,	ECompressionLevel Level,
+							bool CompressIndependentChunks = false);
 
 /**
 * Decode compressed data that was made by Compress
 *
 * UncompressedSize must match exactly the uncompressed size that was used at encode time.  No partial decodes.
 *
+* If buffer can ever be large, use DecompressParallel instead
+*
 * @param	OutUncompressedData		output buffer where uncompressed data is written
 * @param	UncompressedSize		number of bytes to decompress
 * @param	InCompressedData		input buffer containing compressed data
@@ -193,6 +243,25 @@ bool CORE_API Decompress(
 						const void * InCompressedData, int64 CompressedSize
 						);
 
+/**
+* Decode compressed data that was made by Compress, using multiple threads when possible.
+*
+* UncompressedSize must match exactly the uncompressed size that was used at encode time.  No partial decodes.
+*
+* There's no penalty to using this on small buffers, it will just decode synchronously on the calling thread in that case.
+*
+* DecompressParallel can only use more than 1 thread if the encoding was done with CompressIndependentChunks = true
+*
+* @param	OutUncompressedData		output buffer where uncompressed data is written
+* @param	UncompressedSize		number of bytes to decompress
+* @param	InCompressedData		input buffer containing compressed data
+* @param	CompressedSize			size of the input buffer, must be greater or equal to the number of compressed bytes needed
+* @return							boolean success
+*/
+bool CORE_API DecompressParallel(
+						void * OutUncompressedData, int64 UncompressedSize,
+						const void * InCompressedData, int64 CompressedSize
+						);
 
 // For Compression to/from TArray and such higher level actions use OodleDataCompressionUtil.h
 
diff --git a/Engine/Source/Runtime/Datasmith/DatasmithCore/Private/DatasmithMeshSerialization.cpp b/Engine/Source/Runtime/Datasmith/DatasmithCore/Private/DatasmithMeshSerialization.cpp
--- a/Engine/Source/Runtime/Datasmith/DatasmithCore/Private/DatasmithMeshSerialization.cpp
+++ b/Engine/Source/Runtime/Datasmith/DatasmithCore/Private/DatasmithMeshSerialization.cpp
@@ -82,7 +82,7 @@ bool CompressInline(TArray<uint8>& UncompressedData, ECompressionMethod Method)
 	{
 		auto Compressor = FOodleDataCompression::ECompressor::Kraken;
 		auto Level = FOodleDataCompression::ECompressionLevel::VeryFast;
-		if (int64 ActualCompressedSize = FOodleDataCompression::Compress(CompressedData.GetData() + HeaderSize, CompressedSize, UncompressedData.GetData(), UncompressedSize, Compressor, Level))
+		if (int64 ActualCompressedSize = FOodleDataCompression::CompressParallel(CompressedData.GetData() + HeaderSize, CompressedSize, UncompressedData.GetData(), UncompressedSize, Compressor, Level))
 		{
 			CompressedData.SetNum(ActualCompressedSize + HeaderSize, true);
 			UncompressedData = MoveTemp(CompressedData);
diff --git a/Engine/Source/Runtime/Engine/Private/TextureDerivedData.cpp b/Engine/Source/Runtime/Engine/Private/TextureDerivedData.cpp
--- a/Engine/Source/Runtime/Engine/Private/TextureDerivedData.cpp
+++ b/Engine/Source/Runtime/Engine/Private/TextureDerivedData.cpp
@@ -1724,7 +1724,7 @@ static void EstimateOnDiskCompressionForTextureData(
 			CopyBytes == InCompressionBlockSize) // we can fit in this chunk
 		{
 			// Direct.
-			Compressed.Empty();
+			Compressed.SetNum(0,false);
 			FOodleCompressedArray::CompressData(
 				Compressed,
 				CurrentContainer.GetData() + CurrentOffsetInContainer,
@@ -1748,7 +1748,7 @@ static void EstimateOnDiskCompressionForTextureData(
 			if (ContinuousMemory.Num() == InCompressionBlockSize)
 			{
 				// Filled a block - kick.
-				Compressed.Empty();
+				Compressed.SetNum(0,false);
 				FOodleCompressedArray::CompressData(
 					Compressed,
 					ContinuousMemory.GetData(),
@@ -1787,7 +1787,7 @@ static void EstimateOnDiskCompressionForTextureData(
 	if (ContinuousMemory.Num())
 	{
 		// If we ran out of source data before we completely filled, kick here.
-		Compressed.Empty();
+		Compressed.SetNum(0,false);
 		FOodleCompressedArray::CompressData(
 			Compressed,
 			ContinuousMemory.GetData(),
diff --git a/Engine/Source/Runtime/NavigationSystem/Private/NavMesh/RecastNavMeshGenerator.cpp b/Engine/Source/Runtime/NavigationSystem/Private/NavMesh/RecastNavMeshGenerator.cpp
--- a/Engine/Source/Runtime/NavigationSystem/Private/NavMesh/RecastNavMeshGenerator.cpp
+++ b/Engine/Source/Runtime/NavigationSystem/Private/NavMesh/RecastNavMeshGenerator.cpp
@@ -1488,7 +1488,7 @@ struct FTileCacheCompressor : public dtTileCacheCompressor
 
 		if (GNavmeshUseOodleCompression)
 		{
-			const int64 CompressedSize = FOodleDataCompression::Compress((void*)DataPtr, DataSize, (const void*)buffer, bufferSize, GNavmeshTileCacheCompressor, GNavmeshTileCacheCompressionLevel);
+			const int64 CompressedSize = FOodleDataCompression::CompressParallel((void*)DataPtr, DataSize, (const void*)buffer, bufferSize, GNavmeshTileCacheCompressor, GNavmeshTileCacheCompressionLevel);
 			if (CompressedSize > 0)
 			{
 				*outCompressedSize = IntCastChecked<int32>(CompressedSize + HeaderSize);
