OodleDataCompression just comments behavior nop

diff --git a/Engine/Source/Runtime/Core/Private/Compression/OodleDataCompression.cpp b/Engine/Source/Runtime/Core/Private/Compression/OodleDataCompression.cpp
--- a/Engine/Source/Runtime/Core/Private/Compression/OodleDataCompression.cpp
+++ b/Engine/Source/Runtime/Core/Private/Compression/OodleDataCompression.cpp
@@ -600,6 +600,8 @@ int64 CompressParallelSub(
 							ECompressor Compressor,	ECompressionLevel Level,
 							bool CompressIndependentChunks)
 {
+	// beware: doing the split based on NumWorkers means that compressed output depends on the machine
+
 	int32 NumWorkers = FTaskGraphInterface::Get().GetNumWorkerThreads();
 	int64 ChunkLen = OodleLZ_MakeSeekChunkLen(UncompressedSize,NumWorkers);
 	int64 NumChunks = (UncompressedSize + ChunkLen-1)/ChunkLen;
diff --git a/Engine/Source/Runtime/Core/Public/Compression/OodleDataCompression.h b/Engine/Source/Runtime/Core/Public/Compression/OodleDataCompression.h
--- a/Engine/Source/Runtime/Core/Public/Compression/OodleDataCompression.h
+++ b/Engine/Source/Runtime/Core/Public/Compression/OodleDataCompression.h
@@ -19,6 +19,8 @@
 //
 // OodleDataCompression.h is for utility compression in non-shipping game package scenarios
 //	eg. uassets, storage caches and back-ends, large network transfers
+//
+// For Compression to/from TArray and such higher level actions use OodleDataCompressionUtil.h
 
 DECLARE_LOG_CATEGORY_EXTERN(OodleDataCompression, Log, All);
 
@@ -188,6 +190,8 @@ int64 CORE_API Compress(
 * CompressParallel can be used in all cases instead of Compress
 * it is fast even on small buffers (will just run synchronously on the calling thread)
 * the output compressed data is the same as Compress
+* can be decoded with either Decompress or DecompressParallel
+* DecompressParallel can only parallelize if CompressIndependentChunks is true at encode time
 *
 * @param	OutCompressedData			output buffer where compressed data is written
 * @param	CompressedBufferSize		bytes available to write in OutCompressedData
@@ -195,7 +199,7 @@ int64 CORE_API Compress(
 * @param	UncompressedSize			number of bytes in InUncompressedData to read
 * @param	Compressor					ECompressor to encode with (this is saved in the stream)
 * @param	Level						ECompressionLevel to encode with (this is not saved in the stream)
-* @param	CompressIndependentChunks				should chunks be made independent (allows for parallel decode)
+* @param	CompressIndependentChunks	(optional) should chunks be made independent (allows for parallel decode)
 * @return								Compressed size written or zero for failure
 */
 int64 CORE_API CompressParallel(
@@ -210,13 +214,15 @@ int64 CORE_API CompressParallel(
 * CompressParallel can be used in all cases instead of Compress
 * it is fast even on small buffers (will just run synchronously on the calling thread)
 * the output compressed data is the same as Compress
+* can be decoded with either Decompress or DecompressParallel
+* DecompressParallel can only parallelize if CompressIndependentChunks is true at encode time
 *
 * @param	OutCompressedData			output buffer where compressed data is written; array is appended to
 * @param	InUncompressedData			input buffer containing data to compress
 * @param	UncompressedSize			number of bytes in InUncompressedData to read
 * @param	Compressor					ECompressor to encode with (this is saved in the stream)
 * @param	Level						ECompressionLevel to encode with (this is not saved in the stream)
-* @param	CompressIndependentChunks				should chunks be made independent (allows for parallel decode)
+* @param	CompressIndependentChunks	should chunks be made independent (allows for parallel decode)
 * @return								Compressed size written or zero for failure
 */
 int64 CORE_API CompressParallel(
@@ -249,6 +255,8 @@ bool CORE_API Decompress(
 * UncompressedSize must match exactly the uncompressed size that was used at encode time.  No partial decodes.
 *
 * There's no penalty to using this on small buffers, it will just decode synchronously on the calling thread in that case.
+* DecompressParallel can be used to decode data written by Compress or CompressParallel.
+* DecompressParallel can be used in all places you would call Decompress.
 *
 * DecompressParallel can only use more than 1 thread if the encoding was done with CompressIndependentChunks = true
 *
@@ -263,8 +271,6 @@ bool CORE_API DecompressParallel(
 						const void * InCompressedData, int64 CompressedSize
 						);
 
-// For Compression to/from TArray and such higher level actions use OodleDataCompressionUtil.h
-
 // from Compression.cpp :
 void CORE_API CompressionFormatInitOnFirstUseFromLock();
 
