Fixed TBitArray being implicitly constructible from a TBitArray of another allocator to stop accidental dangling references when constructing a const bit iterator with the wrong allocator type.
Added UE_LIFETIMEBOUND to bit iterators and bit references to warn when lifetime issues occur (like holding onto an iterator to a temporary bitarray).
Fixed implicit conversions when comparing or doing bitwise operations between bit arrays with different allocators.

diff --git a/Engine/Source/Runtime/Core/Public/Containers/BitArray.h b/Engine/Source/Runtime/Core/Public/Containers/BitArray.h
--- a/Engine/Source/Runtime/Core/Public/Containers/BitArray.h
+++ b/Engine/Source/Runtime/Core/Public/Containers/BitArray.h
@@ -69,12 +69,128 @@ enum class EBitwiseOperatorFlags
 };
 ENUM_CLASS_FLAGS(EBitwiseOperatorFlags)
 
+namespace UE::Core::Private
+{
+	template <typename AllocatorA, typename AllocatorB, typename OutAllocator, typename ProjectionType>
+	void BitwiseBinaryOperatorImpl(const TBitArray<AllocatorA>& InA, const TBitArray<AllocatorB>& InB, TBitArray<OutAllocator>& OutResult, EBitwiseOperatorFlags InFlags, ProjectionType&& InProjection)
+	{
+		check((const void*)&InA != (const void*)&InB && (const void*)&InA != (const void*)&OutResult && (const void*)&InB != (const void*)&OutResult);
+
+		if (EnumHasAnyFlags(InFlags, EBitwiseOperatorFlags::MinSize))
+		{
+
+			const int32 MinNumBits = FMath::Min(InA.Num(), InB.Num());
+			if (MinNumBits > 0)
+			{
+				OutResult.Reserve(MinNumBits);
+				OutResult.NumBits = MinNumBits;
+
+				typename TBitArray<AllocatorA>::FConstWordIterator IteratorA(InA);
+				typename TBitArray<AllocatorB>::FConstWordIterator IteratorB(InB);
+
+				typename TBitArray<OutAllocator>::FWordIterator IteratorResult(OutResult);
+
+				for ( ; IteratorResult; ++IteratorResult, ++IteratorA, ++IteratorB)
+				{
+					const uint32 NewValue = Invoke(InProjection, IteratorA.GetWord(), IteratorB.GetWord());
+					IteratorResult.SetWord(NewValue);
+				}
+			}
+
+		}
+		else if (EnumHasAnyFlags(InFlags, EBitwiseOperatorFlags::MaxSize))
+		{
+
+			const int32 MaxNumBits = FMath::Max(InA.Num(), InB.Num());
+			const uint32 MissingBitsFill = EnumHasAnyFlags(InFlags, EBitwiseOperatorFlags::OneFillMissingBits) ? ~0u : 0;
+
+			if (MaxNumBits)
+			{
+				OutResult.Reserve(MaxNumBits);
+				OutResult.NumBits = MaxNumBits;
+
+				typename TBitArray<AllocatorA>::FConstWordIterator IteratorA(InA);
+				typename TBitArray<AllocatorB>::FConstWordIterator IteratorB(InB);
+
+				IteratorA.FillMissingBits(MissingBitsFill);
+				IteratorB.FillMissingBits(MissingBitsFill);
+
+				typename TBitArray<OutAllocator>::FWordIterator IteratorResult(OutResult);
+
+				for ( ; IteratorResult; ++IteratorResult, ++IteratorA, ++IteratorB)
+				{
+					uint32 A = IteratorA ? IteratorA.GetWord() : MissingBitsFill;
+					uint32 B = IteratorB ? IteratorB.GetWord() : MissingBitsFill;
+
+					IteratorResult.SetWord(Invoke(InProjection, A, B));
+				}
+			}
+
+		}
+		else
+		{
+			checkf(false, TEXT("Invalid size flag specified for binary bitwise AND"));
+		}
+
+		OutResult.CheckInvariants();
+	}
+
+
+	template<typename OtherAllocator, typename OutAllocator, typename ProjectionType>
+	void BitwiseOperatorImpl(const TBitArray<OtherAllocator>& InOther, TBitArray<OutAllocator>& OutResult, EBitwiseOperatorFlags InFlags, ProjectionType&& InProjection)
+	{
+		if constexpr (std::is_same_v<OtherAllocator, OutAllocator>)
+		{
+			check(&InOther != &OutResult);
+		}
+
+		int32 NewNumBits = OutResult.NumBits;
+		if (EnumHasAnyFlags(InFlags, EBitwiseOperatorFlags::MinSize))
+		{
+			NewNumBits = FMath::Min(InOther.Num(), OutResult.Num());
+		}
+		else if (EnumHasAnyFlags(InFlags, EBitwiseOperatorFlags::MaxSize))
+		{
+			NewNumBits = FMath::Max(InOther.Num(), OutResult.Num());
+		}
+
+		const int32 SizeDifference = NewNumBits - OutResult.NumBits;
+		if (SizeDifference < 0)
+		{
+			OutResult.NumBits = NewNumBits;
+			OutResult.ClearPartialSlackBits();
+		}
+		else if (SizeDifference > 0)
+		{
+			const bool bPadValue = EnumHasAnyFlags(InFlags, EBitwiseOperatorFlags::OneFillMissingBits);
+			OutResult.Add(bPadValue, SizeDifference);
+		}
+
+		const uint32 MissingBitsFill = EnumHasAnyFlags(InFlags, EBitwiseOperatorFlags::OneFillMissingBits) ? ~0u : 0;
+		if (OutResult.NumBits != 0)
+		{
+			typename TBitArray<OtherAllocator>::FConstWordIterator IteratorOther(InOther);
+			IteratorOther.FillMissingBits(MissingBitsFill);
+
+			typename TBitArray<OutAllocator>::FWordIterator IteratorResult(OutResult);
+
+			for ( ; IteratorResult; ++IteratorResult, ++IteratorOther)
+			{
+				const uint32 OtherValue = IteratorOther ? IteratorOther.GetWord() : MissingBitsFill;
+				IteratorResult.SetWord(Invoke(InProjection, IteratorResult.GetWord(), OtherValue));
+			}
+		}
+
+		OutResult.CheckInvariants();
+	}
+}
+
 /** Used to read/write a bit in the array as a bool. */
 class FBitReference
 {
 public:
 
-	FORCEINLINE FBitReference(uint32& InData,uint32 InMask)
+	FORCEINLINE explicit FBitReference(uint32& InData UE_LIFETIMEBOUND,uint32 InMask)
 	:	Data(InData)
 	,	Mask(InMask)
 	{}
@@ -160,7 +276,7 @@ class FConstBitReference
 {
 public:
 
-	FORCEINLINE FConstBitReference(const uint32& InData,uint32 InMask)
+	FORCEINLINE explicit FConstBitReference(const uint32& InData UE_LIFETIMEBOUND,uint32 InMask)
 	:	Data(InData)
 	,	Mask(InMask)
 	{}
@@ -241,6 +357,12 @@ class TBitArray
 	template <typename, typename>
 	friend class TScriptBitArray;
 
+	template <typename AllocatorA, typename AllocatorB, typename OutAllocator, typename ProjectionType>
+	friend void UE::Core::Private::BitwiseBinaryOperatorImpl(const TBitArray<AllocatorA>&, const TBitArray<AllocatorB>&, TBitArray<OutAllocator>&, EBitwiseOperatorFlags, ProjectionType&&);
+
+	template<typename OtherAllocator, typename OutAllocator, typename ProjectionType>
+	friend void UE::Core::Private::BitwiseOperatorImpl(const TBitArray<OtherAllocator>&, TBitArray<OutAllocator>&, EBitwiseOperatorFlags, ProjectionType&&);
+
 	typedef uint32 WordType;
 	static constexpr WordType FullWordMask = (WordType)-1;
 
@@ -291,7 +413,7 @@ public:
 	}
 
 	template<typename OtherAllocator>
-	FORCEINLINE TBitArray(const TBitArray<OtherAllocator> & Copy)
+	FORCEINLINE explicit TBitArray(const TBitArray<OtherAllocator> & Copy)
 		: NumBits(0)
 		, MaxBits(0)
 	{
@@ -1106,59 +1228,68 @@ public:
 	/**
 	 * Return the bitwise AND of two bit arrays. The resulting bit array will be sized according to InFlags.
 	 */
-	static TBitArray BitwiseAND(const TBitArray& A, const TBitArray& B, EBitwiseOperatorFlags InFlags)
+	template <typename AllocatorA, typename AllocatorB>
+	static TBitArray BitwiseAND(const TBitArray<AllocatorA>& A, const TBitArray<AllocatorB>& B, EBitwiseOperatorFlags InFlags)
 	{
 		TBitArray Result;
-		BitwiseBinaryOperatorImpl(A, B, Result, InFlags, [](uint32 InA, uint32 InB) { return InA & InB; });
+		UE::Core::Private::BitwiseBinaryOperatorImpl(A, B, Result, InFlags, [](uint32 InA, uint32 InB) { return InA & InB; });
 		return Result;
 	}
 
 	/**
 	 * Perform a bitwise AND on this bit array with another. This array receives the result and will be sized max(A.Num(), B.Num()).
 	 */
-	TBitArray& CombineWithBitwiseAND(const TBitArray& InOther, EBitwiseOperatorFlags InFlags)
+	template <typename OtherAllocator>
+	TBitArray& CombineWithBitwiseAND(const TBitArray<OtherAllocator>& InOther, EBitwiseOperatorFlags InFlags)
 	{
-		BitwiseOperatorImpl(InOther, *this, InFlags, [](uint32 InA, uint32 InB) { return InA & InB; });
+		UE::Core::Private::BitwiseOperatorImpl(InOther, *this, InFlags, [](uint32 InA, uint32 InB) { return InA & InB; });
 		return *this;
 	}
 
 	/**
 	 * Return the bitwise OR of two bit arrays. The resulting bit array will be sized according to InFlags.
 	 */
-	static TBitArray BitwiseOR(const TBitArray& A, const TBitArray& B, EBitwiseOperatorFlags InFlags)
+	template <typename AllocatorA, typename AllocatorB>
+	static TBitArray BitwiseOR(const TBitArray<AllocatorA>& A, const TBitArray<AllocatorB>& B, EBitwiseOperatorFlags InFlags)
 	{
-		check(&A != &B);
+		if constexpr (std::is_same_v<AllocatorA, AllocatorB>)
+		{
+			check(&A != &B);
+		}
 
 		TBitArray Result;
-		BitwiseBinaryOperatorImpl(A, B, Result, InFlags, [](uint32 InA, uint32 InB) { return InA | InB; });
+		UE::Core::Private::BitwiseBinaryOperatorImpl(A, B, Result, InFlags, [](uint32 InA, uint32 InB) { return InA | InB; });
 		return Result;
 	}
 
 	/**
 	 * Return the bitwise OR of two bit arrays. The resulting bit array will be sized according to InFlags.
 	 */
-	TBitArray& CombineWithBitwiseOR(const TBitArray& InOther, EBitwiseOperatorFlags InFlags)
+	template <typename OtherAllocator>
+	TBitArray& CombineWithBitwiseOR(const TBitArray<OtherAllocator>& InOther, EBitwiseOperatorFlags InFlags)
 	{
-		BitwiseOperatorImpl(InOther, *this, InFlags, [](uint32 InA, uint32 InB) { return InA | InB; });
+		UE::Core::Private::BitwiseOperatorImpl(InOther, *this, InFlags, [](uint32 InA, uint32 InB) { return InA | InB; });
 		return *this;
 	}
 
 	/**
 	 * Return the bitwise XOR of two bit arrays. The resulting bit array will be sized according to InFlags.
 	 */
-	static TBitArray BitwiseXOR(const TBitArray& A, const TBitArray& B, EBitwiseOperatorFlags InFlags)
+	template <typename AllocatorA, typename AllocatorB>
+	static TBitArray BitwiseXOR(const TBitArray<AllocatorA>& A, const TBitArray<AllocatorB>& B, EBitwiseOperatorFlags InFlags)
 	{
 		TBitArray Result;
-		BitwiseBinaryOperatorImpl(A, B, Result, InFlags, [](uint32 InA, uint32 InB) { return InA ^ InB; });
+		UE::Core::Private::BitwiseBinaryOperatorImpl(A, B, Result, InFlags, [](uint32 InA, uint32 InB) { return InA ^ InB; });
 		return Result;
 	}
 
 	/**
 	 * Return the bitwise XOR of two bit arrays. The resulting bit array will be sized according to InFlags.
 	 */
-	TBitArray& CombineWithBitwiseXOR(const TBitArray& InOther, EBitwiseOperatorFlags InFlags)
+	template <typename OtherAllocator>
+	TBitArray& CombineWithBitwiseXOR(const TBitArray<OtherAllocator>& InOther, EBitwiseOperatorFlags InFlags)
 	{
-		BitwiseOperatorImpl(InOther, *this, InFlags, [](uint32 InA, uint32 InB) { return InA ^ InB; });
+		UE::Core::Private::BitwiseOperatorImpl(InOther, *this, InFlags, [](uint32 InA, uint32 InB) { return InA ^ InB; });
 		return *this;
 	}
 
@@ -1202,12 +1333,13 @@ public:
 	 * @param bMissingBitValue The value to use for missing bits when considering bits that are outside the range of either array
 	 * @return true if this array matches Other, including any missing bits, false otherwise
 	 */
-	bool CompareSetBits(const TBitArray& Other, const bool bMissingBitValue) const
+	template <typename OtherAllocator>
+	bool CompareSetBits(const TBitArray<OtherAllocator>& Other, const bool bMissingBitValue) const
 	{
 		const uint32 MissingBitsFill = bMissingBitValue ? ~0u : 0;
 
 		FConstWordIterator ThisIterator(*this);
-		FConstWordIterator OtherIterator(Other);
+		typename TBitArray<OtherAllocator>::FConstWordIterator OtherIterator(Other);
 
 		ThisIterator.FillMissingBits(MissingBitsFill);
 		OtherIterator.FillMissingBits(MissingBitsFill);
@@ -1306,7 +1438,7 @@ public:
 	class FIterator : public FRelativeBitReference
 	{
 	public:
-		FORCEINLINE FIterator(TBitArray<Allocator>& InArray,int32 StartIndex = 0)
+		FORCEINLINE explicit FIterator(TBitArray<Allocator>& InArray UE_LIFETIMEBOUND,int32 StartIndex = 0)
 		:	FRelativeBitReference(StartIndex)
 		,	Array(InArray)
 		,	Index(StartIndex)
@@ -1346,7 +1478,7 @@ public:
 	class FConstIterator : public FRelativeBitReference
 	{
 	public:
-		FORCEINLINE FConstIterator(const TBitArray<Allocator>& InArray,int32 StartIndex = 0)
+		FORCEINLINE explicit FConstIterator(const TBitArray<Allocator>& InArray UE_LIFETIMEBOUND,int32 StartIndex = 0)
 		:	FRelativeBitReference(StartIndex)
 		,	Array(InArray)
 		,	Index(StartIndex)
@@ -1387,7 +1519,7 @@ public:
 	class FConstReverseIterator : public FRelativeBitReference
 	{
 	public:
-		FORCEINLINE FConstReverseIterator(const TBitArray<Allocator>& InArray)
+		FORCEINLINE explicit FConstReverseIterator(const TBitArray<Allocator>& InArray UE_LIFETIMEBOUND)
 			:	FRelativeBitReference(InArray.Num() - 1)
 			,	Array(InArray)
 			,	Index(InArray.Num() - 1)
@@ -1435,118 +1567,6 @@ public:
 	}
 
 private:
-
-	template<typename ProjectionType>
-	static void BitwiseBinaryOperatorImpl(const TBitArray& InA, const TBitArray& InB, TBitArray& OutResult, EBitwiseOperatorFlags InFlags, ProjectionType&& InProjection)
-	{
-		check((&InA != &InB) && (&InA != &OutResult) && (&InB != &OutResult));
-
-		if (EnumHasAnyFlags(InFlags, EBitwiseOperatorFlags::MinSize))
-		{
-
-			const int32 MinNumBits = FMath::Min(InA.Num(), InB.Num());
-			if (MinNumBits > 0)
-			{
-				OutResult.Reserve(MinNumBits);
-				OutResult.NumBits = MinNumBits;
-
-				FConstWordIterator IteratorA(InA);
-				FConstWordIterator IteratorB(InB);
-
-				FWordIterator IteratorResult(OutResult);
-
-				for ( ; IteratorResult; ++IteratorResult, ++IteratorA, ++IteratorB)
-				{
-					const uint32 NewValue = Invoke(InProjection, IteratorA.GetWord(), IteratorB.GetWord());
-					IteratorResult.SetWord(NewValue);
-				}
-			}
-
-		}
-		else if (EnumHasAnyFlags(InFlags, EBitwiseOperatorFlags::MaxSize))
-		{
-
-			const int32 MaxNumBits = FMath::Max(InA.Num(), InB.Num());
-			const uint32 MissingBitsFill = EnumHasAnyFlags(InFlags, EBitwiseOperatorFlags::OneFillMissingBits) ? ~0u : 0;
-
-			if (MaxNumBits)
-			{
-				OutResult.Reserve(MaxNumBits);
-				OutResult.NumBits = MaxNumBits;
-
-				FConstWordIterator IteratorA(InA);
-				FConstWordIterator IteratorB(InB);
-
-				IteratorA.FillMissingBits(MissingBitsFill);
-				IteratorB.FillMissingBits(MissingBitsFill);
-
-				FWordIterator IteratorResult(OutResult);
-
-				for ( ; IteratorResult; ++IteratorResult, ++IteratorA, ++IteratorB)
-				{
-					uint32 A = IteratorA ? IteratorA.GetWord() : MissingBitsFill;
-					uint32 B = IteratorB ? IteratorB.GetWord() : MissingBitsFill;
-
-					IteratorResult.SetWord(Invoke(InProjection, A, B));
-				}
-			}
-
-		}
-		else
-		{
-			checkf(false, TEXT("Invalid size flag specified for binary bitwise AND"));
-		}
-
-		OutResult.CheckInvariants();
-	}
-
-
-	template<typename ProjectionType>
-	static void BitwiseOperatorImpl(const TBitArray& InOther, TBitArray& OutResult, EBitwiseOperatorFlags InFlags, ProjectionType&& InProjection)
-	{
-		check(&InOther != &OutResult);
-
-		int32 NewNumBits = OutResult.NumBits;
-		if (EnumHasAnyFlags(InFlags, EBitwiseOperatorFlags::MinSize))
-		{
-			NewNumBits = FMath::Min(InOther.Num(), OutResult.Num());
-		}
-		else if (EnumHasAnyFlags(InFlags, EBitwiseOperatorFlags::MaxSize))
-		{
-			NewNumBits = FMath::Max(InOther.Num(), OutResult.Num());
-		}
-
-		const int32 SizeDifference = NewNumBits - OutResult.NumBits;
-		if (SizeDifference < 0)
-		{
-			OutResult.NumBits = NewNumBits;
-			OutResult.ClearPartialSlackBits();
-		}
-		else if (SizeDifference > 0)
-		{
-			const bool bPadValue = EnumHasAnyFlags(InFlags, EBitwiseOperatorFlags::OneFillMissingBits);
-			OutResult.Add(bPadValue, SizeDifference);
-		}
-
-		const uint32 MissingBitsFill = EnumHasAnyFlags(InFlags, EBitwiseOperatorFlags::OneFillMissingBits) ? ~0u : 0;
-		if (OutResult.NumBits != 0)
-		{
-			FConstWordIterator IteratorOther(InOther);
-			IteratorOther.FillMissingBits(MissingBitsFill);
-
-			FWordIterator IteratorResult(OutResult);
-
-			for ( ; IteratorResult; ++IteratorResult, ++IteratorOther)
-			{
-				const uint32 OtherValue = IteratorOther ? IteratorOther.GetWord() : MissingBitsFill;
-				IteratorResult.SetWord(Invoke(InProjection, IteratorResult.GetWord(), OtherValue));
-			}
-		}
-
-		OutResult.CheckInvariants();
-	}
-
-
 	template<typename WordType>
 	struct TWordIteratorBase
 	{
@@ -1632,11 +1652,11 @@ private:
 public:
 	struct FConstWordIterator : TWordIteratorBase<const uint32>
 	{
-		explicit FConstWordIterator(const TBitArray<Allocator>& InArray)
+		explicit FConstWordIterator(const TBitArray<Allocator>& InArray UE_LIFETIMEBOUND)
 			: TWordIteratorBase<const uint32>(InArray.GetData(), 0, InArray.Num())
 		{}
 
-		explicit FConstWordIterator(const TBitArray<Allocator>& InArray, int32 InStartBitIndex, int32 InEndBitIndex)
+		explicit FConstWordIterator(const TBitArray<Allocator>& InArray UE_LIFETIMEBOUND, int32 InStartBitIndex, int32 InEndBitIndex)
 			: TWordIteratorBase<const uint32>(InArray.GetData(), InStartBitIndex, InEndBitIndex)
 		{
 			checkSlow(InStartBitIndex <= InEndBitIndex && InStartBitIndex <= InArray.Num() && InEndBitIndex <= InArray.Num());
@@ -1647,7 +1667,7 @@ public:
 
 	struct FWordIterator : TWordIteratorBase<uint32>
 	{
-		explicit FWordIterator(TBitArray<Allocator>& InArray)
+		explicit FWordIterator(TBitArray<Allocator>& InArray UE_LIFETIMEBOUND)
 			: TWordIteratorBase<uint32>(InArray.GetData(), 0, InArray.Num())
 		{}
 
@@ -1765,7 +1785,7 @@ class TConstSetBitIterator : public FRelativeBitReference
 public:
 
 	/** Constructor. */
-	TConstSetBitIterator(const TBitArray<Allocator>& InArray,int32 StartIndex = 0)
+	explicit TConstSetBitIterator(const TBitArray<Allocator>& InArray UE_LIFETIMEBOUND,int32 StartIndex = 0)
 		: FRelativeBitReference(StartIndex)
 		, Array                (InArray)
 		, UnvisitedBitMask     ((~0U) << (StartIndex & (NumBitsPerDWORD - 1)))
