Logging: Fixed and optimized log tracing

#rnx

diff --git a/Engine/Source/Runtime/Core/Private/Logging/StructuredLog.cpp b/Engine/Source/Runtime/Core/Private/Logging/StructuredLog.cpp
--- a/Engine/Source/Runtime/Core/Private/Logging/StructuredLog.cpp
+++ b/Engine/Source/Runtime/Core/Private/Logging/StructuredLog.cpp
@@ -28,6 +28,9 @@ void         StaticFailDebugV(const TCHAR* Error, const ANSICHAR* Expression, co
 namespace UE::Logging::Private
 {
 
+// Temporary override until performance and functionality are sufficient for this to be the default.
+CORE_API bool GConvertBasicLogToLogRecord = false;
+
 static constexpr FAsciiSet ValidLogFieldName("0123456789ABCDEFGHIJKLMNOPQRSTUVWXYZabcdefghijklmnopqrstuvwxyz_");
 
 ///////////////////////////////////////////////////////////////////////////////////////////////////////////////////////
@@ -630,7 +633,7 @@ FORCENOINLINE static void LogToTrace(const void* LogPoint, const FLogRecord& Rec
 #if LOGTRACE_ENABLED
 	TStringBuilder<1024> Message;
 	Record.FormatMessageTo(Message);
-	FLogTrace::OutputLogMessageSimple(LogPoint, TEXT("%s"), *Message);
+	FLogTrace::OutputLogMessage(LogPoint, *Message);
 #endif
 }
 
@@ -702,7 +705,10 @@ template <typename StaticLogRecordType>
 FORCENOINLINE static FLogTemplate& CreateLogTemplate(const FLogCategoryBase& Category, const StaticLogRecordType& Log, const FLogField* Fields, const int32 FieldCount)
 {
 #if LOGTRACE_ENABLED
-	FLogTrace::OutputLogMessageSpec(&Log, &Category, Log.Verbosity, Log.File, Log.Line, TEXT("%s"));
+	if (UE_TRACE_CHANNELEXPR_IS_ENABLED(LogChannel))
+	{
+		FLogTrace::OutputLogMessageSpec(&Log, &Category, Log.Verbosity, Log.File, Log.Line, TEXT("%s"));
+	}
 #endif
 
 	FLogTemplate* NewTemplate = FLogTemplate::CreateStatic(Log, Fields, FieldCount);
@@ -899,7 +905,10 @@ static FStaticBasicLogTemplateManager GStaticBasicLogTemplateManager;
 FORCENOINLINE static void OutputBasicLogMessageSpec(const FLogCategoryBase& Category, const FStaticBasicLogRecord& Log)
 {
 #if LOGTRACE_ENABLED
-	FLogTrace::OutputLogMessageSpec(&Log, &Category, Log.Verbosity, Log.File, Log.Line, TEXT("%s"));
+	if (UE_TRACE_CHANNELEXPR_IS_ENABLED(LogChannel))
+	{
+		FLogTrace::OutputLogMessageSpec(&Log, &Category, Log.Verbosity, Log.File, Log.Line, TEXT("%s"));
+	}
 #endif
 	Log.DynamicData.bInitialized.store(true, std::memory_order_release);
 }
@@ -912,6 +921,17 @@ inline static void EnsureBasicLogMessageSpec(const FLogCategoryBase& Category, c
 	}
 }
 
+// Tracing the log happens in its own function because that allows stack space for the message to
+// be returned before calling into the output devices.
+FORCENOINLINE static void BasicLogToTrace(const void* LogPoint, const TCHAR* Format, va_list Args)
+{
+#if LOGTRACE_ENABLED
+	TStringBuilder<1024> Message;
+	Message.AppendV(Format, Args);
+	FLogTrace::OutputLogMessage(LogPoint, *Message);
+#endif
+}
+
 // Serializing the log to compact binary happens in its own function because that allows stack
 // space for the writer to be returned before calling into the output devices.
 FORCENOINLINE static FCbObject SerializeBasicLogMessage(const FStaticBasicLogRecord& Log, va_list Args)
@@ -920,7 +940,10 @@ FORCENOINLINE static FCbObject SerializeBasicLogMessage(const FStaticBasicLogRec
 	Message.AppendV(Log.Format, Args);
 
 #if LOGTRACE_ENABLED
-	FLogTrace::OutputLogMessageSimple(&Log, TEXT("%s"), *Message);
+	if (UE_TRACE_CHANNELEXPR_IS_ENABLED(LogChannel))
+	{
+		FLogTrace::OutputLogMessage(&Log, *Message);
+	}
 #endif
 
 	TCbWriter<512> Writer;
@@ -935,41 +958,52 @@ void BasicLog(const FLogCategoryBase& Category, const FStaticBasicLogRecord* Log
 #if !NO_LOGGING
 	EnsureBasicLogMessageSpec(Category, *Log);
 
-#if 0 // Enable to use structured logging for these logs. Experimental.
-	va_list Args;
-	va_start(Args, Log);
-	FCbObject Fields = SerializeBasicLogMessage(*Log, Args);
-	va_end(Args);
-
-	FLogRecord Record;
-	Record.SetFormat(GStaticBasicLogFormat);
-	Record.SetTemplate(&GStaticBasicLogTemplateManager.EnsureTemplate());
-	Record.SetFields(MoveTemp(Fields));
-	Record.SetFile(Log->File);
-	Record.SetLine(Log->Line);
-	Record.SetCategory(Category.GetCategoryName());
-	Record.SetVerbosity(Log->Verbosity);
-	Record.SetTime(FLogTime::Now());
-
-	FOutputDevice* OutputDevice = nullptr;
-	switch (Log->Verbosity)
+	if (GConvertBasicLogToLogRecord)
 	{
-	case ELogVerbosity::Error:
-	case ELogVerbosity::Warning:
-	case ELogVerbosity::Display:
-	case ELogVerbosity::SetColor:
-		OutputDevice = GWarn;
-		break;
-	default:
-		break;
+		va_list Args;
+		va_start(Args, Log);
+		FCbObject Fields = SerializeBasicLogMessage(*Log, Args);
+		va_end(Args);
+
+		FLogRecord Record;
+		Record.SetFormat(GStaticBasicLogFormat);
+		Record.SetTemplate(&GStaticBasicLogTemplateManager.EnsureTemplate());
+		Record.SetFields(MoveTemp(Fields));
+		Record.SetFile(Log->File);
+		Record.SetLine(Log->Line);
+		Record.SetCategory(Category.GetCategoryName());
+		Record.SetVerbosity(Log->Verbosity);
+		Record.SetTime(FLogTime::Now());
+
+		FOutputDevice* OutputDevice = nullptr;
+		switch (Log->Verbosity)
+		{
+		case ELogVerbosity::Error:
+		case ELogVerbosity::Warning:
+		case ELogVerbosity::Display:
+		case ELogVerbosity::SetColor:
+			OutputDevice = GWarn;
+			break;
+		default:
+			break;
+		}
+		(OutputDevice ? OutputDevice : GLog)->SerializeRecord(Record);
+	}
+	else
+	{
+		va_list Args;
+	#if LOGTRACE_ENABLED
+		if (UE_TRACE_CHANNELEXPR_IS_ENABLED(LogChannel))
+		{
+			va_start(Args, Log);
+			BasicLogToTrace(Log, Log->Format, Args);
+			va_end(Args);
+		}
+	#endif
+		va_start(Args, Log);
+		FMsg::LogV(Log->File, Log->Line, Category.GetCategoryName(), Log->Verbosity, Log->Format, Args);
+		va_end(Args);
 	}
-	(OutputDevice ? OutputDevice : GLog)->SerializeRecord(Record);
-#else
-	va_list Args;
-	va_start(Args, Log);
-	FMsg::LogV(Log->File, Log->Line, Category.GetCategoryName(), Log->Verbosity, Log->Format, Args);
-	va_end(Args);
-#endif
 #endif
 }
 
