Logging: Added support for structured debug output

Debug output is written as JSON lines when the global named mutex UE_LOG_JSON exists on Windows, or when -JsonDebugOutput is present on the command line. The command line option does not affect logging that occurs before FCommandLine is initialized.

The global named mutex is meant to be used by IDEs, debuggers, and other tools to request structured log output.

diff --git a/Engine/Source/Runtime/Core/Private/GenericPlatform/GenericPlatformMisc.cpp b/Engine/Source/Runtime/Core/Private/GenericPlatform/GenericPlatformMisc.cpp
--- a/Engine/Source/Runtime/Core/Private/GenericPlatform/GenericPlatformMisc.cpp
+++ b/Engine/Source/Runtime/Core/Private/GenericPlatform/GenericPlatformMisc.cpp
@@ -911,6 +911,16 @@ void FGenericPlatformMisc::LowLevelOutputDebugStringf(const TCHAR *Fmt, ... )
 	);
 }
 
+bool FGenericPlatformMisc::IsLowLevelOutputDebugStringStructured()
+{
+	if (!FCommandLine::IsInitialized())
+	{
+		return false;
+	}
+	static bool bJsonDebugOutput = FParse::Param(FCommandLine::Get(), TEXT("JsonDebugOutput"));
+	return bJsonDebugOutput;
+}
+
 void FGenericPlatformMisc::SetUTF8Output()
 {
 	// assume that UTF-8 is possible by default anyway
diff --git a/Engine/Source/Runtime/Core/Private/Misc/OutputDeviceDebug.cpp b/Engine/Source/Runtime/Core/Private/Misc/OutputDeviceDebug.cpp
--- a/Engine/Source/Runtime/Core/Private/Misc/OutputDeviceDebug.cpp
+++ b/Engine/Source/Runtime/Core/Private/Misc/OutputDeviceDebug.cpp
@@ -5,12 +5,55 @@
 #include "Containers/StringFwd.h"
 #include "CoreGlobals.h"
 #include "HAL/PlatformMisc.h"
+#include "HAL/PlatformTime.h"
+#include "Logging/StructuredLog.h"
+#include "Misc/DateTime.h"
 #include "Misc/OutputDeviceHelper.h"
 #include "Misc/StringBuilder.h"
+#include "Misc/TVariant.h"
+#include "Serialization/CompactBinary.h"
+#include "Serialization/CompactBinarySerialization.h"
+#include "Serialization/CompactBinaryWriter.h"
+
+/** Seconds between calls to FOutputDeviceDebug::TickAsync. */
+constexpr static double GOutputDeviceDebugTickPeriod = 1.0;
 
 void FOutputDeviceDebug::Serialize(const TCHAR* Data, ELogVerbosity::Type Verbosity, const FName& Category, const double Time)
 {
-	if (Verbosity != ELogVerbosity::SetColor)
+	if (Verbosity == ELogVerbosity::SetColor)
+	{
+		return;
+	}
+
+	if (Time > 0)
+	{
+		ConditionalTickAsync(Time);
+	}
+
+	if (bSerializeAsJson.load(std::memory_order_relaxed))
+	{
+		TVariant<TCbWriter<1024>, TStringBuilder<512>> WriterMemory;
+
+		WriterMemory.Emplace<TCbWriter<1024>>();
+		TCbWriter<1024>& Writer = WriterMemory.Get<TCbWriter<1024>>();
+		Writer.BeginObject();
+		Writer.AddDateTime(ANSITEXTVIEW("Time"), FDateTime::UtcNow());
+		Writer.AddString(ANSITEXTVIEW("Category"), WriteToUtf8String<64>(Category));
+		Writer.AddString(ANSITEXTVIEW("Verbosity"), ToString(Verbosity));
+		Writer.AddString(ANSITEXTVIEW("Message"), Data);
+		Writer.EndObject();
+
+		TArray<uint8, TInlineAllocator64<1024>> Buffer;
+		Buffer.AddUninitialized((int64)Writer.GetSaveSize());
+		FCbFieldView Object = Writer.Save(MakeMemoryView(Buffer));
+
+		WriterMemory.Emplace<TStringBuilder<512>>();
+		TStringBuilder<512>& Json = WriterMemory.Get<TStringBuilder<512>>();
+		CompactBinaryToCompactJson(Object, Json);
+		Json.Append(LINE_TERMINATOR);
+		FPlatformMisc::LowLevelOutputDebugString(*Json);
+	}
+	else
 	{
 		TStringBuilder<512> Line;
 		FOutputDeviceHelper::AppendFormatLogLine(Line, Verbosity, Category, Data, GPrintLogTimes, Time);
@@ -24,7 +67,82 @@ void FOutputDeviceDebug::Serialize(const TCHAR* Data, ELogVerbosity::Type Verbos
 	Serialize(Data, Verbosity, Category, -1.0);
 }
 
+void FOutputDeviceDebug::SerializeRecord(const UE::FLogRecord& Record)
+{
+	const double Time = FPlatformTime::ToSeconds64(Record.GetTime().GetCycles());
+	ConditionalTickAsync(Time);
+
+	if (bSerializeAsJson.load(std::memory_order_relaxed))
+	{
+		TVariant<TUtf8StringBuilder<1024>, TArray<uint8, TInlineAllocator64<1024>>> ScratchMemory;
+		TVariant<TCbWriter<1024>, TStringBuilder<512>> WriterMemory;
+
+		ScratchMemory.Emplace<TUtf8StringBuilder<1024>>();
+		TUtf8StringBuilder<1024>& Message = ScratchMemory.Get<TUtf8StringBuilder<1024>>();
+		Record.FormatMessageTo(Message);
+
+		WriterMemory.Emplace<TCbWriter<1024>>();
+		TCbWriter<1024>& Writer = WriterMemory.Get<TCbWriter<1024>>();
+		Writer.BeginObject();
+		Writer.AddDateTime(ANSITEXTVIEW("Time"), Record.GetTime().GetUtcTime());
+		Writer.AddString(ANSITEXTVIEW("Category"), WriteToUtf8String<64>(Record.GetCategory()));
+		Writer.AddString(ANSITEXTVIEW("Verbosity"), ToString(Record.GetVerbosity()));
+		Writer.AddString(ANSITEXTVIEW("Message"), Message);
+		if (const TCHAR* Format = Record.GetFormat())
+		{
+			Writer.AddString(ANSITEXTVIEW("Format"), Format);
+		}
+		if (const ANSICHAR* File = Record.GetFile())
+		{
+			Writer.AddString(ANSITEXTVIEW("File"), File);
+		}
+		if (int32 Line = Record.GetLine())
+		{
+			Writer.AddInteger(ANSITEXTVIEW("Line"), Line);
+		}
+		if (const FCbObject& Fields = Record.GetFields())
+		{
+			Writer.AddObject(ANSITEXTVIEW("Fields"), Fields);
+		}
+		Writer.EndObject();
+
+		ScratchMemory.Emplace<TArray<uint8, TInlineAllocator64<1024>>>();
+		TArray<uint8, TInlineAllocator64<1024>>& Buffer = ScratchMemory.Get<TArray<uint8, TInlineAllocator64<1024>>>();
+		Buffer.AddUninitialized((int64)Writer.GetSaveSize());
+		FCbFieldView Object = Writer.Save(MakeMemoryView(Buffer));
+
+		WriterMemory.Emplace<TStringBuilder<512>>();
+		TStringBuilder<512>& Json = WriterMemory.Get<TStringBuilder<512>>();
+		CompactBinaryToCompactJson(Object, Json);
+		Json.Append(LINE_TERMINATOR);
+		FPlatformMisc::LowLevelOutputDebugString(*Json);
+	}
+	else
+	{
+		TStringBuilder<512> Line;
+		FOutputDeviceHelper::AppendFormatLogLine(Line, Record.GetVerbosity(), Record.GetCategory(), nullptr, GPrintLogTimes, Time);
+		Record.FormatMessageTo(Line);
+		Line.Append(LINE_TERMINATOR);
+		FPlatformMisc::LowLevelOutputDebugString(*Line);
+	}
+}
+
 bool FOutputDeviceDebug::CanBeUsedOnMultipleThreads() const
 {
 	return FPlatformMisc::IsLocalPrintThreadSafe();
 }
+
+FORCEINLINE void FOutputDeviceDebug::ConditionalTickAsync(double Time)
+{
+	double LocalLastTickTime = LastTickTime.load(std::memory_order_relaxed);
+	if ((LocalLastTickTime == 0.0 || Time >= LocalLastTickTime + GOutputDeviceDebugTickPeriod) &&
+		LastTickTime.compare_exchange_strong(LocalLastTickTime, Time))
+	{
+		TickAsync();
+	}
+}
+
+void FOutputDeviceDebug::TickAsync()
+{
+	bSerializeAsJson = FPlatformMisc::IsLowLevelOutputDebugStringStructured();
+}
diff --git a/Engine/Source/Runtime/Core/Private/Windows/WindowsPlatformMisc.cpp b/Engine/Source/Runtime/Core/Private/Windows/WindowsPlatformMisc.cpp
--- a/Engine/Source/Runtime/Core/Private/Windows/WindowsPlatformMisc.cpp
+++ b/Engine/Source/Runtime/Core/Private/Windows/WindowsPlatformMisc.cpp
@@ -27,6 +27,7 @@
 #include "Misc/FeedbackContext.h"
 #include "Misc/CoreDelegates.h"
 #include "Misc/App.h"
+#include "Misc/ScopeExit.h"
 #include "HAL/ExceptionHandling.h"
 #include "Misc/SecureHash.h"
 #include "HAL/IConsoleManager.h"
@@ -983,6 +984,13 @@ void FWindowsPlatformMisc::LocalPrint( const TCHAR *Message )
 #endif
 }
 
+bool FWindowsPlatformMisc::IsLowLevelOutputDebugStringStructured()
+{
+	HANDLE Mutex = OpenMutexW(SYNCHRONIZE, /*bInheritHandle*/ false, LR"--(Global\UE_LOG_JSON)--");
+	ON_SCOPE_EXIT { CloseHandle(Mutex); };
+	return !!Mutex || FGenericPlatformMisc::IsLowLevelOutputDebugStringStructured();
+}
+
 void FWindowsPlatformMisc::RequestExit( bool Force )
 {
 	UE_LOG(LogWindows, Log,  TEXT("FPlatformMisc::RequestExit(%i)"), Force );
diff --git a/Engine/Source/Runtime/Core/Public/GenericPlatform/GenericPlatformMisc.h b/Engine/Source/Runtime/Core/Public/GenericPlatform/GenericPlatformMisc.h
--- a/Engine/Source/Runtime/Core/Public/GenericPlatform/GenericPlatformMisc.h
+++ b/Engine/Source/Runtime/Core/Public/GenericPlatform/GenericPlatformMisc.h
@@ -975,6 +975,9 @@ public:
 	static void LowLevelOutputDebugString(const TCHAR *Message);
 	static void VARARGS LowLevelOutputDebugStringf(const TCHAR *Format, ... );
 
+	/** Whether LowLevelOutputDebugString[f] can receive structured output as lines of JSON. */
+	static bool IsLowLevelOutputDebugStringStructured();
+
 	/** Sets the default output to UTF8 */
 	static void SetUTF8Output();
 
diff --git a/Engine/Source/Runtime/Core/Public/Misc/OutputDeviceDebug.h b/Engine/Source/Runtime/Core/Public/Misc/OutputDeviceDebug.h
--- a/Engine/Source/Runtime/Core/Public/Misc/OutputDeviceDebug.h
+++ b/Engine/Source/Runtime/Core/Public/Misc/OutputDeviceDebug.h
@@ -7,6 +7,8 @@
 #include "Misc/OutputDevice.h"
 #include "UObject/NameTypes.h"
 
+#include <atomic>
+
 class CORE_API FOutputDeviceDebug : public FOutputDevice
 {
 public:
@@ -16,9 +18,11 @@ public:
 	* @param	Data	Text to log
 	* @param	Event	Event name used for suppression purposes
 	*/
-	virtual void Serialize(const TCHAR* Data, ELogVerbosity::Type Verbosity, const class FName& Category, const double Time) override;
+	virtual void Serialize(const TCHAR* Data, ELogVerbosity::Type Verbosity, const FName& Category, double Time) override;
+
+	virtual void Serialize(const TCHAR* Data, ELogVerbosity::Type Verbosity, const FName& Category) override;
 
-	virtual void Serialize(const TCHAR* Data, ELogVerbosity::Type Verbosity, const class FName& Category) override;
+	virtual void SerializeRecord(const UE::FLogRecord& Record) override;
 
 	virtual bool CanBeUsedOnAnyThread() const override
 	{
@@ -31,5 +35,11 @@ public:
 	{
 		return true;
 	}
-};
 
+private:
+	void ConditionalTickAsync(double Time);
+	void TickAsync();
+
+	std::atomic<double> LastTickTime = 0.0;
+	std::atomic<bool> bSerializeAsJson = false;
+};
diff --git a/Engine/Source/Runtime/Core/Public/Windows/WindowsPlatformMisc.h b/Engine/Source/Runtime/Core/Public/Windows/WindowsPlatformMisc.h
--- a/Engine/Source/Runtime/Core/Public/Windows/WindowsPlatformMisc.h
+++ b/Engine/Source/Runtime/Core/Public/Windows/WindowsPlatformMisc.h
@@ -90,6 +90,8 @@ struct CORE_API FWindowsPlatformMisc
 	static void SetUTF8Output();
 	static void LocalPrint(const TCHAR *Message);
 
+	static bool IsLowLevelOutputDebugStringStructured();
+
 	static bool IsLocalPrintThreadSafe()
 	{ 
 		//returning true when the debugger is attached is to allow
