Undefined float to int conversion test for VectorRound removed from "System::Core::Math::RoundToIntHalfToEven Vector"

diff --git a/Engine/Source/Runtime/Core/Tests/Math/UnrealMathTest.cpp b/Engine/Source/Runtime/Core/Tests/Math/UnrealMathTest.cpp
--- a/Engine/Source/Runtime/Core/Tests/Math/UnrealMathTest.cpp
+++ b/Engine/Source/Runtime/Core/Tests/Math/UnrealMathTest.cpp
@@ -3912,7 +3912,6 @@ public:
 		TestEqualAfterVectorRounding(TEXT("VectorRound32-NegZeroBitPrecision"), -16777217.0f, -16777216);
 		TestEqualAfterVectorRounding(TEXT("VectorRound32-NegZeroBitPrecision"), -16777217.5f, -16777218);
 
-		TestEqualAfterVectorRounding(TEXT("VectorRound32-FloatMax"), FLT_MAX, INT32_MIN);
 		TestEqualAfterVectorRounding(TEXT("VectorRound32-FloatMax"), -FLT_MAX, INT32_MIN);
 
 		TestEqualAfterVectorRounding(TEXT("VectorRound32-FloatMin"), FLT_MIN, 0);
