Logging: Fixed the callstack and logging for fatal errors

#rnx

diff --git a/Engine/Source/Runtime/Core/Private/Logging/StructuredLog.cpp b/Engine/Source/Runtime/Core/Private/Logging/StructuredLog.cpp
--- a/Engine/Source/Runtime/Core/Private/Logging/StructuredLog.cpp
+++ b/Engine/Source/Runtime/Core/Private/Logging/StructuredLog.cpp
@@ -1014,9 +1014,6 @@ void BasicFatalLog(const FLogCategoryBase& Category, const FStaticBasicLogRecord
 	va_start(Args, Log);
 	StaticFailDebugV(TEXT("Fatal error:"), "", Log->File, Log->Line, /*bIsEnsure*/ false, PLATFORM_RETURN_ADDRESS(), Log->Format, Args);
 	va_end(Args);
-	va_start(Args, Log);
-	FDebug::AssertFailedV("", Log->File, Log->Line, Log->Format, Args);
-	va_end(Args);
 
 	UE_DEBUG_BREAK_AND_PROMPT_FOR_REMOTE();
 	FDebug::ProcessFatalError(PLATFORM_RETURN_ADDRESS());
