Update __LINE__ in UE_LOG  to __builtin_LINE() so it compiles with edit and continue enabled

diff --git a/Engine/Source/Runtime/Core/Public/Logging/LogMacros.h b/Engine/Source/Runtime/Core/Public/Logging/LogMacros.h
--- a/Engine/Source/Runtime/Core/Public/Logging/LogMacros.h
+++ b/Engine/Source/Runtime/Core/Public/Logging/LogMacros.h
@@ -301,7 +301,7 @@ CORE_API void BasicFatalLog(const FLogCategoryBase& Category, const FStaticBasic
 		static_assert(TIsArrayOrRefOfTypeByPredicate<decltype(Format), TIsCharEncodingCompatibleWithTCHAR>::Value, "Formatting string must be a TCHAR array."); \
 		UE_VALIDATE_FORMAT_STRING(Format, ##__VA_ARGS__); \
 		static ::UE::Logging::Private::FStaticBasicLogDynamicData LOG_Dynamic; \
-		static constexpr ::UE::Logging::Private::FStaticBasicLogRecord LOG_Static(Format, __FILE__, __LINE__, ::ELogVerbosity::Verbosity, LOG_Dynamic); \
+		static constexpr ::UE::Logging::Private::FStaticBasicLogRecord LOG_Static(Format, __FILE__, __builtin_LINE(), ::ELogVerbosity::Verbosity, LOG_Dynamic); \
 		static_assert((::ELogVerbosity::Verbosity & ::ELogVerbosity::VerbosityMask) < ::ELogVerbosity::NumVerbosity && ::ELogVerbosity::Verbosity > 0, "Verbosity must be constant and in range."); \
 		if constexpr ((::ELogVerbosity::Verbosity & ELogVerbosity::VerbosityMask) == ::ELogVerbosity::Fatal) \
 		{ \
