AutomationTest: Macro to exit early upon failure of AddErrorIfFalse

diff --git a/Engine/Source/Runtime/Core/Private/Misc/AutomationTest.cpp b/Engine/Source/Runtime/Core/Private/Misc/AutomationTest.cpp
--- a/Engine/Source/Runtime/Core/Private/Misc/AutomationTest.cpp
+++ b/Engine/Source/Runtime/Core/Private/Misc/AutomationTest.cpp
@@ -1142,12 +1142,13 @@ void FAutomationTestBase::AddError(const FString& InError, int32 StackOffset)
 	}
 }
 
-void FAutomationTestBase::AddErrorIfFalse(bool bCondition, const FString& InError, int32 StackOffset)
+bool FAutomationTestBase::AddErrorIfFalse(bool bCondition, const FString& InError, int32 StackOffset)
 {
 	if (!bCondition)
 	{
 		AddError(InError, StackOffset);
 	}
+	return bCondition;
 }
 
 void FAutomationTestBase::AddErrorS(const FString& InError, const FString& InFilename, int32 InLineNumber)
diff --git a/Engine/Source/Runtime/Core/Public/Misc/AutomationTest.h b/Engine/Source/Runtime/Core/Public/Misc/AutomationTest.h
--- a/Engine/Source/Runtime/Core/Public/Misc/AutomationTest.h
+++ b/Engine/Source/Runtime/Core/Public/Misc/AutomationTest.h
@@ -65,6 +65,10 @@ class FAutomationTestBase;
 		Info.LineNumber = 1;																				\
 	}
 
+#ifndef UE_RETURN_ON_ERROR
+#define UE_RETURN_ON_ERROR(Condition, Message) if(!AddErrorIfFalse(Condition, Message)) return false;
+#endif
+
 /**
 * Flags for specifying automation test requirements/behavior
 * Update GetTestFlagsMap when updating this enum.
@@ -1317,8 +1321,9 @@ public:
 	 *
 	 * @param   bCondition The condition to validate.
 	 * @param   InError	   Error message to add to this test
+	 * @return	False if there was an error
 	 */
-	virtual void AddErrorIfFalse( bool bCondition, const FString& InError, int32 StackOffset = 0 );
+	virtual bool AddErrorIfFalse( bool bCondition, const FString& InError, int32 StackOffset = 0 );
 
 	/**
 	 * Adds an error message to this test
