Fixed const correctness in a UnrealMathUtility function.

diff --git a/Engine/Source/Runtime/Core/Public/Math/UnrealMathUtility.h b/Engine/Source/Runtime/Core/Public/Math/UnrealMathUtility.h
--- a/Engine/Source/Runtime/Core/Public/Math/UnrealMathUtility.h
+++ b/Engine/Source/Runtime/Core/Public/Math/UnrealMathUtility.h
@@ -2347,9 +2347,9 @@ public:
 	 * Get a bit in memory created from bitflags (uint32 Value:1), used for EngineShowFlags,
 	 * TestBitFieldFunctions() tests the implementation
 	 */
-	UE_NODISCARD static constexpr bool ExtractBoolFromBitfield(uint8* Ptr, uint32 Index)
+	UE_NODISCARD static constexpr bool ExtractBoolFromBitfield(const uint8* Ptr, uint32 Index)
 	{
-		uint8* BytePtr = Ptr + Index / 8;
+		const uint8* BytePtr = Ptr + Index / 8;
 		uint8 Mask = (uint8)(1 << (Index & 0x7));
 
 		return (*BytePtr & Mask) != 0;
