Add a new way of changing how CSVProfiler sets the Target frame rate.
Used CVars to avoid pulling the ApplicationCore module into Core.

diff --git a/Engine/Source/Runtime/Core/Private/ProfilingDebugging/CsvProfiler.cpp b/Engine/Source/Runtime/Core/Private/ProfilingDebugging/CsvProfiler.cpp
--- a/Engine/Source/Runtime/Core/Private/ProfilingDebugging/CsvProfiler.cpp
+++ b/Engine/Source/Runtime/Core/Private/ProfilingDebugging/CsvProfiler.cpp
@@ -102,6 +102,12 @@ TAutoConsoleVariable<int32> CVarCsvForceExit(
 	ECVF_Default
 );
 
+TAutoConsoleVariable<int32> CVarCsvTargetFrameRateOverride(
+	TEXT("csv.TargetFrameRateOverride"),
+	0,
+	TEXT("If 0, Defaults to calculating the target frame rate using rhi.SyncInterval and Max refresh rate."),
+	ECVF_Default
+);
 
 #if UE_BUILD_SHIPPING
 TAutoConsoleVariable<int32> CVarCsvShippingContinuousWrites(
@@ -2899,18 +2905,27 @@ void FCsvProfiler::BeginFrame()
 					SetMetadataInternal(TEXT("CsvID"), *CsvId);
 					UE_LOG(LogCsvProfiler, Display, TEXT("Capture started. CSV ID: %s"), *CsvId);
 
-					// Figure out the target framerate
-					int TargetFPS = FPlatformMisc::GetMaxRefreshRate();
+					int32 TargetFPS = FPlatformMisc::GetMaxRefreshRate();
+					static IConsoleVariable* CsvTargetFrameRateCVar = IConsoleManager::Get().FindConsoleVariable(TEXT("csv.TargetFrameRateOverride"));
 					static IConsoleVariable* MaxFPSCVar = IConsoleManager::Get().FindConsoleVariable(TEXT("t.MaxFPS"));
 					static IConsoleVariable* SyncIntervalCVar = IConsoleManager::Get().FindConsoleVariable(TEXT("rhi.SyncInterval"));
-					if (MaxFPSCVar && MaxFPSCVar->GetInt() > 0)
+					if (CsvTargetFrameRateCVar && CsvTargetFrameRateCVar->GetInt() > 0)
 					{
-						TargetFPS = MaxFPSCVar->GetInt();
+						TargetFPS = CsvTargetFrameRateCVar->GetInt();
 					}
-					if (SyncIntervalCVar && SyncIntervalCVar->GetInt() > 0)
+					else
 					{
-						TargetFPS = FMath::Min(TargetFPS, FPlatformMisc::GetMaxRefreshRate() / SyncIntervalCVar->GetInt());
+						// Figure out the target framerate
+						if (MaxFPSCVar && MaxFPSCVar->GetInt() > 0)
+						{
+							TargetFPS = MaxFPSCVar->GetInt();
+						}
+						if (SyncIntervalCVar && SyncIntervalCVar->GetInt() > 0)
+						{
+							TargetFPS = FMath::Min(TargetFPS, FPlatformMisc::GetMaxRefreshRate() / SyncIntervalCVar->GetInt());
+						}
 					}
+
 					SetMetadataInternal(TEXT("TargetFramerate"), *FString::FromInt(TargetFPS));
 					SetMetadataInternal(TEXT("StartTimestamp"), *FString::Printf(TEXT("%lld"), FDateTime::UtcNow().ToUnixTimestamp()));
 					SetMetadataInternal(TEXT("NamedEvents"), GCycleStatsShouldEmitNamedEvents ? TEXT("1") : TEXT("0"));
