TRingBuffer: Add MoveAppendRange to reduce costs from callsites that need to MoveAppend many elements.
#rnx
#PullRequest 10190
#10190

diff --git a/Engine/Source/Runtime/AssetRegistry/Private/AssetDataGatherer.cpp b/Engine/Source/Runtime/AssetRegistry/Private/AssetDataGatherer.cpp
--- a/Engine/Source/Runtime/AssetRegistry/Private/AssetDataGatherer.cpp
+++ b/Engine/Source/Runtime/AssetRegistry/Private/AssetDataGatherer.cpp
@@ -3760,11 +3760,7 @@ void FAssetDataGatherer::GetAndTrimSearchResults(FResults& InOutResults, FResult
 
 	auto MoveAppendRangeToRingBuffer = [](auto& InOutRingBuffer, auto& InArray)
 	{
-		InOutRingBuffer.Reserve(InOutRingBuffer.Num() + InArray.Num());
-		for (auto& Element : InArray)
-		{
-			InOutRingBuffer.Add(MoveTemp(Element));
-		}
+		InOutRingBuffer.MoveAppendRange(InArray.GetData(), InArray.Num());
 		InArray.Reset();
 	};
 
diff --git a/Engine/Source/Runtime/Core/Public/Containers/RingBuffer.h b/Engine/Source/Runtime/Core/Public/Containers/RingBuffer.h
--- a/Engine/Source/Runtime/Core/Public/Containers/RingBuffer.h
+++ b/Engine/Source/Runtime/Core/Public/Containers/RingBuffer.h
@@ -337,6 +337,45 @@ public:
 		return GetAtIndexNoCheck(AddUninitialized());
 	}
 
+	/**
+	 * Append elements from a range onto the back pointer of the RingBuffer, resizing if necessary.
+	 * Each element is move-constructed into the RingBuffer; source elements may therefore be modified.
+	 */
+	void MoveAppendRange(ElementType* OtherData, SizeType OtherNum)
+	{
+		if (OtherNum == 0)
+		{
+			return;
+		}
+		check(OtherData);
+		const SizeType OldNum = static_cast<SizeType>(Num());
+		const SizeType NewNum = OldNum + OtherNum;
+		checkf(NewNum > OldNum, TEXT("Overflow: Num() == %d, OtherNum == %d, NewNum == %d"), OldNum, OtherNum, NewNum);
+		Reserve(NewNum);
+		const SizeType LocalIndexMask = IndexMask;
+
+		const SizeType MoveRangeStart = AfterBack;
+		const SizeType MoveRangeEnd = AfterBack + OtherNum;
+		const StorageModuloType MaskedMoveRangeStart = MoveRangeStart & LocalIndexMask;
+		const StorageModuloType MaskedMoveRangeEnd = MoveRangeEnd & LocalIndexMask;
+		ElementType* const Data = GetStorage();
+		if (MaskedMoveRangeStart >= MaskedMoveRangeEnd)
+		{
+			// We should not be reaching capacity and overwriting front
+			check(MaskedMoveRangeEnd <= (Front & LocalIndexMask) && (Front & LocalIndexMask) <= MaskedMoveRangeStart);
+			const SizeType FirstMoveCount = (LocalIndexMask + 1 - MaskedMoveRangeStart);
+			MoveConstructItems(Data + MaskedMoveRangeStart, OtherData, FirstMoveCount);
+			MoveConstructItems(Data, OtherData + FirstMoveCount, MaskedMoveRangeEnd);
+		}
+		else
+		{
+			// We should not be reaching capacity and overwriting front
+			check((Front & LocalIndexMask) <= MaskedMoveRangeStart || MaskedMoveRangeEnd <= (Front & LocalIndexMask));
+			MoveConstructItems(Data + MaskedMoveRangeStart, OtherData, OtherNum);
+		}
+		AfterBack += OtherNum;
+	}
+
 	/** Add a new element before the front pointer of the RingBuffer, resizing if necessary.  The new element is move constructed from the argument. Returns the index of the added element. */
 	IndexType AddFront(ElementType&& Element)
 	{
@@ -923,7 +962,6 @@ private:
 		}
 	}
 
-
 	friend class FRingBufferTest;
 
 	/**
@@ -945,7 +983,7 @@ private:
 	 */
 	StorageModuloType Front;
 	/**
-	 * Pointer to the first location after the back poointer of the RingBuffer; when the RingBuffer is non-empty, one element before this index is the back of the RingBuffer
+	 * Pointer to the first location after the back pointer of the RingBuffer; when the RingBuffer is non-empty, one element before this index is the back of the RingBuffer
 	 * AfterBack is in StorageModulo space.
 	 * Like all values in StorageModulo space, AfterBack is allowed to underflow or overflow its StorageModuloType; it might be 0 and then subtract 1 to wrap around to 0xffffffff, but
 	 * it is always true that (AfterBack - Front) <= capacity.
diff --git a/Engine/Source/Runtime/Core/Tests/Containers/RingBufferTest.cpp b/Engine/Source/Runtime/Core/Tests/Containers/RingBufferTest.cpp
--- a/Engine/Source/Runtime/Core/Tests/Containers/RingBufferTest.cpp
+++ b/Engine/Source/Runtime/Core/Tests/Containers/RingBufferTest.cpp
@@ -6,7 +6,6 @@
 
 #include "Tests/TestHarnessAdapter.h"
 
-
 class FRingBufferTest
 {
 public:
@@ -51,6 +50,18 @@ public:
 		{
 			++NumDestruct;
 		}
+		Counter& operator=(const Counter& Other)
+		{
+			Value = Other.Value;
+			++NumCopy;
+			return *this;
+		}
+		Counter& operator=(Counter&& Other)
+		{
+			Value = Other.Value;
+			++NumMove;
+			return *this;
+		}
 
 		operator uint32() const
 		{
@@ -733,6 +744,125 @@ public:
 			}
 		}
 
+		// Test MoveAppendRange
+		{
+			constexpr uint32 MarkerValue = 0x54321;
+			constexpr int32 SourceNum = 4;
+			Counter SourceRange[SourceNum];
+			auto Clear = [MarkerValue, &SourceRange]()
+			{
+				int32 Index = 0;
+				for (Counter& Element : SourceRange)
+				{
+					Element = Counter(MarkerValue + Index++);
+				}
+				Counter::Clear();
+			};
+			auto TestResults = [this](const TCHAR* HeaderMessage, TRingBuffer<Counter>& Q, TArrayView<Counter> Prefix,
+				int32 NumVoid, int32 NumCopy, int32 NumMove, int32 NumDestruct)
+			{
+				if (!(NumVoid == Counter::NumVoid && NumCopy == Counter::NumCopy && NumMove == Counter::NumMove && NumDestruct == Counter::NumDestruct))
+				{
+					CHECK_MESSAGE(*FString::Printf(TEXT("%s: Constructor/destructor callcounts do not match"), HeaderMessage), false);
+				}
+
+				int32 SourceStart = Prefix.Num();
+				int32 ExpectedNum = SourceNum + SourceStart;
+				if (Q.Num() != ExpectedNum)
+				{
+					CHECK_MESSAGE(*FString::Printf(TEXT("%s: Length does not match"), HeaderMessage), false);
+				}
+				else
+				{
+					for (int32 Index = 0; Index < ExpectedNum; ++Index)
+					{
+						if ((Index < SourceStart && Q[Index] != Prefix[Index]) ||
+							(Index >= SourceStart && Q[Index] != MarkerValue + Index - SourceStart))
+						{
+							CHECK_MESSAGE(*FString::Printf(TEXT("%s: Element does not match"), HeaderMessage), false);
+							break;
+						}
+					}
+				}
+			};
+
+			{
+				Clear();
+				TRingBuffer<Counter> Q;
+				Q.MoveAppendRange(SourceRange, SourceNum);
+				TestResults(TEXT("MoveAppendRange from empty"),
+					Q, TArrayView<Counter>(), 0, 0, SourceNum, 0);
+			}
+
+			{
+				Clear();
+				TRingBuffer<Counter> Q;
+				Q.Reserve(SourceNum*2);
+				Q.Front = Q.IndexMask + 1 - SourceNum/2;
+				Q.AfterBack = Q.Front;
+				Q.MoveAppendRange(SourceRange, SourceNum);
+				TestResults(TEXT("MoveAppendRange from empty with wraparound"),
+					Q, TArrayView<Counter>(), 0, 0, SourceNum, 0);
+			}
+
+			{
+				Clear();
+				TRingBuffer<Counter> Q;
+				Q.Reserve(SourceNum*2);
+				Q.Front = Q.IndexMask + 1 - SourceNum;
+				Q.AfterBack = Q.Front;
+				Q.MoveAppendRange(SourceRange, SourceNum);
+				TestResults(TEXT("MoveAppendRange from empty with AfterBack + OtherNum at end"),
+					Q, TArrayView<Counter>(), 0, 0, SourceNum, 0);
+			}
+
+			{
+				TRingBuffer<Counter> Q;
+				Counter Prefix[] = { 1,2,3 };
+				Q.Reserve(UE_ARRAY_COUNT(Prefix) + SourceNum);
+				for (Counter& C : Prefix)
+				{
+					Q.Add(C);
+				}
+				Clear();
+				Q.MoveAppendRange(SourceRange, SourceNum);
+				TestResults(TEXT("MoveAppendRange with existing elements"),
+					Q, Prefix, 0, 0, SourceNum, 0);
+			}
+
+			{
+				TRingBuffer<Counter> Q;
+				Counter Prefix[] = { 1,2,3 };
+				Q.Reserve(SourceNum * 2 + UE_ARRAY_COUNT(Prefix));
+				Q.Front = Q.IndexMask + 1 - SourceNum / 2 - UE_ARRAY_COUNT(Prefix);
+				Q.AfterBack = Q.Front;
+				for (Counter& C : Prefix)
+				{
+					Q.Add(C);
+				}
+				Clear();
+				Q.MoveAppendRange(SourceRange, SourceNum);
+				TestResults(TEXT("MoveAppendRange with existing elements and wraparound"),
+					Q, Prefix, 0, 0, SourceNum, 0);
+			}
+
+			{
+				TRingBuffer<Counter> Q;
+				Counter Prefix[] = { 1,2,3 };
+				Q.Reserve(UE_ARRAY_COUNT(Prefix) + SourceNum);
+				Q.Front = Q.IndexMask;
+				Q.AfterBack = Q.Front;
+				for (Counter& C : Prefix)
+				{
+					Q.Add(C);
+				}
+				Clear();
+				Q.MoveAppendRange(SourceRange, SourceNum);
+				TestResults(TEXT("MoveAppendRange with existing elements and AfterBackMasked < FrontMasked"),
+					Q, Prefix, 0, 0, SourceNum, 0);
+			}
+		}
+
 		TestShiftIndex<uint32>();
 		TestShiftIndex<Counter>();
 
