Fixed compile errors in test

diff --git a/Engine/Source/Runtime/Core/Tests/Containers/RingBufferTest.cpp b/Engine/Source/Runtime/Core/Tests/Containers/RingBufferTest.cpp
--- a/Engine/Source/Runtime/Core/Tests/Containers/RingBufferTest.cpp
+++ b/Engine/Source/Runtime/Core/Tests/Containers/RingBufferTest.cpp
@@ -758,7 +758,7 @@ public:
 				}
 				Counter::Clear();
 			};
-			auto TestResults = [this](const TCHAR* HeaderMessage, TRingBuffer<Counter>& Q, TArrayView<Counter> Prefix,
+			auto TestResults = [this, MarkerValue, SourceNum](const TCHAR* HeaderMessage, TRingBuffer<Counter>& Q, TArrayView<Counter> Prefix,
 				int32 NumVoid, int32 NumCopy, int32 NumMove, int32 NumDestruct)
 			{
 				if (!(NumVoid == Counter::NumVoid && NumCopy == Counter::NumCopy && NumMove == Counter::NumMove && NumDestruct == Counter::NumDestruct))
