[Insights] MetadataTrace: Added UE_TRACE_METADATA_RESTORE_SCOPE(MetadataId) macro. Also did minor fixes to comments and code style.

diff --git a/Engine/Source/Runtime/Core/Private/ProfilingDebugging/MetadataTrace.cpp b/Engine/Source/Runtime/Core/Private/ProfilingDebugging/MetadataTrace.cpp
--- a/Engine/Source/Runtime/Core/Private/ProfilingDebugging/MetadataTrace.cpp
+++ b/Engine/Source/Runtime/Core/Private/ProfilingDebugging/MetadataTrace.cpp
@@ -1,10 +1,12 @@
 ﻿// Copyright Epic Games, Inc. All Rights Reserved.
+
 #include "ProfilingDebugging/MetadataTrace.h"
 
-////////////////////////////////////////////////////////////////////////////////////////////////////
 #if UE_TRACE_METADATA_ENABLED
+////////////////////////////////////////////////////////////////////////////////////////////////////
 
 UE_TRACE_CHANNEL_DEFINE(MetadataChannel);
+
 UE_TRACE_EVENT_DEFINE(MetadataStack, ClearScope);
 UE_TRACE_EVENT_DEFINE(MetadataStack, SaveStack);
 UE_TRACE_EVENT_DEFINE(MetadataStack, RestoreStack);
@@ -30,7 +32,8 @@ FMetadataRestoreScope::FMetadataRestoreScope(uint32 SavedMetadataIdentifier)
 
 ////////////////////////////////////////////////////////////////////////////////////////////////////
 FMetadataRestoreScope::~FMetadataRestoreScope()
-{}
+{
+}
 
 ////////////////////////////////////////////////////////////////////////////////////////////////////
 void FMetadataRestoreScope::ActivateScope(uint32 InStackId)
@@ -45,6 +48,5 @@ void FMetadataRestoreScope::ActivateScope(uint32 InStackId)
 	}
 }
 
+////////////////////////////////////////////////////////////////////////////////////////////////////
 #endif // UE_TRACE_METADATA_ENABLED
-
-
diff --git a/Engine/Source/Runtime/Core/Public/ProfilingDebugging/MetadataTrace.h b/Engine/Source/Runtime/Core/Public/ProfilingDebugging/MetadataTrace.h
--- a/Engine/Source/Runtime/Core/Public/ProfilingDebugging/MetadataTrace.h
+++ b/Engine/Source/Runtime/Core/Public/ProfilingDebugging/MetadataTrace.h
@@ -5,9 +5,6 @@
 #include "Trace/Config.h"
 #include "Trace/Trace.h"
 
-namespace UE { namespace Trace { class FChannel; } }
-
-////////////////////////////////////////////////////////////////////////////////////////////////////
 # if !defined(UE_TRACE_METADATA_ENABLED)
 	#if !IS_PROGRAM
 		#define	UE_TRACE_METADATA_ENABLED UE_TRACE_ENABLED
@@ -21,10 +18,10 @@ namespace UE { namespace Trace { class FChannel; } }
 // Only needed when enabled
 #include "Trace/Trace.inl"
 
-////////////////////////////////////////////////////////////////////////////////////////////////////
+namespace UE { namespace Trace { class FChannel; } }
 
 /**
- * Channel that asset meta data is output on
+ * Channel that enables the metadata tracing
  */
 CORE_API UE_TRACE_CHANNEL_EXTERN(MetadataChannel);
 
@@ -39,15 +36,12 @@ CORE_API UE_TRACE_EVENT_BEGIN_EXTERN(MetadataStack, RestoreStack)
 	UE_TRACE_EVENT_FIELD(uint32, Id)
 UE_TRACE_EVENT_END()
 
-////////////////////////////////////////////////////////////////////////////////////////////////////
 class FMetadataTrace
 {
 public:
 	CORE_API static uint32 SaveStack();
 };
 
-////////////////////////////////////////////////////////////////////////////////////////////////////
-
 // For now just forward to the event definition macros.
 #define UE_TRACE_METADATA_EVENT_BEGIN_EXTERN(ScopeName, ...) UE_TRACE_EVENT_BEGIN_EXTERN(Metadata, ScopeName, ##__VA_ARGS__)
 #define UE_TRACE_METADATA_EVENT_DEFINE(ScopeName) UE_TRACE_EVENT_DEFINE(Metadata, ScopeName)
@@ -55,6 +49,10 @@ public:
 #define UE_TRACE_METADATA_EVENT_FIELD(Type, Name) UE_TRACE_EVENT_FIELD(Type, Name)
 #define UE_TRACE_METADATA_EVENT_REFERENCE_FIELD(RefLogger, RefEventName, Name) UE_TRACE_EVENT_REFERENCE_FIELD(RefLogger, RefEventName, Name)
 #define UE_TRACE_METADATA_EVENT_END() UE_TRACE_EVENT_END()
+
+/**
+ * A scope to push/pop data on the metadata stack.
+ */
 #define UE_TRACE_METADATA_SCOPE(ScopeName, ChannelExpr) UE_TRACE_LOG_SCOPED(Metadata, ScopeName, ChannelExpr)
 
 /**
@@ -79,24 +77,30 @@ class FMetadataRestoreScope
 {
 public:
 	CORE_API FMetadataRestoreScope(uint32 SavedMetadataIdentifier);
-	CORE_API ~FMetadataRestoreScope();;
+	CORE_API ~FMetadataRestoreScope();
 private:
 	void ActivateScope(uint32 InStackId);
 	UE::Trace::Private::FScopedLogScope Inner;
 };
-////////////////////////////////////////////////////////////////////////////////////////////////////
 
-#else
+/**
+ * Temporarily applies a previously saved metadata stack. See UE_TRACE_METADATA_SAVE_STACK.
+ * When this scope is exited, the previous metadata stack is restored.
+ */
+#define UE_TRACE_METADATA_RESTORE_SCOPE(SavedMetadataIdentifier) \
+	FMetadataRestoreScope PREPROCESSOR_JOIN(__MetadataRestoreScope, __LINE__)(SavedMetadataIdentifier);
+
+#else // UE_TRACE_METADATA_ENABLED
 
-#define UE_TRACE_METADATA_EVENT_BEGIN_EXTERN(ScopeName, ...) 
-#define UE_TRACE_METADATA_EVENT_DEFINE(ScopeName) 
-#define UE_TRACE_METADATA_EVENT_BEGIN(ScopeName, ...) 
-#define UE_TRACE_METADATA_EVENT_FIELD(Type, Name) 
-#define UE_TRACE_METADATA_EVENT_END() 
-#define UE_TRACE_METADATA_SCOPE(ScopeName) 
+#define UE_TRACE_METADATA_EVENT_BEGIN_EXTERN(ScopeName, ...)
+#define UE_TRACE_METADATA_EVENT_DEFINE(ScopeName)
+#define UE_TRACE_METADATA_EVENT_BEGIN(ScopeName, ...)
+#define UE_TRACE_METADATA_EVENT_FIELD(Type, Name)
+#define UE_TRACE_METADATA_EVENT_END()
+#define UE_TRACE_METADATA_SCOPE(ScopeName)
 
 #define UE_TRACE_METADATA_CLEAR_SCOPE()
 #define UE_TRACE_METADATA_SAVE_STACK() 0
-#define UE_TRACE_METADATA_RESTORE_STACK(...)
+#define UE_TRACE_METADATA_RESTORE_SCOPE(SavedMetadataIdentifier)
 
-#endif
\ No newline at end of file
+#endif // UE_TRACE_METADATA_ENABLED
