Re-enable the range allocator after tuning. Adjusted allocation sizes handled by each memory pool to reduce fragmentation. Fix a bug where some memory pool csv stats are not counted correctly. Fix a bug in FAllocPlaybackLogModifier that prevents allocations from being freed during a playback.

Test results:
- Replay ID 580E23C8CB854F8C82D0DEAC74C3F6D8: MemFreeMB Min Avg increased from 476 to 490. 14 MB improvement. MVP and hitces per minute stayed mostly unchanged;
- Replay ID F367ADD3E0B749D1963FE5E95ECF889B: MemFreeMB Min Avg increased from 493 to 496. 3 MB improvement. MVP and hitces per minute stayed mostly unchanged.

diff --git a/Engine/Source/Runtime/Core/Public/Misc/RangeAllocator.h b/Engine/Source/Runtime/Core/Public/Misc/RangeAllocator.h
--- a/Engine/Source/Runtime/Core/Public/Misc/RangeAllocator.h
+++ b/Engine/Source/Runtime/Core/Public/Misc/RangeAllocator.h
@@ -121,7 +121,7 @@ protected:
 			: FChunkModifier(Allocator)
 			, InfoIndex(INDEX_NONE)
 		{
-			FreeRanges.Add(FRange(0, ChunkSize / MinAlignment));
+			FreeRanges.Add(FRange(0, Allocator.ChunkSize / MinAlignment));
 		}
 
 		SIZE_T Alloc(uint32 InSize, uint32 InAlignment, uint16& InOutMaxFreeRangeSize, uint16& OutOffset, uint16& OutSize)
