Add back change that will disable CRC from spawning when running Ensures
If we have disabled spawning CRC from the Engine

diff --git a/Engine/Source/Runtime/Core/Private/Unix/UnixPlatformCrashContext.cpp b/Engine/Source/Runtime/Core/Private/Unix/UnixPlatformCrashContext.cpp
--- a/Engine/Source/Runtime/Core/Private/Unix/UnixPlatformCrashContext.cpp
+++ b/Engine/Source/Runtime/Core/Private/Unix/UnixPlatformCrashContext.cpp
@@ -761,7 +761,7 @@ void FUnixCrashContext::GenerateCrashInfoAndLaunchReporter() const
 				GConfig->GetBool(TEXT("CrashReportClient"), TEXT("bStartCRCFromEngineHandler"), bStartCRCFromEngineHandler, GEngineIni);
 			}
 
-			if (bReportingNonCrash)
+			if (bReportingNonCrash && bStartCRCFromEngineHandler)
 			{
 				bool bFoundEmptySlot = false;
 
