ParkingLot: Added more comments to the tests

#rnx

diff --git a/Engine/Source/Runtime/Core/Tests/Async/ParkingLotTest.cpp b/Engine/Source/Runtime/Core/Tests/Async/ParkingLotTest.cpp
--- a/Engine/Source/Runtime/Core/Tests/Async/ParkingLotTest.cpp
+++ b/Engine/Source/Runtime/Core/Tests/Async/ParkingLotTest.cpp
@@ -23,6 +23,7 @@ TEST_CASE("Core::Async::ParkingLot", "[Core][Async]")
 		// Launch tasks that wait on the address of WaitCount.
 		for (int32 Index = 0; Index < TaskCount; ++Index)
 		{
+			// Using FThread for now because UE::Tasks::Launch does not always wake a worker thread.
 			Threads[Index] = FThread(TEXT("ParkingLotTest"), [&WaitCount, OutState = &WaitStates[Index]]
 			{
 				int32 CanWaitCount = 0;
@@ -41,7 +42,7 @@ TEST_CASE("Core::Async::ParkingLot", "[Core][Async]")
 			}
 		}
 
-		// Wake each task with a sequence number.
+		// Wake each task with a sequence number, with an extra wake call that has no thread to wake.
 		uint64 Sequence = 0;
 		for (int32 Index = 0; Index <= TaskCount; ++Index)
 		{
@@ -53,6 +54,7 @@ TEST_CASE("Core::Async::ParkingLot", "[Core][Async]")
 				CHECK(WakeState.bHasWaitingThreads == (Index + 1 < TaskCount));
 				return ++Sequence;
 			});
+			// The callback must be invoked exactly once.
 			CHECK(WakeCount == 1);
 		}
 
