Update __LINE__ in UE_LOG  to __builtin_LINE() so it compiles with edit and continue enabled

diff --git a/Engine/Source/Runtime/Core/Public/Logging/StructuredLog.h b/Engine/Source/Runtime/Core/Public/Logging/StructuredLog.h
--- a/Engine/Source/Runtime/Core/Public/Logging/StructuredLog.h
+++ b/Engine/Source/Runtime/Core/Public/Logging/StructuredLog.h
@@ -442,7 +442,7 @@ inline void LogIfActive(const LogCategoryType& Category, const LogRecordType& Lo
 	do \
 	{ \
 		static ::UE::Logging::Private::FStaticLogDynamicData LOG_Dynamic; \
-		static constexpr ::UE::Logging::Private::FStaticLogRecord LOG_Static UE_PRIVATE_LOGFMT_AGGREGATE(TEXT(Format), __FILE__, __LINE__, ::ELogVerbosity::Verbosity, LOG_Dynamic); \
+		static constexpr ::UE::Logging::Private::FStaticLogRecord LOG_Static UE_PRIVATE_LOGFMT_AGGREGATE(TEXT(Format), __FILE__, __builtin_LINE(), ::ELogVerbosity::Verbosity, LOG_Dynamic); \
 		::UE::Logging::Private::LogIfActive<FLogCategory##CategoryName, ::ELogVerbosity::Verbosity>(CategoryName, LOG_Static, ##__VA_ARGS__); \
 	} \
 	while (false)
@@ -451,7 +451,7 @@ inline void LogIfActive(const LogCategoryType& Category, const LogRecordType& Lo
 	do \
 	{ \
 		static ::UE::Logging::Private::FStaticLogDynamicData LOG_Dynamic; \
-		static constexpr ::UE::Logging::Private::FStaticLocalizedLogRecord LOG_Static UE_PRIVATE_LOGFMT_AGGREGATE(TEXT(Namespace), TEXT(Key), TEXT(Format), __FILE__, __LINE__, ::ELogVerbosity::Verbosity, LOG_Dynamic); \
+		static constexpr ::UE::Logging::Private::FStaticLocalizedLogRecord LOG_Static UE_PRIVATE_LOGFMT_AGGREGATE(TEXT(Namespace), TEXT(Key), TEXT(Format), __FILE__, __builtin_LINE(), ::ELogVerbosity::Verbosity, LOG_Dynamic); \
 		::UE::Logging::Private::LogIfActive<FLogCategory##CategoryName, ::ELogVerbosity::Verbosity>(CategoryName, LOG_Static, ##__VA_ARGS__); \
 	} \
 	while (false)
diff --git a/Engine/Source/Runtime/RenderCore/Public/ShaderParameterMacros.h b/Engine/Source/Runtime/RenderCore/Public/ShaderParameterMacros.h
--- a/Engine/Source/Runtime/RenderCore/Public/ShaderParameterMacros.h
+++ b/Engine/Source/Runtime/RenderCore/Public/ShaderParameterMacros.h
@@ -1322,7 +1322,7 @@ struct TShaderParameterStructTypeInfo<StructType[InNumElements]>
 			static constexpr int32 Alignment = SHADER_PARAMETER_STRUCT_ALIGNMENT; \
 			static constexpr bool bIsStoredInConstantBuffer = true; \
 			static constexpr const ANSICHAR* const FileName = UE_LOG_SOURCE_FILE(__FILE__); \
-			static constexpr int32 FileLine = __LINE__; \
+			static constexpr int32 FileLine = __builtin_LINE(); \
 			using TAlignedType = StructTypeName; \
 			static const FShaderParametersMetadata* GetStructMetadata() GetStructMetadataScope \
 		}; \
