- Adding a new `Settings` parameter to a new `FPlatformProcess::GetApplicationSettingsDir` function for getting `%appdata%` in Windows, and `$HOME` in macOS/Linux.
      This also allows for trimming the `Epic` path-specific components returned in the result. The default is to preserve the existing behaviour.

diff --git a/Engine/Source/Runtime/Core/Private/GenericPlatform/GenericPlatformProcess.cpp b/Engine/Source/Runtime/Core/Private/GenericPlatform/GenericPlatformProcess.cpp
--- a/Engine/Source/Runtime/Core/Private/GenericPlatform/GenericPlatformProcess.cpp
+++ b/Engine/Source/Runtime/Core/Private/GenericPlatform/GenericPlatformProcess.cpp
@@ -114,6 +114,13 @@ const TCHAR* FGenericPlatformProcess::ApplicationSettingsDir()
 	return FPlatformMisc::RootDir();
 }
 
+FString FGenericPlatformProcess::GetApplicationSettingsDir(const ApplicationSettingsContext& Settings)
+{
+	// Default to returning the original `ApplicationSettingsDir` if it's unimplemented.
+	return ApplicationSettingsDir();
+}
+
+
 const TCHAR* FGenericPlatformProcess::ComputerName()
 {
 	return TEXT("GenericComputer");
diff --git a/Engine/Source/Runtime/Core/Private/Mac/MacPlatformProcess.cpp b/Engine/Source/Runtime/Core/Private/Mac/MacPlatformProcess.cpp
--- a/Engine/Source/Runtime/Core/Private/Mac/MacPlatformProcess.cpp
+++ b/Engine/Source/Runtime/Core/Private/Mac/MacPlatformProcess.cpp
@@ -1153,6 +1153,24 @@ const TCHAR* FMacPlatformProcess::ApplicationSettingsDir()
 	return Result;
 }
 
+FString FMacPlatformProcess::GetApplicationSettingsDir(const ApplicationSettingsContext& Settings)
+{
+	TCHAR Result[MAC_MAX_PATH] = TEXT("");
+	SCOPED_AUTORELEASE_POOL;
+	NSString* ApplicationSupportFolder = [NSSearchPathForDirectoriesInDomains(NSApplicationSupportDirectory, NSUserDomainMask, YES) objectAtIndex:0];
+	FPlatformString::CFStringToTCHAR((CFStringRef) ApplicationSupportFolder, Result);
+	if (Settings.bIsEpic)
+	{
+		// @todo rocket this folder should be based on your company name, not just be hard coded to /Epic/
+		FCString::Strcat(Result, TEXT("/Epic/"));
+	}
+	else
+	{
+		FCString::Strcat(Result, TEXT("/"));
+	}
+	return FString(Result);
+}
+
 const TCHAR* FMacPlatformProcess::ComputerName()
 {
 	static TCHAR Result[256]=TEXT("");
diff --git a/Engine/Source/Runtime/Core/Private/Unix/UnixPlatformProcess.cpp b/Engine/Source/Runtime/Core/Private/Unix/UnixPlatformProcess.cpp
--- a/Engine/Source/Runtime/Core/Private/Unix/UnixPlatformProcess.cpp
+++ b/Engine/Source/Runtime/Core/Private/Unix/UnixPlatformProcess.cpp
@@ -309,6 +309,23 @@ const TCHAR* FUnixPlatformProcess::ApplicationSettingsDir()
 	return Result;
 }
 
+FString FUnixPlatformProcess::GetApplicationSettingsDir(const ApplicationSettingsContext& Settings)
+{
+	// The ApplicationSettingsDir is where the engine stores settings and configuration
+	// data.  On linux this corresponds to $HOME/.config/Epic
+	TCHAR Result[UNIX_MAX_PATH] = TEXT("");
+	FCString::Strncpy(Result, FPlatformProcess::UserHomeDir(), UE_ARRAY_COUNT(Result));
+	if (Settings.bIsEpic)
+	{
+		FCString::Strncat(Result, TEXT("/.config/Epic/"), UE_ARRAY_COUNT(Result));
+	}
+	else
+	{
+		FCString::Strncat(Result, TEXT("/.config/"), UE_ARRAY_COUNT(Result));
+	}
+	return FString(Result);
+}
+
 bool FUnixPlatformProcess::SetProcessLimits(EProcessResource::Type Resource, uint64 Limit)
 {
 	rlimit NativeLimit;
diff --git a/Engine/Source/Runtime/Core/Private/Windows/WindowsPlatformProcess.cpp b/Engine/Source/Runtime/Core/Private/Windows/WindowsPlatformProcess.cpp
--- a/Engine/Source/Runtime/Core/Private/Windows/WindowsPlatformProcess.cpp
+++ b/Engine/Source/Runtime/Core/Private/Windows/WindowsPlatformProcess.cpp
@@ -1184,6 +1184,64 @@ const TCHAR* FWindowsPlatformProcess::ApplicationSettingsDir()
 	return *WindowsApplicationSettingsDir;
 }
 
+FString FWindowsPlatformProcess::GetApplicationSettingsDir(const ApplicationSettingsContext& Settings)
+{
+	FString WindowsApplicationSettingsDir;
+	TCHAR* ApplicationSettingsPath;
+	switch (Settings.Location)
+	{
+		case ApplicationSettingsContext::Context::ApplicationSpecific:
+		{
+			const HRESULT Ret = SHGetKnownFolderPath(FOLDERID_ProgramData, 0, NULL, &ApplicationSettingsPath);
+			if (SUCCEEDED(Ret))
+			{
+				WindowsApplicationSettingsDir = FString(ApplicationSettingsPath);
+				break;
+			}
+			else
+			{
+				return "";
+			}
+		}
+		case ApplicationSettingsContext::Context::LocalUser:
+		{
+			const HRESULT Ret = SHGetKnownFolderPath(FOLDERID_LocalAppData, 0, NULL, &ApplicationSettingsPath);
+			if (SUCCEEDED(Ret))
+			{
+				WindowsApplicationSettingsDir = FString(ApplicationSettingsPath);
+				break;
+			}
+			else
+			{
+				return "";
+			}
+		}
+		case ApplicationSettingsContext::Context::RoamingUser:
+		{
+			const HRESULT Ret = SHGetKnownFolderPath(FOLDERID_RoamingAppData, 0, NULL, &ApplicationSettingsPath);
+			if (SUCCEEDED(Ret))
+			{
+				WindowsApplicationSettingsDir = FString(ApplicationSettingsPath);
+				break;
+			}
+			else
+			{
+				return "";
+			}
+		}
+		default:
+			checkf(false, TEXT("Attempting to call `ApplicationSettingsDir` with an invalid context!"));
+			return "";
+	}
+	WindowsApplicationSettingsDir = WindowsApplicationSettingsDir.Replace(TEXT("\\"), TEXT("/")) + TEXT("/");
+	CoTaskMemFree(ApplicationSettingsPath);
+	if (Settings.bIsEpic)
+	{
+		WindowsApplicationSettingsDir += TEXT("Epic/");
+	}
+	return WindowsApplicationSettingsDir;
+}
+
 const TCHAR* FWindowsPlatformProcess::ComputerName()
 {
 	static TCHAR Result[256] = {};
diff --git a/Engine/Source/Runtime/Core/Public/GenericPlatform/GenericPlatformProcess.h b/Engine/Source/Runtime/Core/Public/GenericPlatform/GenericPlatformProcess.h
--- a/Engine/Source/Runtime/Core/Public/GenericPlatform/GenericPlatformProcess.h
+++ b/Engine/Source/Runtime/Core/Public/GenericPlatform/GenericPlatformProcess.h
@@ -289,8 +289,29 @@ struct CORE_API FGenericPlatformProcess
 	/** Get the user home directory.  NOTE: Only one return value is valid at a time! **/
 	static const TCHAR *UserHomeDir();
 
+	struct ApplicationSettingsContext
+	{
+		enum class Context : int8_t
+		{
+			LocalUser,
+			RoamingUser,
+			ApplicationSpecific
+		};
+		Context Location;
+		bool bIsEpic;
+	};
+
 	/** Get application settings directory.  NOTE: Only one return value is valid at a time! **/
 	static const TCHAR* ApplicationSettingsDir();
+	
+	/** 
+	 * Get application settings directory for a given context.
+	 * 
+	 * @param Settings 	The context in which the application settings should be stored for.
+	 * 
+	 * @return 			A string to the appropriate directory.
+	 */
+	static FString GetApplicationSettingsDir(const ApplicationSettingsContext& Settings);
 
 	/** Get computer name.  NOTE: Only one return value is valid at a time! **/
 	static const TCHAR* ComputerName();
diff --git a/Engine/Source/Runtime/Core/Public/Mac/MacPlatformProcess.h b/Engine/Source/Runtime/Core/Public/Mac/MacPlatformProcess.h
--- a/Engine/Source/Runtime/Core/Public/Mac/MacPlatformProcess.h
+++ b/Engine/Source/Runtime/Core/Public/Mac/MacPlatformProcess.h
@@ -229,6 +229,7 @@ struct CORE_API FMacPlatformProcess : public FGenericPlatformProcess
 	static const TCHAR* UserTempDir();
 	static const TCHAR* UserSettingsDir();
 	static const TCHAR* ApplicationSettingsDir();
+	static FString GetApplicationSettingsDir(const ApplicationSettingsContext& Settings);
 	static const TCHAR* ComputerName();
 	static const TCHAR* UserName(bool bOnlyAlphaNumeric = true);
 	static void SetCurrentWorkingDirectoryToBaseDir();
diff --git a/Engine/Source/Runtime/Core/Public/Unix/UnixPlatformProcess.h b/Engine/Source/Runtime/Core/Public/Unix/UnixPlatformProcess.h
--- a/Engine/Source/Runtime/Core/Public/Unix/UnixPlatformProcess.h
+++ b/Engine/Source/Runtime/Core/Public/Unix/UnixPlatformProcess.h
@@ -273,6 +273,7 @@ struct CORE_API FUnixPlatformProcess : public FGenericPlatformProcess
 	static const TCHAR* UserDir();
 	static const TCHAR* UserSettingsDir();
 	static const TCHAR* ApplicationSettingsDir();
+	static FString GetApplicationSettingsDir(const ApplicationSettingsContext& Settings);
 	static void SetCurrentWorkingDirectoryToBaseDir();
 	static FString GetCurrentWorkingDirectory();
 	static FString GenerateApplicationPath(const FString& AppName, EBuildConfiguration BuildConfiguration);
diff --git a/Engine/Source/Runtime/Core/Public/Windows/WindowsPlatformProcess.h b/Engine/Source/Runtime/Core/Public/Windows/WindowsPlatformProcess.h
--- a/Engine/Source/Runtime/Core/Public/Windows/WindowsPlatformProcess.h
+++ b/Engine/Source/Runtime/Core/Public/Windows/WindowsPlatformProcess.h
@@ -141,6 +141,7 @@ public:
 	static const TCHAR* UserTempDir();
 	static const TCHAR* UserSettingsDir();
 	static const TCHAR* ApplicationSettingsDir();
+	static FString GetApplicationSettingsDir(const ApplicationSettingsContext& Settings);
 	static const TCHAR* ComputerName();
 	static const TCHAR* UserName(bool bOnlyAlphaNumeric = true);
 	static void SetCurrentWorkingDirectoryToBaseDir();
