Add FScriptSparseArray::IsCompact() + IsAllocated()

#rnx

diff --git a/Engine/Source/Runtime/Core/Public/Containers/SparseArray.h b/Engine/Source/Runtime/Core/Public/Containers/SparseArray.h
--- a/Engine/Source/Runtime/Core/Public/Containers/SparseArray.h
+++ b/Engine/Source/Runtime/Core/Public/Containers/SparseArray.h
@@ -1293,6 +1293,11 @@ public:
 		return AllocationFlags.IsValidIndex(Index) && AllocationFlags[Index];
 	}
 
+	bool IsAllocated(int32 Index) const
+	{
+		return AllocationFlags[Index];
+	}
+
 	bool IsEmpty() const
 	{
 		return Data.Num() == NumFreeIndices;
@@ -1308,6 +1313,11 @@ public:
 		return Data.Num();
 	}
 
+	bool IsCompact() const
+	{
+		return NumFreeIndices == 0;
+	}
+
 	void* GetData(int32 Index, const FScriptSparseArrayLayout& Layout)
 	{
 		return (uint8*)Data.GetData() + Layout.Size * Index;
