Add an ensure in BinnedAllocFromOS to catch large allocations > 1GB

diff --git a/Engine/Source/Runtime/Core/Private/GenericPlatform/GenericPlatformMemory.cpp b/Engine/Source/Runtime/Core/Private/GenericPlatform/GenericPlatformMemory.cpp
--- a/Engine/Source/Runtime/Core/Private/GenericPlatform/GenericPlatformMemory.cpp
+++ b/Engine/Source/Runtime/Core/Private/GenericPlatform/GenericPlatformMemory.cpp
@@ -57,11 +57,31 @@ namespace GenericPlatformMemory
 		TEXT("memory.MemoryPressureCriticalThresholdMB"),
 		GMemoryPressureCriticalThresholdMB,
 		TEXT("When the available physical memory drops below this threshold memory stats will consider this to be at critical pressure.\n"
-		     "Where a platform can specifically state it's memory pressure this test maybe ignored.\n"
-		     "0 (default) critical pressure will not use the threshold."),
+			"Where a platform can specifically state it's memory pressure this test maybe ignored.\n"
+			"0 (default) critical pressure will not use the threshold."),
 		ECVF_Default);
 }
 
+#if UE_CHECK_LARGE_ALLOCATIONS
+namespace UE::Memory::Private
+{
+	// this is a console variable ref rather than auto console variable as it's used before the
+	// console manager has been set up
+	bool GEnableLargeAllocationChecks = false;
+	FAutoConsoleVariableRef CVarEnableLargeAllocationChecks(
+		TEXT("memory.EnableLargeAllocationChecks"),
+		GEnableLargeAllocationChecks,
+		TEXT("Turn on ensure which checks no single allocation is greater than 'LargeAllocationThreshold'"),
+		ECVF_Default);
+	int32 GLargeAllocationThreshold = 1 * 1024 * 1024 * 1024;
+	FAutoConsoleVariableRef CVarLargeAllocationThreshold(
+		TEXT("memory.LargeAllocationThreshold"),
+		GLargeAllocationThreshold,
+		TEXT("Maximum size a single allocation can be before setting off the ensure enabled by 'EnableLargeAllocationChecks'"),
+		ECVF_Default);
+}
+#endif
+
 struct TUnalignedTester
 {
 	FGenericPlatformMemory::TUnaligned<uint8> A;
diff --git a/Engine/Source/Runtime/Core/Private/Unix/UnixPlatformMemory.cpp b/Engine/Source/Runtime/Core/Private/Unix/UnixPlatformMemory.cpp
--- a/Engine/Source/Runtime/Core/Private/Unix/UnixPlatformMemory.cpp
+++ b/Engine/Source/Runtime/Core/Private/Unix/UnixPlatformMemory.cpp
@@ -424,6 +424,15 @@ struct FOSAllocationDescriptor
 
 void* FUnixPlatformMemory::BinnedAllocFromOS(SIZE_T Size)
 {
+#if UE_CHECK_LARGE_ALLOCATIONS
+	if (UE::Memory::Private::GEnableLargeAllocationChecks)
+	{
+		// catch possibly erroneous large allocations
+		ensureMsgf(Size <= UE::Memory::Private::GLargeAllocationThreshold,
+			TEXT("Single allocation exceeded large allocation threshold"));
+	}
+#endif
+
 	static SIZE_T OSPageSize = FPlatformMemory::GetConstants().PageSize;
 	// guard against someone not passing size in whole pages
 	SIZE_T SizeInWholePages = (Size % OSPageSize) ? (Size + OSPageSize - (Size % OSPageSize)) : Size;
diff --git a/Engine/Source/Runtime/Core/Private/Unix/UnixPlatformProcess.cpp b/Engine/Source/Runtime/Core/Private/Unix/UnixPlatformProcess.cpp
--- a/Engine/Source/Runtime/Core/Private/Unix/UnixPlatformProcess.cpp
+++ b/Engine/Source/Runtime/Core/Private/Unix/UnixPlatformProcess.cpp
@@ -47,6 +47,14 @@ __thread uint32 FUnixTLS::ThreadIdTLS = 0;
 uint32 FUnixTLS::ThreadIdTLSKey = FUnixTLS::AllocTlsSlot();
 #endif
 
+#if UE_CHECK_LARGE_ALLOCATIONS
+static TAutoConsoleVariable<int32> CVarEnableLargeAllocationChecksAfterFork(
+	TEXT("memory.EnableLargeAllocationChecksAfterFork"),
+	false,
+	TEXT("After forking, Turn on ensure which checks no single allocation is greater than 'LargeAllocationThreshold'"),
+	ECVF_Default);
+#endif
+
 void* FUnixPlatformProcess::GetDllHandle( const TCHAR* Filename )
 {
 	check( Filename );
@@ -1619,6 +1627,13 @@ FGenericPlatformProcess::EWaitAndForkResult FUnixPlatformProcess::WaitAndFork()
 				OnEndFrameHandle = FCoreDelegates::OnEndFrame.AddStatic(FUnixPlatformProcess::OnChildEndFramePostFork);
 				FCoreDelegates::OnPostFork.Broadcast(EForkProcessRole::Child);
 
+#if UE_CHECK_LARGE_ALLOCATIONS
+				if (CVarEnableLargeAllocationChecksAfterFork.GetValueOnAnyThread())
+				{
+					UE::Memory::Private::CVarEnableLargeAllocationChecks->Set(true);
+				}
+#endif
+
 				// Children break out of the loop and return
 				RetVal = EWaitAndForkResult::Child;
 				break;
diff --git a/Engine/Source/Runtime/Core/Private/Windows/WindowsPlatformMemory.cpp b/Engine/Source/Runtime/Core/Private/Windows/WindowsPlatformMemory.cpp
--- a/Engine/Source/Runtime/Core/Private/Windows/WindowsPlatformMemory.cpp
+++ b/Engine/Source/Runtime/Core/Private/Windows/WindowsPlatformMemory.cpp
@@ -416,6 +416,14 @@ bool FWindowsPlatformMemory::PageProtect(void* const Ptr, const SIZE_T Size, con
 }
 void* FWindowsPlatformMemory::BinnedAllocFromOS( SIZE_T Size )
 {
+#if UE_CHECK_LARGE_ALLOCATIONS
+	if (UE::Memory::Private::GEnableLargeAllocationChecks)
+	{
+		// catch possibly erroneous large allocations
+		ensureMsgf(Size <= UE::Memory::Private::GLargeAllocationThreshold,
+			TEXT("Single allocation exceeded large allocation threshold"));
+	}
+#endif
 	void* Ptr = VirtualAlloc( NULL, Size, MEM_COMMIT | MEM_RESERVE, PAGE_READWRITE );
 	LLM(FLowLevelMemTracker::Get().OnLowLevelAlloc(ELLMTracker::Platform, Ptr, Size));
 	return Ptr;
diff --git a/Engine/Source/Runtime/Core/Public/GenericPlatform/GenericPlatformMemory.h b/Engine/Source/Runtime/Core/Public/GenericPlatform/GenericPlatformMemory.h
--- a/Engine/Source/Runtime/Core/Public/GenericPlatform/GenericPlatformMemory.h
+++ b/Engine/Source/Runtime/Core/Public/GenericPlatform/GenericPlatformMemory.h
@@ -21,6 +21,20 @@ class FString;
 struct FGenericMemoryStats;
 struct FPlatformMemoryStats;
 
+#ifndef UE_CHECK_LARGE_ALLOCATIONS
+#define UE_CHECK_LARGE_ALLOCATIONS 0
+#endif
+
+#if UE_CHECK_LARGE_ALLOCATIONS
+class FAutoConsoleVariableRef;
+namespace UE::Memory::Private
+{
+	extern bool GEnableLargeAllocationChecks;
+	extern int32 GLargeAllocationThreshold;
+	extern FAutoConsoleVariableRef CVarEnableLargeAllocationChecks;
+}
+#endif
+
 /**
 * Platform-dependent "bucket" for memory size, where Default is the normal, or possibly the largest.
 * This is generally used for texture LOD settings for how to fit in smaller memory devices
