Core: Removed deprecated code related to logging, names, and strings

#rnx

diff --git a/Engine/Source/Runtime/Core/Private/UObject/UnrealNames.cpp b/Engine/Source/Runtime/Core/Private/UObject/UnrealNames.cpp
--- a/Engine/Source/Runtime/Core/Private/UObject/UnrealNames.cpp
+++ b/Engine/Source/Runtime/Core/Private/UObject/UnrealNames.cpp
@@ -3192,39 +3192,6 @@ FName::FName(int32 Len, const UTF8CHAR* Name, EFindName FindType)
 	: FName(FNameHelper::MakeDetectNumber(MakeUnconvertedView(Name, Len), FindType))
 {}
 
-FName::FName(const WIDECHAR* Name, int32 InNumber, EFindName FindType)
-	: FName(FNameHelper::MakeWithNumber(MakeUnconvertedView(Name), FindType, InNumber))
-{}
-
-FName::FName(const ANSICHAR* Name, int32 InNumber, EFindName FindType)
-	: FName(FNameHelper::MakeWithNumber(MakeUnconvertedView(Name), FindType, InNumber))
-{}
-
-FName::FName(const UTF8CHAR* Name, int32 InNumber, EFindName FindType)
-	: FName(FNameHelper::MakeWithNumber(MakeUnconvertedView(Name), FindType, InNumber))
-{}
-
-FName::FName(int32 Len, const WIDECHAR* Name, int32 InNumber, EFindName FindType)
-	: FName(InNumber != NAME_NO_NUMBER_INTERNAL ? FNameHelper::MakeWithNumber(MakeUnconvertedView(Name, Len), FindType, InNumber)
-		: FNameHelper::MakeDetectNumber(MakeUnconvertedView(Name, Len), FindType))
-{}
-
-FName::FName(int32 Len, const ANSICHAR* Name, int32 InNumber, EFindName FindType)
-	: FName(InNumber != NAME_NO_NUMBER_INTERNAL ? FNameHelper::MakeWithNumber(MakeUnconvertedView(Name, Len), FindType, InNumber)
-		: FNameHelper::MakeDetectNumber(MakeUnconvertedView(Name, Len), FindType))
-{}
-
-FName::FName(int32 Len, const UTF8CHAR* Name, int32 InNumber, EFindName FindType)
-	: FName(InNumber != NAME_NO_NUMBER_INTERNAL ? FNameHelper::MakeWithNumber(MakeUnconvertedView(Name, Len), FindType, InNumber)
-		: FNameHelper::MakeDetectNumber(MakeUnconvertedView(Name, Len), FindType))
-{}
-
-FName::FName(const TCHAR* Name, int32 InNumber, EFindName FindType, bool bSplitName)
-	: FName(InNumber == NAME_NO_NUMBER_INTERNAL && bSplitName
-		? FNameHelper::MakeDetectNumber(MakeUnconvertedView(Name), FindType)
-		: FNameHelper::MakeWithNumber(MakeUnconvertedView(Name), FindType, InNumber))
-{}
-
 FName::FName(const WIDECHAR* Name, int32 InNumber)
 	: FName(FNameHelper::MakeWithNumber(MakeUnconvertedView(Name), FNAME_Add, InNumber))
 {}
@@ -3810,12 +3777,6 @@ void FName::AutoTest()
 	check(FName("_2147483646").IsNone());
 	check(!FName("_2147483647").IsNone()); // invalid number
 
-	// Find existing numbered none 
-PRAGMA_DISABLE_DEPRECATION_WARNINGS
-	check(FName("None_7", FNAME_Find) == None_7);
-	check(FName("None", NAME_EXTERNAL_TO_INTERNAL(7), FNAME_Find) == None_7);
-PRAGMA_ENABLE_DEPRECATION_WARNINGS
-
 	const FName Cylinder(NAME_Cylinder);
 	check(Cylinder == FName("Cylinder"));
 	check(Cylinder.ToEName());
@@ -3861,22 +3822,6 @@ PRAGMA_ENABLE_DEPRECATION_WARNINGS
 	if (Once)
 	{
 		check(FName("UniqueUnicorn!!", FNAME_Find) == FName());
-#if UE_FNAME_OUTLINE_NUMBER
-PRAGMA_DISABLE_DEPRECATION_WARNINGS
-		check(FName("UniqueUnicorn!!", 17, FNAME_Find) == FName());			// We can't find the suffix version either
-PRAGMA_ENABLE_DEPRECATION_WARNINGS
-#else
-PRAGMA_DISABLE_DEPRECATION_WARNINGS
-		check(FName("UniqueUnicorn!!", 17, FNAME_Find) == FName(FNameEntryId(), FNameEntryId(), 17));	 // We fail to find the string part so we make the name "None_16"
-PRAGMA_ENABLE_DEPRECATION_WARNINGS
-#endif
-
-#if UE_FNAME_OUTLINE_NUMBER
-PRAGMA_DISABLE_DEPRECATION_WARNINGS
-		const FName NumberedNone = FName("None", 17, FNAME_Add);
-		check(FName("UniqueUnicorn!!", 17, FNAME_Find) == FName());		// Still can't find a numbered name just because we added a None with the same number
-PRAGMA_ENABLE_DEPRECATION_WARNINGS
-#endif
 
 		// Check that FNAME_Find can find entries
 		const FName UniqueName("UniqueUnicorn!!", FNAME_Add);
@@ -3886,14 +3831,6 @@ PRAGMA_ENABLE_DEPRECATION_WARNINGS
 		check(FName(TEXT("UNIQUEUNICORN!!"), FNAME_Find) == UniqueName);
 		check(FName("uniqueunicorn!!", FNAME_Find) == UniqueName);
 
-#if UE_FNAME_OUTLINE_NUMBER
-PRAGMA_DISABLE_DEPRECATION_WARNINGS
-		check(FName("UniqueUnicorn!!", 17, FNAME_Find) == FName());		// Still can't find a numbered name just because we made the base version and a None with the same number
-		check(FName("UniqueUnicorn!!", 17, FNAME_Add) != FName());		// Explicitly add it
-		check(FName("UniqueUnicorn!!", 17, FNAME_Find) != FName());		// Now we can find it
-		check(FName("UniqueUnicorn!!", 127, FNAME_Find) == FName());	// But we can't find one with a different number
-PRAGMA_ENABLE_DEPRECATION_WARNINGS
-#endif	
 		Once = false;
 	}
 
diff --git a/Engine/Source/Runtime/Core/Public/Containers/StringFwd.h b/Engine/Source/Runtime/Core/Public/Containers/StringFwd.h
--- a/Engine/Source/Runtime/Core/Public/Containers/StringFwd.h
+++ b/Engine/Source/Runtime/Core/Public/Containers/StringFwd.h
@@ -31,11 +31,6 @@ template <int32 BufferSize> using TWideStringBuilder = TStringBuilderWithBuffer<
 /** An extendable string builder for UTF8CHAR. */
 template <int32 BufferSize> using TUtf8StringBuilder = TStringBuilderWithBuffer<UTF8CHAR, BufferSize>;
 
-/** A fixed-size string builder for TCHAR. */
-template <int32 BufferSize> using TFixedStringBuilder UE_DEPRECATED(4.25, "'TFixedStringBuilder' is deprecated. Please use 'TStringBuilder' instead!") = TStringBuilderWithBuffer<TCHAR, BufferSize>;
-/** A fixed-size string builder for ANSICHAR. */
-template <int32 BufferSize> using TFixedAnsiStringBuilder UE_DEPRECATED(4.25, "'TFixedAnsiStringBuilder' is deprecated. Please use 'TAnsiStringBuilder' instead!") = TStringBuilderWithBuffer<ANSICHAR, BufferSize>;
-
 // String View
 
 template <typename CharType> class TStringView;
diff --git a/Engine/Source/Runtime/Core/Public/Containers/StringView.h b/Engine/Source/Runtime/Core/Public/Containers/StringView.h
--- a/Engine/Source/Runtime/Core/Public/Containers/StringView.h
+++ b/Engine/Source/Runtime/Core/Public/Containers/StringView.h
@@ -79,7 +79,6 @@ class TStringView
 {
 public:
 	using ElementType = CharType;
-	using SizeType UE_DEPRECATED(5.0, "Use int32 instead.") = int32;
 	using ViewType = TStringView<CharType>;
 
 public:
@@ -471,40 +470,6 @@ constexpr inline auto MakeStringView(CharPtrType&& CharPtr UE_LIFETIMEBOUND, int
 
 //////////////////////////////////////////////////////////////////////////
 
-#if PLATFORM_TCHAR_IS_UTF8CHAR
-UE_DEPRECATED(5.1, "The _SV literal is deprecated. Use TEXTVIEW.")
-inline FStringView operator "" _SV(const ANSICHAR* String, size_t Size)
-{
-	return FStringView((const UTF8CHAR*)String, Size);
-}
-#else
-UE_DEPRECATED(5.1, "The _SV literal is deprecated. Use TEXTVIEW.")
-constexpr inline FStringView operator "" _SV(const TCHAR* String, size_t Size)
-{
-	return FStringView(String, Size);
-}
-#endif
-
-UE_DEPRECATED(5.1, "The _ASV literal is deprecated. Use ANSITEXTVIEW.")
-constexpr inline FAnsiStringView operator "" _ASV(const ANSICHAR* String, size_t Size)
-{
-	return FAnsiStringView(String, Size);
-}
-
-UE_DEPRECATED(5.1, "The _WSV literal is deprecated. Use WIDETEXTVIEW.")
-constexpr inline FWideStringView operator "" _WSV(const WIDECHAR* String, size_t Size)
-{
-	return FWideStringView(String, Size);
-}
-
-UE_DEPRECATED(5.1, "The _U8SV literal is deprecated. Use UTF8TEXTVIEW.")
-/*constexpr*/ inline FUtf8StringView operator "" _U8SV(const ANSICHAR* String, size_t Size)
-{
-	// Would like this operator to be constexpr, but cannot be until after this operator can take a UTF8CHAR*
-	// rather than an ANSICHAR*, which won't be until we have C++20 char8_t string literals.
-	return FUtf8StringView(reinterpret_cast<const UTF8CHAR*>(String), Size);
-}
-
 #if PLATFORM_TCHAR_IS_UTF8CHAR
 inline FStringView operator "" _PrivateSV(const ANSICHAR* String, size_t Size)
 {
diff --git a/Engine/Source/Runtime/Core/Public/Logging/LogMacros.h b/Engine/Source/Runtime/Core/Public/Logging/LogMacros.h
--- a/Engine/Source/Runtime/Core/Public/Logging/LogMacros.h
+++ b/Engine/Source/Runtime/Core/Public/Logging/LogMacros.h
@@ -196,7 +196,6 @@ CORE_API void BasicFatalLog(const FLogCategoryBase& Category, const FStaticBasic
 	#define DEFINE_LOG_CATEGORY_STATIC(...)
 	#define DECLARE_LOG_CATEGORY_CLASS(...)
 	#define DEFINE_LOG_CATEGORY_CLASS(...)
-	#define UE_SECURITY_LOG(...) DEPRECATED_MACRO(5.1, "UE_SECURITY_LOG has been deprecated in favor of UE_LOG")
 
 #else
 
@@ -261,26 +260,6 @@ CORE_API void BasicFatalLog(const FLogCategoryBase& Category, const FStaticBasic
 	// DO NOT USE. Use UE_LOG because this will be deprecated in a future release.
 	#define UE_LOG_CLINKAGE UE_LOG
 
-	/**
-	* A  macro that outputs a formatted message to the log specifically used for security events
-	* @param NetConnection, a valid UNetConnection
-	* @param SecurityEventType, a security event type (ESecurityEvent::Type)
-	* @param Format, format text
-	***/
-	#define UE_SECURITY_LOG(NetConnection, SecurityEventType, Format, ...) \
-	{ \
-		DEPRECATED_MACRO(5.1, "UE_SECURITY_LOG has been deprecated in favor of UE_LOG") \
-		static_assert(TIsArrayOrRefOfTypeByPredicate<decltype(Format), TIsCharEncodingCompatibleWithTCHAR>::Value, "Formatting string must be a TCHAR array."); \
-		check(NetConnection != nullptr); \
-		CA_CONSTANT_IF((ELogVerbosity::Warning & ELogVerbosity::VerbosityMask) <= ELogVerbosity::COMPILED_IN_MINIMUM_VERBOSITY && (ELogVerbosity::Warning & ELogVerbosity::VerbosityMask) <= FLogCategoryLogSecurity::CompileTimeVerbosity) \
-		{ \
-			if (!LogSecurity.IsSuppressed(ELogVerbosity::Warning)) \
-			{ \
-				FMsg::Logf_Internal(UE_LOG_SOURCE_FILE(__FILE__), __LINE__, LogSecurity.GetCategoryName(), ELogVerbosity::Warning, TEXT("%s: %s: %s"), *(NetConnection->RemoteAddressToString()), ToString(SecurityEventType), *FString::Printf(Format, ##__VA_ARGS__)); \
-			} \
-		} \
-	}
-
 	/**
 	 * A macro that conditionally logs a formatted message if the log category is active at the requested verbosity level.
 	 *
@@ -403,31 +382,6 @@ CORE_API void BasicFatalLog(const FLogCategoryBase& Category, const FStaticBasic
 	NetConnection->FlushNet(true)
 #endif
 
-/**
-	* A  macro that closes the connection and logs the security event on the server and the client
-	* @param NetConnection, a valid UNetConnection
-	* @param SecurityEventType, a security event type (ESecurityEvent::Type)
-	* @param Format, format text
-***/
-#define CLOSE_CONNECTION_DUE_TO_SECURITY_VIOLATION_INNER(NetConnection, SecurityEventType, Format, ...) \
-{ \
-	DEPRECATED_MACRO(5.1, "CLOSE_CONNECTION_DUE_TO_SECURITY_VIOLATION has been deprecated") \
-	static_assert(TIsArrayOrRefOfTypeByPredicate<decltype(Format), TIsCharEncodingCompatibleWithTCHAR>::Value, "Formatting string must be a TCHAR array."); \
-	check(NetConnection != nullptr); \
-	FString SecurityPrint = FString::Printf(Format, ##__VA_ARGS__); \
-	UE_LOG(LogSecurity, Warning, TEXT("%s: %s: %s"), *(NetConnection)->RemoteAddressToString(), ToString(SecurityEventType), SecurityEventType, *SecurityPrint); \
-	UE_LOG(LogSecurity, Warning, TEXT("%s: Closed: Connection closed"), *(NetConnection)->RemoteAddressToString()); \
-	NetConnection->Close({FromSecurityEvent(SecurityEventType), SecurityPrint}); \
-}
-#if USE_SERVER_PERF_COUNTERS
-#define CLOSE_CONNECTION_DUE_TO_SECURITY_VIOLATION(NetConnection, SecurityEventType, Format, ...) \
-	CLOSE_CONNECTION_DUE_TO_SECURITY_VIOLATION_INNER(NetConnection, SecurityEventType, Format, ##__VA_ARGS__) \
-	PerfCountersIncrement(TEXT("ClosedConnectionsDueToSecurityViolations"));
-#else
-#define CLOSE_CONNECTION_DUE_TO_SECURITY_VIOLATION(NetConnection, SecurityEventType, Format, ...) \
-	CLOSE_CONNECTION_DUE_TO_SECURITY_VIOLATION_INNER(NetConnection, SecurityEventType, Format, ##__VA_ARGS__)
-#endif
-
 extern CORE_API int32 GEnsureOnNANDiagnostic;
 
 // Macro to either log an error or ensure on a NaN error.
diff --git a/Engine/Source/Runtime/Core/Public/Misc/OutputDeviceRedirector.h b/Engine/Source/Runtime/Core/Public/Misc/OutputDeviceRedirector.h
--- a/Engine/Source/Runtime/Core/Public/Misc/OutputDeviceRedirector.h
+++ b/Engine/Source/Runtime/Core/Public/Misc/OutputDeviceRedirector.h
@@ -71,9 +71,6 @@ ENUM_CLASS_FLAGS(EOutputDeviceRedirectorFlushOptions);
 class CORE_API FOutputDeviceRedirector final : public FOutputDevice
 {
 public:
-	UE_DEPRECATED(5.1, "TLocalOutputDevicesArray is being removed. Use TArray<FOutputDevice*, TInlineAllocator<16>>.")
-	typedef TArray<FOutputDevice*, TInlineAllocator<16>> TLocalOutputDevicesArray;
-
 	/** Initialization constructor. */
 	FOutputDeviceRedirector();
 
@@ -105,10 +102,6 @@ public:
 	/** Flushes lines buffered by secondary threads. */
 	void FlushThreadedLogs(EOutputDeviceRedirectorFlushOptions Options = EOutputDeviceRedirectorFlushOptions::None);
 
-	/** See Panic. */
-	UE_DEPRECATED(5.1, "Use Panic() when the caller is handling a crash, otherwise use FlushThreadedLogs().")
-	void PanicFlushThreadedLogs() { Panic(); }
-
 	/**
 	 * Serializes the current backlog to the specified output device.
 	 * @param OutputDevice   Output device that will receive the current backlog.
@@ -131,9 +124,6 @@ public:
 	 */
 	void SetCurrentThreadAsPrimaryThread();
 
-	UE_DEPRECATED(5.1, "Use SetCurrentThreadAsPrimaryThread().")
-	inline void SetCurrentThreadAsMasterThread() { SetCurrentThreadAsPrimaryThread(); }
-
 	/**
 	 * Starts a dedicated primary thread that redirects logs to buffered output devices.
 	 *
diff --git a/Engine/Source/Runtime/Core/Public/Misc/StringBuilder.h b/Engine/Source/Runtime/Core/Public/Misc/StringBuilder.h
--- a/Engine/Source/Runtime/Core/Public/Misc/StringBuilder.h
+++ b/Engine/Source/Runtime/Core/Public/Misc/StringBuilder.h
@@ -215,21 +215,8 @@ public:
 		return *this;
 	}
 
-	template <
-		typename AppendedCharType,
-		std::enable_if_t<TIsCharType<AppendedCharType>::Value>* = nullptr
-	>
-	UE_DEPRECATED(5.0, "Use AppendChar instead of Append.")
-	inline BuilderType& Append(AppendedCharType Char)
-	{
-		return AppendChar(Char);
-	}
-
+	UE_DEPRECATED(5.3, "Use Append instead of AppendAnsi.")
 	inline BuilderType& AppendAnsi(const FAnsiStringView String) { return Append(String); }
-	UE_DEPRECATED(5.0, "Use Append instead of AppendAnsi.")
-	inline BuilderType& AppendAnsi(const ANSICHAR* const String) { return Append(String); }
-	UE_DEPRECATED(5.0, "Use Append instead of AppendAnsi.")
-	inline BuilderType& AppendAnsi(const ANSICHAR* const String, const int32 Length) { return Append(String, Length); }
 
 	/** Replace characters at given position and length with substring */
 	void ReplaceAt(int32 Pos, int32 RemoveLen, ViewType Str)
diff --git a/Engine/Source/Runtime/Core/Public/UObject/NameTypes.h b/Engine/Source/Runtime/Core/Public/UObject/NameTypes.h
--- a/Engine/Source/Runtime/Core/Public/UObject/NameTypes.h
+++ b/Engine/Source/Runtime/Core/Public/UObject/NameTypes.h
@@ -77,11 +77,6 @@ struct FNameEntryId
 	// Returns true if this FNameEntryId is not equivalent to NAME_None
 	explicit operator bool() const { return Value != 0; }
 
-	UE_DEPRECATED(4.23, "NAME_INDEX is replaced by FNameEntryId, which is no longer a contiguous integer. "
-						"Please use 'GetTypeHash(MyId)' instead of 'MyId' for hash functions. "
-						"ToUnstableInt() can be used in other advanced cases.")
-	operator int32() const;
-
 	/** Get process specific integer */
 	uint32 ToUnstableInt() const { return Value; }
 
@@ -710,20 +705,6 @@ public:
 		return !(*this == Other);
 	}
 
-
-	UE_DEPRECATED(4.23, "Please use FastLess() / FNameFastLess or LexicalLess() / FNameLexicalLess instead. "
-		"Default lexical sort order is deprecated to avoid unintended expensive sorting. ")
-	FORCEINLINE bool operator<( const FName& Other ) const
-	{
-		return LexicalLess(Other);
-	}
-
-	UE_DEPRECATED(4.23, "Please use B.FastLess(A) or B.LexicalLess(A) instead of A > B.")
-	FORCEINLINE bool operator>(const FName& Other) const
-	{
-		return Other.LexicalLess(*this);
-	}
-
 	/** Fast non-alphabetical order that is only stable during this process' lifetime. */
 	FORCEINLINE bool FastLess(const FName& Other) const
 	{
@@ -990,18 +971,6 @@ public:
 	 * @param Name Value for the string portion of the name
 	 * @param Number Value for the number portion of the name
 	 */
-	UE_DEPRECATED(5.1, "EFindName has been removed from constructors taking a Number argument to add clarity around UE_FNAME_OUTLINE_NUMBER.")
-	CORE_API FName(const WIDECHAR* Name, int32 Number, EFindName FindType);
-	UE_DEPRECATED(5.1, "EFindName has been removed from constructors taking a Number argument to add clarity around UE_FNAME_OUTLINE_NUMBER.")
-	CORE_API FName(const ANSICHAR* Name, int32 Number, EFindName FindType);
-	UE_DEPRECATED(5.1, "EFindName has been removed from constructors taking a Number argument to add clarity around UE_FNAME_OUTLINE_NUMBER.")
-	CORE_API FName(const UTF8CHAR* Name, int32 Number, EFindName FindType);
-	UE_DEPRECATED(5.1, "EFindName has been removed from constructors taking a Number argument to add clarity around UE_FNAME_OUTLINE_NUMBER.")
-	CORE_API FName(int32 Len, const WIDECHAR* Name, int32 Number, EFindName FindType);
-	UE_DEPRECATED(5.1, "EFindName has been removed from constructors taking a Number argument to add clarity around UE_FNAME_OUTLINE_NUMBER.")
-	CORE_API FName(int32 Len, const ANSICHAR* Name, int32 Number, EFindName FindType);
-	UE_DEPRECATED(5.1, "EFindName has been removed from constructors taking a Number argument to add clarity around UE_FNAME_OUTLINE_NUMBER.")
-	CORE_API FName(int32 Len, const UTF8CHAR* Name, int32 Number, EFindName FindType);
 	CORE_API FName(const WIDECHAR* Name, int32 Number);
 	CORE_API FName(const ANSICHAR* Name, int32 Number);
 	CORE_API FName(const UTF8CHAR* Name, int32 Number);
@@ -1009,31 +978,6 @@ public:
 	CORE_API FName(int32 Len, const ANSICHAR* Name, int32 Number);
 	CORE_API FName(int32 Len, const UTF8CHAR* Name, int32 Number);
 
-	UE_DEPRECATED(5.1, "EFindName has been removed from constructors taking a Number argument to add clarity around UE_FNAME_OUTLINE_NUMBER.")
-	inline FName(TStringView<ANSICHAR> View, int32 InNumber, EFindName FindType)
-		: FName(NoInit)
-	{
-PRAGMA_DISABLE_DEPRECATION_WARNINGS
-		*this = FName(View.Len(), View.GetData(), InNumber, FindType);
-PRAGMA_ENABLE_DEPRECATION_WARNINGS
-	}
-	UE_DEPRECATED(5.1, "EFindName has been removed from constructors taking a Number argument to add clarity around UE_FNAME_OUTLINE_NUMBER.")
-	inline FName(TStringView<WIDECHAR> View, int32 InNumber, EFindName FindType)
-		: FName(NoInit)
-	{
-PRAGMA_DISABLE_DEPRECATION_WARNINGS
-		*this = FName(View.Len(), View.GetData(), InNumber, FindType);
-PRAGMA_ENABLE_DEPRECATION_WARNINGS
-	}
-	UE_DEPRECATED(5.1, "EFindName has been removed from constructors taking a Number argument to add clarity around UE_FNAME_OUTLINE_NUMBER.")
-	inline FName(TStringView<UTF8CHAR> View, int32 InNumber, EFindName FindType)
-		: FName(NoInit)
-	{
-PRAGMA_DISABLE_DEPRECATION_WARNINGS
-		*this = FName(View.Len(), View.GetData(), InNumber, FindType);
-PRAGMA_ENABLE_DEPRECATION_WARNINGS
-	}
-
 	inline FName(TStringView<ANSICHAR> View, int32 InNumber)
 		: FName(NoInit)
 	{
@@ -1059,8 +1003,6 @@ PRAGMA_ENABLE_DEPRECATION_WARNINGS
 	 * @param FindType Action to take (see EFindName)
 	 * @param bSplitName true if the trailing number should be split from the name when Number == NAME_NO_NUMBER_INTERNAL, or false to always use the name as-is
 	 */
-	UE_DEPRECATED(5.1, "EFindName has been removed from constructors taking a Number argument to add clarity around UE_FNAME_OUTLINE_NUMBER.")
-	CORE_API FName(const TCHAR* Name, int32 InNumber, EFindName FindType, bool bSplitName);
 	CORE_API FName(const TCHAR* Name, int32 InNumber, bool bSplitName);
 
 	/**
@@ -1600,9 +1542,6 @@ struct FNameFastLess
 	}
 };
 
-UE_DEPRECATED(4.23, "Please use FNameFastLess instead.")
-typedef FNameFastLess FNameSortIndexes;
-
 /** Slow alphabetical order that is stable / deterministic over process runs */
 struct FNameLexicalLess
 {
