Logging: Work around Cycles64 and Seconds from FPlatformTime having different epochs

Request Seconds where needed instead of converting from Cycles64. Times will be a little off, but not as much as with the wrong epoch, and no worse than log times have been to date.

#rnx

diff --git a/Engine/Source/Runtime/Core/Private/Misc/OutputDevice.cpp b/Engine/Source/Runtime/Core/Private/Misc/OutputDevice.cpp
--- a/Engine/Source/Runtime/Core/Private/Misc/OutputDevice.cpp
+++ b/Engine/Source/Runtime/Core/Private/Misc/OutputDevice.cpp
@@ -74,7 +74,7 @@ void FOutputDevice::SerializeRecord(const UE::FLogRecord& Record)
 {
 	TStringBuilder<512> Text;
 	Record.FormatMessageTo(Text);
-	Serialize(*Text, Record.GetVerbosity(), Record.GetCategory(), FPlatformTime::ToSeconds64(Record.GetTime().GetCycles()));
+	Serialize(*Text, Record.GetVerbosity(), Record.GetCategory());
 }
 
 /** Critical errors. */
diff --git a/Engine/Source/Runtime/Core/Private/Misc/OutputDeviceDebug.cpp b/Engine/Source/Runtime/Core/Private/Misc/OutputDeviceDebug.cpp
--- a/Engine/Source/Runtime/Core/Private/Misc/OutputDeviceDebug.cpp
+++ b/Engine/Source/Runtime/Core/Private/Misc/OutputDeviceDebug.cpp
@@ -18,17 +18,14 @@
 /** Seconds between calls to FOutputDeviceDebug::TickAsync. */
 constexpr static double GOutputDeviceDebugTickPeriod = 1.0;
 
-void FOutputDeviceDebug::Serialize(const TCHAR* Data, ELogVerbosity::Type Verbosity, const FName& Category, const double Time)
+void FOutputDeviceDebug::Serialize(const TCHAR* Data, ELogVerbosity::Type Verbosity, const FName& Category, double Time)
 {
 	if (Verbosity == ELogVerbosity::SetColor)
 	{
 		return;
 	}
 
-	if (Time > 0)
-	{
-		ConditionalTickAsync(Time);
-	}
+	ConditionalTickAsync(Time >= 0.0 ? Time : FPlatformTime::Seconds() - GStartTime);
 
 	if (bSerializeAsJson.load(std::memory_order_relaxed))
 	{
@@ -69,8 +66,7 @@ void FOutputDeviceDebug::Serialize(const TCHAR* Data, ELogVerbosity::Type Verbos
 
 void FOutputDeviceDebug::SerializeRecord(const UE::FLogRecord& Record)
 {
-	const double Time = FPlatformTime::ToSeconds64(Record.GetTime().GetCycles());
-	ConditionalTickAsync(Time);
+	ConditionalTickAsync(FPlatformTime::Seconds() - GStartTime);
 
 	if (bSerializeAsJson.load(std::memory_order_relaxed))
 	{
@@ -120,7 +116,7 @@ void FOutputDeviceDebug::SerializeRecord(const UE::FLogRecord& Record)
 	else
 	{
 		TStringBuilder<512> Line;
-		FOutputDeviceHelper::AppendFormatLogLine(Line, Record.GetVerbosity(), Record.GetCategory(), nullptr, GPrintLogTimes, Time);
+		FOutputDeviceHelper::AppendFormatLogLine(Line, Record.GetVerbosity(), Record.GetCategory(), nullptr, GPrintLogTimes);
 		Record.FormatMessageTo(Line);
 		Line.Append(LINE_TERMINATOR);
 		FPlatformMisc::LowLevelOutputDebugString(*Line);
diff --git a/Engine/Source/Runtime/Core/Private/Misc/OutputDeviceRedirector.cpp b/Engine/Source/Runtime/Core/Private/Misc/OutputDeviceRedirector.cpp
--- a/Engine/Source/Runtime/Core/Private/Misc/OutputDeviceRedirector.cpp
+++ b/Engine/Source/Runtime/Core/Private/Misc/OutputDeviceRedirector.cpp
@@ -505,7 +505,7 @@ void FOutputDeviceRedirectorState::AddToBacklog(const FLogRecord& Record)
 {
 	TStringBuilder<512> Text;
 	Record.FormatMessageTo(Text);
-	AddToBacklog(*Text, Record.GetVerbosity(), Record.GetCategory(), FPlatformTime::ToSeconds64(Record.GetTime().GetCycles()));
+	AddToBacklog(*Text, Record.GetVerbosity(), Record.GetCategory(), FPlatformTime::Seconds() - GStartTime);
 }
 
 void FOutputDeviceRedirectorState::AddToBacklog(
