Engine: Memory: Unified the behavior of FPlatformMemory::BaseAlllocator() between platforms. This function is allowed to be called multiple times now. It caches the created allocator the first time it's called and returns it in subsequent runs.

#rnx

diff --git a/Engine/Platforms/Hololens/Source/Runtime/Core/Private/HoloLensPlatformMemory.cpp b/Engine/Platforms/Hololens/Source/Runtime/Core/Private/HoloLensPlatformMemory.cpp
--- a/Engine/Platforms/Hololens/Source/Runtime/Core/Private/HoloLensPlatformMemory.cpp
+++ b/Engine/Platforms/Hololens/Source/Runtime/Core/Private/HoloLensPlatformMemory.cpp
@@ -67,12 +67,18 @@ void FHoloLensPlatformMemory::Init()
 
 FMalloc* FHoloLensPlatformMemory::BaseAllocator()
 {
+	static FMalloc* Instance = nullptr;
+	if (Instance != nullptr)
+	{
+		return Instance;
+	}
+
 #if FORCE_ANSI_ALLOCATOR
-	return new FMallocAnsi();
+	return (Instance = new FMallocAnsi());
 #elif (WITH_EDITORONLY_DATA || IS_PROGRAM) && TBBMALLOC_ENABLED
-	return new FMallocTBB();
+	return (Instance = new FMallocTBB());
 #else
-	return new FMallocBinned3();
+	return (Instance = new FMallocBinned3());
 #endif
 
 	//	_CrtSetAllocHook(HoloLensAllocHook); // Enable to track down windows allocs not handled by our wrapper
diff --git a/Engine/Source/Runtime/Core/Private/Android/AndroidPlatformMemory.cpp b/Engine/Source/Runtime/Core/Private/Android/AndroidPlatformMemory.cpp
--- a/Engine/Source/Runtime/Core/Private/Android/AndroidPlatformMemory.cpp
+++ b/Engine/Source/Runtime/Core/Private/Android/AndroidPlatformMemory.cpp
@@ -535,6 +535,12 @@ EPlatformMemorySizeBucket FAndroidPlatformMemory::GetMemorySizeBucket()
 
 FMalloc* FAndroidPlatformMemory::BaseAllocator()
 {
+	static FMalloc* Instance = nullptr;
+	if (Instance != nullptr)
+	{
+		return Instance;
+	}
+
 #if ENABLE_LOW_LEVEL_MEM_TRACKER
 	// make sure LLM is using UsedPhysical for program size, instead of Available-Free
 	FPlatformMemoryStats Stats = FAndroidPlatformMemory::GetStats();
@@ -542,7 +548,7 @@ FMalloc* FAndroidPlatformMemory::BaseAllocator()
 #endif
 
 #if RUNNING_WITH_ASAN
-	return new FMallocAnsi();
+	return (Instance = new FMallocAnsi());
 #endif
 
 	const bool bHeapProfilingSupported = AndroidHeapProfiling::Init();
@@ -550,15 +556,15 @@ FMalloc* FAndroidPlatformMemory::BaseAllocator()
 #if USE_MALLOC_BINNED3 && PLATFORM_ANDROID_ARM64
 	if (bHeapProfilingSupported)
 	{
-		return new FMallocProfilingProxy<FMallocBinned3>();
+		return (Instance = new FMallocProfilingProxy<FMallocBinned3>());
 	}
-	return new FMallocBinned3();
+	return (Instance = new FMallocBinned3());
 #elif USE_MALLOC_BINNED2
 	if (bHeapProfilingSupported)
 	{
-		return new FMallocProfilingProxy<FMallocBinned2>();
+		return (Instance = new FMallocProfilingProxy<FMallocBinned2>());
 	}
-	return new FMallocBinned2();
+	return (Instance = new FMallocBinned2());
 #else
 	const FPlatformMemoryConstants& MemoryConstants = FPlatformMemory::GetConstants();
 	// 1 << FMath::CeilLogTwo(MemoryConstants.TotalPhysical) should really be FMath::RoundUpToPowerOfTwo,
@@ -568,7 +574,7 @@ FMalloc* FAndroidPlatformMemory::BaseAllocator()
 
 	// todo: Verify MallocBinned2 on 32bit
 	// [RCL] 2017-03-06 FIXME: perhaps BinnedPageSize should be used here, but leaving this change to the Android platform owner.
-	return new FMallocBinned(MemoryConstants.PageSize, MemoryLimit);
+	return (Instance = new FMallocBinned(MemoryConstants.PageSize, MemoryLimit));
 #endif
 }
 
diff --git a/Engine/Source/Runtime/Core/Private/Apple/ApplePlatformMemory.cpp b/Engine/Source/Runtime/Core/Private/Apple/ApplePlatformMemory.cpp
--- a/Engine/Source/Runtime/Core/Private/Apple/ApplePlatformMemory.cpp
+++ b/Engine/Source/Runtime/Core/Private/Apple/ApplePlatformMemory.cpp
@@ -291,6 +291,12 @@ void FApplePlatformMemory::Init()
 
 FMalloc* FApplePlatformMemory::BaseAllocator()
 {
+	static FMalloc* Instance = nullptr;
+	if (Instance != nullptr)
+	{
+		return Instance;
+	}
+
 #if ENABLE_LOW_LEVEL_MEM_TRACKER
 	FPlatformMemoryStats MemStats = FApplePlatformMemory::GetStats();
 	FLowLevelMemTracker::Get().SetProgramSize(MemStats.UsedPhysical);
@@ -318,10 +324,12 @@ FMalloc* FApplePlatformMemory::BaseAllocator()
 	switch (AllocatorToUse)
 	{
 		case EMemoryAllocatorToUse::Ansi:
-			return new FMallocAnsi();
+			Instance = new FMallocAnsi();
+			break;
 
 		case EMemoryAllocatorToUse::Binned2:
-			return new FMallocBinned2();
+			Instance = new FMallocBinned2();
+			break;
 			
 		default:	// intentional fall-through
 		case EMemoryAllocatorToUse::Binned:
@@ -336,10 +344,12 @@ FMalloc* FApplePlatformMemory::BaseAllocator()
 			uint64 MemoryLimit = FMath::Min<uint64>( uint64(1) << FMath::CeilLogTwo((Stats.free_count + Stats.inactive_count) * GetConstants().PageSize), 0x100000000);
 			
 			// [RCL] 2017-03-06 FIXME: perhaps BinnedPageSize should be used here, but leaving this change to the Mac platform owner.
-			return new FMallocBinned((uint32)(GetConstants().PageSize&MAX_uint32), MemoryLimit);
+			Instance = new FMallocBinned((uint32)(GetConstants().PageSize&MAX_uint32), MemoryLimit);
+			break;
 		}
 	}
 	
+	return Instance;
 }
 
 FPlatformMemoryStats FApplePlatformMemory::GetStats()
diff --git a/Engine/Source/Runtime/Core/Private/GenericPlatform/GenericPlatformMemory.cpp b/Engine/Source/Runtime/Core/Private/GenericPlatform/GenericPlatformMemory.cpp
--- a/Engine/Source/Runtime/Core/Private/GenericPlatform/GenericPlatformMemory.cpp
+++ b/Engine/Source/Runtime/Core/Private/GenericPlatform/GenericPlatformMemory.cpp
@@ -267,7 +267,15 @@ void FGenericPlatformMemory::OnOutOfMemory(uint64 Size, uint32 Alignment)
 
 FMalloc* FGenericPlatformMemory::BaseAllocator()
 {
-	return new FMallocAnsi();
+	static FMalloc* Instance = nullptr;
+	if (Instance != nullptr)
+	{
+		return Instance;
+	}
+
+	Instance = new FMallocAnsi();
+
+	return Instance;
 }
 
 FPlatformMemoryStats FGenericPlatformMemory::GetStats()
diff --git a/Engine/Source/Runtime/Core/Private/Mac/MacPlatformMemory.cpp b/Engine/Source/Runtime/Core/Private/Mac/MacPlatformMemory.cpp
--- a/Engine/Source/Runtime/Core/Private/Mac/MacPlatformMemory.cpp
+++ b/Engine/Source/Runtime/Core/Private/Mac/MacPlatformMemory.cpp
@@ -64,6 +64,12 @@ static bool HasArg(const char *Arg)
 
 FMalloc* FMacPlatformMemory::BaseAllocator()
 {
+	static FMalloc* Instance = nullptr;
+	if (Instance != nullptr)
+	{
+		return Instance;
+	}
+
 	if (FORCE_ANSI_ALLOCATOR || IS_PROGRAM)
 	{
 		AllocatorToUse = EMemoryAllocatorToUse::Ansi;
@@ -109,28 +115,35 @@ FMalloc* FMacPlatformMemory::BaseAllocator()
 	switch (AllocatorToUse)
 	{
 	case EMemoryAllocatorToUse::Ansi:
-		return new FMallocAnsi();
+		Instance = new FMallocAnsi();
+		break;
 #if WITH_MALLOC_STOMP
 	case EMemoryAllocatorToUse::Stomp:
-		return new FMallocStomp();
+		Instance = new FMallocStomp();
+		break;
 #endif
 #if TBBMALLOC_ENABLED
 	case EMemoryAllocatorToUse::TBB:
-		return new FMallocTBB();
+		Instance = new FMallocTBB();
+		break;
 #endif
 #if MIMALLOC_ENABLED
 	case EMemoryAllocatorToUse::Mimalloc:
-		return new FMallocMimalloc();
+		Instance = new FMallocMimalloc();
+		break;
 #endif
 	case EMemoryAllocatorToUse::Binned2:
-		return new FMallocBinned2();
+		Instance = new FMallocBinned2();
+		break;
 
 	default:	// intentional fall-through
 	case EMemoryAllocatorToUse::Binned:
 		// [RCL] 2017-03-06 FIXME: perhaps BinnedPageSize should be used here, but leaving this change to the Mac platform owner.
-		return new FMallocBinned((uint32)(GetConstants().PageSize&MAX_uint32), 0x100000000);
+		Instance = new FMallocBinned((uint32)(GetConstants().PageSize&MAX_uint32), 0x100000000);
+		break;
 	}
 
+	return Instance;
 }
 
 FPlatformMemoryStats FMacPlatformMemory::GetStats()
diff --git a/Engine/Source/Runtime/Core/Private/Unix/UnixPlatformMemory.cpp b/Engine/Source/Runtime/Core/Private/Unix/UnixPlatformMemory.cpp
--- a/Engine/Source/Runtime/Core/Private/Unix/UnixPlatformMemory.cpp
+++ b/Engine/Source/Runtime/Core/Private/Unix/UnixPlatformMemory.cpp
@@ -132,6 +132,12 @@ bool FUnixPlatformMemory::HasForkPageProtectorEnabled()
 
 class FMalloc* FUnixPlatformMemory::BaseAllocator()
 {
+	static FMalloc* Allocator = nullptr;
+	if (Allocator != nullptr)
+	{
+		return Allocator;
+	}
+
 #if UE4_DO_ROOT_PRIVILEGE_CHECK && !IS_PROGRAM
 	// This function gets executed very early, way before main() (because global constructors will allocate memory).
 	// This makes it ideal, if unobvious, place for a root privilege check.
@@ -304,8 +310,6 @@ class FMalloc* FUnixPlatformMemory::BaseAllocator()
 		AllocatorToUse = EMemoryAllocatorToUse::Ansi;
 	}
 
-	FMalloc * Allocator = NULL;
-
 	switch (AllocatorToUse)
 	{
 	case EMemoryAllocatorToUse::Ansi:
diff --git a/Engine/Source/Runtime/Core/Private/Windows/WindowsPlatformMemory.cpp b/Engine/Source/Runtime/Core/Private/Windows/WindowsPlatformMemory.cpp
--- a/Engine/Source/Runtime/Core/Private/Windows/WindowsPlatformMemory.cpp
+++ b/Engine/Source/Runtime/Core/Private/Windows/WindowsPlatformMemory.cpp
@@ -114,6 +114,12 @@ void FWindowsPlatformMemory::Init()
 
 FMalloc* FWindowsPlatformMemory::BaseAllocator()
 {
+	static FMalloc* Instance = nullptr;
+	if (Instance != nullptr)
+	{
+		return Instance;
+	}
+
 #if ENABLE_WIN_ALLOC_TRACKING
 	// This allows tracking of allocations that don't happen within the engine's wrappers.
 	// This actually won't be compiled unless bDebugBuildsActuallyUseDebugCRT is set in the
@@ -212,29 +218,38 @@ FMalloc* FWindowsPlatformMemory::BaseAllocator()
 	switch (AllocatorToUse)
 	{
 	case EMemoryAllocatorToUse::Ansi:
-		return new FMallocAnsi();
+		Instance = new FMallocAnsi();
+		break;
 #if WITH_MALLOC_STOMP
 	case EMemoryAllocatorToUse::Stomp:
-		return new FMallocStomp();
+		Instance = new FMallocStomp();
+		break;
 #endif
 #if TBBMALLOC_ENABLED
 	case EMemoryAllocatorToUse::TBB:
-		return new FMallocTBB();
+		Instance = new FMallocTBB();
+		break;
 #endif
 #if MIMALLOC_ENABLED
 	case EMemoryAllocatorToUse::Mimalloc:
-		return new FMallocMimalloc();
+		Instance = new FMallocMimalloc();
+		break;
 #endif
 	case EMemoryAllocatorToUse::Binned2:
-		return new FMallocBinned2();
+		Instance = new FMallocBinned2();
+		break;
 #if PLATFORM_64BITS
 	case EMemoryAllocatorToUse::Binned3:
-		return new FMallocBinned3();
+		Instance = new FMallocBinned3();
+		break;
 #endif
 	default:	// intentional fall-through
 	case EMemoryAllocatorToUse::Binned:
-		return new FMallocBinned((uint32)(GetConstants().BinnedPageSize&MAX_uint32), (uint64)MAX_uint32 + 1);
+		Instance = new FMallocBinned((uint32)(GetConstants().BinnedPageSize&MAX_uint32), (uint64)MAX_uint32 + 1);
+		break;
 	}
+
+	return Instance;
 }
 
 FPlatformMemoryStats FWindowsPlatformMemory::GetStats()
