Nightly Low Level tests must run only slow and performance tests
Additional small quality of life fixes

diff --git a/Engine/Source/Developer/LowLevelTestsRunner/Private/TestListeners/EventListener.cpp b/Engine/Source/Developer/LowLevelTestsRunner/Private/TestListeners/EventListener.cpp
--- a/Engine/Source/Developer/LowLevelTestsRunner/Private/TestListeners/EventListener.cpp
+++ b/Engine/Source/Developer/LowLevelTestsRunner/Private/TestListeners/EventListener.cpp
@@ -45,7 +45,7 @@ public:
 			TStringBuilder<64> DurationText;
 			const int32 Hours = Duration.GetHours();
 			const int32 Minutes = Duration.GetMinutes();
-			const double Seconds = Duration.GetTotalSeconds();
+			const double Seconds = Duration.GetSeconds();
 			if (Hours > 0)
 			{
 				DurationText << Hours << TEXTVIEW("h ");
diff --git a/Engine/Source/Programs/AutomationTool/LowLevelTests/RunLowLevelTests.cs b/Engine/Source/Programs/AutomationTool/LowLevelTests/RunLowLevelTests.cs
--- a/Engine/Source/Programs/AutomationTool/LowLevelTests/RunLowLevelTests.cs
+++ b/Engine/Source/Programs/AutomationTool/LowLevelTests/RunLowLevelTests.cs
@@ -643,6 +643,10 @@ namespace LowLevelTests
 				{
 					CachedConfig.CommandLineParams.AddRawCommandline("--attach-to-debugger");
 				}
+				if (Environment.GetEnvironmentVariable("IsBuildMachine") == "1")
+				{
+					CachedConfig.CommandLineParams.AddRawCommandline("--buildmachine");
+				}
 				if (LowLevelTestsExtension != null)
 				{
 					string ExtraCmd = LowLevelTestsExtension.ExtraCommandLine(Platform, TestApp, BuildPath);
diff --git a/Engine/Source/Runtime/Core/Tests/Containers/ConcurrentQueuesTest.cpp b/Engine/Source/Runtime/Core/Tests/Containers/ConcurrentQueuesTest.cpp
--- a/Engine/Source/Runtime/Core/Tests/Containers/ConcurrentQueuesTest.cpp
+++ b/Engine/Source/Runtime/Core/Tests/Containers/ConcurrentQueuesTest.cpp
@@ -247,7 +247,7 @@ namespace UE { namespace ConcurrentQueuesTests
 		check(Produced == NumConcumed);
 	}
 
-	TEST_CASE_NAMED(FConcurrentQueuesTest, "System::Core::Async::ConcurrentQueuesTest", "[.][ApplicationContextMask][EngineFilter]")
+	TEST_CASE_NAMED(FConcurrentQueuesTest, "System::Core::Async::ConcurrentQueuesTest", "[.][ApplicationContextMask][EngineFilter][Perf]")
 	{
 		{	// test for support of not default constructible types
 			struct FNonDefaultConstructable
@@ -723,7 +723,7 @@ namespace ClosableMpscQueueTests
 		UE_LOG(LogTemp, Display, TEXT("items queued %d"), NumProduced);
 	}
 
-	TEST_CASE_NAMED(FClosableMpscQueueTest, "System::Core::Async::ClosableMpscQueueTest", "[.][ApplicationContextMask][EngineFilter]")
+	TEST_CASE_NAMED(FClosableMpscQueueTest, "System::Core::Async::ClosableMpscQueueTest", "[.][ApplicationContextMask][EngineFilter][Perf]")
 	{
 		{
 			TClosableMpscQueue<int> Q;
@@ -1046,7 +1046,7 @@ namespace DepletableMpscQueueTests
 		Queue.Deplete([](void*) { checkNoEntry(); }); // empty
 	}
 
-	TEST_CASE_NAMED(FDepletableMpscQueueTest, "System::Core::Async::DepletableMpscQueueTest", "[.][ApplicationContextMask][EngineFilter]")
+	TEST_CASE_NAMED(FDepletableMpscQueueTest, "System::Core::Async::DepletableMpscQueueTest", "[.][ApplicationContextMask][EngineFilter][Perf]")
 	{
 		{
 			TDepletableMpmcQueue<int> Q;
diff --git a/Engine/Source/Runtime/Core/Tests/Math/UnrealMathTest.cpp b/Engine/Source/Runtime/Core/Tests/Math/UnrealMathTest.cpp
--- a/Engine/Source/Runtime/Core/Tests/Math/UnrealMathTest.cpp
+++ b/Engine/Source/Runtime/Core/Tests/Math/UnrealMathTest.cpp
@@ -4696,7 +4696,7 @@ namespace UnrealMathTestInternal {
 
 }
 
-TEST_CASE_NAMED(FColorConversionHeavyTest, "System::Core::Math::ColorConversionHeavy", "[.][ApplicationContextMask][EngineFilter]")
+TEST_CASE_NAMED(FColorConversionHeavyTest, "System::Core::Math::ColorConversionHeavy", "[.][ApplicationContextMask][EngineFilter][Slow]")
 {
 	// WARNING: This test runs for a good while (at time of writing, ~90s when using
 	// a single core, proportionately less on many-core machines) with no user
