Engine: Memory: ConcurrentLinearAllocator: Fixed incorrect memory tracking.

There were two issues:

1) ConcurrentLinearAllocator uses different underlying allocators but they generally end up leading to FMalloc. While we do have tagging in place, it didn't include ELLMAllocType::FMalloc, without which the memory stats system cannot correctly match these allocations in Default and Platform trackers. This leads to incorrect FMallocUnused stat value being calculated, which later on contributes to Untracked memory.

2) Memory was accounted for twice: (1) by FMemory::Malloc used by the underlying allocator providing memory pages, and then by (2) our custom tracking (tracking individual memory objects returned to the allocator user). It wasn't a critical issue as ConcurrentLinearAllocator provides memory that is transient in nature, but stats were incorrect gathered which could lead to artifacts. Resolved by using FPlatformMemory::Malloc instead, which bypasses the tracking layer in FMemory and goes directly to the system allocator such as MallocBinned3.

Updating OutputDeviceRedirector to make sure that all systems using the ConcurrentLinearAllocator get memory using FPlatformMemory::BaseAllocator()->Malloc() instead of FMemory::Malloc() to avoid the double accounting problem mentioned above. Most systems either use FDefaultBlockAllocationTag or a custom tag with FOsAllocator, both of which are defined in ConcurrentLinearAllocator.h and updated in this change. OutputDeviceRedirector is the only use case defining its own 'tag' with it's own sub-allocator (FOutputDeviceBlockAllocationTag::Allocator), so it had to be updated separately to be consistent.

diff --git a/Engine/Source/Runtime/Core/Private/Misc/OutputDeviceRedirector.cpp b/Engine/Source/Runtime/Core/Private/Misc/OutputDeviceRedirector.cpp
--- a/Engine/Source/Runtime/Core/Private/Misc/OutputDeviceRedirector.cpp
+++ b/Engine/Source/Runtime/Core/Private/Misc/OutputDeviceRedirector.cpp
@@ -48,15 +48,16 @@ struct FOutputDeviceBlockAllocationTag : FDefaultBlockAllocationTag
 	struct Allocator
 	{
 		static constexpr bool SupportsAlignment = false;
+		static constexpr bool UsesFMalloc       = true;
 
 		FORCEINLINE static void* Malloc(SIZE_T Size, uint32 Alignment)
 		{
-			return FMemory::Malloc(Size, DEFAULT_ALIGNMENT);
+			return FPlatformMemory::BaseAllocator()->Malloc(Size, DEFAULT_ALIGNMENT);
 		}
 
 		FORCEINLINE static void Free(void* Pointer, SIZE_T Size)
 		{
-			return FMemory::Free(Pointer);
+			return FPlatformMemory::BaseAllocator()->Free(Pointer);
 		}
 	};
 };
diff --git a/Engine/Source/Runtime/Core/Public/Experimental/ConcurrentLinearAllocator.h b/Engine/Source/Runtime/Core/Public/Experimental/ConcurrentLinearAllocator.h
--- a/Engine/Source/Runtime/Core/Public/Experimental/ConcurrentLinearAllocator.h
+++ b/Engine/Source/Runtime/Core/Public/Experimental/ConcurrentLinearAllocator.h
@@ -43,10 +43,16 @@ struct FOsAllocator
 {
 	static constexpr bool SupportsAlignment = false;
 
+	// Indicates whether memory is acquired from the same source as used by FMalloc i.e. FMemory/MPlatformMemory::BaseAllocator().
+	// It's needed so that ConcurrentLinearAllocator could correctly account for this memory in the memory stats and tracking systems.
+	// It is 'true' by default, which is suitable for most implementations. However, if you create a custom allocator that
+	// acquires memory in a different way e.g. directly from the OS using FPlatformMemory::BinnedAllosFromOS, then you
+	// should set this variable to 'false'.
+	static constexpr bool UsesFMalloc = true;
+
 	FORCEINLINE static void* Malloc(SIZE_T Size, uint32 Alignment)
 	{
-		//LLM_SCOPED_PAUSE_TRACKING(ELLMAllocType::System);
-		void* Pointer = FMemory::Malloc(Size);
+		void* Pointer = FPlatformMemory::BaseAllocator()->Malloc(Size);
 		MemoryTrace_Alloc(uint64(Pointer), Size, Alignment);
 		MemoryTrace_MarkAllocAsHeap(uint64(Pointer), EMemoryTraceRootHeap::SystemMemory);
 		return Pointer;
@@ -55,7 +61,7 @@ struct FOsAllocator
 	FORCEINLINE static void Free(void* Pointer, SIZE_T Size)
 	{
 		MemoryTrace_Free(uint64(Pointer));
-		FMemory::Free(Pointer);
+		FPlatformMemory::BaseAllocator()->Free(Pointer);
 	}
 };
 
@@ -63,6 +69,7 @@ struct FOsAllocator
 struct FAlignedAllocatorInternal
 {
 	static constexpr bool SupportsAlignment = true;
+	static constexpr bool UsesFMalloc       = true;
 
 	FORCEINLINE static void* Malloc(SIZE_T Size, uint32 Alignment)
 	{
@@ -113,6 +120,7 @@ class TBlockAllocationCache
 
 public:
 	static constexpr bool SupportsAlignment = Allocator::SupportsAlignment;
+	static constexpr bool UsesFMalloc       = Allocator::UsesFMalloc;
 
 	FORCEINLINE static void* Malloc(SIZE_T Size, uint32 Alignment)
 	{
@@ -147,6 +155,7 @@ class TBlockAllocationLockFreeCache
 public:
 	static_assert(BlockSize == FPageAllocator::PageSize, "Only 64k pages are supported with this cache.");
 	static constexpr bool SupportsAlignment = Allocator::SupportsAlignment;
+	static constexpr bool UsesFMalloc       = Allocator::UsesFMalloc;
 
 	FORCEINLINE static void* Malloc(SIZE_T Size, uint32 Alignment)
 	{
@@ -200,6 +209,8 @@ class TConcurrentLinearAllocator
 		&& !IS_ASAN_ENABLED && !BlockAllocationTag::RequiresAccurateSize // Only enabled when not using ASAN or when the Size of an allocation is not required 
 		&& BlockAllocationTag::Allocator::SupportsAlignment; //The allocator needs to align the BlockHeader by Blocksize to find the BlockHeader
 
+	LLM(static constexpr ELLMAllocType LLMAllocationType = BlockAllocationTag::Allocator::UsesFMalloc ? ELLMAllocType::FMalloc : ELLMAllocType::None;)
+
 	struct FBlockHeader;
 
 	class FAllocationHeader
@@ -274,7 +285,7 @@ class TConcurrentLinearAllocator
 					//if all allocations are already freed we can reuse the Block again
 					Header->~FBlockHeader();
 					ASAN_UNPOISON_MEMORY_REGION( Header, BlockAllocationTag::BlockSize );
-					BlockAllocationTag::Allocator::Free(Header, BlockAllocationTag::BlockSize);				
+					BlockAllocationTag::Allocator::Free(Header, BlockAllocationTag::BlockSize);
 				}
 			}
 		}
@@ -344,7 +355,7 @@ public:
 				Header->Num++;
 
 				MemoryTrace_Alloc(uint64(AlignedOffset), Size, Alignment);
-				LLM_IF_ENABLED( FLowLevelMemTracker::Get().OnLowLevelAlloc(ELLMTracker::Default, reinterpret_cast<void*>(AlignedOffset), Size, ELLMTag::LinearAllocator) );
+				LLM_IF_ENABLED( FLowLevelMemTracker::Get().OnLowLevelAlloc(ELLMTracker::Default, reinterpret_cast<void*>(AlignedOffset), Size, ELLMTag::LinearAllocator, LLMAllocationType) );
 				return reinterpret_cast<void*>(AlignedOffset);
 			}
 
@@ -365,7 +376,7 @@ public:
 					checkSlow(LargeAlignedOffset + Size <= LargeHeader->NextAllocationPtr);
 
 					MemoryTrace_Alloc(uint64(LargeAlignedOffset), Size, Alignment);
-					LLM_IF_ENABLED( FLowLevelMemTracker::Get().OnLowLevelAlloc(ELLMTracker::Default, reinterpret_cast<void*>(LargeAlignedOffset), Size, ELLMTag::LinearAllocator) );
+					LLM_IF_ENABLED( FLowLevelMemTracker::Get().OnLowLevelAlloc(ELLMTracker::Default, reinterpret_cast<void*>(LargeAlignedOffset), Size, ELLMTag::LinearAllocator, LLMAllocationType) );
 					return reinterpret_cast<void*>(LargeAlignedOffset);
 				}
 			}
@@ -385,7 +396,7 @@ public:
 				ASAN_POISON_MEMORY_REGION( AllocationHeader, sizeof(FAllocationHeader) );
 
 				MemoryTrace_Alloc(uint64(AllocationHeader), Size + sizeof(FAllocationHeader), Alignment);
-				LLM_IF_ENABLED( FLowLevelMemTracker::Get().OnLowLevelAlloc(ELLMTracker::Default, AllocationHeader, Size + sizeof(FAllocationHeader), ELLMTag::LinearAllocator) );
+				LLM_IF_ENABLED( FLowLevelMemTracker::Get().OnLowLevelAlloc(ELLMTracker::Default, AllocationHeader, Size + sizeof(FAllocationHeader), ELLMTag::LinearAllocator, LLMAllocationType) );
 				return reinterpret_cast<void*>(AlignedOffset);
 			}
 
@@ -408,7 +419,7 @@ public:
 					ASAN_POISON_MEMORY_REGION( AllocationHeader, sizeof(FAllocationHeader) );
 
 					MemoryTrace_Alloc(uint64(AllocationHeader), Size + sizeof(FAllocationHeader), Alignment);
-					LLM_IF_ENABLED( FLowLevelMemTracker::Get().OnLowLevelAlloc(ELLMTracker::Default, AllocationHeader, Size + sizeof(FAllocationHeader), ELLMTag::LinearAllocator) );
+					LLM_IF_ENABLED( FLowLevelMemTracker::Get().OnLowLevelAlloc(ELLMTracker::Default, AllocationHeader, Size + sizeof(FAllocationHeader), ELLMTag::LinearAllocator, LLMAllocationType) );
 					return reinterpret_cast<void*>(LargeAlignedOffset);
 				}
 			}
@@ -440,7 +451,7 @@ public:
 			if constexpr (SupportsFastPath)
 			{
 				MemoryTrace_Free(uint64(Pointer));
-				LLM_IF_ENABLED(FLowLevelMemTracker::Get().OnLowLevelFree(ELLMTracker::Default, Pointer));
+				LLM_IF_ENABLED(FLowLevelMemTracker::Get().OnLowLevelFree(ELLMTracker::Default, Pointer, LLMAllocationType));
 				FBlockHeader* Header = reinterpret_cast<FBlockHeader*>(AlignDown(Pointer, BlockAllocationTag::BlockSize));
 
 				// this deletes complete blocks when the last allocation is freed
@@ -455,7 +466,7 @@ public:
 				FAllocationHeader* AllocationHeader = GetAllocationHeader(Pointer);
 				ASAN_UNPOISON_MEMORY_REGION( AllocationHeader, sizeof(FAllocationHeader) );
 				MemoryTrace_Free(uint64(AllocationHeader));
-				LLM_IF_ENABLED(FLowLevelMemTracker::Get().OnLowLevelFree(ELLMTracker::Default, AllocationHeader));
+				LLM_IF_ENABLED(FLowLevelMemTracker::Get().OnLowLevelFree(ELLMTracker::Default, AllocationHeader, LLMAllocationType));
 				FBlockHeader* Header = AllocationHeader->GetBlockHeader();
 				ASAN_POISON_MEMORY_REGION( AllocationHeader, sizeof(FAllocationHeader) + AllocationHeader->GetAllocationSize() );
 
