AutoRTFM runtime initial commit

diff --git a/Engine/Source/Runtime/Core/Private/AutoRTFM/API.cpp b/Engine/Source/Runtime/Core/Private/AutoRTFM/API.cpp
new file mode 100644
--- /dev/null
+++ b/Engine/Source/Runtime/Core/Private/AutoRTFM/API.cpp
@@ -0,0 +1,371 @@
+// Copyright Epic Games, Inc. All Rights Reserved.
+
+#if (defined(__AUTORTFM) && __AUTORTFM)
+
+#include "AutoRTFM/AutoRTFM.h"
+#include "AutoRTFM/AutoRTFMConstants.h"
+#include "Context.h"
+#include "ContextInlines.h"
+#include "Debug.h"
+#include "FunctionMapInlines.h"
+#include "TransactionInlines.h"
+#include "Utils.h"
+
+// This is the implementation of the AutoRTFM.h API. Ideally, functions here should just delegate to some internal API.
+// For now, I have these functions also perform some error checking.
+
+namespace AutoRTFM
+{
+
+extern "C" bool autortfm_is_transactional()
+{
+    return FContext::Get()->GetStatus() == EContextStatus::OnTrack;
+}
+
+extern "C" bool autortfm_is_closed()
+{
+    return false;
+}
+
+// First Part - the API exposed outside transactions.
+extern "C" autortfm_result autortfm_transact(void (*Work)(void* Arg), void* Arg)
+{
+    return static_cast<autortfm_result>(FContext::Get()->Transact(Work, Arg));
+}
+
+extern "C" void autortfm_commit(void (*Work)(void* Arg), void* Arg)
+{
+    autortfm_result Result = autortfm_transact(Work, Arg);
+    if (Result != autortfm_committed)
+    {
+        fprintf(stderr, "FATAL: Unexpected transaction result: %u.\n", Result);
+        abort();
+    }
+}
+
+extern "C" void autortfm_abort()
+{
+    if (FContext::IsTransactional())
+    {
+        FContext::Get()->AbortByRequestAndThrow();
+    }
+    else
+    {
+        fprintf(stderr, "autortfm_abort called from outside a transaction.\n");
+        abort();
+    }
+}
+
+extern "C" void autortfm_abort_if_transactional()
+{
+    if (FContext::IsTransactional())
+    {
+        fprintf(stderr, "autortfm_abort_if_transactional called from an open nest inside a transaction.\n");
+        abort();
+    }
+}
+
+extern "C" void autortfm_abort_if_closed()
+{
+}
+
+extern "C" void autortfm_open(void (*Work)(void* Arg), void* Arg)
+{
+    Work(Arg);
+}
+
+extern "C" void autortfm_close(void (*Work)(void* Arg), void* Arg)
+{
+    if (FContext::IsTransactional())
+    {
+        FContext* Context = FContext::Get();
+        void (*WorkClone)(void* Arg, FContext* Context) = FunctionMapLookup(Work, Context, "autortfm_close");
+        if (WorkClone)
+        {
+            WorkClone(Arg, Context);
+        }
+    }
+    else
+    {
+        fprintf(stderr, "autortfm_close called from outside a transaction.\n");
+        abort();
+    }
+}
+
+extern "C" void autortfm_register_open_function(void* OriginalFunction, void* NewFunction)
+{
+    constexpr bool bVerbose = false;
+    if (bVerbose)
+    {
+        fprintf(GetLogFile(), "Registering open %p->%p\n", OriginalFunction, NewFunction);
+    }
+    FunctionMapAdd(OriginalFunction, NewFunction);
+}
+
+void DeferUntilCommit(std::function<void()>&& Work)
+{
+    if (FContext::IsTransactional())
+    {
+        FContext::Get()->GetCurrentTransaction()->DeferUntilCommit(std::move(Work));
+    }
+    else
+    {
+        Work();
+    }
+}
+
+void DeferUntilAbort(std::function<void()>&& Work)
+{
+    if (FContext::IsTransactional())
+    {
+        FContext::Get()->GetCurrentTransaction()->DeferUntilAbort(std::move(Work));
+    }
+}
+
+void OpenCommit(std::function<void()>&& Work)
+{
+    Work();
+}
+
+void OpenAbort(std::function<void()>&& Work)
+{
+}
+
+extern "C" void autortfm_defer_until_commit(void (*Work)(void* Arg), void* Arg)
+{
+    DeferUntilCommit([Work, Arg] () { Work(Arg); });
+}
+
+extern "C" void autortfm_defer_until_abort(void (*Work)(void* arg), void* Arg)
+{
+    DeferUntilAbort([Work, Arg] () { Work(Arg); });
+}
+
+extern "C" void autortfm_open_commit(void (*Work)(void* Arg), void* Arg)
+{
+    Work(Arg);
+}
+
+extern "C" void autortfm_open_abort(void (*Work)(void* arg), void* Arg)
+{
+}
+
+extern "C" void* autortfm_did_allocate(void* Ptr, size_t Size)
+{
+    return Ptr;
+}
+
+extern "C" void autortfm_will_deallocate(void* Ptr, size_t Size)
+{
+}
+
+extern "C" void autortfm_check_consistency_assuming_no_races()
+{
+    if (FContext::IsTransactional())
+    {
+        AutoRTFM::Unreachable();
+    }
+}
+
+extern "C" void autortfm_check_abi(void* const Ptr, const size_t Size)
+{
+    struct FConstants final
+    {
+        const size_t LogLineBytes = Constants::LogLineBytes;
+        const size_t LineBytes = Constants::LineBytes;
+        const size_t LineTableSize = Constants::LineTableSize;
+        const size_t Offset_Context_CurrentTransaction = Constants::Offset_Context_CurrentTransaction;
+        const size_t Offset_Context_LineTable = Constants::Offset_Context_LineTable;
+        const size_t Offset_Context_Status = Constants::Offset_Context_Status;
+        const size_t LogSize_LineEntry = Constants::LogSize_LineEntry;
+        const size_t Size_LineEntry = Constants::Size_LineEntry;
+        const size_t Offset_LineEntry_LogicalLine = Constants::Offset_LineEntry_LogicalLine;
+        const size_t Offset_LineEntry_ActiveLine = Constants::Offset_LineEntry_ActiveLine;
+        const size_t Offset_LineEntry_LoggingTransaction = Constants::Offset_LineEntry_LoggingTransaction;
+        const size_t Offset_LineEntry_AccessMask = Constants::Offset_LineEntry_AccessMask;
+        const uint32_t Context_Status_OnTrack = Constants::Context_Status_OnTrack;
+
+		// This is messy - but we want to do comparisons but without comparing any padding bytes.
+		// Before C++20 we cannot use a default created operator== and operator!=, so we use this
+		// ugly trick to just compare the members.
+	private:
+		auto Tied() const
+		{
+			return std::tie(LogLineBytes, LineBytes, LineTableSize, Offset_Context_CurrentTransaction, Offset_Context_LineTable, Offset_Context_Status, LogSize_LineEntry, Size_LineEntry, Offset_LineEntry_LogicalLine, Offset_LineEntry_ActiveLine, Offset_LineEntry_LoggingTransaction, Offset_LineEntry_AccessMask, Context_Status_OnTrack);
+		}
+
+	public:
+		bool operator==(const FConstants& Other) const
+		{
+			return Tied() == Other.Tied();
+		}
+
+		bool operator!=(const FConstants& Other) const
+		{
+			return !(*this == Other);
+		}
+    } RuntimeConstants;
+
+    if (sizeof(FConstants) != Size)
+    {
+        fprintf(GetLogFile(), "Fatal: found ABI error between AutoRTFM compiler and runtime\n");
+        abort();
+    }
+
+    const FConstants* const CompilerConstants = static_cast<FConstants*>(Ptr);
+
+    if (RuntimeConstants != *CompilerConstants)
+    {
+        fprintf(GetLogFile(), "Fatal: found ABI error between AutoRTFM compiler and runtime\n");
+        abort();
+    }
+}
+
+// Second Part - the same API exposed inside transactions. Note that we don't expose all of the API
+// to transactions! That's intentional. However, things like autortfm_defer_until_commit can be called
+// from an open nest in a transaction.
+bool STM_autortfm_is_transactional(FContext* Context)
+{
+    return true;
+}
+UE_AUTORTFM_REGISTER_OPEN_FUNCTION(autortfm_is_transactional);
+
+bool STM_autortfm_is_closed(FContext* Context)
+{
+    return true;
+}
+UE_AUTORTFM_REGISTER_OPEN_FUNCTION(autortfm_is_closed);
+
+autortfm_result STM_autortfm_transact(void (*Work)(void* Arg), void* Arg, FContext* Context)
+{
+    return static_cast<autortfm_result>(Context->Transact(Work, Arg));
+}
+UE_AUTORTFM_REGISTER_OPEN_FUNCTION(autortfm_transact);
+
+void STM_autortfm_commit(void (*Work)(void* Arg), void* Arg, FContext*)
+{
+    autortfm_result Result = autortfm_transact(Work, Arg);
+    if (Result != autortfm_committed)
+    {
+        fprintf(stderr, "FATAL: Unexpected transaction result: %u.\n", Result);
+        abort();
+    }
+}
+UE_AUTORTFM_REGISTER_OPEN_FUNCTION(autortfm_commit);
+
+void STM_autortfm_abort(FContext* Context)
+{
+    Context->AbortByRequestAndThrow();
+}
+UE_AUTORTFM_REGISTER_OPEN_FUNCTION(autortfm_abort);
+
+void STM_autortfm_abort_if_transactional(FContext* Context)
+{
+    if (FDebug::bVerbose)
+    {
+        fprintf(stderr, "autortfm_abort_if_transactional called from inside a transaction.\n");
+    }
+    Context->AbortByLanguageAndThrow();
+}
+UE_AUTORTFM_REGISTER_OPEN_FUNCTION(autortfm_abort_if_transactional);
+
+void STM_autortfm_abort_if_closed(FContext* Context)
+{
+    if (FDebug::bVerbose)
+    {
+        fprintf(stderr, "autortfm_abort_if_closed called from a closed nest inside a transaction.\n");
+    }
+    
+    Context->AbortByLanguageAndThrow();
+}
+UE_AUTORTFM_REGISTER_OPEN_FUNCTION(autortfm_abort_if_closed);
+
+void STM_autortfm_open(void (*Work)(void* Arg), void* Arg, FContext* Context)
+{
+    Work(Arg);
+}
+UE_AUTORTFM_REGISTER_OPEN_FUNCTION(autortfm_open);
+
+void STM_autortfm_close(void (*Work)(void* Arg), void* Arg, FContext* Context)
+{
+    void (*WorkClone)(void* Arg, FContext* Context) = FunctionMapLookup(Work, Context, "STM_autortfm_close");
+    if (WorkClone)
+    {
+        WorkClone(Arg, Context);
+    }
+}
+UE_AUTORTFM_REGISTER_OPEN_FUNCTION(autortfm_close);
+
+void STM_DeferUntilCommit(std::function<void()>&& Work, FContext* Context)
+{
+    ASSERT(Context->GetStatus() == EContextStatus::OnTrack);
+    Context->GetCurrentTransaction()->DeferUntilCommit(std::move(Work));
+}
+UE_AUTORTFM_REGISTER_OPEN_FUNCTION(DeferUntilCommit);
+
+void STM_DeferUntilAbort(std::function<void()>&& Work, FContext* Context)
+{
+    ASSERT(Context->GetStatus() == EContextStatus::OnTrack);
+    Context->GetCurrentTransaction()->DeferUntilAbort(std::move(Work));
+}
+UE_AUTORTFM_REGISTER_OPEN_FUNCTION(DeferUntilAbort);
+
+void STM_OpenCommit(std::function<void()>&& Work, FContext* Context)
+{
+    ASSERT(Context->GetStatus() == EContextStatus::OnTrack);
+    Context->GetCurrentTransaction()->DeferUntilCommit(std::move(Work));
+}
+UE_AUTORTFM_REGISTER_OPEN_FUNCTION(OpenCommit);
+
+void STM_OpenAbort(std::function<void()>&& Work, FContext* Context)
+{
+    ASSERT(Context->GetStatus() == EContextStatus::OnTrack);
+    Context->GetCurrentTransaction()->DeferUntilAbort(std::move(Work));
+}
+UE_AUTORTFM_REGISTER_OPEN_FUNCTION(OpenAbort);
+
+extern "C" void STM_autortfm_defer_until_commit(void (*Work)(void* Arg), void* Arg, FContext* Context)
+{
+    STM_DeferUntilCommit([Work, Arg] { Work(Arg); }, Context);
+}
+UE_AUTORTFM_REGISTER_OPEN_FUNCTION(autortfm_defer_until_commit);
+
+extern "C" void STM_autortfm_defer_until_abort(void (*Work)(void* arg), void* Arg, FContext* Context)
+{
+    STM_DeferUntilAbort([Work, Arg] { Work(Arg); }, Context);
+}
+UE_AUTORTFM_REGISTER_OPEN_FUNCTION(autortfm_defer_until_abort);
+
+extern "C" void STM_autortfm_open_commit(void (*Work)(void* Arg), void* Arg, FContext* Context)
+{
+    STM_OpenCommit([Work, Arg] { Work(Arg); }, Context);
+}
+UE_AUTORTFM_REGISTER_OPEN_FUNCTION(autortfm_open_commit);
+
+extern "C" void STM_autortfm_open_abort(void (*Work)(void* arg), void* Arg, FContext* Context)
+{
+    STM_OpenAbort([Work, Arg] { Work(Arg); }, Context);
+}
+UE_AUTORTFM_REGISTER_OPEN_FUNCTION(autortfm_open_abort);
+
+void* STM_autortfm_did_allocate(void* Ptr, size_t Size, FContext* Context)
+{
+    Context->DidAllocate(Ptr, Size);
+    return Ptr;
+}
+UE_AUTORTFM_REGISTER_OPEN_FUNCTION(autortfm_did_allocate);
+
+void STM_autortfm_will_deallocate(void* Ptr, size_t Size, FContext* Context)
+{
+    Context->WillDeallocate(Ptr, Size);
+}
+UE_AUTORTFM_REGISTER_OPEN_FUNCTION(autortfm_will_deallocate);
+
+void STM_autortfm_check_consistency_assuming_no_races(FContext* Context)
+{
+}
+UE_AUTORTFM_REGISTER_OPEN_FUNCTION(autortfm_check_consistency_assuming_no_races);
+
+} // namespace AutoRTFM
+
+#endif // defined(__AUTORTFM) && __AUTORTFM
diff --git a/Engine/Source/Runtime/Core/Private/AutoRTFM/Atomic.h b/Engine/Source/Runtime/Core/Private/AutoRTFM/Atomic.h
new file mode 100644
--- /dev/null
+++ b/Engine/Source/Runtime/Core/Private/AutoRTFM/Atomic.h
@@ -0,0 +1,76 @@
+// Copyright Epic Games, Inc. All Rights Reserved.
+
+#pragma once
+
+#include <atomic>
+
+namespace AutoRTFM
+{
+
+inline void Fence()
+{
+    atomic_thread_fence(std::memory_order_seq_cst);
+}
+
+inline void CompilerFence()
+{
+    atomic_signal_fence(std::memory_order_seq_cst);
+}
+
+template<typename T, typename U, typename V>
+bool WeakUnfencedCAS(T& Ref, U PassedExpectedValue, V PassedNewValue)
+{
+    T ExpectedValue = PassedExpectedValue;
+    T NewValue = PassedNewValue;
+    std::atomic<T>& AtomicRef = reinterpret_cast<std::atomic<T>&>(Ref);
+    return AtomicRef.compare_exchange_weak(ExpectedValue, NewValue, std::memory_order_relaxed);
+}
+
+template<typename T, typename U, typename V>
+bool WeakFencedCAS(T& Ref, U PassedExpectedValue, V PassedNewValue)
+{
+    T ExpectedValue = PassedExpectedValue;
+    T NewValue = PassedNewValue;
+    std::atomic<T>& AtomicRef = reinterpret_cast<std::atomic<T>&>(Ref);
+    return AtomicRef.compare_exchange_weak(ExpectedValue, NewValue);
+}
+
+template<typename T, typename U, typename V>
+T StrongFencedCAS(T& Ref, U PassedExpectedValue, V PassedNewValue)
+{
+    T ExpectedValue = PassedExpectedValue;
+    T NewValue = PassedNewValue;
+    std::atomic<T>& AtomicRef = reinterpret_cast<std::atomic<T>&>(Ref);
+    AtomicRef.compare_exchange_weak(ExpectedValue, NewValue);
+    return ExpectedValue;
+}
+
+template<typename T>
+T FencedLoad(T& Ref)
+{
+    std::atomic<T>& AtomicRef = reinterpret_cast<std::atomic<T>&>(Ref);
+    return AtomicRef;
+}
+
+template<typename T>
+T LoadAcquire(T& Ref)
+{
+    std::atomic<T>& AtomicRef = reinterpret_cast<std::atomic<T>&>(Ref);
+    return AtomicRef.load(std::memory_order_acquire);
+}
+
+template<typename T>
+void FencedStore(T& Ref, T Value)
+{
+    std::atomic<T>& AtomicRef = reinterpret_cast<std::atomic<T>&>(Ref);
+    AtomicRef = Value;
+}
+
+class FFenceGuard
+{
+public:
+    FFenceGuard() { Fence(); }
+    ~FFenceGuard() { Fence(); }
+};
+
+} // namespace AutoRTFM
diff --git a/Engine/Source/Runtime/Core/Private/AutoRTFM/CompilerInterface.cpp b/Engine/Source/Runtime/Core/Private/AutoRTFM/CompilerInterface.cpp
new file mode 100644
--- /dev/null
+++ b/Engine/Source/Runtime/Core/Private/AutoRTFM/CompilerInterface.cpp
@@ -0,0 +1,94 @@
+// Copyright Epic Games, Inc. All Rights Reserved.
+
+#if (defined(__AUTORTFM) && __AUTORTFM)
+
+#include "ContextInlines.h"
+#include "Debug.h"
+#include "FunctionMapInlines.h"
+#include "Memcpy.h"
+
+namespace AutoRTFM
+{
+
+void AbortDueToBadAlignment(FContext* Context, void* Ptr, size_t Alignment, const char* Message = nullptr)
+{
+    Context->DumpState();
+    fprintf(stderr, "Aborting because alignment error: expected alignment %zu, got pointer %p.\n", Alignment, Ptr);
+    if (Message)
+    {
+        fprintf(stderr, "%s\n", Message);
+    }
+    abort();
+}
+
+void CheckAlignment(FContext* Context, void* Ptr, size_t AlignmentMask)
+{
+    if (reinterpret_cast<uintptr_t>(Ptr) & AlignmentMask)
+    {
+        AbortDueToBadAlignment(Context, Ptr, AlignmentMask + 1);
+    }
+}
+
+extern "C" void autortfm_record_write(FContext* Context, void* Ptr, size_t Size)
+{
+    Context->RecordWrite(Ptr, Size);
+}
+
+extern "C" void* autortfm_lookup_function(FContext* Context, void* OriginalFunction, const char* Where)
+{
+    FDebug Debug(Context, OriginalFunction, nullptr, 0, 0, __FUNCTION__);
+    return FunctionMapLookup(OriginalFunction, Context, Where);
+}
+
+extern "C" void autortfm_memcpy(void* Dst, const void* Src, size_t Size, FContext* Context)
+{
+    FDebug Debug(Context, Dst, Src, Size, 0, __FUNCTION__);
+    Memcpy(Dst, Src, Size, Context);
+}
+
+extern "C" void autortfm_memmove(void* Dst, const void* Src, size_t Size, FContext* Context)
+{
+    FDebug Debug(Context, Dst, Src, Size, 0, __FUNCTION__);
+    Memmove(Dst, Src, Size, Context);
+}
+
+extern "C" void autortfm_memset(void* Dst, int Value, size_t Size, FContext* Context)
+{
+    FDebug Debug(Context, Dst, nullptr, Size, 0, __FUNCTION__);
+    Memset(Dst, Value, Size, Context);
+}
+
+extern "C" void autortfm_llvm_fail(FContext* Context, const char* Message)
+{
+    if (Message)
+    {
+        fprintf(stderr, "Transaction failing because of language issue:\n%s\n", Message);
+    }
+    else
+    {
+        fprintf(stderr, "Transaction failing because of language issue.\n");
+    }
+    Context->AbortByLanguageAndThrow();
+}
+
+extern "C" void autortfm_llvm_alignment_error(FContext* Context, void* Ptr, size_t Alignment, const char* Message)
+{
+    AbortDueToBadAlignment(Context, Ptr, Alignment, Message);
+}
+
+extern "C" void autortfm_llvm_error(FContext* Context, const char* Message)
+{
+    if (Message)
+    {
+        fprintf(stderr, "Aborting because LLVM error:\n%s\n", Message);
+    }
+    else
+    {
+        fprintf(stderr, "Aborting because LLVM error.\n");
+    }
+    abort();
+}
+
+} // namespace AutoRTFM
+
+#endif // defined(__AUTORTFM) && __AUTORTFM
diff --git a/Engine/Source/Runtime/Core/Private/AutoRTFM/Context.cpp b/Engine/Source/Runtime/Core/Private/AutoRTFM/Context.cpp
new file mode 100644
--- /dev/null
+++ b/Engine/Source/Runtime/Core/Private/AutoRTFM/Context.cpp
@@ -0,0 +1,246 @@
+// Copyright Epic Games, Inc. All Rights Reserved.
+
+#if (defined(__AUTORTFM) && __AUTORTFM)
+#include "Context.h"
+#include "ContextInlines.h"
+#include "FunctionMap.h"
+#include "GlobalData.h"
+#include "ScopedGuard.h"
+#include "TransactionInlines.h"
+#include <algorithm>
+#include <memory>
+
+namespace AutoRTFM
+{
+
+thread_local std::unique_ptr<FContext> ContextTls;
+
+void FContext::InitializeGlobalData()
+{
+}
+
+FContext* FContext::TryGet()
+{
+    return ContextTls.get();
+}
+
+void FContext::Set()
+{
+    ContextTls.reset(this);
+}
+
+FContext* FContext::Get()
+{
+    FContext* Result = TryGet();
+
+    if (!Result)
+    {
+        Result = new FContext();
+        Result->Set();
+    }
+
+    return Result;
+}
+
+bool FContext::IsTransactional()
+{
+    FContext* Context = TryGet();
+    if (!Context)
+    {
+        return false;
+    }
+
+    if (Context->GetStatus() == EContextStatus::OnTrack)
+    {
+        return true;
+    }
+    else
+    {
+        ASSERT(Context->GetStatus() == EContextStatus::Idle);
+        return false;
+    }
+}
+
+ETransactionResult FContext::Transact(void (*Function)(void* Arg), void* Arg)
+{
+    constexpr bool bVerbose = false;
+    
+    ASSERT(Status == EContextStatus::Idle || Status == EContextStatus::OnTrack);
+
+    void (*ClonedFunction)(void* Arg, FContext* Context) = FunctionMapTryLookup(Function);
+    if (!ClonedFunction)
+    {
+        fprintf(stderr, "Could not find function %p (%s) in AutoRTFM::FContext::Transact.\n", Function, GetFunctionDescription(Function).c_str());
+        return ETransactionResult::AbortedByLanguage;
+    }
+    
+    std::unique_ptr<FTransaction> NewTransactionUniquePtr(new FTransaction(this));
+    FTransaction* NewTransaction = NewTransactionUniquePtr.get();
+
+    void* TransactStackAddress = &NewTransactionUniquePtr;
+    ASSERT(TransactStackAddress > StackBegin);
+    ASSERT(TransactStackAddress < StackEnd);
+    TScopedGuard<void*> CurrentTransactStackAddressGuard(CurrentTransactStackAddress, TransactStackAddress);
+
+    if (!CurrentTransaction)
+    {
+        ASSERT(Status == EContextStatus::Idle);
+        CurrentTransaction = NewTransaction;
+        OuterTransactStackAddress = TransactStackAddress;
+        ETransactionResult Result = ETransactionResult::Committed; // Initialize to something to make the compiler happy.
+        for (;;)
+        {
+            Status = EContextStatus::OnTrack;
+            ASSERT(CurrentTransaction->IsFresh());
+            CurrentTransaction->Try([&] () { ClonedFunction(Arg, this); });
+            ASSERT(Status != EContextStatus::Idle);
+            if (Status == EContextStatus::OnTrack)
+            {
+                if (bVerbose)
+                {
+                    fprintf(GetLogFile(), "About to commit; my state is:\n");
+                    DumpState();
+                    fprintf(GetLogFile(), "Committing...\n");
+                }
+
+                if (CurrentTransaction->AttemptToCommit())
+                {
+                    Result = ETransactionResult::Committed;
+                    break;
+                }
+
+                if (bVerbose)
+                {
+                    fprintf(GetLogFile(), "Commit failed!\n");
+                }
+
+                ASSERT(Status != EContextStatus::OnTrack);
+                ASSERT(Status != EContextStatus::Idle);
+            }
+            if (Status == EContextStatus::AbortedByRequest)
+            {
+                Result = ETransactionResult::AbortedByRequest;
+                break;
+            }
+            if (Status == EContextStatus::AbortedByLanguage)
+            {
+                Result = ETransactionResult::AbortedByLanguage;
+                break;
+            }
+            ASSERT(Status == EContextStatus::AbortedByLineInitialization
+                   || Status == EContextStatus::AbortedByCommitTimeLineCheck
+                   || Status == EContextStatus::AbortedByFailedLockAcquisition);
+
+            if (bVerbose)
+            {
+                fprintf(GetLogFile(), "About to prelock some locks!\n");
+            }
+            ASSERT(!bLocksToHoldAreHeld);
+            
+            bLocksToHoldAreHeld = true;
+        }
+
+        Reset();
+        // No need to SetIsDone or whatever since all of the transactions have now been blown away.
+        return Result;
+    }
+    else
+    {
+        ASSERT(Status == EContextStatus::OnTrack);
+        FTransaction* PreviousTransaction = CurrentTransaction;
+        CurrentTransaction = NewTransaction;
+        CurrentTransaction->Try([&] () { ClonedFunction(Arg, this); });
+
+        // We just use this bit to help assertions for now (though we could use it more strongly). Because of how we use this right now,
+        // it's OK that it's set before we commit but after we abort.
+        NewTransaction->SetIsDone();
+        
+        if (Status == EContextStatus::OnTrack)
+        {
+            bool bCommitResult = NewTransaction->AttemptToCommit();
+            CurrentTransaction = PreviousTransaction;
+            ASSERT(bCommitResult);
+            ASSERT(Status == EContextStatus::OnTrack);
+            return ETransactionResult::Committed;
+        }
+        
+        CurrentTransaction = PreviousTransaction;
+
+        switch (Status)
+        {
+        case EContextStatus::AbortedByLineInitialization:
+            CurrentTransaction->AbortAndThrow();
+            AutoRTFM::Unreachable();
+        case EContextStatus::AbortedByRequest:
+            Status = EContextStatus::OnTrack;
+            return ETransactionResult::AbortedByRequest;
+        case EContextStatus::AbortedByLanguage:
+            Status = EContextStatus::OnTrack;
+            return ETransactionResult::AbortedByLanguage;
+        default:
+            AutoRTFM::Unreachable();
+        }
+    }
+}
+
+void FContext::AbortByRequestAndThrow()
+{
+    ASSERT(Status == EContextStatus::OnTrack);
+    Status = EContextStatus::AbortedByRequest;
+    CurrentTransaction->AbortAndThrow();
+}
+
+void FContext::AbortByLanguageAndThrow()
+{
+    constexpr bool bAbortProgram = false;
+    ASSERT(Status == EContextStatus::OnTrack);
+    if (bAbortProgram)
+    {
+        fprintf(stderr, "FATAL: Unexpected language abort.\n");
+        abort();
+    }
+    Status = EContextStatus::AbortedByLanguage;
+    CurrentTransaction->AbortAndThrow();
+}
+
+FContext::FContext()
+{
+#if _WIN32
+#if 0
+    GetCurrentThreadStackLimits(reinterpret_cast<PULONG_PTR>(&StackBegin),
+                                reinterpret_cast<PULONG_PTR>(&StackEnd));
+#endif
+#else
+    pthread_attr_t Attr;
+    pthread_getattr_np(pthread_self(), &Attr);
+    size_t StackSize;
+    pthread_attr_getstack(&Attr, &StackBegin, &StackSize);
+    StackEnd = static_cast<char*>(StackBegin) + StackSize;
+#endif
+    ASSERT(StackEnd > StackBegin);
+}
+
+void FContext::Reset()
+{
+    OuterTransactStackAddress = nullptr;
+    CurrentTransactStackAddress = nullptr;
+    CurrentTransaction = nullptr;
+    Status = EContextStatus::Idle;
+}
+
+void FContext::StaticAsserts()
+{
+    /*static_assert(PrettyStaticAssert<Constants::Offset_Context_CurrentTransaction, offsetof(FContext, CurrentTransaction)>::_cResult, "Not equal");
+    static_assert(PrettyStaticAssert<Constants::Offset_Context_LineTable, offsetof(FContext, LineTable)>::_cResult, "Not equal");
+    static_assert(PrettyStaticAssert<Constants::Offset_Context_Status, offsetof(FContext, Status)>::_cResult, "Not equal");*/
+}
+
+void FContext::DumpState() const
+{
+    fprintf(GetLogFile(), "Context at %p.\n", this);
+    fprintf(GetLogFile(), "Transaction stack: %p...%p\n", StackBegin, OuterTransactStackAddress);
+}
+
+} // namespace AutoRTFM
+
+#endif // defined(__AUTORTFM) && __AUTORTFM
diff --git a/Engine/Source/Runtime/Core/Private/AutoRTFM/Context.h b/Engine/Source/Runtime/Core/Private/AutoRTFM/Context.h
new file mode 100644
--- /dev/null
+++ b/Engine/Source/Runtime/Core/Private/AutoRTFM/Context.h
@@ -0,0 +1,64 @@
+// Copyright Epic Games, Inc. All Rights Reserved.
+
+#pragma once
+
+#include "AutoRTFM/AutoRTFM.h"
+#include "ContextStatus.h"
+#include <memory>
+
+namespace AutoRTFM
+{
+
+class FLineLock;
+class FTransaction;
+
+class FContext
+{
+public:
+    static FContext* TryGet();
+    static FContext* Get();
+    static bool IsTransactional();
+    
+    // This is public API
+    ETransactionResult Transact(void (*Function)(void* Arg), void* Arg);
+    void AbortByRequestAndThrow();
+
+    // Record that a write is about to occur at the given LogicalAddress of Size bytes.
+    void RecordWrite(void* LogicalAddress, size_t Size);
+
+    void DidAllocate(void* LogicalAddress, size_t Size);
+    void WillDeallocate(void* LogicalAddress, size_t Size);
+
+    // The rest of this is internalish.
+    void AbortByLanguageAndThrow();
+
+    FTransaction* GetCurrentTransaction() const { return CurrentTransaction; }
+    bool IsTransactionStack(void* LogicalAddress) const { return LogicalAddress >= StackBegin && LogicalAddress < OuterTransactStackAddress; }
+    bool IsInnerTransactionStack(void* LogicalAddress) const { return LogicalAddress >= StackBegin && LogicalAddress < CurrentTransactStackAddress; }
+    EContextStatus GetStatus() const { return Status; }
+
+    void DumpState() const;
+
+    static void InitializeGlobalData();
+
+private:
+    FContext();
+    FContext(const FContext&) = delete;
+
+    void Set();
+    
+    // All of this other stuff ought to be private?
+    void Reset();
+
+    void StaticAsserts();
+    
+    FTransaction* CurrentTransaction{nullptr};
+    void* StackBegin{nullptr}; // begin as in the smaller of the two
+    void* StackEnd{nullptr};
+    void* OuterTransactStackAddress{nullptr};
+    void* CurrentTransactStackAddress{nullptr};
+    bool bLocksToHoldAreHeld{false};
+    EContextStatus Status{EContextStatus::Idle};
+};
+
+} // namespace AutoRTFM
diff --git a/Engine/Source/Runtime/Core/Private/AutoRTFM/ContextInlines.h b/Engine/Source/Runtime/Core/Private/AutoRTFM/ContextInlines.h
new file mode 100644
--- /dev/null
+++ b/Engine/Source/Runtime/Core/Private/AutoRTFM/ContextInlines.h
@@ -0,0 +1,25 @@
+// Copyright Epic Games, Inc. All Rights Reserved.
+
+#pragma once
+
+#include "Context.h"
+#include "TransactionInlines.h"
+
+namespace AutoRTFM
+{
+
+inline void FContext::RecordWrite(void* LogicalAddress, size_t Size)
+{
+    CurrentTransaction->RecordWrite(LogicalAddress, Size);
+}
+
+inline void FContext::DidAllocate(void* LogicalAddress, size_t Size)
+{
+    CurrentTransaction->DidAllocate(LogicalAddress, Size);
+}
+
+inline void FContext::WillDeallocate(void* LogicalAddress, size_t Size)
+{
+}
+
+} // namespace AutoRTFM
diff --git a/Engine/Source/Runtime/Core/Private/AutoRTFM/ContextStatus.h b/Engine/Source/Runtime/Core/Private/AutoRTFM/ContextStatus.h
new file mode 100644
--- /dev/null
+++ b/Engine/Source/Runtime/Core/Private/AutoRTFM/ContextStatus.h
@@ -0,0 +1,47 @@
+// Copyright Epic Games, Inc. All Rights Reserved.
+
+#pragma once
+
+#include "AutoRTFM/AutoRTFMConstants.h"
+#include "Utils.h"
+
+namespace AutoRTFM
+{
+
+enum class EContextStatus
+{
+    Idle,
+    OnTrack,
+    AbortedByLineInitialization,
+    AbortedByCommitTimeLineCheck,
+    AbortedByFailedLockAcquisition,
+    AbortedByLanguage,
+    AbortedByRequest
+};
+
+static_assert(static_cast<unsigned>(EContextStatus::OnTrack) == Constants::Context_Status_OnTrack, "Not equal");
+
+inline const char* GetContextStatusName(EContextStatus Status)
+{
+    switch (Status)
+    {
+    case EContextStatus::Idle:
+        return "Idle";
+    case EContextStatus::OnTrack:
+        return "OnTrack";
+    case EContextStatus::AbortedByLineInitialization:
+        return "AbortedByLineInitialization";
+    case EContextStatus::AbortedByCommitTimeLineCheck:
+        return "AbortedByCommitTimeLineCheck";
+    case EContextStatus::AbortedByFailedLockAcquisition:
+        return "AbortedByFailedLockAcquisition";
+    case EContextStatus::AbortedByLanguage:
+        return "AbortedByLanguage";
+    case EContextStatus::AbortedByRequest:
+        return "AbortedByRequest";
+    }
+    ASSERT(!"Should not be reached");
+    return NULL;
+}
+
+} // namespace AutoRTFM
diff --git a/Engine/Source/Runtime/Core/Private/AutoRTFM/CppNewSupport.cpp b/Engine/Source/Runtime/Core/Private/AutoRTFM/CppNewSupport.cpp
new file mode 100644
--- /dev/null
+++ b/Engine/Source/Runtime/Core/Private/AutoRTFM/CppNewSupport.cpp
@@ -0,0 +1,82 @@
+// Copyright Epic Games, Inc. All Rights Reserved.
+
+#if (defined(__AUTORTFM) && __AUTORTFM)
+#include "CppNewSupport.h"
+#include "ContextInlines.h"
+#include "Debug.h"
+#include "FunctionMap.h"
+#include "TransactionInlines.h"
+#include "Utils.h"
+
+namespace AutoRTFM
+{
+	using FCppNew = void* (*)(size_t);
+	using FCppDeleteWithSize = void(*)(void*, size_t);
+	using FCppDelete = void(*)(void*);
+
+	FCppNew CppNew;
+	FCppNew CppNewWithSize;
+	FCppDeleteWithSize CppDeleteWithSize;
+	FCppDelete CppDelete;
+
+// These are also exposed as compiler ABI, currently just for the benefit of Windows.
+extern "C" void* autortfm_cpp_new(size_t Size, FContext* Context)
+{
+    FDebug Debug(Context, nullptr, nullptr, Size, 0, __FUNCTION__);
+    void* Result = CppNew(Size);
+
+    Context->GetCurrentTransaction()->DeferUntilAbort([Result]
+    {
+        CppDelete(Result);
+    });
+    Context->DidAllocate(Result, Size);
+    return Result;
+}
+
+extern "C" void autortfm_cpp_delete(void* Ptr, FContext* Context)
+{
+    FDebug Debug(Context, Ptr, nullptr, 0, 0, __FUNCTION__);
+    if (Ptr)
+    {
+		const size_t AllocSize = GetAllocationSize(Ptr);
+        Context->WillDeallocate(Ptr, AllocSize);
+        Context->GetCurrentTransaction()->DeferUntilCommit([Ptr]
+        {
+            CppDelete(Ptr);
+        });
+    }
+}
+extern "C" void autortfm_cpp_delete_with_size(void* Ptr, size_t Size, FContext* Context)
+{
+    autortfm_cpp_delete(Ptr, Context);
+}
+
+#ifdef _WIN32
+#endif // _WIN32
+
+struct FRegisterCppNewDelete
+{
+    FRegisterCppNewDelete()
+    {
+#ifdef _WIN32
+		CppNew = &operator new;
+		CppNewWithSize = &operator new[];
+		CppDeleteWithSize = &operator delete[];
+		CppDelete = &operator delete;
+
+        FunctionMapAdd(reinterpret_cast<void*>(CppNew), reinterpret_cast<void*>(autortfm_cpp_new));
+        FunctionMapAdd(reinterpret_cast<void*>(CppNewWithSize), reinterpret_cast<void*>(autortfm_cpp_new));
+        FunctionMapAdd(reinterpret_cast<void*>(CppDeleteWithSize), reinterpret_cast<void*>(autortfm_cpp_delete_with_size));
+        FunctionMapAdd(reinterpret_cast<void*>(CppDelete), reinterpret_cast<void*>(autortfm_cpp_delete));
+#else
+        FunctionMapAdd(reinterpret_cast<void*>(_Znwm), reinterpret_cast<void*>(autortfm_cpp_new));
+        FunctionMapAdd(reinterpret_cast<void*>(_Znam), reinterpret_cast<void*>(autortfm_cpp_new));
+        FunctionMapAdd(reinterpret_cast<void*>(_ZdlPv), reinterpret_cast<void*>(autortfm_cpp_delete));
+#endif
+    }
+};
+const FRegisterCppNewDelete RegisterCppNewDelete;
+
+} // namespace AutoRTFM
+
+#endif // defined(__AUTORTFM) && __AUTORTFM
diff --git a/Engine/Source/Runtime/Core/Private/AutoRTFM/CppNewSupport.h b/Engine/Source/Runtime/Core/Private/AutoRTFM/CppNewSupport.h
new file mode 100644
--- /dev/null
+++ b/Engine/Source/Runtime/Core/Private/AutoRTFM/CppNewSupport.h
@@ -0,0 +1,30 @@
+// Copyright Epic Games, Inc. All Rights Reserved.
+
+#pragma once
+
+#include "AutoRTFM/AutoRTFM.h"
+
+namespace AutoRTFM
+{
+
+#ifndef _WIN32
+extern "C" void* _Znwm(size_t Size); // C++ new
+extern "C" void* _Znam(size_t Size); // C++ array new
+extern "C" void _ZdlPv(void* Ptr); // C++ delete
+
+inline void* CppNew(size_t Size)
+{
+    return _Znwm(Size);
+}
+inline void CppDelete(void* Ptr)
+{
+    _ZdlPv(Ptr);
+}
+inline void CppDeleteWithSize(void* Ptr, size_t Size)
+{
+    CppDelete(Ptr);
+}
+#endif
+
+}
+
diff --git a/Engine/Source/Runtime/Core/Private/AutoRTFM/Debug.h b/Engine/Source/Runtime/Core/Private/AutoRTFM/Debug.h
new file mode 100644
--- /dev/null
+++ b/Engine/Source/Runtime/Core/Private/AutoRTFM/Debug.h
@@ -0,0 +1,41 @@
+// Copyright Epic Games, Inc. All Rights Reserved.
+
+#pragma once
+
+#include "Context.h"
+#include "Utils.h"
+
+namespace AutoRTFM
+{
+
+class FDebug final
+{
+public:
+	static constexpr bool bVerbose = false;
+    static constexpr bool bCheck = false;
+    static constexpr bool bCheckOnCommit = false;
+
+    FDebug(FContext* Context, const void* Dst, const void* Src, size_t Size, size_t Align, const char* Action)
+        : Dst(Dst), Src(Src), Size(Size), Align(Align), Action(Action), Context(Context)
+    {
+        if (bVerbose)
+        {
+            fprintf(GetLogFile(), "Compiler: %s [dst %p | src %p | size %zu | align %zu]\n", Action, Dst, Src, Size, Align);
+        }
+    }
+
+    ~FDebug()
+    {
+    }
+
+private:
+    const void* const Dst;
+    const void* const Src;
+    const size_t Size;
+    const size_t Align;
+    const char* const Action;
+    FContext* const Context;
+};
+
+} // namespace AutoRTFM
+
diff --git a/Engine/Source/Runtime/Core/Private/AutoRTFM/FastLock.h b/Engine/Source/Runtime/Core/Private/AutoRTFM/FastLock.h
new file mode 100644
--- /dev/null
+++ b/Engine/Source/Runtime/Core/Private/AutoRTFM/FastLock.h
@@ -0,0 +1,35 @@
+// Copyright Epic Games, Inc. All Rights Reserved.
+
+#pragma once
+
+#include <atomic>
+
+namespace AutoRTFM
+{
+
+// This should eventually be a ParkingLot-based lock.
+class FFastLock
+{
+public:
+    FFastLock() = default;
+
+    void Lock()
+    {
+        while (bIsHeld.exchange(true));
+    }
+
+    bool TryLock()
+    {
+        return !bIsHeld.exchange(true);
+    }
+
+    void Unlock()
+    {
+        bIsHeld.store(false);
+    }
+
+private:
+    std::atomic<bool> bIsHeld{false};
+};
+
+} // namespace AutoRTFM
diff --git a/Engine/Source/Runtime/Core/Private/AutoRTFM/FunctionMap.cpp b/Engine/Source/Runtime/Core/Private/AutoRTFM/FunctionMap.cpp
new file mode 100644
--- /dev/null
+++ b/Engine/Source/Runtime/Core/Private/AutoRTFM/FunctionMap.cpp
@@ -0,0 +1,106 @@
+// Copyright Epic Games, Inc. All Rights Reserved.
+
+#if (defined(__AUTORTFM) && __AUTORTFM)
+#include "FunctionMap.h"
+#include "Atomic.h"
+#include "FastLock.h"
+#include "LockGuard.h"
+#include "Utils.h"
+#include <stdlib.h>
+#include <string.h>
+
+namespace AutoRTFM
+{
+
+FFunctionMap* AllocateFunctionMap(size_t Size)
+{
+    ASSERT(Size && !(Size & (Size - 1)));
+    size_t TotalSize = offsetof(FFunctionMap, Entries) + sizeof(FFunctionEntry) * Size;
+    FFunctionMap* Result = static_cast<FFunctionMap*>(malloc(TotalSize));
+    Result->KeyCount = 0;
+    Result->Size = Size;
+    Result->IndexMask = Size - 1;
+    memset(Result->Entries, 0, sizeof(FFunctionEntry) * Size);
+    Fence(); // Really we just need a store-store fence...
+    return Result;
+}
+
+void FunctionMapAddImpl(void* OldFunction, void* NewFunction)
+{
+    FFunctionMap* Map = GlobalData->FunctionMap;
+    if (!Map)
+    {
+        Map = AllocateFunctionMap(1024);
+        GlobalData->FunctionMap = Map;
+    }
+    else if (Map->KeyCount * 2 >= Map->Size)
+    {
+        size_t NewSize = Map->Size * 2;
+        FFunctionMap* NewMap = AllocateFunctionMap(NewSize);
+        for (size_t OldIndex = 0; OldIndex < Map->Size; OldIndex++)
+        {
+            FFunctionEntry* OldEntry = Map->Entries + OldIndex;
+            if (!OldEntry->OldFunction)
+            {
+                continue;
+            }
+            // FIXME: Could have a secondary stronger hash here that is used on the second iteration
+            // followed by +1 for the subsequent iterations.
+            for (size_t Hash = FunctionPtrHash(OldEntry->OldFunction); ; ++Hash)
+            {
+                size_t NewIndex = Hash & NewMap->IndexMask;
+                FFunctionEntry* NewEntry = NewMap->Entries + NewIndex;
+                if (!NewEntry->OldFunction)
+                {
+                    *NewEntry = *OldEntry;
+                    break;
+                }
+            }
+        }
+        NewMap->KeyCount = Map->KeyCount;
+        GlobalData->FunctionMap = NewMap;
+        Map = NewMap;
+    }
+
+    for (size_t Hash = FunctionPtrHash(OldFunction); ; ++Hash) {
+        size_t Index = Hash & Map->IndexMask;
+        FFunctionEntry* Entry = Map->Entries + Index;
+
+        if (!Entry->OldFunction) {
+            Entry->OldFunction = OldFunction;
+            Entry->NewFunction = NewFunction;
+            Map->KeyCount++;
+            return;
+        }
+        
+        if (Entry->OldFunction == OldFunction) {
+            Entry->NewFunction = NewFunction;
+            return;
+        }
+    }
+}
+
+void FunctionMapAdd(void* OldFunction, void* NewFunction)
+{
+    if (!OldFunction)
+    {
+        // Silently ignore null OldFunction to make weak linking easy.
+        return;
+    }
+
+    InitializeGlobalDataIfNecessary();
+    
+    TLockGuard<FFastLock> LockGuard(GlobalData->FunctionMapLock);
+    std::vector<void*> Functions;
+    Functions.push_back(OldFunction);
+    while (!Functions.empty())
+    {
+        void* Function = Functions.back();
+        Functions.pop_back();
+        FunctionMapAddImpl(Function, NewFunction);
+    }
+}
+
+} // namespace AutoRTFM
+
+#endif // defined(__AUTORTFM) && __AUTORTFM
diff --git a/Engine/Source/Runtime/Core/Private/AutoRTFM/FunctionMap.h b/Engine/Source/Runtime/Core/Private/AutoRTFM/FunctionMap.h
new file mode 100644
--- /dev/null
+++ b/Engine/Source/Runtime/Core/Private/AutoRTFM/FunctionMap.h
@@ -0,0 +1,71 @@
+// Copyright Epic Games, Inc. All Rights Reserved.
+
+#pragma once
+
+#include "GlobalData.h"
+#include <cstddef>
+
+namespace AutoRTFM
+{
+
+class FContext;
+
+// The compiler has to agree to this.
+inline uintptr_t FunctionPtrHash(void* functionPtr)
+{
+    // FIXME: We'll need something better.
+    return reinterpret_cast<uintptr_t>(functionPtr) / 4;
+}
+
+struct FFunctionEntry
+{
+    FFunctionEntry() = default;
+
+    FFunctionEntry(void* OldFunction, void* NewFunction)
+        : OldFunction(OldFunction), NewFunction(NewFunction) { }
+    
+    void* OldFunction{nullptr};
+    void* NewFunction{nullptr};
+};
+
+struct FFunctionMap {
+    size_t KeyCount;
+    size_t Size;
+    size_t IndexMask;
+    FFunctionEntry Entries[1];
+};
+
+void FunctionMapAdd(void* OldFunction, void* NewFuncton);
+
+inline void* FunctionMapTryLookup(void* OldFunction)
+{
+    FFunctionMap* Map = GlobalData->FunctionMap;
+
+    if (!Map)
+        return nullptr;
+
+    for (size_t Hash = FunctionPtrHash(OldFunction); ; ++Hash) {
+        size_t Index = Hash & Map->IndexMask;
+        FFunctionEntry* Entry = Map->Entries + Index;
+
+        if (Entry->OldFunction == OldFunction)
+        {
+            // This races with Add, but in a benign way. If _OldFunction is set and matches our OldFunction but the _NewFunction
+            // is nullptr, then it means that an Add is happening right now, and before the Add, this entry would have been nullptr.
+            return Entry->NewFunction;
+        }
+
+        if (!Entry->OldFunction)
+        {
+            return nullptr;
+        }
+    }
+}
+
+template<typename TReturnType, typename... TParameterTypes>
+auto FunctionMapTryLookup(TReturnType (*Function)(TParameterTypes...)) -> TReturnType (*)(TParameterTypes..., FContext*)
+{
+    return reinterpret_cast<TReturnType (*)(TParameterTypes..., FContext*)>(FunctionMapTryLookup(reinterpret_cast<void*>(Function)));
+}
+
+} // namespace AutoRTFM
diff --git a/Engine/Source/Runtime/Core/Private/AutoRTFM/FunctionMapInlines.h b/Engine/Source/Runtime/Core/Private/AutoRTFM/FunctionMapInlines.h
new file mode 100644
--- /dev/null
+++ b/Engine/Source/Runtime/Core/Private/AutoRTFM/FunctionMapInlines.h
@@ -0,0 +1,33 @@
+// Copyright Epic Games, Inc. All Rights Reserved.
+
+#pragma once
+
+#include "Context.h"
+#include "FunctionMap.h"
+
+namespace AutoRTFM
+{
+
+inline void* FunctionMapLookup(void* OldFunction, FContext* Context, const char* Where)
+{
+    void* Result = FunctionMapTryLookup(OldFunction);
+    if (!Result)
+    {
+        fprintf(stderr, "Could not find function %p (%s)\n", OldFunction, GetFunctionDescription(OldFunction).c_str());
+        if (Where)
+        {
+            fprintf(stderr, "%s\n", Where);
+        }
+        Context->AbortByLanguageAndThrow();
+    }
+    return Result;
+}
+
+template<typename TReturnType, typename... TParameterTypes>
+auto FunctionMapLookup(TReturnType (*Function)(TParameterTypes...), FContext* Context, const char* Where) -> TReturnType (*)(TParameterTypes..., FContext*)
+{
+    return reinterpret_cast<TReturnType (*)(TParameterTypes..., FContext*)>(FunctionMapLookup(reinterpret_cast<void*>(Function), Context, Where));
+}
+
+} // namespace AutoRTFM
+
diff --git a/Engine/Source/Runtime/Core/Private/AutoRTFM/GlobalData.cpp b/Engine/Source/Runtime/Core/Private/AutoRTFM/GlobalData.cpp
new file mode 100644
--- /dev/null
+++ b/Engine/Source/Runtime/Core/Private/AutoRTFM/GlobalData.cpp
@@ -0,0 +1,72 @@
+// Copyright Epic Games, Inc. All Rights Reserved.
+
+#if (defined(__AUTORTFM) && __AUTORTFM)
+#include "GlobalData.h"
+#include "Context.h"
+#include <mutex>
+
+namespace AutoRTFM
+{
+
+std::once_flag GlobalDataInitializationFlag;
+FGlobalData* GlobalData;
+
+void InitializeGlobalDataIfNecessary()
+{
+    std::call_once(GlobalDataInitializationFlag, []
+    {
+        // AutoRTFM is intended to be used primarily from monolithic UE binaries, where having
+        // some additional DLL/so/dylib would just be annoying. However, we guard against the
+        // possibility that these things somehow get linked together. In that case, the only
+        // bad thing you get is code bloat, but otherwise everything works out:
+        //
+        // - Each AutoRTFM runtime instance registers functions for whatever version of the
+        //   standard library it sees.
+        //
+        // - All AutoRTFM instances coordinate together on things like the function table,
+        //   lock table, and TLS key.
+        
+        constexpr const char* EnvName = "AutoRTFMGlobalData";
+
+#ifdef _MSC_VER
+/*
+   Disable warning about deprecated STD C functions.
+*/
+#pragma warning(disable : 4996)
+
+#pragma warning(push)
+#endif
+
+        if (char* EnvString = getenv(EnvName))
+        {
+            sscanf(EnvString, "%p", &GlobalData);
+            ASSERT(GlobalData);
+        }
+        else
+        {
+            GlobalData = new FGlobalData;
+            FContext::InitializeGlobalData();
+            char Buffer[128];
+            snprintf(Buffer, sizeof(Buffer), "%s=%p", EnvName, GlobalData);
+            putenv(Buffer);
+        }
+
+#ifdef _MSC_VER
+#pragma warning(pop)
+#endif
+    });
+}
+
+struct FInitializeGlobalData
+{
+    FInitializeGlobalData()
+    {
+        InitializeGlobalDataIfNecessary();
+    }
+};
+
+FInitializeGlobalData InitializeGlobalData;
+
+} // namespace AutoRTFM
+
+#endif // defined(__AUTORTFM) && __AUTORTFM
diff --git a/Engine/Source/Runtime/Core/Private/AutoRTFM/GlobalData.h b/Engine/Source/Runtime/Core/Private/AutoRTFM/GlobalData.h
new file mode 100644
--- /dev/null
+++ b/Engine/Source/Runtime/Core/Private/AutoRTFM/GlobalData.h
@@ -0,0 +1,25 @@
+// Copyright Epic Games, Inc. All Rights Reserved.
+
+#pragma once
+
+#include "FastLock.h"
+#include <unordered_map>
+#include <vector>
+
+namespace AutoRTFM
+{
+struct FFunctionMap;
+
+struct FGlobalData
+{
+    FFunctionMap* FunctionMap{nullptr};
+    FFastLock FunctionMapLock;
+};
+
+extern FGlobalData* GlobalData;
+
+// It's only necessary to call this in code that might be called from static initializers.
+void InitializeGlobalDataIfNecessary();
+
+} // namespace AutoRTFM
+
diff --git a/Engine/Source/Runtime/Core/Private/AutoRTFM/HitSet.h b/Engine/Source/Runtime/Core/Private/AutoRTFM/HitSet.h
new file mode 100644
--- /dev/null
+++ b/Engine/Source/Runtime/Core/Private/AutoRTFM/HitSet.h
@@ -0,0 +1,194 @@
+// Copyright Epic Games, Inc. All Rights Reserved.
+
+#pragma once
+
+#include "TaggedPtr.h"
+#include "Utils.h"
+
+namespace AutoRTFM
+{
+	class FHitSet final
+	{
+		// TODO: Revisit a good probe depth for the hashset.
+		static constexpr uint32_t LinearProbeDepth = 21;
+
+		// TODO: Revisit a good initial capacity for the hitset.
+		static constexpr uint32_t InitialCapacity = 128;
+
+	public:
+		using Key = TTaggedPtr<void>;
+
+		explicit FHitSet()
+		{
+			// Do not want to deal with null payloads.
+			ASSERT(0 != InitialCapacity);
+
+			// The capacity increases by a power of two so that the modulus is optimal.
+			ASSERT(0 == (InitialCapacity & (InitialCapacity - 1)));
+
+			// To make the linear probe check simpler, our payload has a padded
+			// linear probe depth so the last element doesn't have to wrap
+			// around and insert at the front.
+			Payload = new uintptr_t[InitialCapacity + LinearProbeDepth]();
+			ASSERT(nullptr != Payload);
+
+			Capacity = InitialCapacity;
+			Size = 0;
+		}
+
+		~FHitSet()
+		{
+			delete[] Payload;
+		}
+
+		// Insert something in the HitSet, returning true if the put succeeded
+		// (EG. the key was not already in the set).
+		bool Insert(Key K)
+		{
+			while (true)
+			{
+				switch (InsertNoResize(K.GetPayload()))
+				{
+				default:
+					AutoRTFM::Unreachable();
+				case EInsertResult::Inserted:
+					return true;
+				case EInsertResult::Exists:
+					return false;
+				case EInsertResult::NeedResize:
+					Resize();
+					continue;
+				}
+			}
+		}
+
+		bool IsEmpty() const
+		{
+			return 0 == Size;
+		}
+
+		// Clear out the data stored in the set, but does not reduce the capacity.
+		void Reset()
+		{
+			memset(Payload, 0, sizeof(uintptr_t) * Capacity);
+			Size = 0;
+		}
+
+	private:
+		uintptr_t* Payload;
+
+		// TODO: Make this (Capacity - 1) to save a minus in the modulus calculation.
+		uint32_t Capacity;
+
+		uint32_t Size;
+
+		void Resize()
+		{
+			const uintptr_t* const OldPayload = Payload;
+			const uint32_t OldCapacity = Capacity;
+			const uint32_t OldSize = Size;
+
+			while (true)
+			{
+				bool NeedAnotherResize = false;
+
+				Size = 0;
+
+				Capacity *= 2;
+
+				// Check that we haven't overflowed our capacity!
+				ASSERT(0 != Capacity);
+
+				Payload = new uintptr_t[Capacity + LinearProbeDepth]();
+				ASSERT(nullptr != Payload);
+
+				// Now we need to rehash and reinsert all the items. We need to
+				// remember about the extra linear probe section here too.
+				for (size_t I = 0; I < OldCapacity + LinearProbeDepth; I++)
+				{
+					// Skip empty locations.
+					if (0 == OldPayload[I])
+					{
+						continue;
+					}
+
+					const EInsertResult Result = InsertNoResize(OldPayload[I]);
+
+					if (EInsertResult::NeedResize == Result)
+					{
+						NeedAnotherResize = true;
+						break;
+					}
+					else
+					{
+						ASSERT(EInsertResult::Inserted == Result);
+					}
+				}
+
+				if (NeedAnotherResize)
+				{
+					delete[] Payload;
+					continue;
+				}
+
+				break;
+			}
+
+			ASSERT(OldSize == Size);
+
+			delete[] OldPayload;
+
+			return;
+		}
+
+		enum class EInsertResult : uint8_t
+		{
+			Exists,
+			Inserted,
+			NeedResize
+		};
+
+		// Insert something in the HitSet, returning true if the insert
+		// succeeded (EG. the key was not already in the set).
+		EInsertResult InsertNoResize(const uintptr_t Raw)
+		{
+			// TODO: Find a great hashing function for the set.
+			constexpr uintptr_t A = 0x65d200ce55b19ad8;
+			constexpr uintptr_t B = 0x4f2162926e40c299;
+			constexpr uintptr_t C = 0x162dd799029970f8;
+
+			const uintptr_t Low = Raw & 0x00000000FFFFFFFF;
+			const uintptr_t High = Raw >> 32;
+
+			const uint32_t Hash = static_cast<uint32_t>((A * Low + B * High + C) >> 32);
+
+			if (Size < Capacity)
+			{
+				const uint32_t HashModCapacity = Hash & (Capacity - 1);
+
+				for (uint32_t D = 0; D < LinearProbeDepth; D++)
+				{
+					const uintptr_t I = HashModCapacity + D;
+
+					// We have a free location in the set.
+					if (0 == Payload[I])
+					{
+						Payload[I] = Raw;
+						Size++;
+						return EInsertResult::Inserted;
+					}
+
+					// We're already in the set.
+					if (Raw == Payload[I])
+					{
+						return EInsertResult::Exists;
+					}
+
+					// Otherwise we need to keep going with our linear probe...
+				}
+			}
+
+			return EInsertResult::NeedResize;
+		}
+	};
+}
diff --git a/Engine/Source/Runtime/Core/Private/AutoRTFM/LockGuard.h b/Engine/Source/Runtime/Core/Private/AutoRTFM/LockGuard.h
new file mode 100644
--- /dev/null
+++ b/Engine/Source/Runtime/Core/Private/AutoRTFM/LockGuard.h
@@ -0,0 +1,27 @@
+// Copyright Epic Games, Inc. All Rights Reserved.
+
+#pragma once
+
+namespace AutoRTFM
+{
+
+template<typename TLockType>
+class TLockGuard
+{
+public:
+    TLockGuard(TLockType& Lock)
+        : Lock(Lock)
+    {
+        Lock.Lock();
+    }
+
+    ~TLockGuard()
+    {
+        Lock.Unlock();
+    }
+
+private:
+    TLockType& Lock;
+};
+
+} // namespace AutoRTFM
diff --git a/Engine/Source/Runtime/Core/Private/AutoRTFM/LongJump.h b/Engine/Source/Runtime/Core/Private/AutoRTFM/LongJump.h
new file mode 100644
--- /dev/null
+++ b/Engine/Source/Runtime/Core/Private/AutoRTFM/LongJump.h
@@ -0,0 +1,65 @@
+// Copyright Epic Games, Inc. All Rights Reserved.
+
+#pragma once
+
+#include "AutoRTFM/AutoRTFM.h"
+#include "Utils.h"
+#include <setjmp.h>
+
+namespace AutoRTFM
+{
+
+// setjmp/longjmp that doesn't do any unwinding (no C++ destructor calls, no messing with OS
+// signal states - just saving/restoring CPU state). Although it's the setjmp/longjmp everyone
+// knows and loves, it's exposed as a try/catch/throw API to make it less error-prone.
+class FLongJump
+{
+public:
+    FLongJump()
+    {
+        memset(this, 0, sizeof(*this));
+    }
+
+    template<typename TTryFunctor, typename TCatchFunctor>
+    void TryCatch(const TTryFunctor& TryFunctor, const TCatchFunctor& CatchFunctor);
+    
+    void Throw();
+    
+private:
+    jmp_buf JmpBuf;
+    bool bIsSet;
+};
+
+template<typename TTryFunctor, typename TCatchFunctor>
+void FLongJump::TryCatch(const TTryFunctor& TryFunctor, const TCatchFunctor& CatchFunctor)
+{
+    ASSERT(!bIsSet);
+#ifdef _WIN32
+	if (!setjmp(JmpBuf))
+#else
+	if (!_setjmp(JmpBuf))
+#endif
+    {
+        bIsSet = true;
+        TryFunctor();
+        bIsSet = false;
+    }
+    else
+    {
+        ASSERT(bIsSet);
+        bIsSet = false;
+        CatchFunctor();
+    }
+}
+
+inline void FLongJump::Throw()
+{
+    ASSERT(bIsSet);
+#ifdef _WIN32
+	longjmp(JmpBuf, 1);
+#else
+	_longjmp(JmpBuf, 1);
+#endif
+}
+
+} // namespace AutoRTFM
diff --git a/Engine/Source/Runtime/Core/Private/AutoRTFM/Memcpy.cpp b/Engine/Source/Runtime/Core/Private/AutoRTFM/Memcpy.cpp
new file mode 100644
--- /dev/null
+++ b/Engine/Source/Runtime/Core/Private/AutoRTFM/Memcpy.cpp
@@ -0,0 +1,64 @@
+// Copyright Epic Games, Inc. All Rights Reserved.
+
+#if (defined(__AUTORTFM) && __AUTORTFM)
+#include "Memcpy.h"
+#include "ContextInlines.h"
+#include "Debug.h"
+
+namespace AutoRTFM
+{
+
+constexpr bool bVerbose = FDebug::bVerbose;
+
+void* MemcpyToNew(void* InDst, const void* InSrc, size_t Size, FContext* Context)
+{
+    if (bVerbose)
+    {
+        fprintf(GetLogFile(), "MemcpyToNew(%p, %p, %zu)\n", InDst, InSrc, Size);
+    }
+    FDebug Debug(Context, InDst, InSrc, Size, 0, __FUNCTION__);
+    AutoRTFM::Unreachable();
+}
+
+void* Memcpy(void* InDst, const void* InSrc, size_t Size, FContext* Context)
+{
+    if (bVerbose)
+    {
+        fprintf(GetLogFile(), "Memcpy(%p, %p, %zu)\n", InDst, InSrc, Size);
+    }
+    FDebug Debug(Context, InDst, InSrc, Size, 0, __FUNCTION__);
+    
+    Context->RecordWrite(InDst, Size);
+
+    return memcpy(InDst, InSrc, Size);
+}
+
+void* Memmove(void* InDst, const void* InSrc, size_t Size, FContext* Context)
+{
+    if (bVerbose)
+    {
+        fprintf(GetLogFile(), "Memmove(%p, %p, %zu)\n", InDst, InSrc, Size);
+    }
+    FDebug Debug(Context, InDst, InSrc, Size, 0, __FUNCTION__);
+
+	Context->RecordWrite(InDst, Size);
+
+	return memmove(InDst, InSrc, Size);
+}
+
+void* Memset(void* InDst, int Value, size_t Size, FContext* Context)
+{
+    if (bVerbose)
+    {
+        fprintf(GetLogFile(), "Memset(%p, %d, %zu)\n", InDst, Value, Size);
+    }
+    FDebug Debug(Context, InDst, nullptr, Size, 0, __FUNCTION__);
+
+	Context->RecordWrite(InDst, Size);
+
+	return memset(InDst, Value, Size);
+}
+
+} // namespace AutoRTFM
+
+#endif // defined(__AUTORTFM) && __AUTORTFM
diff --git a/Engine/Source/Runtime/Core/Private/AutoRTFM/Memcpy.h b/Engine/Source/Runtime/Core/Private/AutoRTFM/Memcpy.h
new file mode 100644
--- /dev/null
+++ b/Engine/Source/Runtime/Core/Private/AutoRTFM/Memcpy.h
@@ -0,0 +1,16 @@
+// Copyright Epic Games, Inc. All Rights Reserved.
+
+#pragma once
+
+namespace AutoRTFM
+{
+
+class FContext;
+
+void* MemcpyToNew(void* Dst, const void* Src, size_t Size, FContext* Context);
+void* Memcpy(void* Dst, const void* Src, size_t Size, FContext* Context);
+void* Memmove(void* Dst, const void* Src, size_t Size, FContext* Context);
+void* Memset(void* Dst, int Value, size_t Size, FContext* Context);
+
+} // namespace AutoRTFM
+
diff --git a/Engine/Source/Runtime/Core/Private/AutoRTFM/ObjectPool.h b/Engine/Source/Runtime/Core/Private/AutoRTFM/ObjectPool.h
new file mode 100644
--- /dev/null
+++ b/Engine/Source/Runtime/Core/Private/AutoRTFM/ObjectPool.h
@@ -0,0 +1,45 @@
+// Copyright Epic Games, Inc. All Rights Reserved.
+
+#pragma once
+
+namespace AutoRTFM
+{
+
+template<typename T, size_t PoolSize>
+class TObjectPool
+{
+public:
+    TObjectPool() = default;
+
+    T* Allocate()
+    {
+        if (FFreeNode* FreeNode = NextFree)
+        {
+            T* Result = reinterpret_cast<T*>(FreeNode);
+            NextFree = FreeNode->Next;
+            return Result;
+        }
+
+        ASSERT(HighWatermark < PoolSize);
+        return Lines + HighWatermark++;
+    }
+
+    void Deallocate(T* Line)
+    {
+        FFreeNode* FreeNode = reinterpret_cast<FFreeNode*>(Line);
+        FreeNode->Next = NextFree;
+        NextFree = FreeNode;
+    }
+    
+private:
+    struct FFreeNode
+    {
+        FFreeNode* Next;
+    };
+    
+    T Lines[PoolSize];
+    size_t HighWatermark{0};
+    FFreeNode* NextFree{nullptr};
+};
+
+} // namespace AutoRTFM
diff --git a/Engine/Source/Runtime/Core/Private/AutoRTFM/ScopedGuard.h b/Engine/Source/Runtime/Core/Private/AutoRTFM/ScopedGuard.h
new file mode 100644
--- /dev/null
+++ b/Engine/Source/Runtime/Core/Private/AutoRTFM/ScopedGuard.h
@@ -0,0 +1,29 @@
+// Copyright Epic Games, Inc. All Rights Reserved.
+
+#pragma once
+
+namespace AutoRTFM
+{
+
+template<typename T>
+class TScopedGuard
+{
+public:
+    TScopedGuard(T& Ref, const T& Value)
+        : OldValue(Ref)
+        , Ref(Ref)
+    {
+        Ref = Value;
+    }
+
+    ~TScopedGuard()
+    {
+        Ref = OldValue;
+    }
+
+private:
+    T OldValue;
+    T& Ref;
+};
+
+} // namespace AutoRTFM
diff --git a/Engine/Source/Runtime/Core/Private/AutoRTFM/SharedPtr.h b/Engine/Source/Runtime/Core/Private/AutoRTFM/SharedPtr.h
new file mode 100644
--- /dev/null
+++ b/Engine/Source/Runtime/Core/Private/AutoRTFM/SharedPtr.h
@@ -0,0 +1,112 @@
+// Copyright Epic Games, Inc. All Rights Reserved.
+
+#pragma once
+
+#include "Utils.h"
+
+namespace AutoRTFM
+{
+
+template<typename T>
+class TSharedPtr final
+{
+public:
+    TSharedPtr() = default;
+
+    TSharedPtr(const TSharedPtr<T>& Other)
+        : Ptr(Other.Ptr), Count(Other.Count)
+    {
+        if (nullptr != Ptr)
+        {
+            ASSERT(nullptr != Count);
+            *Count += 1;
+        }
+        else
+        {
+            ASSERT(nullptr == Count);
+        }
+    }
+
+	TSharedPtr(TSharedPtr<T>&& Other)
+		: Ptr(Other.Ptr), Count(Other.Count)
+	{
+        // The move constructor **must** behave like the copy constructor,
+        // otherwise MSVC will miscompile our code and free still referenced
+        // data (it seems to move the data, but then _somehow_ keep the original
+        // data around which it then calls the destructor for!). This is cheap
+        // anyway since our data has already been allocated, and the destructor
+        // for the moved data will just decrement the already allocated count.
+        if (nullptr != Ptr)
+        {
+            ASSERT(nullptr != Count);
+            *Count += 1;
+        }
+        else
+        {
+            ASSERT(nullptr == Count);
+        }
+	}
+    
+    ~TSharedPtr()
+    {
+        if (Ptr)
+        {
+            ASSERT(nullptr != Count);
+
+            *Count -= 1;
+
+            if (0 == *Count)
+            {
+                delete Ptr;
+                delete Count;
+            }
+        }
+        else
+        {
+            ASSERT(nullptr == Count);
+        }
+    }
+
+    template<typename... TArguments>
+    static TSharedPtr New(TArguments&&... Arguments)
+    {
+        TSharedPtr Result;
+        Result.Ptr = new T(std::forward<TArguments>(Arguments)...);
+        Result.Count = new unsigned(1);
+        return Result;
+    }
+
+    TSharedPtr& operator=(const TSharedPtr<T>& Other)
+    {
+        ASSERT(nullptr == Ptr);
+        ASSERT(nullptr == Count);
+
+        Ptr = Other.Ptr;
+        Count = Other.Count;
+        
+        *Count += 1;
+    }
+
+	TSharedPtr& operator=(TSharedPtr<T>&& Other) = delete;
+
+    explicit operator bool() const { return !!Ptr; }
+
+    bool operator==(const TSharedPtr<T>& Other) const
+    {
+        return Ptr == Other.Ptr;
+    }
+
+    bool operator!=(const TSharedPtr<T>& Other) const
+    {
+        return Ptr != Other.Ptr;
+    }
+
+    T& operator*() const { return *Ptr; }
+    T* operator->() const { return Ptr; }
+    
+private:
+    T* Ptr{nullptr};
+    unsigned* Count{nullptr};
+};
+
+} // namespace AutoRTFM
diff --git a/Engine/Source/Runtime/Core/Private/AutoRTFM/StdLibSupport.cpp b/Engine/Source/Runtime/Core/Private/AutoRTFM/StdLibSupport.cpp
new file mode 100644
--- /dev/null
+++ b/Engine/Source/Runtime/Core/Private/AutoRTFM/StdLibSupport.cpp
@@ -0,0 +1,208 @@
+// Copyright Epic Games, Inc. All Rights Reserved.
+
+#if (defined(__AUTORTFM) && __AUTORTFM)
+#include "ContextInlines.h"
+#include "Utils.h"
+#include "Memcpy.h"
+#include <stdarg.h>
+#include <stdio.h>
+#include <string.h>
+#include <wchar.h>
+#include <float.h>
+#include <math.h>
+
+namespace AutoRTFM
+{
+
+UE_AUTORTFM_REGISTER_OPEN_FUNCTION_EXPLICIT(memcpy, Memcpy);
+UE_AUTORTFM_REGISTER_OPEN_FUNCTION_EXPLICIT(memmove, Memmove);
+UE_AUTORTFM_REGISTER_OPEN_FUNCTION_EXPLICIT(memset, Memset);
+
+// TODO: We need this to stop link errors looking for `DidAllocate` below?
+#pragma clang optimize off
+
+void* STM_malloc(size_t Size, FContext* Context)
+{
+    void* Result = malloc(Size);
+    Context->GetCurrentTransaction()->DeferUntilAbort([Result]
+    {
+        free(Result);
+    });
+    Context->DidAllocate(Result, Size);
+    return Result;
+}
+
+#pragma clang optimize on
+
+UE_AUTORTFM_REGISTER_OPEN_FUNCTION(malloc);
+
+void STM_free(void* Ptr, FContext* Context)
+{
+    if (Ptr)
+    {
+		const size_t AllocSize = GetAllocationSize(Ptr);
+        Context->WillDeallocate(Ptr, AllocSize);
+
+        Context->GetCurrentTransaction()->DeferUntilCommit([Ptr]
+        {
+            free(Ptr);
+        });
+    }
+}
+UE_AUTORTFM_REGISTER_OPEN_FUNCTION(free);
+
+void* STM_realloc(void* Ptr, size_t Size, FContext* Context)
+{
+    void* NewObject = STM_malloc(Size, Context);
+    if (Ptr)
+    {
+    	const size_t OldSize = GetAllocationSize(Ptr);
+        MemcpyToNew(NewObject, Ptr, std::min(OldSize, Size), Context);
+        STM_free(Ptr, Context);
+    }
+    return NewObject;
+}
+UE_AUTORTFM_REGISTER_OPEN_FUNCTION(realloc);
+
+
+size_t STM_strlen(const char* Str, FContext* Context)
+{
+	AutoRTFM::Unreachable();
+}
+UE_AUTORTFM_REGISTER_OPEN_FUNCTION(strlen);
+
+// FIXME: This is only correct when:
+// - Str is newly allocated
+// - Format is either newly allocated or not mutated
+// - any strings passed as arguments are either newly allocated or not mutated
+int STM_snprintf(char* Str, size_t Size, char* Format, FContext* Context, ...)
+{
+    va_list ArgList;
+    va_start(ArgList, Context);
+    int Result = vsnprintf(Str, Size, Format, ArgList);
+    va_end(ArgList);
+    return Result;
+}
+UE_AUTORTFM_REGISTER_OPEN_FUNCTION(snprintf);
+
+char *STM_strncpy(char* Dst, const char* Src, size_t Len, FContext* Context)
+{
+    AutoRTFM::Unreachable();
+}
+
+#ifdef _MSC_VER
+/*
+   Disable warning about deprecated STD C functions.
+*/
+#pragma warning(disable : 4996)
+
+#pragma warning(push)
+#endif
+
+UE_AUTORTFM_REGISTER_OPEN_FUNCTION(strncpy);
+
+#ifdef _MSC_VER
+#pragma warning(pop)
+#endif
+
+int STM_printf(const char* Format, FContext* Context, ...)
+{
+    va_list ArgList;
+    va_start(ArgList, Context);
+    int Result = vprintf(Format, ArgList);
+    va_end(ArgList);
+    return Result;
+}
+UE_AUTORTFM_REGISTER_OPEN_FUNCTION(printf);
+
+int STM_putchar(int Char, FContext* Context)
+{
+    return putchar(Char);
+}
+UE_AUTORTFM_REGISTER_OPEN_FUNCTION(putchar);
+
+int STM_puts(const char* Str, FContext* Context)
+{
+    return puts(Str);
+}
+UE_AUTORTFM_REGISTER_OPEN_FUNCTION(puts);
+
+#ifdef _WIN32
+
+FILE* STM___acrt_iob_func(int Index, FContext* Context)
+{
+    switch (Index)
+    {
+    case 1:
+    case 2:
+        return __acrt_iob_func(Index);
+    default:
+        fprintf(stderr, "Attempt to get file descriptor %d (not 1 or 2) in __acrt_iob_func.\n", Index);
+        Context->AbortByLanguageAndThrow();
+        return NULL;
+    }
+}
+UE_AUTORTFM_REGISTER_OPEN_FUNCTION(__acrt_iob_func);
+
+UE_AUTORTFM_REGISTER_SELF_FUNCTION(__stdio_common_vfprintf);
+
+UE_AUTORTFM_REGISTER_SELF_FUNCTION(puts);
+
+UE_AUTORTFM_REGISTER_SELF_FUNCTION(__stdio_common_vswprintf);
+
+UE_AUTORTFM_REGISTER_SELF_FUNCTION(__stdio_common_vfwprintf);
+
+#else // _WIN32 -> so !_WIN32
+extern "C" size_t _ZNSt3__112__next_primeEm(size_t N) __attribute__((weak));
+UE_AUTORTFM_REGISTER_SELF_FUNCTION(_ZNSt3__112__next_primeEm);
+#endif // _WIN32 -> so end of !_WIN32
+
+UE_AUTORTFM_REGISTER_SELF_FUNCTION(static_cast<float(*)(float, float)>(&powf));
+UE_AUTORTFM_REGISTER_SELF_FUNCTION(static_cast<double(*)(double, double)>(&pow));
+
+#ifdef _WIN32
+UE_AUTORTFM_REGISTER_SELF_FUNCTION(_isnan);
+UE_AUTORTFM_REGISTER_SELF_FUNCTION(_finite);
+
+UE_AUTORTFM_REGISTER_SELF_FUNCTION(IsDebuggerPresent);
+
+UE_AUTORTFM_REGISTER_SELF_FUNCTION(EnterCriticalSection);
+UE_AUTORTFM_REGISTER_SELF_FUNCTION(LeaveCriticalSection);
+
+UE_AUTORTFM_REGISTER_SELF_FUNCTION(QueryPerformanceCounter);
+#endif
+
+wchar_t* STM_wcsncpy(wchar_t* Dst, const wchar_t* Src, size_t Count, FContext* Context)
+{
+	AutoRTFM::Unreachable();
+}
+
+#ifdef _MSC_VER
+/*
+   Disable warning about deprecated STD C functions.
+*/
+#pragma warning(disable : 4996)
+
+#pragma warning(push)
+#endif
+
+UE_AUTORTFM_REGISTER_OPEN_FUNCTION(wcsncpy);
+
+#ifdef _MSC_VER
+#pragma warning(pop)
+#endif
+
+int STM_atexit(void(__cdecl*Callback)(void), FContext* Context)
+{
+	Context->GetCurrentTransaction()->DeferUntilCommit([Callback]
+		{
+			atexit(Callback);
+		});
+
+	return 0;
+}
+UE_AUTORTFM_REGISTER_OPEN_FUNCTION(atexit);
+
+} // namespace AutoRTFM
+
+#endif // defined(__AUTORTFM) && __AUTORTFM
diff --git a/Engine/Source/Runtime/Core/Private/AutoRTFM/TaggedPtr.h b/Engine/Source/Runtime/Core/Private/AutoRTFM/TaggedPtr.h
new file mode 100644
--- /dev/null
+++ b/Engine/Source/Runtime/Core/Private/AutoRTFM/TaggedPtr.h
@@ -0,0 +1,202 @@
+// Copyright Epic Games, Inc. All Rights Reserved.
+
+#pragma once
+
+#include "Utils.h"
+#include <stdint.h>
+
+namespace AutoRTFM
+{
+	// TODO: Could get lower bits too due to alignment?
+	// TODO: Do we need to care about 5-level paging and their 57-bit pointers?
+	template<typename T, typename TopTagT = uint16_t> class TTaggedPtr final
+	{
+		// Our tagged pointer assumes 64-bit pointers.
+		static_assert(sizeof(uintptr_t) == sizeof(uint64_t));
+
+		// For our tagged pointers we assume the bottom 48-bits are used, and
+		// so the top 16-bits are free for our tag.
+		static_assert(sizeof(TopTagT) == sizeof(uint16_t));
+
+		static constexpr unsigned PointerBits = 48;
+		static constexpr uintptr_t PointerMask = (static_cast<uintptr_t>(1) << PointerBits) - 1;
+		static constexpr unsigned TopTagBits = sizeof(TopTagT);
+		static constexpr uintptr_t TopTagMask = (static_cast<uintptr_t>(1) << TopTagBits) - 1;
+
+	public:
+		struct Hasher final
+		{
+		    size_t operator()(const TTaggedPtr<T, TopTagT>& TaggedPtr) const
+		    {
+		      return TaggedPtr.Payload;
+		    }
+		};
+
+		// Assumes little-endian, so top tag is the last 2-bytes in memory.
+		static constexpr size_t OffsetOfTopTag = PointerBits / 8;
+
+		TTaggedPtr()
+		{
+			Reset();
+		}
+
+		explicit TTaggedPtr(T* const Pointer) : Payload(reinterpret_cast<uintptr_t>(Pointer))
+		{
+			ASSERT(0 == (Payload & ~PointerMask));
+		}
+
+		explicit TTaggedPtr(const TTaggedPtr& Other) : Payload(Other.Payload) {}
+		explicit TTaggedPtr(const TTaggedPtr&& Other) : Payload(Other.Payload) {}
+
+		TTaggedPtr& operator=(T* const Pointer)
+		{
+			const uintptr_t NewPayload = reinterpret_cast<uintptr_t>(Pointer);
+
+			ASSERT(0 == (NewPayload & ~PointerMask));
+			Payload &= ~PointerMask;
+			Payload |= NewPayload;
+			return *this;
+		}
+
+		TTaggedPtr& operator=(TTaggedPtr& Other)
+		{
+			Payload = Other.Payload;
+		}
+
+		TTaggedPtr& operator=(TTaggedPtr&& Other)
+		{
+			Payload = Other.Payload;
+		}
+
+		T& operator*() const
+	    {
+	    	return *reinterpret_cast<T*>(Payload & PointerMask);
+	    }
+
+	    T* operator->() const
+	    {
+	    	return reinterpret_cast<T*>(Payload & PointerMask);
+	    }
+
+	    T* Get() const
+	    {
+	    	return reinterpret_cast<T*>(Payload & PointerMask);
+	    }
+
+	    void SetTopTag(const TopTagT Tag)
+	    {
+	    	*reinterpret_cast<TopTagT*>(reinterpret_cast<uint8_t*>(&Payload) + OffsetOfTopTag) = std::move(Tag);
+	    }
+
+	    TopTagT GetTopTag() const
+	    {
+	    	return *reinterpret_cast<const TopTagT*>(reinterpret_cast<const uint8_t*>(&Payload) + OffsetOfTopTag);
+	    }
+
+	    void Reset()
+	    {
+	    	Payload = 0;
+	    }
+
+	    uintptr_t GetPayload() const
+	    {
+	    	return Payload;
+	    }
+
+	private:
+		uintptr_t Payload;
+	};
+
+	template<typename TopTagT> class TTaggedPtr<void, TopTagT> final
+	{
+		// Our tagged pointer assumes 64-bit pointers.
+		static_assert(sizeof(uintptr_t) == sizeof(uint64_t));
+
+		// For our tagged pointers we assume the bottom 48-bits are used, and
+		// so the top 16-bits are free for our tag.
+		static_assert(sizeof(TopTagT) == sizeof(uint16_t));
+
+		static constexpr unsigned PointerBits = 48;
+		static constexpr uintptr_t PointerMask = (static_cast<uintptr_t>(1) << PointerBits) - 1;
+		static constexpr unsigned TopTagBits = sizeof(TopTagT);
+		static constexpr uintptr_t TopTagMask = (static_cast<uintptr_t>(1) << TopTagBits) - 1;
+
+	public:
+		struct Hasher final
+		{
+		    size_t operator()(const TTaggedPtr<void, TopTagT>& TaggedPtr) const
+		    {
+		      return TaggedPtr.Payload;
+		    }
+		};
+
+		// Assumes little-endian, so top tag is the last 2-bytes in memory.
+		static constexpr size_t OffsetOfTopTag = PointerBits / 8;
+
+		TTaggedPtr()
+		{
+			Reset();
+		}
+
+		TTaggedPtr(void* const Pointer) : Payload(reinterpret_cast<uintptr_t>(Pointer))
+		{
+			ASSERT(0 == (Payload & ~PointerMask));
+		}
+
+		TTaggedPtr(const TTaggedPtr& Other) : Payload(Other.Payload) {}
+		TTaggedPtr(const TTaggedPtr&& Other) : Payload(Other.Payload) {}
+
+		TTaggedPtr& operator=(void* const Pointer)
+		{
+			const uintptr_t NewPayload = reinterpret_cast<uintptr_t>(Pointer);
+
+			ASSERT(0 == (NewPayload & ~PointerMask));
+			Payload &= ~PointerMask;
+			Payload |= NewPayload;
+			return *this;
+		}
+
+		TTaggedPtr& operator=(TTaggedPtr& Other)
+		{
+			Payload = Other.Payload;
+		}
+
+		TTaggedPtr& operator=(TTaggedPtr&& Other)
+		{
+			Payload = Other.Payload;
+		}
+
+		bool operator==(const TTaggedPtr& Other) const
+		{
+			return Payload == Other.Payload;
+		}
+
+		void* Get() const
+	    {
+	    	return reinterpret_cast<void*>(Payload & PointerMask);
+	    }
+
+	    void SetTopTag(const TopTagT Tag)
+	    {
+	    	*reinterpret_cast<TopTagT*>(reinterpret_cast<uint8_t*>(&Payload) + OffsetOfTopTag) = std::move(Tag);
+	    }
+
+	    TopTagT GetTopTag() const
+	    {
+	    	return *reinterpret_cast<const TopTagT*>(reinterpret_cast<const uint8_t*>(&Payload) + OffsetOfTopTag);
+	    }
+
+	    void Reset()
+	    {
+	    	Payload = 0;
+	    }
+
+	    uintptr_t GetPayload() const
+	    {
+	    	return Payload;
+	    }
+
+	private:
+		uintptr_t Payload;
+	};
+}
diff --git a/Engine/Source/Runtime/Core/Private/AutoRTFM/TaskArray.h b/Engine/Source/Runtime/Core/Private/AutoRTFM/TaskArray.h
new file mode 100644
--- /dev/null
+++ b/Engine/Source/Runtime/Core/Private/AutoRTFM/TaskArray.h
@@ -0,0 +1,122 @@
+// Copyright Epic Games, Inc. All Rights Reserved.
+
+#pragma once
+
+#include "SharedPtr.h"
+#include <vector>
+
+namespace AutoRTFM
+{
+
+template<typename T>
+class TTaskArray
+{
+public:
+    TTaskArray() = default;
+
+    bool IsEmpty() const { return Latest.empty() && Stash.empty(); }
+
+    void Add(T&& value)
+    {
+        Latest.push_back(std::move(value));
+    }
+    void Add(const T& value)
+    {
+        Latest.push_back(value);
+    }
+
+    void AddAll(TTaskArray<T>&& Other)
+    {
+        Canonicalize();
+        for (TSharedPtr<std::vector<T>>& StashedVectorBox : Other.Stash)
+        {
+            Stash.push_back(std::move(StashedVectorBox));
+        }
+        Other.Stash.clear(); 
+        if (!Other.Latest.empty())
+        {
+            Stash.push_back(TSharedPtr<std::vector<T>>::New(std::move(Other.Latest)));
+        }
+    }
+
+    void AddAll(const TTaskArray<T>& Other)
+    {
+        Canonicalize();
+        Other.Canonicalize();
+        for (const TSharedPtr<std::vector<T>>& StashedVectorBox : Other.Stash)
+        {
+            Stash.push_back(StashedVectorBox);
+        }
+    }
+
+    template<typename TFunc>
+    bool ForEachForward(const TFunc& Func)
+    {
+        for (TSharedPtr<std::vector<T>>& StashedVectorBox : Stash)
+        {
+            for (const T& Entry : *StashedVectorBox)
+            {
+                if (!Func(Entry))
+                {
+                    return false;
+                }
+            }
+        }
+        for (const T& Entry : Latest)
+        {
+            if (!Func(Entry))
+            {
+                return false;
+            }
+        }
+        return true;
+    }
+
+    template<typename TFunc>
+    bool ForEachBackward(const TFunc& Func)
+    {
+        for (size_t Index = Latest.size(); Index--;)
+        {
+            if (!Func(Latest[Index]))
+            {
+                return false;
+            }
+        }
+        for (size_t IndexInStash = Stash.size(); IndexInStash--;)
+        {
+            const std::vector<T>& StashedVector = *Stash[IndexInStash];
+            for (size_t Index = StashedVector.size(); Index--;)
+            {
+                if (!Func(StashedVector[Index]))
+                {
+                    return false;
+                }
+            }
+        }
+        return true;
+    }
+
+    void Reset()
+    {
+        Latest.clear();
+        Stash.clear();
+    }
+
+private:
+    // We don't want to do this too often; currently we just do it where it's asymptitically relevant like AddAll. This doesn't
+    // logically change the TaskArray but it changes its internal representation. Hence the use of `mutable` and hence why this
+    // method is `const`.
+    void Canonicalize() const
+    {
+        if (!Latest.empty())
+        {
+            Stash.push_back(TSharedPtr<std::vector<T>>::New(std::move(Latest)));
+        }
+    }
+    
+    mutable std::vector<T> Latest;
+    
+    mutable std::vector<TSharedPtr<std::vector<T>>> Stash;
+};
+
+} // namespace AutoRTFM
diff --git a/Engine/Source/Runtime/Core/Private/AutoRTFM/Transaction.cpp b/Engine/Source/Runtime/Core/Private/AutoRTFM/Transaction.cpp
new file mode 100644
--- /dev/null
+++ b/Engine/Source/Runtime/Core/Private/AutoRTFM/Transaction.cpp
@@ -0,0 +1,176 @@
+// Copyright Epic Games, Inc. All Rights Reserved.
+
+#if (defined(__AUTORTFM) && __AUTORTFM)
+#include "Transaction.h"
+#include "TransactionInlines.h"
+#include "Debug.h"
+#include "GlobalData.h"
+
+namespace AutoRTFM
+{
+
+FTransaction::FTransaction(FContext* Context)
+    : Context(Context)
+    , Parent(Context->GetCurrentTransaction())
+{
+}
+
+bool FTransaction::IsFresh() const
+{
+    return HitSet.IsEmpty()
+        && WriteLog.IsEmpty()
+        && CommitTasks.IsEmpty()
+        && AbortTasks.IsEmpty()
+        && !bIsDone;
+}
+
+void FTransaction::AbortWithoutThrowing()
+{
+    if (FDebug::bVerbose)
+    {
+        fprintf(GetLogFile(), "Aborting (%s)\n", GetContextStatusName(Context->GetStatus()));
+    }
+    ASSERT(Context->GetStatus() == EContextStatus::AbortedByLineInitialization
+           || Context->GetStatus() == EContextStatus::AbortedByCommitTimeLineCheck
+           || Context->GetStatus() == EContextStatus::AbortedByFailedLockAcquisition
+           || Context->GetStatus() == EContextStatus::AbortedByLanguage
+           || Context->GetStatus() == EContextStatus::AbortedByRequest);
+    ASSERT(Context->GetCurrentTransaction() == this);
+    if (IsNested())
+    {
+        AbortNested();
+    }
+    else
+    {
+        AbortOuterNest();
+    }
+    Reset();
+}
+
+void FTransaction::AbortAndThrow()
+{
+    AbortWithoutThrowing();
+    AbortJump.Throw();
+}
+
+bool FTransaction::AttemptToCommit()
+{
+    ASSERT(Context->GetStatus() == EContextStatus::OnTrack);
+    ASSERT(Context->GetCurrentTransaction() == this);
+    bool bResult;
+    if (IsNested())
+    {
+        CommitNested();
+        bResult = true;
+    }
+    else
+    {
+        bResult = AttemptToCommitOuterNest();
+    }
+    Reset();
+    return bResult;
+}
+
+void FTransaction::Undo()
+{
+    for (FWriteLogEntry& Entry : WriteLog)
+    {
+        // Skip writes that are into our current transactions stack.
+        if (Context->IsInnerTransactionStack(Entry.OriginalAndSize.Get()))
+        {
+            continue;
+        }
+
+        void* const Original = Entry.OriginalAndSize.Get();
+        const size_t Size = Entry.OriginalAndSize.GetTopTag();
+        void* const Copy = Entry.Copy;
+        memcpy(Original, Copy, Size);
+    }
+}
+
+void FTransaction::AbortNested()
+{
+    ASSERT(Parent);
+
+    Undo();
+
+    Parent->CommitTasks.AddAll(AbortTasks);
+    Parent->AbortTasks.AddAll(std::move(AbortTasks));
+}
+
+void FTransaction::AbortOuterNest()
+{
+    Undo();
+
+    AbortTasks.ForEachBackward([] (const std::function<void()>& Task) -> bool { Task(); return true; });
+
+    switch (Context->GetStatus())
+    {
+    case EContextStatus::AbortedByLineInitialization:
+    case EContextStatus::AbortedByCommitTimeLineCheck:
+    case EContextStatus::AbortedByFailedLockAcquisition:
+        break;
+    case EContextStatus::AbortedByRequest:
+    case EContextStatus::AbortedByLanguage:
+        break;
+    default:
+        ASSERT(!"Should not be reached");
+        break;
+    }
+}
+
+void FTransaction::CommitNested()
+{
+    ASSERT(Parent);
+
+    // We need to pass our write log to our parent transaction, but with care!
+    // We need to discard any writes to locations within the stack of our
+    // current transaction, which could be placed there if a child of the
+    // current transaction had written to stack local memory in the parent.
+
+    for (FWriteLogEntry& Write : WriteLog)
+    {
+        // Skip writes that are into our current transactions stack.
+        if (Context->IsInnerTransactionStack(Write.OriginalAndSize.Get()))
+        {
+            continue;
+        }
+
+        Parent->WriteLog.Push(Write);
+        Parent->HitSet.Insert(Write.OriginalAndSize);
+    }
+
+    Parent->WriteLogBumpAllocator.Merge(std::move(WriteLogBumpAllocator));
+
+    Parent->CommitTasks.AddAll(std::move(CommitTasks));
+    Parent->AbortTasks.AddAll(std::move(AbortTasks));
+}
+
+bool FTransaction::AttemptToCommitOuterNest()
+{
+    ASSERT(!Parent);
+
+    if (FDebug::bVerbose)
+    {
+        fprintf(GetLogFile(), "About to run commit tasks!\n");
+        Context->DumpState();
+        fprintf(GetLogFile(), "Running commit tasks...\n");
+    }
+
+    CommitTasks.ForEachForward([] (const std::function<void()>& Task) -> bool { Task(); return true; });
+
+    return true;
+}
+
+void FTransaction::Reset()
+{
+    CommitTasks.Reset();
+    AbortTasks.Reset();
+	HitSet.Reset();
+	WriteLog.Reset();
+	WriteLogBumpAllocator.Reset();
+}
+
+} // namespace AutoRTFM
+
+#endif // defined(__AUTORTFM) && __AUTORTFM
diff --git a/Engine/Source/Runtime/Core/Private/AutoRTFM/Transaction.h b/Engine/Source/Runtime/Core/Private/AutoRTFM/Transaction.h
new file mode 100644
--- /dev/null
+++ b/Engine/Source/Runtime/Core/Private/AutoRTFM/Transaction.h
@@ -0,0 +1,90 @@
+// Copyright Epic Games, Inc. All Rights Reserved.
+
+#pragma once
+
+#include "HitSet.h"
+#include "LongJump.h"
+#include "TaggedPtr.h"
+#include "TaskArray.h"
+#include "WriteLog.h"
+#include "WriteLogBumpAllocator.h"
+#include <functional>
+
+namespace AutoRTFM
+{
+class FContext;
+
+class FTransaction final
+{
+public:
+    FTransaction(FContext* Context);
+    
+    bool IsNested() const { return !!Parent; }
+    FTransaction* GetParent() const { return Parent; }
+
+    // This should just use type displays or ranges. Maybe ranges could even work out great.
+    bool IsNestedWithin(const FTransaction* Other) const
+    {
+        for (const FTransaction* Current = this; ; Current = Current->Parent)
+        {
+            if (!Current)
+            {
+                return false;
+            }
+            else if (Current == Other)
+            {
+                return true;
+            }
+        }
+    }
+
+    bool IsFresh() const;
+    bool IsDone() const { return bIsDone; }
+    void SetIsDone() { bIsDone = true; }
+    
+    void DeferUntilCommit(std::function<void()>&&);
+    void DeferUntilAbort(std::function<void()>&&);
+
+    // Whether this succeeded or not is reflected in Context::GetStatus().
+    template<typename TTryFunctor>
+    void Try(const TTryFunctor& TryFunctor);
+
+    void AbortAndThrow();
+    void AbortWithoutThrowing();
+    bool AttemptToCommit();
+
+    // Record that a write is about to occur at the given LogicalAddress of Size bytes.
+    void RecordWrite(void* LogicalAddress, size_t Size);
+    void RecordWriteMaxPageSized(void* LogicalAddress, size_t Size);
+
+    void DidAllocate(void* LogicalAddress, size_t Size);
+
+private:
+    void Undo();
+    void AbortNested();
+    void AbortOuterNest();
+    
+    void CommitNested();
+    bool AttemptToCommitOuterNest();
+
+    void Reset(); // Frees memory and sets us up for possible retry.
+    
+    FContext* Context;
+    
+    // Are we nested? Then this is the parent.
+    FTransaction* Parent;
+
+    // Commit tasks run on commit in forward order. Abort tasks run on abort in reverse order.
+    TTaskArray<std::function<void()>> CommitTasks;
+    TTaskArray<std::function<void()>> AbortTasks;
+
+    bool bIsDone{false};
+
+    FLongJump AbortJump;
+
+    FHitSet HitSet;
+    FWriteLog WriteLog;
+    FWriteLogBumpAllocator WriteLogBumpAllocator;
+};
+
+} // namespace AutoRTFM
diff --git a/Engine/Source/Runtime/Core/Private/AutoRTFM/TransactionInlines.h b/Engine/Source/Runtime/Core/Private/AutoRTFM/TransactionInlines.h
new file mode 100644
--- /dev/null
+++ b/Engine/Source/Runtime/Core/Private/AutoRTFM/TransactionInlines.h
@@ -0,0 +1,96 @@
+// Copyright Epic Games, Inc. All Rights Reserved.
+
+#pragma once
+
+#include "Context.h"
+#include "Transaction.h"
+#include <cstddef>
+
+namespace AutoRTFM
+{
+
+inline void FTransaction::RecordWriteMaxPageSized(void* LogicalAddress, size_t Size)
+{
+    ASSERT(Size <= UINT16_MAX);
+
+    FMemoryLocation Key(LogicalAddress);
+    Key.SetTopTag(static_cast<uint16_t>(Size));
+
+    // If we are recording a stack address that is relative to our current
+    // transactions stack location, we do not need to record the data in the
+    // write log because if that transaction aborted, that memory will cease to
+    // be meaningful anyway!
+    // TODO: with more advanced LLVM-skillz we could just detect writes to
+    // alloca's in the LLVM code and that should probably catch all or nearly
+    // all of the cases this checks would catch, so we could revisit doing this
+    // check entirely in future!
+    if (Context->IsInnerTransactionStack(LogicalAddress))
+    {
+        return;
+    }
+
+    if (!HitSet.Insert(Key))
+    {
+        return;
+    }
+
+    void* CopyAddress = WriteLogBumpAllocator.Allocate(Size);
+    memcpy(CopyAddress, LogicalAddress, Size);
+
+    WriteLog.Push(FWriteLogEntry(LogicalAddress, Size, CopyAddress));
+}
+
+inline void FTransaction::RecordWrite(void* LogicalAddress, size_t Size)
+{
+    uint8_t* Address = reinterpret_cast<uint8_t*>(LogicalAddress);
+
+    size_t I = 0;
+
+    for (; (I + UINT16_MAX) < Size; I += UINT16_MAX)
+    {
+        RecordWriteMaxPageSized(Address + I, UINT16_MAX);
+    }
+
+    // Remainder at the end of the memcpy.
+    RecordWriteMaxPageSized(Address + I, Size - I);
+}
+
+inline void FTransaction::DidAllocate(void* LogicalAddress, size_t Size)
+{
+    if (Size <= UINT16_MAX)
+    {
+        FMemoryLocation Key(LogicalAddress);
+        Key.SetTopTag(static_cast<uint16_t>(Size));
+        // Otherwise we need to record the write.
+        const bool DidInsert = HitSet.Insert(Key);
+        ASSERT(DidInsert); 
+    }
+}
+
+inline void FTransaction::DeferUntilCommit(std::function<void()>&& Callback)
+{
+    CommitTasks.Add(std::move(Callback));
+}
+
+inline void FTransaction::DeferUntilAbort(std::function<void()>&& Callback)
+{
+    AbortTasks.Add(std::move(Callback));
+}
+
+template<typename TTryFunctor>
+void FTransaction::Try(const TTryFunctor& TryFunctor)
+{
+    AbortJump.TryCatch(
+        [&] ()
+        {
+            TryFunctor();
+            ASSERT(Context->GetStatus() == EContextStatus::OnTrack);
+        },
+        [&] ()
+        {
+            ASSERT(Context->GetStatus() != EContextStatus::Idle);
+            ASSERT(Context->GetStatus() != EContextStatus::OnTrack);
+        });
+}
+
+} // namespace AutoRTFM
diff --git a/Engine/Source/Runtime/Core/Private/AutoRTFM/Utils.cpp b/Engine/Source/Runtime/Core/Private/AutoRTFM/Utils.cpp
new file mode 100644
--- /dev/null
+++ b/Engine/Source/Runtime/Core/Private/AutoRTFM/Utils.cpp
@@ -0,0 +1,96 @@
+// Copyright Epic Games, Inc. All Rights Reserved.
+
+#if (defined(__AUTORTFM) && __AUTORTFM)
+#include "Utils.h"
+#include <errno.h>
+#include <mutex>
+#include <stdio.h>
+#include <string.h>
+
+#ifdef _WIN32
+#include <dbghelp.h>
+#else
+#include <execinfo.h>
+#endif
+
+namespace AutoRTFM
+{
+
+// To log to a file, comment out LogFilename = nullptr and uncomment the other one.
+constexpr const char* LogFilename = nullptr;
+// constexpr const char* LogFilename = "log.txt";
+
+constexpr const char* LogFileMode = "wt";
+
+std::once_flag LogInitializationFlag;
+FILE* LogFile;
+
+FILE* GetLogFile()
+{
+    std::call_once(LogInitializationFlag, []
+    {
+        if (LogFilename)
+        {
+            ASSERT(LogFileMode);
+
+#ifdef _MSC_VER
+/*
+   Disable warning about deprecated STD C functions.
+*/
+#pragma warning(disable : 4996)
+
+#pragma warning(push)
+#endif
+
+            LogFile = fopen(LogFilename, LogFileMode);
+
+#ifdef _MSC_VER
+#pragma warning(pop)
+#endif
+
+            if (!LogFile)
+            {
+                fprintf(stderr, "Could not open %s: %s\n", LogFilename, strerror(errno));
+            }
+        }
+        if (!LogFile)
+        {
+            LogFile = stderr;
+        }
+        setvbuf(LogFile, NULL, _IONBF, 0);
+    });
+    return LogFile;
+}
+
+std::string GetFunctionDescription(void* FunctionPtr)
+{
+#ifdef _WIN32
+    // This is gross, but it works. It's possible for someone to have SymInitialized before. But if they had, then this
+    // will just fail. Also, this function is called in cases where we're failing, so it's ok if we do dirty things.
+    SymInitialize(GetCurrentProcess(), nullptr, true);
+    
+    DWORD64 Displacement = 0;
+    DWORD64 Address = reinterpret_cast<DWORD64>(FunctionPtr);
+    char Buffer[sizeof(SYMBOL_INFO) + MAX_SYM_NAME * sizeof(TCHAR)];
+    PSYMBOL_INFO Symbol = reinterpret_cast<PSYMBOL_INFO>(Buffer);
+    Symbol->SizeOfStruct = sizeof(SYMBOL_INFO);
+    Symbol->MaxNameLen = MAX_SYM_NAME;
+    if (SymFromAddr(GetCurrentProcess(), Address, &Displacement, Symbol))
+    {
+        return Symbol->Name;
+    }
+    else
+    {
+        return "<error getting description>";
+    }
+#else // _WIN32 -> so !_WIN32
+    char** const symbols = backtrace_symbols(&FunctionPtr, 1);
+    std::string name(*symbols);
+    free(symbols);
+    return name;
+#endif // _WIN32 -> so end of !_WIN32
+}
+
+} // namespace AutoRTFM
+
+#endif // defined(__AUTORTFM) && __AUTORTFM
\ No newline at end of file
diff --git a/Engine/Source/Runtime/Core/Private/AutoRTFM/Utils.h b/Engine/Source/Runtime/Core/Private/AutoRTFM/Utils.h
new file mode 100644
--- /dev/null
+++ b/Engine/Source/Runtime/Core/Private/AutoRTFM/Utils.h
@@ -0,0 +1,64 @@
+// Copyright Epic Games, Inc. All Rights Reserved.
+
+#pragma once
+
+#include "AutoRTFM/AutoRTFM.h"
+#define __STDC_FORMAT_MACROS 1
+#include <inttypes.h>
+#include <stdio.h>
+#include <stdlib.h>
+#include <string>
+#ifdef _WIN32
+#define NOMINMAX 1
+#include <windows.h>
+#else
+#include <malloc.h>
+#endif
+
+#define ASSERT(exp) do { \
+    if (!(exp)) { \
+        fprintf(stderr, "%s:%d:%s: assertion %s failed.\n", __FILE__, __LINE__, __PRETTY_FUNCTION__, #exp); \
+        abort(); \
+    } \
+} while (false)
+
+namespace AutoRTFM
+{
+
+[[noreturn]] inline void Unreachable()
+{
+    fprintf(stderr, "Unreachable encountered!\n");
+
+#ifdef _WIN32
+    __assume(false);
+#else
+    __builtin_unreachable();
+#endif
+}
+
+FILE* GetLogFile();
+
+std::string GetFunctionDescription(void* FunctionPtr);
+
+template<typename TReturnType, typename... TParameterTypes>
+std::string GetFunctionDescription(TReturnType (*FunctionPtr)(TParameterTypes...))
+{
+    return GetFunctionDescription(reinterpret_cast<void*>(FunctionPtr));
+}
+
+inline size_t GetAllocationSize(void* const Ptr)
+{
+#if defined(_WIN32)
+    return _msize(Ptr);
+#else
+    return malloc_usable_size(Ptr);
+#endif
+}
+
+template<size_t A, size_t B> struct PrettyStaticAssert final
+{
+  static_assert(A == B, "Not equal");
+  static constexpr bool _cResult = (A == B);
+};
+
+} // namespace AutoRTFM
diff --git a/Engine/Source/Runtime/Core/Private/AutoRTFM/WriteLog.h b/Engine/Source/Runtime/Core/Private/AutoRTFM/WriteLog.h
new file mode 100644
--- /dev/null
+++ b/Engine/Source/Runtime/Core/Private/AutoRTFM/WriteLog.h
@@ -0,0 +1,132 @@
+// Copyright Epic Games, Inc. All Rights Reserved.
+
+#pragma once
+
+#include <stdint.h>
+
+#include "Utils.h"
+
+namespace AutoRTFM
+{
+	using FMemoryLocation = FHitSet::Key;
+
+    struct FWriteLogEntry final
+    {
+        FMemoryLocation OriginalAndSize;
+        void* Copy;
+
+        FWriteLogEntry() = default;
+        FWriteLogEntry(FWriteLogEntry&) = default;
+        FWriteLogEntry& operator=(FWriteLogEntry&) = default;
+
+        explicit FWriteLogEntry(void* Original, size_t Size, void* Copy) :
+            OriginalAndSize(Original), Copy(Copy)
+        {
+            OriginalAndSize.SetTopTag(static_cast<uint16_t>(Size));
+        }
+    };
+
+    class FWriteLog final
+    {
+        struct FWriteLogEntryBucket final
+        {
+            static constexpr uint32_t BucketSize = 1024;
+
+            FWriteLogEntryBucket() = default;
+
+            FWriteLogEntry Entries[BucketSize];
+            size_t Size = 0;
+            FWriteLogEntryBucket* Next = nullptr;
+        };
+
+    public:
+        void Push(FWriteLogEntry Entry)
+        {
+            if (nullptr == Start)
+            {
+                Start = new FWriteLogEntryBucket();
+                Current = Start;
+            }
+
+            if (FWriteLogEntryBucket::BucketSize == Current->Size)
+            {
+                Current->Next = new FWriteLogEntryBucket();
+                Current = Current->Next;
+            }
+
+            Current->Entries[Current->Size++] = Entry;
+
+            TotalSize++;
+            return;
+        }
+
+        ~FWriteLog()
+        {
+            Reset();
+        }
+
+        struct Iterator final
+        {
+            explicit Iterator(FWriteLogEntryBucket* const Bucket) : Bucket(Bucket) {}
+
+            FWriteLogEntry& operator*()
+            {
+                return Bucket->Entries[Offset];
+            }
+
+            void operator++()
+            {
+                Offset++;
+
+                if (Offset == Bucket->Size)
+                {
+                    Bucket = Bucket->Next;
+                    Offset = 0;
+                }
+            }
+
+            bool operator!=(Iterator& Other) const
+            {
+                return (Other.Bucket != Bucket) || (Other.Offset != Offset);
+            }
+
+        private:
+            FWriteLogEntryBucket* Bucket;
+            size_t Offset = 0;
+        };
+
+        Iterator begin()
+        {
+            return Iterator(Start);
+        }
+
+        Iterator end()
+        {
+            return Iterator(nullptr);
+        }
+
+        void Reset()
+        {
+            while (nullptr != Start)
+            {
+                FWriteLogEntryBucket* const Old = Start;
+                Start = Start->Next;
+                delete Old;
+            }
+
+            Current = nullptr;
+            TotalSize = 0;
+        }
+
+        bool IsEmpty() const
+        {
+            return 0 == TotalSize;
+        }
+
+    private:
+
+        FWriteLogEntryBucket* Start = nullptr;
+        FWriteLogEntryBucket* Current = nullptr;
+        size_t TotalSize = 0;
+    };
+}
diff --git a/Engine/Source/Runtime/Core/Private/AutoRTFM/WriteLogBumpAllocator.h b/Engine/Source/Runtime/Core/Private/AutoRTFM/WriteLogBumpAllocator.h
new file mode 100644
--- /dev/null
+++ b/Engine/Source/Runtime/Core/Private/AutoRTFM/WriteLogBumpAllocator.h
@@ -0,0 +1,97 @@
+// Copyright Epic Games, Inc. All Rights Reserved.
+
+#pragma once
+
+#include <stdint.h>
+
+#include "Utils.h"
+
+namespace AutoRTFM
+{
+	class FWriteLogBumpAllocator final
+	{
+	public:
+		FWriteLogBumpAllocator()
+		{
+			Reset();
+		}
+
+		FWriteLogBumpAllocator(const FWriteLogBumpAllocator&) = delete;
+
+		~FWriteLogBumpAllocator()
+		{
+			Reset();
+		}
+
+		void* Allocate(size_t Bytes)
+		{
+			ASSERT(Bytes <= FPage::MaxSize);
+
+			if (nullptr == Start)
+			{
+				Start = new FPage();
+				Current = Start;
+			}
+
+			if (Bytes <= (FPage::MaxSize - Current->Size))
+			{
+				void* Result = &Current->Bytes[Current->Size];
+				Current->Size += Bytes;
+				return Result;
+			}
+
+			Current->Next = new FPage();
+			Current = Current->Next;
+
+			return Allocate(Bytes);
+		}
+
+		void Reset()
+		{
+			FPage* Page = Start;
+
+			while (nullptr != Page)
+			{
+				FPage* const Next = Page->Next;
+				delete Page;
+				Page = Next;
+			}
+
+			Start = nullptr;
+			Current = nullptr;
+		}
+
+		void Merge(FWriteLogBumpAllocator&& Other)
+		{
+			if (nullptr != Start)
+			{
+				Current->Next = Other.Start;
+			}
+			else
+			{
+				Start = Other.Start;
+				Current = Start;
+			}
+
+			Other.Start = nullptr;
+			Other.Current = nullptr;
+		}
+	private:
+
+		struct FPage final
+		{
+			// TODO: This is because we used tagged pointers elsewhere and so
+			// the max size of an allocation we will track in the bump allocator
+			// is 16 bits.
+			constexpr static size_t MaxSize = UINT16_MAX;
+
+			uint8_t Bytes[MaxSize];
+
+			FPage* Next = nullptr;
+			size_t Size = 0;
+		};
+
+		FPage* Start = nullptr;
+		FPage* Current = nullptr;
+	};
+}
diff --git a/Engine/Source/Runtime/Core/Public/AutoRTFM/AutoRTFM.h b/Engine/Source/Runtime/Core/Public/AutoRTFM/AutoRTFM.h
new file mode 100644
--- /dev/null
+++ b/Engine/Source/Runtime/Core/Public/AutoRTFM/AutoRTFM.h
@@ -0,0 +1,404 @@
+// Copyright Epic Games, Inc. All Rights Reserved.
+
+#ifndef AUTORTFM_H
+#define AUTORTFM_H
+
+#if (defined(__AUTORTFM) && __AUTORTFM)
+#define UE_AUTORTFM 1
+#else
+#define UE_AUTORTFM 0
+#endif
+
+#if UE_AUTORTFM
+	#define UE_PRAGMA_AUTORTFM _Pragma("autortfm")
+#else
+	#define UE_PRAGMA_AUTORTFM
+#endif
+
+#include <stdbool.h>
+#include <stdlib.h>
+
+#ifdef __cplusplus
+extern "C"
+{
+#endif
+
+// The C API exists for two reasons:
+//
+// - It makes linking easy. AutoRTFM has to deal with a weird kind of linking
+//   where the compiler directly emits calls to functions with a given name.
+//   It's easiest to do that in llvm if the functions have C linkage and C ABI.
+//
+// - It makes testing easy. Even seemingly simple C++ code introduces pitfalls
+//   for AutoRTFM. So very focused tests work best when written in C.
+//
+// We use snake_case for C API surface area to make it easy to distinguish.
+
+// This must match AutoRTFM::ETransactionResult.
+typedef enum
+{
+    autortfm_aborted_by_request,
+    autortfm_aborted_by_language,
+    autortfm_committed
+} autortfm_result;
+
+// Tells if we are currently running in a transaction. This will return true in an
+// open nest (see autortfm_open).
+#if UE_AUTORTFM
+bool autortfm_is_transactional(void);
+#else
+inline bool autortfm_is_transactional(void)
+{
+    return false;
+}
+#endif
+
+// Tells if we are currently running in the closed nest of a transaction. By
+// default, transactional code is in a closed nest; the only way to be in an open
+// nest is to request it via autortfm_open.
+//
+// The advantages of this function over autortfm_is_transactional are:
+//
+// - It's faster. Once all of the optimizations are implemented, the compiler will
+//   constant-fold this.
+//
+// - Usually, if you are doing special things for transactions, it's to work around
+//   the transactional openation in a closed nest. So, it's often more correct
+//   to test is_closed than is_transactional.
+#if UE_AUTORTFM
+bool autortfm_is_closed(void);
+#else
+inline bool autortfm_is_closed(void)
+{
+    return false;
+}
+#endif
+
+// Run the callback in a transaction. None of its effects will be visible until
+// it finishes, assuming it finishes without aborting. It's fine to call this
+// within a transaction. In that case, the nested transaction is only useful if
+// you might choose to abort it (then only the nested transaction aborts).
+#if UE_AUTORTFM
+autortfm_result autortfm_transact(void (*work)(void* arg), void* arg);
+#else
+inline autortfm_result autortfm_transact(void (*work)(void* arg), void* arg)
+{
+    return autortfm_aborted_by_language;
+}
+#endif
+
+// Run the callback in a transaction like autortfm_transact, but abort program
+// execution if the result is anything other than autortfm_committed. Useful for
+// testing.
+#if UE_AUTORTFM
+void autortfm_commit(void (*work)(void* arg), void* arg);
+#else
+inline void autortfm_commit(void (*work)(void* arg), void* arg)
+{
+    abort();
+}
+#endif
+
+// Request that the current transaction aborts. Only aborts the innermost
+// transaction nest. The corresponding autortfm_transact call will return
+// autortfm_aborted_by_request.
+//
+// Guaranteed to crash if called outside a transaction.
+#if UE_AUTORTFM
+void autortfm_abort(void);
+#else
+inline void autortfm_abort(void)
+{
+    abort();
+}
+#endif
+
+// Abort if running in a transaction. If called from closed code, this does
+// an abort_by_language, which will either abort the program or the transaction,
+// depending on configuration. If called from open code that is inside a
+// transaction, this aborts the program.
+#if UE_AUTORTFM
+void autortfm_abort_if_transactional(void);
+#else
+inline void autortfm_abort_if_transactional(void) { }
+#endif
+
+// Abort if running in closed code. This does an abort_by_language, which will
+// either abort the program or the transaction, depending on configuration.
+#if UE_AUTORTFM
+void autortfm_abort_if_closed(void);
+#else
+inline void autortfm_abort_if_closed(void) { }
+#endif
+
+// Executes the given code non-transactionally regardless of whether we are in
+// a transaction or not.
+#if UE_AUTORTFM
+void autortfm_open(void (*work)(void* arg), void* arg);
+#else
+inline void autortfm_open(void (*work)(void* arg), void* arg)
+{
+    work(arg);
+}
+#endif
+
+// Executes the given code transactionally while inside a autortfm_open. Also
+// executes the given code transactionally if called from transactional code.
+//
+// Guaranteed to crash if called outside a transaction.
+#if UE_AUTORTFM
+void autortfm_close(void (*work)(void* arg), void* arg);
+#else
+inline void autortfm_close(void (*work)(void* arg), void* arg)
+{
+    abort();
+}
+#endif
+
+// Register a transactional version of a function that wasn't compiled by the
+// autortfm compiler. Normally, code is transactionalized by the compiler by
+// emitting a clone that has transactional openation, with some magic to
+// redirect all function calls within a transaction to the transactional clone.
+// This allows you to hook in your own transactionalized implementations of
+// functions that the compiler did not see.
+//
+// Use with great caution!
+//
+// This results in calls to new_function to happen in open mode. We will call
+// new_function's nontransactional version within the transaction. This happens
+// with the additional caveat that if the original_function signature is:
+//
+//     foo (*original_function)(bar, baz)
+//
+// then the new_function signature is expected to be:
+//
+//     foo (*original_function)(bar, baz, void*)
+//
+// In general, a void* argument -- which you cannot introspect -- is added as
+// the last parameter.
+//
+// Note that the string describing the function may be any non-NULL string. It
+// does not have to be unique.
+#if UE_AUTORTFM
+void autortfm_register_open_function(void* original_function, void* new_function);
+#else
+inline void autortfm_register_open_function(void* original_function, void* new_function) { }
+#endif
+
+// Have some work happen when this transaction commits. For nested transactions,
+// this just adds the work to the work deferred until the outer nest's commit.
+// If this is called outside a transaction then the work happens immediately. An
+// open nest within a transaction will be deferred until commit too.
+#if UE_AUTORTFM
+void autortfm_defer_until_commit(void (*work)(void* arg), void* arg);
+#else
+inline void autortfm_defer_until_commit(void (*work)(void* arg), void* arg)
+{
+    work(arg);
+}
+#endif
+
+// Have some work happen when this transaction aborts. For nested transactions,
+// this work will be processed only when the outer nest actually commits or
+// aborts. If this is called outside a transaction then the work is ignored. An
+// open nest within a transaction will be deferred until abort too.
+#if UE_AUTORTFM
+void autortfm_defer_until_abort(void (*work)(void* arg), void* arg);
+#else
+inline void autortfm_defer_until_abort(void (*work)(void* arg), void* arg) { }
+#endif
+
+// Have some work happen when this transaction commits. For nested transactions,
+// this just adds the work to the work deferred until the outer nest's commit.
+// If this is called outside a transaction or from an open nest then the work
+// happens immediately.
+#if UE_AUTORTFM
+void autortfm_open_commit(void (*work)(void* arg), void* arg);
+#else
+inline void autortfm_open_commit(void (*work)(void* arg), void* arg)
+{
+    work(arg);
+}
+#endif
+
+// Have some work happen when this transaction aborts. If this is called
+// outside a transaction or from an open nest then the work is ignored.
+#if UE_AUTORTFM
+void autortfm_open_abort(void (*work)(void* arg), void* arg);
+#else
+inline void autortfm_open_abort(void (*work)(void* arg), void* arg) { }
+#endif
+
+// Inform the runtime that we have performed a new object allocation. It's only
+// necessary to call this inside of custom malloc implementations. As an
+// optimization, you can choose to then only have your malloc return the pointer
+// returned by this function. It's guaranteed to be equal to the pointer you
+// passed, but it's blessed specially from the compiler's perspective, leading
+// to some nice optimizations. This does nothing when called from open code.
+#if UE_AUTORTFM
+void* autortfm_did_allocate(void* ptr, size_t size);
+#else
+inline void* autortfm_did_allocate(void* ptr, size_t size)
+{
+    return ptr;
+}
+#endif
+
+// Inform the runtime that we have freed the object (meaning, we deferred it to
+// be freed in abort). This does nothing when called from open code.
+#if UE_AUTORTFM
+void autortfm_will_deallocate(void* ptr, size_t size);
+#else
+inline void autortfm_will_deallocate(void* ptr, size_t size) { }
+#endif
+
+// If running in a transaction, then perform a consistency check of the
+// transaction's read-write set. If possible, this compares the read-write set's
+// expected values with the actual values in global memory. Does nothing when
+// called outside of a transaction. May do nothing if debugging features aren't
+// enabled in the autortfm runtime.
+#if UE_AUTORTFM
+void autortfm_check_consistency_assuming_no_races(void);
+#else
+inline void autortfm_check_consistency_assuming_no_races(void) { }
+#endif
+
+// If running with AutoRTFM enabled, then perform an ABI check between the
+// AutoRTFM compiler and the AutoRTFM runtime, to ensure that memory is being
+// laid out in an identical manner between the AutoRTFM runtime and the AutoRTFM
+// compiler pass. Should not be called manually by the user, a call to this will
+// be injected by the compiler into a global constructor in the AutoRTFM compiled
+// code.
+#if UE_AUTORTFM
+void autortfm_check_abi(void* ptr, size_t size);
+#else
+inline void autortfm_check_abi(void* ptr, size_t size) { }
+#endif
+
+#ifdef __cplusplus
+}
+
+#include <functional>
+
+namespace AutoRTFM
+{
+
+// This must match autortfm_result.
+enum class ETransactionResult
+{
+    AbortedByRequest = autortfm_aborted_by_request,
+    AbortedByLanguage = autortfm_aborted_by_language,
+    Committed = autortfm_committed
+};
+
+inline bool IsTransactional() { return autortfm_is_transactional(); }
+inline bool IsClosed() { return autortfm_is_closed(); }
+
+template<typename TFunctor>
+ETransactionResult Transact(const TFunctor& Functor)
+{
+    return static_cast<ETransactionResult>(autortfm_transact(
+        [] (void* Arg) { (*static_cast<const TFunctor*>(Arg))(); },
+        const_cast<void*>(static_cast<const void*>(&Functor))));
+}
+
+template<typename TFunctor>
+void Commit(const TFunctor& Functor)
+{
+    autortfm_commit(
+        [] (void* Arg) { (*static_cast<const TFunctor*>(Arg))(); },
+        const_cast<void*>(static_cast<const void*>(&Functor)));
+}
+
+inline void Abort()
+{
+    autortfm_abort();
+}
+
+inline void AbortIfTransactional()
+{
+    autortfm_abort_if_transactional();
+}
+
+inline void AbortIfClosed()
+{
+    autortfm_abort_if_closed();
+}
+
+template<typename TFunctor>
+void Open(const TFunctor& Functor)
+{
+    autortfm_open(
+        [] (void* Arg) { (*static_cast<const TFunctor*>(Arg))(); },
+        const_cast<void*>(static_cast<const void*>(&Functor)));
+}
+
+template<typename TFunctor>
+void Close(const TFunctor& Functor)
+{
+    autortfm_close(
+        [] (void* Arg) { (*static_cast<const TFunctor*>(Arg))(); },
+        const_cast<void*>(static_cast<const void*>(&Functor)));
+}
+
+inline void RegisterOpenFunction(void* OriginalFunction, void* NewFunction)
+{
+    autortfm_register_open_function(OriginalFunction, NewFunction);
+}
+
+#if UE_AUTORTFM
+void DeferUntilCommit(std::function<void()>&& Work);
+void DeferUntilAbort(std::function<void()>&& Work);
+void OpenCommit(std::function<void()>&& Work);
+void OpenAbort(std::function<void()>&& Work);
+#else
+template<typename TFunctor>
+void DeferUntilCommit(const TFunctor& Work) { Work(); }
+template<typename TFunctor>
+void DeferUntilAbort(const TFunctor& Work) { }
+template<typename TFunctor>
+void OpenCommit(const TFunctor& Work) { Work(); }
+template<typename TFunctor>
+void OpenAbort(const TFunctor& Work) { }
+#endif
+
+inline void* DidAllocate(void* Ptr, size_t Size)
+{
+    return autortfm_did_allocate(Ptr, Size);
+}
+
+inline void WillDeallocate(void* Ptr, size_t Size)
+{
+    autortfm_will_deallocate(Ptr, Size);
+}
+
+inline void CheckConsistencyAssumingNoRaces()
+{
+    autortfm_check_consistency_assuming_no_races();
+}
+
+struct FRegisterOpenFunction
+{
+    FRegisterOpenFunction(void* OriginalFunction, void* NewFunction)
+    {
+        RegisterOpenFunction(OriginalFunction, NewFunction);
+    }
+};
+
+#define UE_AUTORTFM_CONCAT_IMPL(A, B) A ## B
+#define UE_AUTORTFM_CONCAT(A, B) UE_AUTORTFM_CONCAT_IMPL(A, B)
+
+#if UE_AUTORTFM
+#define UE_AUTORTFM_REGISTER_OPEN_FUNCTION_EXPLICIT(OriginalFunction, NewFunction) static const FRegisterOpenFunction UE_AUTORTFM_CONCAT(AutoRTFMFunctionRegistration, __COUNTER__)(reinterpret_cast<void*>(OriginalFunction), reinterpret_cast<void*>(NewFunction))
+#define UE_AUTORTFM_REGISTER_OPEN_FUNCTION(OriginalFunction) UE_AUTORTFM_REGISTER_OPEN_FUNCTION_EXPLICIT(OriginalFunction, STM_ ## OriginalFunction)
+#define UE_AUTORTFM_REGISTER_SELF_FUNCTION(OriginalFunction) UE_AUTORTFM_REGISTER_OPEN_FUNCTION_EXPLICIT(OriginalFunction, OriginalFunction)
+#else
+#define UE_AUTORTFM_REGISTER_OPEN_FUNCTION_EXPLICIT(OriginalFunction, NewFunction)
+#define UE_AUTORTFM_REGISTER_OPEN_FUNCTION(OriginalFunction)
+#define UE_AUTORTFM_REGISTER_SELF_FUNCTION(OriginalFunction)
+#endif
+
+} // namespace AutoRTFM
+#endif // __cplusplus
+
+#endif /* AUTORTFM_H */
diff --git a/Engine/Source/Runtime/Core/Public/AutoRTFM/AutoRTFMConstants.h b/Engine/Source/Runtime/Core/Public/AutoRTFM/AutoRTFMConstants.h
new file mode 100644
--- /dev/null
+++ b/Engine/Source/Runtime/Core/Public/AutoRTFM/AutoRTFMConstants.h
@@ -0,0 +1,36 @@
+// Copyright Epic Games, Inc. All Rights Reserved.
+
+#pragma once
+
+#include <cstddef>
+#include <cstdint>
+
+// FIXME: This is going to need support for cross-compilation. So, we will possibly have different offsets for
+// different platforms, and both the compiler and the runtime will have to make choices about which to use.
+
+// WARNING: Any change in these constants will require a re-patch and re-build of LLVM!
+
+namespace AutoRTFM
+{
+namespace Constants
+{
+
+constexpr size_t LogLineBytes = 3;
+constexpr size_t LineBytes = static_cast<size_t>(1) << LogLineBytes;
+constexpr size_t LineTableSize = 33554432;
+
+constexpr size_t Offset_Context_CurrentTransaction = 0;
+constexpr size_t Offset_Context_LineTable = 192;
+constexpr size_t Offset_Context_Status = 152;
+
+constexpr size_t LogSize_LineEntry = 5;
+constexpr size_t Size_LineEntry = static_cast<size_t>(1) << LogSize_LineEntry;
+constexpr size_t Offset_LineEntry_LogicalLine = 0;
+constexpr size_t Offset_LineEntry_ActiveLine = 16;
+constexpr size_t Offset_LineEntry_LoggingTransaction = 24;
+constexpr size_t Offset_LineEntry_AccessMask = 14;
+
+constexpr uint32_t Context_Status_OnTrack = 1;
+
+} // namespace Constants
+} // namespace AutoRTFM
