autortfm: runtime build fixes

diff --git a/Engine/Source/Runtime/Core/Private/AutoRTFM/Utils.cpp b/Engine/Source/Runtime/Core/Private/AutoRTFM/Utils.cpp
--- a/Engine/Source/Runtime/Core/Private/AutoRTFM/Utils.cpp
+++ b/Engine/Source/Runtime/Core/Private/AutoRTFM/Utils.cpp
@@ -8,6 +8,7 @@
 #include <string.h>
 
 #ifdef _WIN32
+#include <windows.h>
 #include <dbghelp.h>
 #else
 #include <execinfo.h>
diff --git a/Engine/Source/Runtime/Core/Private/AutoRTFM/Utils.h b/Engine/Source/Runtime/Core/Private/AutoRTFM/Utils.h
--- a/Engine/Source/Runtime/Core/Private/AutoRTFM/Utils.h
+++ b/Engine/Source/Runtime/Core/Private/AutoRTFM/Utils.h
@@ -8,12 +8,7 @@
 #include <stdio.h>
 #include <stdlib.h>
 #include <string>
-#ifdef _WIN32
-#define NOMINMAX 1
-#include <windows.h>
-#else
 #include <malloc.h>
-#endif
 
 #define ASSERT(exp) do { \
     if (!(exp)) { \
