Refactor pragma code to support Clang-cl.

Clang does not allow the text following #pragma to be expanded, and does not allow the standard _Pragma operator with multiple string literals, but does support the MSVC extension __pragma.

#tested MSVC in Debug/Test/Shipping, Clang 14 in Debug/Test/Shipping.
#rnx

diff --git a/Engine/Source/Runtime/Core/Private/MemPro/MemPro.cpp b/Engine/Source/Runtime/Core/Private/MemPro/MemPro.cpp
--- a/Engine/Source/Runtime/Core/Private/MemPro/MemPro.cpp
+++ b/Engine/Source/Runtime/Core/Private/MemPro/MemPro.cpp
@@ -136,8 +136,8 @@ namespace MemPro
 //------------------------------------------------------------------------
 // disable some warnings we are not interested in so that we can compile at warning level4
 #ifdef MEMPRO_ENABLE_WARNING_PRAGMAS
-	#pragma MEMPRO_PUSH_WARNING_DISABLE
-	#pragma MEMPRO_DISABLE_WARNING(4100)
+	MEMPRO_PUSH_WARNING_DISABLE
+	MEMPRO_DISABLE_WARNING(4100)
 #endif
 
 //------------------------------------------------------------------------
@@ -201,7 +201,7 @@ namespace MemPro
 
 //------------------------------------------------------------------------
 #ifdef MEMPRO_ENABLE_WARNING_PRAGMAS
-	#pragma MEMPRO_POP_WARNING_DISABLE
+	MEMPRO_POP_WARNING_DISABLE
 #endif
 
 //------------------------------------------------------------------------
@@ -712,8 +712,8 @@ namespace MemPro
 //------------------------------------------------------------------------
 // disable some warnings we are not interested in so that we can compile at warning level4
 #ifdef MEMPRO_ENABLE_WARNING_PRAGMAS
-	#pragma MEMPRO_PUSH_WARNING_DISABLE
-	#pragma MEMPRO_DISABLE_WARNING(4127)
+	MEMPRO_PUSH_WARNING_DISABLE
+	MEMPRO_DISABLE_WARNING(4127)
 #endif
 
 //------------------------------------------------------------------------
@@ -799,7 +799,7 @@ namespace MemPro
 
 //------------------------------------------------------------------------
 #ifdef MEMPRO_ENABLE_WARNING_PRAGMAS
-	#pragma MEMPRO_POP_WARNING_DISABLE
+	MEMPRO_POP_WARNING_DISABLE
 #endif
 
 //------------------------------------------------------------------------
@@ -1173,9 +1173,9 @@ namespace MemPro
 //------------------------------------------------------------------------
 // disable some warnings we are not interested in so that we can compile at warning level4
 #ifdef MEMPRO_ENABLE_WARNING_PRAGMAS
-	#pragma MEMPRO_PUSH_WARNING_DISABLE
-	#pragma MEMPRO_DISABLE_WARNING(4127)
-	#pragma MEMPRO_DISABLE_WARNING(4100)
+	MEMPRO_PUSH_WARNING_DISABLE
+	MEMPRO_DISABLE_WARNING(4127)
+	MEMPRO_DISABLE_WARNING(4100)
 #endif
 
 //------------------------------------------------------------------------
@@ -2657,7 +2657,7 @@ void MemPro::FlushDumpFile()
 
 //------------------------------------------------------------------------
 #ifdef MEMPRO_ENABLE_WARNING_PRAGMAS
-	#pragma MEMPRO_POP_WARNING_DISABLE
+	MEMPRO_POP_WARNING_DISABLE
 #endif
 
 //------------------------------------------------------------------------
diff --git a/Engine/Source/Runtime/Core/Public/MemPro/MemPro.h b/Engine/Source/Runtime/Core/Public/MemPro/MemPro.h
--- a/Engine/Source/Runtime/Core/Public/MemPro/MemPro.h
+++ b/Engine/Source/Runtime/Core/Public/MemPro/MemPro.h
@@ -236,9 +236,9 @@ namespace
 	#define MEMPRO_INTERLOCKED_ALIGN __declspec(align(8))
 	#define MEMPRO_INSTRUCTION_BARRIER
 	#define MEMPRO_ENABLE_WARNING_PRAGMAS
-	#define MEMPRO_PUSH_WARNING_DISABLE warning(push)
-	#define MEMPRO_DISABLE_WARNING(w) warning(disable : w)
-	#define MEMPRO_POP_WARNING_DISABLE warning(pop)
+	#define MEMPRO_PUSH_WARNING_DISABLE __pragma(warning(push))
+	#define MEMPRO_DISABLE_WARNING(w) __pragma(warning(disable : w))
+	#define MEMPRO_POP_WARNING_DISABLE __pragma(warning(pop))
 	#define MEMPRO_FORCEINLINE FORCEINLINE
 	#define ENUMERATE_ALL_MODULES // if you are having problems compiling this on your platform undefine ENUMERATE_ALL_MODULES and it send info for just the main module
 	#define THREAD_LOCAL_STORAGE __declspec(thread)
