Converted TimespanTest to use LowLevelTests

diff --git a/Engine/Source/Runtime/Core/Private/Tests/Misc/TimespanTest.cpp b/Engine/Source/Runtime/Core/Private/Tests/Misc/TimespanTest.cpp
--- a/Engine/Source/Runtime/Core/Private/Tests/Misc/TimespanTest.cpp
+++ b/Engine/Source/Runtime/Core/Private/Tests/Misc/TimespanTest.cpp
@@ -1,70 +1,66 @@
 // Copyright Epic Games, Inc. All Rights Reserved.
 
-#include "CoreTypes.h"
-#include "Containers/UnrealString.h"
-#include "Misc/Timespan.h"
-#include "Misc/AutomationTest.h"
-
-
-#if WITH_DEV_AUTOMATION_TESTS
+#if WITH_LOW_LEVEL_TESTS
 
-IMPLEMENT_SIMPLE_AUTOMATION_TEST(FTimespanTest, "System.Core.Misc.Timespan", EAutomationTestFlags::ApplicationContextMask | EAutomationTestFlags::SmokeFilter)
+#include "Misc/Timespan.h"
 
+#include "Containers/UnrealString.h"
+#include "TestHarness.h"
 
-bool FTimespanTest::RunTest(const FString& Parameters)
+TEST_CASE("Core::Time::Timespan", "[Core][Time][Smoke]")
 {
-	// constructors must create equal objects
+	SECTION("Constructors")
 	{
 		FTimespan ts1_1 = FTimespan(3, 2, 1);
 		FTimespan ts1_2 = FTimespan(0, 3, 2, 1);
 		FTimespan ts1_3 = FTimespan(0, 3, 2, 1, 0);
 
-		TestEqual(TEXT("Constructors must create equal objects (Hours/Minutes/Seconds vs. Days/Hours/Minutes/Seconds)"), ts1_1, ts1_2);
-		TestEqual(TEXT("Constructors must create equal objects (Hours/Minutes/Seconds vs. Days/Hours/Minutes/Seconds/FractionNano)"), ts1_1, ts1_3);
+		CHECK(ts1_1 == ts1_2);
+		CHECK(ts1_1 == ts1_3);
 	}
 
-	// component getters must return correct values
+	SECTION("Components")
 	{
 		FTimespan ts2_1 = FTimespan(1, 2, 3, 4, 123456789);
 
-		TestEqual(TEXT("Component getters must return correct values (Days)"), ts2_1.GetDays(), 1);
-		TestEqual(TEXT("Component getters must return correct values (Hours)"), ts2_1.GetHours(), 2);
-		TestEqual(TEXT("Component getters must return correct values (Minutes)"), ts2_1.GetMinutes(), 3);
-		TestEqual(TEXT("Component getters must return correct values (Seconds)"), ts2_1.GetSeconds(), 4);
-		TestEqual(TEXT("Component getters must return correct values (FractionMilli)"), ts2_1.GetFractionMilli(), 123);
-		TestEqual(TEXT("Component getters must return correct values (FractionMicro)"), ts2_1.GetFractionMicro(), 123456);
-		TestEqual(TEXT("Component getters must return correct values (FractionNano)"), ts2_1.GetFractionNano(), 123456700);
+		CHECK(ts2_1.GetDays() == 1);
+		CHECK(ts2_1.GetHours() == 2);
+		CHECK(ts2_1.GetMinutes() == 3);
+		CHECK(ts2_1.GetSeconds() == 4);
+		CHECK(ts2_1.GetFractionMilli() == 123);
+		CHECK(ts2_1.GetFractionMicro() == 123456);
+		CHECK(ts2_1.GetFractionNano() == 123456700);
 	}
 
-	// durations of positive and negative time spans must match
+	SECTION("Duration")
 	{
 		FTimespan ts3_1 = FTimespan(1, 2, 3, 4, 123456789);
 		FTimespan ts3_2 = FTimespan(-1, -2, -3, -4, -123456789);
 
-		TestEqual(TEXT("Durations of positive and negative time spans must match"), ts3_1.GetDuration(), ts3_2.GetDuration());
+		CHECK(ts3_1.GetDuration() == ts3_2.GetDuration());
 	}
 
-	// static constructors must create correct values
+	SECTION("FromXxx to GetTotalXxx")
 	{
-		TestEqual(TEXT("Static constructors must create correct values (FromDays)"), FTimespan::FromDays(123).GetTotalDays(), 123.0);
-		TestEqual(TEXT("Static constructors must create correct values (FromHours)"), FTimespan::FromHours(123).GetTotalHours(), 123.0);
-		TestEqual(TEXT("Static constructors must create correct values (FromMinutes)"), FTimespan::FromMinutes(123).GetTotalMinutes(), 123.0);
-		TestEqual(TEXT("Static constructors must create correct values (FromSeconds)"), FTimespan::FromSeconds(123).GetTotalSeconds(), 123.0);
-		TestEqual(TEXT("Static constructors must create correct values (FromMilliseconds)"), FTimespan::FromMilliseconds(123).GetTotalMilliseconds(), 123.0);
-		TestEqual(TEXT("Static constructors must create correct values (FromMicroseconds)"), FTimespan::FromMicroseconds(123).GetTotalMicroseconds(), 123.0);
+		CHECK(FTimespan::FromDays(123).GetTotalDays() == 123.0);
+		CHECK(FTimespan::FromHours(123).GetTotalHours() == 123.0);
+		CHECK(FTimespan::FromMinutes(123).GetTotalMinutes() == 123.0);
+		CHECK(FTimespan::FromSeconds(123).GetTotalSeconds() == 123.0);
+		CHECK(FTimespan::FromMilliseconds(123).GetTotalMilliseconds() == 123.0);
+		CHECK(FTimespan::FromMicroseconds(123).GetTotalMicroseconds() == 123.0);
 	}
 
-	// string conversions must return correct strings
+	SECTION("ToString")
 	{
 		FTimespan ts5_1 = FTimespan(1, 2, 3, 4, 123456789);
 
-		TestEqual<FString>(TEXT("String conversion (Default)"), ts5_1.ToString(), TEXT("+1.02:03:04.123"));
-		TestEqual<FString>(TEXT("String conversion (%d.%h:%m:%s.%f)"), ts5_1.ToString(TEXT("%d.%h:%m:%s.%f")), TEXT("+1.02:03:04.123"));
-		TestEqual<FString>(TEXT("String conversion (%d.%h:%m:%s.%u)"), ts5_1.ToString(TEXT("%d.%h:%m:%s.%u")), TEXT("+1.02:03:04.123456"));
-		TestEqual<FString>(TEXT("String conversion (%D.%h:%m:%s.%n)"), ts5_1.ToString(TEXT("%D.%h:%m:%s.%n")), TEXT("+00000001.02:03:04.123456700"));
+		CHECK(ts5_1.ToString() == TEXT("+1.02:03:04.123"));
+		CHECK(ts5_1.ToString(TEXT("%d.%h:%m:%s.%f")) == TEXT("+1.02:03:04.123"));
+		CHECK(ts5_1.ToString(TEXT("%d.%h:%m:%s.%u")) == TEXT("+1.02:03:04.123456"));
+		CHECK(ts5_1.ToString(TEXT("%D.%h:%m:%s.%n")) == TEXT("+00000001.02:03:04.123456700"));
 	}
 
-	// parsing valid strings must succeed
+	SECTION("Parse Valid")
 	{
 		FTimespan ts6_t;
 
@@ -72,68 +68,66 @@ bool FTimespanTest::RunTest(const FString& Parameters)
 		FTimespan ts6_2 = FTimespan(1, 2, 3, 4, 123456000);
 		FTimespan ts6_3 = FTimespan(1, 2, 3, 4, 123456700);
 
-		TestTrue(TEXT("Parsing valid strings must succeed (+1.02:03:04.123)"), FTimespan::Parse(TEXT("+1.02:03:04.123"), ts6_t));
-		TestEqual(TEXT("Parsing valid strings must result in correct values (+1.02:03:04.123)"), ts6_t, ts6_1);
+		CHECK(FTimespan::Parse(TEXT("+1.02:03:04.123"), ts6_t));
+		CHECK(ts6_t == ts6_1);
 
-		TestTrue(TEXT("Parsing valid strings must succeed (+1.02:03:04.123456)"), FTimespan::Parse(TEXT("+1.02:03:04.123456"), ts6_t));
-		TestEqual(TEXT("Parsing valid strings must result in correct values (+1.02:03:04.123456)"), ts6_t, ts6_2);
+		CHECK(FTimespan::Parse(TEXT("+1.02:03:04.123456"), ts6_t));
+		CHECK(ts6_t == ts6_2);
 
-		TestTrue(TEXT("Parsing valid strings must succeed (+1.02:03:04.123456789)"), FTimespan::Parse(TEXT("+1.02:03:04.123456789"), ts6_t));
-		TestEqual(TEXT("Parsing valid strings must result in correct values (+1.02:03:04.123456789)"), ts6_t, ts6_3);
+		CHECK(FTimespan::Parse(TEXT("+1.02:03:04.123456789"), ts6_t));
+		CHECK(ts6_t == ts6_3);
 
 		FTimespan ts6_4 = FTimespan(-1, -2, -3, -4, -123000000);
 		FTimespan ts6_5 = FTimespan(-1, -2, -3, -4, -123456000);
 		FTimespan ts6_6 = FTimespan(-1, -2, -3, -4, -123456700);
 
-		TestTrue(TEXT("Parsing valid strings must succeed (-1.02:03:04.123)"), FTimespan::Parse(TEXT("-1.02:03:04.123"), ts6_t));
-		TestEqual(TEXT("Parsing valid strings must result in correct values (-1.02:03:04.123)"), ts6_t, ts6_4);
+		CHECK(FTimespan::Parse(TEXT("-1.02:03:04.123"), ts6_t));
+		CHECK(ts6_t == ts6_4);
 
-		TestTrue(TEXT("Parsing valid strings must succeed (-1.02:03:04.123456)"), FTimespan::Parse(TEXT("-1.02:03:04.123456"), ts6_t));
-		TestEqual(TEXT("Parsing valid strings must result in correct values (-1.02:03:04.123456)"), ts6_t, ts6_5);
+		CHECK(FTimespan::Parse(TEXT("-1.02:03:04.123456"), ts6_t));
+		CHECK(ts6_t == ts6_5);
 
-		TestTrue(TEXT("Parsing valid strings must succeed (-1.02:03:04.123456789)"), FTimespan::Parse(TEXT("-1.02:03:04.123456789"), ts6_t));
-		TestEqual(TEXT("Parsing valid strings must result in correct values (-1.02:03:04.123456789)"), ts6_t, ts6_6);
+		CHECK(FTimespan::Parse(TEXT("-1.02:03:04.123456789"), ts6_t));
+		CHECK(ts6_t == ts6_6);
 	}
 
-	// parsing invalid strings must fail
+	SECTION("Parse Invalid")
 	{
 		FTimespan ts7_1;
 
-		//TestFalse(TEXT("Parsing invalid strings must fail (1,02:03:04.005)"), FTimespan::Parse(TEXT("1,02:03:04.005"), ts7_1));
-		//TestFalse(TEXT("Parsing invalid strings must fail (1.1.02:03:04:005)"), FTimespan::Parse(TEXT("1.1.02:03:04:005"), ts7_1));
-		//TestFalse(TEXT("Parsing invalid strings must fail (04:005)"), FTimespan::Parse(TEXT("04:005"), ts7_1));
+		//CHECK_FALSE(FTimespan::Parse(TEXT("1,02:03:04.005"), ts7_1));
+		//CHECK_FALSE(FTimespan::Parse(TEXT("1.1.02:03:04:005"), ts7_1));
+		//CHECK_FALSE(FTimespan::Parse(TEXT("04:005"), ts7_1));
 	}
 
 	// `From*` converters must return correct value
 	// test normal and edge cases for polar conversions (FromMicroseconds() and FromDay()) and just normal case for all others
+	SECTION("FromXxx Exact Ticks")
 	{
-		TestEqual(TEXT("FromMicroseconds(0) results in correct value"), FTimespan::FromMicroseconds(0), FTimespan(0));
-		TestEqual(TEXT("FromMicroseconds(1) results in correct value"), FTimespan::FromMicroseconds(1), FTimespan(1 * ETimespan::TicksPerMicrosecond));
-		TestEqual(TEXT("FromMicroseconds(1.1) results in correct value"), FTimespan::FromMicroseconds(1.1), FTimespan(1 * ETimespan::TicksPerMicrosecond + 1));
-		TestEqual(TEXT("FromMicroseconds(1.5) results in correct value"), FTimespan::FromMicroseconds(1.5), FTimespan(1 * ETimespan::TicksPerMicrosecond + 5));
-		TestEqual(TEXT("FromMicroseconds(1.499999999999997) results in 1.5 microsecs of ticks"), FTimespan::FromMicroseconds(1.499999999999997), FTimespan(1 * ETimespan::TicksPerMicrosecond + 5));
-		TestEqual(TEXT("FromMicroseconds(1.50000001) results in  in 1.5 microsecs of ticks"), FTimespan::FromMicroseconds(1.50000001), FTimespan(1 * ETimespan::TicksPerMicrosecond + 5));
-		TestEqual(TEXT("FromMicroseconds(-1) results in correct value"), FTimespan::FromMicroseconds(-1), FTimespan(-1 * ETimespan::TicksPerMicrosecond));
-		TestEqual(TEXT("FromMicroseconds(-1.1) results in correct value"), FTimespan::FromMicroseconds(-1.1), FTimespan(-1 * ETimespan::TicksPerMicrosecond - 1));
-		TestEqual(TEXT("FromMicroseconds(-1.5) results in correct value"), FTimespan::FromMicroseconds(-1.5), FTimespan(-1 * ETimespan::TicksPerMicrosecond - 5));
-
-		TestEqual(TEXT("FromMilliseconds(1) results in correct value"), FTimespan::FromMilliseconds(1), FTimespan(1 * ETimespan::TicksPerMillisecond));
-		TestEqual(TEXT("FromSeconds(1) results in correct value"), FTimespan::FromSeconds(1), FTimespan(1 * ETimespan::TicksPerSecond));
-		TestEqual(TEXT("FromMinutes(1) results in correct value"), FTimespan::FromMinutes(1), FTimespan(1 * ETimespan::TicksPerMinute));
-		TestEqual(TEXT("FromHours(1) results in correct value"), FTimespan::FromHours(1), FTimespan(1 * ETimespan::TicksPerHour));
-
-		TestEqual(TEXT("FromDays(0) results in correct value"), FTimespan::FromDays(0), FTimespan(0));
-		TestEqual(TEXT("FromDays(1) results in correct value"), FTimespan::FromDays(1), FTimespan(1 * ETimespan::TicksPerDay));
-		TestEqual(TEXT("FromDays(1.25) results in correct value (1 day and 6 hours of ticks)"), FTimespan::FromDays(1.25), FTimespan(1 * ETimespan::TicksPerDay + 6 * ETimespan::TicksPerHour));
-		TestEqual(TEXT("FromDays(1.5) results in correct value (1 day and 12 hours of ticks)"), FTimespan::FromDays(1.5), FTimespan(1 * ETimespan::TicksPerDay + 12 * ETimespan::TicksPerHour));
-		TestEqual(TEXT("FromDays(1.499999999999997) results in correct value (1 day and 12 hours of ticks)"), FTimespan::FromDays(1.499999999999997), FTimespan(1 * ETimespan::TicksPerDay + 12 * ETimespan::TicksPerHour));
-		TestEqual(TEXT("FromDays(-1) results in correct value"), FTimespan::FromDays(-1), FTimespan(-1 * ETimespan::TicksPerDay));
-		TestEqual(TEXT("FromDays(-1.25) results in correct value (minus 1 day and 6 hours of ticks)"), FTimespan::FromDays(-1.25), FTimespan(-1 * ETimespan::TicksPerDay - 6 * ETimespan::TicksPerHour));
-		TestEqual(TEXT("FromDays(-1.5) results in correct value (minus 1 day and 12 hours of ticks)"), FTimespan::FromDays(-1.5), FTimespan(-1 * ETimespan::TicksPerDay - 12 * ETimespan::TicksPerHour));
+		CHECK(FTimespan::FromMicroseconds(0) == FTimespan(0));
+		CHECK(FTimespan::FromMicroseconds(1) == FTimespan(1 * ETimespan::TicksPerMicrosecond));
+		CHECK(FTimespan::FromMicroseconds(1.1) == FTimespan(1 * ETimespan::TicksPerMicrosecond + 1));
+		CHECK(FTimespan::FromMicroseconds(1.5) == FTimespan(1 * ETimespan::TicksPerMicrosecond + 5));
+		CHECK(FTimespan::FromMicroseconds(1.499999999999997) == FTimespan(1 * ETimespan::TicksPerMicrosecond + 5));
+		CHECK(FTimespan::FromMicroseconds(1.50000001) == FTimespan(1 * ETimespan::TicksPerMicrosecond + 5));
+		CHECK(FTimespan::FromMicroseconds(-1) == FTimespan(-1 * ETimespan::TicksPerMicrosecond));
+		CHECK(FTimespan::FromMicroseconds(-1.1) == FTimespan(-1 * ETimespan::TicksPerMicrosecond - 1));
+		CHECK(FTimespan::FromMicroseconds(-1.5) == FTimespan(-1 * ETimespan::TicksPerMicrosecond - 5));
+
+		CHECK(FTimespan::FromMilliseconds(1) == FTimespan(1 * ETimespan::TicksPerMillisecond));
+		CHECK(FTimespan::FromSeconds(1) == FTimespan(1 * ETimespan::TicksPerSecond));
+		CHECK(FTimespan::FromMinutes(1) == FTimespan(1 * ETimespan::TicksPerMinute));
+		CHECK(FTimespan::FromHours(1) == FTimespan(1 * ETimespan::TicksPerHour));
+
+		CHECK(FTimespan::FromDays(0) == FTimespan(0));
+		CHECK(FTimespan::FromDays(1) == FTimespan(1 * ETimespan::TicksPerDay));
+		CHECK(FTimespan::FromDays(1.25) == FTimespan(1 * ETimespan::TicksPerDay + 6 * ETimespan::TicksPerHour));
+		CHECK(FTimespan::FromDays(1.5) == FTimespan(1 * ETimespan::TicksPerDay + 12 * ETimespan::TicksPerHour));
+		CHECK(FTimespan::FromDays(1.499999999999997) == FTimespan(1 * ETimespan::TicksPerDay + 12 * ETimespan::TicksPerHour));
+		CHECK(FTimespan::FromDays(-1) == FTimespan(-1 * ETimespan::TicksPerDay));
+		CHECK(FTimespan::FromDays(-1.25) == FTimespan(-1 * ETimespan::TicksPerDay - 6 * ETimespan::TicksPerHour));
+		CHECK(FTimespan::FromDays(-1.5) == FTimespan(-1 * ETimespan::TicksPerDay - 12 * ETimespan::TicksPerHour));
 	}
-
-	return true;
 }
 
-
-#endif //WITH_DEV_AUTOMATION_TESTS
+#endif //WITH_LOW_LEVEL_TESTS
