Added FCycleTimePoint and FCycleTimeSpan as an API for the monotonic clock exposed by FPlatformTime::Cycles64

diff --git a/Engine/Source/Runtime/Core/Private/Misc/CycleTime.cpp b/Engine/Source/Runtime/Core/Private/Misc/CycleTime.cpp
new file mode 100644
--- /dev/null
+++ b/Engine/Source/Runtime/Core/Private/Misc/CycleTime.cpp
@@ -0,0 +1,41 @@
+// Copyright Epic Games, Inc. All Rights Reserved.
+
+#include "Misc/CycleTime.h"
+
+#include "HAL/PlatformTime.h"
+#include "Misc/Timespan.h"
+
+namespace UE
+{
+
+FCycleTimeSpan::FCycleTimeSpan(const FTimespan Span)
+	: Cycles(int64(Span.GetTotalSeconds() / FPlatformTime::GetSecondsPerCycle64()))
+{
+}
+
+double FCycleTimeSpan::ToSeconds() const
+{
+	return FPlatformTime::GetSecondsPerCycle64() * double(Cycles);
+}
+
+double FCycleTimeSpan::ToMilliseconds() const
+{
+	return FPlatformTime::GetSecondsPerCycle64() * 1000.0 * double(Cycles);
+}
+
+FCycleTimeSpan FCycleTimeSpan::FromSeconds(const double Seconds)
+{
+	return FromCycles(int64(Seconds / FPlatformTime::GetSecondsPerCycle64()));
+}
+
+FCycleTimeSpan FCycleTimeSpan::FromMilliseconds(const double Milliseconds)
+{
+	return FromCycles(int64(Milliseconds / 1000.0 / FPlatformTime::GetSecondsPerCycle64()));
+}
+
+FCycleTimePoint FCycleTimePoint::Now()
+{
+	return FromCycles(FPlatformTime::Cycles64());
+}
+
+} // namespace UE
diff --git a/Engine/Source/Runtime/Core/Public/Misc/CycleTime.h b/Engine/Source/Runtime/Core/Public/Misc/CycleTime.h
new file mode 100644
--- /dev/null
+++ b/Engine/Source/Runtime/Core/Public/Misc/CycleTime.h
@@ -0,0 +1,136 @@
+// Copyright Epic Games, Inc. All Rights Reserved.
+
+#pragma once
+
+#include "CoreTypes.h"
+#include "Math/NumericLimits.h"
+
+#define UE_API CORE_API
+
+struct FTimespan;
+
+namespace UE
+{
+
+/**
+ * A span of time between two values of FPlatformTime::Cycles64().
+ *
+ * @see FCycleTimePoint
+ */
+struct FCycleTimeSpan
+{
+public:
+	UE_API explicit FCycleTimeSpan(FTimespan Span);
+
+	constexpr static FCycleTimeSpan Zero() { return FCycleTimeSpan(); }
+
+	constexpr static FCycleTimeSpan Infinity() { return FromCycles(MAX_int64); }
+
+	constexpr static FCycleTimeSpan FromCycles(const int64 InCycles)
+	{
+		FCycleTimeSpan TimeSpan;
+		TimeSpan.Cycles = InCycles;
+		return TimeSpan;
+	}
+
+	constexpr FCycleTimeSpan() = default;
+
+	constexpr int64 GetCycles() const { return Cycles; }
+
+	constexpr bool IsInfinity() const { return (Cycles == MAX_int64) | (Cycles == MIN_int64); }
+
+	constexpr bool operator==(const FCycleTimeSpan Other) const { return Cycles == Other.Cycles; }
+	constexpr bool operator!=(const FCycleTimeSpan Other) const { return Cycles != Other.Cycles; }
+	constexpr bool operator<=(const FCycleTimeSpan Other) const { return Cycles <= Other.Cycles; }
+	constexpr bool operator< (const FCycleTimeSpan Other) const { return Cycles <  Other.Cycles; }
+	constexpr bool operator>=(const FCycleTimeSpan Other) const { return Cycles >= Other.Cycles; }
+	constexpr bool operator> (const FCycleTimeSpan Other) const { return Cycles >  Other.Cycles; }
+
+	constexpr FCycleTimeSpan operator+() const { return *this; }
+	constexpr FCycleTimeSpan operator-() const
+	{
+		return UNLIKELY(Cycles == MIN_int64) ? FromCycles(MAX_int64)
+			 : UNLIKELY(Cycles == MAX_int64) ? FromCycles(MIN_int64)
+											 : FromCycles(-Cycles);
+	}
+
+	constexpr FCycleTimeSpan operator+(const FCycleTimeSpan Span) const
+	{
+		return UNLIKELY(IsInfinity()) ? *this : UNLIKELY(Span.IsInfinity()) ? Span : FromCycles(Cycles + Span.GetCycles());
+	}
+
+	constexpr FCycleTimeSpan operator-(const FCycleTimeSpan Span) const
+	{
+		return UNLIKELY(IsInfinity()) ? *this : UNLIKELY(Span.IsInfinity()) ? -Span : FromCycles(Cycles - Span.GetCycles());
+	}
+
+	UE_API double ToSeconds() const;
+	UE_API double ToMilliseconds() const;
+
+	UE_API static FCycleTimeSpan FromSeconds(double Seconds);
+	UE_API static FCycleTimeSpan FromMilliseconds(double Milliseconds);
+
+private:
+	int64 Cycles = 0;
+};
+
+/**
+ * A point in time as reported by FPlatformTime::Cycles64().
+ *
+ * This is a monotonic clock which means the current time will never decrease. This time is meant
+ * primarily for measuring intervals. The interval between ticks of this clock is constant except
+ * for the time that the system is suspended on certain platforms. The tick frequency will differ
+ * between platforms, and must not be used as a means of communicating time without communicating
+ * the tick frequency together with the time.
+ */
+struct FCycleTimePoint
+{
+public:
+	UE_API static FCycleTimePoint Now();
+
+	constexpr static FCycleTimePoint Infinity() { return FromCycles(MAX_uint64); }
+
+	constexpr static FCycleTimePoint FromCycles(const uint64 InCycles)
+	{
+		FCycleTimePoint TimePoint;
+		TimePoint.Cycles = InCycles;
+		return TimePoint;
+	}
+
+	constexpr FCycleTimePoint() = default;
+
+	constexpr int64 GetCycles() const { return Cycles; }
+
+	constexpr bool IsInfinity() const { return Cycles == MAX_uint64; }
+
+	constexpr bool operator==(const FCycleTimePoint Other) const { return Cycles == Other.Cycles; }
+	constexpr bool operator!=(const FCycleTimePoint Other) const { return Cycles != Other.Cycles; }
+	constexpr bool operator<=(const FCycleTimePoint Other) const { return Cycles <= Other.Cycles; }
+	constexpr bool operator< (const FCycleTimePoint Other) const { return Cycles <  Other.Cycles; }
+	constexpr bool operator>=(const FCycleTimePoint Other) const { return Cycles >= Other.Cycles; }
+	constexpr bool operator> (const FCycleTimePoint Other) const { return Cycles >  Other.Cycles; }
+
+	constexpr FCycleTimePoint operator+(const FCycleTimeSpan Span) const
+	{
+		return UNLIKELY(IsInfinity() || Span.IsInfinity()) ? Infinity() : FromCycles(Cycles + Span.GetCycles());
+	}
+
+	constexpr FCycleTimePoint operator-(const FCycleTimeSpan Span) const
+	{
+		return UNLIKELY(IsInfinity() || Span.IsInfinity()) ? Infinity() : FromCycles(Cycles - Span.GetCycles());
+	}
+
+	constexpr FCycleTimeSpan operator-(const FCycleTimePoint Point) const
+	{
+		return UNLIKELY(IsInfinity()) ? FCycleTimeSpan::Infinity()
+			 : UNLIKELY(Point.IsInfinity()) ? -FCycleTimeSpan::Infinity()
+			 : FCycleTimeSpan::FromCycles(int64(Cycles - Point.Cycles));
+	}
+
+private:
+	uint64 Cycles = 0;
+};
+
+} // namespace UE
+
+#undef UE_API
diff --git a/Engine/Source/Runtime/Core/Tests/Misc/CycleTimeTest.cpp b/Engine/Source/Runtime/Core/Tests/Misc/CycleTimeTest.cpp
new file mode 100644
--- /dev/null
+++ b/Engine/Source/Runtime/Core/Tests/Misc/CycleTimeTest.cpp
@@ -0,0 +1,152 @@
+// Copyright Epic Games, Inc. All Rights Reserved.
+
+#if WITH_LOW_LEVEL_TESTS
+
+#include "Misc/CycleTime.h"
+
+#include "Containers/UnrealString.h"
+#include "Misc/Timespan.h"
+#include "TestHarness.h"
+
+namespace UE
+{
+
+TEST_CASE("Core::Time::CycleTimeSpan", "[Core][Time][Smoke]")
+{
+	SECTION("Constructors")
+	{
+		STATIC_CHECK(FCycleTimeSpan().GetCycles() == 0);
+		STATIC_CHECK(FCycleTimeSpan::FromCycles(0).GetCycles() == 0);
+		STATIC_CHECK(FCycleTimeSpan::FromCycles(123).GetCycles() == 123);
+		STATIC_CHECK(FCycleTimeSpan::Zero().GetCycles() == 0);
+	}
+
+	SECTION("Comparison")
+	{
+		STATIC_CHECK(FCycleTimeSpan::FromCycles(0) == FCycleTimeSpan());
+		STATIC_CHECK(FCycleTimeSpan::FromCycles(1) != FCycleTimeSpan());
+		STATIC_CHECK(FCycleTimeSpan::FromCycles(0) <= FCycleTimeSpan());
+		STATIC_CHECK(FCycleTimeSpan::FromCycles(0) >= FCycleTimeSpan());
+		STATIC_CHECK(FCycleTimeSpan::FromCycles(0) < FCycleTimeSpan::FromCycles(1));
+		STATIC_CHECK(FCycleTimeSpan::FromCycles(1) > FCycleTimeSpan::FromCycles(0));
+	}
+
+	SECTION("Infinity")
+	{
+		STATIC_CHECK(FCycleTimeSpan::Infinity().IsInfinity());
+		STATIC_CHECK_FALSE(FCycleTimeSpan().IsInfinity());
+		STATIC_CHECK_FALSE(FCycleTimeSpan::FromCycles(123).IsInfinity());
+
+		STATIC_CHECK(FCycleTimeSpan::Infinity() == FCycleTimeSpan::Infinity());
+		STATIC_CHECK(FCycleTimeSpan::Infinity() <= FCycleTimeSpan::Infinity());
+		STATIC_CHECK(FCycleTimeSpan::Infinity() >= FCycleTimeSpan::Infinity());
+
+		STATIC_CHECK(FCycleTimeSpan::Infinity() >= FCycleTimeSpan());
+		STATIC_CHECK(FCycleTimeSpan::Infinity() > FCycleTimeSpan());
+
+		STATIC_CHECK(FCycleTimeSpan() != FCycleTimeSpan::Infinity());
+		STATIC_CHECK(FCycleTimeSpan() <= FCycleTimeSpan::Infinity());
+		STATIC_CHECK(FCycleTimeSpan() < FCycleTimeSpan::Infinity());
+	}
+
+	SECTION("Addition")
+	{
+		STATIC_CHECK(FCycleTimeSpan::FromCycles(234) + FCycleTimeSpan::FromCycles(123) == FCycleTimeSpan::FromCycles(357));
+		STATIC_CHECK(FCycleTimeSpan::FromCycles(234) + (-FCycleTimeSpan::Infinity()) == -FCycleTimeSpan::Infinity());
+		STATIC_CHECK(FCycleTimeSpan::FromCycles(234) + FCycleTimeSpan::Infinity() == FCycleTimeSpan::Infinity());
+		STATIC_CHECK(FCycleTimeSpan::Infinity() + FCycleTimeSpan::FromCycles(1) == FCycleTimeSpan::Infinity());
+		STATIC_CHECK(FCycleTimeSpan::Infinity() + (-FCycleTimeSpan::Infinity()) == FCycleTimeSpan::Infinity());
+		STATIC_CHECK(FCycleTimeSpan::Infinity() + FCycleTimeSpan::Infinity() == FCycleTimeSpan::Infinity());
+		STATIC_CHECK((-FCycleTimeSpan::Infinity()) + FCycleTimeSpan::FromCycles(1) == -FCycleTimeSpan::Infinity());
+		STATIC_CHECK((-FCycleTimeSpan::Infinity()) + (-FCycleTimeSpan::Infinity()) == -FCycleTimeSpan::Infinity());
+		STATIC_CHECK((-FCycleTimeSpan::Infinity()) + FCycleTimeSpan::Infinity() == -FCycleTimeSpan::Infinity());
+	}
+
+	SECTION("Subtraction")
+	{
+		STATIC_CHECK(FCycleTimeSpan::FromCycles(234) - FCycleTimeSpan::FromCycles(123) == FCycleTimeSpan::FromCycles(111));
+		STATIC_CHECK(FCycleTimeSpan::FromCycles(234) - (-FCycleTimeSpan::Infinity()) == FCycleTimeSpan::Infinity());
+		STATIC_CHECK(FCycleTimeSpan::FromCycles(234) - FCycleTimeSpan::Infinity() == -FCycleTimeSpan::Infinity());
+		STATIC_CHECK(FCycleTimeSpan::Infinity() - FCycleTimeSpan::FromCycles(1) == FCycleTimeSpan::Infinity());
+		STATIC_CHECK(FCycleTimeSpan::Infinity() - (-FCycleTimeSpan::Infinity()) == FCycleTimeSpan::Infinity());
+		STATIC_CHECK(FCycleTimeSpan::Infinity() - FCycleTimeSpan::Infinity() == FCycleTimeSpan::Infinity());
+		STATIC_CHECK((-FCycleTimeSpan::Infinity()) - FCycleTimeSpan::FromCycles(1) == -FCycleTimeSpan::Infinity());
+		STATIC_CHECK((-FCycleTimeSpan::Infinity()) - (-FCycleTimeSpan::Infinity()) == -FCycleTimeSpan::Infinity());
+		STATIC_CHECK((-FCycleTimeSpan::Infinity()) - FCycleTimeSpan::Infinity() == -FCycleTimeSpan::Infinity());
+	}
+
+	SECTION("Conversions")
+	{
+		CHECK(FMath::IsNearlyEqual(FCycleTimeSpan::FromSeconds(123.0).ToSeconds(), 123.0));
+		CHECK(FMath::IsNearlyEqual(FCycleTimeSpan::FromMilliseconds(123.0).ToMilliseconds(), 123.0));
+		CHECK(FMath::IsNearlyEqual(FCycleTimeSpan(FTimespan::FromSeconds(123.0)).ToSeconds(), 123.0));
+		CHECK(FMath::IsNearlyEqual(FCycleTimeSpan(FTimespan::FromSeconds(-123.0)).ToSeconds(), -123.0));
+	}
+}
+
+TEST_CASE("Core::Time::CycleTimePoint", "[Core][Time][Smoke]")
+{
+	SECTION("Constructors")
+	{
+		STATIC_CHECK(FCycleTimePoint().GetCycles() == 0);
+		STATIC_CHECK(FCycleTimePoint::FromCycles(0).GetCycles() == 0);
+		STATIC_CHECK(FCycleTimePoint::FromCycles(123).GetCycles() == 123);
+	}
+
+	SECTION("Comparison")
+	{
+		STATIC_CHECK(FCycleTimePoint::FromCycles(0) == FCycleTimePoint());
+		STATIC_CHECK(FCycleTimePoint::FromCycles(1) != FCycleTimePoint());
+		STATIC_CHECK(FCycleTimePoint::FromCycles(0) <= FCycleTimePoint());
+		STATIC_CHECK(FCycleTimePoint::FromCycles(0) >= FCycleTimePoint());
+		STATIC_CHECK(FCycleTimePoint::FromCycles(0) < FCycleTimePoint::FromCycles(1));
+		STATIC_CHECK(FCycleTimePoint::FromCycles(1) > FCycleTimePoint::FromCycles(0));
+	}
+
+	SECTION("Infinity")
+	{
+		STATIC_CHECK(FCycleTimePoint::Infinity().IsInfinity());
+		STATIC_CHECK_FALSE(FCycleTimePoint().IsInfinity());
+		STATIC_CHECK_FALSE(FCycleTimePoint::FromCycles(123).IsInfinity());
+
+		STATIC_CHECK(FCycleTimePoint::Infinity() == FCycleTimePoint::Infinity());
+		STATIC_CHECK(FCycleTimePoint::Infinity() <= FCycleTimePoint::Infinity());
+		STATIC_CHECK(FCycleTimePoint::Infinity() >= FCycleTimePoint::Infinity());
+
+		STATIC_CHECK(FCycleTimePoint::Infinity() >= FCycleTimePoint());
+		STATIC_CHECK(FCycleTimePoint::Infinity() > FCycleTimePoint());
+
+		STATIC_CHECK(FCycleTimePoint() != FCycleTimePoint::Infinity());
+		STATIC_CHECK(FCycleTimePoint() <= FCycleTimePoint::Infinity());
+		STATIC_CHECK(FCycleTimePoint() < FCycleTimePoint::Infinity());
+	}
+
+	SECTION("Addition")
+	{
+		STATIC_CHECK(FCycleTimePoint::FromCycles(234) + FCycleTimeSpan::FromCycles(123) == FCycleTimePoint::FromCycles(357));
+		STATIC_CHECK(FCycleTimePoint::Infinity() + FCycleTimeSpan::FromCycles(1) == FCycleTimePoint::Infinity());
+		STATIC_CHECK(FCycleTimePoint::Infinity() + (-FCycleTimeSpan::Infinity()) == FCycleTimePoint::Infinity());
+		STATIC_CHECK(FCycleTimePoint::Infinity() + FCycleTimeSpan::Infinity() == FCycleTimePoint::Infinity());
+	}
+
+	SECTION("Subtraction")
+	{
+		STATIC_CHECK(FCycleTimePoint::FromCycles(234) - FCycleTimeSpan::FromCycles(123) == FCycleTimePoint::FromCycles(111));
+		STATIC_CHECK(FCycleTimePoint::Infinity() - FCycleTimeSpan::FromCycles(1) == FCycleTimePoint::Infinity());
+		STATIC_CHECK(FCycleTimePoint::Infinity() - (-FCycleTimeSpan::Infinity()) == FCycleTimePoint::Infinity());
+		STATIC_CHECK(FCycleTimePoint::Infinity() - FCycleTimeSpan::Infinity() == FCycleTimePoint::Infinity());
+	}
+
+	SECTION("Span")
+	{
+		STATIC_CHECK(FCycleTimePoint::FromCycles(357) - FCycleTimePoint::FromCycles(234) == FCycleTimeSpan::FromCycles(123));
+		STATIC_CHECK(FCycleTimePoint::FromCycles(234) - FCycleTimePoint::FromCycles(357) == FCycleTimeSpan::FromCycles(-123));
+
+		STATIC_CHECK(FCycleTimePoint::Infinity() - FCycleTimePoint::FromCycles(123) == FCycleTimeSpan::Infinity());
+		STATIC_CHECK(FCycleTimePoint::FromCycles(123) - FCycleTimePoint::Infinity() == -FCycleTimeSpan::Infinity());
+	}
+}
+
+} // namespace UE
+
+#endif // WITH_LOW_LEVEL_TESTS
