Make the AutoRTFM tests work now that the runtime has moved into UE Core.

diff --git a/Engine/Source/Runtime/Core/Private/AutoRTFM/Context.cpp b/Engine/Source/Runtime/Core/Private/AutoRTFM/Context.cpp
--- a/Engine/Source/Runtime/Core/Private/AutoRTFM/Context.cpp
+++ b/Engine/Source/Runtime/Core/Private/AutoRTFM/Context.cpp
@@ -203,13 +203,14 @@ void FContext::AbortByLanguageAndThrow()
     CurrentTransaction->AbortAndThrow();
 }
 
+#if _WIN32
+extern "C" __declspec(dllimport) void __stdcall GetCurrentThreadStackLimits(void**, void**);
+#endif
+
 FContext::FContext()
 {
 #if _WIN32
-#if 0
-    GetCurrentThreadStackLimits(reinterpret_cast<PULONG_PTR>(&StackBegin),
-                                reinterpret_cast<PULONG_PTR>(&StackEnd));
-#endif
+    GetCurrentThreadStackLimits(&StackBegin, &StackEnd);
 #else
     pthread_attr_t Attr;
     pthread_getattr_np(pthread_self(), &Attr);
@@ -228,13 +229,6 @@ void FContext::Reset()
     Status = EContextStatus::Idle;
 }
 
-void FContext::StaticAsserts()
-{
-    /*static_assert(PrettyStaticAssert<Constants::Offset_Context_CurrentTransaction, offsetof(FContext, CurrentTransaction)>::_cResult, "Not equal");
-    static_assert(PrettyStaticAssert<Constants::Offset_Context_LineTable, offsetof(FContext, LineTable)>::_cResult, "Not equal");
-    static_assert(PrettyStaticAssert<Constants::Offset_Context_Status, offsetof(FContext, Status)>::_cResult, "Not equal");*/
-}
-
 void FContext::DumpState() const
 {
     fprintf(GetLogFile(), "Context at %p.\n", this);
diff --git a/Engine/Source/Runtime/Core/Private/AutoRTFM/Context.h b/Engine/Source/Runtime/Core/Private/AutoRTFM/Context.h
--- a/Engine/Source/Runtime/Core/Private/AutoRTFM/Context.h
+++ b/Engine/Source/Runtime/Core/Private/AutoRTFM/Context.h
@@ -49,8 +49,6 @@ private:
     
     // All of this other stuff ought to be private?
     void Reset();
-
-    void StaticAsserts();
     
     FTransaction* CurrentTransaction{nullptr};
     void* StackBegin{nullptr}; // begin as in the smaller of the two
diff --git a/Engine/Source/Runtime/Core/Private/AutoRTFM/CppNewSupport.cpp b/Engine/Source/Runtime/Core/Private/AutoRTFM/CppNewSupport.cpp
--- a/Engine/Source/Runtime/Core/Private/AutoRTFM/CppNewSupport.cpp
+++ b/Engine/Source/Runtime/Core/Private/AutoRTFM/CppNewSupport.cpp
@@ -10,6 +10,7 @@
 
 namespace AutoRTFM
 {
+#ifdef _WIN32
 	using FCppNew = void* (*)(size_t);
 	using FCppDeleteWithSize = void(*)(void*, size_t);
 	using FCppDelete = void(*)(void*);
@@ -18,6 +19,7 @@ namespace AutoRTFM
 	FCppNew CppNewWithSize;
 	FCppDeleteWithSize CppDeleteWithSize;
 	FCppDelete CppDelete;
+#endif
 
 // These are also exposed as compiler ABI, currently just for the benefit of Windows.
 extern "C" void* autortfm_cpp_new(size_t Size, FContext* Context)
diff --git a/Engine/Source/Runtime/Core/Private/AutoRTFM/Memcpy.h b/Engine/Source/Runtime/Core/Private/AutoRTFM/Memcpy.h
--- a/Engine/Source/Runtime/Core/Private/AutoRTFM/Memcpy.h
+++ b/Engine/Source/Runtime/Core/Private/AutoRTFM/Memcpy.h
@@ -2,6 +2,8 @@
 
 #pragma once
 
+#include <stdlib.h>
+
 namespace AutoRTFM
 {
 
diff --git a/Engine/Source/Runtime/Core/Private/AutoRTFM/StdLibSupport.cpp b/Engine/Source/Runtime/Core/Private/AutoRTFM/StdLibSupport.cpp
--- a/Engine/Source/Runtime/Core/Private/AutoRTFM/StdLibSupport.cpp
+++ b/Engine/Source/Runtime/Core/Private/AutoRTFM/StdLibSupport.cpp
@@ -164,11 +164,20 @@ UE_AUTORTFM_REGISTER_SELF_FUNCTION(static_cast<double(*)(double, double)>(&pow))
 UE_AUTORTFM_REGISTER_SELF_FUNCTION(_isnan);
 UE_AUTORTFM_REGISTER_SELF_FUNCTION(_finite);
 
+extern "C" __declspec(dllimport) int __stdcall IsDebuggerPresent(void);
+
 UE_AUTORTFM_REGISTER_SELF_FUNCTION(IsDebuggerPresent);
 
+extern "C" __declspec(dllimport) void EnterCriticalSection(void*);
+
 UE_AUTORTFM_REGISTER_SELF_FUNCTION(EnterCriticalSection);
+
+extern "C" __declspec(dllimport) void LeaveCriticalSection(void*);
 UE_AUTORTFM_REGISTER_SELF_FUNCTION(LeaveCriticalSection);
 
+extern "C" __declspec(dllimport) int __stdcall QueryPerformanceCounter(void *);
+extern "C" __declspec(dllimport) int __stdcall QueryPerformanceFrequency(void *);
+
 UE_AUTORTFM_REGISTER_SELF_FUNCTION(QueryPerformanceCounter);
 #endif
 
diff --git a/Engine/Source/Runtime/Core/Private/AutoRTFM/TaggedPtr.h b/Engine/Source/Runtime/Core/Private/AutoRTFM/TaggedPtr.h
--- a/Engine/Source/Runtime/Core/Private/AutoRTFM/TaggedPtr.h
+++ b/Engine/Source/Runtime/Core/Private/AutoRTFM/TaggedPtr.h
@@ -61,11 +61,13 @@ namespace AutoRTFM
 		TTaggedPtr& operator=(TTaggedPtr& Other)
 		{
 			Payload = Other.Payload;
+			return *this;
 		}
 
 		TTaggedPtr& operator=(TTaggedPtr&& Other)
 		{
 			Payload = Other.Payload;
+			return *this;
 		}
 
 		T& operator*() const
@@ -159,11 +161,13 @@ namespace AutoRTFM
 		TTaggedPtr& operator=(TTaggedPtr& Other)
 		{
 			Payload = Other.Payload;
+			return *this;
 		}
 
 		TTaggedPtr& operator=(TTaggedPtr&& Other)
 		{
 			Payload = Other.Payload;
+			return *this;
 		}
 
 		bool operator==(const TTaggedPtr& Other) const
