Logging: Removed cycles from FLogTime until we have a better abstraction of time

#rnx

diff --git a/Engine/Source/Runtime/Core/Private/Logging/StructuredLog.cpp b/Engine/Source/Runtime/Core/Private/Logging/StructuredLog.cpp
--- a/Engine/Source/Runtime/Core/Private/Logging/StructuredLog.cpp
+++ b/Engine/Source/Runtime/Core/Private/Logging/StructuredLog.cpp
@@ -506,7 +506,6 @@ FLogTime FLogTime::Now()
 {
 	FLogTime Time;
 	Time.UtcTicks = FDateTime::UtcNow().GetTicks();
-	Time.Cycles = FPlatformTime::Cycles64();
 	return Time;
 }
 
diff --git a/Engine/Source/Runtime/Core/Public/Logging/StructuredLog.h b/Engine/Source/Runtime/Core/Public/Logging/StructuredLog.h
--- a/Engine/Source/Runtime/Core/Public/Logging/StructuredLog.h
+++ b/Engine/Source/Runtime/Core/Public/Logging/StructuredLog.h
@@ -126,9 +126,6 @@ class FLogTemplate;
 
 /**
  * Time that a log event occurred.
- *
- * This abstraction exists to hide the multiple time representations.
- * When cycles can be converted to FDateTime, this can stop storing ticks.
  */
 class FLogTime
 {
@@ -137,15 +134,10 @@ public:
 
 	constexpr FLogTime() = default;
 
-	/** Returns the cycle count. 0 when the time was not set. */
-	constexpr uint64 GetCycles() const { return Cycles; }
-
 	/** Returns the UTC time. 0 ticks when the time was not set. */
 	UE_API FDateTime GetUtcTime() const;
 
 private:
-	/** Cycles from FPlatformTime::Cycles64() */
-	uint64 Cycles = 0;
 	/** Ticks from FDateTime::UtcNow() */
 	int64 UtcTicks = 0;
 };
diff --git a/Engine/Source/Runtime/Launch/Private/LaunchEngineLoop.cpp b/Engine/Source/Runtime/Launch/Private/LaunchEngineLoop.cpp
--- a/Engine/Source/Runtime/Launch/Private/LaunchEngineLoop.cpp
+++ b/Engine/Source/Runtime/Launch/Private/LaunchEngineLoop.cpp
@@ -439,7 +439,7 @@ private:
 	{
 		TStringBuilder<512> V;
 		Record.FormatMessageTo(V);
-		Serialize(*V, Record.GetVerbosity(), Record.GetCategory(), FPlatformTime::ToSeconds64(Record.GetTime().GetCycles()));
+		Serialize(*V, Record.GetVerbosity(), Record.GetCategory(), -1.0);
 	}
 
 	template <typename CharType>
