Logging: Removed the Global\ prefix from the UE_LOG_JSON named mutex used on Windows

Local is believed to be sufficient based on discussions with partners.

#rnx

diff --git a/Engine/Source/Runtime/Core/Private/Windows/WindowsPlatformMisc.cpp b/Engine/Source/Runtime/Core/Private/Windows/WindowsPlatformMisc.cpp
--- a/Engine/Source/Runtime/Core/Private/Windows/WindowsPlatformMisc.cpp
+++ b/Engine/Source/Runtime/Core/Private/Windows/WindowsPlatformMisc.cpp
@@ -986,7 +986,7 @@ void FWindowsPlatformMisc::LocalPrint( const TCHAR *Message )
 
 bool FWindowsPlatformMisc::IsLowLevelOutputDebugStringStructured()
 {
-	HANDLE Mutex = OpenMutexW(SYNCHRONIZE, /*bInheritHandle*/ false, LR"--(Global\UE_LOG_JSON)--");
+	HANDLE Mutex = OpenMutexW(SYNCHRONIZE, /*bInheritHandle*/ false, L"UE_LOG_JSON");
 	ON_SCOPE_EXIT { CloseHandle(Mutex); };
 	return !!Mutex || FGenericPlatformMisc::IsLowLevelOutputDebugStringStructured();
 }
