Bidrectional backwards compatibility support, i.e. support between TScriptDelegate<FWeakObjectPtr> and FMulticastScriptDelgate, and also between FScriptDelegate and TMulticastScriptDelgate<FWeakObjectPtr>.

diff --git a/Engine/Source/Runtime/Core/Public/UObject/ScriptDelegates.h b/Engine/Source/Runtime/Core/Public/UObject/ScriptDelegates.h
--- a/Engine/Source/Runtime/Core/Public/UObject/ScriptDelegates.h
+++ b/Engine/Source/Runtime/Core/Public/UObject/ScriptDelegates.h
@@ -33,6 +33,12 @@ namespace UE::Core::Private
 
 		using WeakPtrType = FWeakObjectPtr;
 	};
+
+	// This type only exists to allow compatibility between multicast and unicast delegate types which use an explicit FWeakObjectPtr template parameter
+	template <typename From, typename To>
+	inline constexpr /*UE_DEPRECATED(5.3, "Deprecated - remove after TScriptDelegateTraits<FWeakObjectPtr> is removed")*/ bool TBackwardCompatibilityCheck_V =
+		(std::is_same_v<From, TScriptDelegateDefault> && std::is_same_v<To, FWeakObjectPtr>) ||
+		(std::is_same_v<From, FWeakObjectPtr> && std::is_same_v<To, TScriptDelegateDefault>);
 }
 
 /**
@@ -54,11 +60,11 @@ public:
 		  FunctionName( NAME_None )
 	{ }
 
-	/* UE_DEPRECATED(5.3, "Deprecated - remove after TScriptDelegateTraits<FWeakObjectPtr> is removed") */
 	template <
 		typename OtherDummy
-		UE_REQUIRES(std::is_same_v<OtherDummy, FWeakObjectPtr>)
+		UE_REQUIRES(UE::Core::Private::TBackwardCompatibilityCheck_V<Dummy, OtherDummy>)
 	>
+	/* UE_DEPRECATED(5.3, "Deprecated - remove after TScriptDelegateTraits<FWeakObjectPtr> is removed") */
 	TScriptDelegate(const TScriptDelegate<OtherDummy>& Other)
 		: Object(Other.Object)
 		, FunctionName(Other.FunctionName)
@@ -206,6 +212,17 @@ public:
 
 		return Object == Other.Object && FunctionName == Other.FunctionName;
 	}
+	template <
+		typename OtherDummy
+		UE_REQUIRES(UE::Core::Private::TBackwardCompatibilityCheck_V<Dummy, OtherDummy>)
+	>
+	/* UE_DEPRECATED(5.3, "Deprecated - remove after TScriptDelegateTraits<FWeakObjectPtr> is removed") */
+	FORCEINLINE bool operator==(const TScriptDelegate<OtherDummy>& Other) const
+	{
+		UE_DELEGATES_MT_SCOPED_READ_ACCESS(AccessDetector);
+
+		return Object == Other.Object && FunctionName == Other.FunctionName;
+	}
 
 	FORCEINLINE bool operator!=( const TScriptDelegate& Other ) const
 	{
@@ -213,6 +230,17 @@ public:
 
 		return Object != Other.Object || FunctionName != Other.FunctionName;
 	}
+	template <
+		typename OtherDummy
+		UE_REQUIRES(UE::Core::Private::TBackwardCompatibilityCheck_V<Dummy, OtherDummy>)
+	>
+	/* UE_DEPRECATED(5.3, "Deprecated - remove after TScriptDelegateTraits<FWeakObjectPtr> is removed") */
+	FORCEINLINE bool operator!=(const TScriptDelegate<OtherDummy>& Other) const
+	{
+		UE_DELEGATES_MT_SCOPED_READ_ACCESS(AccessDetector);
+
+		return Object != Other.Object || FunctionName != Other.FunctionName;
+	}
 
 	void operator=( const TScriptDelegate& Other )
 	{
@@ -221,6 +249,18 @@ public:
 		Object = Other.Object;
 		FunctionName = Other.FunctionName;
 	}
+	template <
+		typename OtherDummy
+		UE_REQUIRES(UE::Core::Private::TBackwardCompatibilityCheck_V<Dummy, OtherDummy>)
+	>
+	/* UE_DEPRECATED(5.3, "Deprecated - remove after TScriptDelegateTraits<FWeakObjectPtr> is removed") */
+	FORCEINLINE bool operator=(const TScriptDelegate<OtherDummy>& Other)
+	{
+		UE_DELEGATES_MT_SCOPED_WRITE_ACCESS(AccessDetector);
+
+		Object = Other.Object;
+		FunctionName = Other.FunctionName;
+	}
 
 	/** 
 	 * Gets the object bound to this delegate
@@ -396,7 +436,7 @@ public:
 	}
 	template <
 		typename OtherDummy
-		UE_REQUIRES(std::is_same_v<OtherDummy, FWeakObjectPtr>)
+		UE_REQUIRES(UE::Core::Private::TBackwardCompatibilityCheck_V<Dummy, OtherDummy>)
 	>
 	/* UE_DEPRECATED(5.3, "Deprecated - remove after TScriptDelegateTraits<FWeakObjectPtr> is removed") */
 	FORCEINLINE bool Contains(const TScriptDelegate<OtherDummy>& InDelegate) const
@@ -437,7 +477,7 @@ public:
 	}
 	template <
 		typename OtherDummy
-		UE_REQUIRES(std::is_same_v<OtherDummy, FWeakObjectPtr>)
+		UE_REQUIRES(UE::Core::Private::TBackwardCompatibilityCheck_V<Dummy, OtherDummy>)
 	>
 	/* UE_DEPRECATED(5.3, "Deprecated - remove after TScriptDelegateTraits<FWeakObjectPtr> is removed") */
 	FORCEINLINE void Add(const TScriptDelegate<OtherDummy>& InDelegate)
@@ -463,7 +503,7 @@ public:
 	}
 	template <
 		typename OtherDummy
-		UE_REQUIRES(std::is_same_v<OtherDummy, FWeakObjectPtr>)
+		UE_REQUIRES(UE::Core::Private::TBackwardCompatibilityCheck_V<Dummy, OtherDummy>)
 	>
 	/* UE_DEPRECATED(5.3, "Deprecated - remove after TScriptDelegateTraits<FWeakObjectPtr> is removed") */
 	FORCEINLINE void AddUnique(const TScriptDelegate<OtherDummy>& InDelegate)
@@ -489,7 +529,7 @@ public:
 	}
 	template <
 		typename OtherDummy
-		UE_REQUIRES(std::is_same_v<OtherDummy, FWeakObjectPtr>)
+		UE_REQUIRES(UE::Core::Private::TBackwardCompatibilityCheck_V<Dummy, OtherDummy>)
 	>
 	/* UE_DEPRECATED(5.3, "Deprecated - remove after TScriptDelegateTraits<FWeakObjectPtr> is removed") */
 	FORCEINLINE void Remove(const TScriptDelegate<OtherDummy>& InDelegate)
