Fix stack corruption in FArchive::SerializeIntPacked64.

diff --git a/Engine/Source/Runtime/Core/Private/Serialization/Archive.cpp b/Engine/Source/Runtime/Core/Private/Serialization/Archive.cpp
--- a/Engine/Source/Runtime/Core/Private/Serialization/Archive.cpp
+++ b/Engine/Source/Runtime/Core/Private/Serialization/Archive.cpp
@@ -1346,7 +1346,7 @@ void FArchive::SerializeIntPacked64(uint64& Value)
 	}
 	else
 	{
-		uint8 PackedBytes[9];
+		uint8 PackedBytes[10];
 		int32 PackedByteCount = 0;
 		uint64 Remaining = Value;
 		while (true)
