RequestExit Feedback: Add Callsite argument to FPlatformMisc::RequestExit
Add the argument to many but not all callsites.
#rnx

diff --git a/Engine/Platforms/Hololens/Source/Runtime/Core/Private/HoloLensPlatformMisc.cpp b/Engine/Platforms/Hololens/Source/Runtime/Core/Private/HoloLensPlatformMisc.cpp
--- a/Engine/Platforms/Hololens/Source/Runtime/Core/Private/HoloLensPlatformMisc.cpp
+++ b/Engine/Platforms/Hololens/Source/Runtime/Core/Private/HoloLensPlatformMisc.cpp
@@ -162,9 +162,10 @@ void FHoloLensMisc::LowLevelOutputDebugString(const TCHAR *Message)
 	OutputDebugString(Message);
 }
 
-void FHoloLensMisc::RequestExit(bool Force)
+void FHoloLensMisc::RequestExit(bool Force, const TCHAR* CallSite)
 {
-	UE_LOG(LogTemp, Log, TEXT("FHoloLensMisc::RequestExit(%i)"), Force);
+	UE_LOG(LogTemp, Log, TEXT("FHoloLensMisc::RequestExit(%i, %s)"), Force,
+		CallSite ? CallSite : TEXT("<NoCallSiteInfo>"));
 	if (Force)
 	{
 		// Force immediate exit.
diff --git a/Engine/Platforms/Hololens/Source/Runtime/Core/Public/HoloLensPlatformMisc.h b/Engine/Platforms/Hololens/Source/Runtime/Core/Public/HoloLensPlatformMisc.h
--- a/Engine/Platforms/Hololens/Source/Runtime/Core/Public/HoloLensPlatformMisc.h
+++ b/Engine/Platforms/Hololens/Source/Runtime/Core/Public/HoloLensPlatformMisc.h
@@ -62,7 +62,7 @@ struct CORE_API FHoloLensMisc : public FGenericPlatformMisc
     
     static void PumpMessages(bool bFromMainLoop);
 	static void LowLevelOutputDebugString(const TCHAR *Message);
-	static void RequestExit(bool Force);
+	static void RequestExit(bool Force, const TCHAR* CallSite = nullptr);
 	static const TCHAR* GetSystemErrorMessage(TCHAR* OutBuffer, int32 BufferCount, int32 Error);
 	static void CreateGuid(struct FGuid& Result);
 	static int32 NumberOfCores();
diff --git a/Engine/Source/Runtime/Core/Private/Android/AndroidErrorOutputDevice.cpp b/Engine/Source/Runtime/Core/Private/Android/AndroidErrorOutputDevice.cpp
--- a/Engine/Source/Runtime/Core/Private/Android/AndroidErrorOutputDevice.cpp
+++ b/Engine/Source/Runtime/Core/Private/Android/AndroidErrorOutputDevice.cpp
@@ -38,7 +38,7 @@ void FAndroidErrorOutputDevice::Serialize( const TCHAR* Msg, ELogVerbosity::Type
 	else
 	{
 		HandleError();
-		FPlatformMisc::RequestExit(true);
+		FPlatformMisc::RequestExit(true, TEXT("FAndroidErrorOutputDevice::Serialize.!GIsGuarded"));
 	}
 }
 
diff --git a/Engine/Source/Runtime/Core/Private/Android/AndroidPlatformMisc.cpp b/Engine/Source/Runtime/Core/Private/Android/AndroidPlatformMisc.cpp
--- a/Engine/Source/Runtime/Core/Private/Android/AndroidPlatformMisc.cpp
+++ b/Engine/Source/Runtime/Core/Private/Android/AndroidPlatformMisc.cpp
@@ -231,7 +231,7 @@ static void InitCpuThermalSensor()
 	}
 }
 
-void FAndroidMisc::RequestExit( bool Force )
+void FAndroidMisc::RequestExit( bool Force, const TCHAR* CallSite)
 {
 
 #if PLATFORM_COMPILER_OPTIMIZATION_PG_PROFILING
@@ -245,7 +245,8 @@ void FAndroidMisc::RequestExit( bool Force )
 	}
 #endif
 
-	UE_LOG(LogAndroid, Log, TEXT("FAndroidMisc::RequestExit(%i)"), Force);
+	UE_LOG(LogAndroid, Log, TEXT("FAndroidMisc::RequestExit(%i, %s)"), Force,
+		CallSite ? CallSite : TEXT("<NoCallSiteInfo>"));
 	if (GLog)
 	{
 		GLog->Flush();
diff --git a/Engine/Source/Runtime/Core/Private/GenericPlatform/GenericPlatformMallocCrash.cpp b/Engine/Source/Runtime/Core/Private/GenericPlatform/GenericPlatformMallocCrash.cpp
--- a/Engine/Source/Runtime/Core/Private/GenericPlatform/GenericPlatformMallocCrash.cpp
+++ b/Engine/Source/Runtime/Core/Private/GenericPlatform/GenericPlatformMallocCrash.cpp
@@ -155,7 +155,7 @@ struct FMallocCrashPool
 			UE_DEBUG_BREAK();
 			FPlatformMisc::LowLevelOutputDebugStringf( TEXT( "AllocateFromPool run out of memory allocating %u bytes for %u allocations\n" ), InAllocationSize, MaxNumAllocations );
 			FPlatformMisc::LowLevelOutputDebugString( TEXT( "Please increase MaxNumAllocations for that pool, exiting...\n" ) );
-			FPlatformMisc::RequestExit( true );
+			FPlatformMisc::RequestExit( true, TEXT("GenericPlatformmallocCrash::AllocateFromPool"));
 		}
 		return nullptr;
 	}
@@ -253,7 +253,7 @@ FGenericPlatformMallocCrash::FGenericPlatformMallocCrash( FMalloc* MainMalloc )
 	if( !SmallMemoryPool || !LargeMemoryPool || !BookkeepingPool )
 	{
 		FPlatformMisc::LowLevelOutputDebugString( TEXT( "Memory pools allocations failed, exiting...\n" ) );
-		FPlatformMisc::RequestExit(true);
+		FPlatformMisc::RequestExit(true, TEXT("GenericPlatformmallocCrash().MemoryPoolsAllocationsFailed"));
 	}
 
 	if(!IsAligned(LargeMemoryPool, REQUIRED_ALIGNMENT) 
@@ -261,7 +261,7 @@ FGenericPlatformMallocCrash::FGenericPlatformMallocCrash( FMalloc* MainMalloc )
 	|| !IsAligned(BookkeepingPool, REQUIRED_ALIGNMENT))
 	{
 		FPlatformMisc::LowLevelOutputDebugString( TEXT( "OS allocations must be aligned to a value multiple of 16, exiting...\n" ) );
-		FPlatformMisc::RequestExit(true);
+		FPlatformMisc::RequestExit(true, TEXT("GenericPlatformmallocCrash().InvalidAlignment"));
 	}
 
 	InitializeSmallPools();
@@ -335,7 +335,7 @@ void* FGenericPlatformMallocCrash::Malloc( SIZE_T Size, uint32 Alignment )
 				FPlatformMisc::LowLevelOutputDebugStringf( TEXT( "MallocCrash run out of memory allocating %u bytes, free %u bytes\n" ), Size32, LARGE_MEMORYPOOL_SIZE-LargeMemoryPoolOffset );
 				FPlatformMisc::LowLevelOutputDebugString(TEXT("Please increase LARGE_MEMORYPOOL_SIZE, exiting...\n"));
 				UE_DEBUG_BREAK();
-				FPlatformMisc::RequestExit( true );			
+				FPlatformMisc::RequestExit( true, TEXT("GenericPlatformmallocCrash::Malloc.OutOfMemory"));
 			}
 		}
 	}
@@ -363,14 +363,14 @@ void* FGenericPlatformMallocCrash::Realloc( void* Ptr, SIZE_T NewSize, uint32 Al
 					if (!PreviousMalloc->GetAllocationSize(Ptr, PtrSize) || PtrSize == 0)
 					{
 						FPlatformMisc::LowLevelOutputDebugString( TEXT( "Realloc from previous malloc - we were not able to get correct allocation size, exiting...\n" ) );
-						FPlatformMisc::RequestExit( true );
+						FPlatformMisc::RequestExit( true, TEXT("GenericPlatformmallocCrash::Realloc.InvalidAllocationSize"));
 					}
 				}
 				// There is nothing we can do about it.
 				else
 				{
 					FPlatformMisc::LowLevelOutputDebugString( TEXT( "Realloc from previous malloc - we don't know how to get allocation size, exiting...\n" ) );
-					FPlatformMisc::RequestExit( true );	
+					FPlatformMisc::RequestExit( true, TEXT("GenericPlatformmallocCrash::Realloc.CanNotHandleMissingAllocationSize"));
 				}
 			}
 			else
diff --git a/Engine/Source/Runtime/Core/Private/GenericPlatform/GenericPlatformMisc.cpp b/Engine/Source/Runtime/Core/Private/GenericPlatform/GenericPlatformMisc.cpp
--- a/Engine/Source/Runtime/Core/Private/GenericPlatform/GenericPlatformMisc.cpp
+++ b/Engine/Source/Runtime/Core/Private/GenericPlatform/GenericPlatformMisc.cpp
@@ -942,9 +942,10 @@ bool FGenericPlatformMisc::HasSeparateChannelForDebugOutput()
 	return true;
 }
 
-void FGenericPlatformMisc::RequestExit( bool Force )
+void FGenericPlatformMisc::RequestExit( bool Force, const TCHAR* CallSite)
 {
-	UE_LOG(LogGenericPlatformMisc, Log,  TEXT("FPlatformMisc::RequestExit(%i)"), Force );
+	UE_LOG(LogGenericPlatformMisc, Log,  TEXT("FPlatformMisc::RequestExit(%i, %s)"), Force,
+		CallSite ? CallSite : TEXT("<NoCallSiteInfo>"));
 	if( Force )
 	{
 		// Force immediate exit.
@@ -971,12 +972,12 @@ bool FGenericPlatformMisc::RestartApplication()
 	return false;
 }
 
-void FGenericPlatformMisc::RequestExitWithStatus(bool Force, uint8 ReturnCode)
+void FGenericPlatformMisc::RequestExitWithStatus(bool Force, uint8 ReturnCode, const TCHAR* CallSite)
 {
 	// Generic implementation will ignore the return code - this may be important, so warn.
 	UE_LOG(LogGenericPlatformMisc, Warning, TEXT("FPlatformMisc::RequestExitWithStatus(%i, %d) - return code will be ignored by the generic implementation."), Force, ReturnCode);
 
-	return FPlatformMisc::RequestExit(Force);
+	return FPlatformMisc::RequestExit(Force, CallSite);
 }
 
 const TCHAR* FGenericPlatformMisc::GetSystemErrorMessage(TCHAR* OutBuffer, int32 BufferCount, int32 Error)
diff --git a/Engine/Source/Runtime/Core/Private/HAL/ThreadHeartBeat.cpp b/Engine/Source/Runtime/Core/Private/HAL/ThreadHeartBeat.cpp
--- a/Engine/Source/Runtime/Core/Private/HAL/ThreadHeartBeat.cpp
+++ b/Engine/Source/Runtime/Core/Private/HAL/ThreadHeartBeat.cpp
@@ -289,7 +289,7 @@ void FORCENOINLINE FThreadHeartBeat::OnHang(double HangDuration, uint32 ThreadTh
 					*NSLOCTEXT("MessageDialog", "ReportHangError_Title", "Application Hang Detected").ToString());
 			}
 
-			FPlatformMisc::RequestExit(true);
+			FPlatformMisc::RequestExit(true, TEXT("FThreadHeartbeat::OnHang"));
 		}
 #else
 		if (bHangsAreFatal)
diff --git a/Engine/Source/Runtime/Core/Private/IOS/IOSErrorOutputDevice.cpp b/Engine/Source/Runtime/Core/Private/IOS/IOSErrorOutputDevice.cpp
--- a/Engine/Source/Runtime/Core/Private/IOS/IOSErrorOutputDevice.cpp
+++ b/Engine/Source/Runtime/Core/Private/IOS/IOSErrorOutputDevice.cpp
@@ -26,7 +26,7 @@ void FIOSErrorOutputDevice::Serialize( const TCHAR* Msg, ELogVerbosity::Type Ver
 	{
 		// We crashed outside the guarded code (e.g. appExit).
 		HandleError();
-		FPlatformMisc::RequestExit( true );
+		FPlatformMisc::RequestExit( true, TEXT("FIOSErrorOutputDevice::Serialize"));
 	}
 }
 
diff --git a/Engine/Source/Runtime/Core/Private/IOS/IOSPlatformMisc.cpp b/Engine/Source/Runtime/Core/Private/IOS/IOSPlatformMisc.cpp
--- a/Engine/Source/Runtime/Core/Private/IOS/IOSPlatformMisc.cpp
+++ b/Engine/Source/Runtime/Core/Private/IOS/IOSPlatformMisc.cpp
@@ -1634,11 +1634,11 @@ bool FIOSPlatformMisc::HasSeparateChannelForDebugOutput()
 #endif
 }
 
-void FIOSPlatformMisc::RequestExit(bool Force)
+void FIOSPlatformMisc::RequestExit(bool Force, const TCHAR* CallSite)
 {
 	if (Force)
 	{
-		FApplePlatformMisc::RequestExit(Force);
+		FApplePlatformMisc::RequestExit(Force, CallSite);
 	}
 	else
 	{
@@ -1647,16 +1647,17 @@ void FIOSPlatformMisc::RequestExit(bool Force)
 	}
 }
 
-void FIOSPlatformMisc::RequestExitWithStatus(bool Force, uint8 ReturnCode)
+void FIOSPlatformMisc::RequestExitWithStatus(bool Force, uint8 ReturnCode, const TCHAR* CallSite)
 {
 	if (Force)
 	{
-		FApplePlatformMisc::RequestExit(Force);
+		FApplePlatformMisc::RequestExit(Force, CallSite);
 	}
 	else
 	{
 		// Implementation will ignore the return code - this may be important, so warn.
-		UE_LOG(LogIOS, Warning, TEXT("FIOSPlatformMisc::RequestExitWithStatus(%i, %d) - return code will be ignored by the generic implementation."), Force, ReturnCode);
+		UE_LOG(LogIOS, Warning, TEXT("FIOSPlatformMisc::RequestExitWithStatus(%i, %d, %s) - return code will be ignored by the generic implementation."),
+			Force, ReturnCode, CallSite ? CallSite : TEXT("<NoCallSiteInfo>"));
 
 		// ForceExit is sort of a misnomer here.  This will exit the engine loop before calling _Exit() from the app delegate
 		[[IOSAppDelegate GetDelegate] ForceExit];
diff --git a/Engine/Source/Runtime/Core/Private/Mac/MacErrorOutputDevice.cpp b/Engine/Source/Runtime/Core/Private/Mac/MacErrorOutputDevice.cpp
--- a/Engine/Source/Runtime/Core/Private/Mac/MacErrorOutputDevice.cpp
+++ b/Engine/Source/Runtime/Core/Private/Mac/MacErrorOutputDevice.cpp
@@ -69,7 +69,7 @@ void FMacErrorOutputDevice::Serialize( const TCHAR* Msg, ELogVerbosity::Type Ver
 	{
 		// We crashed outside the guarded code (e.g. appExit).
 		HandleError();
-		FPlatformMisc::RequestExit( true );
+		FPlatformMisc::RequestExit( true, TEXT("FMacErrorOutputDevice::Serialize.!GIsGuarded"));
 	}
 }
 
diff --git a/Engine/Source/Runtime/Core/Private/Mac/MacPlatformMisc.cpp b/Engine/Source/Runtime/Core/Private/Mac/MacPlatformMisc.cpp
--- a/Engine/Source/Runtime/Core/Private/Mac/MacPlatformMisc.cpp
+++ b/Engine/Source/Runtime/Core/Private/Mac/MacPlatformMisc.cpp
@@ -706,9 +706,10 @@ TArray<uint8> FMacPlatformMisc::GetMacAddress()
 	return Result;
 }
 
-void FMacPlatformMisc::RequestExit( bool Force )
+void FMacPlatformMisc::RequestExit( bool Force, const TCHAR* CallSite)
 {
-	UE_LOG(LogMac, Log,  TEXT("FPlatformMisc::RequestExit(%i)"), Force );
+	UE_LOG(LogMac, Log,  TEXT("FPlatformMisc::RequestExit(%i, %s)"), Force,
+		CallSite ? CallSite : TEXT("<NoCallSiteInfo>"));
 
 	FCoreDelegates::ApplicationWillTerminateDelegate.Broadcast();
 
diff --git a/Engine/Source/Runtime/Core/Private/Misc/OutputDeviceAnsiError.cpp b/Engine/Source/Runtime/Core/Private/Misc/OutputDeviceAnsiError.cpp
--- a/Engine/Source/Runtime/Core/Private/Misc/OutputDeviceAnsiError.cpp
+++ b/Engine/Source/Runtime/Core/Private/Misc/OutputDeviceAnsiError.cpp
@@ -57,7 +57,7 @@ void FOutputDeviceAnsiError::Serialize( const TCHAR* Msg, ELogVerbosity::Type Ve
 		// pop up a crash window if we are not in unattended mode
 		if( FApp::IsUnattended() == false )
 		{
-			FPlatformMisc::RequestExit( true );
+			FPlatformMisc::RequestExit( true, TEXT("FOutputDeviceAnsiError::Serialize.!GIsGuarded"));
 		}
 		else
 		{
diff --git a/Engine/Source/Runtime/Core/Private/ProfilingDebugging/CsvProfiler.cpp b/Engine/Source/Runtime/Core/Private/ProfilingDebugging/CsvProfiler.cpp
--- a/Engine/Source/Runtime/Core/Private/ProfilingDebugging/CsvProfiler.cpp
+++ b/Engine/Source/Runtime/Core/Private/ProfilingDebugging/CsvProfiler.cpp
@@ -3142,7 +3142,7 @@ void FCsvProfiler::EndFrame()
 			if (bCaptureEnded && (GCsvExitOnCompletion || FParse::Param(FCommandLine::Get(), TEXT("ExitAfterCsvProfiling"))))
 			{
 				bool bForceExit = !!CVarCsvForceExit.GetValueOnGameThread();
-				FPlatformMisc::RequestExit(bForceExit);
+				FPlatformMisc::RequestExit(bForceExit, TEXT("CsvProfiler.ExitAfterCsvProfiling"));
 			}
 		}
 	}
diff --git a/Engine/Source/Runtime/Core/Private/UObject/UnrealNames.cpp b/Engine/Source/Runtime/Core/Private/UObject/UnrealNames.cpp
--- a/Engine/Source/Runtime/Core/Private/UObject/UnrealNames.cpp
+++ b/Engine/Source/Runtime/Core/Private/UObject/UnrealNames.cpp
@@ -1624,7 +1624,7 @@ FNamePool::FNamePool()
 		{
 			FPlatformMisc::PromptForRemoteDebugging(false);
 			FMessageDialog::Open(EAppMsgType::Ok, NSLOCTEXT("UnrealEd", "DuplicatedHardcodedName", "Duplicate hardcoded name"));
-			FPlatformMisc::RequestExit(false);
+			FPlatformMisc::RequestExit(false, TEXT("FNamePool.DuplicateHardcodedName"));
 		}
 	}
 }
diff --git a/Engine/Source/Runtime/Core/Private/Unix/UnixErrorOutputDevice.cpp b/Engine/Source/Runtime/Core/Private/Unix/UnixErrorOutputDevice.cpp
--- a/Engine/Source/Runtime/Core/Private/Unix/UnixErrorOutputDevice.cpp
+++ b/Engine/Source/Runtime/Core/Private/Unix/UnixErrorOutputDevice.cpp
@@ -75,7 +75,7 @@ void FUnixErrorOutputDevice::Serialize(const TCHAR* Msg, ELogVerbosity::Type Ver
 	{
 		// We crashed outside the guarded code (e.g. appExit).
 		HandleError();
-		FPlatformMisc::RequestExit(true);
+		FPlatformMisc::RequestExit(true, TEXT("FUnixErrorOutputDevice.Serialize.!GIsGuarded"));
 	}
 }
 
diff --git a/Engine/Source/Runtime/Core/Private/Unix/UnixPlatformCrashContext.cpp b/Engine/Source/Runtime/Core/Private/Unix/UnixPlatformCrashContext.cpp
--- a/Engine/Source/Runtime/Core/Private/Unix/UnixPlatformCrashContext.cpp
+++ b/Engine/Source/Runtime/Core/Private/Unix/UnixPlatformCrashContext.cpp
@@ -160,7 +160,7 @@ void GracefulTerminationHandler(int32 Signal, siginfo_t* Info, void* Context)
 	}
 	else
 	{
-		FPlatformMisc::RequestExit(true);
+		FPlatformMisc::RequestExit(true, TEXT("UnixPlatformCrashContext.GracefulTerminationHandler"));
 	}
 
 	GEnteredSignalHandler = 0;
diff --git a/Engine/Source/Runtime/Core/Private/Unix/UnixPlatformMemory.cpp b/Engine/Source/Runtime/Core/Private/Unix/UnixPlatformMemory.cpp
--- a/Engine/Source/Runtime/Core/Private/Unix/UnixPlatformMemory.cpp
+++ b/Engine/Source/Runtime/Core/Private/Unix/UnixPlatformMemory.cpp
@@ -144,7 +144,7 @@ class FMalloc* FUnixPlatformMemory::BaseAllocator()
 	if (geteuid() == 0)
 	{
 		fprintf(stderr, "Refusing to run with the root privileges.\n");
-		FPlatformMisc::RequestExit(true);
+		FPlatformMisc::RequestExit(true, TEXT("FUnixPlatformMemory.BaseAllocator"));
 		// unreachable
 		return nullptr;
 	}
diff --git a/Engine/Source/Runtime/Core/Private/Unix/UnixPlatformMisc.cpp b/Engine/Source/Runtime/Core/Private/Unix/UnixPlatformMisc.cpp
--- a/Engine/Source/Runtime/Core/Private/Unix/UnixPlatformMisc.cpp
+++ b/Engine/Source/Runtime/Core/Private/Unix/UnixPlatformMisc.cpp
@@ -345,7 +345,7 @@ extern volatile sig_atomic_t GEnteredSignalHandler;
 uint8 GOverriddenReturnCode = 0;
 bool GHasOverriddenReturnCode = false;
 
-void FUnixPlatformMisc::RequestExit(bool Force)
+void FUnixPlatformMisc::RequestExit(bool Force, const TCHAR* CallSite)
 {
 	if (GEnteredSignalHandler)
 	{
@@ -357,7 +357,8 @@ void FUnixPlatformMisc::RequestExit(bool Force)
 	}
 	else
 	{
-		UE_LOG(LogCore, Log,  TEXT("FUnixPlatformMisc::RequestExit(%i)"), Force );
+		UE_LOG(LogCore, Log,  TEXT("FUnixPlatformMisc::RequestExit(%i, %s)"), Force,
+			CallSite ? CallSite : TEXT("<NoCallSiteInfo>"));
 	}
 
 	if(Force)
@@ -398,11 +399,11 @@ PRAGMA_ENABLE_DEPRECATION_WARNINGS
 	else
 	{
 		// Tell the platform specific code we want to exit cleanly from the main loop.
-		FGenericPlatformMisc::RequestExit(Force);
+		FGenericPlatformMisc::RequestExit(Force, CallSite);
 	}
 }
 
-void FUnixPlatformMisc::RequestExitWithStatus(bool Force, uint8 ReturnCode)
+void FUnixPlatformMisc::RequestExitWithStatus(bool Force, uint8 ReturnCode, const TCHAR* CallSite)
 {
 	if (GEnteredSignalHandler)
 	{
@@ -414,13 +415,14 @@ void FUnixPlatformMisc::RequestExitWithStatus(bool Force, uint8 ReturnCode)
 	}
 	else
 	{
-		UE_LOG(LogCore, Log, TEXT("FUnixPlatformMisc::RequestExit(bForce=%s, ReturnCode=%d)"), Force ? TEXT("true") : TEXT("false"), ReturnCode);
+		UE_LOG(LogCore, Log, TEXT("FUnixPlatformMisc::RequestExit(bForce=%s, ReturnCode=%d, CallSite=%s)"),
+			Force ? TEXT("true") : TEXT("false"), ReturnCode, CallSite ? CallSite : TEXT("<NoCallSiteInfo>"));
 	}
 
 	GOverriddenReturnCode = ReturnCode;
 	GHasOverriddenReturnCode = true;
 
-	return FPlatformMisc::RequestExit(Force);
+	return FPlatformMisc::RequestExit(Force, CallSite);
 }
 
 bool FUnixPlatformMisc::HasOverriddenReturnCode(uint8 * OverriddenReturnCodeToUsePtr)
diff --git a/Engine/Source/Runtime/Core/Private/Windows/WindowsErrorOutputDevice.cpp b/Engine/Source/Runtime/Core/Private/Windows/WindowsErrorOutputDevice.cpp
--- a/Engine/Source/Runtime/Core/Private/Windows/WindowsErrorOutputDevice.cpp
+++ b/Engine/Source/Runtime/Core/Private/Windows/WindowsErrorOutputDevice.cpp
@@ -86,7 +86,7 @@ void FWindowsErrorOutputDevice::Serialize( const TCHAR* Msg, ELogVerbosity::Type
 	{
 		// We crashed outside the guarded code (e.g. appExit).
 		HandleError();
-		FPlatformMisc::RequestExit( true );
+		FPlatformMisc::RequestExit( true, TEXT("WindowsErrorOutputDevice::Serialize.!GIsGuarded"));
 	}
 }
 
diff --git a/Engine/Source/Runtime/Core/Private/Windows/WindowsPlatformCrashContext.cpp b/Engine/Source/Runtime/Core/Private/Windows/WindowsPlatformCrashContext.cpp
--- a/Engine/Source/Runtime/Core/Private/Windows/WindowsPlatformCrashContext.cpp
+++ b/Engine/Source/Runtime/Core/Private/Windows/WindowsPlatformCrashContext.cpp
@@ -1003,7 +1003,7 @@ static void AbortHandler(int Signal)
 	__except(ReportCrash(GetExceptionInformation()))
 	{
 		GIsCriticalError = true;
-		FPlatformMisc::RequestExit(true);
+		FPlatformMisc::RequestExit(true, TEXT("WindowsPlatformCrashContext.AbortHandler"));
 	}
 #endif
 }
@@ -1558,7 +1558,7 @@ LONG WINAPI EngineUnhandledExceptionFilter(LPEXCEPTION_POINTERS ExceptionInfo)
 {
 	ReportCrash(ExceptionInfo);
 	GIsCriticalError = true;
-	FPlatformMisc::RequestExit(true);
+	FPlatformMisc::RequestExit(true, TEXT("WindowsPlatformCrashContext.EngineUnhandledExceptionFilter"));
 
 	return EXCEPTION_CONTINUE_SEARCH; // Not really important, RequestExit() terminates the process just above.
 }
@@ -1751,7 +1751,7 @@ FORCENOINLINE void ReportGPUCrash(const TCHAR* ErrorMessage, void* ProgramCounte
 	}
 	__except (ReportCrash(GetExceptionInformation()))
 	{
-		FPlatformMisc::RequestExit(false);
+		FPlatformMisc::RequestExit(false, TEXT("WindowsPlatformCrashContext.ReportGPUCrash"));
 	}
 #endif
 }
diff --git a/Engine/Source/Runtime/Core/Private/Windows/WindowsPlatformMisc.cpp b/Engine/Source/Runtime/Core/Private/Windows/WindowsPlatformMisc.cpp
--- a/Engine/Source/Runtime/Core/Private/Windows/WindowsPlatformMisc.cpp
+++ b/Engine/Source/Runtime/Core/Private/Windows/WindowsPlatformMisc.cpp
@@ -582,7 +582,7 @@ void FWindowsPlatformMisc::PlatformPreInit()
 	if ( ::GetSystemMetrics(SM_CXSCREEN) < MinResolution[0] || ::GetSystemMetrics(SM_CYSCREEN) < MinResolution[1] )
 	{
 		FMessageDialog::Open( EAppMsgType::Ok, NSLOCTEXT("Launch", "Error_ResolutionTooLow", "The current resolution is too low to run this game.") );
-		FPlatformMisc::RequestExit( false );
+		FPlatformMisc::RequestExit( false, TEXT("FWindowsPlatformMisc::PlatformPreInit.ResolutionTooLow"));
 	}
 
 	// initialize the file SHA hash mapping
@@ -991,24 +991,26 @@ bool FWindowsPlatformMisc::IsLowLevelOutputDebugStringStructured()
 	return !!Mutex || FGenericPlatformMisc::IsLowLevelOutputDebugStringStructured();
 }
 
-void FWindowsPlatformMisc::RequestExit( bool Force )
+void FWindowsPlatformMisc::RequestExit( bool Force, const TCHAR* CallSite )
 {
-	UE_LOG(LogWindows, Log,  TEXT("FPlatformMisc::RequestExit(%i)"), Force );
+	UE_LOG(LogWindows, Log,  TEXT("FPlatformMisc::RequestExit(%i, %s)"),
+		Force, CallSite ? CallSite : TEXT("<NoCallSiteInfo>"));
 
 	// Legacy behavior that now calls through to RequestExitWithStatus
 	if( Force )
 	{
-		RequestExitWithStatus(Force, GIsCriticalError ? 3 : 0);
+		RequestExitWithStatus(Force, GIsCriticalError ? 3 : 0, CallSite);
 	}
 	else
 	{
-		RequestExitWithStatus(false, 0);
+		RequestExitWithStatus(false, 0, CallSite);
 	}
 }
 
-void FWindowsPlatformMisc::RequestExitWithStatus(bool Force, uint8 ReturnCode)
+void FWindowsPlatformMisc::RequestExitWithStatus(bool Force, uint8 ReturnCode, const TCHAR* CallSite)
 {
-	UE_LOG(LogWindows, Log, TEXT("FPlatformMisc::RequestExitWithStatus(%i, %i)"), Force, ReturnCode);
+	UE_LOG(LogWindows, Log, TEXT("FPlatformMisc::RequestExitWithStatus(%i, %i, %s)"), Force, ReturnCode,
+		CallSite ? CallSite : TEXT("<NoCallSiteInfo>"));
 
 #if ENABLE_PGO_PROFILE
 	// save current PGO profiling data and terminate immediately
diff --git a/Engine/Source/Runtime/Core/Private/Windows/WindowsRunnableThread.cpp b/Engine/Source/Runtime/Core/Private/Windows/WindowsRunnableThread.cpp
--- a/Engine/Source/Runtime/Core/Private/Windows/WindowsRunnableThread.cpp
+++ b/Engine/Source/Runtime/Core/Private/Windows/WindowsRunnableThread.cpp
@@ -94,7 +94,7 @@ uint32 FRunnableThreadWin::GuardedRun()
 				// Crashed.
 				ExitCode = 1;
 				GError->HandleError();
-				FPlatformMisc::RequestExit( true );
+				FPlatformMisc::RequestExit(true, TEXT("FRunnableThreadWin::GuardedRun.ExceptionHandler"));
 			}
 			__except(EXCEPTION_EXECUTE_HANDLER)
 			{
diff --git a/Engine/Source/Runtime/Core/Public/Android/AndroidPlatformMisc.h b/Engine/Source/Runtime/Core/Public/Android/AndroidPlatformMisc.h
--- a/Engine/Source/Runtime/Core/Public/Android/AndroidPlatformMisc.h
+++ b/Engine/Source/Runtime/Core/Public/Android/AndroidPlatformMisc.h
@@ -28,7 +28,7 @@ enum class ECrashContextType;
  */
 struct CORE_API FAndroidMisc : public FGenericPlatformMisc
 {
-	static void RequestExit( bool Force );
+	static void RequestExit( bool Force, const TCHAR* CallSite = nullptr);
 	static bool RestartApplication();
 	static void LocalPrint(const TCHAR *Message);
 	static bool IsLocalPrintThreadSafe() { return true; }
diff --git a/Engine/Source/Runtime/Core/Public/GenericPlatform/GenericPlatformMisc.h b/Engine/Source/Runtime/Core/Public/GenericPlatform/GenericPlatformMisc.h
--- a/Engine/Source/Runtime/Core/Public/GenericPlatform/GenericPlatformMisc.h
+++ b/Engine/Source/Runtime/Core/Public/GenericPlatform/GenericPlatformMisc.h
@@ -998,8 +998,9 @@ public:
 	 *
 	 * @param	Force	If true, perform immediate exit (dangerous because config code isn't flushed, etc).
 	 *				  If false, request clean main-loop exit from the platform specific code.
+	 * @param	CallSite a diagnostic message identifying the callsite.
 	 */
-	static void RequestExit( bool Force );
+	static void RequestExit( bool Force, const TCHAR* CallSite = nullptr);
 
 	/**
 	 * Requests application exit with a specified return code. Name is different from RequestExit() so overloads of just one of functions are possible.
@@ -1007,8 +1008,9 @@ public:
 	 * @param	Force 	   If true, perform immediate exit (dangerous because config code isn't flushed, etc).
 	 *					 If false, request clean main-loop exit from the platform specific code.
 	 * @param   ReturnCode This value will be returned from the program (on the platforms where it's possible). Limited to 0-255 to conform with POSIX.
+	 * @param	CallSite a diagnostic message identifying the callsite.
 	 */
-	static void RequestExitWithStatus( bool Force, uint8 ReturnCode );
+	static void RequestExitWithStatus( bool Force, uint8 ReturnCode, const TCHAR* CallSite = nullptr);
 
 
 	/**
diff --git a/Engine/Source/Runtime/Core/Public/IOS/IOSPlatformMisc.h b/Engine/Source/Runtime/Core/Public/IOS/IOSPlatformMisc.h
--- a/Engine/Source/Runtime/Core/Public/IOS/IOSPlatformMisc.h
+++ b/Engine/Source/Runtime/Core/Public/IOS/IOSPlatformMisc.h
@@ -46,8 +46,8 @@ struct CORE_API FIOSPlatformMisc : public FApplePlatformMisc
 	static const TCHAR* GamePersistentDownloadDir();
     static bool HasSeparateChannelForDebugOutput();
 
-	static void RequestExit(bool Force);
-	static void RequestExitWithStatus(bool Force, uint8 ReturnCode);
+	static void RequestExit(bool Force, const TCHAR* CallSite = nullptr);
+	static void RequestExitWithStatus(bool Force, uint8 ReturnCode, const TCHAR* CallSite = nullptr);
 
 	UE_DEPRECATED(4.21, "Use GetDeviceVolume, it is now callable on all platforms.")
 	static int GetAudioVolume();
diff --git a/Engine/Source/Runtime/Core/Public/Mac/MacPlatformMisc.h b/Engine/Source/Runtime/Core/Public/Mac/MacPlatformMisc.h
--- a/Engine/Source/Runtime/Core/Public/Mac/MacPlatformMisc.h
+++ b/Engine/Source/Runtime/Core/Public/Mac/MacPlatformMisc.h
@@ -31,7 +31,7 @@ struct CORE_API FMacPlatformMisc : public FApplePlatformMisc
 
 	static TArray<uint8> GetMacAddress();
 
-	static void RequestExit(bool Force);
+	static void RequestExit(bool Force, const TCHAR* CallSite = nullptr);
 	static EAppReturnType::Type MessageBoxExt( EAppMsgType::Type MsgType, const TCHAR* Text, const TCHAR* Caption );
 	static bool CommandLineCommands();
 	static int32 NumberOfCores();
diff --git a/Engine/Source/Runtime/Core/Public/Unix/UnixPlatformMisc.h b/Engine/Source/Runtime/Core/Public/Unix/UnixPlatformMisc.h
--- a/Engine/Source/Runtime/Core/Public/Unix/UnixPlatformMisc.h
+++ b/Engine/Source/Runtime/Core/Public/Unix/UnixPlatformMisc.h
@@ -42,8 +42,8 @@ struct CORE_API FUnixPlatformMisc : public FGenericPlatformMisc
 
 	static void LowLevelOutputDebugString(const TCHAR *Message);
 
-	static void RequestExit(bool Force);
-	static void RequestExitWithStatus(bool Force, uint8 ReturnCode);
+	static void RequestExit(bool Force, const TCHAR* CallSite = nullptr);
+	static void RequestExitWithStatus(bool Force, uint8 ReturnCode, const TCHAR* CallSite = nullptr);
 	static const TCHAR* GetSystemErrorMessage(TCHAR* OutBuffer, int32 BufferCount, int32 Error);
 
 	static void NormalizePath(FString& InPath);
diff --git a/Engine/Source/Runtime/Core/Public/Windows/WindowsPlatformMisc.h b/Engine/Source/Runtime/Core/Public/Windows/WindowsPlatformMisc.h
--- a/Engine/Source/Runtime/Core/Public/Windows/WindowsPlatformMisc.h
+++ b/Engine/Source/Runtime/Core/Public/Windows/WindowsPlatformMisc.h
@@ -100,8 +100,8 @@ struct CORE_API FWindowsPlatformMisc
 		return IsDebuggerPresent();
 	}
 	
-	static void RequestExitWithStatus(bool Force, uint8 ReturnCode);
-	static void RequestExit(bool Force);
+	static void RequestExitWithStatus(bool Force, uint8 ReturnCode, const TCHAR* CallSite = nullptr);
+	static void RequestExit(bool Force, const TCHAR* CallSite = nullptr);
 	static const TCHAR* GetSystemErrorMessage(TCHAR* OutBuffer, int32 BufferCount, int32 Error);
 	static void CreateGuid(struct FGuid& Result);
 	static EAppReturnType::Type MessageBoxExt( EAppMsgType::Type MsgType, const TCHAR* Text, const TCHAR* Caption );
diff --git a/Engine/Source/Runtime/D3D12RHI/Private/D3D12Util.cpp b/Engine/Source/Runtime/D3D12RHI/Private/D3D12Util.cpp
--- a/Engine/Source/Runtime/D3D12RHI/Private/D3D12Util.cpp
+++ b/Engine/Source/Runtime/D3D12RHI/Private/D3D12Util.cpp
@@ -725,7 +725,7 @@ static void TerminateOnOutOfMemory(ID3D12Device* InDevice, HRESULT D3DResult, bo
 	else
 	{
 		// Exit silently without reporting a crash because an OOM is not necessarily our fault
-		FPlatformMisc::RequestExit(true);
+		FPlatformMisc::RequestExit(true, TEXT("D3D12Util.TerminateOnOutOfMemory"));
 	}
 
 #else // PLATFORM_WINDOWS
@@ -834,7 +834,7 @@ namespace D3D12RHI
 			ReportGPUCrash(TEXT("Aftermath GPU Crash dump Triggered"), nullptr);
 
 			// Force shutdown, we can't do anything useful anymore.
-			FPlatformMisc::RequestExit(true);
+			FPlatformMisc::RequestExit(true, TEXT("D3D12Util.TerminateOnGPUCrash"));
 		}
 #endif // PLATFORM_WINDOWS
 
@@ -873,7 +873,7 @@ namespace D3D12RHI
 		UE_LOG(LogD3D12RHI, Fatal, TEXT("%s failed \n at %s:%u \n with error %s\n%s"), ANSI_TO_TCHAR(Code), ANSI_TO_TCHAR(Filename), Line, *ErrorString, *Message);
 
 		// Force shutdown, we can't do anything useful anymore.
-		FPlatformMisc::RequestExit(true);
+		FPlatformMisc::RequestExit(true, TEXT("D3D12Util.VerifyD3D12Result"));
 	}
 
 	void VerifyD3D12CreateTextureResult(HRESULT D3DResult, const ANSICHAR* Code, const ANSICHAR* Filename, uint32 Line, const D3D12_RESOURCE_DESC& TextureDesc, ID3D12Device* Device)
@@ -931,7 +931,7 @@ namespace D3D12RHI
 			*GetD3D12TextureFlagString(TextureDesc.Flags));
 
 		// Force shutdown, we can't do anything useful anymore.
-		FPlatformMisc::RequestExit(true);
+		FPlatformMisc::RequestExit(true, TEXT("D3D12Util.VerifyD3D12CreateTextureResult"));
 	}
 
 	void VerifyComRefCount(IUnknown* Object, int32 ExpectedRefs, const TCHAR* Code, const TCHAR* Filename, int32 Line)
diff --git a/Engine/Source/Runtime/Engine/Private/DemoNetDriver.cpp b/Engine/Source/Runtime/Engine/Private/DemoNetDriver.cpp
--- a/Engine/Source/Runtime/Engine/Private/DemoNetDriver.cpp
+++ b/Engine/Source/Runtime/Engine/Private/DemoNetDriver.cpp
@@ -2879,7 +2879,7 @@ void UDemoNetDriver::TickDemoPlayback(float DeltaSeconds)
 			{
 				if (FParse::Param(FCommandLine::Get(), TEXT("ExitAfterReplay")) && (!LocalPlaylistTracker || LocalPlaylistTracker->IsOnLastReplay()))
 				{
-					FPlatformMisc::RequestExit(false);
+					FPlatformMisc::RequestExit(false, TEXT("DemoNetDriver"));
 				}
 				else
 				{
diff --git a/Engine/Source/Runtime/Engine/Private/GameEngine.cpp b/Engine/Source/Runtime/Engine/Private/GameEngine.cpp
--- a/Engine/Source/Runtime/Engine/Private/GameEngine.cpp
+++ b/Engine/Source/Runtime/Engine/Private/GameEngine.cpp
@@ -1539,7 +1539,7 @@ bool UGameEngine::HandleExitCommand( const TCHAR* Cmd, FOutputDevice& Ar )
 
 	FGameDelegates::Get().GetExitCommandDelegate().Broadcast();
 
-	FPlatformMisc::RequestExit( 0 );
+	FPlatformMisc::RequestExit(false, TEXT("UGameEngine::HandleExitCommand"));
 	return true;
 }
 
@@ -1671,7 +1671,7 @@ void UGameEngine::Tick( float DeltaSeconds, bool bIdleMode )
 	if (GIsClient && (GameViewport == nullptr) && FApp::CanEverRender())
 	{
 		UE_LOG(LogEngine, Log,  TEXT("All Windows Closed") );
-		FPlatformMisc::RequestExit( 0 );
+		FPlatformMisc::RequestExit(false, TEXT("UGameEngine::Tick.ViewportClosed"));
 		return;
 	}
 
@@ -2003,6 +2003,6 @@ void UGameEngine::HandleTravelFailure_NotifyGameInstance(UWorld* World, ETravelF
 void UGameEngine::HandleBrowseToDefaultMapFailure(FWorldContext& Context, const FString& TextURL, const FString& Error)
 {
 	Super::HandleBrowseToDefaultMapFailure(Context, TextURL, Error);
-	FPlatformMisc::RequestExit(false);
+	FPlatformMisc::RequestExit(false, TEXT("UGameEngine::HandleBrowseToDefaultMapFailure"));
 }
 
diff --git a/Engine/Source/Runtime/Engine/Private/GameInstance.cpp b/Engine/Source/Runtime/Engine/Private/GameInstance.cpp
--- a/Engine/Source/Runtime/Engine/Private/GameInstance.cpp
+++ b/Engine/Source/Runtime/Engine/Private/GameInstance.cpp
@@ -654,7 +654,7 @@ void UGameInstance::StartGameInstance()
 				&& FMessageDialog::Open(EAppMsgType::OkCancel, Message) != EAppReturnType::Ok)
 			{
 				// user canceled (maybe a typo while attempting to run a commandlet)
-				FPlatformMisc::RequestExit(false);
+				FPlatformMisc::RequestExit(false, TEXT("UGameInstance::StartGameInstance.Cancelled"));
 				return;
 			}
 			else
@@ -666,7 +666,7 @@ void UGameInstance::StartGameInstance()
 		{
 			const FText Message = FText::Format(NSLOCTEXT("Engine", "MapNotFoundNoFallback", "The map specified on the commandline '{0}' could not be found. Exiting."), FText::FromString(URL.Map));
 			FMessageDialog::Open(EAppMsgType::Ok, Message);
-			FPlatformMisc::RequestExit(false);
+			FPlatformMisc::RequestExit(false, TEXT("UGameInstance::StartGameInstance.MapNotFound"));
 			return;
 		}
 	}
@@ -677,7 +677,7 @@ void UGameInstance::StartGameInstance()
 		UE_LOG(LogLoad, Error, TEXT("%s"), *FString::Printf(TEXT("Failed to enter %s: %s. Please check the log for errors."), *DefaultMap, *Error));
 		const FText Message = FText::Format(NSLOCTEXT("Engine", "DefaultMapNotFound", "The default map '{0}' could not be found. Exiting."), FText::FromString(DefaultMap));
 		FMessageDialog::Open(EAppMsgType::Ok, Message);
-		FPlatformMisc::RequestExit(false);
+		FPlatformMisc::RequestExit(false, TEXT("UGameInstance::StartGameInstance.BrowseFailure"));
 		return;
 	}
 
diff --git a/Engine/Source/Runtime/Engine/Private/NetDriver.cpp b/Engine/Source/Runtime/Engine/Private/NetDriver.cpp
--- a/Engine/Source/Runtime/Engine/Private/NetDriver.cpp
+++ b/Engine/Source/Runtime/Engine/Private/NetDriver.cpp
@@ -5055,7 +5055,7 @@ struct FReplicationAutoCapture
 			if (KillProcessFrames == 0)
 			{
 				GLog->Flush();
-				FPlatformMisc::RequestExit(1);
+				FPlatformMisc::RequestExit(true, TEXT("FReplicationAutoCapture"));
 			}
 		}
 #endif
diff --git a/Engine/Source/Runtime/Engine/Private/ShaderCompiler/ShaderCompiler.cpp b/Engine/Source/Runtime/Engine/Private/ShaderCompiler/ShaderCompiler.cpp
--- a/Engine/Source/Runtime/Engine/Private/ShaderCompiler/ShaderCompiler.cpp
+++ b/Engine/Source/Runtime/Engine/Private/ShaderCompiler/ShaderCompiler.cpp
@@ -895,7 +895,7 @@ static void ModalErrorOrLog(const FString& Text, int64 CurrentFilePos = 0, int64
 		{
 			// Show dialog box with error message and request exit
 			FMessageDialog::Open(EAppMsgType::Ok, FText::FromString(Text));
-			FPlatformMisc::RequestExit(false);
+			FPlatformMisc::RequestExit(false, TEXT("ShaderCompiler.ModalErrorOrLog"));
 		}
 		else
 		{
@@ -7936,7 +7936,7 @@ void CompileGlobalShaderMap(EShaderPlatform Platform, const ITargetPlatform* Tar
 			{
 				UE_LOG(LogShaders, Error, TEXT("%s"), *Message.ToString());
 				FMessageDialog::Open(EAppMsgType::Ok, Message);
-				FPlatformMisc::RequestExit(false);
+				FPlatformMisc::RequestExit(false, TEXT("CompileGlobalShaderMap"));
 				return;
 			}
 			else
diff --git a/Engine/Source/Runtime/Engine/Private/Tests/AutomationCommon.cpp b/Engine/Source/Runtime/Engine/Private/Tests/AutomationCommon.cpp
--- a/Engine/Source/Runtime/Engine/Private/Tests/AutomationCommon.cpp
+++ b/Engine/Source/Runtime/Engine/Private/Tests/AutomationCommon.cpp
@@ -439,7 +439,7 @@ bool FExitGameCommand::Update()
 
 bool FRequestExitCommand::Update()
 {
-	FPlatformMisc::RequestExit(true);
+	FPlatformMisc::RequestExit(true, TEXT("FRequestExitCommand"));
 	return true;
 }
 
diff --git a/Engine/Source/Runtime/Launch/Private/Android/LaunchAndroid.cpp b/Engine/Source/Runtime/Launch/Private/Android/LaunchAndroid.cpp
--- a/Engine/Source/Runtime/Launch/Private/Android/LaunchAndroid.cpp
+++ b/Engine/Source/Runtime/Launch/Private/Android/LaunchAndroid.cpp
@@ -622,7 +622,7 @@ int32 AndroidMain(struct android_app* state)
 
 	UE_LOG(LogAndroid, Log, TEXT("Exiting is over"));
 
-	FPlatformMisc::RequestExit(1);
+	FPlatformMisc::RequestExit(true, TEXT("AndroidMain"));
 	return 0;
 }
 
diff --git a/Engine/Source/Runtime/Launch/Private/LaunchEngineLoop.cpp b/Engine/Source/Runtime/Launch/Private/LaunchEngineLoop.cpp
--- a/Engine/Source/Runtime/Launch/Private/LaunchEngineLoop.cpp
+++ b/Engine/Source/Runtime/Launch/Private/LaunchEngineLoop.cpp
@@ -2822,7 +2822,7 @@ int32 FEngineLoop::PreInitPreStartupScreen(const TCHAR* CmdLine)
 		GWarn = &UnrealEdWarn;
 #else //WITH_EDITOR
 		FMessageDialog::Open(EAppMsgType::Ok, NSLOCTEXT("Engine", "EditorNotSupported", "Editor not supported in this mode."));
-		FPlatformMisc::RequestExit(false);
+		FPlatformMisc::RequestExit(false, TEXT("FEngineLoop::PreInitPreStartupScreen.bHasEditorToken"));
 		return 1;
 #endif //WITH_EDITOR
 	}
@@ -2984,7 +2984,7 @@ PRAGMA_ENABLE_DEPRECATION_WARNINGS
 
 #if CHECK_PUREVIRTUALS
 	FMessageDialog::Open(EAppMsgType::Ok, NSLOCTEXT("Engine", "Error_PureVirtualsEnabled", "The game cannot run with CHECK_PUREVIRTUALS enabled.  Please disable CHECK_PUREVIRTUALS and rebuild the executable."));
-	FPlatformMisc::RequestExit(false);
+	FPlatformMisc::RequestExit(false, TEXT("FEngineLoop::PreInitPreStartupScreen.Check_PureVirtuals"));
 #endif
 
 	FEmbeddedCommunication::ForceTick(2);
@@ -2996,7 +2996,7 @@ PRAGMA_ENABLE_DEPRECATION_WARNINGS
 	// allow for game explorer processing (including parental controls) and firewalls installation
 	if (!FPlatformMisc::CommandLineCommands())
 	{
-		FPlatformMisc::RequestExit(false);
+		FPlatformMisc::RequestExit(false, TEXT("FEngineLoop::PreInitPreStartupScreen.CommandLineCommands"));
 	}
 
 #if !UE_EDITOR
@@ -5236,7 +5236,7 @@ void DumpEarlyReads(bool bDumpEarlyConfigReads, bool bDumpEarlyPakFileReads, boo
 		}
 		else
 		{
-			FPlatformMisc::RequestExit(true);
+			FPlatformMisc::RequestExit(true, TEXT("DumpEarlyReads"));
 		}
 	}
 }
@@ -5540,7 +5540,7 @@ void FEngineLoop::Tick()
 	if (GScopedTestExit.IsValid() && GScopedTestExit->RequestExit())
 	{
 		UE_LOG(LogExit, Display, TEXT("**** TestExit: %s ****"), *GScopedTestExit->RequestExitPhrase());
-		FPlatformMisc::RequestExit(true);
+		FPlatformMisc::RequestExit(true, TEXT("FEngineLoop::Tick.GScopedTestExit"));
 	}
 #endif
 
@@ -5651,7 +5651,7 @@ void FEngineLoop::Tick()
 		if ((FApp::IsBenchmarking() && MaxFrameCounter && (GFrameCounter > MaxFrameCounter)) ||
 			(MaxTickTime && (TotalTickTime > MaxTickTime)))
 		{
-			FPlatformMisc::RequestExit(0);
+			FPlatformMisc::RequestExit(false, TEXT("FEngineLoop::Tick.Benchmarking"));
 		}
 
 		// set FApp::CurrentTime, FApp::DeltaTime and potentially wait to enforce max tick rate
diff --git a/Engine/Source/Runtime/Launch/Private/Windows/LaunchWindows.cpp b/Engine/Source/Runtime/Launch/Private/Windows/LaunchWindows.cpp
--- a/Engine/Source/Runtime/Launch/Private/Windows/LaunchWindows.cpp
+++ b/Engine/Source/Runtime/Launch/Private/Windows/LaunchWindows.cpp
@@ -257,7 +257,7 @@ LAUNCH_API int32 LaunchWindowsStartup( HINSTANCE hInInstance, HINSTANCE hPrevIns
 			}
 			LaunchStaticShutdownAfterError();
 			FPlatformMallocCrash::Get().PrintPoolsUsage();
-			FPlatformMisc::RequestExit( true );
+			FPlatformMisc::RequestExit( true, TEXT("LaunchWindowsStartup.ExceptionHandler"));
 		}
 #endif
 	}
diff --git a/Engine/Source/Runtime/MovieSceneCapture/Private/MovieSceneCaptureModule.cpp b/Engine/Source/Runtime/MovieSceneCapture/Private/MovieSceneCaptureModule.cpp
--- a/Engine/Source/Runtime/MovieSceneCapture/Private/MovieSceneCaptureModule.cpp
+++ b/Engine/Source/Runtime/MovieSceneCapture/Private/MovieSceneCaptureModule.cpp
@@ -56,7 +56,7 @@ private:
 			if (!GameEngine->SceneViewport->GetClient()->GetWorld())
 			{
 				// @todo: Set exit code to EMovieSceneCaptureExitCode::WorldNotFound when we have the ability to do so
-				FPlatformMisc::RequestExit(false);
+				FPlatformMisc::RequestExit(false, TEXT("FMovieSceneCaptureModule::OnPostLoadMap"));
 			}
 			else
 			{
@@ -193,7 +193,7 @@ private:
 		FActiveMovieSceneCaptures::Get().Add( Capture );
 
 		Capture->OnCaptureFinished().AddLambda([]{
-			FPlatformMisc::RequestExit(0);
+			FPlatformMisc::RequestExit(false, TEXT("FMovieSceneCaptureModule::InitializeFromCommandLine"));
 		});
 		
 		return Capture;
diff --git a/Engine/Source/Runtime/OpenGLDrv/Private/Linux/OpenGLLinux.cpp b/Engine/Source/Runtime/OpenGLDrv/Private/Linux/OpenGLLinux.cpp
--- a/Engine/Source/Runtime/OpenGLDrv/Private/Linux/OpenGLLinux.cpp
+++ b/Engine/Source/Runtime/OpenGLDrv/Private/Linux/OpenGLLinux.cpp
@@ -128,7 +128,7 @@ void Linux_PlatformCreateDummyGLWindow( FPlatformOpenGLContext *OutContext )
 		FPlatformMisc::MessageBoxExt(EAppMsgType::Ok,
 			*ErrorMessage.ToString(),
 			*(NSLOCTEXT("Renderer", "LinuxCannotCreatePlatformWindowTitle", "Cannot create SDL window.").ToString()));
-		FPlatformMisc::RequestExit(true);
+		FPlatformMisc::RequestExit(true, TEXT("Linux_PlatformCreateDummyGLWindow"));
 		// unreachable
 		return;
 	}
@@ -218,7 +218,7 @@ struct FPlatformOpenGLDevice
 			FPlatformMisc::MessageBoxExt(EAppMsgType::Ok,
 				*(NSLOCTEXT("Renderer", "LinuxInsufficientDriversText", "Cannot create OpenGL context. Check that the drivers and hardware support at least OpenGL 4.3 (or re-run with -opengl3)").ToString()),
 				*(NSLOCTEXT("Renderer", "LinuxInsufficientDriversTitle", "Insufficient drivers or hardware").ToString()));
-			FPlatformMisc::RequestExit(true);
+			FPlatformMisc::RequestExit(true, TEXT("FPlatformOpenGLDevice()"));
 			// unreachable
 			return;
 		}
@@ -764,7 +764,7 @@ bool PlatformInitOpenGL()
 			FPlatformMisc::MessageBoxExt(EAppMsgType::Ok,
 				*FString::Printf(TEXT("%s. SDL error: \"%s\""), *(NSLOCTEXT("Renderer", "LinuxCannotLoadLibGLText", "Unable to dynamically load libGL").ToString()), UTF8_TO_TCHAR(SDL_GetError())),
 				*(NSLOCTEXT("Renderer", "LinuxInsufficientDriversTitle", "Insufficient drivers or hardware").ToString()));
-			FPlatformMisc::RequestExit(true);
+			FPlatformMisc::RequestExit(true, TEXT("PlatformInitOpenGL"));
 			// unreachable
 			return false;
 		}
diff --git a/Engine/Source/Runtime/RHI/Private/Android/AndroidDynamicRHI.cpp b/Engine/Source/Runtime/RHI/Private/Android/AndroidDynamicRHI.cpp
--- a/Engine/Source/Runtime/RHI/Private/Android/AndroidDynamicRHI.cpp
+++ b/Engine/Source/Runtime/RHI/Private/Android/AndroidDynamicRHI.cpp
@@ -41,7 +41,7 @@ FDynamicRHI* PlatformCreateDynamicRHI()
 	{
 
 	//	FMessageDialog::Open(EAppMsgType::Ok, TEXT("OpenGL 3.2 is required to run the engine."));
-		FPlatformMisc::RequestExit(1);
+		FPlatformMisc::RequestExit(true, TEXT("PlatformCreateDynamicRHI"));
 		DynamicRHIModule = NULL;
 	}
 
diff --git a/Engine/Source/Runtime/RHI/Private/DynamicRHI.cpp b/Engine/Source/Runtime/RHI/Private/DynamicRHI.cpp
--- a/Engine/Source/Runtime/RHI/Private/DynamicRHI.cpp
+++ b/Engine/Source/Runtime/RHI/Private/DynamicRHI.cpp
@@ -74,7 +74,7 @@ void InitNullRHI()
 	if ((DynamicRHIModule == 0) || !DynamicRHIModule->IsSupported())
 	{
 		FMessageDialog::Open( EAppMsgType::Ok, NSLOCTEXT("DynamicRHI", "NullDrvFailure", "NullDrv failure?"));
-		FPlatformMisc::RequestExit(1);
+		FPlatformMisc::RequestExit(true, TEXT("InitNullRHI"));
 	}
 
 	GDynamicRHI = DynamicRHIModule->CreateRHI();
@@ -364,7 +364,7 @@ void RHIInit(bool bHasEditorToken)
 					FString ShaderPlatformString = LegacyShaderPlatformToShaderFormat(GetFeatureLevelShaderPlatform(GMaxRHIFeatureLevel)).ToString();
 					FString Error = FString::Printf(TEXT("A Feature Level 5 video card is required to run the editor.\nAvailableFeatureLevel = %s, ShaderPlatform = %s"), *FeatureLevelString, *ShaderPlatformString);
 					FMessageDialog::Open(EAppMsgType::Ok, FText::FromString(Error));
-					FPlatformMisc::RequestExit(1);
+					FPlatformMisc::RequestExit(true, TEXT("RHIInit"));
 				}
 
 				// Update the crash context analytics
diff --git a/Engine/Source/Runtime/RHI/Private/Windows/WindowsDynamicRHI.cpp b/Engine/Source/Runtime/RHI/Private/Windows/WindowsDynamicRHI.cpp
--- a/Engine/Source/Runtime/RHI/Private/Windows/WindowsDynamicRHI.cpp
+++ b/Engine/Source/Runtime/RHI/Private/Windows/WindowsDynamicRHI.cpp
@@ -808,7 +808,7 @@ static bool HandleUnsupportedFeatureLevel(EWindowsRHI& WindowsRHI, ERHIFeatureLe
 			UE_LOG(LogRHI, Log, TEXT("RHI %s with Feature Level %s is not supported on your system."), GetLogName(WindowsRHI), GetLogName(FeatureLevel));
 
 			FMessageDialog::Open(EAppMsgType::Ok, LOCTEXT("RequiredDX12SM6", "DirectX 12 with Feature Level SM6 is not supported on your system. Try running without the -sm6 command line argument."));
-			FPlatformMisc::RequestExit(1);
+			FPlatformMisc::RequestExit(true, TEXT("HandleUnsupportedFeatureLevel"));
 		}
 
 		return false;
@@ -837,7 +837,7 @@ static bool HandleUnsupportedRHI(EWindowsRHI& WindowsRHI, ERHIFeatureLevel::Type
 			UE_LOG(LogRHI, Log, TEXT("RHI %s with Feature Level %s is not supported on your system."), GetLogName(WindowsRHI), GetLogName(FeatureLevel));
 
 			FMessageDialog::Open(EAppMsgType::Ok, LOCTEXT("RequiredDX12", "DirectX 12 is not supported on your system. Try running without the -dx12 or -d3d12 command line argument."));
-			FPlatformMisc::RequestExit(1);
+			FPlatformMisc::RequestExit(true, TEXT("HandleUnsupportedRHI.ForcedRHI"));
 			return false;
 		}
 	}
@@ -857,25 +857,25 @@ static bool HandleUnsupportedRHI(EWindowsRHI& WindowsRHI, ERHIFeatureLevel::Type
 	if (WindowsRHI == EWindowsRHI::D3D12)
 	{
 		FMessageDialog::Open(EAppMsgType::Ok, LOCTEXT("RequiredDX12", "DirectX 12 is not supported on your system. Try running without the -dx12 or -d3d12 command line argument."));
-		FPlatformMisc::RequestExit(1);
+		FPlatformMisc::RequestExit(true, TEXT("HandleUnsupportedRHI.D3D12"));
 	}
 
 	if (WindowsRHI == EWindowsRHI::D3D11)
 	{
 		FMessageDialog::Open(EAppMsgType::Ok, LOCTEXT("RequiredDX11Feature_11_SM5", "A D3D11-compatible GPU (Feature Level 11.0, Shader Model 5.0) is required to run the engine."));
-		FPlatformMisc::RequestExit(1);
+		FPlatformMisc::RequestExit(true, TEXT("HandleUnsupportedRHI.D3D11"));
 	}
 
 	if (WindowsRHI == EWindowsRHI::Vulkan)
 	{
 		FMessageDialog::Open(EAppMsgType::Ok, LOCTEXT("RequiredVulkan", "Vulkan Driver is required to run the engine."));
-		FPlatformMisc::RequestExit(1);
+		FPlatformMisc::RequestExit(true, TEXT("HandleUnsupportedRHI.Vulkan"));
 	}
 
 	if (WindowsRHI == EWindowsRHI::OpenGL)
 	{
 		FMessageDialog::Open(EAppMsgType::Ok, LOCTEXT("RequiredOpenGL", "OpenGL 4.3 is required to run the engine."));
-		FPlatformMisc::RequestExit(1);
+		FPlatformMisc::RequestExit(true, TEXT("HandleUnsupportedRHI.OpenGL"));
 	}
 
 	return false;
diff --git a/Engine/Source/Runtime/RenderCore/Private/ShaderCodeLibrary.cpp b/Engine/Source/Runtime/RenderCore/Private/ShaderCodeLibrary.cpp
--- a/Engine/Source/Runtime/RenderCore/Private/ShaderCodeLibrary.cpp
+++ b/Engine/Source/Runtime/RenderCore/Private/ShaderCodeLibrary.cpp
@@ -3344,7 +3344,7 @@ void FShaderCodeLibrary::InitForRuntime(EShaderPlatform ShaderPlatform)
                 FMessageDialog::Open(EAppMsgType::Ok, NSLOCTEXT("MessageDialog", "MissingGlobalShaderLibraryFilesClient_Body", "Game files required to initialize the global shader and cooked content are most likely missing. Refer to Engine log for details."));
 				UE_LOG(LogShaderLibrary, Fatal, TEXT("Failed to initialize ShaderCodeLibrary required by the project because part of the Global shader library is missing from %s."), *FPaths::ProjectContentDir());
 			}
-			FPlatformMisc::RequestExit(true);
+			FPlatformMisc::RequestExit(true, TEXT("FShaderCodeLibrary::InitForRuntime"));
 #endif // !WITH_EDITOR	
 		}
 	}
diff --git a/Engine/Source/Runtime/Windows/D3D11RHI/Private/D3D11Util.cpp b/Engine/Source/Runtime/Windows/D3D11RHI/Private/D3D11Util.cpp
--- a/Engine/Source/Runtime/Windows/D3D11RHI/Private/D3D11Util.cpp
+++ b/Engine/Source/Runtime/Windows/D3D11RHI/Private/D3D11Util.cpp
@@ -187,7 +187,7 @@ static void TerminateOnDeviceRemoved(HRESULT D3DResult, ID3D11Device* Direct3DDe
 		// Workaround for the fact that in non-monolithic builds the exe gets into a weird state and exception handling fails. 
 		// @todo investigate why non-monolithic builds fail to capture the exception when graphics driver crashes.
 #if !IS_MONOLITHIC
-		FPlatformMisc::RequestExit(true);
+		FPlatformMisc::RequestExit(true, TEXT("TerminateOnDeviceRemoved"));
 #endif
 	}
 }
@@ -233,7 +233,7 @@ static void TerminateOnOutOfMemory(HRESULT D3DResult, bool bCreatingTextures)
 		}
 		else
 		{
-			FPlatformMisc::RequestExit(true);
+			FPlatformMisc::RequestExit(true, TEXT("TerminateOnOutOfMemory"));
 		}
 	}
 }
