Attempted fix for static analysis false positive.

diff --git a/Engine/Source/Runtime/Core/Public/UObject/ScriptDelegates.h b/Engine/Source/Runtime/Core/Public/UObject/ScriptDelegates.h
--- a/Engine/Source/Runtime/Core/Public/UObject/ScriptDelegates.h
+++ b/Engine/Source/Runtime/Core/Public/UObject/ScriptDelegates.h
@@ -34,11 +34,24 @@ namespace UE::Core::Private
 		using WeakPtrType = FWeakObjectPtr;
 	};
 
-	// This type only exists to allow compatibility between multicast and unicast delegate types which use an explicit FWeakObjectPtr template parameter
+	// This function only exists to allow compatibility between multicast and unicast delegate types which use an explicit FWeakObjectPtr template parameter
 	template <typename From, typename To>
-	inline constexpr /*UE_DEPRECATED(5.3, "Deprecated - remove after TScriptDelegateTraits<FWeakObjectPtr> is removed")*/ bool TBackwardCompatibilityCheck_V =
-		(std::is_same_v<From, TScriptDelegateDefault> && std::is_same_v<To, FWeakObjectPtr>) ||
-		(std::is_same_v<From, FWeakObjectPtr> && std::is_same_v<To, TScriptDelegateDefault>);
+	/* UE_DEPRECATED(5.3, "Deprecated - remove after TScriptDelegateTraits<FWeakObjectPtr> is removed") */
+	inline constexpr bool BackwardCompatibilityCheck()
+	{
+		if constexpr (std::is_same_v<From, TScriptDelegateDefault>)
+		{
+			return std::is_same_v<To, FWeakObjectPtr>;
+		}
+		else if constexpr (std::is_same_v<From, FWeakObjectPtr>)
+		{
+			return std::is_same_v<To, TScriptDelegateDefault>;
+		}
+		else
+		{
+			return false;
+		}
+	}
 }
 
 /**
@@ -62,7 +75,7 @@ public:
 
 	template <
 		typename OtherDummy
-		UE_REQUIRES(UE::Core::Private::TBackwardCompatibilityCheck_V<Dummy, OtherDummy>)
+		UE_REQUIRES(UE::Core::Private::BackwardCompatibilityCheck<Dummy, OtherDummy>())
 	>
 	/* UE_DEPRECATED(5.3, "Deprecated - remove after TScriptDelegateTraits<FWeakObjectPtr> is removed") */
 	TScriptDelegate(const TScriptDelegate<OtherDummy>& Other)
@@ -214,7 +227,7 @@ public:
 	}
 	template <
 		typename OtherDummy
-		UE_REQUIRES(UE::Core::Private::TBackwardCompatibilityCheck_V<Dummy, OtherDummy>)
+		UE_REQUIRES(UE::Core::Private::BackwardCompatibilityCheck<Dummy, OtherDummy>())
 	>
 	/* UE_DEPRECATED(5.3, "Deprecated - remove after TScriptDelegateTraits<FWeakObjectPtr> is removed") */
 	FORCEINLINE bool operator==(const TScriptDelegate<OtherDummy>& Other) const
@@ -232,7 +245,7 @@ public:
 	}
 	template <
 		typename OtherDummy
-		UE_REQUIRES(UE::Core::Private::TBackwardCompatibilityCheck_V<Dummy, OtherDummy>)
+		UE_REQUIRES(UE::Core::Private::BackwardCompatibilityCheck<Dummy, OtherDummy>())
 	>
 	/* UE_DEPRECATED(5.3, "Deprecated - remove after TScriptDelegateTraits<FWeakObjectPtr> is removed") */
 	FORCEINLINE bool operator!=(const TScriptDelegate<OtherDummy>& Other) const
@@ -251,7 +264,7 @@ public:
 	}
 	template <
 		typename OtherDummy
-		UE_REQUIRES(UE::Core::Private::TBackwardCompatibilityCheck_V<Dummy, OtherDummy>)
+		UE_REQUIRES(UE::Core::Private::BackwardCompatibilityCheck<Dummy, OtherDummy>())
 	>
 	/* UE_DEPRECATED(5.3, "Deprecated - remove after TScriptDelegateTraits<FWeakObjectPtr> is removed") */
 	FORCEINLINE bool operator=(const TScriptDelegate<OtherDummy>& Other)
@@ -436,7 +449,7 @@ public:
 	}
 	template <
 		typename OtherDummy
-		UE_REQUIRES(UE::Core::Private::TBackwardCompatibilityCheck_V<Dummy, OtherDummy>)
+		UE_REQUIRES(UE::Core::Private::BackwardCompatibilityCheck<Dummy, OtherDummy>())
 	>
 	/* UE_DEPRECATED(5.3, "Deprecated - remove after TScriptDelegateTraits<FWeakObjectPtr> is removed") */
 	FORCEINLINE bool Contains(const TScriptDelegate<OtherDummy>& InDelegate) const
@@ -477,7 +490,7 @@ public:
 	}
 	template <
 		typename OtherDummy
-		UE_REQUIRES(UE::Core::Private::TBackwardCompatibilityCheck_V<Dummy, OtherDummy>)
+		UE_REQUIRES(UE::Core::Private::BackwardCompatibilityCheck<Dummy, OtherDummy>())
 	>
 	/* UE_DEPRECATED(5.3, "Deprecated - remove after TScriptDelegateTraits<FWeakObjectPtr> is removed") */
 	FORCEINLINE void Add(const TScriptDelegate<OtherDummy>& InDelegate)
@@ -503,7 +516,7 @@ public:
 	}
 	template <
 		typename OtherDummy
-		UE_REQUIRES(UE::Core::Private::TBackwardCompatibilityCheck_V<Dummy, OtherDummy>)
+		UE_REQUIRES(UE::Core::Private::BackwardCompatibilityCheck<Dummy, OtherDummy>())
 	>
 	/* UE_DEPRECATED(5.3, "Deprecated - remove after TScriptDelegateTraits<FWeakObjectPtr> is removed") */
 	FORCEINLINE void AddUnique(const TScriptDelegate<OtherDummy>& InDelegate)
@@ -529,7 +542,7 @@ public:
 	}
 	template <
 		typename OtherDummy
-		UE_REQUIRES(UE::Core::Private::TBackwardCompatibilityCheck_V<Dummy, OtherDummy>)
+		UE_REQUIRES(UE::Core::Private::BackwardCompatibilityCheck<Dummy, OtherDummy>())
 	>
 	/* UE_DEPRECATED(5.3, "Deprecated - remove after TScriptDelegateTraits<FWeakObjectPtr> is removed") */
 	FORCEINLINE void Remove(const TScriptDelegate<OtherDummy>& InDelegate)
