Renamed CycleTime to MonotonicTime and moved to Seconds() instead of Cycles64()

Using a double in units of seconds for storage has several benefits compared to uint64 Cycles.

- Avoids arithmetic errors (like overflow) gracefully with infinity and NaN.
- Infinity and NaN are part of the type.
- Consistent units across every platform.
- Immediately understandable when viewed in a debugger.
- No awkward edge cases with unsigned cycles and signed time spans.
- Conversions between clocks or units required arithmetic with doubles anyway. seconds or other time types.

#rnx

diff --git a/Engine/Source/Runtime/Core/Private/Misc/CycleTime.cpp b/Engine/Source/Runtime/Core/Private/Misc/CycleTime.cpp
deleted file mode 100644
diff --git a/Engine/Source/Runtime/Core/Private/Misc/MonotonicTime.cpp b/Engine/Source/Runtime/Core/Private/Misc/MonotonicTime.cpp
new file mode 100644
--- /dev/null
+++ b/Engine/Source/Runtime/Core/Private/Misc/MonotonicTime.cpp
@@ -0,0 +1,15 @@
+// Copyright Epic Games, Inc. All Rights Reserved.
+
+#include "Misc/MonotonicTime.h"
+
+#include "HAL/PlatformTime.h"
+
+namespace UE
+{
+
+FMonotonicTimePoint FMonotonicTimePoint::Now()
+{
+	return FromSeconds(FPlatformTime::Seconds());
+}
+
+} // namespace UE
diff --git a/Engine/Source/Runtime/Core/Public/Misc/CycleTime.h b/Engine/Source/Runtime/Core/Public/Misc/CycleTime.h
deleted file mode 100644
diff --git a/Engine/Source/Runtime/Core/Public/Misc/MonotonicTime.h b/Engine/Source/Runtime/Core/Public/Misc/MonotonicTime.h
new file mode 100644
--- /dev/null
+++ b/Engine/Source/Runtime/Core/Public/Misc/MonotonicTime.h
@@ -0,0 +1,115 @@
+// Copyright Epic Games, Inc. All Rights Reserved.
+
+#pragma once
+
+#include "CoreTypes.h"
+#include <limits>
+
+#define UE_API CORE_API
+
+namespace UE
+{
+
+/**
+ * A span of time measured in seconds between two time points.
+ *
+ * @see FMonotonicTimePoint
+ */
+struct FMonotonicTimeSpan
+{
+public:
+	constexpr FMonotonicTimeSpan() = default;
+
+	constexpr static FMonotonicTimeSpan Zero() { return FMonotonicTimeSpan(); }
+
+	constexpr static FMonotonicTimeSpan Infinity() { return FromSeconds(std::numeric_limits<double>::infinity()); }
+
+	constexpr static FMonotonicTimeSpan FromSeconds(double Seconds)
+	{
+		FMonotonicTimeSpan TimeSpan;
+		TimeSpan.Time = Seconds;
+		return TimeSpan;
+	}
+
+	constexpr static FMonotonicTimeSpan FromMilliseconds(double Milliseconds) { return FromSeconds(Milliseconds * 0.001); }
+
+	constexpr double ToSeconds() const { return Time; }
+	constexpr double ToMilliseconds() const { return Time * 1000.0; }
+
+	constexpr bool IsInfinity() const { return *this == Infinity() || *this == -Infinity(); }
+	constexpr bool IsNaN() const { return Time != Time; }
+
+	constexpr bool operator==(const FMonotonicTimeSpan Other) const { return Time == Other.Time; }
+	constexpr bool operator!=(const FMonotonicTimeSpan Other) const { return Time != Other.Time; }
+	constexpr bool operator<=(const FMonotonicTimeSpan Other) const { return Time <= Other.Time; }
+	constexpr bool operator< (const FMonotonicTimeSpan Other) const { return Time <  Other.Time; }
+	constexpr bool operator>=(const FMonotonicTimeSpan Other) const { return Time >= Other.Time; }
+	constexpr bool operator> (const FMonotonicTimeSpan Other) const { return Time >  Other.Time; }
+
+	constexpr FMonotonicTimeSpan operator+() const { return *this; }
+	constexpr FMonotonicTimeSpan operator-() const { return FromSeconds(-Time); }
+
+	constexpr FMonotonicTimeSpan operator+(const FMonotonicTimeSpan Span) const { return FromSeconds(Time + Span.Time); }
+	constexpr FMonotonicTimeSpan operator-(const FMonotonicTimeSpan Span) const { return FromSeconds(Time - Span.Time); }
+
+	constexpr FMonotonicTimeSpan& operator+=(const FMonotonicTimeSpan Span) { return *this = *this + Span; }
+	constexpr FMonotonicTimeSpan& operator-=(const FMonotonicTimeSpan Span) { return *this = *this - Span; }
+
+private:
+	double Time = 0;
+};
+
+/**
+ * A point in time measured in seconds since an arbitrary epoch.
+ *
+ * This is a monotonic clock which means the current time will never decrease. This time is meant
+ * primarily for measuring intervals. The interval between ticks of this clock is constant except
+ * for the time that the system is suspended on certain platforms. The tick frequency will differ
+ * between platforms, and must not be used as a means of communicating time without communicating
+ * the tick frequency together with the time.
+ */
+struct FMonotonicTimePoint
+{
+public:
+	UE_API static FMonotonicTimePoint Now();
+
+	constexpr FMonotonicTimePoint() = default;
+
+	constexpr static FMonotonicTimePoint Infinity() { return FromSeconds(std::numeric_limits<double>::infinity()); }
+
+	/** Construct from seconds since the epoch. */
+	constexpr static FMonotonicTimePoint FromSeconds(const double Seconds)
+	{
+		FMonotonicTimePoint TimePoint;
+		TimePoint.Time = Seconds;
+		return TimePoint;
+	}
+
+	/** Seconds since the epoch. */
+	constexpr double ToSeconds() const { return Time; }
+
+	constexpr bool IsInfinity() const { return *this == Infinity(); }
+	constexpr bool IsNaN() const { return Time != Time; }
+
+	constexpr bool operator==(const FMonotonicTimePoint Other) const { return Time == Other.Time; }
+	constexpr bool operator!=(const FMonotonicTimePoint Other) const { return Time != Other.Time; }
+	constexpr bool operator<=(const FMonotonicTimePoint Other) const { return Time <= Other.Time; }
+	constexpr bool operator< (const FMonotonicTimePoint Other) const { return Time <  Other.Time; }
+	constexpr bool operator>=(const FMonotonicTimePoint Other) const { return Time >= Other.Time; }
+	constexpr bool operator> (const FMonotonicTimePoint Other) const { return Time >  Other.Time; }
+
+	constexpr FMonotonicTimePoint operator+(const FMonotonicTimeSpan Span) const { return FromSeconds(Time + Span.ToSeconds()); }
+	constexpr FMonotonicTimePoint operator-(const FMonotonicTimeSpan Span) const { return FromSeconds(Time - Span.ToSeconds()); }
+
+	constexpr FMonotonicTimeSpan operator-(const FMonotonicTimePoint Point) const
+	{
+		return FMonotonicTimeSpan::FromSeconds(Time - Point.Time);
+	}
+
+private:
+	double Time = 0;
+};
+
+} // namespace UE
+
+#undef UE_API
diff --git a/Engine/Source/Runtime/Core/Tests/Misc/CycleTimeTest.cpp b/Engine/Source/Runtime/Core/Tests/Misc/CycleTimeTest.cpp
deleted file mode 100644
diff --git a/Engine/Source/Runtime/Core/Tests/Misc/MonotonicTimeTest.cpp b/Engine/Source/Runtime/Core/Tests/Misc/MonotonicTimeTest.cpp
new file mode 100644
--- /dev/null
+++ b/Engine/Source/Runtime/Core/Tests/Misc/MonotonicTimeTest.cpp
@@ -0,0 +1,142 @@
+// Copyright Epic Games, Inc. All Rights Reserved.
+
+#if WITH_LOW_LEVEL_TESTS
+
+#include "Misc/MonotonicTime.h"
+
+#include "TestHarness.h"
+
+namespace UE
+{
+
+TEST_CASE("Core::Time::MonotonicTimeSpan", "[Core][Time][Smoke]")
+{
+	SECTION("Constructors")
+	{
+		STATIC_CHECK(FMonotonicTimeSpan().ToSeconds() == 0.0);
+		STATIC_CHECK(FMonotonicTimeSpan::FromSeconds(0.0).ToSeconds() == 0.0);
+		STATIC_CHECK(FMonotonicTimeSpan::FromSeconds(123.0).ToSeconds() == 123.0);
+		STATIC_CHECK(FMonotonicTimeSpan::Zero().ToSeconds() == 0.0);
+	}
+
+	SECTION("Comparison")
+	{
+		STATIC_CHECK(FMonotonicTimeSpan::FromSeconds(0.0) == FMonotonicTimeSpan());
+		STATIC_CHECK(FMonotonicTimeSpan::FromSeconds(1.0) != FMonotonicTimeSpan());
+		STATIC_CHECK(FMonotonicTimeSpan::FromSeconds(0.0) <= FMonotonicTimeSpan());
+		STATIC_CHECK(FMonotonicTimeSpan::FromSeconds(0.0) >= FMonotonicTimeSpan());
+		STATIC_CHECK(FMonotonicTimeSpan::FromSeconds(0.0) < FMonotonicTimeSpan::FromSeconds(1.0));
+		STATIC_CHECK(FMonotonicTimeSpan::FromSeconds(1.0) > FMonotonicTimeSpan::FromSeconds(0.0));
+	}
+
+	SECTION("Infinity")
+	{
+		STATIC_CHECK(FMonotonicTimeSpan::Infinity().IsInfinity());
+		STATIC_CHECK_FALSE(FMonotonicTimeSpan().IsInfinity());
+		STATIC_CHECK_FALSE(FMonotonicTimeSpan::FromSeconds(123.0).IsInfinity());
+
+		STATIC_CHECK(FMonotonicTimeSpan::Infinity() == FMonotonicTimeSpan::Infinity());
+		STATIC_CHECK(FMonotonicTimeSpan::Infinity() <= FMonotonicTimeSpan::Infinity());
+		STATIC_CHECK(FMonotonicTimeSpan::Infinity() >= FMonotonicTimeSpan::Infinity());
+
+		STATIC_CHECK(FMonotonicTimeSpan::Infinity() >= FMonotonicTimeSpan());
+		STATIC_CHECK(FMonotonicTimeSpan::Infinity() > FMonotonicTimeSpan());
+
+		STATIC_CHECK(FMonotonicTimeSpan() != FMonotonicTimeSpan::Infinity());
+		STATIC_CHECK(FMonotonicTimeSpan() <= FMonotonicTimeSpan::Infinity());
+		STATIC_CHECK(FMonotonicTimeSpan() < FMonotonicTimeSpan::Infinity());
+	}
+
+	SECTION("Addition")
+	{
+		STATIC_CHECK(FMonotonicTimeSpan::FromSeconds(234.0) + FMonotonicTimeSpan::FromSeconds(123.0) == FMonotonicTimeSpan::FromSeconds(357.0));
+		STATIC_CHECK(FMonotonicTimeSpan::FromSeconds(234.0) + (-FMonotonicTimeSpan::Infinity()) == -FMonotonicTimeSpan::Infinity());
+		STATIC_CHECK(FMonotonicTimeSpan::FromSeconds(234.0) + FMonotonicTimeSpan::Infinity() == FMonotonicTimeSpan::Infinity());
+		STATIC_CHECK(FMonotonicTimeSpan::Infinity() + FMonotonicTimeSpan::FromSeconds(1.0) == FMonotonicTimeSpan::Infinity());
+		STATIC_CHECK(FMonotonicTimeSpan::Infinity() + FMonotonicTimeSpan::Infinity() == FMonotonicTimeSpan::Infinity());
+		STATIC_CHECK((-FMonotonicTimeSpan::Infinity()) + FMonotonicTimeSpan::FromSeconds(1.0) == -FMonotonicTimeSpan::Infinity());
+		STATIC_CHECK((-FMonotonicTimeSpan::Infinity()) + (-FMonotonicTimeSpan::Infinity()) == -FMonotonicTimeSpan::Infinity());
+		CHECK((FMonotonicTimeSpan::Infinity() + (-FMonotonicTimeSpan::Infinity())).IsNaN());
+		CHECK(((-FMonotonicTimeSpan::Infinity()) + FMonotonicTimeSpan::Infinity()).IsNaN());
+	}
+
+	SECTION("Subtraction")
+	{
+		STATIC_CHECK(FMonotonicTimeSpan::FromSeconds(234.0) - FMonotonicTimeSpan::FromSeconds(123.0) == FMonotonicTimeSpan::FromSeconds(111.0));
+		STATIC_CHECK(FMonotonicTimeSpan::FromSeconds(234.0) - (-FMonotonicTimeSpan::Infinity()) == FMonotonicTimeSpan::Infinity());
+		STATIC_CHECK(FMonotonicTimeSpan::FromSeconds(234.0) - FMonotonicTimeSpan::Infinity() == -FMonotonicTimeSpan::Infinity());
+		STATIC_CHECK(FMonotonicTimeSpan::Infinity() - FMonotonicTimeSpan::FromSeconds(1.0) == FMonotonicTimeSpan::Infinity());
+		STATIC_CHECK(FMonotonicTimeSpan::Infinity() - (-FMonotonicTimeSpan::Infinity()) == FMonotonicTimeSpan::Infinity());
+		STATIC_CHECK((-FMonotonicTimeSpan::Infinity()) - FMonotonicTimeSpan::FromSeconds(1.0) == -FMonotonicTimeSpan::Infinity());
+		STATIC_CHECK((-FMonotonicTimeSpan::Infinity()) - FMonotonicTimeSpan::Infinity() == -FMonotonicTimeSpan::Infinity());
+		CHECK((FMonotonicTimeSpan::Infinity() - FMonotonicTimeSpan::Infinity()).IsNaN());
+		CHECK(((-FMonotonicTimeSpan::Infinity()) - (-FMonotonicTimeSpan::Infinity())).IsNaN());
+	}
+}
+
+TEST_CASE("Core::Time::MonotonicTimePoint", "[Core][Time][Smoke]")
+{
+	SECTION("Constructors")
+	{
+		STATIC_CHECK(FMonotonicTimePoint().ToSeconds() == 0.0);
+		STATIC_CHECK(FMonotonicTimePoint::FromSeconds(0.0).ToSeconds() == 0.0);
+		STATIC_CHECK(FMonotonicTimePoint::FromSeconds(123.0).ToSeconds() == 123.0);
+	}
+
+	SECTION("Comparison")
+	{
+		STATIC_CHECK(FMonotonicTimePoint::FromSeconds(0.0) == FMonotonicTimePoint());
+		STATIC_CHECK(FMonotonicTimePoint::FromSeconds(1.0) != FMonotonicTimePoint());
+		STATIC_CHECK(FMonotonicTimePoint::FromSeconds(0.0) <= FMonotonicTimePoint());
+		STATIC_CHECK(FMonotonicTimePoint::FromSeconds(0.0) >= FMonotonicTimePoint());
+		STATIC_CHECK(FMonotonicTimePoint::FromSeconds(0.0) < FMonotonicTimePoint::FromSeconds(1.0));
+		STATIC_CHECK(FMonotonicTimePoint::FromSeconds(1.0) > FMonotonicTimePoint::FromSeconds(0.0));
+	}
+
+	SECTION("Infinity")
+	{
+		STATIC_CHECK(FMonotonicTimePoint::Infinity().IsInfinity());
+		STATIC_CHECK_FALSE(FMonotonicTimePoint().IsInfinity());
+		STATIC_CHECK_FALSE(FMonotonicTimePoint::FromSeconds(123.0).IsInfinity());
+
+		STATIC_CHECK(FMonotonicTimePoint::Infinity() == FMonotonicTimePoint::Infinity());
+		STATIC_CHECK(FMonotonicTimePoint::Infinity() <= FMonotonicTimePoint::Infinity());
+		STATIC_CHECK(FMonotonicTimePoint::Infinity() >= FMonotonicTimePoint::Infinity());
+
+		STATIC_CHECK(FMonotonicTimePoint::Infinity() >= FMonotonicTimePoint());
+		STATIC_CHECK(FMonotonicTimePoint::Infinity() > FMonotonicTimePoint());
+
+		STATIC_CHECK(FMonotonicTimePoint() != FMonotonicTimePoint::Infinity());
+		STATIC_CHECK(FMonotonicTimePoint() <= FMonotonicTimePoint::Infinity());
+		STATIC_CHECK(FMonotonicTimePoint() < FMonotonicTimePoint::Infinity());
+	}
+
+	SECTION("Addition")
+	{
+		STATIC_CHECK(FMonotonicTimePoint::FromSeconds(234.0) + FMonotonicTimeSpan::FromSeconds(123.0) == FMonotonicTimePoint::FromSeconds(357.0));
+		STATIC_CHECK(FMonotonicTimePoint::Infinity() + FMonotonicTimeSpan::FromSeconds(1.0) == FMonotonicTimePoint::Infinity());
+		STATIC_CHECK(FMonotonicTimePoint::Infinity() + FMonotonicTimeSpan::Infinity() == FMonotonicTimePoint::Infinity());
+		CHECK((FMonotonicTimePoint::Infinity() + (-FMonotonicTimeSpan::Infinity())).IsNaN());
+	}
+
+	SECTION("Subtraction")
+	{
+		STATIC_CHECK(FMonotonicTimePoint::FromSeconds(234.0) - FMonotonicTimeSpan::FromSeconds(123.0) == FMonotonicTimePoint::FromSeconds(111.0));
+		STATIC_CHECK(FMonotonicTimePoint::Infinity() - FMonotonicTimeSpan::FromSeconds(1.0) == FMonotonicTimePoint::Infinity());
+		STATIC_CHECK(FMonotonicTimePoint::Infinity() - (-FMonotonicTimeSpan::Infinity()) == FMonotonicTimePoint::Infinity());
+		CHECK((FMonotonicTimePoint::Infinity() - FMonotonicTimeSpan::Infinity()).IsNaN());
+	}
+
+	SECTION("Span")
+	{
+		STATIC_CHECK(FMonotonicTimePoint::FromSeconds(357.0) - FMonotonicTimePoint::FromSeconds(234.0) == FMonotonicTimeSpan::FromSeconds(123.0));
+		STATIC_CHECK(FMonotonicTimePoint::FromSeconds(234.0) - FMonotonicTimePoint::FromSeconds(357.0) == FMonotonicTimeSpan::FromSeconds(-123.0));
+
+		STATIC_CHECK(FMonotonicTimePoint::Infinity() - FMonotonicTimePoint::FromSeconds(123.0) == FMonotonicTimeSpan::Infinity());
+		STATIC_CHECK(FMonotonicTimePoint::FromSeconds(123.0) - FMonotonicTimePoint::Infinity() == -FMonotonicTimeSpan::Infinity());
+	}
+}
+
+} // namespace UE
+
+#endif // WITH_LOW_LEVEL_TESTS
