Some macOS fixes for AutoRTFM.

diff --git a/Engine/Source/Runtime/Core/Private/AutoRTFM/Context.cpp b/Engine/Source/Runtime/Core/Private/AutoRTFM/Context.cpp
--- a/Engine/Source/Runtime/Core/Private/AutoRTFM/Context.cpp
+++ b/Engine/Source/Runtime/Core/Private/AutoRTFM/Context.cpp
@@ -211,6 +211,10 @@ FContext::FContext()
 {
 #if _WIN32
     GetCurrentThreadStackLimits(&StackBegin, &StackEnd);
+#elif defined(__APPLE__)         
+   StackEnd = pthread_get_stackaddr_np(pthread_self());   
+   size_t StackSize = pthread_get_stacksize_np(pthread_self());    
+   StackBegin = static_cast<char*>(StackEnd) - StackSize;
 #else
     pthread_attr_t Attr;
     pthread_getattr_np(pthread_self(), &Attr);
diff --git a/Engine/Source/Runtime/Core/Private/AutoRTFM/Utils.h b/Engine/Source/Runtime/Core/Private/AutoRTFM/Utils.h
--- a/Engine/Source/Runtime/Core/Private/AutoRTFM/Utils.h
+++ b/Engine/Source/Runtime/Core/Private/AutoRTFM/Utils.h
@@ -8,7 +8,12 @@
 #include <stdio.h>
 #include <stdlib.h>
 #include <string>
+
+#ifdef __APPLE__
+#include <malloc/malloc.h>
+#else
 #include <malloc.h>
+#endif
 
 #define ASSERT(exp) do { \
     if (!(exp)) { \
@@ -43,7 +48,9 @@ std::string GetFunctionDescription(TReturnType (*FunctionPtr)(TParameterTypes...
 
 inline size_t GetAllocationSize(void* const Ptr)
 {
-#if defined(_WIN32)
+#if defined(__APPLE__)
+    return malloc_size(Ptr);
+#elif defined(_WIN32)
     return _msize(Ptr);
 #else
     return malloc_usable_size(Ptr);
