ParkingLot: Added WaitUntil to wait until a specific point in time

#rnx

diff --git a/Engine/Source/Runtime/Core/Private/Async/ParkingLot.cpp b/Engine/Source/Runtime/Core/Private/Async/ParkingLot.cpp
--- a/Engine/Source/Runtime/Core/Private/Async/ParkingLot.cpp
+++ b/Engine/Source/Runtime/Core/Private/Async/ParkingLot.cpp
@@ -517,10 +517,12 @@ FThreadLocalData::~FThreadLocalData()
 
 ///////////////////////////////////////////////////////////////////////////////////////////////////////////////////////
 
-FWaitState WaitFor(const void* Address, TFunctionRef<bool()> CanWait, TFunctionRef<void()> BeforeWait, uint32 WaitMs)
+FWaitState WaitUntil(const void* Address, TFunctionRef<bool()> CanWait, TFunctionRef<void()> BeforeWait, FMonotonicTimePoint WaitTime)
 {
 	using namespace Private;
 
+	check(!WaitTime.IsNaN());
+
 	FThread& Self = FThreadLocalData::Get();
 
 	checkfSlow(!Self.WaitAddress, TEXT("WaitAddress must be null. This can happen if Wait is called by BeforeWait."));
@@ -547,13 +549,15 @@ FWaitState WaitFor(const void* Address, TFunctionRef<bool()> CanWait, TFunctionR
 #if UE_PARKINGLOT_USE_WAITONADDRESS
 	for (;;)
 	{
-		const DWORD SystemWaitMs = (WaitMs == MAX_uint32) ? INFINITE : WaitMs;
-		if (!WaitOnAddress(&Self.WaitAddress, &Address, sizeof(Address), SystemWaitMs))
+		const FMonotonicTimeSpan WaitSpan = WaitTime - FMonotonicTimePoint::Now();
+		if (WaitSpan <= FMonotonicTimeSpan::Zero())
 		{
-			if (GetLastError() == ERROR_TIMEOUT)
-			{
-				break;
-			}
+			break;
+		}
+		const DWORD WaitMs = WaitSpan.IsInfinity() ? INFINITE : DWORD(FMath::CeilToInt64(WaitSpan.ToMilliseconds()));
+		if (!WaitOnAddress(&Self.WaitAddress, &Address, sizeof(Address), WaitMs) && GetLastError() == ERROR_TIMEOUT)
+		{
+			break;
 		}
 		if (!Self.WaitAddress)
 		{
@@ -567,12 +571,14 @@ FWaitState WaitFor(const void* Address, TFunctionRef<bool()> CanWait, TFunctionR
 	{
 		// Wait until the timeout or until the thread has been dequeued.
 		std::unique_lock SelfLock(Self.Lock);
-		if (WaitMs == MAX_uint32)
+		const FMonotonicTimeSpan WaitSpan = WaitTime - FMonotonicTimePoint::Now();
+		if (WaitSpan.IsInfinity())
 		{
 			Self.Condition.wait(SelfLock, [&Self] { return !Self.WaitAddress; });
 		}
-		else
+		else if (WaitSpan > FMonotonicTimeSpan::Zero())
 		{
+			const int64 WaitMs = FMath::CeilToInt64(WaitSpan.ToMilliseconds());
 			Self.Condition.wait_for(SelfLock, std::chrono::milliseconds(WaitMs), [&Self] { return !Self.WaitAddress; });
 		}
 
diff --git a/Engine/Source/Runtime/Core/Public/Async/ParkingLot.h b/Engine/Source/Runtime/Core/Public/Async/ParkingLot.h
--- a/Engine/Source/Runtime/Core/Public/Async/ParkingLot.h
+++ b/Engine/Source/Runtime/Core/Public/Async/ParkingLot.h
@@ -3,7 +3,7 @@
 #pragma once
 
 #include "CoreTypes.h"
-#include "Math/NumericLimits.h"
+#include "Misc/MonotonicTime.h"
 #include "Templates/Function.h"
 
 #define UE_API CORE_API
@@ -34,7 +34,7 @@ struct FWakeState final
 
 namespace Private
 {
-UE_API FWaitState WaitFor(const void* Address, TFunctionRef<bool()> CanWait, TFunctionRef<void()> BeforeWait, uint32 WaitMs);
+UE_API FWaitState WaitUntil(const void* Address, TFunctionRef<bool()> CanWait, TFunctionRef<void()> BeforeWait, FMonotonicTimePoint WaitTime);
 UE_API void WakeOne(const void* Address, TFunctionRef<uint64(FWakeState)> OnWakeState);
 }
 
@@ -48,7 +48,7 @@ UE_API void WakeOne(const void* Address, TFunctionRef<uint64(FWakeState)> OnWake
 template <typename CanWaitType, typename BeforeWaitType>
 inline FWaitState Wait(const void* Address, CanWaitType&& CanWait, BeforeWaitType&& BeforeWait)
 {
-	return Private::WaitFor(Address, Forward<CanWaitType>(CanWait), Forward<BeforeWaitType>(BeforeWait), MAX_uint32);
+	return Private::WaitUntil(Address, Forward<CanWaitType>(CanWait), Forward<BeforeWaitType>(BeforeWait), FMonotonicTimePoint::Infinity());
 }
 
 /**
@@ -57,12 +57,26 @@ inline FWaitState Wait(const void* Address, CanWaitType&& CanWait, BeforeWaitTyp
  * @param Address      Address to use as the key for the queue. The same address is used to wake the thread.
  * @param CanWait      Function called while the queue is locked. A return of false cancels the wait.
  * @param BeforeWait   Function called after the queue is unlocked and before the thread waits.
- * @param WaitMs       Minimum wait time after which waiting is automatically canceled and the thread wakes.
+ * @param WaitTime     Relative time after which waiting is automatically canceled and the thread wakes.
  */
 template <typename CanWaitType, typename BeforeWaitType>
-inline FWaitState WaitFor(const void* Address, CanWaitType&& CanWait, BeforeWaitType&& BeforeWait, uint32 WaitMs)
+inline FWaitState WaitFor(const void* Address, CanWaitType&& CanWait, BeforeWaitType&& BeforeWait, FMonotonicTimeSpan WaitTime)
 {
-	return Private::WaitFor(Address, Forward<CanWaitType>(CanWait), Forward<BeforeWaitType>(BeforeWait), WaitMs);
+	return Private::WaitUntil(Address, Forward<CanWaitType>(CanWait), Forward<BeforeWaitType>(BeforeWait), FMonotonicTimePoint::Now() + WaitTime);
+}
+
+/**
+ * Queue the calling thread to wait if CanWait returns true. BeforeWait is only called if CanWait returns true.
+ *
+ * @param Address      Address to use as the key for the queue. The same address is used to wake the thread.
+ * @param CanWait      Function called while the queue is locked. A return of false cancels the wait.
+ * @param BeforeWait   Function called after the queue is unlocked and before the thread waits.
+ * @param WaitTime     Absolute time after which waiting is automatically canceled and the thread wakes.
+ */
+template <typename CanWaitType, typename BeforeWaitType>
+inline FWaitState WaitUntil(const void* Address, CanWaitType&& CanWait, BeforeWaitType&& BeforeWait, FMonotonicTimePoint WaitTime)
+{
+	return Private::WaitUntil(Address, Forward<CanWaitType>(CanWait), Forward<BeforeWaitType>(BeforeWait), WaitTime);
 }
 
 /**
diff --git a/Engine/Source/Runtime/Core/Tests/Async/ParkingLotTest.cpp b/Engine/Source/Runtime/Core/Tests/Async/ParkingLotTest.cpp
--- a/Engine/Source/Runtime/Core/Tests/Async/ParkingLotTest.cpp
+++ b/Engine/Source/Runtime/Core/Tests/Async/ParkingLotTest.cpp
@@ -102,7 +102,7 @@ TEST_CASE("Core::Async::ParkingLot", "[Core][Async]")
 		ParkingLot::FWaitState State = ParkingLot::WaitFor(&Value,
 			[&CanWaitCount] { ++CanWaitCount; return true; },
 			[&BeforeWaitCount] { ++BeforeWaitCount; },
-			/*WaitMs*/ 1);
+			FMonotonicTimeSpan::FromMilliseconds(1.0));
 		CHECK(CanWaitCount == 1);
 		CHECK(BeforeWaitCount == 1);
 		CHECK(State.bDidWait);
@@ -124,7 +124,7 @@ TEST_CASE("Core::Async::ParkingLot", "[Core][Async]")
 					[&Remaining]
 					{
 						Remaining.fetch_sub(1, std::memory_order_relaxed);
-					}, /*WaitMs*/ 1);
+					}, FMonotonicTimeSpan::FromMilliseconds(1.0));
 			});
 			CHECK(Remaining == 0);
 		});
