Fixed IsNaN for the new monotonic time types

The previous implementation was meant to avoid including extra headers, but doesn't work with fastmath.

#rnx

diff --git a/Engine/Source/Runtime/Core/Public/Misc/MonotonicTime.h b/Engine/Source/Runtime/Core/Public/Misc/MonotonicTime.h
--- a/Engine/Source/Runtime/Core/Public/Misc/MonotonicTime.h
+++ b/Engine/Source/Runtime/Core/Public/Misc/MonotonicTime.h
@@ -3,6 +3,7 @@
 #pragma once
 
 #include "CoreTypes.h"
+#include "HAL/PlatformMath.h"
 #include <limits>
 
 #define UE_API CORE_API
@@ -37,7 +38,7 @@ public:
 	constexpr double ToMilliseconds() const { return Time * 1000.0; }
 
 	constexpr bool IsInfinity() const { return *this == Infinity() || *this == -Infinity(); }
-	constexpr bool IsNaN() const { return Time != Time; }
+	bool IsNaN() const { return FPlatformMath::IsNaN(Time); }
 
 	constexpr bool operator==(const FMonotonicTimeSpan Other) const { return Time == Other.Time; }
 	constexpr bool operator!=(const FMonotonicTimeSpan Other) const { return Time != Other.Time; }
@@ -89,7 +90,7 @@ public:
 	constexpr double ToSeconds() const { return Time; }
 
 	constexpr bool IsInfinity() const { return *this == Infinity(); }
-	constexpr bool IsNaN() const { return Time != Time; }
+	bool IsNaN() const { return FPlatformMath::IsNaN(Time); }
 
 	constexpr bool operator==(const FMonotonicTimePoint Other) const { return Time == Other.Time; }
 	constexpr bool operator!=(const FMonotonicTimePoint Other) const { return Time != Other.Time; }
