CrashHandling instrumentation: Log a message before calling OnHandleSystemError.Broadcast to indicate whether we're crashing from inside the crash handler time during Broadcast.
#rnx

diff --git a/Engine/Source/Runtime/Core/Private/Windows/WindowsErrorOutputDevice.cpp b/Engine/Source/Runtime/Core/Private/Windows/WindowsErrorOutputDevice.cpp
--- a/Engine/Source/Runtime/Core/Private/Windows/WindowsErrorOutputDevice.cpp
+++ b/Engine/Source/Runtime/Core/Private/Windows/WindowsErrorOutputDevice.cpp
@@ -109,6 +109,9 @@ void FWindowsErrorOutputDevice::HandleError()
 
 	// Trigger the OnSystemFailure hook if it exists
 	// make sure it happens after GIsGuarded is set to 0 in case this hook crashes
+#if !NO_LOGGING
+	FMsg::Logf(__FILE__, __LINE__, LogWindows.GetCategoryName(), ELogVerbosity::Display, TEXT("OnHandleSystemError.Broadcast"));
+#endif
 	FCoreDelegates::OnHandleSystemError.Broadcast();
 
 	// Dump the error and flush the log. If you change behavior here, you should probably update RenderingThread's __except block in FRenderingThread::Run
