CrashHandling: more instrumentation for missing callstacks.
#rnx

diff --git a/Engine/Source/Runtime/Core/Private/Windows/WindowsErrorOutputDevice.cpp b/Engine/Source/Runtime/Core/Private/Windows/WindowsErrorOutputDevice.cpp
--- a/Engine/Source/Runtime/Core/Private/Windows/WindowsErrorOutputDevice.cpp
+++ b/Engine/Source/Runtime/Core/Private/Windows/WindowsErrorOutputDevice.cpp
@@ -113,6 +113,7 @@ void FWindowsErrorOutputDevice::HandleError()
 	FMsg::Logf(__FILE__, __LINE__, LogWindows.GetCategoryName(), ELogVerbosity::Display, TEXT("OnHandleSystemError.Broadcast"));
 #endif
 	FCoreDelegates::OnHandleSystemError.Broadcast();
+	GLog->Panic();
 
 	// Dump the error and flush the log. If you change behavior here, you should probably update RenderingThread's __except block in FRenderingThread::Run
 #if !NO_LOGGING
diff --git a/Engine/Source/Runtime/Launch/Private/Windows/LaunchWindows.cpp b/Engine/Source/Runtime/Launch/Private/Windows/LaunchWindows.cpp
--- a/Engine/Source/Runtime/Launch/Private/Windows/LaunchWindows.cpp
+++ b/Engine/Source/Runtime/Launch/Private/Windows/LaunchWindows.cpp
@@ -249,6 +249,7 @@ LAUNCH_API int32 LaunchWindowsStartup( HINSTANCE hInInstance, HINSTANCE hPrevIns
 				? ( GEnableInnerException ? EXCEPTION_EXECUTE_HANDLER : ReportCrash(GetExceptionInformation()) )
 				: EXCEPTION_CONTINUE_SEARCH )	
 		{
+			UE_LOG(LogLaunchWindows, Log, TEXT("LaunchWindowsStartup.ExceptionHandler"));
 			// Crashed.
 			ErrorLevel = 1;
 			if(GError)
