CrashHandling: instrumentation for missing callstacks attempt 3.
#rnx

diff --git a/Engine/Source/Runtime/Core/Private/Windows/WindowsErrorOutputDevice.cpp b/Engine/Source/Runtime/Core/Private/Windows/WindowsErrorOutputDevice.cpp
--- a/Engine/Source/Runtime/Core/Private/Windows/WindowsErrorOutputDevice.cpp
+++ b/Engine/Source/Runtime/Core/Private/Windows/WindowsErrorOutputDevice.cpp
@@ -112,7 +112,11 @@ void FWindowsErrorOutputDevice::HandleError()
 #if !NO_LOGGING
 	FMsg::Logf(__FILE__, __LINE__, LogWindows.GetCategoryName(), ELogVerbosity::Display, TEXT("OnHandleSystemError.Broadcast"));
 #endif
+	GLog->Panic();
 	FCoreDelegates::OnHandleSystemError.Broadcast();
+#if !NO_LOGGING
+	FMsg::Logf(__FILE__, __LINE__, LogWindows.GetCategoryName(), ELogVerbosity::Display, TEXT("OnHandleSystemError.Broadcast complete"));
+#endif
 	GLog->Panic();
 
 	// Dump the error and flush the log. If you change behavior here, you should probably update RenderingThread's __except block in FRenderingThread::Run
