Add TMultiUniqueLock, which lets you acquire multiple locks at once according to address order.

diff --git a/Engine/Source/Runtime/Core/Public/Experimental/Async/MultiUniqueLock.h b/Engine/Source/Runtime/Core/Public/Experimental/Async/MultiUniqueLock.h
new file mode 100644
--- /dev/null
+++ b/Engine/Source/Runtime/Core/Public/Experimental/Async/MultiUniqueLock.h
@@ -0,0 +1,54 @@
+// Copyright Epic Games, Inc. All Rights Reserved.
+
+#pragma once
+
+#include "Algo/Sort.h"
+#include "Containers/Array.h"
+
+namespace UE
+{
+
+// Like TUniqueLock, but acquires multiple locks and holds then until destruction. Acquires the locks in address
+// order to prevent deadlocks if multiple threads are using TMultiUniqueLock for similar sets of locks.
+template<typename TLockType>
+class TMultiUniqueLock final
+{
+public:
+    explicit TMultiUniqueLock(const TArray<TLockType*>& Locks)
+        : Locks(Locks)
+    {
+        SortAndLock();
+    }
+
+    explicit TMultiUniqueLock(TArray<TLockType*>&& Locks)
+        : Locks(MoveTemp(Locks))
+    {
+        SortAndLock();
+    }
+
+    TMultiUniqueLock& operator=(const TArray<TLockType*>&) = delete;
+    TMultiUniqueLock& operator=(TArray<TLockType*>&&) = delete;
+
+    ~TMultiUniqueLock()
+    {
+        for (TLockType* Lock : Locks)
+        {
+            Lock->Unlock();
+        }
+    }
+
+private:
+    void SortAndLock()
+    {
+        Algo::Sort(Locks);
+        for (TLockType* Lock : Locks)
+        {
+            Lock->Lock();
+        }
+    }
+
+    TArray<TLockType*> Locks;
+};
+
+}
+
