CrashHandling: instrumentation for missing callstacks attempt 4.
#rnx

diff --git a/Engine/Source/Runtime/Core/Private/Windows/WindowsErrorOutputDevice.cpp b/Engine/Source/Runtime/Core/Private/Windows/WindowsErrorOutputDevice.cpp
--- a/Engine/Source/Runtime/Core/Private/Windows/WindowsErrorOutputDevice.cpp
+++ b/Engine/Source/Runtime/Core/Private/Windows/WindowsErrorOutputDevice.cpp
@@ -121,6 +121,9 @@ void FWindowsErrorOutputDevice::HandleError()
 
 	// Dump the error and flush the log. If you change behavior here, you should probably update RenderingThread's __except block in FRenderingThread::Run
 #if !NO_LOGGING
+	FMsg::Logf(__FILE__, __LINE__, LogWindows.GetCategoryName(), ELogVerbosity::Error, TEXT("LogFormattedMessageWithCallstackFlushed"));
+	GLog->Panic();
+	FMsg::Logf(__FILE__, __LINE__, LogWindows.GetCategoryName(), ELogVerbosity::Error, TEXT("LogFormattedMessageWithCallstackNotFlushed"));
 	FDebug::LogFormattedMessageWithCallstack(LogWindows.GetCategoryName(), __FILE__, __LINE__, TEXT("=== Critical error: ==="), GErrorHist, ELogVerbosity::Error);
 #endif
 	GLog->Panic();
