Translated tests from Engine\Source\Runtime\Core\Private\Tests\HAL

diff --git a/Engine/Source/Developer/LowLevelTestsRunner/Public/TestHarness.h b/Engine/Source/Developer/LowLevelTestsRunner/Public/TestHarness.h
--- a/Engine/Source/Developer/LowLevelTestsRunner/Public/TestHarness.h
+++ b/Engine/Source/Developer/LowLevelTestsRunner/Public/TestHarness.h
@@ -125,6 +125,8 @@ struct Catch::StringMaker<TTuple<KeyType, ValueType>>
 #define ADD_ERROR(What)\
 	FAIL(What)
 
+#define FAIL_ON_MESSAGE(What)
+
 #define CHECK_EQUAL(Actual, Expected)\
 	CHECK(Actual == Expected)
 
diff --git a/Engine/Source/Runtime/Core/Private/Tests/HAL/FileManagerGenericTest.cpp b/Engine/Source/Runtime/Core/Private/Tests/HAL/FileManagerGenericTest.cpp
deleted file mode 100644
diff --git a/Engine/Source/Runtime/Core/Private/Tests/HAL/PlatformAtomicsTest.cpp b/Engine/Source/Runtime/Core/Private/Tests/HAL/PlatformAtomicsTest.cpp
deleted file mode 100644
diff --git a/Engine/Source/Runtime/Core/Private/Tests/HAL/PlatformStringTest.cpp b/Engine/Source/Runtime/Core/Private/Tests/HAL/PlatformStringTest.cpp
deleted file mode 100644
diff --git a/Engine/Source/Runtime/Core/Private/Tests/HAL/ThreadTest.cpp b/Engine/Source/Runtime/Core/Private/Tests/HAL/ThreadTest.cpp
deleted file mode 100644
diff --git a/Engine/Source/Runtime/Core/Public/Misc/LowLevelTestAdapter.h b/Engine/Source/Runtime/Core/Public/Misc/LowLevelTestAdapter.h
--- a/Engine/Source/Runtime/Core/Public/Misc/LowLevelTestAdapter.h
+++ b/Engine/Source/Runtime/Core/Public/Misc/LowLevelTestAdapter.h
@@ -130,6 +130,7 @@ FString CaptureExpressionsAndValues(const FString& InExpressions, ArgTypes&&...
 #define FAIL_CHECK(Message) FAutomationTestFramework::Get().GetCurrentTest()->AddError(Message);
 
 #define CAPTURE(...) FAutomationTestFramework::Get().GetCurrentTest()->AddInfo(CaptureExpressionsAndValues(#__VA_ARGS__, __VA_ARGS__));
-#define WARN(Message) { FAutomationTestFramework::Get().GetCurrentTest()->AddWarning(Message); }
+#define WARN(Message) FAutomationTestFramework::Get().GetCurrentTest()->AddWarning(Message); 
+#define FAIL_ON_MESSAGE(Message) FAutomationTestFramework::Get().GetCurrentTest()->AddExpectedError(Message); 
 
 #endif // !WITH_LOW_LEVEL_TESTS
\ No newline at end of file
diff --git a/Engine/Source/Runtime/Core/Tests/HAL/FileManagerGenericTest.cpp b/Engine/Source/Runtime/Core/Tests/HAL/FileManagerGenericTest.cpp
new file mode 100644
--- /dev/null
+++ b/Engine/Source/Runtime/Core/Tests/HAL/FileManagerGenericTest.cpp
@@ -0,0 +1,288 @@
+// Copyright Epic Games, Inc. All Rights Reserved.
+
+#include "HAL/FileManagerGeneric.h"
+#include "Containers/UnrealString.h"
+#include "HAL/FileManager.h"
+#include "HAL/PlatformFileManager.h"
+#include "Misc/Paths.h"
+#include "Misc/StringBuilder.h"
+#include "Tests/TestHarnessAdapter.h"
+
+#if WITH_TESTS
+
+class FArchiveFileReaderGenericTest
+{
+public:
+	~FArchiveFileReaderGenericTest();
+	void TestInternalPrecache();
+
+private:
+	void CreateTestFile();
+	void TestBytesValid(const TCHAR* What, uint8* Data, int64 NumBytes, int64 FileOffset);
+	int32 GetExpectedValue(int64 IntOffset);
+	void SetPosAndBuffer(TUniquePtr<FArchiveFileReaderGeneric>& Reader, int64 Pos, int64 BufferBase, int64 BufferSize);
+
+	FString TestFileName;
+	TUniquePtr<IFileHandle> ReadHandle;
+	int64 FileSize;
+};
+
+FArchiveFileReaderGenericTest::~FArchiveFileReaderGenericTest()
+{
+	ReadHandle.Reset();
+	if (!TestFileName.IsEmpty())
+	{
+		IFileManager& FileManager = IFileManager::Get();
+		FileManager.Delete(*TestFileName, false /* bRequireExists */, true /* bEvenReadOnly */, true /* Quiet */);
+	}
+	TestFileName.Empty();
+}
+
+void FArchiveFileReaderGenericTest::TestInternalPrecache()
+{
+	CreateTestFile();
+
+	uint32 BufferSize = 1024;
+	TArray<uint8> UnusedBytes;
+	UnusedBytes.SetNumUninitialized(BufferSize);
+	TUniquePtr<FArchiveFileReaderGeneric> Reader(new FArchiveFileReaderGeneric(ReadHandle.Release(), *TestFileName, FileSize, BufferSize));
+	CHECK_MESSAGE(TEXT("Initial Pos should be 0"), Reader->Pos == 0LL);
+	CHECK_MESSAGE(TEXT("Size should be what was passed in"), Reader->Size == FileSize);
+	CHECK_MESSAGE(TEXT("BufferSize should be what was passed in"), Reader->BufferSize == BufferSize);
+
+	// Vanilla InternalPrecache at start of file
+	bool Result = Reader->InternalPrecache(0, BufferSize);
+	CHECK_MESSAGE(TEXT("Vanilla0 - InternalPrecache should succeed"), Result);
+	CHECK_MESSAGE(TEXT("Vanilla0 - InternalPrecache should not alter Pos"), Reader->Pos == 0LL);
+	CHECK_MESSAGE(TEXT("Vanilla0 - BufferBase should have be set to PrecacheOffset aka Pos"), Reader->BufferBase == 0LL);
+	CHECK_MESSAGE(TEXT("Vanilla0 - Bytes precached should be set to BufferSize unless it runs out of room"), Reader->BufferArray.Num() == BufferSize);
+	TestBytesValid(TEXT("Vanilla0 - BufferArray should be expected bytes"), Reader->BufferArray.GetData(), Reader->BufferArray.Num(), Reader->BufferBase);
+
+	// InternalPrecache at PrecacheOffset != Pos is ignored, and returns false if the first byte at PrecacheOffset is not cached
+	Result = Reader->InternalPrecache(BufferSize  + BufferSize, BufferSize);
+	CHECK_FALSE_MESSAGE(TEXT("PrecacheNotAtPos - Should fail"), Result);
+	CHECK_MESSAGE(TEXT("PrecacheNotAtPos - BufferBase should not be altered"), Reader->BufferBase == 0LL);
+	CHECK_MESSAGE(TEXT("PrecacheNotAtPos - Buffe should not be altered"), Reader->BufferArray.Num() == BufferSize);
+
+	// InternalPrecache at PrecacheOffset != Pos is ignored, but returns true if the first byte at PrecacheOffset is cached
+	Result = Reader->InternalPrecache(BufferSize - 1, BufferSize);
+	CHECK_MESSAGE(TEXT("PrecacheNotAtPos - Should succeed since the first byte at PrecacheOffset is buffered"), Result);
+	CHECK_MESSAGE(TEXT("PrecacheNotAtPos - BufferBase should not be altered"), Reader->BufferBase == 0LL);
+	CHECK_MESSAGE(TEXT("PrecacheNotAtPos - Buffe should not be altered"), Reader->BufferArray.Num() == BufferSize);
+
+	// InternalPrecache partway through the buffer should allocate the full buffer/not allocate anything, depending on bPrecacheAsSoonAsPossible
+	int64 PosStart = BufferSize / 2;
+	SetPosAndBuffer(Reader, PosStart, 0, BufferSize);
+	Result = Reader->InternalPrecache(PosStart, BufferSize);
+	CHECK_MESSAGE(TEXT("PrecachePartwayPosMoreThanBuffer - Should succeed"), Result);
+	if (FArchiveFileReaderGeneric::bPrecacheAsSoonAsPossible)
+	{
+		CHECK_MESSAGE(TEXT("PrecachePartwayPosMoreThanBuffer - BufferBase should be updated"), Reader->BufferBase == PosStart);
+	}
+	else
+	{
+		CHECK_MESSAGE(TEXT("PrecachePartwayPosMoreThanBuffer - BufferBase should not be updated"), Reader->BufferBase == 0LL);
+	}
+	CHECK_MESSAGE(TEXT("PrecachePartwayPosMoreThanBuffer - BufferCount should be set to BufferSize"), Reader->BufferArray.Num() == BufferSize);
+	TestBytesValid(TEXT("PrecachePartwayPosMoreThanBuffer - BufferBytes should match BufferBase"), Reader->BufferArray.GetData(), Reader->BufferArray.Num(), Reader->BufferBase);
+
+	// InternalPrecache partway through the buffer should allocate the full buffer/not allocate anything, and should handle the case of BufferArray's allocation being smaller than BufferSize
+	SetPosAndBuffer(Reader, 3*BufferSize/4, BufferSize/2, BufferSize/2);
+	Reader->BufferArray.Shrink();
+	Result = Reader->InternalPrecache(3*BufferSize/4, BufferSize);
+	CHECK_MESSAGE(TEXT("PrecachePartwayPosMoreThanBufferReallocation - Should succeed"), Result);
+	if (FArchiveFileReaderGeneric::bPrecacheAsSoonAsPossible)
+	{
+		CHECK_MESSAGE(TEXT("PrecachePartwayPosMoreThanBufferReallocation - BufferBase should be updated"), Reader->BufferBase == 3*BufferSize/4);
+		CHECK_MESSAGE(TEXT("PrecachePartwayPosMoreThanBufferReallocation - BufferCount should be set to BufferSize"), Reader->BufferArray.Num() == BufferSize);
+	}
+	else
+	{
+		CHECK_MESSAGE(TEXT("PrecachePartwayPosMoreThanBufferReallocation - BufferBase should not be updated"), Reader->BufferBase == BufferSize/2);
+		CHECK_MESSAGE(TEXT("PrecachePartwayPosMoreThanBufferReallocation - BufferCount should not be updated"), Reader->BufferArray.Num() == BufferSize/2);
+	}
+	TestBytesValid(TEXT("PrecachePartwayPosMoreThanBufferReallocation - BufferBytes should match BufferBase"), Reader->BufferArray.GetData(), Reader->BufferArray.Num(), Reader->BufferBase);
+
+	// InternalPrecache partway through the buffer, when the requested size is within the remaining buffer, should not change the buffer
+	PosStart = BufferSize / 2;
+	SetPosAndBuffer(Reader, PosStart, 0, BufferSize);
+	Result = Reader->InternalPrecache(PosStart, BufferSize / 4);
+	CHECK_MESSAGE(TEXT("PrecachePartwayPosLessThanBuffer - Should succeed"), Result);
+	CHECK_MESSAGE(TEXT("PrecachePartwayPosLessThanBuffer - BufferBase should not be updated"), Reader->BufferBase == 0LL);
+	CHECK_MESSAGE(TEXT("PrecachePartwayPosLessThanBuffer - BufferCount should be set to BufferSize"), Reader->BufferArray.Num() == BufferSize);
+	TestBytesValid(TEXT("PrecachePartwayPosLessThanBuffer - BufferBytes should match BufferBase"), Reader->BufferArray.GetData(), Reader->BufferArray.Num(), Reader->BufferBase);
+
+	// InternalPrecache at the end of the buffer should allocate the full buffer
+	PosStart = BufferSize;
+	SetPosAndBuffer(Reader, PosStart, 0, BufferSize);
+	Result = Reader->InternalPrecache(PosStart, BufferSize);
+	CHECK_MESSAGE(TEXT("PrecachePosEndOfBuffer - Should succeed"), Result);
+	CHECK_MESSAGE(TEXT("PrecachePosEndOfBuffer - BufferBase should be updated"), Reader->BufferBase == PosStart);
+	CHECK_MESSAGE(TEXT("PrecachePosEndOfBuffer - BufferCount should be set to BufferSize"), Reader->BufferArray.Num() == BufferSize);
+	TestBytesValid(TEXT("PrecachePosEndOfBuffer - BufferBytes should match BufferBase"), Reader->BufferArray.GetData(), Reader->BufferArray.Num(), Reader->BufferBase);
+
+	// InternalPrecache near the end of the file should clamp BufferSize to Size - Pos
+	PosStart = FileSize - 1;
+	SetPosAndBuffer(Reader, PosStart, 0, 0);
+	Result = Reader->InternalPrecache(PosStart, BufferSize);
+	CHECK_MESSAGE(TEXT("PrecacheNearEndOfFile - Should succeed"), Result);
+	CHECK_MESSAGE(TEXT("PrecacheNearEndOfFile - BufferBase should be updated"), Reader->BufferBase == PosStart);
+	CHECK_MESSAGE(TEXT("PrecacheNearEndOfFile - BufferCount should be set Min(BufferSize, Size - Pos)"), Reader->BufferArray.Num() == 1LL);
+	TestBytesValid(TEXT("PrecacheNearEndOfFile - BufferBytes should match BufferBase"), Reader->BufferArray.GetData(), Reader->BufferArray.Num(), Reader->BufferBase);
+
+	FAIL_ON_MESSAGE(TEXT("ReadFile failed"));
+	// InternalPrecache at the end of the file should fail
+	SetPosAndBuffer(Reader, FileSize, 0, 16);
+	Result = Reader->InternalPrecache(FileSize, BufferSize);
+	CHECK_FALSE_MESSAGE(TEXT("PrecacheEndOfFile - Should fail"), Result);
+	CHECK_MESSAGE(TEXT("PrecacheEndOfFile - BufferBase should not be updated"), Reader->BufferBase == 0LL);
+	CHECK_MESSAGE(TEXT("PrecacheEndOfFile - BufferCount should not be updated"), Reader->BufferArray.Num() == 16LL);
+
+	// If ReadLowLevel fails, InternalPrecache should return true if the first byte of PrecacheOffset is in the buffer
+	SetPosAndBuffer(Reader, 8, 0, 16);
+	Reader->SeekLowLevel(FileSize);
+	Result = Reader->InternalPrecache(8, 16);
+	CHECK_MESSAGE(TEXT("PrecacheReadFailsBytesRemain - Should succeed if bytes are left in the buffer"), Result);
+	if (FArchiveFileReaderGeneric::bPrecacheAsSoonAsPossible)
+	{
+		CHECK_MESSAGE(TEXT("PrecacheReadFailsBytesRemain - BufferBase should be updated"), Reader->BufferBase == 8LL);
+		CHECK_MESSAGE(TEXT("PrecacheReadFailsBytesRemain - BufferCount should be updated"), Reader->BufferArray.Num() == 8LL);
+	}
+	else
+	{
+		CHECK_MESSAGE(TEXT("PrecacheReadFailsBytesRemain - BufferBase should not be updated"), Reader->BufferBase == 0LL);
+		CHECK_MESSAGE(TEXT("PrecacheReadFailsBytesRemain - BufferCount should not be updated"), Reader->BufferArray.Num() == 16LL);
+	}
+
+	// If ReadLowLevel fails, InternalPrecache should return false if the first byte of PrecacheOffset is not in the buffer
+	SetPosAndBuffer(Reader, 16, 0, 16);
+	Reader->SeekLowLevel(FileSize);
+	Result = Reader->InternalPrecache(16, 16);
+	CHECK_FALSE_MESSAGE(TEXT("PrecacheReadFails - Should fail if bytes are not left in the buffer"), Result);
+	CHECK_MESSAGE(TEXT("PrecacheReadFails - BufferCount should be updated to empty"), Reader->BufferArray.Num() == 0LL);
+
+	// TODO: Call ClearExpectedErrors to clear the AddExpectedError called above when it is implemented. Until then, make sure these failure-testing cases are the last
+	// cases tested in the test
+	//ClearExpectedErrors();
+}
+
+void FArchiveFileReaderGenericTest::CreateTestFile()
+{
+	IFileManager& FileManager = IFileManager::Get();
+	IPlatformFile& PlatformFile = FPlatformFileManager::Get().GetPlatformFile();;
+	TestFileName = FPaths::CreateTempFilename(FPlatformProcess::UserTempDir(), TEXT("ArchiveFileReaderGenericTest"));
+	TUniquePtr<IFileHandle> WriteHandle(PlatformFile.OpenWrite(*TestFileName, false /* bAppend */, false /* bAllowRead */));
+	REQUIRE_MESSAGE(FString::Printf(TEXT("Could not create test file %s."), *TestFileName), WriteHandle.IsValid());
+	TArray<int32> Builder;
+	int32 NumInts = 1024 * 128 + 17; // Make sure the file size is not a multiple of buffersize; we need to test that behavior
+	Builder.SetNumUninitialized(NumInts);
+	int32* Data = Builder.GetData();
+	for (int n = 0; n < Builder.Num(); ++n)
+	{
+		Data[n] = GetExpectedValue(n);
+	}
+	FileSize = NumInts * sizeof(int32);
+	WriteHandle->Write(reinterpret_cast<uint8*>(Data), FileSize);
+	WriteHandle.Reset();
+
+	ReadHandle.Reset(PlatformFile.OpenRead(*TestFileName, false /* bAllowWrite */));
+
+	REQUIRE_MESSAGE(FString::Printf(TEXT("Could not open test file %s."), *TestFileName), ReadHandle.IsValid());
+	REQUIRE_MESSAGE(FString::Printf(TEXT("Received incorrect file size from test file %s. Expected = %lld. Actual = %lld."), *TestFileName, FileSize, ReadHandle->Size()), ReadHandle->Size() == FileSize);
+}
+
+void FArchiveFileReaderGenericTest::TestBytesValid(const TCHAR* What, uint8* InData, int64 InNumBytes, int64 InFileOffset)
+{
+	uint8* Data = InData;
+	int64 NumBytes = InNumBytes;
+	int64 FileOffset = InFileOffset;
+
+	int64 StartByte = FileOffset % sizeof(int32);
+	if (StartByte != 0)
+	{
+		int64 FloorIndex = (FileOffset - StartByte)/sizeof(int32);
+		int64 ExpectedValue = GetExpectedValue(FloorIndex);
+		for (int n = 0; n < NumBytes && n < sizeof(int32) - StartByte; ++n)
+		{
+
+			if (Data[n] != reinterpret_cast<uint8*>(&ExpectedValue)[StartByte + n])
+			{
+				CHECK_MESSAGE(What, Data[n] == reinterpret_cast<uint8*>(&ExpectedValue)[StartByte + n]);
+				return;
+			}
+
+		}
+		if (NumBytes < static_cast<int64>(sizeof(int32) - StartByte))
+		{
+			return;
+		}
+		int32 BytesRead = sizeof(int32) - StartByte;
+		Data += BytesRead;
+		FileOffset += BytesRead;
+		NumBytes -= BytesRead;
+	}
+	int32* IntData = reinterpret_cast<int32*>(Data);
+	int64 IntOffset = FileOffset / sizeof(int32);
+	int64 NumInts = NumBytes / sizeof(int32);
+	for (int n = 0; n < NumInts; ++n)
+	{
+		if (IntData[n] != GetExpectedValue(IntOffset + n))
+		{
+			CHECK_MESSAGE(What, IntData[n] == GetExpectedValue(IntOffset + n));
+			return;
+		}
+	}
+
+	int64 NumEndBytes = NumBytes - NumInts * sizeof(int32);
+	if (NumEndBytes != 0)
+	{
+		uint8* DataStart = Data + NumInts * sizeof(int32);
+		int32 ExpectedValue = GetExpectedValue(IntOffset + NumInts);
+		for (int n = 0; n < NumEndBytes; ++n)
+		{
+
+			if (DataStart[n] != reinterpret_cast<uint8*>(&ExpectedValue)[n])
+			{
+				CHECK_MESSAGE(What, DataStart[n] == reinterpret_cast<uint8*>(&ExpectedValue)[n]);
+				return;
+			}
+		}
+	}
+}
+
+void FArchiveFileReaderGenericTest::SetPosAndBuffer(TUniquePtr<FArchiveFileReaderGeneric>& Reader, int64 Pos, int64 BufferBase, int64 BufferSize)
+{
+	Reader->Pos = Pos;
+	Reader->BufferBase = BufferBase;
+	Reader->BufferArray.SetNumUninitialized(BufferSize, false /* bAllowShrink */);
+	if (BufferSize)
+	{
+		Reader->SeekLowLevel(BufferBase);
+		int64 BytesRead;
+		Reader->ReadLowLevel(Reader->BufferArray.GetData(), BufferSize, BytesRead);
+		CHECK_MESSAGE(TEXT("SetPosAndBuffer - ReadLowLevel read the requested bytes"), BytesRead == BufferSize);
+	}
+	// See the contract for Buffer window and LowLevel Pos in the variable comment on BufferArray
+	bool bPosWithinBuffer = BufferBase <= Pos && Pos < BufferBase + BufferSize;
+	if (!bPosWithinBuffer)
+	{
+		Reader->SeekLowLevel(Pos);
+	}
+	CHECK_MESSAGE(TEXT("SetPosAndBuffer - Pos set correctly"), Reader->Pos == Pos);
+	CHECK_MESSAGE(TEXT("SetPosAndBuffer - BufferBase set correctly"), Reader->BufferBase == BufferBase);
+	CHECK_MESSAGE(TEXT("SetPosAndBuffer - BufferSize set correctly"), Reader->BufferArray.Num() == BufferSize);
+	TestBytesValid(TEXT("SetPosAndBuffer - BufferBytes should match BufferBase"), Reader->BufferArray.GetData(), Reader->BufferArray.Num(), Reader->BufferBase);
+}
+
+int32 FArchiveFileReaderGenericTest::GetExpectedValue(int64 IntOffset)
+{
+	return 0xbe000000 + static_cast<int32>(IntOffset);
+}
+
+TEST_CASE_NAMED(FArchiveFileReaderGenericTestRunner, "System::Core::HAL::FArchiveFileReaderGeneric", "[.][ApplicationContextMask][EngineFilter]")
+{
+	FArchiveFileReaderGenericTest Instance = FArchiveFileReaderGenericTest();
+	Instance.TestInternalPrecache();
+}
+
+#endif //WITH_TESTS
diff --git a/Engine/Source/Runtime/Core/Tests/HAL/PlatformAtomicsTest.cpp b/Engine/Source/Runtime/Core/Tests/HAL/PlatformAtomicsTest.cpp
new file mode 100644
--- /dev/null
+++ b/Engine/Source/Runtime/Core/Tests/HAL/PlatformAtomicsTest.cpp
@@ -0,0 +1,503 @@
+// Copyright Epic Games, Inc. All Rights Reserved.
+
+#include "CoreTypes.h"
+#include "Misc/AssertionMacros.h"
+#include "Tests/TestHarnessAdapter.h"
+
+#if WITH_TESTS
+
+#if PLATFORM_WINDOWS
+#include "Windows/AllowWindowsPlatformAtomics.h"
+#include "Windows/HideWindowsPlatformAtomics.h"
+#endif
+#include "HAL/PlatformAtomics.h"
+
+template<typename T>
+static bool TestInterlocked(volatile T* Dest, T ExpectedReturnValue, T ExpectedFinalValue, TFunctionRef<T(volatile T*)> InterlockedFunc, const TCHAR* FunctionName, const TCHAR* TypeName)
+{
+	const T ReturnValue = InterlockedFunc(Dest);
+	if (ReturnValue != ExpectedReturnValue)
+	{
+		check(false);
+		FAIL_CHECK(FString::Printf(TEXT("FPlatformAtomics::Interlocked%s on %s failed"), FunctionName, TypeName));
+		return false;
+	}
+	if (FPlatformAtomics::AtomicRead(Dest) != ExpectedFinalValue)
+	{
+		check(false);
+		FAIL_CHECK(FString::Printf(TEXT("FPlatformAtomics::Interlocked%s on %s failed"), FunctionName, TypeName));
+		return false;
+	}
+	return true;
+}
+
+static bool TestInterlockedAnd()
+{
+	bool bSuccess = true;
+
+	const TCHAR* FunctionName = TEXT("And");
+	{
+		// Test and with value where some bits are set and some aren't set in the current value.
+		volatile int8 Value0 = 0x30;
+		bSuccess = TestInterlocked<int8>(&Value0, 0x30, 0x20, [](volatile int8* Dest) { return FPlatformAtomics::InterlockedAnd(Dest, 0x66); }, FunctionName, TEXT("int8")) && bSuccess;
+
+		// Test and with all bits set.
+		volatile int8 Value1 = 0x30;
+		bSuccess = TestInterlocked<int8>(&Value1, 0x30, 0x30, [](volatile int8* Dest) { return FPlatformAtomics::InterlockedAnd(Dest, -1); }, FunctionName, TEXT("int8")) && bSuccess;
+
+		// Test and with zero.
+		volatile int8 Value2 = 0x30;
+		bSuccess = TestInterlocked<int8>(&Value2, 0x30, 0, [](volatile int8* Dest) { return FPlatformAtomics::InterlockedAnd(Dest, 0); }, FunctionName, TEXT("int8")) && bSuccess;
+	}
+
+	{
+		// Test and with value where some bits are set and some aren't set in the current value.
+		volatile int16 Value0 = 0x3030;
+		bSuccess = TestInterlocked<int16>(&Value0, 0x3030, 0x2020, [](volatile int16* Dest) { return FPlatformAtomics::InterlockedAnd(Dest, 0x6666); }, FunctionName, TEXT("int16")) && bSuccess;
+
+		// Test and with all bits set.
+		volatile int16 Value1 = 0x3030;
+		bSuccess = TestInterlocked<int16>(&Value1, 0x3030, 0x3030, [](volatile int16* Dest) { return FPlatformAtomics::InterlockedAnd(Dest, -1); }, FunctionName, TEXT("int16")) && bSuccess;
+
+		// Test and with zero.
+		volatile int16 Value2 = 0x3030;
+		bSuccess = TestInterlocked<int16>(&Value2, 0x3030, 0, [](volatile int16* Dest) { return FPlatformAtomics::InterlockedAnd(Dest, 0); }, FunctionName, TEXT("int16")) && bSuccess;
+	}
+
+	{
+		// Test and with value where some bits are set and some aren't set in the current value.
+		volatile int32 Value0 = 0x30303030;
+		bSuccess = TestInterlocked<int32>(&Value0, 0x30303030, 0x20202020, [](volatile int32* Dest) { return FPlatformAtomics::InterlockedAnd(Dest, 0x66666666); }, FunctionName, TEXT("int32")) && bSuccess;
+
+		// Test and with all bits set.
+		volatile int32 Value1 = 0x30303030;
+		bSuccess = TestInterlocked<int32>(&Value1, 0x30303030, 0x30303030, [](volatile int32* Dest) { return FPlatformAtomics::InterlockedAnd(Dest, -1); }, FunctionName, TEXT("int32")) && bSuccess;
+
+		// Test and with zero.
+		volatile int32 Value2 = 0x30303030;
+		bSuccess = TestInterlocked<int32>(&Value2, 0x30303030, 0, [](volatile int32* Dest) { return FPlatformAtomics::InterlockedAnd(Dest, 0); }, FunctionName, TEXT("int32")) && bSuccess;
+	}
+
+	{
+		// Test and with value where some bits are set and some aren't set in the current value.
+		alignas(8) volatile int64 Value0 = 0x3030303030303030LL;
+		bSuccess = TestInterlocked<int64>(&Value0, 0x3030303030303030LL, 0x2020202020202020LL, [](volatile int64* Dest) { return FPlatformAtomics::InterlockedAnd(Dest, 0x6666666666666666); }, FunctionName, TEXT("int64")) && bSuccess;
+
+		// Test and with all bits set.
+		alignas(8) volatile int64 Value1 = 0x3030303030303030LL;
+		bSuccess = TestInterlocked<int64>(&Value1, 0x3030303030303030LL, 0x3030303030303030LL, [](volatile int64* Dest) { return FPlatformAtomics::InterlockedAnd(Dest, -1); }, FunctionName, TEXT("int64")) && bSuccess;
+
+		// Test and with zero.
+		alignas(8) volatile int64 Value2 = 0x3030303030303030LL;
+		bSuccess = TestInterlocked<int64>(&Value2, 0x3030303030303030LL, 0, [](volatile int64* Dest) { return FPlatformAtomics::InterlockedAnd(Dest, 0); }, FunctionName, TEXT("int64")) && bSuccess;
+	}
+
+	return bSuccess;
+}
+
+static bool TestInterlockedOr()
+{
+	bool bSuccess = true;
+
+	const TCHAR* FunctionName = TEXT("Or");
+	{
+		// Test or with value where some bits are set and some aren't set in the current value.
+		volatile int8 Value0 = 0x30;
+		bSuccess = TestInterlocked<int8>(&Value0, 0x30, 0x76, [](volatile int8* Dest) { return FPlatformAtomics::InterlockedOr(Dest, 0x66); }, FunctionName, TEXT("int8")) && bSuccess;
+
+		// Test or with all bits set.
+		volatile int8 Value1 = 0x30;
+		bSuccess = TestInterlocked<int8>(&Value1, 0x30, -1, [](volatile int8* Dest) { return FPlatformAtomics::InterlockedOr(Dest, -1); }, FunctionName, TEXT("int8")) && bSuccess;
+
+		// Test or with zero.
+		volatile int8 Value2 = 0x30;
+		bSuccess = TestInterlocked<int8>(&Value2, 0x30, 0x30, [](volatile int8* Dest) { return FPlatformAtomics::InterlockedOr(Dest, 0); }, FunctionName, TEXT("int8")) && bSuccess;
+	}
+
+	{
+		// Test or with value where some bits are set and some aren't set in the current value.
+		volatile int16 Value0 = 0x3030;
+		bSuccess = TestInterlocked<int16>(&Value0, 0x3030, 0x7676, [](volatile int16* Dest) { return FPlatformAtomics::InterlockedOr(Dest, 0x6666); }, FunctionName, TEXT("int16")) && bSuccess;
+
+		// Test or with all bits set.
+		volatile int16 Value1 = 0x3030;
+		bSuccess = TestInterlocked<int16>(&Value1, 0x3030, -1, [](volatile int16* Dest) { return FPlatformAtomics::InterlockedOr(Dest, -1); }, FunctionName, TEXT("int16")) && bSuccess;
+
+		// Test or with zero.
+		volatile int16 Value2 = 0x3030;
+		bSuccess = TestInterlocked<int16>(&Value2, 0x3030, 0x3030, [](volatile int16* Dest) { return FPlatformAtomics::InterlockedOr(Dest, 0); }, FunctionName, TEXT("int16")) && bSuccess;
+	}
+
+	{
+		// Test or with value where some bits are set and some aren't set in the current value.
+		volatile int32 Value0 = 0x30303030;
+		bSuccess = TestInterlocked<int32>(&Value0, 0x30303030, 0x76767676, [](volatile int32* Dest) { return FPlatformAtomics::InterlockedOr(Dest, 0x66666666); }, FunctionName, TEXT("int32")) && bSuccess;
+
+		// Test or with all bits set.
+		volatile int32 Value1 = 0x30303030;
+		bSuccess = TestInterlocked<int32>(&Value1, 0x30303030, -1, [](volatile int32* Dest) { return FPlatformAtomics::InterlockedOr(Dest, -1); }, FunctionName, TEXT("int32")) && bSuccess;
+
+		// Test or with zero.
+		volatile int32 Value2 = 0x30303030;
+		bSuccess = TestInterlocked<int32>(&Value2, 0x30303030, 0x30303030, [](volatile int32* Dest) { return FPlatformAtomics::InterlockedOr(Dest, 0); }, FunctionName, TEXT("int32")) && bSuccess;
+	}
+
+	{
+		// Test or with value where some bits are set and some aren't set in the current value.
+		alignas(8) volatile int64 Value0 = 0x3030303030303030LL;
+		bSuccess = TestInterlocked<int64>(&Value0, 0x3030303030303030LL, 0x7676767676767676LL, [](volatile int64* Dest) { return FPlatformAtomics::InterlockedOr(Dest, 0x6666666666666666LL); }, FunctionName, TEXT("int64")) && bSuccess;
+
+		// Test or with all bits set.
+		alignas(8) volatile int64 Value1 = 0x3030303030303030LL;
+		bSuccess = TestInterlocked<int64>(&Value1, 0x3030303030303030LL, -1LL, [](volatile int64* Dest) { return FPlatformAtomics::InterlockedOr(Dest, -1); }, FunctionName, TEXT("int64")) && bSuccess;
+
+		// Test or with zero.
+		alignas(8) volatile int64 Value2 = 0x3030303030303030LL;
+		bSuccess = TestInterlocked<int64>(&Value2, 0x3030303030303030LL, 0x3030303030303030LL, [](volatile int64* Dest) { return FPlatformAtomics::InterlockedOr(Dest, 0); }, FunctionName, TEXT("int64")) && bSuccess;
+	}
+
+	return bSuccess;
+}
+
+static bool TestInterlockedXor()
+{
+	bool bSuccess = true;
+
+	const TCHAR* FunctionName = TEXT("Xor");
+	{
+		// Test xor with value where some bits are set and some aren't set in the current value.
+		volatile int8 Value0 = 0x30;
+		bSuccess = TestInterlocked<int8>(&Value0, 0x30, 0x56, [](volatile int8* Dest) { return FPlatformAtomics::InterlockedXor(Dest, 0x66); }, FunctionName, TEXT("int8")) && bSuccess;
+
+		// Test xor with all bits set.
+		volatile int8 Value1 = 0x30;
+		bSuccess = TestInterlocked<int8>(&Value1, 0x30, int8(~0x30), [](volatile int8* Dest) { return FPlatformAtomics::InterlockedXor(Dest, -1); }, FunctionName, TEXT("int8")) && bSuccess;
+
+		// Test xor with zero.
+		volatile int8 Value2 = 0x30;
+		bSuccess = TestInterlocked<int8>(&Value2, 0x30, 0x30, [](volatile int8* Dest) { return FPlatformAtomics::InterlockedXor(Dest, 0); }, FunctionName, TEXT("int8")) && bSuccess;
+	}
+
+	{
+		// Test xor with value where some bits are set and some aren't set in the current value.
+		volatile int16 Value0 = 0x3030;
+		bSuccess = TestInterlocked<int16>(&Value0, 0x3030, 0x5656, [](volatile int16* Dest) { return FPlatformAtomics::InterlockedXor(Dest, 0x6666); }, FunctionName, TEXT("int16")) && bSuccess;
+
+		// Test xor with all bits set.
+		volatile int16 Value1 = 0x3030;
+		bSuccess = TestInterlocked<int16>(&Value1, 0x3030, int16(~0x3030), [](volatile int16* Dest) { return FPlatformAtomics::InterlockedXor(Dest, -1); }, FunctionName, TEXT("int16")) && bSuccess;
+
+		// Test xor with zero.
+		volatile int16 Value2 = 0x3030;
+		bSuccess = TestInterlocked<int16>(&Value2, 0x3030, 0x3030, [](volatile int16* Dest) { return FPlatformAtomics::InterlockedXor(Dest, 0); }, FunctionName, TEXT("int16")) && bSuccess;
+	}
+
+	{
+		// Test xor with value where some bits are set and some aren't set in the current value.
+		volatile int32 Value0 = 0x30303030;
+		bSuccess = TestInterlocked<int32>(&Value0, 0x30303030, 0x56565656, [](volatile int32* Dest) { return FPlatformAtomics::InterlockedXor(Dest, 0x66666666); }, FunctionName, TEXT("int32")) && bSuccess;
+
+		// Test xor with all bits set.
+		volatile int32 Value1 = 0x30303030;
+		bSuccess = TestInterlocked<int32>(&Value1, 0x30303030, ~0x30303030, [](volatile int32* Dest) { return FPlatformAtomics::InterlockedXor(Dest, -1); }, FunctionName, TEXT("int32")) && bSuccess;
+
+		// Test xor with zero.
+		volatile int32 Value2 = 0x30303030;
+		bSuccess = TestInterlocked<int32>(&Value2, 0x30303030, 0x30303030, [](volatile int32* Dest) { return FPlatformAtomics::InterlockedXor(Dest, 0); }, FunctionName, TEXT("int32")) && bSuccess;
+	}
+
+	{
+		// Test xor with value where some bits are set and some aren't set in the current value.
+		alignas(8) volatile int64 Value0 = 0x3030303030303030LL;
+		bSuccess = TestInterlocked<int64>(&Value0, 0x3030303030303030LL, 0x5656565656565656LL, [](volatile int64* Dest) { return FPlatformAtomics::InterlockedXor(Dest, 0x6666666666666666LL); }, FunctionName, TEXT("int64")) && bSuccess;
+
+		// Test xor with all bits set.
+		alignas(8) volatile int64 Value1 = 0x3030303030303030LL;
+		bSuccess = TestInterlocked<int64>(&Value1, 0x3030303030303030LL, ~0x3030303030303030LL, [](volatile int64* Dest) { return FPlatformAtomics::InterlockedXor(Dest, -1); }, FunctionName, TEXT("int64")) && bSuccess;
+
+		// Test xor with zero.
+		alignas(8) volatile int64 Value2 = 0x3030303030303030LL;
+		bSuccess = TestInterlocked<int64>(&Value2, 0x3030303030303030LL, 0x3030303030303030LL, [](volatile int64* Dest) { return FPlatformAtomics::InterlockedXor(Dest, 0); }, FunctionName, TEXT("int64")) && bSuccess;
+	}
+
+	return bSuccess;
+}
+
+static bool TestInterlockedAdd()
+{
+	bool bSuccess = true;
+
+	const TCHAR* FunctionName = TEXT("Add");
+	{
+		// Test adding positive value
+		volatile int8 Value0 = 0x0F;
+		bSuccess = TestInterlocked<int8>(&Value0, 0x0F, 0x11, [](volatile int8* Dest) { return FPlatformAtomics::InterlockedAdd(Dest, +0x02); }, FunctionName, TEXT("int8")) && bSuccess;
+
+		// Test adding negative value
+		volatile int8 Value1 = 0x11;
+		bSuccess = TestInterlocked<int8>(&Value1, 0x11, 0x0F, [](volatile int8* Dest) { return FPlatformAtomics::InterlockedAdd(Dest, -0x02); }, FunctionName, TEXT("int8")) && bSuccess;
+
+		// Test value overflow
+		volatile int8 Value2 = MAX_int8 - 1;
+		bSuccess = TestInterlocked<int8>(&Value2, MAX_int8 - 1, MIN_int8 + 2, [](volatile int8* Dest) { return FPlatformAtomics::InterlockedAdd(Dest, +4); }, FunctionName, TEXT("int8")) && bSuccess;
+
+		// Test value underflow
+		volatile int8 Value3 = MIN_int8 + 2;
+		bSuccess = TestInterlocked<int8>(&Value3, MIN_int8 + 2, MAX_int8 - 1, [](volatile int8* Dest) { return FPlatformAtomics::InterlockedAdd(Dest, -4); }, FunctionName, TEXT("int8")) && bSuccess;
+	}
+
+	{
+		// Test adding positive value
+		volatile int16 Value0 = 0x0F00;
+		bSuccess = TestInterlocked<int16>(&Value0, 0x0F00, 0x1001, [](volatile int16* Dest) { return FPlatformAtomics::InterlockedAdd(Dest, 0x0101); }, FunctionName, TEXT("int16")) && bSuccess;
+
+		// Test adding negative value
+		volatile int16 Value1 = 0x1001;
+		bSuccess = TestInterlocked<int16>(&Value1, 0x1001, 0x0F00, [](volatile int16* Dest) { return FPlatformAtomics::InterlockedAdd(Dest, -0x0101); }, FunctionName, TEXT("int16")) && bSuccess;
+
+		// Test value overflow
+		volatile int16 Value2 = MAX_int16 - 1;
+		bSuccess = TestInterlocked<int16>(&Value2, MAX_int16 - 1, MIN_int16 + 2, [](volatile int16* Dest) { return FPlatformAtomics::InterlockedAdd(Dest, +4); }, FunctionName, TEXT("int16")) && bSuccess;
+
+		// Test value underflow
+		volatile int16 Value3 = MIN_int16 + 2;
+		bSuccess = TestInterlocked<int16>(&Value3, MIN_int16 + 2, MAX_int16 - 1, [](volatile int16* Dest) { return FPlatformAtomics::InterlockedAdd(Dest, -4); }, FunctionName, TEXT("int16")) && bSuccess;
+	}
+
+	{
+		// Test adding positive value
+		volatile int32 Value0 = 0x0F000000;
+		bSuccess = TestInterlocked<int32>(&Value0, 0x0F000000, 0x10010101, [](volatile int32* Dest) { return FPlatformAtomics::InterlockedAdd(Dest, 0x01010101); }, FunctionName, TEXT("int32")) && bSuccess;
+
+		// Test adding negative value
+		volatile int32 Value1 = 0x10010101;
+		bSuccess = TestInterlocked<int32>(&Value1, 0x10010101, 0x0F000000, [](volatile int32* Dest) { return FPlatformAtomics::InterlockedAdd(Dest, -0x01010101); }, FunctionName, TEXT("int32")) && bSuccess;
+
+		// Test value overflow
+		volatile int32 Value2 = MAX_int32 - 1;
+		bSuccess = TestInterlocked<int32>(&Value2, MAX_int32 - 1, MIN_int32 + 2, [](volatile int32* Dest) { return FPlatformAtomics::InterlockedAdd(Dest, +4); }, FunctionName, TEXT("int32")) && bSuccess;
+
+		// Test value underflow
+		volatile int32 Value3 = MIN_int32 + 2;
+		bSuccess = TestInterlocked<int32>(&Value3, MIN_int32 + 2, MAX_int32 - 1, [](volatile int32* Dest) { return FPlatformAtomics::InterlockedAdd(Dest, -4); }, FunctionName, TEXT("int32")) && bSuccess;
+	}
+
+	{
+		// Test adding positive value
+		alignas(8) volatile int64 Value0 = 0x0F00000000000000LL;
+		bSuccess = TestInterlocked<int64>(&Value0, 0x0F00000000000000LL, 0x1001010101010101LL, [](volatile int64* Dest) { return FPlatformAtomics::InterlockedAdd(Dest, 0x0101010101010101LL); }, FunctionName, TEXT("int64")) && bSuccess;
+
+		// Test adding negative value
+		alignas(8) volatile int64 Value1 = 0x1001010101010101LL;
+		bSuccess = TestInterlocked<int64>(&Value1, 0x1001010101010101LL, 0x0F00000000000000LL, [](volatile int64* Dest) { return FPlatformAtomics::InterlockedAdd(Dest, -0x0101010101010101LL); }, FunctionName, TEXT("int64")) && bSuccess;
+
+		// Test value overflow
+		alignas(8) volatile int64 Value2 = MAX_int64 - 1LL;
+		bSuccess = TestInterlocked<int64>(&Value2, MAX_int64 - 1, MIN_int64 + 2, [](volatile int64* Dest) { return FPlatformAtomics::InterlockedAdd(Dest, +4); }, FunctionName, TEXT("int64")) && bSuccess;
+
+		// Test value underflow
+		alignas(8) volatile int64 Value3 = MIN_int64 + 2LL;
+		bSuccess = TestInterlocked<int64>(&Value3, MIN_int64 + 2, MAX_int64 - 1, [](volatile int64* Dest) { return FPlatformAtomics::InterlockedAdd(Dest, -4); }, FunctionName, TEXT("int64")) && bSuccess;
+	}
+
+	return bSuccess;
+}
+
+static bool TestInterlockedIncrement()
+{
+	bool bSuccess = true;
+
+	const TCHAR* FunctionName = TEXT("Increment");
+	{
+		volatile int8 Value0 = 0x0F;
+		bSuccess = TestInterlocked<int8>(&Value0, 0x10, 0x10, [](volatile int8* Dest) { return FPlatformAtomics::InterlockedIncrement(Dest); }, FunctionName, TEXT("int8")) && bSuccess;
+
+		// Test value overflow
+		volatile int8 Value1 = MAX_int8;
+		bSuccess = TestInterlocked<int8>(&Value1, MIN_int8, MIN_int8, [](volatile int8* Dest) { return FPlatformAtomics::InterlockedIncrement(Dest); }, FunctionName, TEXT("int8")) && bSuccess;
+	}
+
+	{
+		volatile int16 Value0 = 0x0F0F;
+		bSuccess = TestInterlocked<int16>(&Value0, 0x0F10, 0x0F10, [](volatile int16* Dest) { return FPlatformAtomics::InterlockedIncrement(Dest); }, FunctionName, TEXT("int16")) && bSuccess;
+
+		// Test value overflow
+		volatile int16 Value1 = MAX_int16;
+		bSuccess = TestInterlocked<int16>(&Value1, MIN_int16, MIN_int16, [](volatile int16* Dest) { return FPlatformAtomics::InterlockedIncrement(Dest); }, FunctionName, TEXT("int16")) && bSuccess;
+	}
+
+	{
+		volatile int32 Value0 = 0x0F00000F;
+		bSuccess = TestInterlocked<int32>(&Value0, 0x0F000010, 0x0F000010, [](volatile int32* Dest) { return FPlatformAtomics::InterlockedIncrement(Dest); }, FunctionName, TEXT("int32")) && bSuccess;
+
+		// Test value overflow
+		volatile int32 Value1 = MAX_int32;
+		bSuccess = TestInterlocked<int32>(&Value1, MIN_int32, MIN_int32, [](volatile int32* Dest) { return FPlatformAtomics::InterlockedIncrement(Dest); }, FunctionName, TEXT("int32")) && bSuccess;
+	}
+
+	{
+		alignas(8) volatile int64 Value0 = 0x0F0000000000000FLL;
+		bSuccess = TestInterlocked<int64>(&Value0, 0x0F00000000000010LL, 0x0F00000000000010LL, [](volatile int64* Dest) { return FPlatformAtomics::InterlockedIncrement(Dest); }, FunctionName, TEXT("int64")) && bSuccess;
+
+		// Test value overflow
+		alignas(8) volatile int64 Value1 = MAX_int64;
+		bSuccess = TestInterlocked<int64>(&Value1, MIN_int64, MIN_int64, [](volatile int64* Dest) { return FPlatformAtomics::InterlockedIncrement(Dest); }, FunctionName, TEXT("int64")) && bSuccess;
+	}
+
+	return bSuccess;
+}
+
+static bool TestInterlockedDecrement()
+{
+	bool bSuccess = true;
+
+	const TCHAR* FunctionName = TEXT("Decrement");
+	{
+		volatile int8 Value0 = 0x10;
+		bSuccess = TestInterlocked<int8>(&Value0, 0x0F, 0x0F, [](volatile int8* Dest) { return FPlatformAtomics::InterlockedDecrement(Dest); }, FunctionName, TEXT("int8")) && bSuccess;
+
+		// Test value underflow
+		volatile int8 Value1 = MIN_int8;
+		bSuccess = TestInterlocked<int8>(&Value1, MAX_int8, MAX_int8, [](volatile int8* Dest) { return FPlatformAtomics::InterlockedDecrement(Dest); }, FunctionName, TEXT("int8")) && bSuccess;
+	}
+
+	{
+		volatile int16 Value0 = 0x0F10;
+		bSuccess = TestInterlocked<int16>(&Value0, 0x0F0F, 0x0F0F, [](volatile int16* Dest) { return FPlatformAtomics::InterlockedDecrement(Dest); }, FunctionName, TEXT("int16")) && bSuccess;
+
+		// Test value underflow
+		volatile int16 Value1 = MIN_int16;
+		bSuccess = TestInterlocked<int16>(&Value1, MAX_int16, MAX_int16, [](volatile int16* Dest) { return FPlatformAtomics::InterlockedDecrement(Dest); }, FunctionName, TEXT("int16")) && bSuccess;
+	}
+
+	{
+		volatile int32 Value0 = 0x0F000010;
+		bSuccess = TestInterlocked<int32>(&Value0, 0x0F00000F, 0x0F00000F, [](volatile int32* Dest) { return FPlatformAtomics::InterlockedDecrement(Dest); }, FunctionName, TEXT("int32")) && bSuccess;
+
+		// Test value underflow
+		volatile int32 Value1 = MIN_int32;
+		bSuccess = TestInterlocked<int32>(&Value1, MAX_int32, MAX_int32, [](volatile int32* Dest) { return FPlatformAtomics::InterlockedDecrement(Dest); }, FunctionName, TEXT("int32")) && bSuccess;
+	}
+
+	{
+		alignas(8) volatile int64 Value0 = 0x0F00000000000010LL;
+		bSuccess = TestInterlocked<int64>(&Value0, 0x0F0000000000000FLL, 0x0F0000000000000FLL, [](volatile int64* Dest) { return FPlatformAtomics::InterlockedDecrement(Dest); }, FunctionName, TEXT("int64")) && bSuccess;
+
+		// Test value underflow
+		alignas(8) volatile int64 Value1 = MIN_int64;
+		bSuccess = TestInterlocked<int64>(&Value1, MAX_int64, MAX_int64, [](volatile int64* Dest) { return FPlatformAtomics::InterlockedDecrement(Dest); }, FunctionName, TEXT("int64")) && bSuccess;
+	}
+
+	return bSuccess;
+}
+
+static bool TestInterlockedExchange()
+{
+	bool bSuccess = true;
+
+	const TCHAR* FunctionName = TEXT("Exchange");
+	{
+		volatile int8 Value = 0x10;
+		bSuccess = TestInterlocked<int8>(&Value, 0x10, 0x01, [](volatile int8* Dest) { return FPlatformAtomics::InterlockedExchange(Dest, 0x01); }, FunctionName, TEXT("int8")) && bSuccess;
+	}
+
+	{
+		volatile int16 Value = 0x1000;
+		bSuccess = TestInterlocked<int16>(&Value, 0x1000, 0x0001, [](volatile int16* Dest) { return FPlatformAtomics::InterlockedExchange(Dest, 0x0001); }, FunctionName, TEXT("int16")) && bSuccess;
+	}
+
+	{
+		volatile int32 Value = 0x10000000;
+		bSuccess = TestInterlocked<int32>(&Value, 0x10000000, 0x00000101, [](volatile int32* Dest) { return FPlatformAtomics::InterlockedExchange(Dest, 0x00000101); }, FunctionName, TEXT("int32")) && bSuccess;
+	}
+
+	{
+		alignas(8) volatile int64 Value = 0x1000000000000000LL;
+		bSuccess = TestInterlocked<int64>(&Value, 0x1000000000000000LL, 0x0000000001010101LL, [](volatile int64* Dest) { return FPlatformAtomics::InterlockedExchange(Dest, 0x0000000001010101LL); }, FunctionName, TEXT("int64")) && bSuccess;
+	}
+
+	return bSuccess;
+}
+
+static bool TestInterlockedExchangePtr()
+{
+	int Test = 0;
+	void* volatile Value = &Test;
+	const void* ReturnValue = FPlatformAtomics::InterlockedExchangePtr(&Value, nullptr);
+	if (ReturnValue != &Test)
+	{
+		FAIL_CHECK(TEXT("FPlatformAtomics::InterlockedExchangePtr failed"));
+		return false;
+	}
+	if (Value != nullptr)
+	{
+		FAIL_CHECK(TEXT("FPlatformAtomics::InterlockedExchangePtr failed"));
+		return false;
+	}
+	return true;
+}
+
+static bool TestInterlockedCompareExchange()
+{
+	bool bSuccess = true;
+
+	const TCHAR* FunctionName = TEXT("CompareExchange");
+	{
+		volatile int8 Value = 0x10;
+
+		// Test value isn't changed when comparand differs
+		bSuccess = TestInterlocked<int8>(&Value, 0x10, 0x10, [](volatile int8* Dest) { return FPlatformAtomics::InterlockedCompareExchange(Dest, 0x10, 0x01); }, FunctionName, TEXT("int8")) && bSuccess;
+
+		// Test value is changed when comparand matches
+		bSuccess = TestInterlocked<int8>(&Value, 0x10, 0x01, [](volatile int8* Dest) { return FPlatformAtomics::InterlockedCompareExchange(Dest, 0x01, 0x10); }, FunctionName, TEXT("int8")) && bSuccess;
+	}
+
+	{
+		volatile int16 Value = 0x1000;
+
+		// Test value isn't changed when comparand differs
+		bSuccess = TestInterlocked<int16>(&Value, 0x1000, 0x1000, [](volatile int16* Dest) { return FPlatformAtomics::InterlockedCompareExchange(Dest, 0x1000, 0x0001); }, FunctionName, TEXT("int16")) && bSuccess;
+
+		// Test value is changed when comparand matches
+		bSuccess = TestInterlocked<int16>(&Value, 0x1000, 0x0001, [](volatile int16* Dest) { return FPlatformAtomics::InterlockedCompareExchange(Dest, 0x0001, 0x1000); }, FunctionName, TEXT("int16")) && bSuccess;
+	}
+
+	{
+		volatile int32 Value = 0x10000000;
+
+		// Test value isn't changed when comparand differs
+		bSuccess = TestInterlocked<int32>(&Value, 0x10000000, 0x10000000, [](volatile int32* Dest) { return FPlatformAtomics::InterlockedCompareExchange(Dest, 0x10000000, 0x00000101); }, FunctionName, TEXT("int32")) && bSuccess;
+
+		// Test value is changed when comparand matches
+		bSuccess = TestInterlocked<int32>(&Value, 0x10000000, 0x00000101, [](volatile int32* Dest) { return FPlatformAtomics::InterlockedCompareExchange(Dest, 0x00000101, 0x10000000); }, FunctionName, TEXT("int32")) && bSuccess;
+	}
+
+	{
+		alignas(8) volatile int64 Value = 0x1000000000000000LL;
+
+		// Test value isn't changed when comparand differs
+		bSuccess = TestInterlocked<int64>(&Value, 0x1000000000000000LL, 0x1000000000000000LL, [](volatile int64* Dest) { return FPlatformAtomics::InterlockedCompareExchange(Dest, 0x1000000000000000LL, 0x0000000001010101LL); }, FunctionName, TEXT("int64")) && bSuccess;
+
+		// Test value is changed when comparand matches
+		bSuccess = TestInterlocked<int64>(&Value, 0x1000000000000000LL, 0x0000000001010101LL, [](volatile int64* Dest) { return FPlatformAtomics::InterlockedCompareExchange(Dest, 0x0000000001010101LL, 0x1000000000000000LL); }, FunctionName, TEXT("int64")) && bSuccess;
+	}
+
+	return bSuccess;
+}
+
+TEST_CASE_NAMED(FPlatformAtomicsTest, "System::Core::HAL::PlatformAtomics", "[ApplicationContextMask][SmokeFilter]")
+{
+	bool bSuccess = true;
+	bSuccess = TestInterlockedAnd() && bSuccess;
+	bSuccess = TestInterlockedOr() && bSuccess;
+	bSuccess = TestInterlockedXor() && bSuccess;
+	bSuccess = TestInterlockedAdd() && bSuccess;
+	bSuccess = TestInterlockedIncrement() && bSuccess;
+	bSuccess = TestInterlockedDecrement() && bSuccess;
+	bSuccess = TestInterlockedExchange() && bSuccess;
+	bSuccess = TestInterlockedExchangePtr() && bSuccess;
+	bSuccess = TestInterlockedCompareExchange() && bSuccess;
+	REQUIRE(bSuccess);
+}
+
+#endif //WITH_TESTS
diff --git a/Engine/Source/Runtime/Core/Tests/HAL/PlatformStringTest.cpp b/Engine/Source/Runtime/Core/Tests/HAL/PlatformStringTest.cpp
new file mode 100644
--- /dev/null
+++ b/Engine/Source/Runtime/Core/Tests/HAL/PlatformStringTest.cpp
@@ -0,0 +1,47 @@
+// Copyright Epic Games, Inc. All Rights Reserved.
+
+#include "HAL/PlatformString.h"
+
+#include "Tests/TestHarnessAdapter.h"
+
+#if WITH_TESTS
+
+template <typename CharType, SIZE_T Size>
+static void InvokePlatformStringGetVarArgs(CharType (&Dest)[Size], const CharType* Fmt, ...)
+{
+	va_list ap;
+	va_start(ap, Fmt);
+	FPlatformString::GetVarArgs(Dest, Size, Fmt, ap);
+	va_end(ap);
+}
+
+TEST_CASE_NAMED(FPlatformStringTestGetVarArgs, "System::Core::HAL::PlatformString::GetVarArgs", "[ApplicationContextMask][EngineFilter]")
+{
+	TCHAR Buffer[128];
+	InvokePlatformStringGetVarArgs(Buffer, TEXT("A%.*sZ"), 4, TEXT(" to B"));
+	CHECK_MESSAGE(TEXT("GetVarArgs(%.*s)"), FCString::Strcmp(Buffer, TEXT("A to Z")) == 0);
+}
+
+TEST_CASE_NAMED(FPlatformStringTestStrnlen, "System::Core::HAL::PlatformString::Strnlen", "[ApplicationContextMask][EngineFilter]")
+{
+	CHECK_MESSAGE(TEXT("Strnlen(nullptr, 0)"), FPlatformString::Strnlen((const ANSICHAR*)nullptr, 0) == 0);
+	CHECK_MESSAGE(TEXT("Strnlen(\"\", 0)"), FPlatformString::Strnlen("", 0) == 0);
+	CHECK_MESSAGE(TEXT("Strnlen(\"1\", 0)"), FPlatformString::Strnlen("1", 0) == 0);
+	CHECK_MESSAGE(TEXT("Strnlen(\"1\", 1)"), FPlatformString::Strnlen("1", 1) == 1);
+	CHECK_MESSAGE(TEXT("Strnlen(\"1\", 2)"), FPlatformString::Strnlen("1", 2) == 1);
+	CHECK_MESSAGE(TEXT("Strnlen(\"123\", 2)"), FPlatformString::Strnlen("123", 2) == 2);
+	ANSICHAR AnsiBuffer[128] = "123456789";
+	CHECK_MESSAGE(TEXT("Strnlen(PaddedBuffer)"), FPlatformString::Strnlen(AnsiBuffer, UE_ARRAY_COUNT(AnsiBuffer)) == 9);
+
+	CHECK_MESSAGE(TEXT("Strnlen(nullptr, 0)"), FPlatformString::Strnlen((const TCHAR*)nullptr, 0) == 0); //-V575
+	CHECK_MESSAGE(TEXT("Strnlen(\"\", 0)"), FPlatformString::Strnlen(TEXT(""), 0) == 0); //-V575
+	CHECK_MESSAGE(TEXT("Strnlen(\"1\", 0)"), FPlatformString::Strnlen(TEXT("1"), 0) == 0); //-V575
+	CHECK_MESSAGE(TEXT("Strnlen(\"1\", 1)"), FPlatformString::Strnlen(TEXT("1"), 1) == 1);
+	CHECK_MESSAGE(TEXT("Strnlen(\"1\", 2)"), FPlatformString::Strnlen(TEXT("1"), 2) == 1);
+	CHECK_MESSAGE(TEXT("Strnlen(\"123\", 2)"), FPlatformString::Strnlen(TEXT("123"), 2) == 2);
+	TCHAR Buffer[128] = {};
+	FCString::Strcpy(Buffer, TEXT("123456789"));
+	CHECK_MESSAGE(TEXT("Strnlen(PaddedBuffer)"), FPlatformString::Strnlen(Buffer, UE_ARRAY_COUNT(Buffer)) == 9);
+}
+
+#endif //WITH_TESTS
diff --git a/Engine/Source/Runtime/Core/Tests/HAL/ThreadTest.cpp b/Engine/Source/Runtime/Core/Tests/HAL/ThreadTest.cpp
new file mode 100644
--- /dev/null
+++ b/Engine/Source/Runtime/Core/Tests/HAL/ThreadTest.cpp
@@ -0,0 +1,267 @@
+// Copyright Epic Games, Inc. All Rights Reserved.
+
+#include "CoreTypes.h"
+#include "Misc/AssertionMacros.h"
+#include "Tests/TestHarnessAdapter.h"
+
+#if WITH_TESTS
+
+#include "HAL/Thread.h"
+#include "HAL/ThreadSafeBool.h"
+#include "HAL/ThreadSafeCounter.h"
+#include "HAL/ThreadSingleton.h"
+#include "HAL/Event.h"
+#include "Containers/Queue.h"
+#include "Containers/StringConv.h"
+
+namespace
+{
+	void TestIsJoinableAfterCreation()
+	{
+		FThread Thread(TEXT("Test.Thread.TestIsJoinableAfterCreation"), []() { /*NOOP*/ });
+		CHECK_MESSAGE(TEXT("FThread must be joinable after construction"), Thread.IsJoinable());
+		Thread.Join();
+		UE_LOG(LogTemp, Log, TEXT("%s completed"), StringCast<TCHAR>(__FUNCTION__).Get());
+	}
+
+	void TestIsJoinableAfterCompletion()
+	{
+		FThreadSafeBool bDone = false;
+		FThread Thread(TEXT("Test.Thread.TestIsJoinableAfterCompletion"), [&bDone]() { bDone = true; });
+		while (!bDone); // wait for completion //-V529
+		CHECK_MESSAGE(TEXT("FThread must still be joinable after completion"), Thread.IsJoinable());
+		Thread.Join();
+		UE_LOG(LogTemp, Log, TEXT("%s completed"), StringCast<TCHAR>(__FUNCTION__).Get());
+	}
+
+	void TestIsNotJoinableAfterJoining()
+	{
+		FThread Thread(TEXT("Test.Thread.TestIsNotJoinableAfterJoining"), []() { /*NOOP*/ });
+		Thread.Join();
+		CHECK_FALSE_MESSAGE(TEXT("FThread must not be joinable after joining"), Thread.IsJoinable());
+		UE_LOG(LogTemp, Log, TEXT("%s completed"), StringCast<TCHAR>(__FUNCTION__).Get());
+	}
+
+#if 0 // detaching is not implemented
+	void TestIsNotJoinableAfterDetaching()
+	{
+		// two cases: it's either the calling thread detaches from the thread before the thread is completed
+		{
+			TAtomic<bool> bReady{ false };
+			FThread Thread(TEXT("Test.Thread"), [&bReady]()
+				{
+					while (!bReady) {}
+				});
+			Thread.Detach();
+			bReady = true; // make sure `Detach` is called before thread function exit
+			CHECK_FALSE_MESSAGE(TEXT("FThread must not be joinable after detaching"), Thread.IsJoinable());
+		}
+		// or thread function is completed fast and `FThreadImpl` releases the reference to itself before
+		// `Detach` call
+		{
+			TAtomic<bool> bReady{ false };
+			FThread Thread(TEXT("Test.Thread"), [&bReady]() { /*NOOP*/});
+			FPlatformProcess::Sleep(0.1); // let the thread exit before detaching
+			Thread.Detach();
+			bReady = true; // make sure `Detach` is called before thread function exit
+			CHECK_FALSE_MESSAGE(TEXT("FThread must not be joinable after detaching"), Thread.IsJoinable());
+		}
+		UE_LOG(LogTemp, Log, TEXT("%s completed"), StringCast<TCHAR>(__FUNCTION__).Get());
+	}
+#endif
+
+	void TestAssertIfNotJoinedOrDetached()
+	{
+		// this does fails the `check`, but it seems there's no way to test this by UE4 Automation Testing, so commented out
+		FThread Thread(TEXT("Test.Thread.TestAssertIfNotJoinedOrDetached"), []() { /*NOOP*/ });
+		// should assert in the destructor
+	}
+
+	void TestDefaultConstruction()
+	{
+		{
+			FThread Thread;
+			CHECK_FALSE_MESSAGE(TEXT("Default-constructed FThread must be not joinable"), Thread.IsJoinable());
+		}
+		{	// check that default constructed thread can be "upgraded" to joinable thread
+			FThread Thread;
+			Thread = FThread(TEXT("Test.Thread.TestDefaultConstruction"), []() { /* NOOP */ });
+			CHECK_MESSAGE(TEXT("Move-constructed FThread from joinable thread must be joinable"), Thread.IsJoinable());
+			Thread.Join();
+		}
+		UE_LOG(LogTemp, Log, TEXT("%s completed"), StringCast<TCHAR>(__FUNCTION__).Get());
+	}
+
+	void TestThreadSingleton()
+	{
+		{	// check that ThreadSingleton instances in different threads are isolated from each other
+			class FThreadSingletonTest : public TThreadSingleton<FThreadSingletonTest>
+			{
+			public:
+				void SetTestField(int NewValue) { TestField = NewValue; }
+				int GetTestField() const { return TestField; }
+			private:
+				int TestField{ 0 };
+			};
+			FThreadSingletonTest::Get().SetTestField(1);
+			FThread Thread;
+			TAtomic<bool> bDefaultValuePass{ false };
+			Thread = FThread(TEXT("Test.Thread.TestThreadSingleton"),
+				[&bDefaultValuePass]()
+				{
+					bDefaultValuePass = FThreadSingletonTest::TryGet() == nullptr;
+					bDefaultValuePass = bDefaultValuePass && (FThreadSingletonTest::Get().GetTestField() == 0);
+					FThreadSingletonTest::Get().SetTestField(2);
+				});
+			Thread.Join();
+			CHECK_MESSAGE(TEXT("Thread singleton should be uninitialized by default"), bDefaultValuePass);
+			CHECK_MESSAGE(TEXT("Thread singletons in different threads should be isolated"), FThreadSingletonTest::Get().GetTestField() == 1);
+		}
+		{	// check that ThreadSingleton entries don't point to invalid memory after cleanup
+			class FThreadSingletonFirst : public TThreadSingleton<FThreadSingletonFirst>
+			{
+			};
+			class FThreadSingletonSecond : public TThreadSingleton<FThreadSingletonSecond>
+			{
+			public:
+				virtual ~FThreadSingletonSecond()
+				{
+					// By the time we reach this destructor, the first singleton's destructor should have been executed already.
+					check(FThreadSingletonFirst::TryGet() == nullptr);
+				}
+			};
+			FThread Thread;
+			Thread = FThread(TEXT("Test.Thread.TestThreadSingleton"),
+				[]()
+				{
+					FThreadSingletonFirst::Get();
+					FThreadSingletonSecond::Get();
+				});
+			Thread.Join();
+			CHECK_MESSAGE(TEXT("Thread singletons should be uninitialized by default"), FThreadSingletonFirst::TryGet() == nullptr);
+			CHECK_MESSAGE(TEXT("Thread singletons should be uninitialized by default"), FThreadSingletonSecond::TryGet() == nullptr);
+		}
+		UE_LOG(LogTemp, Log, TEXT("%s completed"), StringCast<TCHAR>(__FUNCTION__).Get());
+	}
+
+	void TestMovability()
+	{
+		{	// move constructor with default-constructed thread
+			FThread Src;
+			FThread Dst(MoveTemp(Src));
+			CHECK_FALSE_MESSAGE(TEXT("Default-constructed thread must stay not joinable after moving out"), Src.IsJoinable());
+			CHECK_FALSE_MESSAGE(TEXT("Move-constructed thread from not joinable thread must be not joinable"), Dst.IsJoinable());
+		}
+		{	// move constructor with joinable thread
+			FThread Src(TEXT("Test.Thread.TestMovability.1"), []() { /* NOOP */ });
+			FThread Dst(MoveTemp(Src));
+			CHECK_FALSE_MESSAGE(TEXT("Moved out thread must be not joinable"), Src.IsJoinable());
+			CHECK_MESSAGE(TEXT("Move-constructed thread from joinable thread must be joinable"), Dst.IsJoinable());
+			Dst.Join();
+		}
+		{	// move assignment operator
+			FThread Src(TEXT("Test.Thread.TestMovability.2"), []() { /* NOOP */ });
+			FThread Dst;
+			Dst = MoveTemp(Src);
+			CHECK_FALSE_MESSAGE(TEXT("Moved out thread must be not joinable"), Src.IsJoinable());
+			CHECK_MESSAGE(TEXT("Move-assigned thread from joinable thread must be joinable"), Dst.IsJoinable());
+			Dst.Join();
+		}
+		{	// Failure test for move assignment operator of joinable thread
+			//FThread Src(TEXT("Test.Thread"), []() { /* NOOP */ });
+			//FThread Dst(TEXT("Test.Thread"), []() { /* NOOP */ });
+			//Dst = MoveTemp(Src); // must assert that joinable thread wasn't joined before move-assignment, no way to test this
+			//Dst.Join();
+		}
+		{	// Move assignment operator of thread that has been joined
+			FThread Src(TEXT("Test.Thread.TestMovability.3"), []() { /* NOOP */ });
+			FThread Dst(TEXT("Test.Thread.TestMovability.4"), []() { /* NOOP */ });
+			Dst.Join();
+			Dst = MoveTemp(Src);
+			Dst.Join();
+		}
+		UE_LOG(LogTemp, Log, TEXT("%s completed"), StringCast<TCHAR>(__FUNCTION__).Get());
+	}
+
+	// An example of possible implementation of Consumer/Producer idiom
+	void TestTypicalUseCase()
+	{
+		FThreadSafeBool bQuitRequested = false;
+		using FWork = uint32;
+		TQueue<FWork> WorkQueue;
+		FEvent* WorkQueuedEvent = FPlatformProcess::GetSynchEventFromPool();
+
+		FThread WorkerThread(TEXT("Test.Thread.TestTypicalUseCase"), [&bQuitRequested, &WorkQueue, WorkQueuedEvent]() 
+		{
+			while (!bQuitRequested)
+			{
+				// get work
+				FWork Work;
+				if (!WorkQueue.Dequeue(Work))
+				{
+					WorkQueuedEvent->Wait();
+					// first check if quitting requested
+					continue;
+				}
+
+				// do work
+				UE_LOG(LogTemp, Log, TEXT("Work #%d consumed"), Work);
+			}
+
+			UE_LOG(LogTemp, Log, TEXT("Quit"));
+		});
+
+		// produce work
+		const int WorkNum = 3;
+		for (FWork Work = 0; Work != WorkNum; ++Work)
+		{
+			WorkQueue.Enqueue(Work);
+			WorkQueuedEvent->Trigger();
+			UE_LOG(LogTemp, Log, TEXT("Work #%d produced"), Work);
+		}
+
+		UE_LOG(LogTemp, Log, TEXT("Request to quit"));
+		bQuitRequested = true;
+		// the thread can be blocked waiting for work, unblock it
+		WorkQueuedEvent->Trigger();
+		WorkerThread.Join();
+
+		FPlatformProcess::ReturnSynchEventToPool(WorkQueuedEvent);
+
+		// example of output:
+		// Work #0 produced
+		//	Work #0 consumed
+		//	Work #1 produced
+		//	Work #1 consumed
+		//	Work #2 produced
+		//	Work #2 consumed
+		//	Request to quit
+		//	The thread 0x96e0 has exited with code 0 (0x0).
+		//	Quit
+	
+		UE_LOG(LogTemp, Log, TEXT("%s completed"), StringCast<TCHAR>(__FUNCTION__).Get());
+	}
+}
+
+TEST_CASE_NAMED(FThreadTest, "System::Core::HAL::Thread", "[ApplicationContextMask][EngineFilter]")
+{
+	UE_LOG(LogTemp, Log, TEXT("%s"), StringCast<TCHAR>(__FUNCTION__).Get());
+
+	TestIsJoinableAfterCreation();
+	TestIsJoinableAfterCompletion();
+	TestIsNotJoinableAfterJoining();
+	
+#if 0 // detaching is not implemented
+	TestIsNotJoinableAfterDetaching();
+#endif
+
+	//TestAssertIfNotJoinedOrDetached();
+
+	TestDefaultConstruction();
+	TestMovability();
+
+	TestTypicalUseCase();
+	TestThreadSingleton();
+}
+
+#endif //WITH_TESTS
