Fix C++20 compile issue. != is automatically defined if == exists in C++20, resulting in new ambiguity in some RingBuffer test code.

could be 'bool FRingBufferTest::Counter::operator ==(const FRingBufferTest::Counter &) const' [rewritten expression '!(x == y)']
or 'built-in C++ operator!=(uint32, uint32)'

diff --git a/Engine/Source/Runtime/Core/Tests/Containers/RingBufferTest.cpp b/Engine/Source/Runtime/Core/Tests/Containers/RingBufferTest.cpp
--- a/Engine/Source/Runtime/Core/Tests/Containers/RingBufferTest.cpp
+++ b/Engine/Source/Runtime/Core/Tests/Containers/RingBufferTest.cpp
@@ -777,7 +777,7 @@ public:
 					for (int32 Index = 0; Index < ExpectedNum; ++Index)
 					{
 						if ((Index < SourceStart && Q[Index] != Prefix[Index]) ||
-							(Index >= SourceStart && Q[Index] != MarkerValue + Index - SourceStart))
+							(Index >= SourceStart && (uint32)Q[Index] != MarkerValue + Index - SourceStart))
 						{
 							CHECK_MESSAGE(*FString::Printf(TEXT("%s: Element does not match"), HeaderMessage), false);
 							break;
