Fix C++20 compile error. FTestMoveOnly can no longer be aggregate initialized in C++20.

diff --git a/Engine/Source/Runtime/Core/Tests/Templates/ValueOrErrorTest.cpp b/Engine/Source/Runtime/Core/Tests/Templates/ValueOrErrorTest.cpp
--- a/Engine/Source/Runtime/Core/Tests/Templates/ValueOrErrorTest.cpp
+++ b/Engine/Source/Runtime/Core/Tests/Templates/ValueOrErrorTest.cpp
@@ -51,6 +51,7 @@ TEST_CASE_NAMED(FTemplatesValueOrErrorTest, "System::Core::Templates::ValueOrErr
 	struct FTestMoveOnly
 	{
 		FTestMoveOnly() = default;
+		FTestMoveOnly(int32 InValue) : Value(InValue) {}
 		FTestMoveOnly(const FTestMoveOnly&) = delete;
 		FTestMoveOnly(FTestMoveOnly&&) = default;
 		int Value = 0;
@@ -150,8 +151,8 @@ TEST_CASE_NAMED(FTemplatesValueOrErrorTest, "System::Core::Templates::ValueOrErr
 
 	SECTION("Move-Only Value/Error")
 	{
-		TValueOrError<FTestMoveOnly, FTestMoveOnly> Value = MakeValue(FTestMoveOnly{1});
-		TValueOrError<FTestMoveOnly, FTestMoveOnly> Error = MakeError(FTestMoveOnly{1});
+		TValueOrError<FTestMoveOnly, FTestMoveOnly> Value = MakeValue(FTestMoveOnly(1));
+		TValueOrError<FTestMoveOnly, FTestMoveOnly> Error = MakeError(FTestMoveOnly(1));
 		FTestMoveOnly MovedValue = MoveTemp(Value).GetValue();
 		FTestMoveOnly MovedError = MoveTemp(Error).GetError();
 		CHECK(MovedValue.Value == 1);
