Get the AutoRTFM testing working in UBT without all the CMake shenanigans.

diff --git a/Engine/Source/Runtime/Core/Private/AutoRTFM/StdLibSupport.cpp b/Engine/Source/Runtime/Core/Private/AutoRTFM/StdLibSupport.cpp
--- a/Engine/Source/Runtime/Core/Private/AutoRTFM/StdLibSupport.cpp
+++ b/Engine/Source/Runtime/Core/Private/AutoRTFM/StdLibSupport.cpp
@@ -11,6 +11,11 @@
 #include <float.h>
 #include <math.h>
 
+#ifdef _WIN32
+#define NOMINMAX
+#include <windows.h>
+#endif
+
 namespace AutoRTFM
 {
 
@@ -163,21 +168,9 @@ UE_AUTORTFM_REGISTER_SELF_FUNCTION(static_cast<double(*)(double, double)>(&pow))
 #ifdef _WIN32
 UE_AUTORTFM_REGISTER_SELF_FUNCTION(_isnan);
 UE_AUTORTFM_REGISTER_SELF_FUNCTION(_finite);
-
-extern "C" __declspec(dllimport) int __stdcall IsDebuggerPresent(void);
-
 UE_AUTORTFM_REGISTER_SELF_FUNCTION(IsDebuggerPresent);
-
-extern "C" __declspec(dllimport) void EnterCriticalSection(void*);
-
 UE_AUTORTFM_REGISTER_SELF_FUNCTION(EnterCriticalSection);
-
-extern "C" __declspec(dllimport) void LeaveCriticalSection(void*);
 UE_AUTORTFM_REGISTER_SELF_FUNCTION(LeaveCriticalSection);
-
-extern "C" __declspec(dllimport) int __stdcall QueryPerformanceCounter(void *);
-extern "C" __declspec(dllimport) int __stdcall QueryPerformanceFrequency(void *);
-
 UE_AUTORTFM_REGISTER_SELF_FUNCTION(QueryPerformanceCounter);
 #endif
 
