Fixed TScriptDelegate::operator= return type.

diff --git a/Engine/Source/Runtime/Core/Public/UObject/ScriptDelegates.h b/Engine/Source/Runtime/Core/Public/UObject/ScriptDelegates.h
--- a/Engine/Source/Runtime/Core/Public/UObject/ScriptDelegates.h
+++ b/Engine/Source/Runtime/Core/Public/UObject/ScriptDelegates.h
@@ -267,7 +267,7 @@ public:
 		UE_REQUIRES(UE::Core::Private::BackwardCompatibilityCheck<Dummy, OtherDummy>())
 	>
 	/* UE_DEPRECATED(5.3, "Deprecated - remove after TScriptDelegateTraits<FWeakObjectPtr> is removed") */
-	FORCEINLINE bool operator=(const TScriptDelegate<OtherDummy>& Other)
+	FORCEINLINE void operator=(const TScriptDelegate<OtherDummy>& Other)
 	{
 		UE_DELEGATES_MT_SCOPED_WRITE_ACCESS(AccessDetector);
 
