Minor edit to AESTest to save ~20s of compile time with MSVC

#rnx

diff --git a/Engine/Source/Runtime/Core/Tests/Misc/AESTest.cpp b/Engine/Source/Runtime/Core/Tests/Misc/AESTest.cpp
--- a/Engine/Source/Runtime/Core/Tests/Misc/AESTest.cpp
+++ b/Engine/Source/Runtime/Core/Tests/Misc/AESTest.cpp
@@ -64,9 +64,9 @@ TEST_CASE_NAMED(FMiscAESTest, "System::Core::Misc::AES", "[Core][AES][SmokeFilte
 
 	struct AesTestCase
 	{
-		FAnsiStringView Key;
-		FAnsiStringView PlainText;
-		FAnsiStringView CipherText;
+		const ANSICHAR* Key = nullptr;
+		const ANSICHAR* PlainText = nullptr;
+		const ANSICHAR* CipherText = nullptr;
 	};
 
 	// Appendix C, https://csrc.nist.gov/csrc/media/publications/fips/197/final/documents/fips-197.pdf
