CrashHandling: FlushBufferedItems should early exit if Panic is started on another thread to avoid having its lock ignored by the Panic thread and causing a crash during crash-handling on the Panic thread due to unguarded simultaneous access to BufferedOutputDevices.
#rnx

diff --git a/Engine/Source/Runtime/Core/Private/Misc/OutputDeviceRedirector.cpp b/Engine/Source/Runtime/Core/Private/Misc/OutputDeviceRedirector.cpp
--- a/Engine/Source/Runtime/Core/Private/Misc/OutputDeviceRedirector.cpp
+++ b/Engine/Source/Runtime/Core/Private/Misc/OutputDeviceRedirector.cpp
@@ -480,25 +480,36 @@ void FOutputDeviceRedirectorState::FlushBufferedItems()
 	TRACE_CPUPROFILER_EVENT_SCOPE(FOutputDeviceRedirector::FlushBufferedItems);
 
 	const uint32 ThreadId = FPlatformTLS::GetCurrentThreadId();
-	BufferedItems.Deplete([this, ThreadId](FOutputDeviceItem&& Item)
+	bool bIsPanicThread = this->IsPanicThread(ThreadId);
+	BufferedItems.Deplete([this, ThreadId, bIsPanicThread](FOutputDeviceItem&& Item)
 	{
-		Visit([this, ThreadId](auto&& Value)
+		if (!bIsPanicThread && HasPanicThread())
 		{
-			using ValueType = std::decay_t<decltype(Value)>;
-			if constexpr (std::is_same_v<ValueType, FOutputDeviceLine>)
-			{
-				const FOutputDeviceLine& Line = Value;
-				BroadcastTo(ThreadId, BufferedOutputDevices, BufferedOutputDevicesCanBeUsedOnPanicThread,
-					UE_PROJECTION_MEMBER(FOutputDevice, Serialize),
-					Line.Data, Line.Verbosity, Line.Category, Line.Time);
-			}
-			else if constexpr (std::is_same_v<ValueType, FLogRecord>)
+			// Enqueuing during Deplete is okay because the items passed to the Deplete consumer
+			// are a moved list that is no longer associated with the queue; Deplete will return
+			// even though new items are in the queue.
+			BufferedItems.Enqueue(MoveTemp(Item));
+		}
+		else
+		{
+			Visit([this, ThreadId](auto&& Value)
 			{
-				const FLogRecord& Record = Value;
-				BroadcastTo(ThreadId, BufferedOutputDevices, BufferedOutputDevicesCanBeUsedOnPanicThread,
-					UE_PROJECTION_MEMBER(FOutputDevice, SerializeRecord), Record);
-			}
-		}, Item.Value);
+				using ValueType = std::decay_t<decltype(Value)>;
+				if constexpr (std::is_same_v<ValueType, FOutputDeviceLine>)
+				{
+					const FOutputDeviceLine& Line = Value;
+					BroadcastTo(ThreadId, BufferedOutputDevices, BufferedOutputDevicesCanBeUsedOnPanicThread,
+						UE_PROJECTION_MEMBER(FOutputDevice, Serialize),
+						Line.Data, Line.Verbosity, Line.Category, Line.Time);
+				}
+				else if constexpr (std::is_same_v<ValueType, FLogRecord>)
+				{
+					const FLogRecord& Record = Value;
+					BroadcastTo(ThreadId, BufferedOutputDevices, BufferedOutputDevicesCanBeUsedOnPanicThread,
+						UE_PROJECTION_MEMBER(FOutputDevice, SerializeRecord), Record);
+				}
+			}, Item.Value);
+		}
 	});
 }
 
