Bring back debugger breaking on assertion macros line instead of cpp file for CheckVerifyImpl

diff --git a/Engine/Source/Runtime/Core/Private/Misc/AssertionMacros.cpp b/Engine/Source/Runtime/Core/Private/Misc/AssertionMacros.cpp
--- a/Engine/Source/Runtime/Core/Private/Misc/AssertionMacros.cpp
+++ b/Engine/Source/Runtime/Core/Private/Misc/AssertionMacros.cpp
@@ -562,7 +562,7 @@ FORCENOINLINE void FDebug::EnsureFailed(const ANSICHAR* Expr, const ANSICHAR* Fi
 	FPlatformAtomics::InterlockedDecrement(&ActiveEnsureCount);
 }
 
-void FORCENOINLINE FDebug::CheckVerifyFailedImpl(
+bool FORCENOINLINE FDebug::CheckVerifyFailedImpl(
 	const ANSICHAR* Expr,
 	const ANSICHAR* File,
 	int32 Line,
@@ -590,8 +590,11 @@ void FORCENOINLINE FDebug::CheckVerifyFailedImpl(
 		va_start(Args, Format);
 		AssertFailedImplV(Expr, File, Line, ProgramCounter, Format, Args);
 		va_end(Args);
+
+		return false;
 	}
-	PLATFORM_BREAK_IF_DESIRED_NONINLINE();
+	
+	return !GIgnoreDebugger;
 }
 
 #endif // DO_CHECK || DO_GUARD_SLOW || DO_ENSURE
@@ -706,8 +709,7 @@ bool UE_DEBUG_SECTION VARARGS CheckVerifyImpl(bool& InOutExecuted, bool Always,
 			return false;
 		}
 
-		PLATFORM_BREAK_IF_DESIRED_NONINLINE();
-		return true;
+		return !GIgnoreDebugger;
 	}
 
 	return false;
diff --git a/Engine/Source/Runtime/Core/Public/Misc/AssertionMacros.h b/Engine/Source/Runtime/Core/Public/Misc/AssertionMacros.h
--- a/Engine/Source/Runtime/Core/Public/Misc/AssertionMacros.h
+++ b/Engine/Source/Runtime/Core/Public/Misc/AssertionMacros.h
@@ -85,7 +85,7 @@ struct CORE_API FDebug
 
 #if DO_CHECK || DO_GUARD_SLOW || DO_ENSURE
 public:
-	static void VARARGS CheckVerifyFailedImpl(const ANSICHAR* Expr, const ANSICHAR* File, int32 Line, void* ProgramCounter, const TCHAR* Format, ...);
+	static bool VARARGS CheckVerifyFailedImpl(const ANSICHAR* Expr, const ANSICHAR* File, int32 Line, void* ProgramCounter, const TCHAR* Format, ...);
 private:
 	static void VARARGS LogAssertFailedMessageImpl(const ANSICHAR* Expr, const ANSICHAR* File, int32 Line, void* ProgramCounter, const TCHAR* Fmt, ...);
 	static void LogAssertFailedMessageImplV(const ANSICHAR* Expr, const ANSICHAR* File, int32 Line, void* ProgramCounter, const TCHAR* Fmt, va_list Args);
@@ -212,17 +212,6 @@ RetType FORCENOINLINE UE_DEBUG_SECTION DispatchCheckVerify(InnerType&& Inner, Ar
 	#define PLATFORM_BREAK_IF_DESIRED() PLATFORM_BREAK();
 #endif // !UE_BUILD_SHIPPING
 
-#if PLATFORM_WINDOWS
-	// We want to break inline of an assertion macro on Windows for ease of debugging
-#	define PLATFORM_BREAK_IF_DESIRED_INLINE() PLATFORM_BREAK_IF_DESIRED()
-#	define PLATFORM_BREAK_IF_DESIRED_NONINLINE()
-#else
-	// ... however on every other platform this causes significant code bloat and we want to avoid it.
-	// So stepping up through callstack in a debugger is a necessary price to pay
-#	define PLATFORM_BREAK_IF_DESIRED_INLINE()
-#	define PLATFORM_BREAK_IF_DESIRED_NONINLINE() PLATFORM_BREAK_IF_DESIRED()
-#endif
-
 
 #if DO_CHECK
 #ifndef checkCode
@@ -242,8 +231,10 @@ RetType FORCENOINLINE UE_DEBUG_SECTION DispatchCheckVerify(InnerType&& Inner, Ar
 		{ \
 			if(UNLIKELY(!(expr))) \
 			{ \
-				FDebug::CheckVerifyFailedImpl(#expr, __FILE__, __LINE__, PLATFORM_RETURN_ADDRESS(), TEXT("")); \
-				PLATFORM_BREAK_IF_DESIRED_INLINE(); \
+				if (FDebug::CheckVerifyFailedImpl(#expr, __FILE__, __LINE__, PLATFORM_RETURN_ADDRESS(), TEXT(""))) \
+				{ \
+					PLATFORM_BREAK(); \
+				} \
 				CA_ASSUME(false); \
 			} \
 		}
@@ -263,8 +254,10 @@ RetType FORCENOINLINE UE_DEBUG_SECTION DispatchCheckVerify(InnerType&& Inner, Ar
 		{ \
 			if(UNLIKELY(!(expr))) \
 			{ \
-				FDebug::CheckVerifyFailedImpl(#expr, __FILE__, __LINE__, PLATFORM_RETURN_ADDRESS(), format, ##__VA_ARGS__); \
-				PLATFORM_BREAK_IF_DESIRED_INLINE(); \
+				if (FDebug::CheckVerifyFailedImpl(#expr, __FILE__, __LINE__, PLATFORM_RETURN_ADDRESS(), format, ##__VA_ARGS__)) \
+				{ \
+					PLATFORM_BREAK(); \
+				} \
 				CA_ASSUME(false); \
 			} \
 		}
@@ -376,7 +369,7 @@ RetType FORCENOINLINE UE_DEBUG_SECTION DispatchCheckVerify(InnerType&& Inner, Ar
 			static bool bExecuted = false; \
 			if (CheckVerifyImpl(bExecuted, Always, __FILE__, __LINE__, PLATFORM_RETURN_ADDRESS(), #InExpression, InFormat)) \
 			{ \
-				PLATFORM_BREAK_IF_DESIRED_INLINE(); \
+				PLATFORM_BREAK(); \
 			} \
 			return false; \
 		})))
@@ -388,7 +381,7 @@ RetType FORCENOINLINE UE_DEBUG_SECTION DispatchCheckVerify(InnerType&& Inner, Ar
 			FValidateArgsInternal(__VA_ARGS__); \
 			if (CheckVerifyImpl(bExecuted, Always, __FILE__, __LINE__, PLATFORM_RETURN_ADDRESS(), #InExpression, ##__VA_ARGS__)) \
 			{ \
-				PLATFORM_BREAK_IF_DESIRED_INLINE(); \
+				PLATFORM_BREAK(); \
 			} \
 			return false; \
 		})))
