CrashHandling: During GIsCriticalError, skip recording of errors for display at the end of the commandlet; FMalloc might be broken and cause a crash during crash handling if we try to allocate any strings.
#rnx

diff --git a/Engine/Source/Runtime/Core/Private/Misc/FeedbackContext.cpp b/Engine/Source/Runtime/Core/Private/Misc/FeedbackContext.cpp
--- a/Engine/Source/Runtime/Core/Private/Misc/FeedbackContext.cpp
+++ b/Engine/Source/Runtime/Core/Private/Misc/FeedbackContext.cpp
@@ -117,6 +117,13 @@ void FFeedbackContext::FormatRecordLine(FStringBuilderBase& Out, const UE::FLogR
 
 void FFeedbackContext::AddToHistory(const TCHAR* V, ELogVerbosity::Type Verbosity, const FName& Category, double Time)
 {
+	if (GIsCriticalError)
+	{
+		// Turn off the storing of Errors and Warnings when GIsCriticalError is true.
+		// In a critical error (e.g. crash handling), it might be the case that Malloc is returning bad data,
+		// so we want to avoid dynamic allocation (such as the construction of strings below) as much as possible.
+		return;
+	}
 	TStringBuilder<512> Line;
 	FormatLine(Line, V, Verbosity, Category, Time, &Verbosity);
 	if (Verbosity == ELogVerbosity::Error)
diff --git a/Engine/Source/Runtime/Launch/Private/Windows/LaunchWindows.cpp b/Engine/Source/Runtime/Launch/Private/Windows/LaunchWindows.cpp
--- a/Engine/Source/Runtime/Launch/Private/Windows/LaunchWindows.cpp
+++ b/Engine/Source/Runtime/Launch/Private/Windows/LaunchWindows.cpp
@@ -249,8 +249,6 @@ LAUNCH_API int32 LaunchWindowsStartup( HINSTANCE hInInstance, HINSTANCE hPrevIns
 				? ( GEnableInnerException ? EXCEPTION_EXECUTE_HANDLER : ReportCrash(GetExceptionInformation()) )
 				: EXCEPTION_CONTINUE_SEARCH )	
 		{
-			UE_LOG(LogLaunchWindows, Log, TEXT("LaunchWindowsStartup.ExceptionHandler"));
-			UE_LOG(LogLaunchWindows, Error, TEXT("LaunchWindowsStartup.ExceptionHandler"));
 			// Crashed.
 			ErrorLevel = 1;
 			if(GError)
