switched to stronger memory order to avoid hitting the "IsCompleted()" check in LowLevelTask::~FTask() on arm64 architecture

diff --git a/Engine/Source/Runtime/Core/Private/Tasks/TaskPrivate.cpp b/Engine/Source/Runtime/Core/Private/Tasks/TaskPrivate.cpp
--- a/Engine/Source/Runtime/Core/Private/Tasks/TaskPrivate.cpp
+++ b/Engine/Source/Runtime/Core/Private/Tasks/TaskPrivate.cpp
@@ -226,7 +226,7 @@ namespace UE::Tasks
 			}
 
 			// the waiting task will be destroyed leaving this scope, wait for the internal reference to it to be released
-			while (WaitingTask.GetRefCount() != 1)
+			while (WaitingTask.GetRefCount(std::memory_order_acquire) != 1)
 			{
 				FPlatformProcess::Yield();
 			}
diff --git a/Engine/Source/Runtime/Core/Public/Tasks/Task.h b/Engine/Source/Runtime/Core/Public/Tasks/Task.h
--- a/Engine/Source/Runtime/Core/Public/Tasks/Task.h
+++ b/Engine/Source/Runtime/Core/Public/Tasks/Task.h
@@ -367,7 +367,7 @@ namespace UE::Tasks
 		}
 
 		// `WaitingTask` is allocated on the stack. we need to wait until its RefCount reaches zero before leaving the function to avoid "use after destruction"
-		while (WaitingTask.GetRefCount() != 1)
+		while (WaitingTask.GetRefCount(std::memory_order_acquire) != 1)
 		{
 			FPlatformProcess::Yield();
 		}
diff --git a/Engine/Source/Runtime/Core/Public/Tasks/TaskPrivate.h b/Engine/Source/Runtime/Core/Public/Tasks/TaskPrivate.h
--- a/Engine/Source/Runtime/Core/Public/Tasks/TaskPrivate.h
+++ b/Engine/Source/Runtime/Core/Public/Tasks/TaskPrivate.h
@@ -120,9 +120,9 @@ namespace UE::Tasks
 				}
 			}
 
-			uint32 GetRefCount() const
+			uint32 GetRefCount(std::memory_order MemoryOrder) const
 			{
-				return RefCount.load(std::memory_order_relaxed);
+				return RefCount.load(MemoryOrder);
 			}
 
 		private:
