[Insights] Fixed memory tracing for block allocators used by TConcurrentLinearAllocator (FOsAllocator, FAlignedAllocatorInternal, FOutputDeviceLinearAllocator's block allocator).

diff --git a/Engine/Source/Runtime/Core/Private/Misc/OutputDeviceRedirector.cpp b/Engine/Source/Runtime/Core/Private/Misc/OutputDeviceRedirector.cpp
--- a/Engine/Source/Runtime/Core/Private/Misc/OutputDeviceRedirector.cpp
+++ b/Engine/Source/Runtime/Core/Private/Misc/OutputDeviceRedirector.cpp
@@ -52,11 +52,16 @@ struct FOutputDeviceBlockAllocationTag : FDefaultBlockAllocationTag
 
 		FORCEINLINE static void* Malloc(SIZE_T Size, uint32 Alignment)
 		{
-			return FPlatformMemory::BaseAllocator()->Malloc(Size, DEFAULT_ALIGNMENT);
+			void* Pointer = FPlatformMemory::BaseAllocator()->Malloc(Size, DEFAULT_ALIGNMENT);
+			MemoryTrace_Alloc(uint64(Pointer), Size, DEFAULT_ALIGNMENT, EMemoryTraceRootHeap::SystemMemory);
+			MemoryTrace_MarkAllocAsHeap(uint64(Pointer), EMemoryTraceRootHeap::SystemMemory);
+			return Pointer;
 		}
 
 		FORCEINLINE static void Free(void* Pointer, SIZE_T Size)
 		{
+			MemoryTrace_UnmarkAllocAsHeap(uint64(Pointer), EMemoryTraceRootHeap::SystemMemory);
+			MemoryTrace_Free(uint64(Pointer), EMemoryTraceRootHeap::SystemMemory);
 			return FPlatformMemory::BaseAllocator()->Free(Pointer);
 		}
 	};
diff --git a/Engine/Source/Runtime/Core/Public/Experimental/ConcurrentLinearAllocator.h b/Engine/Source/Runtime/Core/Public/Experimental/ConcurrentLinearAllocator.h
--- a/Engine/Source/Runtime/Core/Public/Experimental/ConcurrentLinearAllocator.h
+++ b/Engine/Source/Runtime/Core/Public/Experimental/ConcurrentLinearAllocator.h
@@ -43,24 +43,25 @@ struct FOsAllocator
 {
 	static constexpr bool SupportsAlignment = false;
 
-	// Indicates whether memory is acquired from the same source as used by FMalloc i.e. FMemory/MPlatformMemory::BaseAllocator().
+	// Indicates whether memory is acquired from the same source as used by FMalloc i.e. FMemory/FPlatformMemory::BaseAllocator().
 	// It's needed so that ConcurrentLinearAllocator could correctly account for this memory in the memory stats and tracking systems.
 	// It is 'true' by default, which is suitable for most implementations. However, if you create a custom allocator that
-	// acquires memory in a different way e.g. directly from the OS using FPlatformMemory::BinnedAllosFromOS, then you
+	// acquires memory in a different way e.g. directly from the OS using FPlatformMemory::BinnedAllocFromOS, then you
 	// should set this variable to 'false'.
 	static constexpr bool UsesFMalloc = true;
 
 	FORCEINLINE static void* Malloc(SIZE_T Size, uint32 Alignment)
 	{
-		void* Pointer = FPlatformMemory::BaseAllocator()->Malloc(Size);
-		MemoryTrace_Alloc(uint64(Pointer), Size, Alignment);
+		void* Pointer = FPlatformMemory::BaseAllocator()->Malloc(Size, DEFAULT_ALIGNMENT);
+		MemoryTrace_Alloc(uint64(Pointer), Size, DEFAULT_ALIGNMENT, EMemoryTraceRootHeap::SystemMemory);
 		MemoryTrace_MarkAllocAsHeap(uint64(Pointer), EMemoryTraceRootHeap::SystemMemory);
 		return Pointer;
 	}
 
 	FORCEINLINE static void Free(void* Pointer, SIZE_T Size)
 	{
-		MemoryTrace_Free(uint64(Pointer));
+		MemoryTrace_UnmarkAllocAsHeap(uint64(Pointer), EMemoryTraceRootHeap::SystemMemory);
+		MemoryTrace_Free(uint64(Pointer), EMemoryTraceRootHeap::SystemMemory);
 		FPlatformMemory::BaseAllocator()->Free(Pointer);
 	}
 };
@@ -73,18 +74,19 @@ struct FAlignedAllocatorInternal
 
 	FORCEINLINE static void* Malloc(SIZE_T Size, uint32 Alignment)
 	{
-		//LLM_SCOPED_PAUSE_TRACKING(ELLMAllocType::FMalloc);
+		void* Pointer = AnsiMalloc(Size, Alignment);
 #if DO_CHECK
-		void* Ret = AnsiMalloc(Size, Alignment);
-		check(IsAligned(Ret, Alignment));
-		return Ret;
-#else
-		return AnsiMalloc(Size, Alignment);
+		check(IsAligned(Pointer, Alignment));
 #endif
+		MemoryTrace_Alloc(uint64(Pointer), Size, Alignment, EMemoryTraceRootHeap::SystemMemory);
+		MemoryTrace_MarkAllocAsHeap(uint64(Pointer), EMemoryTraceRootHeap::SystemMemory);
+		return Pointer;
 	}
 
 	FORCEINLINE static void Free(void* Pointer, SIZE_T Size)
 	{
+		MemoryTrace_UnmarkAllocAsHeap(uint64(Pointer), EMemoryTraceRootHeap::SystemMemory);
+		MemoryTrace_Free(uint64(Pointer), EMemoryTraceRootHeap::SystemMemory);
 		return AnsiFree(Pointer);
 	}
 };
