Added explict MemoryTrace_UnmarkAllocAsHeap trace events for VirtualFree() API (to match the MemoryTrace_MarkAllocAsHeap events).

diff --git a/Engine/Source/Runtime/Core/Private/ProfilingDebugging/Microsoft/WindowsMemoryTrace.cpp b/Engine/Source/Runtime/Core/Private/ProfilingDebugging/Microsoft/WindowsMemoryTrace.cpp
--- a/Engine/Source/Runtime/Core/Private/ProfilingDebugging/Microsoft/WindowsMemoryTrace.cpp
+++ b/Engine/Source/Runtime/Core/Private/ProfilingDebugging/Microsoft/WindowsMemoryTrace.cpp
@@ -263,7 +263,7 @@ void FVirtualWinApiHooks::Initialize(bool bInLight)
 	// Note that hooking alloc functions is done last as applying the hook can
 	// allocate some memory pages.
 
-	VmFreeOrig = Editor.Hook(VirtualFree, &FVirtualWinApiHooks::VmFree); 
+	VmFreeOrig = Editor.Hook(VirtualFree, &FVirtualWinApiHooks::VmFree);
 	VmFreeExOrig = Editor.Hook(VirtualFreeEx, &FVirtualWinApiHooks::VmFreeEx);
 
 #if PLATFORM_WINDOWS
@@ -300,6 +300,7 @@ void FVirtualWinApiHooks::Initialize(bool bInLight)
 LPVOID WINAPI FVirtualWinApiHooks::VmAlloc(LPVOID Address, SIZE_T Size, DWORD Type, DWORD Protect)
 {
 	LPVOID Ret = VmAllocOrig(Address, Size, Type, Protect);
+
 	// Track any reserve for now. Going forward we need events to differentiate reserves/commits and
 	// corresponding information on frees.
 	if (Ret != nullptr && (Type & MEM_RESERVE))
@@ -316,8 +317,10 @@ BOOL WINAPI FVirtualWinApiHooks::VmFree(LPVOID Address, SIZE_T Size, DWORD Type)
 {
 	if (Type & MEM_RELEASE)
 	{
+		MemoryTrace_UnmarkAllocAsHeap((uint64)Address, EMemoryTraceRootHeap::SystemMemory);
 		MemoryTrace_Free((uint64)Address, EMemoryTraceRootHeap::SystemMemory);
 	}
+
 	return VmFreeOrig(Address, Size, Type);
 }
 
@@ -325,9 +328,10 @@ BOOL WINAPI FVirtualWinApiHooks::VmFree(LPVOID Address, SIZE_T Size, DWORD Type)
 LPVOID WINAPI FVirtualWinApiHooks::VmAllocEx(HANDLE Process, LPVOID Address, SIZE_T Size, DWORD Type, DWORD Protect)
 {
 	LPVOID Ret = VmAllocExOrig(Process, Address, Size, Type, Protect);
+
 	if (Process == GetCurrentProcess() && Ret != nullptr && (Type & MEM_RESERVE))
 	{
-		MemoryTrace_Alloc((uint64)Ret, Size, 0);
+		MemoryTrace_Alloc((uint64)Ret, Size, 0, EMemoryTraceRootHeap::SystemMemory);
 		MemoryTrace_MarkAllocAsHeap((uint64)Ret, EMemoryTraceRootHeap::SystemMemory);
 	}
 
@@ -339,6 +343,7 @@ BOOL WINAPI FVirtualWinApiHooks::VmFreeEx(HANDLE Process, LPVOID Address, SIZE_T
 {
 	if (Process == GetCurrentProcess() && (Type & MEM_RELEASE))
 	{
+		MemoryTrace_UnmarkAllocAsHeap((uint64)Address, EMemoryTraceRootHeap::SystemMemory);
 		MemoryTrace_Free((uint64)Address, EMemoryTraceRootHeap::SystemMemory);
 	}
 
@@ -349,9 +354,10 @@ BOOL WINAPI FVirtualWinApiHooks::VmFreeEx(HANDLE Process, LPVOID Address, SIZE_T
 LPVOID WINAPI FVirtualWinApiHooks::VmAlloc2(HANDLE Process, LPVOID BaseAddress, SIZE_T Size, ULONG Type, ULONG PageProtection, /*MEM_EXTENDED_PARAMETER* */ void* ExtendedParameters, ULONG ParameterCount)
 {
 	LPVOID Ret = VmAlloc2Orig(Process, BaseAddress, Size, Type, PageProtection, ExtendedParameters, ParameterCount);
+
 	if (Process == GetCurrentProcess() && Ret != nullptr && (Type & MEM_RESERVE))
 	{
-		MemoryTrace_Alloc((uint64)Ret, Size, 0);
+		MemoryTrace_Alloc((uint64)Ret, Size, 0, EMemoryTraceRootHeap::SystemMemory);
 		MemoryTrace_MarkAllocAsHeap((uint64)Ret, EMemoryTraceRootHeap::SystemMemory);
 	}
 
