Logging: Changed FOutputDeviceBlockAllocationTag to use FOsAllocator for simplicity

#rnx

diff --git a/Engine/Source/Runtime/Core/Private/Misc/OutputDeviceRedirector.cpp b/Engine/Source/Runtime/Core/Private/Misc/OutputDeviceRedirector.cpp
--- a/Engine/Source/Runtime/Core/Private/Misc/OutputDeviceRedirector.cpp
+++ b/Engine/Source/Runtime/Core/Private/Misc/OutputDeviceRedirector.cpp
@@ -43,28 +43,9 @@ namespace UE::Private
 
 struct FOutputDeviceBlockAllocationTag : FDefaultBlockAllocationTag
 {
-	static constexpr const TCHAR* TagName = TEXT("OutputDeviceLinear");
+	static constexpr const char* TagName = "OutputDeviceLinear";
 
-	struct Allocator
-	{
-		static constexpr bool SupportsAlignment = false;
-		static constexpr bool UsesFMalloc       = true;
-
-		FORCEINLINE static void* Malloc(SIZE_T Size, uint32 Alignment)
-		{
-			void* Pointer = FPlatformMemory::BaseAllocator()->Malloc(Size, DEFAULT_ALIGNMENT);
-			MemoryTrace_Alloc(uint64(Pointer), Size, DEFAULT_ALIGNMENT, EMemoryTraceRootHeap::SystemMemory);
-			MemoryTrace_MarkAllocAsHeap(uint64(Pointer), EMemoryTraceRootHeap::SystemMemory);
-			return Pointer;
-		}
-
-		FORCEINLINE static void Free(void* Pointer, SIZE_T Size)
-		{
-			MemoryTrace_UnmarkAllocAsHeap(uint64(Pointer), EMemoryTraceRootHeap::SystemMemory);
-			MemoryTrace_Free(uint64(Pointer), EMemoryTraceRootHeap::SystemMemory);
-			return FPlatformMemory::BaseAllocator()->Free(Pointer);
-		}
-	};
+	using Allocator = FOsAllocator;
 };
 
 struct FOutputDeviceLinearAllocator
