Added PlatformManualResetEvent

This private abstraction is used to signal a single waiting thread to wake.
This first implementation is only specialized for Microsoft platforms. Other platforms will be specialized later.

diff --git a/Engine/Platforms/Hololens/Source/Runtime/Core/Private/HoloLensPlatformManualResetEvent.h b/Engine/Platforms/Hololens/Source/Runtime/Core/Private/HoloLensPlatformManualResetEvent.h
new file mode 100644
--- /dev/null
+++ b/Engine/Platforms/Hololens/Source/Runtime/Core/Private/HoloLensPlatformManualResetEvent.h
@@ -0,0 +1,7 @@
+// Copyright Epic Games, Inc. All Rights Reserved.
+
+#pragma once
+
+#include "Microsoft/MicrosoftPlatformManualResetEvent.h"
+
+namespace UE { using FPlatformManualResetEvent = UE::HAL::Private::FMicrosoftPlatformManualResetEvent; }
diff --git a/Engine/Source/Runtime/Core/Private/Android/AndroidPlatformManualResetEvent.h b/Engine/Source/Runtime/Core/Private/Android/AndroidPlatformManualResetEvent.h
new file mode 100644
--- /dev/null
+++ b/Engine/Source/Runtime/Core/Private/Android/AndroidPlatformManualResetEvent.h
@@ -0,0 +1,7 @@
+// Copyright Epic Games, Inc. All Rights Reserved.
+
+#pragma once
+
+#include "GenericPlatform/GenericPlatformManualResetEvent.h"
+
+namespace UE { using FPlatformManualResetEvent = UE::HAL::Private::FGenericPlatformManualResetEvent; }
diff --git a/Engine/Source/Runtime/Core/Private/Apple/ApplePlatformManualResetEvent.h b/Engine/Source/Runtime/Core/Private/Apple/ApplePlatformManualResetEvent.h
new file mode 100644
--- /dev/null
+++ b/Engine/Source/Runtime/Core/Private/Apple/ApplePlatformManualResetEvent.h
@@ -0,0 +1,7 @@
+// Copyright Epic Games, Inc. All Rights Reserved.
+
+#pragma once
+
+#include "GenericPlatform/GenericPlatformManualResetEvent.h"
+
+namespace UE::HAL::Private { using FApplePlatformManualResetEvent = FGenericPlatformManualResetEvent; }
diff --git a/Engine/Source/Runtime/Core/Private/Async/ParkingLot.cpp b/Engine/Source/Runtime/Core/Private/Async/ParkingLot.cpp
--- a/Engine/Source/Runtime/Core/Private/Async/ParkingLot.cpp
+++ b/Engine/Source/Runtime/Core/Private/Async/ParkingLot.cpp
@@ -6,10 +6,10 @@
 #include "Async/UniqueLock.h"
 #include "Containers/Array.h"
 #include "Containers/ArrayView.h"
+#include "HAL/PlatformManualResetEvent.h"
 #include "Templates/RefCounting.h"
 #include <atomic>
 
-#define UE_PARKINGLOT_USE_WAITONADDRESS PLATFORM_WINDOWS
 #define UE_PARKINGLOT_USE_WORDMUTEX 0
 
 #if UE_PARKINGLOT_USE_WORDMUTEX
@@ -18,18 +18,11 @@
 #include "HAL/CriticalSection.h"
 #endif
 
-#if UE_PARKINGLOT_USE_WAITONADDRESS
-#include "Microsoft/WindowsHWrapper.h"
-#else
-#include <condition_variable>
-#include <mutex>
-#endif
-
 namespace UE::ParkingLot::Private
 {
 
 #if UE_PARKINGLOT_USE_WORDMUTEX
-using FBucketMutex = UE::Private::FWordMutex;
+using FBucketMutex = FWordMutex;
 #else
 using FBucketMutex = FCriticalSection;
 #endif
@@ -42,16 +35,9 @@ using FBucketMutex = FCriticalSection;
 struct FThread final
 {
 	FThread* Next = nullptr;
-
-#if UE_PARKINGLOT_USE_WAITONADDRESS
-	const void* volatile WaitAddress = nullptr;
-#else
 	const void* WaitAddress = nullptr;
-	std::mutex Lock;
-	std::condition_variable Condition;
-#endif
-
 	uint64 WakeToken = 0;
+	FPlatformManualResetEvent Event;
 	mutable std::atomic<uint32> ReferenceCount = 0;
 
 	inline void AddRef() const
@@ -540,58 +526,24 @@ FWaitState WaitUntil(const void* Address, TFunctionRef<bool()> CanWait, TFunctio
 			return State;
 		}
 		Self.WaitAddress = Address;
+		Self.Event.Reset();
 		Bucket.Enqueue(&Self);
 	}
 
 	// BeforeWait must be invoked after the bucket is unlocked.
 	BeforeWait();
 
-#if UE_PARKINGLOT_USE_WAITONADDRESS
-	for (;;)
-	{
-		const FMonotonicTimeSpan WaitSpan = WaitTime - FMonotonicTimePoint::Now();
-		if (WaitSpan <= FMonotonicTimeSpan::Zero())
-		{
-			break;
-		}
-		const DWORD WaitMs = WaitSpan.IsInfinity() ? INFINITE : DWORD(FMath::CeilToInt64(WaitSpan.ToMilliseconds()));
-		if (!WaitOnAddress(&Self.WaitAddress, &Address, sizeof(Address), WaitMs) && GetLastError() == ERROR_TIMEOUT)
-		{
-			break;
-		}
-		if (!Self.WaitAddress)
-		{
-			State.bDidWake = true;
-			State.WakeToken = Self.WakeToken;
-			Self.WakeToken = 0;
-			return State;
-		}
-	}
-#else
-	{
-		// Wait until the timeout or until the thread has been dequeued.
-		std::unique_lock SelfLock(Self.Lock);
-		const FMonotonicTimeSpan WaitSpan = WaitTime - FMonotonicTimePoint::Now();
-		if (WaitSpan.IsInfinity())
-		{
-			Self.Condition.wait(SelfLock, [&Self] { return !Self.WaitAddress; });
-		}
-		else if (WaitSpan > FMonotonicTimeSpan::Zero())
-		{
-			const int64 WaitMs = FMath::CeilToInt64(WaitSpan.ToMilliseconds());
-			Self.Condition.wait_for(SelfLock, std::chrono::milliseconds(WaitMs), [&Self] { return !Self.WaitAddress; });
-		}
+	// Wait until the timeout or until the thread has been dequeued.
+	Self.Event.WaitUntil(WaitTime);
 
-		// WaitAddress is reset when the thread is dequeued.
-		if (!Self.WaitAddress)
-		{
-			State.bDidWake = true;
-			State.WakeToken = Self.WakeToken;
-			Self.WakeToken = 0;
-			return State;
-		}
+	// WaitAddress is reset when the thread is dequeued.
+	if (!Self.WaitAddress)
+	{
+		State.bDidWake = true;
+		State.WakeToken = Self.WakeToken;
+		Self.WakeToken = 0;
+		return State;
 	}
-#endif
 
 	// The timeout was reached and the thread needs to dequeue itself.
 	// This can race with a call to wake a thread, which means Self is unsafe to access outside of the lock.
@@ -614,15 +566,7 @@ FWaitState WaitUntil(const void* Address, TFunctionRef<bool()> CanWait, TFunctio
 	// has finished waking this thread by setting its wait address to null.
 	if (!bDequeued)
 	{
-	#if UE_PARKINGLOT_USE_WAITONADDRESS
-		while (Self.WaitAddress)
-		{
-			WaitOnAddress(&Self.WaitAddress, &Address, sizeof(Address), INFINITE);
-		}
-	#else
-		std::unique_lock SelfLock(Self.Lock);
-		Self.Condition.wait(SelfLock, [&Self] { return !Self.WaitAddress; });
-	#endif
+		Self.Event.Wait();
 		State.bDidWake = true;
 		State.WakeToken = Self.WakeToken;
 		Self.WakeToken = 0;
@@ -660,16 +604,8 @@ void WakeOne(const void* Address, TFunctionRef<uint64(FWakeState)> OnWakeState)
 	{
 		checkSlow(WakeThread->WaitAddress == Address);
 		WakeThread->WakeToken = WakeToken;
-	#if UE_PARKINGLOT_USE_WAITONADDRESS
 		WakeThread->WaitAddress = nullptr;
-		WakeByAddressSingle((void*)&WakeThread->WaitAddress);
-	#else
-		{
-			std::unique_lock WakeThreadLock(WakeThread->Lock);
-			WakeThread->WaitAddress = nullptr;
-		}
-		WakeThread->Condition.notify_one();
-	#endif
+		WakeThread->Event.Signal();
 	}
 }
 
@@ -711,16 +647,8 @@ uint32 WakeMultiple(const void* const Address, const uint32 WakeCount)
 	for (FThread* WakeThread : WakeThreads)
 	{
 		checkSlow(WakeThread->WaitAddress == Address);
-	#if UE_PARKINGLOT_USE_WAITONADDRESS
 		WakeThread->WaitAddress = nullptr;
-		WakeByAddressSingle((void*)&WakeThread->WaitAddress);
-	#else
-		{
-			std::unique_lock WakeThreadLock(WakeThread->Lock);
-			WakeThread->WaitAddress = nullptr;
-		}
-		WakeThread->Condition.notify_one();
-	#endif
+		WakeThread->Event.Signal();
 	}
 
 	return uint32(WakeThreads.Num());
diff --git a/Engine/Source/Runtime/Core/Private/GenericPlatform/GenericPlatformManualResetEvent.h b/Engine/Source/Runtime/Core/Private/GenericPlatform/GenericPlatformManualResetEvent.h
new file mode 100644
--- /dev/null
+++ b/Engine/Source/Runtime/Core/Private/GenericPlatform/GenericPlatformManualResetEvent.h
@@ -0,0 +1,95 @@
+// Copyright Epic Games, Inc. All Rights Reserved.
+
+#pragma once
+
+#include "HAL/PlatformMath.h"
+#include "Misc/MonotonicTime.h"
+
+#include <atomic>
+#include <chrono>
+#include <condition_variable>
+#include <mutex>
+
+namespace UE::HAL::Private
+{
+
+/**
+ * A manual reset event that supports only one thread waiting and one thread signaling at a time.
+ *
+ * Only one waiting thread may call Reset() or the Wait() functions.
+ * Only one signaling thread may call Signal() once until the event is reset.
+ */
+class FGenericPlatformManualResetEvent
+{
+public:
+	FGenericPlatformManualResetEvent() = default;
+	FGenericPlatformManualResetEvent(const FGenericPlatformManualResetEvent&) = delete;
+	FGenericPlatformManualResetEvent& operator=(const FGenericPlatformManualResetEvent&) = delete;
+
+	/**
+	 * Resets the event to permit another Wait/Signal cycle.
+	 *
+	 * Must only be called by the waiting thread, and only when there is no possibility of waking
+	 * occurring concurrently with the reset.
+	 */
+	void Reset()
+	{
+		bWait = true;
+	}
+
+	/**
+	 * Waits for Signal() to be called.
+	 *
+	 * Signal() may be called prior to Wait(), and this will return immediately in that case.
+	 */
+	void Wait()
+	{
+		WaitUntil(FMonotonicTimePoint::Infinity());
+	}
+
+	/**
+	 * Waits until the wait time for Signal() to be called.
+	 *
+	 * Signal() may be called prior to WaitUntil(), and this will return immediately in that case.
+	 *
+	 * @param WaitTime   Absolute time after which waiting is canceled and the thread wakes.
+	 * @return True if Signal() was called before the wait time elapsed, otherwise false.
+	 */
+	bool WaitUntil(FMonotonicTimePoint WaitTime)
+	{
+		std::unique_lock SelfLock(Lock);
+		if (WaitTime.IsInfinity())
+		{
+			Condition.wait(SelfLock, [this] { return !bWait; });
+			return true;
+		}
+		if (FMonotonicTimeSpan WaitSpan = WaitTime - FMonotonicTimePoint::Now(); WaitSpan > FMonotonicTimeSpan::Zero())
+		{
+			const int64 WaitMs = FPlatformMath::CeilToInt64(WaitSpan.ToMilliseconds());
+			return Condition.wait_for(SelfLock, std::chrono::milliseconds(WaitMs), [this] { return !bWait; });
+		}
+		return !bWait;
+	}
+
+	/**
+	 * Signals the waiting thread.
+	 *
+	 * Signal() may be called prior to one of the wait functions, and the eventual wait call will
+	 * return immediately when that occurs.
+	 */
+	void Signal()
+	{
+		{
+			std::unique_lock SelfLock(Lock);
+			bWait = false;
+		}
+		Condition.notify_one();
+	}
+
+private:
+	std::mutex Lock;
+	std::condition_variable Condition;
+	std::atomic<bool> bWait = true;
+};
+
+} // UE::HAL::Private
diff --git a/Engine/Source/Runtime/Core/Private/HAL/PlatformManualResetEvent.h b/Engine/Source/Runtime/Core/Private/HAL/PlatformManualResetEvent.h
new file mode 100644
--- /dev/null
+++ b/Engine/Source/Runtime/Core/Private/HAL/PlatformManualResetEvent.h
@@ -0,0 +1,7 @@
+// Copyright Epic Games, Inc. All Rights Reserved.
+
+#pragma once
+
+#include "CoreTypes.h"
+
+#include COMPILED_PLATFORM_HEADER(PlatformManualResetEvent.h)
diff --git a/Engine/Source/Runtime/Core/Private/IOS/IOSPlatformManualResetEvent.h b/Engine/Source/Runtime/Core/Private/IOS/IOSPlatformManualResetEvent.h
new file mode 100644
--- /dev/null
+++ b/Engine/Source/Runtime/Core/Private/IOS/IOSPlatformManualResetEvent.h
@@ -0,0 +1,7 @@
+// Copyright Epic Games, Inc. All Rights Reserved.
+
+#pragma once
+
+#include "Apple/ApplePlatformManualResetEvent.h"
+
+namespace UE { using FPlatformManualResetEvent = UE::HAL::Private::FApplePlatformManualResetEvent; }
diff --git a/Engine/Source/Runtime/Core/Private/Linux/LinuxPlatformManualResetEvent.h b/Engine/Source/Runtime/Core/Private/Linux/LinuxPlatformManualResetEvent.h
new file mode 100644
--- /dev/null
+++ b/Engine/Source/Runtime/Core/Private/Linux/LinuxPlatformManualResetEvent.h
@@ -0,0 +1,7 @@
+// Copyright Epic Games, Inc. All Rights Reserved.
+
+#pragma once
+
+#include "Unix/UnixPlatformManualResetEvent.h"
+
+namespace UE { using FPlatformManualResetEvent = UE::HAL::Private::FUnixPlatformManualResetEvent; }
diff --git a/Engine/Source/Runtime/Core/Private/Mac/MacPlatformManualResetEvent.h b/Engine/Source/Runtime/Core/Private/Mac/MacPlatformManualResetEvent.h
new file mode 100644
--- /dev/null
+++ b/Engine/Source/Runtime/Core/Private/Mac/MacPlatformManualResetEvent.h
@@ -0,0 +1,7 @@
+// Copyright Epic Games, Inc. All Rights Reserved.
+
+#pragma once
+
+#include "Apple/ApplePlatformManualResetEvent.h"
+
+namespace UE { using FPlatformManualResetEvent = UE::HAL::Private::FApplePlatformManualResetEvent; }
diff --git a/Engine/Source/Runtime/Core/Private/Microsoft/MicrosoftPlatformManualResetEvent.cpp b/Engine/Source/Runtime/Core/Private/Microsoft/MicrosoftPlatformManualResetEvent.cpp
new file mode 100644
--- /dev/null
+++ b/Engine/Source/Runtime/Core/Private/Microsoft/MicrosoftPlatformManualResetEvent.cpp
@@ -0,0 +1,55 @@
+// Copyright Epic Games, Inc. All Rights Reserved.
+
+#include "Microsoft/MicrosoftPlatformManualResetEvent.h"
+
+#include "HAL/PlatformMath.h"
+#include "Misc/MonotonicTime.h"
+
+#include "Microsoft/WindowsHWrapper.h"
+
+namespace UE::HAL::Private
+{
+
+void FMicrosoftPlatformManualResetEvent::Wait()
+{
+	WaitUntil(FMonotonicTimePoint::Infinity());
+}
+
+bool FMicrosoftPlatformManualResetEvent::WaitUntil(FMonotonicTimePoint WaitTime)
+{
+	bool bLocalWait = true;
+	if (WaitTime.IsInfinity())
+	{
+		for (;;)
+		{
+			if (WaitOnAddress(&bWait, &bLocalWait, sizeof(bool), INFINITE) && !bWait)
+			{
+				return true;
+			}
+		}
+	}
+	else
+	{
+		for (;;)
+		{
+			FMonotonicTimeSpan WaitSpan = WaitTime - FMonotonicTimePoint::Now();
+			if (WaitSpan <= FMonotonicTimeSpan::Zero())
+			{
+				return false;
+			}
+			const DWORD WaitMs = DWORD(FPlatformMath::CeilToInt64(WaitSpan.ToMilliseconds()));
+			if (WaitOnAddress(&bWait, &bLocalWait, sizeof(bool), WaitMs) && !bWait)
+			{
+				return true;
+			}
+		}
+	}
+}
+
+void FMicrosoftPlatformManualResetEvent::Signal()
+{
+	bWait = false;
+	WakeByAddressSingle((void*)&bWait);
+}
+
+} // UE::HAL::Private
diff --git a/Engine/Source/Runtime/Core/Private/Microsoft/MicrosoftPlatformManualResetEvent.h b/Engine/Source/Runtime/Core/Private/Microsoft/MicrosoftPlatformManualResetEvent.h
new file mode 100644
--- /dev/null
+++ b/Engine/Source/Runtime/Core/Private/Microsoft/MicrosoftPlatformManualResetEvent.h
@@ -0,0 +1,34 @@
+// Copyright Epic Games, Inc. All Rights Reserved.
+
+#pragma once
+
+#include "CoreTypes.h"
+
+#define UE_API CORE_API
+
+namespace UE { struct FMonotonicTimePoint; }
+
+namespace UE::HAL::Private
+{
+
+/** @see FGenericPlatformManualResetEvent */
+class FMicrosoftPlatformManualResetEvent
+{
+public:
+	FMicrosoftPlatformManualResetEvent() = default;
+	FMicrosoftPlatformManualResetEvent(const FMicrosoftPlatformManualResetEvent&) = delete;
+	FMicrosoftPlatformManualResetEvent& operator=(const FMicrosoftPlatformManualResetEvent&) = delete;
+
+	void Reset() { bWait = true; }
+
+	UE_API void Wait();
+	UE_API bool WaitUntil(FMonotonicTimePoint WaitTime);
+	UE_API void Signal();
+
+private:
+	bool bWait = true;
+};
+
+} // UE::HAL::Private
+
+#undef UE_API
diff --git a/Engine/Source/Runtime/Core/Private/Unix/UnixPlatformManualResetEvent.h b/Engine/Source/Runtime/Core/Private/Unix/UnixPlatformManualResetEvent.h
new file mode 100644
--- /dev/null
+++ b/Engine/Source/Runtime/Core/Private/Unix/UnixPlatformManualResetEvent.h
@@ -0,0 +1,7 @@
+// Copyright Epic Games, Inc. All Rights Reserved.
+
+#pragma once
+
+#include "GenericPlatform/GenericPlatformManualResetEvent.h"
+
+namespace UE::HAL::Private { using FUnixPlatformManualResetEvent = FGenericPlatformManualResetEvent; }
diff --git a/Engine/Source/Runtime/Core/Private/Windows/WindowsPlatformManualResetEvent.h b/Engine/Source/Runtime/Core/Private/Windows/WindowsPlatformManualResetEvent.h
new file mode 100644
--- /dev/null
+++ b/Engine/Source/Runtime/Core/Private/Windows/WindowsPlatformManualResetEvent.h
@@ -0,0 +1,7 @@
+// Copyright Epic Games, Inc. All Rights Reserved.
+
+#pragma once
+
+#include "Microsoft/MicrosoftPlatformManualResetEvent.h"
+
+namespace UE { using FPlatformManualResetEvent = UE::HAL::Private::FMicrosoftPlatformManualResetEvent; }
