Fix race in FGenericPlatformTime::StrTimestamp()

#rnx

diff --git a/Engine/Source/Runtime/Core/Private/GenericPlatform/GenericPlatformTime.cpp b/Engine/Source/Runtime/Core/Private/GenericPlatform/GenericPlatformTime.cpp
--- a/Engine/Source/Runtime/Core/Private/GenericPlatform/GenericPlatformTime.cpp
+++ b/Engine/Source/Runtime/Core/Private/GenericPlatform/GenericPlatformTime.cpp
@@ -111,7 +111,7 @@ TCHAR* FGenericPlatformTime::StrTime( TCHAR* Dest, SIZE_T DestSize )
 
 const TCHAR* FGenericPlatformTime::StrTimestamp()
 {
-	static TCHAR Result[1024];
+	static thread_local TCHAR Result[1024];
 	*Result = TEXT('\0');
 	StrDate( Result, UE_ARRAY_COUNT(Result) );
 	FCString::Strcat( Result, TEXT(" ") );
