#UE Enable MallocStomp2 in editor builds

diff --git a/Engine/Source/Runtime/Core/Public/HAL/MallocStomp2.h b/Engine/Source/Runtime/Core/Public/HAL/MallocStomp2.h
--- a/Engine/Source/Runtime/Core/Public/HAL/MallocStomp2.h
+++ b/Engine/Source/Runtime/Core/Public/HAL/MallocStomp2.h
@@ -6,7 +6,7 @@
 #include "HAL/PlatformTLS.h"
 
 #ifndef WITH_MALLOC_STOMP2
-#define WITH_MALLOC_STOMP2 0
+#define WITH_MALLOC_STOMP2 WITH_EDITOR
 #endif
 #if WITH_MALLOC_STOMP2
 #include <atomic>
