removed delegates' old deprecations

diff --git a/Engine/Source/Runtime/Core/Public/Delegates/Delegate.h b/Engine/Source/Runtime/Core/Public/Delegates/Delegate.h
--- a/Engine/Source/Runtime/Core/Public/Delegates/Delegate.h
+++ b/Engine/Source/Runtime/Core/Public/Delegates/Delegate.h
@@ -502,11 +502,3 @@ DECLARE_DELEGATE( FSimpleDelegate );
 DECLARE_MULTICAST_DELEGATE( FSimpleMulticastDelegate );
 DECLARE_TS_MULTICAST_DELEGATE( FTSSimpleMulticastDelegate );
 
-// Legacy typedefs
-template <typename RetType, typename... ArgTypes>
-using TBaseDelegate UE_DEPRECATED(4.26, "TBaseDelegate<ReturnType, ArgTypes...> is deprecated - use TDelegate<ReturnType(ArgTypes...)> instead.")
-	= TDelegate<RetType(ArgTypes...)>;
-
-template <typename RetType, typename... ArgTypes>
-using TBaseMulticastDelegate UE_DEPRECATED(4.26, "TBaseMulticastDelegate<ReturnType, ArgTypes...> is deprecated - use TMulticastDelegate<ReturnType(ArgTypes...)> instead.")
-	= TMulticastDelegate<RetType(ArgTypes...)>;
