Added FWordMutex

This is a mutex that fits in a machine word. It is not fair and does not support recursive locking.

#rnx

diff --git a/Engine/Source/Runtime/Core/Private/Async/ParkingLot.cpp b/Engine/Source/Runtime/Core/Private/Async/ParkingLot.cpp
--- a/Engine/Source/Runtime/Core/Private/Async/ParkingLot.cpp
+++ b/Engine/Source/Runtime/Core/Private/Async/ParkingLot.cpp
@@ -10,7 +10,7 @@
 #include "Templates/RefCounting.h"
 #include <atomic>
 
-#define UE_PARKINGLOT_USE_WORDMUTEX 0
+#define UE_PARKINGLOT_USE_WORDMUTEX 1
 
 #if UE_PARKINGLOT_USE_WORDMUTEX
 #include "Async/WordMutex.h"
diff --git a/Engine/Source/Runtime/Core/Private/Async/WordMutex.cpp b/Engine/Source/Runtime/Core/Private/Async/WordMutex.cpp
new file mode 100644
--- /dev/null
+++ b/Engine/Source/Runtime/Core/Private/Async/WordMutex.cpp
@@ -0,0 +1,168 @@
+// Copyright Epic Games, Inc. All Rights Reserved.
+
+#include "Async/WordMutex.h"
+
+#include "HAL/PlatformProcess.h"
+#include "HAL/PlatformManualResetEvent.h"
+
+namespace UE
+{
+
+struct FWordMutexQueueNode
+{
+	// Points to the next node in the tail-to-head direction. Only null for the current tail.
+	FWordMutexQueueNode* Prev = nullptr;
+	// Points to the next node in the head-to-tail direction.
+	// Null until UnlockSlow() has traversed from the tail to fill in next pointers.
+	FWordMutexQueueNode* Next = nullptr;
+	// Any node that has been the tail may have a non-null head pointer, but only the most recent
+	// node with a non-null head pointer, when traversing from the tail, has a valid pointer.
+	FWordMutexQueueNode* Head = nullptr;
+
+	FPlatformManualResetEvent Event;
+};
+
+void FWordMutex::LockSlow()
+{
+	static_assert((alignof(FWordMutexQueueNode) & QueueMask) == alignof(FWordMutexQueueNode),
+		"Alignment of FWordMutexQueueNode is insufficient to pack flags into the lower bits.");
+
+	constexpr int32 SpinLimit = 40;
+	int32 SpinCount = 0;
+	for (;;)
+	{
+		UPTRINT CurrentState = State.load(std::memory_order_relaxed);
+
+		// Try to acquire the lock if it was unlocked, even if there is a queue.
+		// Acquiring the lock despite the queue means that this lock is not FIFO and thus not fair.
+		if (!(CurrentState & IsLockedFlag))
+		{
+			if (State.compare_exchange_weak(CurrentState, CurrentState | IsLockedFlag, std::memory_order_acquire, std::memory_order_relaxed))
+			{
+				return;
+			}
+			continue;
+		}
+
+		// Spin up to the spin limit while there is no queue.
+		if (!(CurrentState & QueueMask) && SpinCount < SpinLimit)
+		{
+			FPlatformProcess::YieldThread();
+			++SpinCount;
+			continue;
+		}
+
+		// Create the node that will be used to add this thread to the queue.
+		FWordMutexQueueNode Self;
+
+		Self.Event.Reset();
+
+		// The state points to the tail of the queue, and each node points to the previous node.
+		if (FWordMutexQueueNode* Tail = (FWordMutexQueueNode*)(CurrentState & QueueMask))
+		{
+			Self.Prev = Tail;
+		}
+		else
+		{
+			Self.Head = &Self;
+		}
+
+		// Swap this thread in as the tail, which makes it visible to any other thread that acquires the queue lock.
+		if (!State.compare_exchange_weak(CurrentState, (CurrentState & ~QueueMask) | UPTRINT(&Self), std::memory_order_acq_rel, std::memory_order_relaxed))
+		{
+			continue;
+		}
+
+		// Wait until another thread wakes this thread, which can happen as soon as the preceding store completes.
+		Self.Event.Wait();
+
+		// Loop back and try to acquire the lock.
+		SpinCount = 0;
+	}
+}
+
+void FWordMutex::UnlockSlow(UPTRINT CurrentState)
+{
+	// IsLockedFlag was cleared by Unlock().
+	CurrentState &= ~IsLockedFlag;
+
+	for (;;)
+	{
+		// Try to lock the queue.
+		if (State.compare_exchange_weak(CurrentState, CurrentState | IsQueueLockedFlag, std::memory_order_acquire, std::memory_order_relaxed))
+		{
+			CurrentState |= IsQueueLockedFlag;
+			break;
+		}
+
+		// A locked queue indicates that another thread is looking for a thread to wake.
+		if ((CurrentState & IsQueueLockedFlag) || !(CurrentState & QueueMask))
+		{
+			return;
+		}
+	}
+
+	FWordMutexQueueNode* Head;
+
+	for (;;)
+	{
+		// This thread now holds the queue lock. Neither the queue nor State will change while the queue is locked.
+		// The state points to the tail of the queue, and each node points to the previous node.
+		FWordMutexQueueNode* Tail = (FWordMutexQueueNode*)(CurrentState & QueueMask);
+
+		// Traverse from the tail to find the head and set next pointers for any nodes added since the last unlock.
+		for (FWordMutexQueueNode* Node = Tail;;)
+		{
+			Head = Node->Head;
+			if (Head)
+			{
+				break;
+			}
+			FWordMutexQueueNode* Prev = Node->Prev;
+			checkSlow(Prev);
+			Prev->Next = Node;
+			Node = Prev;
+		}
+
+		// Another thread may acquire the lock while this thread has been finding a thread to unlock.
+		// That case will not be detected on the first iteration of the loop, but only when this
+		// thread has failed to unlock the queue at least once. Attempt to unlock the queue here
+		// and allow the next unlock to find a thread to wake.
+		if (CurrentState & IsLockedFlag)
+		{
+			if (State.compare_exchange_weak(CurrentState, CurrentState & ~IsQueueLockedFlag, std::memory_order_release, std::memory_order_acquire))
+			{
+				return;
+			}
+			continue;
+		}
+
+		// Remove the head from the queue and unlock the queue.
+		if (FWordMutexQueueNode* NewHead = Head->Next)
+		{
+			// Clear the pointer to the node being removed.
+			NewHead->Prev = nullptr;
+
+			// Store the head in the tail, which shows that the next pointers have been set in the whole list.
+			Tail->Head = NewHead;
+
+			// Unlock the queue regardless of whether new nodes have been added in the meantime.
+			State.fetch_and(~IsQueueLockedFlag, std::memory_order_release);
+		}
+		else
+		{
+			// Unlock and clear the queue. Failure needs to restart the loop, because newly-added
+			// nodes will have a pointer to the node being removed.
+			if (!State.compare_exchange_strong(CurrentState, CurrentState & IsLockedFlag, std::memory_order_release, std::memory_order_acquire))
+			{
+				continue;
+			}
+		}
+
+		// Wake the thread that was at the head of the queue.
+		Head->Event.Signal();
+		break;
+	}
+}
+
+} // UE
diff --git a/Engine/Source/Runtime/Core/Public/Async/WordMutex.h b/Engine/Source/Runtime/Core/Public/Async/WordMutex.h
new file mode 100644
--- /dev/null
+++ b/Engine/Source/Runtime/Core/Public/Async/WordMutex.h
@@ -0,0 +1,77 @@
+// Copyright Epic Games, Inc. All Rights Reserved.
+
+#pragma once
+
+#include "CoreTypes.h"
+#include "Misc/AssertionMacros.h"
+#include <atomic>
+
+#define UE_API CORE_API
+
+namespace UE
+{
+
+/**
+ * A mutex that is the size of a pointer and does not depend on ParkingLot.
+ *
+ * Prefer FMutex to FWordMutex whenever possible.
+ * This mutex is not fair and does not support recursive locking.
+ */
+class FWordMutex final
+{
+public:
+	constexpr FWordMutex() = default;
+
+	FWordMutex(const FWordMutex&) = delete;
+	FWordMutex& operator=(const FWordMutex&) = delete;
+
+	inline bool TryLock()
+	{
+		UPTRINT Expected = 0;
+		return State.compare_exchange_strong(Expected, IsLockedFlag, std::memory_order_acquire, std::memory_order_relaxed);
+	}
+
+	inline void Lock()
+	{
+		UPTRINT Expected = 0;
+		if (LIKELY(State.compare_exchange_weak(Expected, IsLockedFlag, std::memory_order_acquire, std::memory_order_relaxed)))
+		{
+			return;
+		}
+
+		LockSlow();
+	}
+
+	inline void Unlock()
+	{
+		// Unlock immediately to allow other threads to acquire the lock while this thread looks for a thread to wake.
+		UPTRINT CurrentState = State.fetch_sub(IsLockedFlag, std::memory_order_release);
+		checkSlow(CurrentState & IsLockedFlag);
+
+		// An empty queue indicates that there are no threads to wake.
+		const bool bQueueEmpty = !(CurrentState & QueueMask);
+		// A locked queue indicates that another thread is looking for a thread to wake.
+		const bool bQueueLocked = (CurrentState & IsQueueLockedFlag);
+
+		if (LIKELY(bQueueEmpty || bQueueLocked))
+		{
+			return;
+		}
+
+		UnlockSlow(CurrentState);
+	}
+
+private:
+	UE_API void LockSlow();
+	UE_API void UnlockSlow(UPTRINT CurrentState);
+
+	static constexpr UPTRINT IsLockedFlag = 1 << 0;
+	static constexpr UPTRINT IsQueueLockedFlag = 1 << 1;
+	static constexpr UPTRINT QueueMask = ~(IsLockedFlag | IsQueueLockedFlag);
+
+	std::atomic<UPTRINT> State = 0;
+};
+
+} // UE
+
+#undef UE_API
diff --git a/Engine/Source/Runtime/Core/Tests/Async/WordMutexTest.cpp b/Engine/Source/Runtime/Core/Tests/Async/WordMutexTest.cpp
new file mode 100644
--- /dev/null
+++ b/Engine/Source/Runtime/Core/Tests/Async/WordMutexTest.cpp
@@ -0,0 +1,69 @@
+// Copyright Epic Games, Inc. All Rights Reserved.
+
+#if WITH_LOW_LEVEL_TESTS
+
+#include "Async/WordMutex.h"
+
+#include "HAL/Thread.h"
+#include "TestHarness.h"
+
+namespace UE
+{
+
+TEST_CASE("Core::Async::WordMutex", "[Core][Async]")
+{
+	SECTION("FIFO")
+	{
+		constexpr static int32 TaskCount = 5;
+		FThread Threads[TaskCount];
+		std::atomic<int32> WaitCount = 0;
+		int32 WakeCount = 0;
+		int32 WakeStates[TaskCount];
+
+		FWordMutex Mutex;
+		Mutex.Lock();
+
+		// Launch tasks that wait on the address of WaitCount.
+		for (int32 Index = 0; Index < TaskCount; ++Index)
+		{
+			// Using FThread for now because UE::Tasks::Launch does not always wake a worker thread.
+			Threads[Index] = FThread(TEXT("WordMutexTest"), [&Mutex, &WaitCount, &WakeCount, OutState = &WakeStates[Index]]
+			{
+				WaitCount.fetch_add(1);
+				Mutex.Lock();
+				*OutState = WakeCount++;
+				Mutex.Unlock();
+				WaitCount.fetch_sub(1);
+			});
+
+			// Spin until the task is about to lock.
+			while (WaitCount != Index + 1)
+			{
+			}
+		}
+
+		// Unlock to allow each task to lock in sequence.
+		Mutex.Unlock();
+
+		// Spin until the tasks are complete.
+		while (WaitCount != 0)
+		{
+		}
+
+		// Verify that tasks woke in FIFO order.
+		for (int32 Index = 0; Index < TaskCount; ++Index)
+		{
+			CHECK(WakeStates[Index] == Index);
+		}
+
+		// Wait for the threads to exit.
+		for (FThread& Thread : Threads)
+		{
+			Thread.Join();
+		}
+	}
+}
+
+} // UE
+
+#endif // WITH_LOW_LEVEL_TESTS
