CompactBinary: Optimized FCbWriter by fixing an unintentional array shrink

#rnx

diff --git a/Engine/Source/Runtime/Core/Private/Serialization/CompactBinaryWriter.cpp b/Engine/Source/Runtime/Core/Private/Serialization/CompactBinaryWriter.cpp
--- a/Engine/Source/Runtime/Core/Private/Serialization/CompactBinaryWriter.cpp
+++ b/Engine/Source/Runtime/Core/Private/Serialization/CompactBinaryWriter.cpp
@@ -235,7 +235,7 @@ void FCbWriter::MakeFieldsUniform(const int64 FieldBeginOffset, const int64 Fiel
 	}
 	if (!TargetView.IsEmpty())
 	{
-		Data.RemoveAt(FieldEndOffset - TargetView.GetSize(), TargetView.GetSize());
+		Data.RemoveAt(FieldEndOffset - TargetView.GetSize(), TargetView.GetSize(), /*bAllowShrinking*/ false);
 	}
 }
 
