Fix for race condition in Linux FileMapCache

diff --git a/Engine/Source/Runtime/Core/Private/Unix/UnixPlatformFile.cpp b/Engine/Source/Runtime/Core/Private/Unix/UnixPlatformFile.cpp
--- a/Engine/Source/Runtime/Core/Private/Unix/UnixPlatformFile.cpp
+++ b/Engine/Source/Runtime/Core/Private/Unix/UnixPlatformFile.cpp
@@ -360,6 +360,8 @@ namespace
 		virtual const FileEntry* Find(const FString& Key) = 0;
 		virtual void AddEntry(const FString& Key, const FString& Elem) = 0;
 		virtual void Invalidate(const FString& Key) = 0;
+		virtual void Lock() = 0;
+		virtual void Unlock() = 0;
 	};
 
 	class FileMapCacheDummy : public FileMapCache
@@ -375,6 +377,12 @@ namespace
 
 		void Invalidate(const FString& Key) override
 		{ }
+
+		void Lock() override
+		{ }
+
+		void Unlock() override
+		{ }
 	};
 
 	class FileMapCacheDefault : public FileMapCache
@@ -387,22 +395,29 @@ namespace
 
 		const FileEntry* Find(const FString& Key) override
 		{
-			FScopeLock ScopeLock(&Mutex);
 			return Cache.FindAndTouch(Key);
 		}
 
 		void AddEntry(const FString& Key, const FString& Elem) override
 		{
-			FScopeLock ScopeLock(&Mutex);
 			Cache.Add(Key, {Elem, Elem.IsEmpty(), FPlatformTime::Seconds()});
 		}
 
 		void Invalidate(const FString& Key) override
 		{
-			FScopeLock ScopeLock(&Mutex);
 			Cache.Remove(Key);
 		}
 
+		void Lock() override
+		{
+			Mutex.Lock();
+		}
+
+		void Unlock() override
+		{
+			Mutex.Unlock();
+		}
+
 	private:
 		FCriticalSection Mutex;
 		TLruCache<FString, FileEntry> Cache;
@@ -596,6 +611,7 @@ public:
 		else
 		{
 			FileMapCache& Cache = GetFileMapCache();
+			UE::TScopeLock Lock(Cache);
 			const FileEntry* Entry = Cache.Find(PossiblyWrongFilename);
 
 			if (Entry != nullptr)
@@ -872,7 +888,11 @@ bool FUnixPlatformFile::DeleteFile(const TCHAR* Filename)
 		return false;
 	}
 
-	GetFileMapCache().Invalidate(IntendedFilename);
+	{
+		FileMapCache& Cache = GetFileMapCache();
+		UE::TScopeLock Lock(Cache);
+		Cache.Invalidate(IntendedFilename);
+	}
 
 	// removing mapped file is too dangerous
 	if (IntendedFilename != CaseSensitiveFilename)
@@ -910,7 +930,11 @@ bool FUnixPlatformFile::MoveFile(const TCHAR* To, const TCHAR* From)
 		return false;
 	}
 
-	GetFileMapCache().Invalidate(IntendedFilename);
+	{
+		FileMapCache& Cache = GetFileMapCache();
+		UE::TScopeLock Lock(Cache);
+		Cache.Invalidate(IntendedFilename);
+	}
 
 	int32 Result = rename(TCHAR_TO_UTF8(*CaseSensitiveFilename), TCHAR_TO_UTF8(*NormalizeFilename(To, true)));
 	if (Result == -1 && errno == EXDEV)
@@ -1079,7 +1103,11 @@ IFileHandle* FUnixPlatformFile::OpenWrite(const TCHAR* Filename, bool bAppend, b
 	}
 
 	// We may have cached this as an invalid file, so lets just remove a newly created file from the cache
-	GetFileMapCache().Invalidate(FString(Filename));
+	{
+		FileMapCache& Cache = GetFileMapCache();
+		UE::TScopeLock Lock(Cache);
+		Cache.Invalidate(FString(Filename));
+	}
 
 	// create directories if needed.
 	if (!CreateDirectoriesFromPath(Filename))
@@ -1210,7 +1238,11 @@ bool FUnixPlatformFile::DeleteDirectory(const TCHAR* Directory)
 	CaseSensitiveFilename = IntendedFilename;
 #endif
 
-	GetFileMapCache().Invalidate(IntendedFilename);
+	{
+		FileMapCache& Cache = GetFileMapCache();
+		UE::TScopeLock Lock(Cache);
+		GetFileMapCache().Invalidate(IntendedFilename);
+	}
 
 	// removing mapped directory is too dangerous
 	if (IntendedFilename != CaseSensitiveFilename)
