Optimized WordMutex by removing its head pointer in favor of the tail pointing to the head

#rnx

diff --git a/Engine/Source/Runtime/Core/Private/Async/WordMutex.cpp b/Engine/Source/Runtime/Core/Private/Async/WordMutex.cpp
--- a/Engine/Source/Runtime/Core/Private/Async/WordMutex.cpp
+++ b/Engine/Source/Runtime/Core/Private/Async/WordMutex.cpp
@@ -12,12 +12,9 @@ struct FWordMutexQueueNode
 {
 	// Points to the next node in the tail-to-head direction. Only null for the current tail.
 	FWordMutexQueueNode* Prev = nullptr;
-	// Points to the next node in the head-to-tail direction.
+	// Points to the next node in the head-to-tail direction. The tail points to the head.
 	// Null until UnlockSlow() has traversed from the tail to fill in next pointers.
 	FWordMutexQueueNode* Next = nullptr;
-	// Any node that has been the tail may have a non-null head pointer, but only the most recent
-	// node with a non-null head pointer, when traversing from the tail, has a valid pointer.
-	FWordMutexQueueNode* Head = nullptr;
 
 	FPlatformManualResetEvent Event;
 };
@@ -64,7 +61,7 @@ void FWordMutex::LockSlow()
 		}
 		else
 		{
-			Self.Head = &Self;
+			Self.Next = &Self;
 		}
 
 		// Swap this thread in as the tail, which makes it visible to any other thread that acquires the queue lock.
@@ -102,8 +99,6 @@ void FWordMutex::UnlockSlow(UPTRINT CurrentState)
 		}
 	}
 
-	FWordMutexQueueNode* Head;
-
 	for (;;)
 	{
 		// This thread now holds the queue lock. Neither the queue nor State will change while the queue is locked.
@@ -111,15 +106,11 @@ void FWordMutex::UnlockSlow(UPTRINT CurrentState)
 		FWordMutexQueueNode* Tail = (FWordMutexQueueNode*)(CurrentState & QueueMask);
 
 		// Traverse from the tail to find the head and set next pointers for any nodes added since the last unlock.
-		for (FWordMutexQueueNode* Node = Tail;;)
+		for (FWordMutexQueueNode* Node = Tail; !Tail->Next;)
 		{
-			Head = Node->Head;
-			if (Head)
-			{
-				break;
-			}
 			FWordMutexQueueNode* Prev = Node->Prev;
 			checkSlow(Prev);
+			Tail->Next = Prev->Next;
 			Prev->Next = Node;
 			Node = Prev;
 		}
@@ -137,19 +128,11 @@ void FWordMutex::UnlockSlow(UPTRINT CurrentState)
 			continue;
 		}
 
-		// Remove the head from the queue and unlock the queue.
-		if (FWordMutexQueueNode* NewHead = Head->Next)
-		{
-			// Clear the pointer to the node being removed.
-			NewHead->Prev = nullptr;
-
-			// Store the head in the tail, which shows that the next pointers have been set in the whole list.
-			Tail->Head = NewHead;
+		// The next node from the tail is the head.
+		FWordMutexQueueNode* Head = Tail->Next;
 
-			// Unlock the queue regardless of whether new nodes have been added in the meantime.
-			State.fetch_and(~IsQueueLockedFlag, std::memory_order_release);
-		}
-		else
+		// Remove the head from the queue and unlock the queue.
+		if (FWordMutexQueueNode* NewHead = Head->Next; NewHead == Head)
 		{
 			// Unlock and clear the queue. Failure needs to restart the loop, because newly-added
 			// nodes will have a pointer to the node being removed.
@@ -158,6 +141,16 @@ void FWordMutex::UnlockSlow(UPTRINT CurrentState)
 				continue;
 			}
 		}
+		else
+		{
+			// Clear pointers to the head node being removed.
+			checkSlow(NewHead);
+			NewHead->Prev = nullptr;
+			Tail->Next = NewHead;
+
+			// Unlock the queue regardless of whether new nodes have been added in the meantime.
+			State.fetch_and(~IsQueueLockedFlag, std::memory_order_release);
+		}
 
 		// Wake the thread that was at the head of the queue.
 		Head->Event.Signal();
