Add CrashContext in crashdumps

#rnx

diff --git a/Engine/Source/Runtime/Core/Private/GenericPlatform/GenericPlatformCrashContext.cpp b/Engine/Source/Runtime/Core/Private/GenericPlatform/GenericPlatformCrashContext.cpp
--- a/Engine/Source/Runtime/Core/Private/GenericPlatform/GenericPlatformCrashContext.cpp
+++ b/Engine/Source/Runtime/Core/Private/GenericPlatform/GenericPlatformCrashContext.cpp
@@ -1132,6 +1132,11 @@ void FGenericCrashContext::SetEngineData(const FString& Key, const FString& Valu
 	}
 }
 
+const TMap<FString, FString>& FGenericCrashContext::GetEngineData()
+{
+	return NCached::EngineData;
+}
+
 /** Get arbitrary engine data from the crash context */
 const FString* FGenericCrashContext::GetEngineData(const FString& Key)
 {
@@ -1171,6 +1176,11 @@ void FGenericCrashContext::SetGameData(const FString& Key, const FString& Value)
 	}
 }
 
+const TMap<FString, FString>& FGenericCrashContext::GetGameData()
+{
+	return NCached::GameData;
+}
+
 /** Get arbitrary game data from the crash context */
 const FString* FGenericCrashContext::GetGameData(const FString& Key)
 {
diff --git a/Engine/Source/Runtime/Core/Public/GenericPlatform/GenericPlatformCrashContext.h b/Engine/Source/Runtime/Core/Public/GenericPlatform/GenericPlatformCrashContext.h
--- a/Engine/Source/Runtime/Core/Public/GenericPlatform/GenericPlatformCrashContext.h
+++ b/Engine/Source/Runtime/Core/Public/GenericPlatform/GenericPlatformCrashContext.h
@@ -440,12 +440,18 @@ public:
 	/** Updates (or adds if not already present) arbitrary engine data to the crash context (will remove the key if passed an empty string) */
 	CORE_API static void SetEngineData(const FString& Key, const FString& Value);
 
+	/** Get the engine data dictionary */
+	CORE_API static const TMap<FString, FString>& GetEngineData();
+
 	/** Clears the game data dictionary */
 	CORE_API static void ResetGameData();
 
 	/** Updates (or adds if not already present) arbitrary game data to the crash context (will remove the key if passed an empty string) */
 	CORE_API static void SetGameData(const FString& Key, const FString& Value);
 
+	/** Get the game data dictionary */
+	CORE_API static const TMap<FString, FString>& GetGameData();
+
 	/** Adds a plugin descriptor string to the enabled plugins list in the crash context */
 	CORE_API static void AddPlugin(const FString& PluginDesc);
 
