Added EventCount

This is an event that provides a mechanism for avoiding a race between waiting and signaling through the use of a signal count token.

diff --git a/Engine/Source/Developer/DerivedDataCache/Private/DerivedDataRequestOwner.cpp b/Engine/Source/Developer/DerivedDataCache/Private/DerivedDataRequestOwner.cpp
--- a/Engine/Source/Developer/DerivedDataCache/Private/DerivedDataRequestOwner.cpp
+++ b/Engine/Source/Developer/DerivedDataCache/Private/DerivedDataRequestOwner.cpp
@@ -2,10 +2,10 @@
 
 #include "DerivedDataRequestOwner.h"
 
+#include "Async/EventCount.h"
 #include "Containers/Array.h"
 #include "DerivedDataRequest.h"
 #include "DerivedDataRequestTypes.h"
-#include "Experimental/Async/LazyEvent.h"
 #include "HAL/CriticalSection.h"
 #include "Misc/ScopeExit.h"
 #include "Misc/ScopeRWLock.h"
@@ -17,12 +17,6 @@
 namespace UE::DerivedData::Private
 {
 
-/** A one-time-use event to work around the lack of condition variables. */
-struct FRequestBarrierEvent : public FRefCountBase
-{
-	FLazyEvent Event{EEventMode::ManualReset};
-};
-
 class FRequestOwnerShared final : public IRequestOwner, public FRequestBase
 {
 public:
@@ -49,7 +43,7 @@ public:
 private:
 	mutable FRWLock Lock;
 	TArray<TRefCountPtr<IRequest>, TInlineAllocator<1>> Requests;
-	TRefCountPtr<FRequestBarrierEvent> BarrierEvent;
+	FEventCount BarrierEvent;
 	uint32 BarrierCount{0};
 	uint32 PriorityBarrierCount{0};
 	EPriority Priority{EPriority::Normal};
@@ -128,7 +122,7 @@ void FRequestOwnerShared::BeginBarrier(ERequestBarrierFlags Flags)
 void FRequestOwnerShared::EndBarrier(ERequestBarrierFlags Flags)
 {
 	ON_SCOPE_EXIT { Release(); };
-	TRefCountPtr<FRequestBarrierEvent> LocalBarrierEvent;
+	bool bSignalBarrier = false;
 	{
 		FWriteScopeLock WriteLock(Lock);
 		if (EnumHasAnyFlags(Flags, ERequestBarrierFlags::Priority) && --PriorityBarrierCount == 0)
@@ -137,13 +131,12 @@ void FRequestOwnerShared::EndBarrier(ERequestBarrierFlags Flags)
 		}
 		if (--BarrierCount == 0)
 		{
-			LocalBarrierEvent = MoveTemp(BarrierEvent);
-			BarrierEvent = nullptr;
+			bSignalBarrier = true;
 		}
 	}
-	if (LocalBarrierEvent)
+	if (bSignalBarrier)
 	{
-		LocalBarrierEvent->Event.Trigger();
+		BarrierEvent.Signal();
 	}
 }
 
@@ -172,7 +165,7 @@ void FRequestOwnerShared::Cancel()
 	for (;;)
 	{
 		TRefCountPtr<IRequest> LocalRequest;
-		TRefCountPtr<FRequestBarrierEvent> LocalBarrierEvent;
+		FEventCountToken BarrierToken;
 
 		{
 			FWriteScopeLock WriteLock(Lock);
@@ -184,11 +177,7 @@ void FRequestOwnerShared::Cancel()
 			}
 			else if (BarrierCount > 0)
 			{
-				LocalBarrierEvent = BarrierEvent;
-				if (!LocalBarrierEvent)
-				{
-					LocalBarrierEvent = BarrierEvent = new FRequestBarrierEvent;
-				}
+				BarrierToken = BarrierEvent.PrepareWait();
 			}
 			else
 			{
@@ -202,7 +191,7 @@ void FRequestOwnerShared::Cancel()
 		}
 		else
 		{
-			LocalBarrierEvent->Event.Wait();
+			BarrierEvent.Wait(BarrierToken);
 		}
 	}
 }
@@ -212,7 +201,7 @@ void FRequestOwnerShared::Wait()
 	for (;;)
 	{
 		TRefCountPtr<IRequest> LocalRequest;
-		TRefCountPtr<FRequestBarrierEvent> LocalBarrierEvent;
+		FEventCountToken BarrierToken;
 
 		{
 			FWriteScopeLock WriteLock(Lock);
@@ -223,11 +212,7 @@ void FRequestOwnerShared::Wait()
 			}
 			else if (BarrierCount > 0)
 			{
-				LocalBarrierEvent = BarrierEvent;
-				if (!LocalBarrierEvent)
-				{
-					LocalBarrierEvent = BarrierEvent = new FRequestBarrierEvent;
-				}
+				BarrierToken = BarrierEvent.PrepareWait();
 			}
 			else
 			{
@@ -241,7 +226,7 @@ void FRequestOwnerShared::Wait()
 		}
 		else
 		{
-			LocalBarrierEvent->Event.Wait();
+			BarrierEvent.Wait(BarrierToken);
 		}
 	}
 }
diff --git a/Engine/Source/Runtime/Core/Private/Misc/OutputDeviceRedirector.cpp b/Engine/Source/Runtime/Core/Private/Misc/OutputDeviceRedirector.cpp
--- a/Engine/Source/Runtime/Core/Private/Misc/OutputDeviceRedirector.cpp
+++ b/Engine/Source/Runtime/Core/Private/Misc/OutputDeviceRedirector.cpp
@@ -2,11 +2,10 @@
 
 #include "Misc/OutputDeviceRedirector.h"
 
+#include "Async/EventCount.h"
 #include "Containers/BitArray.h"
 #include "Containers/DepletableMpmcQueue.h"
-#include "Containers/ConsumeAllMpmcQueue.h"
 #include "Experimental/ConcurrentLinearAllocator.h"
-#include "HAL/Event.h"
 #include "HAL/PlatformProcess.h"
 #include "HAL/PlatformTime.h"
 #include "HAL/PlatformTLS.h"
@@ -139,11 +138,11 @@ struct FOutputDeviceRedirectorState
 	/** A lock to synchronize access to the thread. */
 	FRWLock ThreadLock;
 
-	/** A queue of events to trigger when the dedicated primary thread is idle. */
-	TConsumeAllMpmcQueue<FEvent*, FOutputDeviceLinearAllocator> ThreadIdleEvents;
+	/** An event that is signaled when the dedicated primary thread is idle. */
+	FEventCount ThreadIdleEvent;
 
 	/** An event to wake the dedicated primary thread to process buffered items. */
-	std::atomic<FEvent*> ThreadWakeEvent = nullptr;
+	FEventCount ThreadWakeEvent;
 
 	/** The ID of the thread holding the primary lock. */
 	std::atomic<uint32> LockedThreadId = MAX_uint32;
@@ -154,6 +153,9 @@ struct FOutputDeviceRedirectorState
 	/** The ID of the panic thread, which is only set by Panic(). */
 	std::atomic<uint32> PanicThreadId = MAX_uint32;
 
+	/** Whether a dedicated primary thread has been started. */
+	std::atomic<bool> bThreadStarted = false;
+
 	/** Whether the backlog is enabled. */
 	bool bEnableBacklog = false;
 
@@ -408,11 +410,8 @@ void FOutputDeviceRedirectorState::RemoveOutputDevice(FOutputDevice* OutputDevic
 
 bool FOutputDeviceRedirectorState::TryStartThread()
 {
-	if (FWriteScopeLock ThreadScopeLock(ThreadLock); !ThreadWakeEvent.load(std::memory_order_relaxed))
+	if (FWriteScopeLock ThreadScopeLock(ThreadLock); !bThreadStarted.exchange(true, std::memory_order_relaxed))
 	{
-		FEvent* WakeEvent = FPlatformProcess::GetSynchEventFromPool();
-		WakeEvent->Trigger();
-		ThreadWakeEvent.store(WakeEvent, std::memory_order_release);
 		Thread = FThread(TEXT("OutputDeviceRedirector"), [this] { ThreadLoop(); });
 	}
 	return true;
@@ -420,12 +419,10 @@ bool FOutputDeviceRedirectorState::TryStartThread()
 
 bool FOutputDeviceRedirectorState::TryStopThread()
 {
-	if (FWriteScopeLock ThreadScopeLock(ThreadLock); FEvent* WakeEvent = ThreadWakeEvent.exchange(nullptr, std::memory_order_acquire))
+	if (FWriteScopeLock ThreadScopeLock(ThreadLock); bThreadStarted.exchange(false, std::memory_order_relaxed))
 	{
-		WakeEvent->Trigger();
+		ThreadWakeEvent.Signal();
 		Thread.Join();
-		FOutputDevicesWriteScopeLock Lock(*this);
-		FPlatformProcess::ReturnSynchEventToPool(WakeEvent);
 	}
 	return true;
 }
@@ -439,9 +436,13 @@ void FOutputDeviceRedirectorState::ThreadLoop()
 		PrimaryThreadId.store(ThreadId, std::memory_order_relaxed);
 	}
 
-	while (FEvent* WakeEvent = ThreadWakeEvent.load(std::memory_order_acquire))
+	for (;;)
 	{
-		WakeEvent->Wait();
+		FEventCountToken Token = ThreadWakeEvent.PrepareWait();
+		if (!bThreadStarted.load(std::memory_order_relaxed))
+		{
+			break;
+		}
 		while (!BufferedItems.IsEmpty() && IsPrimaryThread(ThreadId))
 		{
 			if (FOutputDevicesPrimaryScopeLock Lock(*this); Lock.IsLocked())
@@ -449,7 +450,8 @@ void FOutputDeviceRedirectorState::ThreadLoop()
 				FlushBufferedItems();
 			}
 		}
-		ThreadIdleEvents.ConsumeAllLifo([](FEvent* Event) { Event->Trigger(); });
+		ThreadIdleEvent.Signal();
+		ThreadWakeEvent.Wait(Token);
 	}
 }
 
@@ -553,17 +555,14 @@ bool FOutputDeviceRedirector::IsRedirectingTo(FOutputDevice* OutputDevice)
 
 void FOutputDeviceRedirector::FlushThreadedLogs(EOutputDeviceRedirectorFlushOptions Options)
 {
-	if (FReadScopeLock ThreadLock(State->ThreadLock); FEvent* WakeEvent = State->ThreadWakeEvent.load(std::memory_order_acquire))
+	if (FReadScopeLock ThreadLock(State->ThreadLock); State->bThreadStarted.load(std::memory_order_relaxed))
 	{
 		if (!EnumHasAnyFlags(Options, EOutputDeviceRedirectorFlushOptions::Async))
 		{
 			TRACE_CPUPROFILER_EVENT_SCOPE(FOutputDeviceRedirector::FlushThreadedLogs);
-			FEventRef IdleEvent(EEventMode::ManualReset);
-			if (State->ThreadIdleEvents.ProduceItem(IdleEvent.Get()) == UE::EConsumeAllMpmcQueueResult::WasEmpty)
-			{
-				WakeEvent->Trigger();
-			}
-			IdleEvent->Wait();
+			UE::FEventCountToken Token = State->ThreadIdleEvent.PrepareWait();
+			State->ThreadWakeEvent.Signal();
+			State->ThreadIdleEvent.Wait(Token);
 		}
 		return;
 	}
@@ -656,10 +655,7 @@ void FOutputDeviceRedirector::SerializeRecord(const UE::FLogRecord& Record)
 	// Queue the record to serialize to buffered output devices from the primary thread.
 	if (State->BufferedItems.EnqueueAndReturnWasEmpty(Record))
 	{
-		if (FEvent* WakeEvent = State->ThreadWakeEvent.load(std::memory_order_acquire))
-		{
-			WakeEvent->Trigger();
-		}
+		State->ThreadWakeEvent.Signal();
 	}
 }
 
@@ -720,10 +716,7 @@ void FOutputDeviceRedirector::Serialize(const TCHAR* const Data, const ELogVerbo
 	// Queue the line to serialize to buffered output devices from the primary thread.
 	if (State->BufferedItems.EnqueueAndReturnWasEmpty(Data, Category, Verbosity, RealTime))
 	{
-		if (FEvent* WakeEvent = State->ThreadWakeEvent.load(std::memory_order_acquire))
-		{
-			WakeEvent->Trigger();
-		}
+		State->ThreadWakeEvent.Signal();
 	}
 }
 
diff --git a/Engine/Source/Runtime/Core/Public/Async/EventCount.h b/Engine/Source/Runtime/Core/Public/Async/EventCount.h
new file mode 100644
--- /dev/null
+++ b/Engine/Source/Runtime/Core/Public/Async/EventCount.h
@@ -0,0 +1,152 @@
+// Copyright Epic Games, Inc. All Rights Reserved.
+
+#pragma once
+
+#include "Async/ParkingLot.h"
+#include <atomic>
+#include <type_traits>
+
+namespace UE
+{
+
+template <typename CounterType>
+class TEventCount;
+
+/** A token used to wait on TEventCount. */
+template <typename CounterType>
+class TEventCountToken
+{
+	static_assert(std::is_unsigned_v<CounterType>);
+
+public:
+	/** Returns true if the token has been assigned by PrepareWait(). */
+	inline explicit operator bool() const { return !(Value & 1); }
+
+private:
+	/** Defaults to an odd value, which is never valid to wait on. */
+	CounterType Value = 1;
+
+	friend TEventCount<CounterType>;
+};
+
+/**
+ * A type of event that avoids missed signals by maintaining a signal count.
+ *
+ * This type of event is suited to waiting on another thread conditionally.
+ * Typical usage looks similar to this example:
+ *
+ *     FEventCount Event;
+ *     std::atomic<uint32> CurrentValue = 0;
+ *
+ * On the waiting thread:
+ *
+ *     FEventCountToken Token = Event.PrepareWait();
+ *     if (CurrentValue < TargetValue)
+ *     {
+ *         Event.Wait(Token);
+ *     }
+ *
+ * On the signaling thread:
+ *
+ *     ++CurrentValue;
+ *     if (CurrentValue == TargetValue)
+ *     {
+ *         Event.Signal();
+ *     }
+ *
+ * Acquiring a token before checking the condition avoids a race because Wait returns immediately
+ * when the token no longer matches the signal count.
+ */
+template <typename CounterType>
+class TEventCount
+{
+	static_assert(std::is_unsigned_v<CounterType>);
+
+public:
+	constexpr TEventCount() = default;
+
+	TEventCount(const TEventCount&) = delete;
+	TEventCount& operator=(const TEventCount&) = delete;
+
+	/**
+	 * Prepare to wait.
+	 *
+	 * Call this before any logic that must re-execute if the event is signaled in the meantime.
+	 *
+	 * @return A token to pass to one of the wait functions to abort the wait if the event has been signaled since.
+	 */
+	inline TEventCountToken<CounterType> PrepareWait()
+	{
+		TEventCountToken<CounterType> Token;
+		Token.Value = Count.fetch_or(1, std::memory_order_relaxed) & ~CounterType(1);
+		return Token;
+	}
+
+	/**
+	 * Wait until the event is signaled. Returns immediately if signaled since the token was acquired.
+	 *
+	 * @param Compare   A token acquired from PrepareWait() before checking the conditions for this wait.
+	 */
+	inline void Wait(TEventCountToken<CounterType> Compare)
+	{
+		if ((Count.load(std::memory_order_acquire) & ~CounterType(1)) == Compare.Value)
+		{
+			ParkingLot::Wait(&Count, [this, Compare] { return (Count.load(std::memory_order_acquire) & ~CounterType(1)) == Compare.Value; }, []{});
+		}
+	}
+
+	/**
+	 * Wait until the event is signaled. Returns immediately if signaled since the token was acquired.
+	 *
+	 * @param Compare    A token acquired from PrepareWait() before checking the conditions for this wait.
+	 * @param WaitTime   Relative time after which waiting is automatically canceled and the thread wakes.
+	 * @return True if the event was signaled before the wait time elapsed, otherwise false.
+	 */
+	inline bool WaitFor(TEventCountToken<CounterType> Compare, FMonotonicTimeSpan WaitTime)
+	{
+		if ((Count.load(std::memory_order_acquire) & ~CounterType(1)) == Compare.Value)
+		{
+			return ParkingLot::WaitFor(&Count, [this, Compare] { return (Count.load(std::memory_order_acquire) & ~CounterType(1)) == Compare.Value; }, []{}, WaitTime).bDidWake;
+		}
+		return true;
+	}
+
+	/**
+	 * Wait until the event is signaled. Returns immediately if signaled since the token was acquired.
+	 *
+	 * @param Compare    A token acquired from PrepareWait() before checking the conditions for this wait.
+	 * @param WaitTime   Absolute time after which waiting is automatically canceled and the thread wakes.
+	 * @return True if the event was signaled before the wait time elapsed, otherwise false.
+	 */
+	inline bool WaitUntil(TEventCountToken<CounterType> Compare, FMonotonicTimePoint WaitTime)
+	{
+		if ((Count.load(std::memory_order_acquire) & ~CounterType(1)) == Compare.Value)
+		{
+			return ParkingLot::WaitUntil(&Count, [this, Compare] { return (Count.load(std::memory_order_acquire) & ~CounterType(1)) == Compare.Value; }, []{}, WaitTime).bDidWake;
+		}
+		return true;
+	}
+
+	/**
+	 * Signals all waiting threads.
+	 *
+	 * Any threads that have called PrepareWait() and not yet waited will be signaled immediately
+	 * if they do wait on a token from a call to PrepareWait() that preceded this call.
+	 */
+	inline void Signal()
+	{
+		CounterType Value = Count.load(std::memory_order_relaxed);
+		if ((Value & 1) && Count.compare_exchange_strong(Value, Value + 1, std::memory_order_release))
+		{
+			ParkingLot::WakeAll(&Count);
+		}
+	}
+
+private:
+	std::atomic<CounterType> Count = 0;
+};
+
+using FEventCount = TEventCount<uint32>;
+using FEventCountToken = TEventCountToken<uint32>;
+
+} // UE
