Added FManualResetEvent

This is a one-byte manual-reset event that can replace that use case for the experimental FLazyEvent.

diff --git a/Engine/Source/Developer/DerivedDataCache/Private/DerivedDataBuildJobContext.cpp b/Engine/Source/Developer/DerivedDataCache/Private/DerivedDataBuildJobContext.cpp
--- a/Engine/Source/Developer/DerivedDataCache/Private/DerivedDataBuildJobContext.cpp
+++ b/Engine/Source/Developer/DerivedDataCache/Private/DerivedDataBuildJobContext.cpp
@@ -140,7 +140,7 @@ void FBuildJobContext::BeginBuild(IRequestOwner& InOwner, TUniqueFunction<void (
 void FBuildJobContext::EndBuild()
 {
 	OnEndBuild();
-	BuildCompleteEvent.Trigger();
+	BuildCompleteEvent.Notify();
 	Owner = nullptr;
 }
 
diff --git a/Engine/Source/Developer/DerivedDataCache/Private/DerivedDataBuildJobContext.h b/Engine/Source/Developer/DerivedDataCache/Private/DerivedDataBuildJobContext.h
--- a/Engine/Source/Developer/DerivedDataCache/Private/DerivedDataBuildJobContext.h
+++ b/Engine/Source/Developer/DerivedDataCache/Private/DerivedDataBuildJobContext.h
@@ -2,6 +2,7 @@
 
 #pragma once
 
+#include "Async/ManualResetEvent.h"
 #include "Compression/CompressedBuffer.h"
 #include "Containers/Map.h"
 #include "Containers/StringView.h"
@@ -9,7 +10,6 @@
 #include "DerivedDataCacheKey.h"
 #include "DerivedDataRequest.h"
 #include "DerivedDataSharedString.h"
-#include "Experimental/Async/LazyEvent.h"
 #include "Memory/MemoryFwd.h"
 #include "Serialization/CompactBinary.h"
 #include "Templates/Function.h"
@@ -88,12 +88,12 @@ private:
 	FBuildOutputBuilder& OutputBuilder;
 	TMap<FUtf8SharedString, FCbObject> Constants;
 	TMap<FUtf8SharedString, FCompressedBuffer> Inputs;
-	UE::FLazyEvent BuildCompleteEvent{EEventMode::ManualReset};
 	TUniqueFunction<void ()> OnEndBuild;
 	IRequestOwner* Owner{};
 	uint64 RequiredMemory{};
 	ECachePolicy CachePolicyMask;
 	EBuildPolicy BuildPolicyMask;
+	FManualResetEvent BuildCompleteEvent;
 	bool bIsAsyncBuild{false};
 	bool bIsAsyncBuildComplete{false};
 	bool bDeterministicOutputCheck{true};
diff --git a/Engine/Source/Developer/DerivedDataCache/Private/DerivedDataCache.cpp b/Engine/Source/Developer/DerivedDataCache/Private/DerivedDataCache.cpp
--- a/Engine/Source/Developer/DerivedDataCache/Private/DerivedDataCache.cpp
+++ b/Engine/Source/Developer/DerivedDataCache/Private/DerivedDataCache.cpp
@@ -20,7 +20,6 @@
 #include "DerivedDataRequest.h"
 #include "DerivedDataRequestOwner.h"
 #include "DerivedDataThreadPoolTask.h"
-#include "Experimental/Async/LazyEvent.h"
 #include "Experimental/ZenServerInterface.h"
 #include "Features/IModularFeatures.h"
 #include "HAL/ThreadSafeCounter.h"
diff --git a/Engine/Source/Developer/DerivedDataCache/Private/DerivedDataThreadPoolTask.cpp b/Engine/Source/Developer/DerivedDataCache/Private/DerivedDataThreadPoolTask.cpp
--- a/Engine/Source/Developer/DerivedDataCache/Private/DerivedDataThreadPoolTask.cpp
+++ b/Engine/Source/Developer/DerivedDataCache/Private/DerivedDataThreadPoolTask.cpp
@@ -3,10 +3,10 @@
 #include "DerivedDataThreadPoolTask.h"
 
 #include "Async/InheritedContext.h"
+#include "Async/ManualResetEvent.h"
 #include "DerivedDataRequest.h"
 #include "DerivedDataRequestOwner.h"
 #include "DerivedDataRequestTypes.h"
-#include "Experimental/Async/LazyEvent.h"
 #include "Misc/IQueuedWork.h"
 #include "Misc/QueuedThreadPool.h"
 #include "Stats/Stats.h"
@@ -56,7 +56,7 @@ private:
 	IRequestOwner& Owner;
 	FQueuedThreadPool& ThreadPool;
 	TUniqueFunction<void ()> TaskBody;
-	FLazyEvent Event{EEventMode::ManualReset};
+	FManualResetEvent Event;
 	std::atomic<bool> bClaimed = false;
 };
 
@@ -91,7 +91,7 @@ TRefCountPtr<IRequest> FThreadPoolTaskRequest::TryEnd()
 		FInheritedContextScope InheritedContextScope = RestoreInheritedContext();
 		FScopeCycleCounter Scope(GetStatId(), /*bAlways*/ true);
 		TaskBody();
-		Event.Trigger();
+		Event.Notify();
 	});
 }
 
diff --git a/Engine/Source/Developer/DerivedDataCache/Private/FileSystemCacheStore.cpp b/Engine/Source/Developer/DerivedDataCache/Private/FileSystemCacheStore.cpp
--- a/Engine/Source/Developer/DerivedDataCache/Private/FileSystemCacheStore.cpp
+++ b/Engine/Source/Developer/DerivedDataCache/Private/FileSystemCacheStore.cpp
@@ -7,6 +7,7 @@
 #include "Algo/StableSort.h"
 #include "Algo/Transform.h"
 #include "Async/Async.h"
+#include "Async/ManualResetEvent.h"
 #include "Containers/StaticBitArray.h"
 #include "DerivedDataBackendInterface.h"
 #include "DerivedDataCacheInterface.h"
@@ -17,7 +18,6 @@
 #include "DerivedDataChunk.h"
 #include "DerivedDataRequestOwner.h"
 #include "DerivedDataValue.h"
-#include "Experimental/Async/LazyEvent.h"
 #include "Features/IModularFeatures.h"
 #include "HAL/Event.h"
 #include "HAL/FileManager.h"
@@ -265,7 +265,7 @@ private:
 
 	IFileManager& FileManager = IFileManager::Get();
 
-	mutable FLazyEvent IdleEvent;
+	mutable FManualResetEvent IdleEvent;
 	FEventRef WaitEvent;
 	FThread Thread;
 
@@ -315,7 +315,6 @@ FFileSystemCacheStoreMaintainer::FFileSystemCacheStoreMaintainer(
 	: Params(InParams)
 	, CachePath(InCachePath)
 	, bExitAfterScan(Params.ScanFrequency.GetTotalDays() > MaxScanFrequencyDays)
-	, IdleEvent(EEventMode::ManualReset)
 	, WaitEvent(EEventMode::AutoReset)
 	, Thread(
 		TEXT("FileSystemCacheStoreMaintainer"),
@@ -350,7 +349,7 @@ void FFileSystemCacheStoreMaintainer::Tick()
 		Loop();
 	}
 	bIdle = true;
-	IdleEvent.Trigger();
+	IdleEvent.Notify();
 }
 
 void FFileSystemCacheStoreMaintainer::Loop()
@@ -370,7 +369,7 @@ void FFileSystemCacheStoreMaintainer::Loop()
 		bIdle = false;
 		Scan();
 		bIdle = true;
-		IdleEvent.Trigger();
+		IdleEvent.Notify();
 		bIgnoreScanRate = false;
 		const FDateTime ScanEnd = FDateTime::Now();
 
@@ -395,7 +394,7 @@ void FFileSystemCacheStoreMaintainer::Loop()
 	}
 
 	bIdle = true;
-	IdleEvent.Trigger();
+	IdleEvent.Notify();
 }
 
 void FFileSystemCacheStoreMaintainer::Scan()
diff --git a/Engine/Source/Developer/DerivedDataCache/Private/Http/CurlHttpClient.cpp b/Engine/Source/Developer/DerivedDataCache/Private/Http/CurlHttpClient.cpp
--- a/Engine/Source/Developer/DerivedDataCache/Private/Http/CurlHttpClient.cpp
+++ b/Engine/Source/Developer/DerivedDataCache/Private/Http/CurlHttpClient.cpp
@@ -3,10 +3,10 @@
 #include "HttpClient.h"
 
 #include "Async/InheritedContext.h"
+#include "Async/ManualResetEvent.h"
 #include "Containers/ConsumeAllMpmcQueue.h"
 #include "Containers/LockFreeList.h"
 #include "Containers/StringView.h"
-#include "HAL/Event.h"
 #include "HAL/PlatformProcess.h"
 #include "HAL/Thread.h"
 #include "Memory/CompositeBuffer.h"
@@ -294,7 +294,7 @@ public:
 	TRefCountPtr<IHttpResponseMonitor> GetMonitor() final { return this; }
 
 	void Cancel() final;
-	void Wait() const final { CompleteEvent->Wait(); }
+	void Wait() const final { CompleteEvent.Wait(); }
 	bool Poll() const final { return EnumHasAnyFlags(State.load(), ECurlHttpResponseState::Complete); }
 	bool IsCanceled() const { return EnumHasAnyFlags(State.load(std::memory_order_relaxed), ECurlHttpResponseState::Canceled); }
 
@@ -340,7 +340,7 @@ private:
 	FAnsiStringView Error;
 	ANSICHAR ErrorBuffer[CURL_ERROR_SIZE]{};
 	FHttpResponseStats Stats;
-	FEventRef CompleteEvent{EEventMode::ManualReset};
+	mutable FManualResetEvent CompleteEvent;
 
 	enum class EHttpReceiverFunction : uint8
 	{
@@ -1138,7 +1138,7 @@ void FCurlHttpResponse::SetComplete(CURLcode Code)
 	}
 
 	CurlHttp::Private::AtomicFetchOr(State, ECurlHttpResponseState::Complete);
-	CompleteEvent->Trigger();
+	CompleteEvent.Notify();
 
 	if (UE_LOG_ACTIVE(LogHttp, Verbose))
 	{
@@ -1184,7 +1184,7 @@ void FCurlHttpResponse::Cancel()
 	// Waiting within a receiver function would prevent the response from completing.
 	if (LocalReceiverFunction == EHttpReceiverFunction::None)
 	{
-		CompleteEvent->Wait();
+		CompleteEvent.Wait();
 	}
 }
 
diff --git a/Engine/Source/Developer/DerivedDataCache/Private/Http/HttpClient.cpp b/Engine/Source/Developer/DerivedDataCache/Private/Http/HttpClient.cpp
--- a/Engine/Source/Developer/DerivedDataCache/Private/Http/HttpClient.cpp
+++ b/Engine/Source/Developer/DerivedDataCache/Private/Http/HttpClient.cpp
@@ -2,8 +2,8 @@
 
 #include "HttpClient.h"
 
+#include "Async/ManualResetEvent.h"
 #include "Containers/LockFreeList.h"
-#include "HAL/Event.h"
 #include "Memory/MemoryView.h"
 #include "Misc/AsciiSet.h"
 #include "ProfilingDebugging/CpuProfilerTrace.h"
@@ -271,7 +271,7 @@ struct Http::Private::FHttpRequestQueueData
 	struct FWaiter
 	{
 		THttpUniquePtr<IHttpRequest> Request;
-		FEventRef Event{EEventMode::ManualReset};
+		FManualResetEvent Event;
 	};
 
 	FHttpRequestQueueData(IHttpConnectionPool& ConnectionPool, const FHttpClientParams& ClientParams)
@@ -290,7 +290,7 @@ struct Http::Private::FHttpRequestQueueData
 					if (FWaiter* Waiter = Waiters.Pop())
 					{
 						Waiter->Request = MoveTemp(Request);
-						Waiter->Event->Trigger();
+						Waiter->Event.Notify();
 					}
 				}
 			}
@@ -310,7 +310,7 @@ struct Http::Private::FHttpRequestQueueData
 			if (FWaiter* Waiter = Waiters.Pop())
 			{
 				Waiter->Request = MoveTemp(Request);
-				Waiter->Event->Trigger();
+				Waiter->Event.Notify();
 			}
 			else
 			{
@@ -326,16 +326,16 @@ struct Http::Private::FHttpRequestQueueData
 			if (FWaiter* Waiter = Waiters.Pop())
 			{
 				Waiter->Request = MoveTemp(Request);
-				Waiter->Event->Trigger();
+				Waiter->Event.Notify();
 			}
-			if (LocalWaiter.Event->Wait(0))
+			if (LocalWaiter.Event.WaitFor(FMonotonicTimeSpan::Zero()))
 			{
 				return MoveTemp(LocalWaiter.Request);
 			}
 		}
 
 		TRACE_CPUPROFILER_EVENT_SCOPE(HttpRequestQueue_Wait);
-		LocalWaiter.Event->Wait();
+		LocalWaiter.Event.Wait();
 		return MoveTemp(LocalWaiter.Request);
 	}
 
diff --git a/Engine/Source/Runtime/Core/Public/Async/ManualResetEvent.h b/Engine/Source/Runtime/Core/Public/Async/ManualResetEvent.h
new file mode 100644
--- /dev/null
+++ b/Engine/Source/Runtime/Core/Public/Async/ManualResetEvent.h
@@ -0,0 +1,95 @@
+// Copyright Epic Games, Inc. All Rights Reserved.
+
+#pragma once
+
+#include "Async/ParkingLot.h"
+#include <atomic>
+
+namespace UE
+{
+
+/**
+ * A type of event that remains notified until manually reset.
+ */
+class FManualResetEvent
+{
+public:
+	constexpr FManualResetEvent() = default;
+
+	FManualResetEvent(const FManualResetEvent&) = delete;
+	FManualResetEvent& operator=(const FManualResetEvent&) = delete;
+
+	/**
+	 * Wait until the event is notified.
+	 */
+	inline void Wait()
+	{
+		WaitUntil(FMonotonicTimePoint::Infinity());
+	}
+
+	/**
+	 * Wait until the event is notified.
+	 *
+	 * @param WaitTime   Relative time after which waiting is automatically canceled and the thread wakes.
+	 * @return True if the event was notified before the wait time elapsed, otherwise false.
+	 */
+	inline bool WaitFor(FMonotonicTimeSpan WaitTime)
+	{
+		return WaitUntil(FMonotonicTimePoint::Now() + WaitTime);
+	}
+
+	/**
+	 * Wait until the event is notified.
+	 *
+	 * @param WaitTime   Absolute time after which waiting is automatically canceled and the thread wakes.
+	 * @return True if the event was notified before the wait time elapsed, otherwise false.
+	 */
+	inline bool WaitUntil(FMonotonicTimePoint WaitTime)
+	{
+		for (uint8 CurrentState = State.load(std::memory_order_acquire);;)
+		{
+			if (CurrentState & IsNotifiedFlag)
+			{
+				return true;
+			}
+
+			if ((CurrentState & MaybeWaitingFlag) ||
+				State.compare_exchange_weak(CurrentState, MaybeWaitingFlag, std::memory_order_relaxed, std::memory_order_acquire))
+			{
+				const ParkingLot::FWaitState WaitState = ParkingLot::WaitUntil(&State, [this, &CurrentState]
+				{
+					return !(State.load(std::memory_order_acquire) & IsNotifiedFlag);
+				}, []{}, WaitTime);
+				return !WaitState.bDidWait || WaitState.bDidWake;
+			}
+		}
+	}
+
+	/**
+	 * Notifies all waiting threads and leaves the event notified until the next call to Reset().
+	 */
+	inline void Notify()
+	{
+		const uint8 CurrentState = State.fetch_or(IsNotifiedFlag, std::memory_order_release);
+		if (!(CurrentState & IsNotifiedFlag) && (CurrentState & MaybeWaitingFlag))
+		{
+			ParkingLot::WakeAll(&State);
+		}
+	}
+
+	/**
+	 * Resets the event to a non-notified state.
+	 */
+	inline void Reset()
+	{
+		State.fetch_and(~IsNotifiedFlag, std::memory_order_relaxed);
+	}
+
+private:
+	static constexpr uint8 IsNotifiedFlag = 1 << 0;
+	static constexpr uint8 MaybeWaitingFlag = 1 << 1;
+
+	std::atomic<uint8> State = 0;
+};
+
+} // UE
diff --git a/Engine/Source/Runtime/Core/Tests/Async/ManualResetEventTest.cpp b/Engine/Source/Runtime/Core/Tests/Async/ManualResetEventTest.cpp
new file mode 100644
--- /dev/null
+++ b/Engine/Source/Runtime/Core/Tests/Async/ManualResetEventTest.cpp
@@ -0,0 +1,36 @@
+// Copyright Epic Games, Inc. All Rights Reserved.
+
+#if WITH_LOW_LEVEL_TESTS
+
+#include "Async/ManualResetEvent.h"
+
+#include "Tasks/Task.h"
+#include "TestHarness.h"
+
+namespace UE
+{
+
+TEST_CASE("Core::Async::ManualResetEvent", "[Core][Async]")
+{
+	FManualResetEvent Event;
+
+	CHECK_FALSE(Event.WaitFor(FMonotonicTimeSpan::Zero()));
+	Event.Notify();
+	Event.Wait();
+	CHECK(Event.WaitFor(FMonotonicTimeSpan::Zero()));
+	CHECK(Event.WaitUntil(FMonotonicTimePoint::Now() - FMonotonicTimeSpan::FromSeconds(1.0)));
+	Event.Reset();
+	CHECK_FALSE(Event.WaitFor(FMonotonicTimeSpan::Zero()));
+
+	Tasks::Launch(UE_SOURCE_LOCATION, [&Event] { Event.Notify(); });
+	Event.Wait();
+
+	Event.Notify();
+	CHECK(Event.WaitFor(FMonotonicTimeSpan::Zero()));
+	Event.Reset();
+	CHECK_FALSE(Event.WaitFor(FMonotonicTimeSpan::Zero()));
+}
+
+} // UE
+
+#endif // WITH_LOW_LEVEL_TESTS
diff --git a/Engine/Source/Runtime/CoreUObject/Private/Serialization/DerivedData.cpp b/Engine/Source/Runtime/CoreUObject/Private/Serialization/DerivedData.cpp
--- a/Engine/Source/Runtime/CoreUObject/Private/Serialization/DerivedData.cpp
+++ b/Engine/Source/Runtime/CoreUObject/Private/Serialization/DerivedData.cpp
@@ -29,7 +29,7 @@
 ///////////////////////////////////////////////////////////////////////////////////////////////////////////////////////
 // TEMPORARY SERIALIZATION AS BULK DATA
 #if WITH_EDITORONLY_DATA
-#include "Experimental/Async/LazyEvent.h"
+#include "Async/ManualResetEvent.h"
 #include "Serialization/BulkData.h"
 #endif // WITH_EDITORONLY_DATA
 // TEMPORARY SERIALIZATION AS BULK DATA
@@ -336,11 +336,11 @@ void FDerivedData::Serialize(FArchive& Ar, UObject* Owner)
 
 		// Blocking read to populate the bulk data to be serialized later by the linker.
 		{
-			FLazyEvent BatchComplete(EEventMode::ManualReset);
+			FManualResetEvent BatchComplete;
 			FDerivedDataIoBatch Batch;
 			FDerivedDataIoRequest Request = Batch.Read(*this);
 			FDerivedDataIoResponse Response;
-			Batch.Dispatch(Response, [&BatchComplete] { BatchComplete.Trigger(); });
+			Batch.Dispatch(Response, [&BatchComplete] { BatchComplete.Notify(); });
 			BatchComplete.Wait();
 
 			const FSharedBuffer Data = Response.GetData(Request);
