Renamed Signal to Notify in new synchronization primitives

#rnx

diff --git a/Engine/Source/Developer/DerivedDataCache/Private/DerivedDataRequestOwner.cpp b/Engine/Source/Developer/DerivedDataCache/Private/DerivedDataRequestOwner.cpp
--- a/Engine/Source/Developer/DerivedDataCache/Private/DerivedDataRequestOwner.cpp
+++ b/Engine/Source/Developer/DerivedDataCache/Private/DerivedDataRequestOwner.cpp
@@ -122,7 +122,7 @@ void FRequestOwnerShared::BeginBarrier(ERequestBarrierFlags Flags)
 void FRequestOwnerShared::EndBarrier(ERequestBarrierFlags Flags)
 {
 	ON_SCOPE_EXIT { Release(); };
-	bool bSignalBarrier = false;
+	bool bNotifyBarrier = false;
 	{
 		FWriteScopeLock WriteLock(Lock);
 		if (EnumHasAnyFlags(Flags, ERequestBarrierFlags::Priority) && --PriorityBarrierCount == 0)
@@ -131,12 +131,12 @@ void FRequestOwnerShared::EndBarrier(ERequestBarrierFlags Flags)
 		}
 		if (--BarrierCount == 0)
 		{
-			bSignalBarrier = true;
+			bNotifyBarrier = true;
 		}
 	}
-	if (bSignalBarrier)
+	if (bNotifyBarrier)
 	{
-		BarrierEvent.Signal();
+		BarrierEvent.Notify();
 	}
 }
 
diff --git a/Engine/Source/Runtime/Core/Private/Async/ParkingLot.cpp b/Engine/Source/Runtime/Core/Private/Async/ParkingLot.cpp
--- a/Engine/Source/Runtime/Core/Private/Async/ParkingLot.cpp
+++ b/Engine/Source/Runtime/Core/Private/Async/ParkingLot.cpp
@@ -605,7 +605,7 @@ void WakeOne(const void* Address, TFunctionRef<uint64(FWakeState)> OnWakeState)
 		checkSlow(WakeThread->WaitAddress == Address);
 		WakeThread->WakeToken = WakeToken;
 		WakeThread->WaitAddress = nullptr;
-		WakeThread->Event.Signal();
+		WakeThread->Event.Notify();
 	}
 }
 
@@ -648,7 +648,7 @@ uint32 WakeMultiple(const void* const Address, const uint32 WakeCount)
 	{
 		checkSlow(WakeThread->WaitAddress == Address);
 		WakeThread->WaitAddress = nullptr;
-		WakeThread->Event.Signal();
+		WakeThread->Event.Notify();
 	}
 
 	return uint32(WakeThreads.Num());
diff --git a/Engine/Source/Runtime/Core/Private/Async/WordMutex.cpp b/Engine/Source/Runtime/Core/Private/Async/WordMutex.cpp
--- a/Engine/Source/Runtime/Core/Private/Async/WordMutex.cpp
+++ b/Engine/Source/Runtime/Core/Private/Async/WordMutex.cpp
@@ -153,7 +153,7 @@ void FWordMutex::UnlockSlow(UPTRINT CurrentState)
 		}
 
 		// Wake the thread that was at the head of the queue.
-		Head->Event.Signal();
+		Head->Event.Notify();
 		break;
 	}
 }
diff --git a/Engine/Source/Runtime/Core/Private/GenericPlatform/GenericPlatformManualResetEvent.h b/Engine/Source/Runtime/Core/Private/GenericPlatform/GenericPlatformManualResetEvent.h
--- a/Engine/Source/Runtime/Core/Private/GenericPlatform/GenericPlatformManualResetEvent.h
+++ b/Engine/Source/Runtime/Core/Private/GenericPlatform/GenericPlatformManualResetEvent.h
@@ -14,10 +14,10 @@ namespace UE::HAL::Private
 {
 
 /**
- * A manual reset event that supports only one thread waiting and one thread signaling at a time.
+ * A manual reset event that supports only one thread waiting and one thread notifying at a time.
  *
  * Only one waiting thread may call Reset() or the Wait() functions.
- * Only one signaling thread may call Signal() once until the event is reset.
+ * Only one notifying thread may call Notify() once until the event is reset.
  */
 class FGenericPlatformManualResetEvent
 {
@@ -27,7 +27,7 @@ public:
 	FGenericPlatformManualResetEvent& operator=(const FGenericPlatformManualResetEvent&) = delete;
 
 	/**
-	 * Resets the event to permit another Wait/Signal cycle.
+	 * Resets the event to permit another Wait/Notify cycle.
 	 *
 	 * Must only be called by the waiting thread, and only when there is no possibility of waking
 	 * occurring concurrently with the reset.
@@ -38,9 +38,9 @@ public:
 	}
 
 	/**
-	 * Waits for Signal() to be called.
+	 * Waits for Notify() to be called.
 	 *
-	 * Signal() may be called prior to Wait(), and this will return immediately in that case.
+	 * Notify() may be called prior to Wait(), and this will return immediately in that case.
 	 */
 	void Wait()
 	{
@@ -48,12 +48,12 @@ public:
 	}
 
 	/**
-	 * Waits until the wait time for Signal() to be called.
+	 * Waits until the wait time for Notify() to be called.
 	 *
-	 * Signal() may be called prior to WaitUntil(), and this will return immediately in that case.
+	 * Notify() may be called prior to WaitUntil(), and this will return immediately in that case.
 	 *
 	 * @param WaitTime   Absolute time after which waiting is canceled and the thread wakes.
-	 * @return True if Signal() was called before the wait time elapsed, otherwise false.
+	 * @return True if Notify() was called before the wait time elapsed, otherwise false.
 	 */
 	bool WaitUntil(FMonotonicTimePoint WaitTime)
 	{
@@ -72,12 +72,12 @@ public:
 	}
 
 	/**
-	 * Signals the waiting thread.
+	 * Notifies the waiting thread.
 	 *
-	 * Signal() may be called prior to one of the wait functions, and the eventual wait call will
+	 * Notify() may be called prior to one of the wait functions, and the eventual wait call will
 	 * return immediately when that occurs.
 	 */
-	void Signal()
+	void Notify()
 	{
 		{
 			std::unique_lock SelfLock(Lock);
diff --git a/Engine/Source/Runtime/Core/Private/Microsoft/MicrosoftPlatformManualResetEvent.cpp b/Engine/Source/Runtime/Core/Private/Microsoft/MicrosoftPlatformManualResetEvent.cpp
--- a/Engine/Source/Runtime/Core/Private/Microsoft/MicrosoftPlatformManualResetEvent.cpp
+++ b/Engine/Source/Runtime/Core/Private/Microsoft/MicrosoftPlatformManualResetEvent.cpp
@@ -46,7 +46,7 @@ bool FMicrosoftPlatformManualResetEvent::WaitUntil(FMonotonicTimePoint WaitTime)
 	}
 }
 
-void FMicrosoftPlatformManualResetEvent::Signal()
+void FMicrosoftPlatformManualResetEvent::Notify()
 {
 	bWait = false;
 	WakeByAddressSingle((void*)&bWait);
diff --git a/Engine/Source/Runtime/Core/Private/Microsoft/MicrosoftPlatformManualResetEvent.h b/Engine/Source/Runtime/Core/Private/Microsoft/MicrosoftPlatformManualResetEvent.h
--- a/Engine/Source/Runtime/Core/Private/Microsoft/MicrosoftPlatformManualResetEvent.h
+++ b/Engine/Source/Runtime/Core/Private/Microsoft/MicrosoftPlatformManualResetEvent.h
@@ -23,7 +23,7 @@ public:
 
 	UE_API void Wait();
 	UE_API bool WaitUntil(FMonotonicTimePoint WaitTime);
-	UE_API void Signal();
+	UE_API void Notify();
 
 private:
 	bool bWait = true;
diff --git a/Engine/Source/Runtime/Core/Private/Misc/OutputDeviceRedirector.cpp b/Engine/Source/Runtime/Core/Private/Misc/OutputDeviceRedirector.cpp
--- a/Engine/Source/Runtime/Core/Private/Misc/OutputDeviceRedirector.cpp
+++ b/Engine/Source/Runtime/Core/Private/Misc/OutputDeviceRedirector.cpp
@@ -138,7 +138,7 @@ struct FOutputDeviceRedirectorState
 	/** A lock to synchronize access to the thread. */
 	FRWLock ThreadLock;
 
-	/** An event that is signaled when the dedicated primary thread is idle. */
+	/** An event that is notified when the dedicated primary thread is idle. */
 	FEventCount ThreadIdleEvent;
 
 	/** An event to wake the dedicated primary thread to process buffered items. */
@@ -421,7 +421,7 @@ bool FOutputDeviceRedirectorState::TryStopThread()
 {
 	if (FWriteScopeLock ThreadScopeLock(ThreadLock); bThreadStarted.exchange(false, std::memory_order_relaxed))
 	{
-		ThreadWakeEvent.Signal();
+		ThreadWakeEvent.Notify();
 		Thread.Join();
 	}
 	return true;
@@ -450,7 +450,7 @@ void FOutputDeviceRedirectorState::ThreadLoop()
 				FlushBufferedItems();
 			}
 		}
-		ThreadIdleEvent.Signal();
+		ThreadIdleEvent.Notify();
 		ThreadWakeEvent.Wait(Token);
 	}
 }
@@ -561,7 +561,7 @@ void FOutputDeviceRedirector::FlushThreadedLogs(EOutputDeviceRedirectorFlushOpti
 		{
 			TRACE_CPUPROFILER_EVENT_SCOPE(FOutputDeviceRedirector::FlushThreadedLogs);
 			UE::FEventCountToken Token = State->ThreadIdleEvent.PrepareWait();
-			State->ThreadWakeEvent.Signal();
+			State->ThreadWakeEvent.Notify();
 			State->ThreadIdleEvent.Wait(Token);
 		}
 		return;
@@ -655,7 +655,7 @@ void FOutputDeviceRedirector::SerializeRecord(const UE::FLogRecord& Record)
 	// Queue the record to serialize to buffered output devices from the primary thread.
 	if (State->BufferedItems.EnqueueAndReturnWasEmpty(Record))
 	{
-		State->ThreadWakeEvent.Signal();
+		State->ThreadWakeEvent.Notify();
 	}
 }
 
@@ -716,7 +716,7 @@ void FOutputDeviceRedirector::Serialize(const TCHAR* const Data, const ELogVerbo
 	// Queue the line to serialize to buffered output devices from the primary thread.
 	if (State->BufferedItems.EnqueueAndReturnWasEmpty(Data, Category, Verbosity, RealTime))
 	{
-		State->ThreadWakeEvent.Signal();
+		State->ThreadWakeEvent.Notify();
 	}
 }
 
diff --git a/Engine/Source/Runtime/Core/Public/Async/EventCount.h b/Engine/Source/Runtime/Core/Public/Async/EventCount.h
--- a/Engine/Source/Runtime/Core/Public/Async/EventCount.h
+++ b/Engine/Source/Runtime/Core/Public/Async/EventCount.h
@@ -30,7 +30,7 @@ private:
 };
 
 /**
- * A type of event that avoids missed signals by maintaining a signal count.
+ * A type of event that avoids missed notifications by maintaining a notification count.
  *
  * This type of event is suited to waiting on another thread conditionally.
  * Typical usage looks similar to this example:
@@ -46,16 +46,16 @@ private:
  *         Event.Wait(Token);
  *     }
  *
- * On the signaling thread:
+ * On the notifying thread:
  *
  *     ++CurrentValue;
  *     if (CurrentValue == TargetValue)
  *     {
- *         Event.Signal();
+ *         Event.Notify();
  *     }
  *
  * Acquiring a token before checking the condition avoids a race because Wait returns immediately
- * when the token no longer matches the signal count.
+ * when the token no longer matches the notification count.
  */
 template <typename CounterType>
 class TEventCount
@@ -71,9 +71,9 @@ public:
 	/**
 	 * Prepare to wait.
 	 *
-	 * Call this before any logic that must re-execute if the event is signaled in the meantime.
+	 * Call this before any logic that must re-execute if the event is notified in the meantime.
 	 *
-	 * @return A token to pass to one of the wait functions to abort the wait if the event has been signaled since.
+	 * @return A token to pass to one of the wait functions to abort the wait if the event has been notified since.
 	 */
 	inline TEventCountToken<CounterType> PrepareWait()
 	{
@@ -83,7 +83,7 @@ public:
 	}
 
 	/**
-	 * Wait until the event is signaled. Returns immediately if signaled since the token was acquired.
+	 * Wait until the event is notified. Returns immediately if notified since the token was acquired.
 	 *
 	 * @param Compare   A token acquired from PrepareWait() before checking the conditions for this wait.
 	 */
@@ -96,11 +96,11 @@ public:
 	}
 
 	/**
-	 * Wait until the event is signaled. Returns immediately if signaled since the token was acquired.
+	 * Wait until the event is notified. Returns immediately if notified since the token was acquired.
 	 *
 	 * @param Compare    A token acquired from PrepareWait() before checking the conditions for this wait.
 	 * @param WaitTime   Relative time after which waiting is automatically canceled and the thread wakes.
-	 * @return True if the event was signaled before the wait time elapsed, otherwise false.
+	 * @return True if the event was notified before the wait time elapsed, otherwise false.
 	 */
 	inline bool WaitFor(TEventCountToken<CounterType> Compare, FMonotonicTimeSpan WaitTime)
 	{
@@ -112,11 +112,11 @@ public:
 	}
 
 	/**
-	 * Wait until the event is signaled. Returns immediately if signaled since the token was acquired.
+	 * Wait until the event is notified. Returns immediately if notified since the token was acquired.
 	 *
 	 * @param Compare    A token acquired from PrepareWait() before checking the conditions for this wait.
 	 * @param WaitTime   Absolute time after which waiting is automatically canceled and the thread wakes.
-	 * @return True if the event was signaled before the wait time elapsed, otherwise false.
+	 * @return True if the event was notified before the wait time elapsed, otherwise false.
 	 */
 	inline bool WaitUntil(TEventCountToken<CounterType> Compare, FMonotonicTimePoint WaitTime)
 	{
@@ -128,12 +128,12 @@ public:
 	}
 
 	/**
-	 * Signals all waiting threads.
+	 * Notifies all waiting threads.
 	 *
-	 * Any threads that have called PrepareWait() and not yet waited will be signaled immediately
+	 * Any threads that have called PrepareWait() and not yet waited will be notified immediately
 	 * if they do wait on a token from a call to PrepareWait() that preceded this call.
 	 */
-	inline void Signal()
+	inline void Notify()
 	{
 		CounterType Value = Count.load(std::memory_order_relaxed);
 		if ((Value & 1) && Count.compare_exchange_strong(Value, Value + 1, std::memory_order_release))
