Removed FMutex::IsLocked that was added by mistake

#rnx

diff --git a/Engine/Source/Runtime/Core/Public/Async/Mutex.h b/Engine/Source/Runtime/Core/Public/Async/Mutex.h
--- a/Engine/Source/Runtime/Core/Public/Async/Mutex.h
+++ b/Engine/Source/Runtime/Core/Public/Async/Mutex.h
@@ -21,11 +21,6 @@ public:
 	FMutex(const FMutex&) = delete;
 	FMutex& operator=(const FMutex&) = delete;
 
-	inline bool IsLocked() const
-	{
-		return (State.load(std::memory_order_acquire) & IsLockedFlag);
-	}
-
 	inline bool TryLock()
 	{
 		uint8 Expected = State.load(std::memory_order_acquire);
