Remove instrumentation that was added to track down missing callstacks during cook crashes.
#rnx

diff --git a/Engine/Source/Runtime/Core/Private/Windows/WindowsErrorOutputDevice.cpp b/Engine/Source/Runtime/Core/Private/Windows/WindowsErrorOutputDevice.cpp
--- a/Engine/Source/Runtime/Core/Private/Windows/WindowsErrorOutputDevice.cpp
+++ b/Engine/Source/Runtime/Core/Private/Windows/WindowsErrorOutputDevice.cpp
@@ -109,21 +109,10 @@ void FWindowsErrorOutputDevice::HandleError()
 
 	// Trigger the OnSystemFailure hook if it exists
 	// make sure it happens after GIsGuarded is set to 0 in case this hook crashes
-#if !NO_LOGGING
-	FMsg::Logf(__FILE__, __LINE__, LogWindows.GetCategoryName(), ELogVerbosity::Display, TEXT("OnHandleSystemError.Broadcast"));
-#endif
-	GLog->Panic();
 	FCoreDelegates::OnHandleSystemError.Broadcast();
-#if !NO_LOGGING
-	FMsg::Logf(__FILE__, __LINE__, LogWindows.GetCategoryName(), ELogVerbosity::Display, TEXT("OnHandleSystemError.Broadcast complete"));
-#endif
-	GLog->Panic();
 
 	// Dump the error and flush the log. If you change behavior here, you should probably update RenderingThread's __except block in FRenderingThread::Run
 #if !NO_LOGGING
-	FMsg::Logf(__FILE__, __LINE__, LogWindows.GetCategoryName(), ELogVerbosity::Error, TEXT("LogFormattedMessageWithCallstackFlushed"));
-	GLog->Panic();
-	FMsg::Logf(__FILE__, __LINE__, LogWindows.GetCategoryName(), ELogVerbosity::Error, TEXT("LogFormattedMessageWithCallstackNotFlushed"));
 	FDebug::LogFormattedMessageWithCallstack(LogWindows.GetCategoryName(), __FILE__, __LINE__, TEXT("=== Critical error: ==="), GErrorHist, ELogVerbosity::Error);
 #endif
 	GLog->Panic();
