Translate tests from Engine\Source\Runtime\Core\Private\Tests\Internationalization

diff --git a/Engine/Source/Runtime/Core/Private/Tests/Internationalization/ArchiveTest.cpp b/Engine/Source/Runtime/Core/Private/Tests/Internationalization/ArchiveTest.cpp
deleted file mode 100644
diff --git a/Engine/Source/Runtime/Core/Private/Tests/Internationalization/DateTimeFormattingRulesTest.cpp b/Engine/Source/Runtime/Core/Private/Tests/Internationalization/DateTimeFormattingRulesTest.cpp
deleted file mode 100644
diff --git a/Engine/Source/Runtime/Core/Private/Tests/Internationalization/KeysCultureTest.cpp b/Engine/Source/Runtime/Core/Private/Tests/Internationalization/KeysCultureTest.cpp
deleted file mode 100644
diff --git a/Engine/Source/Runtime/Core/Private/Tests/Internationalization/ManifestTest.cpp b/Engine/Source/Runtime/Core/Private/Tests/Internationalization/ManifestTest.cpp
deleted file mode 100644
diff --git a/Engine/Source/Runtime/Core/Private/Tests/Internationalization/MetadataTest.cpp b/Engine/Source/Runtime/Core/Private/Tests/Internationalization/MetadataTest.cpp
deleted file mode 100644
diff --git a/Engine/Source/Runtime/Core/Private/Tests/Internationalization/NumberFormatingRulesTest.cpp b/Engine/Source/Runtime/Core/Private/Tests/Internationalization/NumberFormatingRulesTest.cpp
deleted file mode 100644
diff --git a/Engine/Source/Runtime/Core/Private/Tests/Internationalization/TextTest.cpp b/Engine/Source/Runtime/Core/Private/Tests/Internationalization/TextTest.cpp
deleted file mode 100644
diff --git a/Engine/Source/Runtime/Core/Tests/Internationalization/ArchiveTest.cpp b/Engine/Source/Runtime/Core/Tests/Internationalization/ArchiveTest.cpp
new file mode 100644
--- /dev/null
+++ b/Engine/Source/Runtime/Core/Tests/Internationalization/ArchiveTest.cpp
@@ -0,0 +1,218 @@
+// Copyright Epic Games, Inc. All Rights Reserved.
+
+#if WITH_TESTS
+
+#include "CoreTypes.h"
+#include "Containers/UnrealString.h"
+#include "Misc/AutomationTest.h"
+#include "Internationalization/InternationalizationManifest.h"
+#include "Internationalization/InternationalizationArchive.h"
+#include "Internationalization/InternationalizationMetadata.h"
+
+#include "Tests/TestHarnessAdapter.h"
+
+
+TEST_CASE_NAMED(FArchiveTest, "System::Core::Misc::Internationalization Archive", "[EditorContext][ClientContext][SmokeFilter]")
+{
+	// Key metadata
+	TSharedPtr< FLocMetadataObject > KeyMetadataA = MakeShareable(new FLocMetadataObject);
+	TSharedPtr< FLocMetadataObject > KeyMetadataB = MakeShareable(new FLocMetadataObject);
+
+	// Source metadata
+	TSharedPtr< FLocMetadataObject > SourceMetadataA = MakeShareable(new FLocMetadataObject);
+	TSharedPtr< FLocMetadataObject > SourceMetadataB = MakeShareable(new FLocMetadataObject);
+
+
+	// Setup KeyMetadataA
+	KeyMetadataA->SetStringField(TEXT("Gender"), TEXT("Masculine"));
+	KeyMetadataA->SetStringField(TEXT("Plurality"), TEXT("Singular"));
+	KeyMetadataA->SetStringField(TEXT("TargetGender"), TEXT("Masculine"));
+	KeyMetadataA->SetStringField(TEXT("TargetPlurality"), TEXT("Singular"));
+
+	// Setup KeyMetadataB
+	KeyMetadataB->SetStringField(TEXT("Gender"), TEXT("Masculine"));
+	KeyMetadataB->SetStringField(TEXT("Plurality"), TEXT("Singular"));
+	KeyMetadataB->SetStringField(TEXT("TargetGender"), TEXT("Feminine"));
+	KeyMetadataB->SetStringField(TEXT("TargetPlurality"), TEXT("Singular"));
+
+	// Setup source metadata
+	SourceMetadataA->SetBoolField(TEXT("*IsMature"), false);
+	SourceMetadataB->SetBoolField(TEXT("*IsMature"), true);
+
+	// Setup source item
+	FLocItem SourceA(TEXT("TextA"));
+	SourceA.MetadataObj = MakeShareable(new FLocMetadataObject(*SourceMetadataA));
+
+	FLocItem Translation = SourceA;
+	Translation.Text = TEXT("TranslatedTextA");
+
+	FString TestNamespace = TEXT("TestNamespace");
+	FString SourceAKey = TEXT("TextA");
+
+	// Test entry add
+	{
+		bool TestOptionalTrue = true;
+		bool TestOptionalFalse = false;
+
+		// bIsOptional is not used as a key.  We ensure adding entries where the bIsOptional member is the only difference works as expected.
+		FInternationalizationArchive TestArchive;
+		TestArchive.AddEntry(TestNamespace, SourceAKey, SourceA, Translation, NULL, TestOptionalTrue);
+		// Add duplicate entry that differs in bIsOptional flag.  This add should report success because we already have an entry with matching 
+		//  namespace/source/keymetadata.  Differences in bIsOptional are not taken into consideration. 
+		bool bResult = TestArchive.AddEntry(TestNamespace, SourceAKey, SourceA, Translation, NULL, TestOptionalFalse);
+		CHECK_MESSAGE(TEXT("AddEntry result = true"), bResult);
+
+		// We should only have one entry in the archive
+		int32 EntryCount = 0;
+		for (auto Iter = TestArchive.GetEntriesBySourceTextIterator(); Iter; ++Iter, ++EntryCount);
+		CHECK_MESSAGE(TEXT("EntryCount == 1"), EntryCount == 1);
+
+		// Make sure the original bIsOptional value is not overwritten by our second add.
+		TSharedPtr< FArchiveEntry > FoundEntry = TestArchive.FindEntryByKey(TestNamespace, SourceAKey, NULL);
+		if (!FoundEntry.IsValid())
+		{
+			FAIL_CHECK(TEXT("FArchiveEntry could not find entry using FindEntryByKey."));
+		}
+		else
+		{
+			CHECK_MESSAGE(TEXT("FoundEntry->Namespace == Namespace"), FoundEntry->bIsOptional == TestOptionalTrue);
+		}
+	}
+
+	// Test lookup an entry
+	{
+		{
+			FInternationalizationArchive TestArchive;
+			TestArchive.AddEntry(TestNamespace, SourceAKey, SourceA, Translation, KeyMetadataA, false);
+
+			TSharedPtr< FArchiveEntry > FoundEntry = TestArchive.FindEntryByKey(TestNamespace, SourceAKey, KeyMetadataA);
+			if (!FoundEntry.IsValid())
+			{
+				FAIL_CHECK(TEXT("FArchiveEntry could not find entry using FindEntryByKey."));
+			}
+			else
+			{
+				CHECK_MESSAGE(TEXT("FoundEntry->Namespace == Namespace"), FoundEntry->Namespace == TestNamespace);
+				CHECK_MESSAGE(TEXT("FoundEntry->Source == Source"), FoundEntry->Source == SourceA);
+				CHECK_MESSAGE(TEXT("FoundEntry->Translation == Translation"), FoundEntry->Translation == Translation);
+				if (FoundEntry->KeyMetadataObj == KeyMetadataA)
+				{
+					FAIL_CHECK(TEXT("FArchiveEntry KeyMetadataObj is not a unique object."));
+				}
+				CHECK_MESSAGE(TEXT("FoundEntry->KeyMetadataObj == KeyMetadataA"), *(FoundEntry->KeyMetadataObj) == *(KeyMetadataA));
+			}
+
+			// Passing in a mismatched key metadata will fail to find the entry.  Any fallback logic is intended to happen at runtime
+			FoundEntry = TestArchive.FindEntryByKey(TestNamespace, SourceAKey, NULL);
+			CHECK_FALSE_MESSAGE(TEXT("!FoundEntry.IsValid()"), FoundEntry.IsValid());
+
+			FoundEntry = TestArchive.FindEntryByKey(TestNamespace, SourceAKey, MakeShareable(new FLocMetadataObject()));
+			CHECK_FALSE_MESSAGE(TEXT("!FoundEntry.IsValid()"), FoundEntry.IsValid());
+
+			FoundEntry = TestArchive.FindEntryByKey(TestNamespace, SourceAKey, KeyMetadataB);
+			CHECK_FALSE_MESSAGE(TEXT("!FoundEntry.IsValid()"), FoundEntry.IsValid());
+		}
+
+		// Ensure we can properly lookup entries with non-null but empty key metadata
+		{
+			FInternationalizationArchive TestArchive;
+			TestArchive.AddEntry(TestNamespace, SourceAKey, SourceA, Translation, MakeShareable(new FLocMetadataObject()), false);
+
+			TSharedPtr< FArchiveEntry > FoundEntry = TestArchive.FindEntryByKey(TestNamespace, SourceAKey, NULL);
+			if (!FoundEntry.IsValid())
+			{
+				FAIL_CHECK(TEXT("FArchiveEntry could not find entry using FindEntryByKey."));
+			}
+			else
+			{
+				CHECK_MESSAGE(TEXT("FoundEntry->Namespace == Namespace"), FoundEntry->Namespace == TestNamespace);
+				CHECK_MESSAGE(TEXT("FoundEntry->Source == Source"), FoundEntry->Source == SourceA);
+			}
+		}
+
+		// Ensure we can properly lookup entries with null key metadata
+		{
+			FInternationalizationArchive TestArchive;
+			TestArchive.AddEntry(TestNamespace, SourceAKey, SourceA, Translation, NULL, false);
+
+			TSharedPtr< FArchiveEntry > FoundEntry = TestArchive.FindEntryByKey(TestNamespace, SourceAKey, NULL);
+
+			if (!FoundEntry.IsValid())
+			{
+				FAIL_CHECK(TEXT("FArchiveEntry could not find entry using FindEntryByKey."));
+			}
+			else
+			{
+				CHECK_MESSAGE(TEXT("FoundEntry->Namespace == Namespace"), FoundEntry->Namespace == TestNamespace);
+				CHECK_MESSAGE(TEXT("FoundEntry->Source == Source"), FoundEntry->Source == SourceA);
+			}
+
+			// Test lookup with non-null but empty key metadata
+			FoundEntry = TestArchive.FindEntryByKey(TestNamespace, SourceAKey, MakeShareable(new FLocMetadataObject()));
+			if (!FoundEntry.IsValid())
+			{
+				FAIL_CHECK(TEXT("FArchiveEntry could not find entry using FindEntryByKey."));
+			}
+			else
+			{
+				CHECK_MESSAGE(TEXT("FoundEntry->Namespace == Namespace"), FoundEntry->Namespace == TestNamespace);
+				CHECK_MESSAGE(TEXT("FoundEntry->Source == Source"), FoundEntry->Source == SourceA);
+			}
+		}
+
+		// Ensure lookup where source metadata has * prefixed entries work as expected.  Note: The source metadata object
+		//   supports a * prefix on metadata names which modifies the way we perform metadata comparison(ignores entry type and value, only name is checked)
+		{
+			FLocItem SourceCompare(TEXT("TextA"));
+			SourceCompare.MetadataObj = MakeShareable(new FLocMetadataObject());
+			SourceCompare.MetadataObj->SetStringField(TEXT("*IsMature"), TEXT(""));
+
+			FInternationalizationArchive TestArchive;
+			// Added entry with String *IsMature entry
+			TestArchive.AddEntry(TestNamespace, SourceAKey, SourceCompare, Translation, KeyMetadataA, false);
+
+			// Finding entry using a source that has Boolean *IsMature
+			TSharedPtr< FArchiveEntry > FoundEntry = TestArchive.FindEntryByKey(TestNamespace, SourceAKey, KeyMetadataA);
+			if (!FoundEntry.IsValid())
+			{
+				FAIL_CHECK(TEXT("FArchiveEntry could not find entry using FindEntryByKey."));
+			}
+			else
+			{
+				CHECK_MESSAGE(TEXT("FoundEntry->Namespace == Namespace"), FoundEntry->Namespace == TestNamespace);
+				CHECK_MESSAGE(TEXT("FoundEntry->Source == Source"), FoundEntry->Source == SourceA);
+			}
+
+			// Attempting to add an entry that only differs by a * prefexed source metadata entry value or type will result in success since
+			//  a 'matching' entry already exists in the archive.  We should, however, only have one entry in the archive
+			bool bResult = TestArchive.AddEntry(TestNamespace, SourceAKey, SourceA, Translation, KeyMetadataA, false);
+
+			CHECK_MESSAGE(TEXT("AddEntry result = true"), bResult);
+
+			// We should only have one entry in the archive
+			int32 EntryCount = 0;
+			for (auto Iter = TestArchive.GetEntriesBySourceTextIterator(); Iter; ++Iter, ++EntryCount);
+			CHECK_MESSAGE(TEXT("EntryCount == 1"), EntryCount == 1);
+
+			// Check to see that the original type/value of the * prefixed entry did not get modified
+			FoundEntry = TestArchive.FindEntryByKey(TestNamespace, SourceAKey, KeyMetadataA);
+			if (!FoundEntry.IsValid())
+			{
+				FAIL_CHECK(TEXT("FArchiveEntry could not find entry using FindEntryByKey."));
+			}
+			else
+			{
+				if (FoundEntry->Source.MetadataObj->HasTypedField< ELocMetadataType::String >(TEXT("*IsMature")))
+				{
+					CHECK_MESSAGE(TEXT("Metadata Type == String and Value == Empty string"), FoundEntry->Source.MetadataObj->GetStringField(TEXT("*IsMature")) == TEXT(""));
+				}
+				else
+				{
+					FAIL_CHECK(TEXT("FArchiveEntry * prefixed metadata entry on source object was modified unexpectedly."));
+				}
+			}
+		}
+	}
+}
+
+#endif //WITH_TESTS
\ No newline at end of file
diff --git a/Engine/Source/Runtime/Core/Tests/Internationalization/DateTimeFormattingRulesTest.cpp b/Engine/Source/Runtime/Core/Tests/Internationalization/DateTimeFormattingRulesTest.cpp
new file mode 100644
--- /dev/null
+++ b/Engine/Source/Runtime/Core/Tests/Internationalization/DateTimeFormattingRulesTest.cpp
@@ -0,0 +1,165 @@
+// Copyright Epic Games, Inc. All Rights Reserved.
+
+#include "CoreTypes.h"
+#include "Containers/UnrealString.h"
+#include "Misc/DateTime.h"
+#include "Internationalization/Text.h"
+#include "Internationalization/Culture.h"
+#include "Internationalization/Internationalization.h"
+
+#if WITH_TESTS
+
+#include "Tests/TestHarnessAdapter.h"
+
+class FDateTimeFormattingRulesTestClass {
+public:
+	void DateTimeFormattingRulesTest() {
+		auto TestText = [this](const TCHAR* const Desc, const FText& A, const FText& B)
+		{
+			if (!A.EqualTo(B))
+			{
+				FAIL_CHECK(FString::Printf(TEXT("%s - A=%s B=%s"), Desc, *A.ToString(), *B.ToString()));
+			}
+		};
+		auto TestStr = [this](const TCHAR* const Desc, const FString& A, const FString& B)
+		{
+			if (!A.Equals(B))
+			{
+				FAIL_CHECK(FString::Printf(TEXT("%s - A=%s B=%s"), Desc, *A, *B));
+			}
+		};
+
+		FInternationalization& I18N = FInternationalization::Get();
+
+#if UE_ENABLE_ICU
+		{
+			FInternationalization::FCultureStateSnapshot OriginalCultureState;
+			I18N.BackupCultureState(OriginalCultureState);
+
+			const FDateTime UnixEpoch = FDateTime::FromUnixTimestamp(0);
+			const FDateTime UnixBillennium = FDateTime::FromUnixTimestamp(1000000000);
+			const FDateTime UnixOnes = FDateTime::FromUnixTimestamp(1111111111);
+			const FDateTime UnixDecimalSequence = FDateTime::FromUnixTimestamp(1234567890);
+			const FDateTime YearOne(1, 1, 1, 00, 00, 00, 000);
+			const FDateTime TestDateTime(1990, 6, 13, 12, 34, 56, 789);
+
+			const FDateTime LocalTime = FDateTime::Now();
+			const FDateTime UtcTime = FDateTime::UtcNow();
+
+			if (I18N.SetCurrentCulture("en-US"))
+			{
+				// Unix Time Values via Date Time
+				TestText(TEXT("Testing Unix Epoch"), FText::AsDateTime(UnixEpoch, EDateTimeStyle::Short, EDateTimeStyle::Short, FText::GetInvariantTimeZone()), FText::FromString(TEXT("1/1/70, 12:00 AM")));
+				TestText(TEXT("Testing Unix Epoch"), FText::AsDateTime(UnixEpoch, EDateTimeStyle::Medium, EDateTimeStyle::Medium, FText::GetInvariantTimeZone()), FText::FromString(TEXT("Jan 1, 1970, 12:00:00 AM")));
+				TestText(TEXT("Testing Unix Epoch"), FText::AsDateTime(UnixEpoch, EDateTimeStyle::Long, EDateTimeStyle::Long, FText::GetInvariantTimeZone()), FText::FromString(TEXT("January 1, 1970 at 12:00:00 AM GMT")));
+				TestText(TEXT("Testing Unix Epoch"), FText::AsDateTime(UnixEpoch, EDateTimeStyle::Full, EDateTimeStyle::Full, FText::GetInvariantTimeZone()), FText::FromString(TEXT("Thursday, January 1, 1970 at 12:00:00 AM GMT")));
+
+				TestText(TEXT("Testing Unix Billennium"), FText::AsDateTime(UnixBillennium, EDateTimeStyle::Short, EDateTimeStyle::Short, FText::GetInvariantTimeZone()), FText::FromString(TEXT("9/9/01, 1:46 AM")));
+				TestText(TEXT("Testing Unix Billennium"), FText::AsDateTime(UnixBillennium, EDateTimeStyle::Medium, EDateTimeStyle::Medium, FText::GetInvariantTimeZone()), FText::FromString(TEXT("Sep 9, 2001, 1:46:40 AM")));
+				TestText(TEXT("Testing Unix Billennium"), FText::AsDateTime(UnixBillennium, EDateTimeStyle::Long, EDateTimeStyle::Long, FText::GetInvariantTimeZone()), FText::FromString(TEXT("September 9, 2001 at 1:46:40 AM GMT")));
+				TestText(TEXT("Testing Unix Billennium"), FText::AsDateTime(UnixBillennium, EDateTimeStyle::Full, EDateTimeStyle::Full, FText::GetInvariantTimeZone()), FText::FromString(TEXT("Sunday, September 9, 2001 at 1:46:40 AM GMT")));
+
+				TestText(TEXT("Testing Unix Ones"), FText::AsDateTime(UnixOnes, EDateTimeStyle::Short, EDateTimeStyle::Short, FText::GetInvariantTimeZone()), FText::FromString(TEXT("3/18/05, 1:58 AM")));
+				TestText(TEXT("Testing Unix Ones"), FText::AsDateTime(UnixOnes, EDateTimeStyle::Medium, EDateTimeStyle::Medium, FText::GetInvariantTimeZone()), FText::FromString(TEXT("Mar 18, 2005, 1:58:31 AM")));
+				TestText(TEXT("Testing Unix Ones"), FText::AsDateTime(UnixOnes, EDateTimeStyle::Long, EDateTimeStyle::Long, FText::GetInvariantTimeZone()), FText::FromString(TEXT("March 18, 2005 at 1:58:31 AM GMT")));
+				TestText(TEXT("Testing Unix Ones"), FText::AsDateTime(UnixOnes, EDateTimeStyle::Full, EDateTimeStyle::Full, FText::GetInvariantTimeZone()), FText::FromString(TEXT("Friday, March 18, 2005 at 1:58:31 AM GMT")));
+
+				TestText(TEXT("Testing Unix Decimal Sequence"), FText::AsDateTime(UnixDecimalSequence, EDateTimeStyle::Short, EDateTimeStyle::Short, FText::GetInvariantTimeZone()), FText::FromString(TEXT("2/13/09, 11:31 PM")));
+				TestText(TEXT("Testing Unix Decimal Sequence"), FText::AsDateTime(UnixDecimalSequence, EDateTimeStyle::Medium, EDateTimeStyle::Medium, FText::GetInvariantTimeZone()), FText::FromString(TEXT("Feb 13, 2009, 11:31:30 PM")));
+				TestText(TEXT("Testing Unix Decimal Sequence"), FText::AsDateTime(UnixDecimalSequence, EDateTimeStyle::Long, EDateTimeStyle::Long, FText::GetInvariantTimeZone()), FText::FromString(TEXT("February 13, 2009 at 11:31:30 PM GMT")));
+				TestText(TEXT("Testing Unix Decimal Sequence"), FText::AsDateTime(UnixDecimalSequence, EDateTimeStyle::Full, EDateTimeStyle::Full, FText::GetInvariantTimeZone()), FText::FromString(TEXT("Friday, February 13, 2009 at 11:31:30 PM GMT")));
+
+				TestText(TEXT("Testing Year One"), FText::AsDateTime(YearOne, EDateTimeStyle::Short, EDateTimeStyle::Short, FText::GetInvariantTimeZone()), FText::FromString(TEXT("1/1/01, 12:00 AM")));
+				TestText(TEXT("Testing Year One"), FText::AsDateTime(YearOne, EDateTimeStyle::Medium, EDateTimeStyle::Medium, FText::GetInvariantTimeZone()), FText::FromString(TEXT("Jan 1, 1, 12:00:00 AM")));
+				TestText(TEXT("Testing Year One"), FText::AsDateTime(YearOne, EDateTimeStyle::Long, EDateTimeStyle::Long, FText::GetInvariantTimeZone()), FText::FromString(TEXT("January 1, 1 at 12:00:00 AM GMT")));
+				TestText(TEXT("Testing Year One"), FText::AsDateTime(YearOne, EDateTimeStyle::Full, EDateTimeStyle::Full, FText::GetInvariantTimeZone()), FText::FromString(TEXT("Saturday, January 1, 1 at 12:00:00 AM GMT")));
+
+				// Date Time
+				TestText(TEXT("Testing Date-Time"), FText::AsDateTime(TestDateTime, EDateTimeStyle::Short, EDateTimeStyle::Short, FText::GetInvariantTimeZone()), FText::FromString(TEXT("6/13/90, 12:34 PM")));
+				TestText(TEXT("Testing Date-Time"), FText::AsDateTime(TestDateTime, EDateTimeStyle::Medium, EDateTimeStyle::Medium, FText::GetInvariantTimeZone()), FText::FromString(TEXT("Jun 13, 1990, 12:34:56 PM")));
+				TestText(TEXT("Testing Date-Time"), FText::AsDateTime(TestDateTime, EDateTimeStyle::Long, EDateTimeStyle::Long, FText::GetInvariantTimeZone()), FText::FromString(TEXT("June 13, 1990 at 12:34:56 PM GMT")));
+				TestText(TEXT("Testing Date-Time"), FText::AsDateTime(TestDateTime, EDateTimeStyle::Full, EDateTimeStyle::Full, FText::GetInvariantTimeZone()), FText::FromString(TEXT("Wednesday, June 13, 1990 at 12:34:56 PM GMT")));
+
+				// Checks that the default ICU timezone is set correctly (including DST)
+				TestText(TEXT("Testing Local Time"), FText::AsDateTime(UtcTime, EDateTimeStyle::Short, EDateTimeStyle::Short), FText::AsDateTime(LocalTime, EDateTimeStyle::Short, EDateTimeStyle::Short, FText::GetInvariantTimeZone()));
+			}
+			else
+			{
+				WARN(FString::Printf(TEXT("Internationalization data for %s missing - test is partially disabled."), TEXT("en-US")));
+			}
+
+			if (I18N.SetCurrentCulture("ja-JP"))
+			{
+				// Unix Time Values via Date Time
+				TestText(TEXT("Testing Unix Epoch"), FText::AsDateTime(UnixEpoch, EDateTimeStyle::Short, EDateTimeStyle::Short, FText::GetInvariantTimeZone()), FText::FromString(TEXT("1970/01/01 0:00")));
+				TestText(TEXT("Testing Unix Epoch"), FText::AsDateTime(UnixEpoch, EDateTimeStyle::Medium, EDateTimeStyle::Medium, FText::GetInvariantTimeZone()), FText::FromString(TEXT("1970/01/01 0:00:00")));
+				TestText(TEXT("Testing Unix Epoch"), FText::AsDateTime(UnixEpoch, EDateTimeStyle::Long, EDateTimeStyle::Long, FText::GetInvariantTimeZone()), FText::FromString(TEXT("1970\x5E74") TEXT("1\x6708") TEXT("1\x65E5") TEXT(" 0:00:00 GMT")));
+				TestText(TEXT("Testing Unix Epoch"), FText::AsDateTime(UnixEpoch, EDateTimeStyle::Full, EDateTimeStyle::Full, FText::GetInvariantTimeZone()), FText::FromString(TEXT("1970\x5E74") TEXT("1\x6708") TEXT("1\x65E5\x6728\x66DC\x65E5") TEXT(" ") TEXT("0\x6642") TEXT("00\x5206") TEXT("00\x79D2") TEXT(" GMT")));
+
+				TestText(TEXT("Testing Unix Billennium"), FText::AsDateTime(UnixBillennium, EDateTimeStyle::Short, EDateTimeStyle::Short, FText::GetInvariantTimeZone()), FText::FromString(TEXT("2001/09/09 1:46")));
+				TestText(TEXT("Testing Unix Billennium"), FText::AsDateTime(UnixBillennium, EDateTimeStyle::Medium, EDateTimeStyle::Medium, FText::GetInvariantTimeZone()), FText::FromString(TEXT("2001/09/09 1:46:40")));
+				TestText(TEXT("Testing Unix Billennium"), FText::AsDateTime(UnixBillennium, EDateTimeStyle::Long, EDateTimeStyle::Long, FText::GetInvariantTimeZone()), FText::FromString(TEXT("2001\x5E74") TEXT("9\x6708") TEXT("9\x65E5") TEXT(" 1:46:40 GMT")));
+				TestText(TEXT("Testing Unix Billennium"), FText::AsDateTime(UnixBillennium, EDateTimeStyle::Full, EDateTimeStyle::Full, FText::GetInvariantTimeZone()), FText::FromString(TEXT("2001\x5E74") TEXT("9\x6708") TEXT("9\x65E5\x65E5\x66DC\x65E5") TEXT(" ") TEXT("1\x6642") TEXT("46\x5206") TEXT("40\x79D2") TEXT(" GMT")));
+
+				TestText(TEXT("Testing Unix Ones"), FText::AsDateTime(UnixOnes, EDateTimeStyle::Short, EDateTimeStyle::Short, FText::GetInvariantTimeZone()), FText::FromString(TEXT("2005/03/18 1:58")));
+				TestText(TEXT("Testing Unix Ones"), FText::AsDateTime(UnixOnes, EDateTimeStyle::Medium, EDateTimeStyle::Medium, FText::GetInvariantTimeZone()), FText::FromString(TEXT("2005/03/18 1:58:31")));
+				TestText(TEXT("Testing Unix Ones"), FText::AsDateTime(UnixOnes, EDateTimeStyle::Long, EDateTimeStyle::Long, FText::GetInvariantTimeZone()), FText::FromString(TEXT("2005\x5E74") TEXT("3\x6708") TEXT("18\x65E5") TEXT(" 1:58:31 GMT")));
+				TestText(TEXT("Testing Unix Ones"), FText::AsDateTime(UnixOnes, EDateTimeStyle::Full, EDateTimeStyle::Full, FText::GetInvariantTimeZone()), FText::FromString(TEXT("2005\x5E74") TEXT("3\x6708") TEXT("18\x65E5\x91D1\x66DC\x65E5") TEXT(" ") TEXT("1\x6642") TEXT("58\x5206") TEXT("31\x79D2") TEXT(" GMT")));
+
+				TestText(TEXT("Testing Unix Decimal Sequence"), FText::AsDateTime(UnixDecimalSequence, EDateTimeStyle::Short, EDateTimeStyle::Short, FText::GetInvariantTimeZone()), FText::FromString(TEXT("2009/02/13 23:31")));
+				TestText(TEXT("Testing Unix Decimal Sequence"), FText::AsDateTime(UnixDecimalSequence, EDateTimeStyle::Medium, EDateTimeStyle::Medium, FText::GetInvariantTimeZone()), FText::FromString(TEXT("2009/02/13 23:31:30")));
+				TestText(TEXT("Testing Unix Decimal Sequence"), FText::AsDateTime(UnixDecimalSequence, EDateTimeStyle::Long, EDateTimeStyle::Long, FText::GetInvariantTimeZone()), FText::FromString(TEXT("2009\x5E74") TEXT("2\x6708") TEXT("13\x65E5 23:31:30 GMT")));
+				TestText(TEXT("Testing Unix Decimal Sequence"), FText::AsDateTime(UnixDecimalSequence, EDateTimeStyle::Full, EDateTimeStyle::Full, FText::GetInvariantTimeZone()), FText::FromString(TEXT("2009\x5E74") TEXT("2\x6708") TEXT("13\x65E5\x91D1\x66DC\x65E5") TEXT(" ") TEXT("23\x6642") TEXT("31\x5206") TEXT("30\x79D2") TEXT(" GMT")));
+
+				TestText(TEXT("Testing Year One"), FText::AsDateTime(YearOne, EDateTimeStyle::Short, EDateTimeStyle::Short, FText::GetInvariantTimeZone()), FText::FromString(TEXT("1/01/01 0:00")));
+				TestText(TEXT("Testing Year One"), FText::AsDateTime(YearOne, EDateTimeStyle::Medium, EDateTimeStyle::Medium, FText::GetInvariantTimeZone()), FText::FromString(TEXT("1/01/01 0:00:00")));
+				TestText(TEXT("Testing Year One"), FText::AsDateTime(YearOne, EDateTimeStyle::Long, EDateTimeStyle::Long, FText::GetInvariantTimeZone()), FText::FromString(TEXT("1\x5E74") TEXT("1\x6708") TEXT("1\x65E5") TEXT(" 0:00:00 GMT")));
+				TestText(TEXT("Testing Year One"), FText::AsDateTime(YearOne, EDateTimeStyle::Full, EDateTimeStyle::Full, FText::GetInvariantTimeZone()), FText::FromString(TEXT("1\x5E74") TEXT("1\x6708") TEXT("1\x65E5\x571F\x66DC\x65E5") TEXT(" ") TEXT("0\x6642") TEXT("00\x5206") TEXT("00\x79D2 GMT")));
+
+				// Date Time
+				TestText(TEXT("Testing Date-Time"), FText::AsDateTime(TestDateTime, EDateTimeStyle::Short, EDateTimeStyle::Short, FText::GetInvariantTimeZone()), FText::FromString(TEXT("1990/06/13 12:34")));
+				TestText(TEXT("Testing Date-Time"), FText::AsDateTime(TestDateTime, EDateTimeStyle::Medium, EDateTimeStyle::Medium, FText::GetInvariantTimeZone()), FText::FromString(TEXT("1990/06/13 12:34:56")));
+				TestText(TEXT("Testing Date-Time"), FText::AsDateTime(TestDateTime, EDateTimeStyle::Long, EDateTimeStyle::Long, FText::GetInvariantTimeZone()), FText::FromString(TEXT("1990\x5E74") TEXT("6\x6708") TEXT("13\x65E5") TEXT(" 12:34:56 GMT")));
+				TestText(TEXT("Testing Date-Time"), FText::AsDateTime(TestDateTime, EDateTimeStyle::Full, EDateTimeStyle::Full, FText::GetInvariantTimeZone()), FText::FromString(TEXT("1990\x5E74") TEXT("6\x6708") TEXT("13\x65E5\x6C34\x66DC\x65E5") TEXT(" ") TEXT("12\x6642") TEXT("34\x5206") TEXT("56\x79D2 GMT")));
+
+				// Checks that the default ICU timezone is set correctly (including DST)
+				TestText(TEXT("Testing Local Time"), FText::AsDateTime(UtcTime, EDateTimeStyle::Short, EDateTimeStyle::Short), FText::AsDateTime(LocalTime, EDateTimeStyle::Short, EDateTimeStyle::Short, FText::GetInvariantTimeZone()));
+			}
+			else
+			{
+				WARN(FString::Printf(TEXT("Internationalization data for %s missing - test is partially disabled."), TEXT("ja-JP")));
+			}
+
+			I18N.RestoreCultureState(OriginalCultureState);
+		}
+#else
+		WARN("ICU is disabled thus locale-aware date/time formatting is disabled.");
+#endif
+
+		{
+			FCultureRef InvariantCulture = I18N.GetInvariantCulture();
+			const FDateTime TestDateTime = FDateTime(1991, 6, 21, 9, 30);
+
+			TestStr(TEXT("Testing Custom Date-Time Pattern (FDateTime) - %m/%d/%Y"), TestDateTime.ToFormattedString(TEXT("%m/%d/%Y")), TEXT("06/21/1991"));
+			TestStr(TEXT("Testing Custom Date-Time Pattern (FDateTime) - %A, %B %e, %Y"), TestDateTime.ToFormattedString(TEXT("%A, %B %e, %Y")), TEXT("Friday, June 21, 1991"));
+			TestStr(TEXT("Testing Custom Date-Time Pattern (FDateTime) - %a, %b %e"), TestDateTime.ToFormattedString(TEXT("%a, %b %e")), TEXT("Fri, Jun 21"));
+			TestStr(TEXT("Testing Custom Date-Time Pattern (FDateTime) - %H:%M"), TestDateTime.ToFormattedString(TEXT("%H:%M")), TEXT("09:30"));
+			TestStr(TEXT("Testing Custom Date-Time Pattern (FDateTime) - %l:%M %p"), TestDateTime.ToFormattedString(TEXT("%l:%M %p")), TEXT("9:30 AM"));
+
+			TestStr(TEXT("Testing Custom Date-Time Pattern (FText) - %m/%d/%Y"), FText::AsDateTime(TestDateTime, TEXT("%m/%d/%Y"), FText::GetInvariantTimeZone(), InvariantCulture).ToString(), TEXT("06/21/1991"));
+			TestStr(TEXT("Testing Custom Date-Time Pattern (FText) - %A, %B %e, %Y"), FText::AsDateTime(TestDateTime, TEXT("%A, %B %e, %Y"), FText::GetInvariantTimeZone(), InvariantCulture).ToString(), TEXT("Friday, June 21, 1991"));
+			TestStr(TEXT("Testing Custom Date-Time Pattern (FText) - %a, %b %e"), FText::AsDateTime(TestDateTime, TEXT("%a, %b %e"), FText::GetInvariantTimeZone(), InvariantCulture).ToString(), TEXT("Fri, Jun 21"));
+			TestStr(TEXT("Testing Custom Date-Time Pattern (FText) - %H:%M"), FText::AsDateTime(TestDateTime, TEXT("%H:%M"), FText::GetInvariantTimeZone(), InvariantCulture).ToString(), TEXT("09:30"));
+			TestStr(TEXT("Testing Custom Date-Time Pattern (FText) - %l:%M %p"), FText::AsDateTime(TestDateTime, TEXT("%l:%M %p"), FText::GetInvariantTimeZone(), InvariantCulture).ToString(), TEXT("9:30 AM"));
+		}
+	}
+
+};
+
+TEST_CASE_NAMED(FDateTimeFormattingRulesTest, "System::Core::Misc::DateTime Formatting Rules", "[.][EditorContext][ClientContext][EngineFilter]")
+{
+	FDateTimeFormattingRulesTestClass Instance;
+	Instance.DateTimeFormattingRulesTest();
+}
+
+#endif //WITH_TESTS
\ No newline at end of file
diff --git a/Engine/Source/Runtime/Core/Tests/Internationalization/KeysCultureTest.cpp b/Engine/Source/Runtime/Core/Tests/Internationalization/KeysCultureTest.cpp
new file mode 100644
--- /dev/null
+++ b/Engine/Source/Runtime/Core/Tests/Internationalization/KeysCultureTest.cpp
@@ -0,0 +1,147 @@
+// Copyright Epic Games, Inc. All Rights Reserved.
+
+#include "CoreTypes.h"
+#include "Containers/UnrealString.h"
+#include "CoreGlobals.h"
+#include "Internationalization/Internationalization.h"
+#include "Internationalization/LocTesting.h"
+#include "Internationalization/PolyglotTextData.h"
+#include "Internationalization/Culture.h"
+#include "Internationalization/Cultures/KeysCulture.h"
+
+#if WITH_TESTS && ENABLE_LOC_TESTING
+
+#include "Tests/TestHarnessAdapter.h"
+
+#define LOCTEXT_NAMESPACE "Core.Tests.KeysCultureTests"
+
+/** 
+* A utility struct to ensure that the internationalization system is reverted
+* to its original state by the end of a test. Declare this at the start of tests that involve culture changes.
+*/
+struct FRestoreCultureStateGuard
+{
+	FRestoreCultureStateGuard()
+	{
+		FInternationalization::Get().BackupCultureState(OriginalCultureState);
+	}
+	~FRestoreCultureStateGuard()
+	{
+		FInternationalization::Get().RestoreCultureState(OriginalCultureState);
+	}
+private:
+	FInternationalization::FCultureStateSnapshot OriginalCultureState;
+};
+/**
+* A helper function that generates the expected output string when the keys culture is active.
+* If you change the behavior of what gets displayed by the keys culture, be sure to update this function. FTextLocalizationManager::KeyifyAllDisplayStrings is currently the place that is done.
+*/
+FString CreateExpectedOutputFromPolyglotData(const FPolyglotTextData& InData)
+{
+	FString Key;
+	FString Namespace;
+	InData.GetIdentity(Namespace, Key);
+	// Right now, we display it as key, namespace
+	return FString::Printf(TEXT("%s, %s"), *Key, *Namespace);
+}
+
+TEST_CASE_NAMED(FKeysCultureTest, "System::Core::Misc::KeysCulture", "[.][EditorContext][ClientContext][SmokeFilter]")
+{
+	FRestoreCultureStateGuard RestoreCultureStateGuard;
+	
+	static const FString NativeCulture = TEXT("en");
+	static const FString TargetCulture = TEXT("ja");
+
+	// Test Overview
+	// 1. We set up the test environment and change the culture to the native culture
+	// 2. We register text data that as a native string and a target culture translation.
+	// 3. We register text data that as a native string and no target culture translation. 
+	// 4. We change to the target culture. The text with translation should have its display string the same as the localized string in the target culture. The display string for the text with no translation should match the native string.
+	// 5. We change to the keys culture. Both text with and without translations in a target culture should have their display strings show the localization key and namespace of the text.
+	// 6. We revert back to the native culture. Both text with and without translations should have their display strings the same as their native strings with no change.
+	
+	FInternationalization& I18N = FInternationalization::Get();
+	// We initialize the test environment to the native culture to start from a clean slate.
+	I18N.SetCurrentCulture(NativeCulture);
+	// somehow failed to initialize the test environment to the right culture. Early out.
+	REQUIRE(I18N.GetCurrentCulture()->GetName() == NativeCulture);
+
+#define CREATE_TEST_POLYGLOTDATA(InKey, InSourceString) FPolyglotTextData(ELocalizedTextSourceCategory::Editor, TEXT(LOCTEXT_NAMESPACE), TEXT(InKey), TEXT(InSourceString), NativeCulture)
+	// We use polyglot text data here to patch in our test strings so we don't have to create a separte loclization resource file 
+	// Test1 is for text with a native string and a non-empty translation in a target language
+	// Simulates text with a valid translation editor/game
+	FPolyglotTextData Test1 = CREATE_TEST_POLYGLOTDATA("Test1Key", "Test1 Native");
+	Test1.AddLocalizedString(TargetCulture, TEXT("Test1 Target"));
+
+	// Test2 is for text with a valid native string but a missing translation in a target culture
+	// Simulates text with missing translation in editor/game
+	FPolyglotTextData Test2 = CREATE_TEST_POLYGLOTDATA("Test2Key", "Test2 Native");
+#undef CREATE_TEST_POLYGLOTDATA
+
+	// GetText() already handles registration of the polyglot text data with the FTextLocalizationManager internally and adds these strings to the display string table.
+	FText Test1Text = Test1.GetText();
+	FText Test2Text = Test2.GetText();
+	
+	FString Test1NativeCultureDisplayString = Test1Text.ToString();
+	// Test1 display string should have the same value as the native string. Otherwise it could be a problem with registering the text. The set up of the test failed. Early out.
+	REQUIRE(Test1NativeCultureDisplayString == Test1.GetNativeString());
+
+	FString Test2NativeCultureDisplayString = Test2Text.ToString();
+	// Same as above.
+	REQUIRE(Test2NativeCultureDisplayString == Test2.GetNativeString());
+
+		// Change to target culture. 
+	I18N.SetCurrentCulture(TargetCulture);
+	// Test 1 - We should be able to change from the native culture to the target culture 
+	// just early out. Doesn't make sense to keep testing.
+	REQUIRE_MESSAGE(TEXT("Changing from native culture to target culture."), I18N.GetCurrentCulture()->GetName() == TargetCulture);
+
+	// Test 2 - The display string for the text with translation should match the localized string in the target culture 
+	FString Test1TargetCultureDisplayString = Test1Text.ToString();
+	FString Test1TargetCultureString;
+	Test1.GetLocalizedString(TargetCulture, Test1TargetCultureString);
+	CHECK_MESSAGE(TEXT("Text with translation display is the same as the localized string in target culture."), Test1TargetCultureDisplayString == Test1TargetCultureString);
+
+	// Test 3 - The text with no translation should have no retrievable text from it spolyglot data in the target culture.
+	// Test 4 - The display string for the text with no translation should match the native string while in the target culture. Text with no translations for a culture should fall back to the native string. 
+	FString Test2TargetCultureDisplayString = Test2Text.ToString();
+	FString Test2TargetCultureString;
+	Test2.GetLocalizedString(TargetCulture, Test2TargetCultureString);
+	CHECK_MESSAGE(TEXT("Text without translation should not have a localized string in the target culture."), Test2TargetCultureString.IsEmpty());
+	CHECK_MESSAGE(TEXT("Text without translation display string is the same as the native string in target culture."), Test2TargetCultureDisplayString == Test2.GetNativeString());
+	
+	// Change to the keys culture
+	// This culture should output all the display strings as the localization key and namespace of the text.
+	// There will be logs that say the localization resource files associated with the keys culture can't be read. That is normal and can be ignored.
+	I18N.SetCurrentCulture(FKeysCulture::StaticGetName());
+	// Test 5 - Can we successfully change from the native culture to the keys culture 
+	// just early out. Doesn't make sense to keep testing.
+	REQUIRE_MESSAGE(TEXT("Changing from native culture to keys culture."), I18N.GetCurrentCulture()->GetName() == FKeysCulture::StaticGetName());
+
+	// Test 6 - The display string for the text with translation should be the localization key and localization namespace for the text in the keys culture 
+	FString Test1KeysCultureDisplayString = Test1Text.ToString();
+	CHECK_MESSAGE(TEXT("Text with translation now has display string that matches its localization key and namespace in keys culture."), Test1KeysCultureDisplayString == CreateExpectedOutputFromPolyglotData(Test1));
+	
+	// Test 7 - The display string for the text with no translation should be its localization key and localization namespace in the keys culture.
+	FString Test2KeysCultureDisplayString = Test2Text.ToString();
+	CHECK_MESSAGE(TEXT("Text with blank translation has display string that matches its localization key and namespace in keys culture."), Test2KeysCultureDisplayString == CreateExpectedOutputFromPolyglotData(Test2));
+
+	// Revert to native culture. All of the text should still show up as the originals.
+	I18N.SetCurrentCulture(NativeCulture);
+	// Test 8 - Can we change from the keys culture back to the native culture 
+	// early out same as above 
+	REQUIRE_MESSAGE(TEXT("Able to revert from keys culture to native culture.."), I18N.GetCurrentCulture()->GetName() == NativeCulture);
+
+	// Test 9 - The display string for the text with translation should be exactly the same as its original native string before changing to the keys culture
+	FString Test1NativeCultureRevertedDisplayString = Test1Text.ToString();
+	CHECK_MESSAGE(TEXT("Text with translation has display string that is the same as the original native string in native culture."), Test1NativeCultureRevertedDisplayString == Test1.GetNativeString());
+
+	// Test 10 - The display string for the text with no translation should be exactly the same as its original native string before changing to the keys culture
+	FString Test2NativeCultureRevertedDisplayString = Test2Text.ToString();
+	CHECK_MESSAGE(TEXT("Text without translation has display string the same as its native string in native culture."), Test2NativeCultureRevertedDisplayString == Test2.GetNativeString());
+
+// @TODOLocalization: Unregister the polyglot data once the functionality is implemented.
+}
+
+#undef LOCTEXT_NAMESPACE 
+#endif // WITH_TESTS && ENABLE_LOC_TESTING
diff --git a/Engine/Source/Runtime/Core/Tests/Internationalization/ManifestTest.cpp b/Engine/Source/Runtime/Core/Tests/Internationalization/ManifestTest.cpp
new file mode 100644
--- /dev/null
+++ b/Engine/Source/Runtime/Core/Tests/Internationalization/ManifestTest.cpp
@@ -0,0 +1,877 @@
+// Copyright Epic Games, Inc. All Rights Reserved.
+
+#include "CoreTypes.h"
+#include "Containers/UnrealString.h"
+#include "Internationalization/InternationalizationManifest.h"
+#include "Internationalization/InternationalizationMetadata.h"
+
+#if WITH_TESTS
+
+#include "Tests/TestHarnessAdapter.h"
+
+// Helper function to count the number of entries in a manifest
+int32 CountManifestEntries( const FInternationalizationManifest& Manifest )
+{
+	return Manifest.GetNumEntriesByKey();
+}
+
+TEST_CASE_NAMED(FLocContextTest, "System::Core::Misc::Internationalization Context", "[EditorContext][ClientContext][SmokeFilter]")
+{
+	// Key metadata
+	TSharedPtr< FLocMetadataObject > KeyMetadataA = MakeShareable( new FLocMetadataObject );
+	TSharedPtr< FLocMetadataObject > KeyMetadataB = MakeShareable( new FLocMetadataObject );
+
+	// Info metadata
+	TSharedPtr< FLocMetadataObject > InfoMetadataA = MakeShareable( new FLocMetadataObject );
+	TSharedPtr< FLocMetadataObject > InfoMetadataB = MakeShareable( new FLocMetadataObject );
+
+	// Source metadata
+	TSharedPtr< FLocMetadataObject > SourceMetadataA = MakeShareable( new FLocMetadataObject );
+	TSharedPtr< FLocMetadataObject > SourceMetadataB = MakeShareable( new FLocMetadataObject );
+
+
+	// Setup KeyMetadataA
+	KeyMetadataA->SetStringField( TEXT("Gender"			),	TEXT("Masculine")	);
+	KeyMetadataA->SetStringField( TEXT("Plurality"		),	TEXT("Singular")	);
+	KeyMetadataA->SetStringField( TEXT("TargetGender"	),	TEXT("Masculine")	);
+	KeyMetadataA->SetStringField( TEXT("TargetPlurality"),	TEXT("Singular")	);
+
+	// Setup KeyMetadataB
+	KeyMetadataB->SetStringField( TEXT("Gender"			),	TEXT("Masculine")	);
+	KeyMetadataB->SetStringField( TEXT("Plurality"		),	TEXT("Singular")	);
+	KeyMetadataB->SetStringField( TEXT("TargetGender"	),	TEXT("Feminine")	);
+	KeyMetadataB->SetStringField( TEXT("TargetPlurality"),	TEXT("Singular")	);
+
+	// Setup source metadata
+	SourceMetadataA->SetBoolField( TEXT("*IsMature"), false );
+	SourceMetadataB->SetBoolField( TEXT("*IsMature"), true );
+
+	// Set InfoMetadataA
+	InfoMetadataA->SetStringField( TEXT("VoiceActorDirection"), TEXT("Go big or go home!") );
+
+	// Test FManifestContext
+	{
+		FManifestContext ContextA;
+		ContextA.Key			= TEXT("KeyA");
+		ContextA.SourceLocation = TEXT("SourceLocationA");
+		ContextA.InfoMetadataObj = MakeShareable( new FLocMetadataObject(*InfoMetadataA) );
+		ContextA.KeyMetadataObj = MakeShareable( new FLocMetadataObject(*KeyMetadataA) );
+
+		FManifestContext ContextB;
+		ContextB.Key			= TEXT("KeyB");
+		ContextB.SourceLocation = TEXT("SourceLocationB");
+		ContextB.InfoMetadataObj = MakeShareable( new FLocMetadataObject(*InfoMetadataB) );
+		ContextB.KeyMetadataObj = MakeShareable( new FLocMetadataObject(*KeyMetadataB) );
+
+
+		// Test copy ctor
+		{ //-V760
+			FManifestContext ContextAClone = ContextA;
+
+			if (ContextAClone.InfoMetadataObj == ContextA.InfoMetadataObj)
+			{
+				FAIL_CHECK(TEXT("FManifestContext InfoMetadataObj and its Clone are not unique objects."));
+			}
+			if (ContextAClone.KeyMetadataObj == ContextA.KeyMetadataObj)
+			{
+				FAIL_CHECK(TEXT("FManifestContext KeyMetadataObj and its Clone are not unique objects."));
+			}
+
+			CHECK_MESSAGE(TEXT("ContextAClone.Key == ContextA.Key"), ContextAClone.Key == ContextA.Key );
+			CHECK_MESSAGE(TEXT("ContextAClone.SourceLocation == ContextA.SourceLocation"), ContextAClone.SourceLocation == ContextA.SourceLocation );
+			CHECK_MESSAGE(TEXT("ContextAClone.bIsOptional == ContextA.bIsOptional"), ContextAClone.bIsOptional == ContextA.bIsOptional );
+			CHECK_MESSAGE(TEXT("ContextAClone.InfoMetadataObj == ContextA.InfoMetadataObj"), *(ContextAClone.InfoMetadataObj) == *(ContextA.InfoMetadataObj) );
+			CHECK_MESSAGE(TEXT("ContextAClone.KeyMetadataObj == ContextA.KeyMetadataObj"), *(ContextAClone.KeyMetadataObj) == *(ContextA.KeyMetadataObj) );
+
+			CHECK_MESSAGE(TEXT("ContextAClone == ContextA"), ContextAClone == ContextA );
+			CHECK_FALSE_MESSAGE(TEXT("ContextAClone < ContextA"), ContextAClone < ContextA );
+
+		}
+
+		// Test assignment operator
+		{
+			FManifestContext ContextAClone = ContextA;
+
+			if (ContextAClone.InfoMetadataObj == ContextA.InfoMetadataObj)
+			{
+				FAIL_CHECK(TEXT("FManifestContext InfoMetadataObj and its Clone are not unique objects."));
+			}
+			if (ContextAClone.KeyMetadataObj == ContextA.KeyMetadataObj)
+			{
+				FAIL_CHECK(TEXT("FManifestContext KeyMetadataObj and its Clone are not unique objects."));
+			}
+
+			CHECK_MESSAGE(TEXT("ContextAClone.Key == ContextA.Key"), ContextAClone.Key == ContextA.Key );
+			CHECK_MESSAGE(TEXT("ContextAClone.SourceLocation == ContextA.SourceLocation"), ContextAClone.SourceLocation == ContextA.SourceLocation );
+			CHECK_MESSAGE(TEXT("ContextAClone.bIsOptional == ContextA.bIsOptional"), ContextAClone.bIsOptional == ContextA.bIsOptional );
+			CHECK_MESSAGE(TEXT("ContextAClone.InfoMetadataObj == ContextA.InfoMetadataObj"), *(ContextAClone.InfoMetadataObj) == *(ContextA.InfoMetadataObj) );
+			CHECK_MESSAGE(TEXT("ContextAClone.KeyMetadataObj == ContextA.KeyMetadataObj"), *(ContextAClone.KeyMetadataObj) == *(ContextA.KeyMetadataObj) );
+
+			CHECK_MESSAGE(TEXT("ContextAClone == ContextA"), ContextAClone == ContextA );
+			CHECK_FALSE_MESSAGE(TEXT("ContextAClone < ContextA"), ContextAClone < ContextA );
+		}
+
+		// Test comparison operator
+		{
+			// Key and KeyMetadataObj members should be the only items that are taken into account when comparing
+			FManifestContext ContextAClone = ContextA;
+			CHECK_MESSAGE(TEXT("ContextAClone == ContextA"), ContextAClone == ContextA );
+
+			// Arbitrarily change all the non-important members
+			ContextAClone.SourceLocation = ContextA.SourceLocation + TEXT("New");
+			ContextAClone.bIsOptional = !ContextA.bIsOptional;
+			ContextAClone.InfoMetadataObj.Reset();
+			ContextAClone.InfoMetadataObj = MakeShareable( new FLocMetadataObject( *InfoMetadataB ) );
+			CHECK_MESSAGE(TEXT("ContextAClone == ContextA"), ContextAClone == ContextA );
+
+			// Changing the key in any way will cause comparison to fail
+			ContextAClone.Key = ContextAClone.Key.GetString() + TEXT("New");
+			CHECK_FALSE_MESSAGE(TEXT("ContextAClone != ContextA"), ContextAClone == ContextA );
+
+			// Reset and test KeyMetadataObj change to one of the value entries
+			ContextAClone = ContextA;
+			ContextAClone.KeyMetadataObj->SetStringField( TEXT("TargetPlurality"), TEXT("Plural") );
+			CHECK_FALSE_MESSAGE(TEXT("ContextAClone != ContextA"), ContextAClone == ContextA );
+
+			// Reset and test addition of entry to KeyMetadataObj
+			ContextAClone = ContextA;
+			ContextAClone.KeyMetadataObj->SetStringField( TEXT("NewField"), TEXT("NewFieldValue") );
+			CHECK_FALSE_MESSAGE(TEXT("ContextAClone != ContextA"), ContextAClone == ContextA );
+
+			// Reset and test removal of entry from KeyMetadataObj
+			ContextAClone = ContextA;
+			ContextAClone.KeyMetadataObj->RemoveField( TEXT("TargetPlurality") );
+			CHECK_FALSE_MESSAGE(TEXT("ContextAClone != ContextA"), ContextAClone == ContextA );
+
+			// Context with valid but empty KeyMetadataObject should be equivalent to Context with null KeyMetadataObject
+			FManifestContext ContextEmptyA;
+			FManifestContext ContextEmptyB;
+			ContextEmptyB.KeyMetadataObj = MakeShareable( new FLocMetadataObject );
+			CHECK_MESSAGE(TEXT("ContextEmptyA == ContextEmptyB"), ContextEmptyA == ContextEmptyB );
+		}
+
+		// Testing less than operator
+		{
+			CHECK_MESSAGE(TEXT("ContextA < ContextB"), ContextA < ContextB );
+
+			FManifestContext ContextAClone = ContextA;
+
+			// Differences in Key
+			CHECK_FALSE_MESSAGE(TEXT("ContextA < ContextAClone"), ContextA < ContextAClone );
+			ContextAClone.Key = ContextAClone.Key.GetString() + TEXT("A");
+			//currently failing CHECK_MESSAGE( TEXT("ContextA < ContextAClone"), ContextA < ContextAClone );
+
+			// Adding new key metadata entry that will appear before other entries
+			ContextAClone = ContextA;
+			ContextAClone.KeyMetadataObj->SetStringField( TEXT("ANewKey"), TEXT("ANewValue") );
+			CHECK_MESSAGE(TEXT("ContextAClone < ContextA"), ContextAClone < ContextA );
+
+			// Adding new key metadata entry that will appear after other entries
+			ContextAClone = ContextA;
+			ContextAClone.KeyMetadataObj->SetStringField( TEXT("ZNewKey"), TEXT("ZNewValue") );
+			//currently failing CHECK_MESSAGE( TEXT("ContextA < ContextAClone"), ContextA < ContextAClone );
+
+			// Removing a key metadata entry
+			ContextAClone = ContextA;
+			ContextAClone.KeyMetadataObj->RemoveField( TEXT("TargetPlurality") ) ;
+			CHECK_MESSAGE(TEXT("ContextAClone < ContextA"), ContextAClone < ContextA );
+
+			// Changing a key metadata entry value
+			ContextAClone = ContextA;
+			ContextAClone.KeyMetadataObj->SetStringField( TEXT("TargetPlurality"), TEXT("A") ) ;
+			CHECK_MESSAGE(TEXT("ContextAClone < ContextA"), ContextAClone < ContextA );
+
+			FManifestContext ContextEmptyA;
+			FManifestContext ContextEmptyB;
+			CHECK_FALSE_MESSAGE( TEXT("ContextEmptyA < ContextEmptyB"), ContextEmptyA < ContextEmptyB );
+			ContextEmptyB.KeyMetadataObj = MakeShareable( new FLocMetadataObject );
+			CHECK_FALSE_MESSAGE( TEXT("ContextEmptyA < ContextEmptyB"), ContextEmptyA < ContextEmptyB );
+			CHECK_FALSE_MESSAGE( TEXT("ContextEmptyB < ContextEmptyA"), ContextEmptyB < ContextEmptyA );
+			ContextEmptyB.KeyMetadataObj->SetStringField(TEXT("AMetadataKey"), TEXT("AMetadataValue"));
+			CHECK_MESSAGE( TEXT("ContextEmptyA < ContextEmptyB"), ContextEmptyA < ContextEmptyB );
+		}
+	}
+}
+
+TEST_CASE_NAMED(FLocItemTest, "System::Core::Misc::Internationalization LocItem", "[EditorContext][ClientContext][SmokeFilter]")
+{
+	// Key metadata
+	TSharedPtr< FLocMetadataObject > KeyMetadataA = MakeShareable( new FLocMetadataObject );
+	TSharedPtr< FLocMetadataObject > KeyMetadataB = MakeShareable( new FLocMetadataObject );
+
+	// Info metadata
+	TSharedPtr< FLocMetadataObject > InfoMetadataA = MakeShareable( new FLocMetadataObject );
+	TSharedPtr< FLocMetadataObject > InfoMetadataB = MakeShareable( new FLocMetadataObject );
+
+	// Source metadata
+	TSharedPtr< FLocMetadataObject > SourceMetadataA = MakeShareable( new FLocMetadataObject );
+	TSharedPtr< FLocMetadataObject > SourceMetadataB = MakeShareable( new FLocMetadataObject );
+
+
+	// Setup KeyMetadataA
+	KeyMetadataA->SetStringField( TEXT("Gender"			),	TEXT("Masculine")	);
+	KeyMetadataA->SetStringField( TEXT("Plurality"		),	TEXT("Singular")	);
+	KeyMetadataA->SetStringField( TEXT("TargetGender"	),	TEXT("Masculine")	);
+	KeyMetadataA->SetStringField( TEXT("TargetPlurality"),	TEXT("Singular")	);
+
+	// Setup KeyMetadataB
+	KeyMetadataB->SetStringField( TEXT("Gender"			),	TEXT("Masculine")	);
+	KeyMetadataB->SetStringField( TEXT("Plurality"		),	TEXT("Singular")	);
+	KeyMetadataB->SetStringField( TEXT("TargetGender"	),	TEXT("Feminine")	);
+	KeyMetadataB->SetStringField( TEXT("TargetPlurality"),	TEXT("Singular")	);
+
+	// Setup source metadata
+	SourceMetadataA->SetBoolField( TEXT("*IsMature"), false );
+	SourceMetadataB->SetBoolField( TEXT("*IsMature"), true );
+
+	// Set InfoMetadataA
+	InfoMetadataA->SetStringField( TEXT("VoiceActorDirection"), TEXT("Go big or go home!") );
+
+
+	// Test FLocItem
+	{
+		FLocItem LocItemA( TEXT("TextA") );
+		LocItemA.MetadataObj = MakeShareable( new FLocMetadataObject(*SourceMetadataA) );
+
+		FLocItem LocItemB( TEXT("TextB") );
+		LocItemB.MetadataObj = MakeShareable( new FLocMetadataObject(*SourceMetadataB) );
+
+		// Test copy ctor
+		{
+			FLocItem LocItemAClone = LocItemA;
+
+			if (LocItemAClone.MetadataObj == LocItemA.MetadataObj)
+			{
+				FAIL_CHECK(TEXT("FLocItem MetadataObj and its Clone are not unique objects."));
+			}
+			CHECK_MESSAGE(TEXT("LocItemAClone.Text == LocItemA.Text"), LocItemAClone.Text == LocItemA.Text );
+			CHECK_MESSAGE(TEXT("LocItemAClone.MetadataObj == LocItemA.MetadataObj"), *(LocItemAClone.MetadataObj) == *(LocItemA.MetadataObj) );
+
+			CHECK_MESSAGE(TEXT("LocItemAClone == LocItemA"), LocItemAClone == LocItemA );
+			CHECK_FALSE_MESSAGE(TEXT("LocItemAClone < LocItemA"), LocItemAClone < LocItemA );
+
+		}
+
+		// Test assignment operator
+		{
+			FLocItem LocItemAClone( TEXT("New") );
+
+			LocItemAClone = LocItemA;
+
+			if (LocItemAClone.MetadataObj == LocItemA.MetadataObj)
+			{
+				FAIL_CHECK(TEXT("FLocItem MetadataObj and its Clone are not unique objects."));
+			}
+			CHECK_MESSAGE (TEXT("LocItemAClone.Text == LocItemA.Text"), LocItemAClone.Text == LocItemA.Text );
+			CHECK_MESSAGE(TEXT("LocItemAClone.MetadataObj == LocItemA.MetadataObj"), *(LocItemAClone.MetadataObj) == *(LocItemA.MetadataObj) );
+
+			CHECK_MESSAGE(TEXT("LocItemAClone == LocItemA"), LocItemAClone == LocItemA );
+			CHECK_FALSE_MESSAGE(TEXT("LocItemAClone < LocItemA"), LocItemAClone < LocItemA );
+			CHECK_FALSE_MESSAGE(TEXT("LocItemA < LocItemAClone"), LocItemA < LocItemAClone );
+		}
+
+		// Test comparison operator
+		{
+			// Text and MetadataObj members should both be taken into account when comparing.  Note, Metadata supports a special * name prefix that causes the type 
+			//   and value of the metadata to be ignored when performing comparisons
+			FLocItem LocItemAClone = LocItemA;
+			CHECK_MESSAGE(TEXT("LocItemAClone == LocItemA"), LocItemAClone == LocItemA );
+
+			// Metadata with * prefix does not impact comparison but both FLocItems need a metadata where the name matches( type and value can be different )
+			FLocItem LocItemAClone2 = LocItemA;
+			LocItemAClone.MetadataObj->SetStringField( TEXT("*NewNonCompare"), TEXT("NewNonCompareValue") );
+			CHECK_FALSE_MESSAGE(TEXT("LocItemAClone != LocItemAClone2"), LocItemAClone == LocItemAClone2 );
+
+			LocItemAClone2.MetadataObj->SetStringField( TEXT("*NewNonCompare"), TEXT("NewNonCompareValue2") );
+			CHECK_MESSAGE(TEXT("LocItemAClone == LocItemAClone2"), LocItemAClone == LocItemAClone2 );
+
+			// Changing the text in any way will cause comparison to fail
+			LocItemAClone = LocItemA;
+			LocItemAClone.Text = LocItemAClone.Text + TEXT("New");
+			CHECK_FALSE_MESSAGE(TEXT("LocItemAClone != LocItemA"), LocItemAClone == LocItemA );
+
+			// LocItem with valid but empty KeyMetadataObject should be equivalent to LocItem with null KeyMetadataObject
+			FLocItem LocItemEmptyA(TEXT("TestText"));
+			FLocItem LocItemEmptyB(TEXT("TestText"));
+			LocItemEmptyB.MetadataObj = MakeShareable( new FLocMetadataObject );
+			CHECK_MESSAGE(TEXT("LocItemEmptyA == LocItemEmptyB"), LocItemEmptyA == LocItemEmptyB );
+		}
+
+		// Testing less than operator
+		{
+			CHECK_MESSAGE(TEXT("LocItemA < LocItemB"), LocItemA < LocItemB );
+			CHECK_FALSE_MESSAGE(TEXT("LocItemB < LocItemA"), LocItemB < LocItemA );
+
+			FLocItem LocItemAClone = LocItemA;
+
+			// Differences in Text
+			CHECK_FALSE_MESSAGE(TEXT("LocItemA < LocItemAClone"), LocItemA < LocItemAClone );
+			LocItemAClone.Text = LocItemAClone.Text + TEXT("A");
+			CHECK_FALSE_MESSAGE(TEXT("LocItemAClone < LocItemA"), LocItemAClone < LocItemA );
+			//currently failing CHECK_MESSAGE( TEXT("LocItemA < LocItemAClone"), LocItemA < LocItemAClone );
+
+			// Adding new key metadata entry
+			LocItemAClone = LocItemA;
+			LocItemAClone.MetadataObj->SetStringField( TEXT("ANewKey"), TEXT("ANewValue") );
+			//currently failing CHECK_MESSAGE( TEXT("LocItemA < LocItemAClone"), LocItemA < LocItemAClone );
+			CHECK_FALSE_MESSAGE(TEXT("LocItemAClone < LocItemA"), LocItemAClone < LocItemA );
+
+			// Removing a key metadata entry
+			LocItemAClone = LocItemA;
+			LocItemAClone.MetadataObj->RemoveField( TEXT("*IsMature") );
+			CHECK_MESSAGE(TEXT("LocItemAClone < LocItemA"), LocItemAClone < LocItemA );
+			CHECK_FALSE_MESSAGE(TEXT("LocItemA < LocItemAClone"), LocItemA < LocItemAClone );
+
+			// Changing a key metadata entry value
+			LocItemAClone = LocItemA;
+			LocItemAClone.MetadataObj->SetBoolField( TEXT("*IsMature"),true );
+			//currently failing CHECK_MESSAGE( TEXT("LocItemA < LocItemAClone"), LocItemA < LocItemAClone );
+			CHECK_FALSE_MESSAGE(TEXT("LocItemAClone < LocItemA"), LocItemAClone < LocItemA );
+
+			// Test null and non-null but empty Metadata
+			FLocItem LocItemEmptyA(TEXT("SameText") );
+			FLocItem LocItemEmptyB(TEXT("SameText") );
+			CHECK_FALSE_MESSAGE(TEXT("LocItemEmptyA < LocItemEmptyB"), LocItemEmptyA < LocItemEmptyB );
+			LocItemEmptyB.MetadataObj = MakeShareable( new FLocMetadataObject );
+			CHECK_FALSE_MESSAGE(TEXT("LocItemEmptyA < LocItemEmptyB"), LocItemEmptyA < LocItemEmptyB );
+			CHECK_FALSE_MESSAGE(TEXT("LocItemEmptyB < LocItemEmptyA"), LocItemEmptyB < LocItemEmptyA );
+			LocItemEmptyB.MetadataObj->SetStringField(TEXT("AMetadataKey"), TEXT("AMetadataValue"));
+			CHECK_MESSAGE(TEXT("LocItemEmptyA < LocItemEmptyB"), LocItemEmptyA < LocItemEmptyB );
+		}
+	}
+}
+
+TEST_CASE_NAMED(FManifestTest, "System::Core::Misc::Internationalization Manifest", "[EditorContext][ClientContext][SmokeFilter]")
+{
+	// Key metadata
+	TSharedPtr< FLocMetadataObject > KeyMetadataA = MakeShareable( new FLocMetadataObject );
+	TSharedPtr< FLocMetadataObject > KeyMetadataB = MakeShareable( new FLocMetadataObject );
+
+	// Info metadata
+	TSharedPtr< FLocMetadataObject > InfoMetadataA = MakeShareable( new FLocMetadataObject );
+	TSharedPtr< FLocMetadataObject > InfoMetadataB = MakeShareable( new FLocMetadataObject );
+
+	// Source metadata
+	TSharedPtr< FLocMetadataObject > SourceMetadataA = MakeShareable( new FLocMetadataObject );
+	TSharedPtr< FLocMetadataObject > SourceMetadataB = MakeShareable( new FLocMetadataObject );
+
+
+	// Setup KeyMetadataA
+	KeyMetadataA->SetStringField( TEXT("Gender"			),	TEXT("Masculine")	);
+	KeyMetadataA->SetStringField( TEXT("Plurality"		),	TEXT("Singular")	);
+	KeyMetadataA->SetStringField( TEXT("TargetGender"	),	TEXT("Masculine")	);
+	KeyMetadataA->SetStringField( TEXT("TargetPlurality"),	TEXT("Singular")	);
+
+	// Setup KeyMetadataB
+	KeyMetadataB->SetStringField( TEXT("Gender"			),	TEXT("Masculine")	);
+	KeyMetadataB->SetStringField( TEXT("Plurality"		),	TEXT("Singular")	);
+	KeyMetadataB->SetStringField( TEXT("TargetGender"	),	TEXT("Feminine")	);
+	KeyMetadataB->SetStringField( TEXT("TargetPlurality"),	TEXT("Singular")	);
+
+	// Setup source metadata
+	SourceMetadataA->SetBoolField( TEXT("*IsMature"), false );
+	SourceMetadataB->SetBoolField( TEXT("*IsMature"), true );
+
+	// Set InfoMetadataA
+	InfoMetadataA->SetStringField( TEXT("VoiceActorDirection"), TEXT("Go big or go home!") );
+
+	// Test FInternationalizationManifest
+	{
+		FManifestContext ContextA;
+		ContextA.Key			= TEXT("KeyA");
+		ContextA.SourceLocation = TEXT("SourceLocationA");
+		ContextA.InfoMetadataObj = MakeShareable( new FLocMetadataObject(*InfoMetadataA) );
+		ContextA.KeyMetadataObj = MakeShareable( new FLocMetadataObject(*KeyMetadataA) );
+
+		FManifestContext ContextB;
+		ContextB.Key			= TEXT("KeyB");
+		ContextB.SourceLocation = TEXT("SourceLocationB");
+		ContextB.InfoMetadataObj = MakeShareable( new FLocMetadataObject(*InfoMetadataB) );
+		ContextB.KeyMetadataObj = MakeShareable( new FLocMetadataObject(*KeyMetadataB) );
+
+		FLocItem Source( TEXT("TestText") );
+		Source.MetadataObj = MakeShareable( new FLocMetadataObject(*SourceMetadataA) );
+
+		FString TestNamespace = TEXT("TestNamespace");
+
+
+		// Adding entries with exactly matching Source and matching Context
+		{
+			FInternationalizationManifest TestManifest;
+
+			TestManifest.AddSource( TestNamespace, Source, ContextA );
+			bool bResult = TestManifest.AddSource( TestNamespace, Source, ContextA );
+
+			// Adding a duplicate entry reports success but our entry count is not increased after the first entry is added
+			CHECK_MESSAGE(TEXT("AddSource result == true"), bResult);
+			CHECK_MESSAGE(TEXT("ManifestCount == 1"), CountManifestEntries(TestManifest) == 1 );
+		}
+
+		// Adding entries with exactly matching Source but different Contexts
+		{
+			FInternationalizationManifest TestManifest;
+
+			TestManifest.AddSource( TestNamespace, Source, ContextA );
+			TestManifest.AddSource( TestNamespace, Source, ContextB );
+
+			CHECK_MESSAGE(TEXT("ManifestCount == 2"), CountManifestEntries(TestManifest) == 2 );
+
+			// Test find by context
+			TSharedPtr< FManifestEntry > FoundEntry1 = TestManifest.FindEntryByContext( TestNamespace, ContextA );
+			if (!FoundEntry1.IsValid())
+			{
+				FAIL_CHECK(TEXT("FManifestEntry could not find entry using FindEntryByContext."));
+			}
+			else
+			{
+				CHECK_MESSAGE(TEXT("FoundEntry->Source == Source"), FoundEntry1->Source == Source);
+				CHECK_MESSAGE(TEXT("FoundEntry->Context.Num() == 2"), FoundEntry1->Contexts.Num() == 2);
+			}
+
+			TSharedPtr< FManifestEntry > FoundEntry2 = TestManifest.FindEntryByContext( TestNamespace, ContextB );
+			if (!FoundEntry2.IsValid())
+			{
+				FAIL_CHECK(TEXT("FManifestEntry could not find entry using FindEntryByContext."));
+			}
+			else
+			{
+				CHECK_MESSAGE(TEXT("FoundEntry->Source == Source"), FoundEntry2->Source == Source);
+				CHECK_MESSAGE(TEXT("FoundEntry->Context.Num() == 2"), FoundEntry2->Contexts.Num() == 2);
+			}
+
+			// Test find by source
+			TSharedPtr< FManifestEntry > FoundEntry3 = TestManifest.FindEntryBySource( TestNamespace, Source );
+			if (!FoundEntry3.IsValid())
+			{
+				FAIL_CHECK(TEXT("FManifestEntry could not find entry using FindEntryBySource."));
+			}
+			else
+			{
+				CHECK_MESSAGE(TEXT("FoundEntry->Source == Source"), FoundEntry3->Source == Source);
+				CHECK_MESSAGE(TEXT("FoundEntry->Context.Num() == 2"), FoundEntry3->Contexts.Num() == 2);
+			}
+
+			CHECK_MESSAGE(TEXT("FoundEntry1 == FoundEntry2 == FoundEntry3"), (FoundEntry1 == FoundEntry2) && (FoundEntry1 == FoundEntry3) );
+
+		}
+
+		// Adding entries with Source that is NOT an exact match and matching context
+		{
+			// Source mismatched by Source Text.
+			{
+				FInternationalizationManifest TestManifest;
+
+				FLocItem ConflictingSourceA( TEXT("Conflicting TestTextA") );
+				FLocItem ConflictingSourceB( TEXT("Conflicting TestTextB") );
+
+				TestManifest.AddSource( TestNamespace, ConflictingSourceA, ContextA );
+				bool bResult = TestManifest.AddSource( TestNamespace, ConflictingSourceB, ContextA );
+
+				// Adding the second entry reports failure and entry count is not increased
+				CHECK_FALSE_MESSAGE(TEXT("AddSource result == false"), bResult);
+				CHECK_MESSAGE(TEXT("ManifestCount == 1"), CountManifestEntries(TestManifest) == 1 );	
+			}
+
+			// Source mismatched by standard metadata type (not * prefixed) 
+			{
+				FInternationalizationManifest ManifestTypeConflict;
+
+				FLocItem ConflictingSourceMetadataTypeA = Source;
+				FLocItem ConflictingSourceMetadataTypeB = Source;
+
+				// Set metadata with the same name but different type
+				ConflictingSourceMetadataTypeA.MetadataObj->SetBoolField( TEXT("ConflictingType"), true );
+				ConflictingSourceMetadataTypeB.MetadataObj->SetStringField( TEXT("ConflictingType"), TEXT("true") );
+
+				ManifestTypeConflict.AddSource( TestNamespace, ConflictingSourceMetadataTypeA, ContextA );
+				bool bResult = ManifestTypeConflict.AddSource( TestNamespace, ConflictingSourceMetadataTypeB, ContextA );
+
+				CHECK_FALSE_MESSAGE(TEXT("AddSource result == false"), bResult);
+				CHECK_MESSAGE(TEXT("ManifestCount == 1"), CountManifestEntries(ManifestTypeConflict) == 1 );
+
+				// Should not find a match when searching with the second Source we tried to add
+				TSharedPtr< FManifestEntry > FoundEntry = ManifestTypeConflict.FindEntryBySource( TestNamespace, ConflictingSourceMetadataTypeB );
+				CHECK_FALSE_MESSAGE(TEXT("FoundEntry is not valid"), FoundEntry.IsValid());
+
+			}
+
+			// Source mismatched by standard metadata value (not * prefixed)
+			{
+				FInternationalizationManifest TestManifest;
+
+				FLocItem ConflictingSourceMetadataValueA = Source;
+				FLocItem ConflictingSourceMetadataValueB = Source;
+
+				// Set metadata with the same name and type but different value
+				ConflictingSourceMetadataValueA.MetadataObj->SetStringField( TEXT("ConflictingValue"), TEXT("A") );
+				ConflictingSourceMetadataValueB.MetadataObj->SetStringField( TEXT("ConflictingValue"), TEXT("B") );
+
+				TestManifest.AddSource( TestNamespace, ConflictingSourceMetadataValueA, ContextA );
+				bool bResult = TestManifest.AddSource( TestNamespace, ConflictingSourceMetadataValueB, ContextA );
+
+				CHECK_FALSE_MESSAGE(TEXT("AddSource result == false"), bResult);
+				CHECK_MESSAGE(TEXT("ManifestCount == 1"), CountManifestEntries(TestManifest) == 1 );
+
+				// Should not find a match when searching with the second Source we tried to add
+				TSharedPtr< FManifestEntry > FoundEntry = TestManifest.FindEntryBySource( TestNamespace, ConflictingSourceMetadataValueB );
+				CHECK_FALSE_MESSAGE(TEXT("FoundEntry is not valid"), FoundEntry.IsValid());
+
+			}
+
+			// Source mismatched by * prefixed metadata type
+			{
+				FInternationalizationManifest TestManifest;
+
+				FLocItem ConflictingSourceMetadataTypeA = Source;
+				FLocItem ConflictingSourceMetadataTypeB = Source;
+
+				// Set metadata with the same name but different type
+				ConflictingSourceMetadataTypeA.MetadataObj->SetBoolField( TEXT("*ConflictingType"), true );
+				ConflictingSourceMetadataTypeB.MetadataObj->SetStringField( TEXT("*ConflictingType"), TEXT("true") );
+
+				TestManifest.AddSource( TestNamespace, ConflictingSourceMetadataTypeA, ContextA );
+
+				// Though the Source items are considered to be a match(they are equal) in this case, the manifest reports this as a conflict
+				//  and should not add an entry.  The reason is that AddSource does an 'exact' match check on the metadata object
+				bool bResult = TestManifest.AddSource( TestNamespace, ConflictingSourceMetadataTypeB, ContextA );
+
+				CHECK_FALSE_MESSAGE(TEXT("AddSource result == false"), bResult);
+				CHECK_MESSAGE(TEXT("ManifestCount == 1"), CountManifestEntries(TestManifest) == 1 );
+
+				// We should be able to find the entry using either Source FLocItem
+				TSharedPtr< FManifestEntry > FoundEntry1 = TestManifest.FindEntryBySource( TestNamespace, ConflictingSourceMetadataTypeA );
+				CHECK_MESSAGE(TEXT("FoundEntry1 is valid"), FoundEntry1.IsValid());
+
+				TSharedPtr< FManifestEntry > FoundEntry2 = TestManifest.FindEntryBySource( TestNamespace, ConflictingSourceMetadataTypeB );
+				CHECK_MESSAGE(TEXT("FoundEntry2 is valid"), FoundEntry2.IsValid());
+
+			}
+
+			// Source mismatched by * prefixed metadata value
+			{
+				FInternationalizationManifest TestManifest;
+
+				FLocItem ConflictingSourceMetadataValueA = Source;
+				FLocItem ConflictingSourceMetadataValueB = Source;
+
+				// Set metadata with the same name and type but different value
+				ConflictingSourceMetadataValueA.MetadataObj->SetStringField( TEXT("*ConflictingValue"), TEXT("A") );
+				ConflictingSourceMetadataValueB.MetadataObj->SetStringField( TEXT("*ConflictingValue"), TEXT("B") );
+
+				TestManifest.AddSource( TestNamespace, ConflictingSourceMetadataValueA, ContextA );
+
+				// Thought the Source items are considered to be a match(they are equal) in this case, the manifest reports this as a conflict
+				//  and should not add an entry.  The reason is that AddSource does an 'exact' match check on the metadata object
+				bool bResult = TestManifest.AddSource( TestNamespace, ConflictingSourceMetadataValueB, ContextA );
+
+				CHECK_FALSE_MESSAGE(TEXT("AddSource result == false"), bResult);
+				CHECK_MESSAGE(TEXT("ManifestCount == 1"), CountManifestEntries(TestManifest) == 1 );
+
+				// We should be able to find the entry using either Source FLocItem
+				TSharedPtr< FManifestEntry > FoundEntry1 = TestManifest.FindEntryBySource( TestNamespace, ConflictingSourceMetadataValueA );
+				CHECK_MESSAGE(TEXT("FoundEntry1 is valid"), FoundEntry1.IsValid());
+
+				TSharedPtr< FManifestEntry > FoundEntry2 = TestManifest.FindEntryBySource( TestNamespace, ConflictingSourceMetadataValueB );
+				CHECK_MESSAGE(TEXT("FoundEntry2 is valid"), FoundEntry2.IsValid());
+
+			}
+
+		}
+
+		// Adding entries with Source that is NOT an exact match and different context
+		{
+			// Source mismatched by Source Text.
+			{
+				FInternationalizationManifest TestManifest;
+
+				FLocItem ConflictingSourceA( TEXT("Conflicting TestTextA") );
+				FLocItem ConflictingSourceB( TEXT("Conflicting TestTextB") );
+
+				TestManifest.AddSource( TestNamespace, ConflictingSourceA, ContextA );
+				bool bResult = TestManifest.AddSource( TestNamespace, ConflictingSourceB, ContextB );
+
+				CHECK_MESSAGE(TEXT("AddSource result == true"), bResult);
+				CHECK_MESSAGE(TEXT("ManifestCount == 2"), CountManifestEntries(TestManifest) == 2 );
+
+				// We should be able to find two unique entries by source
+				TSharedPtr< FManifestEntry > FoundEntry1 = TestManifest.FindEntryBySource( TestNamespace, ConflictingSourceA );
+				CHECK_MESSAGE(TEXT("FoundEntry1 is valid"), FoundEntry1.IsValid());
+
+				TSharedPtr< FManifestEntry > FoundEntry2 = TestManifest.FindEntryBySource( TestNamespace, ConflictingSourceB );
+				CHECK_MESSAGE(TEXT("FoundEntry2 is valid"), FoundEntry2.IsValid());
+
+				CHECK_MESSAGE(TEXT("FoundEntry1 != FoundEntry2"), FoundEntry1 != FoundEntry2 );
+
+				// We should be able to find two unique entries by context
+				TSharedPtr< FManifestEntry > FoundEntry3 = TestManifest.FindEntryByContext( TestNamespace, ContextA );
+				CHECK_MESSAGE(TEXT("FoundEntry3 is valid"), FoundEntry3.IsValid());
+
+				TSharedPtr< FManifestEntry > FoundEntry4 = TestManifest.FindEntryByContext( TestNamespace, ContextB );
+				CHECK_MESSAGE(TEXT("FoundEntry3 is valid"), FoundEntry4.IsValid());
+
+				CHECK_MESSAGE(TEXT("FoundEntry3 != FoundEntry4"), FoundEntry3 != FoundEntry4 );
+
+				// Found entry looked up by source A should match entry looked up by context A
+				CHECK_MESSAGE(TEXT("FoundEntry1 == FoundEntry3"), FoundEntry1 == FoundEntry3 );
+
+				// Found entry looked up by source B should match entry looked up by context B
+				CHECK_MESSAGE(TEXT("FoundEntry2 == FoundEntry4"), FoundEntry2 == FoundEntry4 );
+
+			}
+
+			// Source mismatched by standard metadata type (not * prefixed)
+			{
+				FInternationalizationManifest TestManifest;
+
+				FLocItem ConflictingSourceMetadataTypeA = Source;
+				FLocItem ConflictingSourceMetadataTypeB = Source;
+
+				// Set metadata with the same name but different type
+				ConflictingSourceMetadataTypeA.MetadataObj->SetBoolField( TEXT("ConflictingType"), true );
+				ConflictingSourceMetadataTypeB.MetadataObj->SetStringField( TEXT("ConflictingType"), TEXT("true") );
+
+				TestManifest.AddSource( TestNamespace, ConflictingSourceMetadataTypeA, ContextA );
+				bool bResult = TestManifest.AddSource( TestNamespace, ConflictingSourceMetadataTypeB, ContextB );
+
+				CHECK_MESSAGE(TEXT("AddSource result == true"), bResult);
+				CHECK_MESSAGE(TEXT("ManifestCount == 2"), CountManifestEntries(TestManifest) == 2 );
+
+				// We should be able to find two unique entries by source
+				TSharedPtr< FManifestEntry > FoundEntry1 = TestManifest.FindEntryBySource( TestNamespace, ConflictingSourceMetadataTypeA );
+				CHECK_MESSAGE(TEXT("FoundEntry1 is valid"), FoundEntry1.IsValid());
+
+				TSharedPtr< FManifestEntry > FoundEntry2 = TestManifest.FindEntryBySource( TestNamespace, ConflictingSourceMetadataTypeB );
+				CHECK_MESSAGE(TEXT("FoundEntry2 is valid"), FoundEntry2.IsValid());
+
+				CHECK_MESSAGE(TEXT("FoundEntry1 != FoundEntry2"), FoundEntry1 != FoundEntry2 );
+
+				// We should be able to find two unique entries by context
+				TSharedPtr< FManifestEntry > FoundEntry3 = TestManifest.FindEntryByContext( TestNamespace, ContextA );
+				CHECK_MESSAGE(TEXT("FoundEntry3 is valid"), FoundEntry3.IsValid());
+
+				TSharedPtr< FManifestEntry > FoundEntry4 = TestManifest.FindEntryByContext( TestNamespace, ContextB );
+				CHECK_MESSAGE(TEXT("FoundEntry3 is valid"), FoundEntry4.IsValid());
+
+				CHECK_MESSAGE(TEXT("FoundEntry3 != FoundEntry4"), FoundEntry3 != FoundEntry4 );
+
+				// Found entry looked up by source A should match entry looked up by context A
+				CHECK_MESSAGE(TEXT("FoundEntry1 == FoundEntry3"), FoundEntry1 == FoundEntry3 );
+
+				// Found entry looked up by source B should match entry looked up by context B
+				CHECK_MESSAGE(TEXT("FoundEntry2 == FoundEntry4"), FoundEntry2 == FoundEntry4 );
+
+			}
+
+			// Source mismatched by standard metadata value (not * prefixed)
+			{
+				FInternationalizationManifest TestManifest;
+
+				FLocItem ConflictingSourceMetadataValueA = Source;
+				FLocItem ConflictingSourceMetadataValueB = Source;
+
+				// Set metadata with the same name and type but different value
+				ConflictingSourceMetadataValueA.MetadataObj->SetStringField( TEXT("ConflictingValue"), TEXT("A") );
+				ConflictingSourceMetadataValueB.MetadataObj->SetStringField( TEXT("ConflictingValue"), TEXT("B") );
+
+				TestManifest.AddSource( TestNamespace, ConflictingSourceMetadataValueA, ContextA );
+				bool bResult = TestManifest.AddSource( TestNamespace, ConflictingSourceMetadataValueB, ContextB );
+
+				CHECK_MESSAGE(TEXT("AddSource result == true"), bResult);
+				CHECK_MESSAGE(TEXT("ManifestCount == 2"), CountManifestEntries(TestManifest) == 2 );
+
+				// We should be able to find two unique entries by source
+				TSharedPtr< FManifestEntry > FoundEntry1 = TestManifest.FindEntryBySource( TestNamespace, ConflictingSourceMetadataValueA );
+				CHECK_MESSAGE(TEXT("FoundEntry1 is valid"), FoundEntry1.IsValid());
+
+				TSharedPtr< FManifestEntry > FoundEntry2 = TestManifest.FindEntryBySource( TestNamespace, ConflictingSourceMetadataValueB );
+				CHECK_MESSAGE(TEXT("FoundEntry2 is valid"), FoundEntry2.IsValid());
+
+				CHECK_MESSAGE(TEXT("FoundEntry1 != FoundEntry2"), FoundEntry1 != FoundEntry2 );
+
+				// We should be able to find two unique entries by context
+				TSharedPtr< FManifestEntry > FoundEntry3 = TestManifest.FindEntryByContext( TestNamespace, ContextA );
+				CHECK_MESSAGE(TEXT("FoundEntry3 is valid"), FoundEntry3.IsValid());
+
+				TSharedPtr< FManifestEntry > FoundEntry4 = TestManifest.FindEntryByContext( TestNamespace, ContextB );
+				CHECK_MESSAGE(TEXT("FoundEntry3 is valid"), FoundEntry4.IsValid());
+
+				CHECK_MESSAGE(TEXT("FoundEntry3 != FoundEntry4"), FoundEntry3 != FoundEntry4 );
+
+				// Found entry looked up by source A should match entry looked up by context A
+				CHECK_MESSAGE(TEXT("FoundEntry1 == FoundEntry3"), FoundEntry1 == FoundEntry3 );
+
+				// Found entry looked up by source B should match entry looked up by context B
+				CHECK_MESSAGE(TEXT("FoundEntry2 == FoundEntry4"), FoundEntry2 == FoundEntry4 );
+
+			}
+
+			// Source mismatched by * prefixed metadata type
+			{
+				FInternationalizationManifest TestManifest;
+
+				FLocItem ConflictingSourceMetadataTypeA = Source;
+				FLocItem ConflictingSourceMetadataTypeB = Source;
+
+				// Set metadata with the same name but different type
+				ConflictingSourceMetadataTypeA.MetadataObj->SetBoolField( TEXT("*ConflictingType"), true );
+				ConflictingSourceMetadataTypeB.MetadataObj->SetStringField( TEXT("*ConflictingType"), TEXT("true") );
+
+				TestManifest.AddSource( TestNamespace, ConflictingSourceMetadataTypeA, ContextA );
+
+				// Though the Source items are considered to be a match(they are equal) in this case, the manifest reports this as a conflict
+				//  and should not add an entry.  The reason is that AddSource does an 'exact' match check on the metadata object
+				bool bResult = TestManifest.AddSource( TestNamespace, ConflictingSourceMetadataTypeB, ContextB );
+
+				CHECK_FALSE_MESSAGE(TEXT("AddSource result == false"), bResult);
+				CHECK_MESSAGE(TEXT("ManifestCount == 1"), CountManifestEntries(TestManifest) == 1 );
+
+				// We should be able to find the entry using either Source FLocItem
+				TSharedPtr< FManifestEntry > FoundEntry1 = TestManifest.FindEntryBySource( TestNamespace, ConflictingSourceMetadataTypeA );
+				CHECK_MESSAGE(TEXT("FoundEntry1 is valid"), FoundEntry1.IsValid());
+
+				TSharedPtr< FManifestEntry > FoundEntry2 = TestManifest.FindEntryBySource( TestNamespace, ConflictingSourceMetadataTypeB );
+				CHECK_MESSAGE(TEXT("FoundEntry2 is valid"), FoundEntry2.IsValid());
+			}
+
+			// Source mismatched by * prefixed metadata value
+			{
+				FInternationalizationManifest TestManifest;
+
+				FLocItem ConflictingSourceMetadataValueA = Source;
+				FLocItem ConflictingSourceMetadataValueB = Source;
+
+				// Set metadata with the same name and type but different value
+				ConflictingSourceMetadataValueA.MetadataObj->SetStringField( TEXT("*ConflictingValue"), TEXT("A") );
+				ConflictingSourceMetadataValueB.MetadataObj->SetStringField( TEXT("*ConflictingValue"), TEXT("B") );
+
+				TestManifest.AddSource( TestNamespace, ConflictingSourceMetadataValueA, ContextA );
+
+				// Thought the Source items are considered to be a match(they are equal) in this case, the manifest reports this as a conflict
+				//  and should not add an entry.  The reason is that AddSource does an 'exact' match check on the metadata object
+				bool bResult = TestManifest.AddSource( TestNamespace, ConflictingSourceMetadataValueB, ContextB );
+
+				CHECK_FALSE_MESSAGE(TEXT("AddSource result == false"), bResult);
+				CHECK_MESSAGE(TEXT("ManifestCount == 1"), CountManifestEntries(TestManifest) == 1 );
+
+				// We should be able to find the entry using either Source FLocItem
+				TSharedPtr< FManifestEntry > FoundEntry1 = TestManifest.FindEntryBySource( TestNamespace, ConflictingSourceMetadataValueA );
+				CHECK_MESSAGE(TEXT("FoundEntry1 is valid"), FoundEntry1.IsValid());
+
+				TSharedPtr< FManifestEntry > FoundEntry2 = TestManifest.FindEntryBySource( TestNamespace, ConflictingSourceMetadataValueB );
+				CHECK_MESSAGE(TEXT("FoundEntry2 is valid"), FoundEntry2.IsValid());
+
+			}
+
+		}
+
+		// Adding an entry that only differs in the optional flag 
+		{
+			FInternationalizationManifest TestManifest;
+
+			// Reports success and our entry count does not change.  bIsOptional is not a key and is not used during lookup.  When
+			//  we AddSource, we find an existing matching entry so AddSource returns true but no new entry is added.  The original
+			//  entry's bIsOptional value is not updated.
+			FManifestContext ContextConflictingOptionalFlag = ContextA;
+			ContextConflictingOptionalFlag.bIsOptional = !ContextA.bIsOptional;
+
+			TestManifest.AddSource( TestNamespace, Source, ContextA );
+			bool bResult = TestManifest.AddSource( TestNamespace, Source, ContextConflictingOptionalFlag);
+
+			CHECK_MESSAGE(TEXT("AddSource result == true"), bResult);
+			CHECK_MESSAGE(TEXT("ManifestCount == 1"), CountManifestEntries(TestManifest) == 1 );
+
+			// We should be able to look up the existing entry using the ContextConflictingOptionalFlag context but the entries bIsOptional flag will match ContextA
+			TSharedPtr< FManifestEntry > FoundEntry = TestManifest.FindEntryByContext( TestNamespace, ContextConflictingOptionalFlag );
+			if (!FoundEntry.IsValid())
+			{
+				FAIL_CHECK(TEXT("FManifestEntry could not find entry using FindEntryByContext."));
+			}
+			else
+			{
+				CHECK_MESSAGE(TEXT("FoundEntry->bIsOptional == ContextA->bIsOptional"), FoundEntry->Contexts[0].bIsOptional == ContextA.bIsOptional);
+			}
+		}
+
+		// Add an entry with null key metadata to see if we can retrieve it with non-null but empty key metadata
+		{
+			FInternationalizationManifest TestManifest;
+
+			FManifestContext ContextC;
+			ContextC.Key			= TEXT("KeyC");
+			ContextC.SourceLocation = TEXT("SourceLocationC");
+			ContextC.InfoMetadataObj = MakeShareable( new FLocMetadataObject(*InfoMetadataB) );
+			ContextC.KeyMetadataObj = NULL;
+
+			Source.MetadataObj = NULL;
+			TestManifest.AddSource( TestNamespace, Source, ContextC );
+
+			//  Now give our context and source valid but empty metadata
+			ContextC.KeyMetadataObj = MakeShareable( new FLocMetadataObject() );
+			Source.MetadataObj = MakeShareable( new FLocMetadataObject() );
+
+			// Ensure we find the entry we added by source
+			TSharedPtr< FManifestEntry > FoundEntry;
+			FoundEntry = TestManifest.FindEntryBySource( TestNamespace, Source );
+			if (!FoundEntry.IsValid())
+			{
+				FAIL_CHECK(TEXT("FManifestEntry could not find entry using FindEntryBySource."));
+			}
+			else
+			{
+				CHECK_MESSAGE(TEXT("FoundEntry->Source == Source"), FoundEntry->Source == Source);
+			}
+			FoundEntry = TestManifest.FindEntryByContext(TestNamespace, ContextC);
+			// Ensure we find the entry we added by context
+			if (!FoundEntry.IsValid())
+			{
+				FAIL_CHECK(TEXT("FManifestEntry could not find entry using FindEntryByContext."));
+			}
+			else
+			{
+				CHECK_MESSAGE(TEXT("FoundEntry->Source == Source"), FoundEntry->Source == Source);
+			}
+		}
+
+		// Add an entry with non-null but empty key metadata and see if we can retrieve null metadata
+		{
+			FInternationalizationManifest TestManifest;
+
+			FManifestContext ContextC;
+			ContextC.Key			= TEXT("KeyC");
+			ContextC.SourceLocation = TEXT("SourceLocationC");
+			ContextC.InfoMetadataObj = MakeShareable( new FLocMetadataObject(*InfoMetadataB) );
+			ContextC.KeyMetadataObj = MakeShareable( new FLocMetadataObject());
+
+			Source.MetadataObj = MakeShareable( new FLocMetadataObject());
+			TestManifest.AddSource( TestNamespace, Source, ContextC );
+
+			//  Now give our context and source Null metadata
+			ContextC.KeyMetadataObj = NULL;
+			Source.MetadataObj = NULL;
+
+			// Ensure we find the entry we added by source
+			TSharedPtr< FManifestEntry > FoundEntry;
+			FoundEntry = TestManifest.FindEntryBySource( TestNamespace, Source );
+			if (!FoundEntry.IsValid())
+			{
+				FAIL_CHECK(TEXT("FManifestEntry could not find entry using FindEntryBySource."));
+			}
+			else
+			{
+				CHECK_MESSAGE(TEXT("FoundEntry->Source == Source"), FoundEntry->Source == Source);
+			}
+			// Ensure we find the entry we added by context
+			FoundEntry = TestManifest.FindEntryByContext( TestNamespace, ContextC );
+			if (!FoundEntry.IsValid())
+			{
+				FAIL_CHECK(TEXT("FManifestEntry could not find entry using FindEntryByContext."));
+			}
+			else
+			{
+				CHECK_MESSAGE(TEXT("FoundEntry->Source == Source"), FoundEntry->Source == Source);
+			}
+		}
+	}
+}
+
+#endif //WITH_TESTS
\ No newline at end of file
diff --git a/Engine/Source/Runtime/Core/Tests/Internationalization/MetadataTest.cpp b/Engine/Source/Runtime/Core/Tests/Internationalization/MetadataTest.cpp
new file mode 100644
--- /dev/null
+++ b/Engine/Source/Runtime/Core/Tests/Internationalization/MetadataTest.cpp
@@ -0,0 +1,366 @@
+// Copyright Epic Games, Inc. All Rights Reserved.
+
+#include "CoreTypes.h"
+#include "Containers/Array.h"
+#include "Containers/UnrealString.h"
+#include "Internationalization/InternationalizationMetadata.h"
+
+#if WITH_TESTS
+
+#include "Tests/TestHarnessAdapter.h"
+
+TEST_CASE_NAMED(FMetadataTest, "System::Core::Misc::Internationalization Metadata", "[EditorContext][ClientContext][SmokeFilter]")
+{
+	// BooleanValue metadata
+	bool BoolFalse = false;
+	bool BoolTrue = true;
+	TSharedPtr<FLocMetadataValue> MetadataValueBoolFalse = MakeShareable( new FLocMetadataValueBoolean( BoolFalse ) );
+	TSharedPtr<FLocMetadataValue> MetadataValueBoolTrue = MakeShareable( new FLocMetadataValueBoolean( BoolTrue ) );
+
+	// StringValue metadata
+	FString StringA = TEXT("A");
+	FString StringB = TEXT("B");
+	TSharedPtr<FLocMetadataValue> MetadataValueStringA = MakeShareable( new FLocMetadataValueString( StringA ) );
+	TSharedPtr<FLocMetadataValue> MetadataValueStringB = MakeShareable( new FLocMetadataValueString( StringB ) );
+
+	// ArrayValue metadata
+	TArray< TSharedPtr< FLocMetadataValue > > ArrayA;
+	TArray< TSharedPtr< FLocMetadataValue > > ArrayB;
+
+	ArrayA.Add( MetadataValueBoolFalse );
+	ArrayA.Add( MetadataValueStringA );
+
+	ArrayB.Add( MetadataValueBoolTrue );
+	ArrayB.Add( MetadataValueStringB );
+
+	TSharedPtr<FLocMetadataValue> MetadataValueArrayA = MakeShareable( new FLocMetadataValueArray( ArrayA ) );
+	TSharedPtr<FLocMetadataValue> MetadataValueArrayB = MakeShareable( new FLocMetadataValueArray( ArrayB ) );
+
+	// Object metadata
+	TSharedPtr< FLocMetadataObject > MetadataObjectA = MakeShareable( new FLocMetadataObject );
+	TSharedPtr< FLocMetadataObject > MetadataObjectB = MakeShareable( new FLocMetadataObject );
+
+	// Setup object A
+	MetadataObjectA->SetField( TEXT("MetadataBoolFalse"), MetadataValueBoolFalse );
+	MetadataObjectA->SetField( TEXT("MetadataStringA"), MetadataValueStringA );
+	MetadataObjectA->SetField( TEXT("MetadataArrayA"), MetadataValueArrayA );
+	MetadataObjectA->SetField( TEXT("*MetadataCompareModifier"), MetadataValueStringA);// Note: The * name prefix modifies the way entries in the object are compared.
+
+	// Setup object B
+	MetadataObjectB->SetField( TEXT("MetadataBoolFalse"), MetadataValueBoolTrue );
+	MetadataObjectB->SetField( TEXT("MetadataStringB"), MetadataValueStringB );
+	MetadataObjectB->SetField( TEXT("MetadataArrayB"), MetadataValueArrayB );
+	MetadataObjectA->SetBoolField( TEXT("*MetadataCompareModifier"), true);// Note: Different type/value. The * name prefix modifies the way entries in the object are compared.
+
+	// ObjectValue metadata
+	TSharedPtr< FLocMetadataValue > MetadataValueObjectA = MakeShareable( new FLocMetadataValueObject( MetadataObjectA ) );
+	TSharedPtr< FLocMetadataValue > MetadataValueObjectB = MakeShareable( new FLocMetadataValueObject( MetadataObjectB ) );
+
+	// Testing the bool meta data value type
+	{
+		CHECK_FALSE_MESSAGE(TEXT("MetadataValueBoolFalse == MetadataValueBoolTrue"), *MetadataValueBoolFalse == *MetadataValueBoolTrue );
+		CHECK_MESSAGE(TEXT("MetadataValueBoolFalse < MetadataValueBoolTrue"), *MetadataValueBoolFalse < *MetadataValueBoolTrue );
+		CHECK_FALSE_MESSAGE(TEXT("MetadataValueBoolTrue < MetadataValueBoolFalse"), *MetadataValueBoolTrue < *MetadataValueBoolFalse);
+
+		CHECK_MESSAGE(TEXT("MetadataValueBoolFalse < MetadataValueString"), *MetadataValueBoolFalse < *MetadataValueStringA );
+		CHECK_MESSAGE(TEXT("MetadataValueBoolTrue < MetadataValueString"), *MetadataValueBoolTrue < *MetadataValueStringA);
+
+		CHECK_MESSAGE(TEXT("MetadataValueBoolFalse < MetadataValueArray"), *MetadataValueBoolFalse < *MetadataValueArrayA );
+		CHECK_MESSAGE(TEXT("MetadataValueBoolTrue < MetadataValueArray"), *MetadataValueBoolTrue < *MetadataValueArrayA);
+
+		CHECK_MESSAGE(TEXT("MetadataValueBoolFalse < MetadataValueObject"), *MetadataValueBoolFalse < *MetadataValueObjectA );
+		CHECK_MESSAGE(TEXT("MetadataValueBoolTrue < MetadataValueObject"), *MetadataValueBoolTrue < *MetadataValueObjectA);
+
+		TSharedPtr<FLocMetadataValue> MetadataValueBoolFalseClone = MetadataValueBoolFalse->Clone();
+		TSharedPtr<FLocMetadataValue> MetadataValueBoolTrueClone = MetadataValueBoolTrue->Clone();
+
+		if (MetadataValueBoolFalse == MetadataValueBoolFalseClone)
+		{
+			FAIL_CHECK(TEXT("MetadataValueBool and its Clone are not unique objects."));
+		}
+		CHECK_MESSAGE(TEXT("MetadataValueBoolFalseClone == MetadataValueBoolFalse"), *MetadataValueBoolFalseClone == *MetadataValueBoolFalse );
+		CHECK_FALSE_MESSAGE(TEXT("MetadataValueBoolFalseClone < MetadataValueBoolFalse"), *MetadataValueBoolFalseClone < *MetadataValueBoolFalse );
+
+		CHECK_MESSAGE(TEXT("MetadataValueBoolTrueClone == MetadataValueBoolTrue"), *MetadataValueBoolTrueClone == *MetadataValueBoolTrue );
+		CHECK_FALSE_MESSAGE(TEXT("MetadataValueBoolTrueClone < MetadataValueBoolTrue"), *MetadataValueBoolTrueClone < *MetadataValueBoolTrue );
+
+		// Test the bool metadata when it is part of an object
+		TSharedPtr< FLocMetadataObject > MetadataObjectFalse = MakeShareable( new FLocMetadataObject );
+		MetadataObjectFalse->SetField(TEXT("MetadataValueBool"), MetadataValueBoolFalse);
+
+		TSharedPtr< FLocMetadataObject > MetadataObjectTrue = MakeShareable( new FLocMetadataObject );
+		MetadataObjectTrue->SetField(TEXT("MetadataValueBool"), MetadataValueBoolTrue);
+
+		CHECK_FALSE_MESSAGE(TEXT("GetBoolField(MetadataValueBool)"), MetadataObjectFalse->GetBoolField( TEXT("MetadataValueBool") ) );
+		CHECK_MESSAGE(TEXT("GetBoolField(MetadataValueBool)"), MetadataObjectTrue->GetBoolField( TEXT("MetadataValueBool") ) );
+
+		CHECK_FALSE_MESSAGE(TEXT("MetadataObjectFalse == MetadataObjectTrue"), *MetadataObjectFalse == *MetadataObjectTrue );
+		CHECK_MESSAGE(TEXT("MetadataObjectFalse < MetadataObjectTrue"), *MetadataObjectFalse < *MetadataObjectTrue );
+	}
+
+	// Testing string metadata value type
+	{
+		CHECK_FALSE_MESSAGE(TEXT("MetadataValueStringA == MetadataValueStringB"), *MetadataValueStringA == *MetadataValueStringB );
+		CHECK_MESSAGE(TEXT("MetadataValueStringA < MetadataValueStringB"), *MetadataValueStringA < *MetadataValueStringB );
+		CHECK_FALSE_MESSAGE(TEXT("MetadataValueStringB < MetadataValueStringA"), *MetadataValueStringB < *MetadataValueStringA);
+
+		CHECK_MESSAGE(TEXT("MetadataValueString < MetadataValueArray"), *MetadataValueStringA < *MetadataValueArrayA );
+
+		CHECK_MESSAGE(TEXT("MetadataValueStringA < MetadataValueObject"), *MetadataValueStringA < *MetadataValueObjectA );
+
+
+		TSharedPtr<FLocMetadataValue> MetadataValueStringAClone = MetadataValueStringA->Clone();
+
+		if (MetadataValueStringA == MetadataValueStringAClone)
+		{
+			FAIL_CHECK(TEXT("MetadataValueString and its Clone are not unique objects."));
+		}
+		CHECK_MESSAGE(TEXT("MetadataValueStringAClone == MetadataValueStringA"), *MetadataValueStringAClone == *MetadataValueStringA );
+		CHECK_FALSE_MESSAGE(TEXT("MetadataValueStringAClone < MetadataValueStringA"), *MetadataValueStringAClone < *MetadataValueStringA );
+		CHECK_MESSAGE(TEXT("MetadataValueStringAClone < MetadataValueStringB"), *MetadataValueStringAClone < *MetadataValueStringB );
+
+		// Test the string metadata when it is part of an object
+		TSharedPtr< FLocMetadataObject > TestMetadataObjectA = MakeShareable( new FLocMetadataObject );
+		TestMetadataObjectA->SetField(TEXT("MetadataValueString"), MetadataValueStringA);
+
+		TSharedPtr< FLocMetadataObject > TestMetadataObjectB = MakeShareable( new FLocMetadataObject );
+		TestMetadataObjectB->SetField(TEXT("MetadataValueString"), MetadataValueStringB);
+
+		CHECK_MESSAGE(TEXT("GetStringField(MetadataValueString) == StringA"), TestMetadataObjectA->GetStringField( TEXT("MetadataValueString") ) == StringA );
+	
+		CHECK_FALSE_MESSAGE(TEXT("TestMetadataObjectA == TestMetadataObjectB"), *TestMetadataObjectA == *TestMetadataObjectB );
+		CHECK_MESSAGE(TEXT("TestMetadataObjectA < TestMetadataObjectB"), *TestMetadataObjectA < *TestMetadataObjectB );
+	}
+
+	// Testing array metadata value type
+	{
+		CHECK_FALSE_MESSAGE(TEXT("MetadataValueArrayA == MetadataValueArrayB"), *MetadataValueArrayA == *MetadataValueArrayB );
+		CHECK_MESSAGE(TEXT("MetadataValueArrayA < MetadataValueArrayB"), *MetadataValueArrayA < *MetadataValueArrayB );
+		CHECK_FALSE_MESSAGE(TEXT("MetadataValueArrayB < MetadataValueArrayA"), *MetadataValueArrayB < *MetadataValueArrayA);
+
+		CHECK_MESSAGE(TEXT("MetadataValueArrayA < MetadataValueObject"), *MetadataValueArrayA < *MetadataValueObjectA );
+
+		TSharedPtr<FLocMetadataValue> MetadataValueArrayAClone = MetadataValueArrayA->Clone();
+
+		if (MetadataValueArrayA == MetadataValueArrayAClone)
+		{
+			FAIL_CHECK(TEXT("MetadataValueString and its Clone are not unique objects."));
+		}
+		CHECK_MESSAGE(TEXT("MetadataValueArrayAClone == MetadataValueArrayA"), *MetadataValueArrayAClone == *MetadataValueArrayA );
+		CHECK_FALSE_MESSAGE(TEXT("MetadataValueArrayAClone < MetadataValueArrayA"), *MetadataValueArrayAClone < *MetadataValueArrayA );
+		CHECK_MESSAGE(TEXT("MetadataValueArrayAClone < MetadataValueArrayB"), *MetadataValueArrayAClone < *MetadataValueArrayB );
+
+		// Test Less than and equality checks.  Metadata arrays are equivalent if they contain equivalent contents in any order.  
+		//  To calculate if a metadata array is less than another, we sort both arrays and check each entry index against its counterpart.
+		//  If we encounter an entry that is less than another we stop looking.
+		TArray< TSharedPtr< FLocMetadataValue > > ArrayC;
+		ArrayC.Add( MetadataValueBoolFalse );
+		ArrayC.Add( MetadataValueBoolFalse->Clone() );
+		TSharedPtr<FLocMetadataValue> MetadataValueArrayC = MakeShareable( new FLocMetadataValueArray( ArrayC ) );
+	
+		CHECK_FALSE_MESSAGE(TEXT("MetadataValueArrayA == MetadataValueArrayC"), *MetadataValueArrayA == *MetadataValueArrayC );
+		CHECK_MESSAGE(TEXT("MetadataValueArrayC < MetadataValueArrayA"), *MetadataValueArrayC < *MetadataValueArrayA );
+		CHECK_MESSAGE(TEXT("MetadataValueArrayC < MetadataValueArrayB"), *MetadataValueArrayC < *MetadataValueArrayB );
+
+		TArray< TSharedPtr< FLocMetadataValue > > ArrayD;
+		ArrayD.Add( MetadataValueBoolFalse );
+		ArrayD.Add( MetadataValueBoolFalse->Clone() );
+		ArrayD.Add( MetadataValueBoolFalse->Clone() );
+		TSharedPtr<FLocMetadataValue> MetadataValueArrayD = MakeShareable( new FLocMetadataValueArray( ArrayD ) );
+
+		CHECK_FALSE_MESSAGE(TEXT("MetadataValueArrayA == MetadataValueArrayD"), *MetadataValueArrayA == *MetadataValueArrayD );
+		CHECK_FALSE_MESSAGE(TEXT("MetadataValueArrayC == MetadataValueArrayD"), *MetadataValueArrayC == *MetadataValueArrayD );
+		CHECK_MESSAGE(TEXT("MetadataValueArrayC < MetadataValueArrayD"), *MetadataValueArrayC < *MetadataValueArrayD );
+		CHECK_MESSAGE(TEXT("MetadataValueArrayD < MetadataValueArrayA"), *MetadataValueArrayD < *MetadataValueArrayA );
+
+		// Test the array metadata when it is part of an object
+		TSharedPtr< FLocMetadataObject > TestMetadataObjectA = MakeShareable( new FLocMetadataObject );
+		TestMetadataObjectA->SetField(TEXT("MetadataValueArray"), MetadataValueArrayA);
+
+		TSharedPtr< FLocMetadataObject > TestMetadataObjectB = MakeShareable( new FLocMetadataObject );
+		TestMetadataObjectB->SetField(TEXT("MetadataValueArray"), MetadataValueArrayB);
+
+		CHECK_MESSAGE(TEXT("GetArrayField(MetadataValueArray) == ArrayA"), TestMetadataObjectA->GetArrayField( TEXT("MetadataValueArray") ) == ArrayA );
+
+		CHECK_FALSE_MESSAGE(TEXT("TestMetadataObjectA == TestMetadataObjectB"), *TestMetadataObjectA == *TestMetadataObjectB );
+		CHECK_MESSAGE(TEXT("TestMetadataObjectA < TestMetadataObjectB"), *TestMetadataObjectA < *TestMetadataObjectB );
+	}
+
+	// Testing object metadata value type
+	{
+		CHECK_FALSE_MESSAGE(TEXT("MetadataValueObjectA == MetadataValueObjectB"), *MetadataValueObjectA == *MetadataValueObjectB );
+		CHECK_MESSAGE(TEXT("MetadataValueObjectA < MetadataValueObjectB"), *MetadataValueObjectA < *MetadataValueObjectB );
+		CHECK_FALSE_MESSAGE(TEXT("MetadataValueObjectB < MetadataValueObjectA"), *MetadataValueObjectB < *MetadataValueObjectA);
+
+		TSharedPtr<FLocMetadataValue> MetadataValueObjectAClone = MetadataValueObjectA->Clone();
+
+		if (MetadataValueObjectA == MetadataValueObjectAClone)
+		{
+			FAIL_CHECK(TEXT("MetadataValueString and its Clone are not unique objects."));
+		}
+
+		CHECK_MESSAGE(TEXT("MetadataValueObjectAClone == MetadataValueObjectA"), *MetadataValueObjectAClone == *MetadataValueObjectA );
+		CHECK_FALSE_MESSAGE(TEXT("MetadataValueObjectAClone < MetadataValueObjectA"), *MetadataValueObjectAClone < *MetadataValueObjectA );
+		CHECK_MESSAGE(TEXT("MetadataValueObjectAClone < MetadataValueObjectB"), *MetadataValueObjectAClone < *MetadataValueObjectB );
+
+
+		// Test the object metadata when it is part of another object
+		TSharedPtr< FLocMetadataObject > TestMetadataObjectA = MakeShareable( new FLocMetadataObject );
+		TestMetadataObjectA->SetField(TEXT("MetadataValueObject"), MetadataValueObjectA);
+
+		TSharedPtr< FLocMetadataObject > TestMetadataObjectB = MakeShareable( new FLocMetadataObject );
+		TestMetadataObjectB->SetField(TEXT("MetadataValueObject"), MetadataValueObjectB);
+
+		CHECK_MESSAGE(TEXT("GetObjectField(MetadataValueObject) == MetadataObjectA"), *TestMetadataObjectA->GetObjectField( TEXT("MetadataValueObject") ) == *MetadataObjectA );
+
+		CHECK_FALSE_MESSAGE(TEXT("TestMetadataObjectA == TestMetadataObjectB"), *TestMetadataObjectA == *TestMetadataObjectB );
+		CHECK_MESSAGE(TEXT("TestMetadataObjectA < TestMetadataObjectB"), *TestMetadataObjectA < *TestMetadataObjectB );
+	}
+
+	// Testing Loc Metadata Object
+	{
+		CHECK_FALSE_MESSAGE(TEXT("MetadataObjectA == MetadataObjectB"), *MetadataObjectA == *MetadataObjectB );
+		CHECK_MESSAGE(TEXT("MetadataObjectA < MetadataObjectB"), *MetadataObjectA < *MetadataObjectB );
+		CHECK_FALSE_MESSAGE(TEXT("MetadataObjectB < MetadataObjectA"), *MetadataObjectB < *MetadataObjectA);
+
+		// Test copy ctor
+		{
+			FLocMetadataObject MetadataObjectAClone = *MetadataObjectA;
+			CHECK_MESSAGE(TEXT("MetadataObjectAClone == MetadataObjectA"), MetadataObjectAClone == *MetadataObjectA );
+		}
+		
+		// Test assignment
+		{
+			FLocMetadataObject MetadataObjectAClone = *MetadataObjectB;
+			// PVS-Studio complains about double initialization, but that's something we're testing
+			// so we disable the warning
+			MetadataObjectAClone = *MetadataObjectA; //-V519
+			CHECK_MESSAGE(TEXT("MetadataObjectAClone == MetadataObjectA"), MetadataObjectAClone == *MetadataObjectA );
+			CHECK_FALSE_MESSAGE(TEXT("MetadataObjectAClone == MetadataObjectB"), MetadataObjectAClone == *MetadataObjectB );
+		}
+		
+		// Test comparison operator  
+		{
+			FLocMetadataObject MetadataObjectAClone = *MetadataObjectA;
+			
+			
+			// Adding standard entry
+			MetadataObjectAClone.SetStringField( TEXT("NewEntry"), TEXT("NewEntryValue") );
+			CHECK_FALSE_MESSAGE(TEXT("MetadataObjectAClone == MetadataObjectA"), MetadataObjectAClone == *MetadataObjectA );
+			
+			// Adding non-standard entry.  Note metadata with * prefix in the name will ignore value and type when performing comparisons
+			MetadataObjectAClone = *MetadataObjectA;
+			MetadataObjectAClone.SetStringField( TEXT("*NewEntry"), TEXT("*NewEntryValue") );
+			CHECK_FALSE_MESSAGE(TEXT("MetadataObjectAClone == MetadataObjectA"), MetadataObjectAClone == *MetadataObjectA );
+
+			// Value mismatch on entry with * prefix with same type
+			MetadataObjectAClone = *MetadataObjectA;
+			MetadataObjectAClone.SetStringField( TEXT("*NoCompare"), TEXT("NoCompare") );
+			FLocMetadataObject MetadataObjectAClone2 = *MetadataObjectA;
+			MetadataObjectAClone2.SetStringField( TEXT("*NoCompare"), TEXT("NoCompare2") );
+			CHECK_MESSAGE(TEXT("MetadataObjectAClone == MetadataObjectAClone2"), MetadataObjectAClone == MetadataObjectAClone2 );
+
+			// Value and type mismatch on entry with * prefix
+			MetadataObjectAClone = *MetadataObjectA;
+			MetadataObjectAClone.SetStringField( TEXT("*NoCompare"), TEXT("NoCompare") );
+			MetadataObjectAClone2 = *MetadataObjectA;
+			MetadataObjectAClone2.SetBoolField( TEXT("*NoCompare"), true );
+			CHECK_MESSAGE(TEXT("MetadataObjectAClone == MetadataObjectAClone2"), MetadataObjectAClone == MetadataObjectAClone2 );
+		
+			// Value mismatch on standard entry
+			MetadataObjectAClone = *MetadataObjectA;
+			MetadataObjectAClone.SetStringField( TEXT("DoCompare"), TEXT("DoCompare") );
+			MetadataObjectAClone2 = *MetadataObjectA;
+			MetadataObjectAClone2.SetStringField( TEXT("DoCompare"), TEXT("DoCompare2") );
+			CHECK_FALSE_MESSAGE(TEXT("MetadataObjectAClone == MetadataObjectAClone2"), MetadataObjectAClone == MetadataObjectAClone2 );
+
+			// Value and type mismatch on standard entry
+			MetadataObjectAClone = *MetadataObjectA;
+			MetadataObjectAClone.SetStringField( TEXT("DoCompare"), TEXT("DoCompare") );
+			MetadataObjectAClone2 = *MetadataObjectA;
+			MetadataObjectAClone2.SetBoolField( TEXT("DoCompare"), true );
+			CHECK_FALSE_MESSAGE(TEXT("MetadataObjectAClone == MetadataObjectAClone2"), MetadataObjectAClone == MetadataObjectAClone2 );
+		}		
+		
+		// Test IsExactMatch function.  Note: Differs from the == operator which takes into account the COMPARISON_MODIFIER_PREFIX
+		{
+			// This function behaves similar to the comparison operator except that it ensures all metadata entries match exactly.  In other
+			//  words, it will also perform exact match checks on * prefixed metadata items.
+			FLocMetadataObject MetadataObjectAClone = *MetadataObjectA;
+
+			CHECK_MESSAGE(TEXT("MetadataObjectAClone == MetadataObjectA"), MetadataObjectAClone.IsExactMatch( *MetadataObjectA) );
+
+			// Adding standard entry
+			MetadataObjectAClone.SetStringField( TEXT("NewEntry"), TEXT("NewEntryValue") );
+			CHECK_FALSE_MESSAGE(TEXT("MetadataObjectAClone == MetadataObjectA"), MetadataObjectAClone.IsExactMatch( *MetadataObjectA) );
+
+			// Adding non-standard entry.  Note metadata with * prefix in the name will ignore value and type when performing comparisons
+			MetadataObjectAClone = *MetadataObjectA;
+			MetadataObjectAClone.SetStringField( TEXT("*NewEntry"), TEXT("*NewEntryValue") );
+			CHECK_FALSE_MESSAGE(TEXT("MetadataObjectAClone == MetadataObjectA"),MetadataObjectAClone.IsExactMatch( *MetadataObjectA)  );
+
+			// Value mismatch on entry with * prefix with same type
+			MetadataObjectAClone = *MetadataObjectA;
+			MetadataObjectAClone.SetStringField( TEXT("*NoCompare"), TEXT("NoCompare") );
+			FLocMetadataObject MetadataObjectAClone2 = *MetadataObjectA;
+			MetadataObjectAClone2.SetStringField( TEXT("*NoCompare"), TEXT("NoCompare2") );
+			CHECK_FALSE_MESSAGE(TEXT("MetadataObjectAClone == MetadataObjectAClone2"), MetadataObjectAClone.IsExactMatch( MetadataObjectAClone2) );
+
+			// Value and type mismatch on entry with * prefix
+			MetadataObjectAClone = *MetadataObjectA;
+			MetadataObjectAClone.SetStringField( TEXT("*NoCompare"), TEXT("NoCompare") );
+			MetadataObjectAClone2 = *MetadataObjectA;
+			MetadataObjectAClone2.SetBoolField( TEXT("*NoCompare"), true );
+			CHECK_FALSE_MESSAGE(TEXT("MetadataObjectAClone == MetadataObjectAClone2"),  MetadataObjectAClone.IsExactMatch( MetadataObjectAClone2) );
+
+			// Value mismatch on standard entry
+			MetadataObjectAClone = *MetadataObjectA;
+			MetadataObjectAClone.SetStringField( TEXT("DoCompare"), TEXT("DoCompare") );
+			MetadataObjectAClone2 = *MetadataObjectA;
+			MetadataObjectAClone2.SetStringField( TEXT("DoCompare"), TEXT("DoCompare2") );
+			CHECK_FALSE_MESSAGE(TEXT("MetadataObjectAClone == MetadataObjectAClone2"),  MetadataObjectAClone.IsExactMatch( MetadataObjectAClone2) );
+
+			// Value and type mismatch on standard entry
+			MetadataObjectAClone = *MetadataObjectA;
+			MetadataObjectAClone.SetStringField( TEXT("DoCompare"), TEXT("DoCompare") );
+			MetadataObjectAClone2 = *MetadataObjectA;
+			MetadataObjectAClone2.SetBoolField( TEXT("DoCompare"), true );
+			CHECK_FALSE_MESSAGE(TEXT("MetadataObjectAClone == MetadataObjectAClone2"),  MetadataObjectAClone.IsExactMatch( MetadataObjectAClone2) );
+		}
+
+		// Test less than operator
+		{
+			
+			// Adding standard entry that would appear before other entries
+			FLocMetadataObject MetadataObjectAClone = *MetadataObjectA;
+			MetadataObjectAClone.SetStringField( TEXT("ANewEntry"), TEXT("NewEntryValue") );
+			CHECK_MESSAGE(TEXT("MetadataObjectAClone < MetadataObjectA"), MetadataObjectAClone < *MetadataObjectA );
+			
+			// Adding standard entry that would appear after other entries
+			MetadataObjectAClone = *MetadataObjectA;
+			MetadataObjectAClone.SetStringField( TEXT("ZNewEntry"), TEXT("NewEntryValue") );
+			CHECK_MESSAGE(TEXT("MetadataObjectA < MetadataObjectAClone"), *MetadataObjectA < MetadataObjectAClone );
+
+			// Adding non-standard entry that would appear before other entries.  Note metadata with * prefix has no special treatment in less than operator
+			MetadataObjectAClone = *MetadataObjectA;
+			MetadataObjectAClone.SetStringField( TEXT("*NewEntry"), TEXT("NewEntryValue") );
+			CHECK_MESSAGE(TEXT("MetadataObjectAClone < MetadataObjectA"), MetadataObjectAClone < *MetadataObjectA );
+
+			// Value mismatch on entry with * prefix with same type. Note metadata with * prefix has no special treatment in less than operator
+			MetadataObjectAClone = *MetadataObjectA;
+			MetadataObjectAClone.SetStringField( TEXT("*NoCompare"), TEXT("NoCompare") );
+			FLocMetadataObject MetadataObjectAClone2 = *MetadataObjectA;
+			MetadataObjectAClone2.SetStringField( TEXT("*NoCompare"), TEXT("NoCompare2") );
+			CHECK_MESSAGE(TEXT("MetadataObjectAClone < MetadataObjectAClone2"), MetadataObjectAClone < MetadataObjectAClone2 );
+
+			// Value and type mismatch on entry with * prefix.  Note metadata with * prefix has no special treatment in less than operator
+			MetadataObjectAClone = *MetadataObjectA;
+			MetadataObjectAClone.SetBoolField( TEXT("*NoCompare"), true );
+			MetadataObjectAClone2 = *MetadataObjectA;
+			MetadataObjectAClone2.SetStringField( TEXT("*NoCompare"), TEXT("NoCompare") );
+			CHECK_MESSAGE(TEXT("MetadataObjectAClone < MetadataObjectAClone2"), MetadataObjectAClone < MetadataObjectAClone2 );
+		}
+	}
+}
+
+#endif //WITH_TESTS
\ No newline at end of file
diff --git a/Engine/Source/Runtime/Core/Tests/Internationalization/NumberFormatingRulesTest.cpp b/Engine/Source/Runtime/Core/Tests/Internationalization/NumberFormatingRulesTest.cpp
new file mode 100644
--- /dev/null
+++ b/Engine/Source/Runtime/Core/Tests/Internationalization/NumberFormatingRulesTest.cpp
@@ -0,0 +1,362 @@
+// Copyright Epic Games, Inc. All Rights Reserved.
+
+#include "CoreTypes.h"
+#include "Containers/UnrealString.h"
+#include "CoreGlobals.h"
+#include "Internationalization/Text.h"
+#include "Internationalization/Culture.h"
+#include "Internationalization/Internationalization.h"
+
+#if WITH_TESTS
+
+#include "Tests/TestHarnessAdapter.h"
+
+class FNumberFormattingRulesTestClass {
+public:
+	void Test(FNumberFormattingRulesTestClass* const ExplicitThis, const TCHAR* const Desc, const FText& A, const FText& B)
+	{
+		if (!A.EqualTo(B))
+		{
+			FAIL_CHECK(FString::Printf(TEXT("%s - A=%s B=%s"), Desc, *A.ToString(), *B.ToString()));
+		}
+	}
+
+	void NumberFormattingRulesTest() {
+
+		/**
+		*	List of functions and types to test.
+		*
+		*	Function	Type
+		*	---------	------
+		*	AsNumber 	double
+		*	AsNumber 	float
+		*	AsNumber 	uint8
+		*	AsNumber 	uint16
+		*	AsNumber 	uint32
+		*	AsNumber 	uint64
+		*	AsNumber 	int8
+		*	AsNumber 	int16
+		*	AsNumber 	int32
+		*	AsNumber 	int64
+		*	AsCurrency 	double
+		*	AsCurrency 	float
+		*	AsCurrency 	uint8
+		*	AsCurrency 	uint16
+		*	AsCurrency 	uint32
+		*	AsCurrency 	uint64
+		*	AsCurrency 	int8
+		*	AsCurrency 	int16
+		*	AsCurrency 	int32
+		*	AsCurrency 	int64
+		*	AsPercent 	double
+		*	AsPercent	float
+		*/
+
+#if UE_ENABLE_ICU
+		double DoubleValue = 12345678.901;
+		float FloatValue = 1234.567f;
+		double DoubleNegativeValue = -12345678.901;
+		float FloatNegativeValue = -1234.567f;
+		uint8 Uint8Value = 0xFF;
+		uint16 Uint16Value = 0xFFFF;
+		uint32 Uint32Value = 0xFFFFFFFF;
+		uint64 Uint64Value = 0x1999999999999999;
+		int8 Int8Value = 123;
+		int16 Int16Value = 12345;
+		int32 Int32Value = 12345;
+		int64 Int64Value = 12345;
+		int8 Int8NegativeValue = -123;
+		int16 Int16NegativeValue = -12345;
+		int32 Int32NegativeValue = -12345;
+		int64 Int64NegativeValue = -12345;
+		int64 Int64MinValue = MIN_int64;
+
+		FInternationalization& I18N = FInternationalization::Get();
+
+		FInternationalization::FCultureStateSnapshot OriginalCultureState;
+		I18N.BackupCultureState(OriginalCultureState);
+
+		if (I18N.SetCurrentCulture("en-US"))
+		{
+			Test(this, TEXT("Convert a Double to a number formatted correct for en-US"), FText::AsNumber(DoubleValue), FText::FromString(TEXT("12,345,678.901")));
+			Test(this, TEXT("Convert a Float to a number formatted correct for en-US"), FText::AsNumber(FloatValue), FText::FromString(TEXT("1,234.567")));
+			Test(this, TEXT("Convert a Negative Double to a number formatted correct for en-US"), FText::AsNumber(DoubleNegativeValue), FText::FromString(TEXT("-12,345,678.901")));
+			Test(this, TEXT("Convert a Negative Float to a number formatted correct for en-US"), FText::AsNumber(FloatNegativeValue), FText::FromString(TEXT("-1,234.567")));
+			Test(this, TEXT("Convert a uint8 to a number formatted correct for en-US"), FText::AsNumber(Uint8Value), FText::FromString(TEXT("255")));
+			Test(this, TEXT("Convert a uint16 to a number formatted correct for en-US"), FText::AsNumber(Uint16Value), FText::FromString(TEXT("65,535")));
+			Test(this, TEXT("Convert a uint32 to a number formatted correct for en-US"), FText::AsNumber(Uint32Value), FText::FromString(TEXT("4,294,967,295")));
+			Test(this, TEXT("Convert a uint64 to a number formatted correct for en-US"), FText::AsNumber(Uint64Value), FText::FromString(TEXT("1,844,674,407,370,955,161")));
+			Test(this, TEXT("Convert a int8 to a number formatted correct for en-US"), FText::AsNumber(Int8Value), FText::FromString(TEXT("123")));
+			Test(this, TEXT("Convert a int16 to a number formatted correct for en-US"), FText::AsNumber(Int16Value), FText::FromString(TEXT("12,345")));
+			Test(this, TEXT("Convert a int32 to a number formatted correct for en-US"), FText::AsNumber(Int32Value), FText::FromString(TEXT("12,345")));
+			Test(this, TEXT("Convert a int64 to a number formatted correct for en-US"), FText::AsNumber(Int64Value), FText::FromString(TEXT("12,345")));
+			Test(this, TEXT("Convert a Negative int8 to a number formatted correct for en-US"), FText::AsNumber(Int8NegativeValue), FText::FromString(TEXT("-123")));
+			Test(this, TEXT("Convert a Negative int16 to a number formatted correct for en-US"), FText::AsNumber(Int16NegativeValue), FText::FromString(TEXT("-12,345")));
+			Test(this, TEXT("Convert a Negative int32 to a number formatted correct for en-US"), FText::AsNumber(Int32NegativeValue), FText::FromString(TEXT("-12,345")));
+			Test(this, TEXT("Convert a Negative int64 to a number formatted correct for en-US"), FText::AsNumber(Int64NegativeValue), FText::FromString(TEXT("-12,345")));
+			Test(this, TEXT("Convert a Negative int64 to a number formatted correct for en-US"), FText::AsNumber(Int64MinValue), FText::FromString(TEXT("-9,223,372,036,854,775,808")));
+
+			// Test Number Formatting Options
+			{
+				FNumberFormattingOptions NumberFormattingOptions;
+				NumberFormattingOptions.MinimumIntegralDigits = 20;
+				NumberFormattingOptions.MaximumIntegralDigits = 20;
+				NumberFormattingOptions.MinimumFractionalDigits = 3;
+				NumberFormattingOptions.MaximumFractionalDigits = 3;
+				NumberFormattingOptions.UseGrouping = false;
+
+				Test(this, TEXT("Convert a Double to a number formatted correct for en-US"), FText::AsNumber(DoubleValue, &(NumberFormattingOptions)), FText::FromString(TEXT("00000000000012345678.901")));
+				Test(this, TEXT("Convert a Float to a number formatted correct for en-US"), FText::AsNumber(FloatValue, &(NumberFormattingOptions)), FText::FromString(TEXT("00000000000000001234.567")));
+				Test(this, TEXT("Convert a Negative Double to a number formatted correct for en-US"), FText::AsNumber(DoubleNegativeValue, &(NumberFormattingOptions)), FText::FromString(TEXT("-00000000000012345678.901")));
+				Test(this, TEXT("Convert a Negative Float to a number formatted correct for en-US"), FText::AsNumber(FloatNegativeValue, &(NumberFormattingOptions)), FText::FromString(TEXT("-00000000000000001234.567")));
+				Test(this, TEXT("Convert a uint8 to a number formatted correct for en-US"), FText::AsNumber(Uint8Value, &(NumberFormattingOptions)), FText::FromString(TEXT("00000000000000000255.000")));
+				Test(this, TEXT("Convert a uint16 to a number formatted correct for en-US"), FText::AsNumber(Uint16Value, &(NumberFormattingOptions)), FText::FromString(TEXT("00000000000000065535.000")));
+				Test(this, TEXT("Convert a uint32 to a number formatted correct for en-US"), FText::AsNumber(Uint32Value, &(NumberFormattingOptions)), FText::FromString(TEXT("00000000004294967295.000")));
+				Test(this, TEXT("Convert a uint64 to a number formatted correct for en-US"), FText::AsNumber(Uint64Value, &(NumberFormattingOptions)), FText::FromString(TEXT("01844674407370955161.000")));
+				Test(this, TEXT("Convert a int8 to a number formatted correct for en-US"), FText::AsNumber(Int8Value, &(NumberFormattingOptions)), FText::FromString(TEXT("00000000000000000123.000")));
+				Test(this, TEXT("Convert a int16 to a number formatted correct for en-US"), FText::AsNumber(Int16Value, &(NumberFormattingOptions)), FText::FromString(TEXT("00000000000000012345.000")));
+				Test(this, TEXT("Convert a int32 to a number formatted correct for en-US"), FText::AsNumber(Int32Value, &(NumberFormattingOptions)), FText::FromString(TEXT("00000000000000012345.000")));
+				Test(this, TEXT("Convert a int64 to a number formatted correct for en-US"), FText::AsNumber(Int64Value, &(NumberFormattingOptions)), FText::FromString(TEXT("00000000000000012345.000")));
+				Test(this, TEXT("Convert a Negative int8 to a number formatted correct for en-US"), FText::AsNumber(Int8NegativeValue, &(NumberFormattingOptions)), FText::FromString(TEXT("-00000000000000000123.000")));
+				Test(this, TEXT("Convert a Negative int16 to a number formatted correct for en-US"), FText::AsNumber(Int16NegativeValue, &(NumberFormattingOptions)), FText::FromString(TEXT("-00000000000000012345.000")));
+				Test(this, TEXT("Convert a Negative int32 to a number formatted correct for en-US"), FText::AsNumber(Int32NegativeValue, &(NumberFormattingOptions)), FText::FromString(TEXT("-00000000000000012345.000")));
+				Test(this, TEXT("Convert a Negative int64 to a number formatted correct for en-US"), FText::AsNumber(Int64NegativeValue, &(NumberFormattingOptions)), FText::FromString(TEXT("-00000000000000012345.000")));
+				Test(this, TEXT("Convert a Negative int64 to a number formatted correct for en-US"), FText::AsNumber(Int64MinValue, &(NumberFormattingOptions)), FText::FromString(TEXT("-09223372036854775808.000")));
+			}
+
+			{
+				FNumberFormattingOptions NumberFormattingOptions;
+				NumberFormattingOptions.MinimumIntegralDigits = 5;
+				NumberFormattingOptions.MaximumIntegralDigits = 5;
+				NumberFormattingOptions.UseGrouping = true;
+
+				Test(this, TEXT("Convert a 5 digit int to 5 grouped integral digits formatted correct for en-US"), FText::AsNumber(12345, &NumberFormattingOptions), FText::FromString(TEXT("12,345")));
+			}
+
+			{
+				FNumberFormattingOptions NumberFormattingOptions;
+				NumberFormattingOptions.MinimumIntegralDigits = 6;
+				NumberFormattingOptions.MaximumIntegralDigits = 6;
+				NumberFormattingOptions.UseGrouping = true;
+
+				Test(this, TEXT("Convert a 5 digit int to 6 grouped integral digits formatted correct for en-US"), FText::AsNumber(12345, &NumberFormattingOptions), FText::FromString(TEXT("012,345")));
+			}
+
+			{
+				FNumberFormattingOptions NumberFormattingOptions;
+				NumberFormattingOptions.MinimumIntegralDigits = 1;
+				NumberFormattingOptions.MaximumIntegralDigits = 1;
+				NumberFormattingOptions.MinimumFractionalDigits = 0;
+				NumberFormattingOptions.MaximumFractionalDigits = 0;
+				NumberFormattingOptions.UseGrouping = false;
+
+				NumberFormattingOptions.RoundingMode = ERoundingMode::HalfToEven;
+
+				Test(this, TEXT("Round a Double to a number formatted correct for en-US using HalfToEven"), FText::AsNumber(-1.50, &(NumberFormattingOptions)), FText::FromString(TEXT("-2")));
+				Test(this, TEXT("Round a Double to a number formatted correct for en-US using HalfToEven"), FText::AsNumber(-1.00, &(NumberFormattingOptions)), FText::FromString(TEXT("-1")));
+				Test(this, TEXT("Round a Double to a number formatted correct for en-US using HalfToEven"), FText::AsNumber(-0.75, &(NumberFormattingOptions)), FText::FromString(TEXT("-1")));
+				Test(this, TEXT("Round a Double to a number formatted correct for en-US using HalfToEven"), FText::AsNumber(-0.50, &(NumberFormattingOptions)), FText::FromString(TEXT("-0")));
+				Test(this, TEXT("Round a Double to a number formatted correct for en-US using HalfToEven"), FText::AsNumber(-0.25, &(NumberFormattingOptions)), FText::FromString(TEXT("-0")));
+				Test(this, TEXT("Round a Double to a number formatted correct for en-US using HalfToEven"), FText::AsNumber(-0.00, &(NumberFormattingOptions)), FText::FromString(TEXT("-0")));
+				Test(this, TEXT("Round a Double to a number formatted correct for en-US using HalfToEven"), FText::AsNumber(0.00, &(NumberFormattingOptions)), FText::FromString(TEXT("0")));
+				Test(this, TEXT("Round a Double to a number formatted correct for en-US using HalfToEven"), FText::AsNumber(0.25, &(NumberFormattingOptions)), FText::FromString(TEXT("0")));
+				Test(this, TEXT("Round a Double to a number formatted correct for en-US using HalfToEven"), FText::AsNumber(0.50, &(NumberFormattingOptions)), FText::FromString(TEXT("0")));
+				Test(this, TEXT("Round a Double to a number formatted correct for en-US using HalfToEven"), FText::AsNumber(0.75, &(NumberFormattingOptions)), FText::FromString(TEXT("1")));
+				Test(this, TEXT("Round a Double to a number formatted correct for en-US using HalfToEven"), FText::AsNumber(1.00, &(NumberFormattingOptions)), FText::FromString(TEXT("1")));
+				Test(this, TEXT("Round a Double to a number formatted correct for en-US using HalfToEven"), FText::AsNumber(1.50, &(NumberFormattingOptions)), FText::FromString(TEXT("2")));
+
+				NumberFormattingOptions.RoundingMode = ERoundingMode::HalfFromZero;
+
+				Test(this, TEXT("Round a Double to a number formatted correct for en-US using HalfFromZero"), FText::AsNumber(-1.00, &(NumberFormattingOptions)), FText::FromString(TEXT("-1")));
+				Test(this, TEXT("Round a Double to a number formatted correct for en-US using HalfFromZero"), FText::AsNumber(-0.75, &(NumberFormattingOptions)), FText::FromString(TEXT("-1")));
+				Test(this, TEXT("Round a Double to a number formatted correct for en-US using HalfFromZero"), FText::AsNumber(-0.50, &(NumberFormattingOptions)), FText::FromString(TEXT("-1")));
+				Test(this, TEXT("Round a Double to a number formatted correct for en-US using HalfFromZero"), FText::AsNumber(-0.25, &(NumberFormattingOptions)), FText::FromString(TEXT("-0")));
+				Test(this, TEXT("Round a Double to a number formatted correct for en-US using HalfFromZero"), FText::AsNumber(-0.00, &(NumberFormattingOptions)), FText::FromString(TEXT("-0")));
+				Test(this, TEXT("Round a Double to a number formatted correct for en-US using HalfFromZero"), FText::AsNumber(0.00, &(NumberFormattingOptions)), FText::FromString(TEXT("0")));
+				Test(this, TEXT("Round a Double to a number formatted correct for en-US using HalfFromZero"), FText::AsNumber(0.25, &(NumberFormattingOptions)), FText::FromString(TEXT("0")));
+				Test(this, TEXT("Round a Double to a number formatted correct for en-US using HalfFromZero"), FText::AsNumber(0.50, &(NumberFormattingOptions)), FText::FromString(TEXT("1")));
+				Test(this, TEXT("Round a Double to a number formatted correct for en-US using HalfFromZero"), FText::AsNumber(0.75, &(NumberFormattingOptions)), FText::FromString(TEXT("1")));
+				Test(this, TEXT("Round a Double to a number formatted correct for en-US using HalfFromZero"), FText::AsNumber(1.00, &(NumberFormattingOptions)), FText::FromString(TEXT("1")));
+
+				NumberFormattingOptions.RoundingMode = ERoundingMode::HalfToZero;
+
+				Test(this, TEXT("Round a Double to a number formatted correct for en-US using HalfToZero"), FText::AsNumber(-1.00, &(NumberFormattingOptions)), FText::FromString(TEXT("-1")));
+				Test(this, TEXT("Round a Double to a number formatted correct for en-US using HalfToZero"), FText::AsNumber(-0.75, &(NumberFormattingOptions)), FText::FromString(TEXT("-1")));
+				Test(this, TEXT("Round a Double to a number formatted correct for en-US using HalfToZero"), FText::AsNumber(-0.50, &(NumberFormattingOptions)), FText::FromString(TEXT("-0")));
+				Test(this, TEXT("Round a Double to a number formatted correct for en-US using HalfToZero"), FText::AsNumber(-0.25, &(NumberFormattingOptions)), FText::FromString(TEXT("-0")));
+				Test(this, TEXT("Round a Double to a number formatted correct for en-US using HalfToZero"), FText::AsNumber(-0.00, &(NumberFormattingOptions)), FText::FromString(TEXT("-0")));
+				Test(this, TEXT("Round a Double to a number formatted correct for en-US using HalfToZero"), FText::AsNumber(0.00, &(NumberFormattingOptions)), FText::FromString(TEXT("0")));
+				Test(this, TEXT("Round a Double to a number formatted correct for en-US using HalfToZero"), FText::AsNumber(0.25, &(NumberFormattingOptions)), FText::FromString(TEXT("0")));
+				Test(this, TEXT("Round a Double to a number formatted correct for en-US using HalfToZero"), FText::AsNumber(0.50, &(NumberFormattingOptions)), FText::FromString(TEXT("0")));
+				Test(this, TEXT("Round a Double to a number formatted correct for en-US using HalfToZero"), FText::AsNumber(0.75, &(NumberFormattingOptions)), FText::FromString(TEXT("1")));
+				Test(this, TEXT("Round a Double to a number formatted correct for en-US using HalfToZero"), FText::AsNumber(1.00, &(NumberFormattingOptions)), FText::FromString(TEXT("1")));
+
+				NumberFormattingOptions.RoundingMode = ERoundingMode::FromZero;
+
+				Test(this, TEXT("Round a Double to a number formatted correct for en-US using FromZero"), FText::AsNumber(-1.0, &(NumberFormattingOptions)), FText::FromString(TEXT("-1")));
+				Test(this, TEXT("Round a Double to a number formatted correct for en-US using FromZero"), FText::AsNumber(-0.5, &(NumberFormattingOptions)), FText::FromString(TEXT("-1")));
+				Test(this, TEXT("Round a Double to a number formatted correct for en-US using FromZero"), FText::AsNumber(-0.0, &(NumberFormattingOptions)), FText::FromString(TEXT("-0")));
+				Test(this, TEXT("Round a Double to a number formatted correct for en-US using FromZero"), FText::AsNumber(0.0, &(NumberFormattingOptions)), FText::FromString(TEXT("0")));
+				Test(this, TEXT("Round a Double to a number formatted correct for en-US using FromZero"), FText::AsNumber(0.5, &(NumberFormattingOptions)), FText::FromString(TEXT("1")));
+				Test(this, TEXT("Round a Double to a number formatted correct for en-US using FromZero"), FText::AsNumber(1.0, &(NumberFormattingOptions)), FText::FromString(TEXT("1")));
+
+				NumberFormattingOptions.RoundingMode = ERoundingMode::ToZero;
+
+				Test(this, TEXT("Round a Double to a number formatted correct for en-US using ToZero"), FText::AsNumber(-1.0, &(NumberFormattingOptions)), FText::FromString(TEXT("-1")));
+				Test(this, TEXT("Round a Double to a number formatted correct for en-US using ToZero"), FText::AsNumber(-0.5, &(NumberFormattingOptions)), FText::FromString(TEXT("-0")));
+				Test(this, TEXT("Round a Double to a number formatted correct for en-US using ToZero"), FText::AsNumber(-0.0, &(NumberFormattingOptions)), FText::FromString(TEXT("-0")));
+				Test(this, TEXT("Round a Double to a number formatted correct for en-US using ToZero"), FText::AsNumber(0.0, &(NumberFormattingOptions)), FText::FromString(TEXT("0")));
+				Test(this, TEXT("Round a Double to a number formatted correct for en-US using ToZero"), FText::AsNumber(0.5, &(NumberFormattingOptions)), FText::FromString(TEXT("0")));
+				Test(this, TEXT("Round a Double to a number formatted correct for en-US using ToZero"), FText::AsNumber(1.0, &(NumberFormattingOptions)), FText::FromString(TEXT("1")));
+
+				NumberFormattingOptions.RoundingMode = ERoundingMode::ToNegativeInfinity;
+
+				Test(this, TEXT("Round a Double to a number formatted correct for en-US using ToNegativeInfinity"), FText::AsNumber(-1.0, &(NumberFormattingOptions)), FText::FromString(TEXT("-1")));
+				Test(this, TEXT("Round a Double to a number formatted correct for en-US using ToNegativeInfinity"), FText::AsNumber(-0.5, &(NumberFormattingOptions)), FText::FromString(TEXT("-1")));
+				Test(this, TEXT("Round a Double to a number formatted correct for en-US using ToNegativeInfinity"), FText::AsNumber(-0.0, &(NumberFormattingOptions)), FText::FromString(TEXT("-0")));
+				Test(this, TEXT("Round a Double to a number formatted correct for en-US using ToNegativeInfinity"), FText::AsNumber(0.0, &(NumberFormattingOptions)), FText::FromString(TEXT("0")));
+				Test(this, TEXT("Round a Double to a number formatted correct for en-US using ToNegativeInfinity"), FText::AsNumber(0.5, &(NumberFormattingOptions)), FText::FromString(TEXT("0")));
+				Test(this, TEXT("Round a Double to a number formatted correct for en-US using ToNegativeInfinity"), FText::AsNumber(1.0, &(NumberFormattingOptions)), FText::FromString(TEXT("1")));
+
+				NumberFormattingOptions.RoundingMode = ERoundingMode::ToPositiveInfinity;
+
+				Test(this, TEXT("Round a Double to a number formatted correct for en-US using ToPositiveInfinity"), FText::AsNumber(-1.0, &(NumberFormattingOptions)), FText::FromString(TEXT("-1")));
+				Test(this, TEXT("Round a Double to a number formatted correct for en-US using ToPositiveInfinity"), FText::AsNumber(-0.5, &(NumberFormattingOptions)), FText::FromString(TEXT("-0")));
+				Test(this, TEXT("Round a Double to a number formatted correct for en-US using ToPositiveInfinity"), FText::AsNumber(-0.0, &(NumberFormattingOptions)), FText::FromString(TEXT("-0")));
+				Test(this, TEXT("Round a Double to a number formatted correct for en-US using ToPositiveInfinity"), FText::AsNumber(0.0, &(NumberFormattingOptions)), FText::FromString(TEXT("0")));
+				Test(this, TEXT("Round a Double to a number formatted correct for en-US using ToPositiveInfinity"), FText::AsNumber(0.5, &(NumberFormattingOptions)), FText::FromString(TEXT("1")));
+				Test(this, TEXT("Round a Double to a number formatted correct for en-US using ToPositiveInfinity"), FText::AsNumber(1.0, &(NumberFormattingOptions)), FText::FromString(TEXT("1")));
+			}
+
+			PRAGMA_DISABLE_DEPRECATION_WARNINGS
+			Test(this, TEXT("Convert a Double to a currency formatted correct for en-US"), FText::AsCurrency(DoubleValue), FText::FromString(TEXT("$12,345,678.90")));
+			Test(this, TEXT("Convert a Float to a currency formatted correct for en-US"), FText::AsCurrency(FloatValue), FText::FromString(TEXT("$1,234.57")));
+			Test(this, TEXT("Convert a Negative Double to a currency formatted correct for en-US"), FText::AsCurrency(DoubleNegativeValue), FText::FromString(TEXT("-$12,345,678.90")));
+			Test(this, TEXT("Convert a Negative Float to a currency formatted correct for en-US"), FText::AsCurrency(FloatNegativeValue), FText::FromString(TEXT("-$1,234.57")));
+			Test(this, TEXT("Convert a uint8 to a currency formatted correct for en-US"), FText::AsCurrency(Uint8Value), FText::FromString(TEXT("$255.00")));
+			Test(this, TEXT("Convert a uint16 to a currency formatted correct for en-US"), FText::AsCurrency(Uint16Value), FText::FromString(TEXT("$65,535.00")));
+			Test(this, TEXT("Convert a uint32 to a currency formatted correct for en-US"), FText::AsCurrency(Uint32Value), FText::FromString(TEXT("$4,294,967,295.00")));
+			Test(this, TEXT("Convert a uint64 to a currency formatted correct for en-US"), FText::AsCurrency(Uint64Value), FText::FromString(TEXT("$1,844,674,407,370,955,161.00")));
+			Test(this, TEXT("Convert a int8 to a currency formatted correct for en-US"), FText::AsCurrency(Int8Value), FText::FromString(TEXT("$123.00")));
+			Test(this, TEXT("Convert a int16 to a currency formatted correct for en-US"), FText::AsCurrency(Int16Value), FText::FromString(TEXT("$12,345.00")));
+			Test(this, TEXT("Convert a int32 to a currency formatted correct for en-US"), FText::AsCurrency(Int32Value), FText::FromString(TEXT("$12,345.00")));
+			Test(this, TEXT("Convert a int64 to a currency formatted correct for en-US"), FText::AsCurrency(Int64Value), FText::FromString(TEXT("$12,345.00")));
+			Test(this, TEXT("Convert a Negative int8 to a currency formatted correct for en-US"), FText::AsCurrency(Int8NegativeValue), FText::FromString(TEXT("-$123.00")));
+			Test(this, TEXT("Convert a Negative int16 to a currency formatted correct for en-US"), FText::AsCurrency(Int16NegativeValue), FText::FromString(TEXT("-$12,345.00")));
+			Test(this, TEXT("Convert a Negative int32 to a currency formatted correct for en-US"), FText::AsCurrency(Int32NegativeValue), FText::FromString(TEXT("-$12,345.00")));
+			Test(this, TEXT("Convert a Negative int64 to a currency formatted correct for en-US"), FText::AsCurrency(Int64NegativeValue), FText::FromString(TEXT("-$12,345.00")));
+			Test(this, TEXT("Convert a Negative int64 to a currency formatted correct for en-US"), FText::AsCurrency(Int64MinValue), FText::FromString(TEXT("-$9,223,372,036,854,775,808.00")));
+			PRAGMA_ENABLE_DEPRECATION_WARNINGS
+
+			Test(this, TEXT("Convert a Double to a percent formatted correct for en-US"), FText::AsPercent(DoubleValue), FText::FromString(TEXT("1,234,567,890%")));
+			Test(this, TEXT("Convert a Float to a percent formatted correct for en-US"), FText::AsPercent(FloatValue), FText::FromString(TEXT("123,457%")));
+			Test(this, TEXT("Convert a Negative Double to a percent formatted correct for en-US"), FText::AsPercent(DoubleNegativeValue), FText::FromString(TEXT("-1,234,567,890%")));
+			Test(this, TEXT("Convert a Negative Float to a percent formatted correct for en-US"), FText::AsPercent(FloatNegativeValue), FText::FromString(TEXT("-123,457%")));
+		}
+		else
+		{
+			WARN(FString::Printf(TEXT("Internationalization data for %s missing - test is partially disabled."), TEXT("en-US")));
+		}
+
+		if (I18N.SetCurrentCulture("hi-IN"))
+		{
+			Test(this, TEXT("Convert a Double to a number formatted correct for hi-IN"), FText::AsNumber(DoubleValue), FText::FromString(TEXT("1,23,45,678.901")));
+			Test(this, TEXT("Convert a Float to a number formatted correct for hi-IN"), FText::AsNumber(FloatValue), FText::FromString(TEXT("1,234.567")));
+			Test(this, TEXT("Convert a Negative Double to a number formatted correct for hi-IN"), FText::AsNumber(DoubleNegativeValue), FText::FromString(TEXT("-1,23,45,678.901")));
+			Test(this, TEXT("Convert a Negative Float to a number formatted correct for hi-IN"), FText::AsNumber(FloatNegativeValue), FText::FromString(TEXT("-1,234.567")));
+			Test(this, TEXT("Convert a uint8 to a number formatted correct for hi-IN"), FText::AsNumber(Uint8Value), FText::FromString(TEXT("255")));
+			Test(this, TEXT("Convert a uint16 to a number formatted correct for hi-IN"), FText::AsNumber(Uint16Value), FText::FromString(TEXT("65,535")));
+			Test(this, TEXT("Convert a uint32 to a number formatted correct for hi-IN"), FText::AsNumber(Uint32Value), FText::FromString(TEXT("4,29,49,67,295")));
+			Test(this, TEXT("Convert a uint64 to a number formatted correct for hi-IN"), FText::AsNumber(Uint64Value), FText::FromString(TEXT("18,44,67,44,07,37,09,55,161")));
+			Test(this, TEXT("Convert a int8 to a number formatted correct for hi-IN"), FText::AsNumber(Int8Value), FText::FromString(TEXT("123")));
+			Test(this, TEXT("Convert a int16 to a number formatted correct for hi-IN"), FText::AsNumber(Int16Value), FText::FromString(TEXT("12,345")));
+			Test(this, TEXT("Convert a int32 to a number formatted correct for hi-IN"), FText::AsNumber(Int32Value), FText::FromString(TEXT("12,345")));
+			Test(this, TEXT("Convert a int64 to a number formatted correct for hi-IN"), FText::AsNumber(Int64Value), FText::FromString(TEXT("12,345")));
+			Test(this, TEXT("Convert a Negative int8 to a number formatted correct for hi-IN"), FText::AsNumber(Int8NegativeValue), FText::FromString(TEXT("-123")));
+			Test(this, TEXT("Convert a Negative int16 to a number formatted correct for hi-IN"), FText::AsNumber(Int16NegativeValue), FText::FromString(TEXT("-12,345")));
+			Test(this, TEXT("Convert a Negative int32 to a number formatted correct for hi-IN"), FText::AsNumber(Int32NegativeValue), FText::FromString(TEXT("-12,345")));
+			Test(this, TEXT("Convert a Negative int64 to a number formatted correct for hi-IN"), FText::AsNumber(Int64NegativeValue), FText::FromString(TEXT("-12,345")));
+			Test(this, TEXT("Convert a Negative int64 to a number formatted correct for hi-IN"), FText::AsNumber(Int64MinValue), FText::FromString(TEXT("-92,23,37,20,36,85,47,75,808")));
+
+			{
+				const FCulturePtr& InvariantCulture = I18N.GetInvariantCulture();
+				Test(this, TEXT("Convert a Double to a number formatted correct for hi-IN but as invariant"), FText::AsNumber(DoubleValue, NULL, InvariantCulture), FText::FromString(TEXT("12345678.901")));
+				Test(this, TEXT("Convert a Float to a number formatted correct for hi-IN but as invariant"), FText::AsNumber(FloatValue, NULL, InvariantCulture), FText::FromString(TEXT("1234.567017")));
+				Test(this, TEXT("Convert a Negative Double to a number formatted correct for hi-IN but as invariant"), FText::AsNumber(DoubleNegativeValue, NULL, InvariantCulture), FText::FromString(TEXT("-12345678.901")));
+				Test(this, TEXT("Convert a Negative Float to a number formatted correct for hi-IN but as invariant"), FText::AsNumber(FloatNegativeValue, NULL, InvariantCulture), FText::FromString(TEXT("-1234.567017")));
+				Test(this, TEXT("Convert a uint8 to a number formatted correct for hi-IN but as invariant"), FText::AsNumber(Uint8Value, NULL, InvariantCulture), FText::FromString(TEXT("255")));
+				Test(this, TEXT("Convert a uint16 to a number formatted correct for hi-IN but as invariant"), FText::AsNumber(Uint16Value, NULL, InvariantCulture), FText::FromString(TEXT("65535")));
+				Test(this, TEXT("Convert a uint32 to a number formatted correct for hi-IN but as invariant"), FText::AsNumber(Uint32Value, NULL, InvariantCulture), FText::FromString(TEXT("4294967295")));
+				Test(this, TEXT("Convert a uint64 to a number formatted correct for hi-IN but as invariant"), FText::AsNumber(Uint64Value, NULL, InvariantCulture), FText::FromString(TEXT("1844674407370955161")));
+				Test(this, TEXT("Convert a int8 to a number formatted correct for hi-IN but as invariant"), FText::AsNumber(Int8Value, NULL, InvariantCulture), FText::FromString(TEXT("123")));
+				Test(this, TEXT("Convert a int16 to a number formatted correct for hi-IN but as invariant"), FText::AsNumber(Int16Value, NULL, InvariantCulture), FText::FromString(TEXT("12345")));
+				Test(this, TEXT("Convert a int32 to a number formatted correct for hi-IN but as invariant"), FText::AsNumber(Int32Value, NULL, InvariantCulture), FText::FromString(TEXT("12345")));
+				Test(this, TEXT("Convert a int64 to a number formatted correct for hi-IN but as invariant"), FText::AsNumber(Int64Value, NULL, InvariantCulture), FText::FromString(TEXT("12345")));
+				Test(this, TEXT("Convert a Negative int8 to a number formatted correct for hi-IN but as invariant"), FText::AsNumber(Int8NegativeValue, NULL, InvariantCulture), FText::FromString(TEXT("-123")));
+				Test(this, TEXT("Convert a Negative int16 to a number formatted correct for hi-IN but as invariant"), FText::AsNumber(Int16NegativeValue, NULL, InvariantCulture), FText::FromString(TEXT("-12345")));
+				Test(this, TEXT("Convert a Negative int32 to a number formatted correct for hi-IN but as invariant"), FText::AsNumber(Int32NegativeValue, NULL, InvariantCulture), FText::FromString(TEXT("-12345")));
+				Test(this, TEXT("Convert a Negative int64 to a number formatted correct for hi-IN but as invariant"), FText::AsNumber(Int64NegativeValue, NULL, InvariantCulture), FText::FromString(TEXT("-12345")));
+				Test(this, TEXT("Convert a Negative int64 to a number formatted correct for hi-IN but as invariant"), FText::AsNumber(Int64MinValue, NULL, InvariantCulture), FText::FromString(TEXT("-9223372036854775808")));
+			}
+
+			PRAGMA_DISABLE_DEPRECATION_WARNINGS
+			Test(this, TEXT("Convert a Double to a currency formatted correct for hi-IN"), FText::AsCurrency(DoubleValue), FText::FromString(TEXT("\x20B9") TEXT("1,23,45,678.90")));
+			Test(this, TEXT("Convert a Float to a currency formatted correct for hi-IN"), FText::AsCurrency(FloatValue), FText::FromString(TEXT("\x20B9") TEXT("1,234.57")));
+			Test(this, TEXT("Convert a Negative Double to a currency formatted correct for hi-IN"), FText::AsCurrency(DoubleNegativeValue), FText::FromString(TEXT("-") TEXT("\x20B9") TEXT("1,23,45,678.90")));
+			Test(this, TEXT("Convert a Negative Float to a currency formatted correct for hi-IN"), FText::AsCurrency(FloatNegativeValue), FText::FromString(TEXT("-") TEXT("\x20B9") TEXT("1,234.57")));
+			Test(this, TEXT("Convert a uint8 to a currency formatted correct for hi-IN"), FText::AsCurrency(Uint8Value), FText::FromString(TEXT("\x20B9") TEXT("255.00")));
+			Test(this, TEXT("Convert a uint16 to a currency formatted correct for hi-IN"), FText::AsCurrency(Uint16Value), FText::FromString(TEXT("\x20B9") TEXT("65,535.00")));
+			Test(this, TEXT("Convert a uint32 to a currency formatted correct for hi-IN"), FText::AsCurrency(Uint32Value), FText::FromString(TEXT("\x20B9") TEXT("4,29,49,67,295.00")));
+			Test(this, TEXT("Convert a uint64 to a currency formatted correct for hi-IN"), FText::AsCurrency(Uint64Value), FText::FromString(TEXT("\x20B9") TEXT("18,44,67,44,07,37,09,55,161.00")));
+			Test(this, TEXT("Convert a int8 to a currency formatted correct for hi-IN"), FText::AsCurrency(Int8Value), FText::FromString(TEXT("\x20B9") TEXT("123.00")));
+			Test(this, TEXT("Convert a int16 to a currency formatted correct for hi-IN"), FText::AsCurrency(Int16Value), FText::FromString(TEXT("\x20B9") TEXT("12,345.00")));
+			Test(this, TEXT("Convert a int32 to a currency formatted correct for hi-IN"), FText::AsCurrency(Int32Value), FText::FromString(TEXT("\x20B9") TEXT("12,345.00")));
+			Test(this, TEXT("Convert a int64 to a currency formatted correct for hi-IN"), FText::AsCurrency(Int64Value), FText::FromString(TEXT("\x20B9") TEXT("12,345.00")));
+			Test(this, TEXT("Convert a Negative int8 to a currency formatted correct for hi-IN"), FText::AsCurrency(Int8NegativeValue), FText::FromString(TEXT("-") TEXT("\x20B9") TEXT("123.00")));
+			Test(this, TEXT("Convert a Negative int16 to a currency formatted correct for hi-IN"), FText::AsCurrency(Int16NegativeValue), FText::FromString(TEXT("-") TEXT("\x20B9") TEXT("12,345.00")));
+			Test(this, TEXT("Convert a Negative int32 to a currency formatted correct for hi-IN"), FText::AsCurrency(Int32NegativeValue), FText::FromString(TEXT("-") TEXT("\x20B9") TEXT("12,345.00")));
+			Test(this, TEXT("Convert a Negative int64 to a currency formatted correct for hi-IN"), FText::AsCurrency(Int64NegativeValue), FText::FromString(TEXT("-") TEXT("\x20B9") TEXT("12,345.00")));
+			Test(this, TEXT("Convert a Negative int64 to a currency formatted correct for hi-IN"), FText::AsCurrency(Int64MinValue), FText::FromString(TEXT("-") TEXT("\x20B9") TEXT("92,23,37,20,36,85,47,75,808.00")));
+			PRAGMA_ENABLE_DEPRECATION_WARNINGS
+
+			Test(this, TEXT("Convert a Double to a percent formatted correct for hi-IN"), FText::AsPercent(DoubleValue), FText::FromString(TEXT("1,23,45,67,890%")));
+			Test(this, TEXT("Convert a Float to a percent formatted correct for hi-IN"), FText::AsPercent(FloatValue), FText::FromString(TEXT("1,23,457%")));
+			Test(this, TEXT("Convert a Negative Double to a percent formatted correct for hi-IN"), FText::AsPercent(DoubleNegativeValue), FText::FromString(TEXT("-1,23,45,67,890%")));
+			Test(this, TEXT("Convert a Negative Float to a percent formatted correct for hi-IN"), FText::AsPercent(FloatNegativeValue), FText::FromString(TEXT("-1,23,457%")));
+		}
+		else
+		{
+			WARN(FString::Printf(TEXT("Internationalization data for %s missing - test is partially disabled."), TEXT("hi-IN")));
+		}
+
+		I18N.RestoreCultureState(OriginalCultureState);
+
+		FText Number = FText::AsNumber(Int64NegativeValue);
+		FText Percent = FText::AsPercent(DoubleValue);
+		FText Currency = FText::AsCurrencyBase(Int64Value, TEXT("USD"));
+
+		if (GIsEditor && (Number.IsTransient() || Percent.IsTransient() || Currency.IsTransient()))
+		{
+			FAIL_CHECK(TEXT("Number formatting functions should not produce transient text in the editor"));
+		}
+		if (!GIsEditor && (!Number.IsTransient() || !Percent.IsTransient() || !Currency.IsTransient()))
+		{
+			FAIL_CHECK(TEXT("Number formatting functions should always produce transient text outside of the editor"));
+		}
+		if (Number.IsCultureInvariant() || Percent.IsCultureInvariant() || Currency.IsCultureInvariant())
+		{
+			FAIL_CHECK(TEXT("Number formatting functions should never produce a Culture Invariant Text"));
+		}
+#else
+		WARN("ICU is disabled thus locale-aware number formatting is disabled.");
+#endif
+	}
+};
+
+TEST_CASE_NAMED(FNumberFormattingRulesTest, "System::Core::Misc::Number Formatting Rules", "[.][EditorContext][ClientContext][EngineFilter]")
+{
+	FNumberFormattingRulesTestClass Instanse;
+	Instanse.NumberFormattingRulesTest();
+}
+
+#endif //WITH_TESTS
\ No newline at end of file
diff --git a/Engine/Source/Runtime/Core/Tests/Internationalization/TextTest.cpp b/Engine/Source/Runtime/Core/Tests/Internationalization/TextTest.cpp
new file mode 100644
--- /dev/null
+++ b/Engine/Source/Runtime/Core/Tests/Internationalization/TextTest.cpp
@@ -0,0 +1,1557 @@
+// Copyright Epic Games, Inc. All Rights Reserved.
+
+#include "CoreTypes.h"
+#include "Templates/UnrealTemplate.h"
+#include "Containers/Array.h"
+#include "Containers/UnrealString.h"
+#include "Misc/DateTime.h"
+#include "CoreGlobals.h"
+#include "Internationalization/Text.h"
+#include "Internationalization/Culture.h"
+#include "Internationalization/FastDecimalFormat.h"
+#include "Internationalization/StringTableRegistry.h"
+#include "Internationalization/Internationalization.h"
+#include "Internationalization/Cultures/LeetCulture.h"
+#include "Serialization/MemoryWriter.h"
+#include "Serialization/MemoryReader.h"
+#include "Internationalization/ICUUtilities.h"
+
+#if WITH_TESTS
+
+#include "Tests/TestHarnessAdapter.h"
+
+#define LOCTEXT_NAMESPACE "Core.Tests.TextFormatTest"
+
+class  FTextTestClass {
+public:
+		FText FormatWithoutArguments(const FText& Pattern)
+		{
+			FFormatOrderedArguments Arguments;
+			return FText::Format(Pattern, Arguments);
+		}
+
+		void ArrayToString(const TArray<FString>& Array, FString& String)
+		{
+			const int32 Count = Array.Num();
+			for (int32 i = 0; i < Count; ++i)
+			{
+				if (i > 0)
+				{
+					String += TEXT(", ");
+				}
+				String += Array[i];
+			}
+		}
+
+		void TestPatternParameterEnumeration(FTextTestClass& Test, const FText& Pattern, TArray<FString>& ActualParameters, const TArray<FString>& ExpectedParameters)
+		{
+			ActualParameters.Empty(ExpectedParameters.Num());
+			FText::GetFormatPatternParameters(Pattern, ActualParameters);
+			if (ActualParameters != ExpectedParameters)
+			{
+				FString ActualParametersString;
+				ArrayToString(ActualParameters, ActualParametersString);
+				FString ExpectedParametersString;
+				ArrayToString(ExpectedParameters, ExpectedParametersString);
+				FAIL_CHECK(FString::Printf(TEXT("\"%s\" contains parameters (%s) but expected parameters (%s)."), *(Pattern.ToString()), *(ActualParametersString), *(ExpectedParametersString)));
+			}
+		}
+
+		void TestIdentical(FTextTestClass& Test, const FText& One, const FText& Two, const ETextIdenticalModeFlags CompareFlags, const bool bExpectedResult, const int32 TestLine)
+		{
+			const bool bActualResult = One.IdenticalTo(Two, CompareFlags);
+			if (bActualResult != bExpectedResult)
+			{
+				FAIL_CHECK(FString::Printf(TEXT("FText(\"%s\").IdenticalTo(FText(\"%s\")) on line %d produced %s when it was expected to produce %s."), *One.ToString(), *Two.ToString(), TestLine, *LexToString(bActualResult), *LexToString(bExpectedResult)));
+			}
+		}
+		void TextTest() {
+			FInternationalization& I18N = FInternationalization::Get();
+
+			FInternationalization::FCultureStateSnapshot OriginalCultureState;
+			I18N.BackupCultureState(OriginalCultureState);
+
+			FText ArgText0 = INVTEXT("Arg0");
+			FText ArgText1 = INVTEXT("Arg1");
+			FText ArgText2 = INVTEXT("Arg2");
+			FText ArgText3 = INVTEXT("Arg3");
+
+#define TEST( A, B, CompareFlags, Expected ) TestIdentical(*this, A, B, CompareFlags, Expected, __LINE__)
+			{
+				const int32 TestNumber1 = 10;
+				const int32 TestNumber2 = 20;
+				const FDateTime TestDateTime = FDateTime(1991, 6, 21, 9, 30);
+				const FText TestIdenticalStr1 = LOCTEXT("TestIdenticalStr1", "Str1");
+				const FText TestIdenticalStr2 = LOCTEXT("TestIdenticalStr2", "Str2");
+
+				TEST(TestIdenticalStr1, TestIdenticalStr1, ETextIdenticalModeFlags::None, true);
+				TEST(TestIdenticalStr1, TestIdenticalStr2, ETextIdenticalModeFlags::None, false);
+				TEST(TestIdenticalStr1, TestIdenticalStr1, ETextIdenticalModeFlags::DeepCompare | ETextIdenticalModeFlags::LexicalCompareInvariants, true);
+				TEST(TestIdenticalStr1, TestIdenticalStr2, ETextIdenticalModeFlags::DeepCompare | ETextIdenticalModeFlags::LexicalCompareInvariants, false);
+
+				TEST(FText::AsCultureInvariant(TEXT("Wooble")), FText::AsCultureInvariant(TEXT("Wooble")), ETextIdenticalModeFlags::None, false);
+				TEST(FText::FromString(TEXT("Wooble")), FText::FromString(TEXT("Wooble")), ETextIdenticalModeFlags::None, false);
+				TEST(FText::AsCultureInvariant(TEXT("Wooble")), FText::AsCultureInvariant(TEXT("Wooble")), ETextIdenticalModeFlags::LexicalCompareInvariants, true);
+				TEST(FText::FromString(TEXT("Wooble")), FText::FromString(TEXT("Wooble")), ETextIdenticalModeFlags::LexicalCompareInvariants, true);
+				TEST(FText::AsCultureInvariant(TEXT("Wooble")), FText::AsCultureInvariant(TEXT("Wooble2")), ETextIdenticalModeFlags::LexicalCompareInvariants, false);
+				TEST(FText::FromString(TEXT("Wooble")), FText::FromString(TEXT("Wooble2")), ETextIdenticalModeFlags::LexicalCompareInvariants, false);
+
+				TEST(FText::Format(LOCTEXT("TestIdenticalPattern", "This takes an arg {0}"), ArgText0), FText::Format(LOCTEXT("TestIdenticalPattern", "This takes an arg {0}"), ArgText0), ETextIdenticalModeFlags::None, false);
+				TEST(FText::Format(LOCTEXT("TestIdenticalPattern", "This takes an arg {0}"), ArgText0), FText::Format(LOCTEXT("TestIdenticalPattern", "This takes an arg {0}"), ArgText0), ETextIdenticalModeFlags::DeepCompare | ETextIdenticalModeFlags::LexicalCompareInvariants, true);
+				TEST(FText::Format(LOCTEXT("TestIdenticalPattern", "This takes an arg {0}"), ArgText0), FText::Format(LOCTEXT("TestIdenticalPattern", "This takes an arg {0}"), ArgText1), ETextIdenticalModeFlags::DeepCompare | ETextIdenticalModeFlags::LexicalCompareInvariants, false);
+				TEST(FText::Format(LOCTEXT("TestIdenticalPattern", "This takes an arg {0}"), ArgText0), FText::Format(LOCTEXT("TestIdenticalPattern2", "This takes an arg {0}!"), ArgText0), ETextIdenticalModeFlags::DeepCompare | ETextIdenticalModeFlags::LexicalCompareInvariants, false);
+
+				TEST(FText::AsDate(TestDateTime), FText::AsDate(TestDateTime), ETextIdenticalModeFlags::None, false);
+				TEST(FText::AsDate(TestDateTime), FText::AsDate(TestDateTime), ETextIdenticalModeFlags::DeepCompare, true);
+				TEST(FText::AsTime(TestDateTime), FText::AsTime(TestDateTime), ETextIdenticalModeFlags::None, false);
+				TEST(FText::AsTime(TestDateTime), FText::AsTime(TestDateTime), ETextIdenticalModeFlags::DeepCompare, true);
+				TEST(FText::AsDateTime(TestDateTime), FText::AsDateTime(TestDateTime), ETextIdenticalModeFlags::None, false);
+				TEST(FText::AsDateTime(TestDateTime), FText::AsDateTime(TestDateTime), ETextIdenticalModeFlags::DeepCompare, true);
+
+				TEST(FText::AsNumber(TestNumber1), FText::AsNumber(TestNumber1), ETextIdenticalModeFlags::None, false);
+				TEST(FText::AsNumber(TestNumber1), FText::AsNumber(TestNumber1), ETextIdenticalModeFlags::DeepCompare, true);
+				TEST(FText::AsNumber(TestNumber1), FText::AsNumber(TestNumber2), ETextIdenticalModeFlags::None, false);
+				TEST(FText::AsNumber(TestNumber1), FText::AsNumber(TestNumber2), ETextIdenticalModeFlags::DeepCompare, false);
+
+				TEST(TestIdenticalStr1.ToUpper(), TestIdenticalStr1.ToUpper(), ETextIdenticalModeFlags::None, false);
+				TEST(TestIdenticalStr1.ToUpper(), TestIdenticalStr1.ToUpper(), ETextIdenticalModeFlags::DeepCompare, true);
+				TEST(TestIdenticalStr1.ToUpper(), TestIdenticalStr1.ToLower(), ETextIdenticalModeFlags::None, false);
+				TEST(TestIdenticalStr1.ToUpper(), TestIdenticalStr1.ToLower(), ETextIdenticalModeFlags::DeepCompare, false);
+				TEST(TestIdenticalStr1.ToUpper(), TestIdenticalStr2.ToUpper(), ETextIdenticalModeFlags::None, false);
+				TEST(TestIdenticalStr1.ToUpper(), TestIdenticalStr2.ToUpper(), ETextIdenticalModeFlags::DeepCompare, false);
+			}
+
+			#undef TEST
+#define TEST( Desc, A, B ) if( !A.EqualTo(B) ) FAIL_CHECK(FString::Printf(TEXT("%s - A=%s B=%s"),*Desc,*A.ToString(),*B.ToString()))
+
+			FText TestText;
+
+			TestText = INVTEXT("Format with single apostrophes quotes: '{0}'");
+			TEST(TestText.ToString(), FText::Format(TestText, ArgText0), INVTEXT("Format with single apostrophes quotes: 'Arg0'"));
+			TestText = INVTEXT("Format with double apostrophes quotes: ''{0}''");
+			TEST(TestText.ToString(), FText::Format(TestText, ArgText0), INVTEXT("Format with double apostrophes quotes: ''Arg0''"));
+			TestText = INVTEXT("Print with single graves: `{0}`");
+			TEST(TestText.ToString(), FText::Format(TestText, ArgText0), INVTEXT("Print with single graves: {0}`"));
+			TestText = INVTEXT("Format with double graves: ``{0}``");
+			TEST(TestText.ToString(), FText::Format(TestText, ArgText0), INVTEXT("Format with double graves: `Arg0`"));
+
+			TestText = INVTEXT("Testing `escapes` here.");
+			TEST(TestText.ToString(), FormatWithoutArguments(TestText), INVTEXT("Testing `escapes` here."));
+			TestText = INVTEXT("Testing ``escapes` here.");
+			TEST(TestText.ToString(), FormatWithoutArguments(TestText), INVTEXT("Testing `escapes` here."));
+			TestText = INVTEXT("Testing ``escapes`` here.");
+			TEST(TestText.ToString(), FormatWithoutArguments(TestText), INVTEXT("Testing `escapes` here."));
+
+			TestText = INVTEXT("Testing `}escapes{ here.");
+			TEST(TestText.ToString(), FormatWithoutArguments(TestText), INVTEXT("Testing }escapes{ here."));
+			TestText = INVTEXT("Testing `}escapes{ here.`");
+			TEST(TestText.ToString(), FormatWithoutArguments(TestText), INVTEXT("Testing }escapes{ here.`"));
+			TestText = INVTEXT("Testing `}escapes{` here.");
+			TEST(TestText.ToString(), FormatWithoutArguments(TestText), INVTEXT("Testing }escapes{` here."));
+			TestText = INVTEXT("Testing }escapes`{ here.");
+			TEST(TestText.ToString(), FormatWithoutArguments(TestText), INVTEXT("Testing }escapes{ here."));
+			TestText = INVTEXT("`Testing }escapes`{ here.");
+			TEST(TestText.ToString(), FormatWithoutArguments(TestText), INVTEXT("`Testing }escapes{ here."));
+
+			TestText = INVTEXT("Testing `{escapes} here.");
+			TEST(TestText.ToString(), FormatWithoutArguments(TestText), INVTEXT("Testing {escapes} here."));
+			TestText = INVTEXT("Testing `{escapes} here.`");
+			TEST(TestText.ToString(), FormatWithoutArguments(TestText), INVTEXT("Testing {escapes} here.`"));
+			TestText = INVTEXT("Testing `{escapes}` here.");
+			TEST(TestText.ToString(), FormatWithoutArguments(TestText), INVTEXT("Testing {escapes}` here."));
+
+			TestText = INVTEXT("Starting text: {0} {1}");
+			TEST(TestText.ToString(), FText::Format(TestText, ArgText0, ArgText1), INVTEXT("Starting text: Arg0 Arg1"));
+			TestText = INVTEXT("{0} {1} - Ending Text.");
+			TEST(TestText.ToString(), FText::Format(TestText, ArgText0, ArgText1), INVTEXT("Arg0 Arg1 - Ending Text."));
+			TestText = INVTEXT("Starting text: {0} {1} - Ending Text.");
+			TEST(TestText.ToString(), FText::Format(TestText, ArgText0, ArgText1), INVTEXT("Starting text: Arg0 Arg1 - Ending Text."));
+			TestText = INVTEXT("{0} {1}");
+			TEST(TestText.ToString(), FText::Format(TestText, ArgText0, ArgText1), INVTEXT("Arg0 Arg1"));
+			TestText = INVTEXT("{1} {0}");
+			TEST(TestText.ToString(), FText::Format(TestText, ArgText0, ArgText1), INVTEXT("Arg1 Arg0"));
+			TestText = INVTEXT("{0}");
+			TEST(TestText.ToString(), FText::Format(TestText, ArgText0), INVTEXT("Arg0"));
+			TestText = INVTEXT("{0} - {1} - {2} - {3}");
+			TEST(TestText.ToString(), FText::Format(TestText, ArgText0, ArgText1, ArgText2, ArgText3), INVTEXT("Arg0 - Arg1 - Arg2 - Arg3"));
+			TestText = INVTEXT("{0} - {0} - {0} - {1}");
+			TEST(TestText.ToString(), FText::Format(TestText, ArgText0, ArgText1), INVTEXT("Arg0 - Arg0 - Arg0 - Arg1"));
+
+			TestText = INVTEXT("Starting text: {1}");
+			TEST(TestText.ToString(), FText::Format(TestText, ArgText0, ArgText1), INVTEXT("Starting text: Arg1"));
+			TestText = INVTEXT("{0} - Ending Text.");
+			TEST(TestText.ToString(), FText::Format(TestText, ArgText0, ArgText1), INVTEXT("Arg0 - Ending Text."));
+			TestText = INVTEXT("Starting text: {0} - Ending Text.");
+			TEST(TestText.ToString(), FText::Format(TestText, ArgText0, ArgText1), INVTEXT("Starting text: Arg0 - Ending Text."));
+
+			TestText = INVTEXT("{0} {2}");
+			TEST(TestText.ToString(), FText::Format(TestText, ArgText0, ArgText1, ArgText2), INVTEXT("Arg0 Arg2"));
+			TestText = INVTEXT("{1}");
+			TEST(TestText.ToString(), FText::Format(TestText, ArgText0, ArgText1, ArgText2), INVTEXT("Arg1"));
+
+			TestText = INVTEXT("Starting text: {0} {1}");
+			TEST(TestText.ToString(), FText::Format(TestText, ArgText0, ArgText1), INVTEXT("Starting text: Arg0 Arg1"));
+			TestText = INVTEXT("{0} {1} - Ending Text.");
+			TEST(TestText.ToString(), FText::Format(TestText, ArgText0, ArgText1), INVTEXT("Arg0 Arg1 - Ending Text."));
+			TestText = INVTEXT("Starting text: {0} {1} - Ending Text.");
+			TEST(TestText.ToString(), FText::Format(TestText, ArgText0, ArgText1), INVTEXT("Starting text: Arg0 Arg1 - Ending Text."));
+			TestText = INVTEXT("{0} {1}");
+			TEST(TestText.ToString(), FText::Format(TestText, ArgText0, ArgText1), INVTEXT("Arg0 Arg1"));
+			TestText = INVTEXT("{1} {0}");
+			TEST(TestText.ToString(), FText::Format(TestText, ArgText0, ArgText1), INVTEXT("Arg1 Arg0"));
+			TestText = INVTEXT("{0}");
+			TEST(TestText.ToString(), FText::Format(TestText, ArgText0), INVTEXT("Arg0"));
+			TestText = INVTEXT("{0} - {1} - {2} - {3}");
+			TEST(TestText.ToString(), FText::Format(TestText, ArgText0, ArgText1, ArgText2, ArgText3), INVTEXT("Arg0 - Arg1 - Arg2 - Arg3"));
+			TestText = INVTEXT("{0} - {0} - {0} - {1}");
+			TEST(TestText.ToString(), FText::Format(TestText, ArgText0, ArgText1), INVTEXT("Arg0 - Arg0 - Arg0 - Arg1"));
+
+			{
+				FFormatNamedArguments Arguments;
+				Arguments.Add(TEXT("Age"), INVTEXT("23"));
+				Arguments.Add(TEXT("Height"), INVTEXT("68"));
+				Arguments.Add(TEXT("Gender"), INVTEXT("male"));
+				Arguments.Add(TEXT("Name"), INVTEXT("Saul"));
+
+				// Not using all the arguments is okay.
+				TestText = INVTEXT("My name is {Name}.");
+				TEST(TestText.ToString(), FText::Format(TestText, Arguments), INVTEXT("My name is Saul."));
+				TestText = INVTEXT("My age is {Age}.");
+				TEST(TestText.ToString(), FText::Format(TestText, Arguments), INVTEXT("My age is 23."));
+				TestText = INVTEXT("My gender is {Gender}.");
+				TEST(TestText.ToString(), FText::Format(TestText, Arguments), INVTEXT("My gender is male."));
+				TestText = INVTEXT("My height is {Height}.");
+				TEST(TestText.ToString(), FText::Format(TestText, Arguments), INVTEXT("My height is 68."));
+
+				// Using arguments out of order is okay.
+				TestText = INVTEXT("My name is {Name}. My age is {Age}. My gender is {Gender}.");
+				TEST(TestText.ToString(), FText::Format(TestText, Arguments), INVTEXT("My name is Saul. My age is 23. My gender is male."));
+				TestText = INVTEXT("My age is {Age}. My gender is {Gender}. My name is {Name}.");
+				TEST(TestText.ToString(), FText::Format(TestText, Arguments), INVTEXT("My age is 23. My gender is male. My name is Saul."));
+				TestText = INVTEXT("My gender is {Gender}. My name is {Name}. My age is {Age}.");
+				TEST(TestText.ToString(), FText::Format(TestText, Arguments), INVTEXT("My gender is male. My name is Saul. My age is 23."));
+				TestText = INVTEXT("My gender is {Gender}. My age is {Age}. My name is {Name}.");
+				TEST(TestText.ToString(), FText::Format(TestText, Arguments), INVTEXT("My gender is male. My age is 23. My name is Saul."));
+				TestText = INVTEXT("My age is {Age}. My name is {Name}. My gender is {Gender}.");
+				TEST(TestText.ToString(), FText::Format(TestText, Arguments), INVTEXT("My age is 23. My name is Saul. My gender is male."));
+				TestText = INVTEXT("My name is {Name}. My gender is {Gender}. My age is {Age}.");
+				TEST(TestText.ToString(), FText::Format(TestText, Arguments), INVTEXT("My name is Saul. My gender is male. My age is 23."));
+
+				// Reusing arguments is okay.
+				TestText = INVTEXT("If my age is {Age}, I have been alive for {Age} year(s).");
+				TEST(TestText.ToString(), FText::Format(TestText, Arguments), INVTEXT("If my age is 23, I have been alive for 23 year(s)."));
+
+				// Not providing an argument leaves the parameter as text.
+				TestText = INVTEXT("What... is the air-speed velocity of an unladen swallow? {AirSpeedOfAnUnladenSwallow}.");
+				TEST(TestText.ToString(), FText::Format(TestText, Arguments), INVTEXT("What... is the air-speed velocity of an unladen swallow? {AirSpeedOfAnUnladenSwallow}."));
+			}
+
+			{
+				FFormatNamedArguments ArgumentList;
+				ArgumentList.Emplace(TEXT("Age"), INVTEXT("23"));
+				ArgumentList.Emplace(TEXT("Height"), INVTEXT("68"));
+				ArgumentList.Emplace(TEXT("Gender"), INVTEXT("male"));
+				ArgumentList.Emplace(TEXT("Name"), INVTEXT("Saul"));
+
+				// Not using all the arguments is okay.
+				TestText = INVTEXT("My name is {Name}.");
+				TEST(TestText.ToString(), FText::Format(TestText, ArgumentList), INVTEXT("My name is Saul."));
+				TestText = INVTEXT("My age is {Age}.");
+				TEST(TestText.ToString(), FText::Format(TestText, ArgumentList), INVTEXT("My age is 23."));
+				TestText = INVTEXT("My gender is {Gender}.");
+				TEST(TestText.ToString(), FText::Format(TestText, ArgumentList), INVTEXT("My gender is male."));
+				TestText = INVTEXT("My height is {Height}.");
+				TEST(TestText.ToString(), FText::Format(TestText, ArgumentList), INVTEXT("My height is 68."));
+
+				// Using arguments out of order is okay.
+				TestText = INVTEXT("My name is {Name}. My age is {Age}. My gender is {Gender}.");
+				TEST(TestText.ToString(), FText::Format(TestText, ArgumentList), INVTEXT("My name is Saul. My age is 23. My gender is male."));
+				TestText = INVTEXT("My age is {Age}. My gender is {Gender}. My name is {Name}.");
+				TEST(TestText.ToString(), FText::Format(TestText, ArgumentList), INVTEXT("My age is 23. My gender is male. My name is Saul."));
+				TestText = INVTEXT("My gender is {Gender}. My name is {Name}. My age is {Age}.");
+				TEST(TestText.ToString(), FText::Format(TestText, ArgumentList), INVTEXT("My gender is male. My name is Saul. My age is 23."));
+				TestText = INVTEXT("My gender is {Gender}. My age is {Age}. My name is {Name}.");
+				TEST(TestText.ToString(), FText::Format(TestText, ArgumentList), INVTEXT("My gender is male. My age is 23. My name is Saul."));
+				TestText = INVTEXT("My age is {Age}. My name is {Name}. My gender is {Gender}.");
+				TEST(TestText.ToString(), FText::Format(TestText, ArgumentList), INVTEXT("My age is 23. My name is Saul. My gender is male."));
+				TestText = INVTEXT("My name is {Name}. My gender is {Gender}. My age is {Age}.");
+				TEST(TestText.ToString(), FText::Format(TestText, ArgumentList), INVTEXT("My name is Saul. My gender is male. My age is 23."));
+
+				// Reusing arguments is okay.
+				TestText = INVTEXT("If my age is {Age}, I have been alive for {Age} year(s).");
+				TEST(TestText.ToString(), FText::Format(TestText, ArgumentList), INVTEXT("If my age is 23, I have been alive for 23 year(s)."));
+
+				// Not providing an argument leaves the parameter as text.
+				TestText = INVTEXT("What... is the air-speed velocity of an unladen swallow? {AirSpeedOfAnUnladenSwallow}.");
+				TEST(TestText.ToString(), FText::Format(TestText, ArgumentList), INVTEXT("What... is the air-speed velocity of an unladen swallow? {AirSpeedOfAnUnladenSwallow}."));
+			}
+
+#undef TEST
+#define TEST( Pattern, Actual, Expected ) TestPatternParameterEnumeration(*this, Pattern, Actual, Expected)
+
+			TArray<FString> ActualArguments;
+			TArray<FString> ExpectedArguments;
+
+			TestText = INVTEXT("My name is {Name}.");
+			ExpectedArguments.Empty(1);
+			ExpectedArguments.Add(TEXT("Name"));
+			TEST(TestText, ActualArguments, ExpectedArguments);
+
+			TestText = INVTEXT("My age is {Age}.");
+			ExpectedArguments.Empty(1);
+			ExpectedArguments.Add(TEXT("Age"));
+			TEST(TestText, ActualArguments, ExpectedArguments);
+
+			TestText = INVTEXT("If my age is {Age}, I have been alive for {Age} year(s).");
+			ExpectedArguments.Empty(1);
+			ExpectedArguments.Add(TEXT("Age"));
+			TEST(TestText, ActualArguments, ExpectedArguments);
+
+			TestText = INVTEXT("{0} - {1} - {2} - {3}");
+			ExpectedArguments.Empty(4);
+			ExpectedArguments.Add(TEXT("0"));
+			ExpectedArguments.Add(TEXT("1"));
+			ExpectedArguments.Add(TEXT("2"));
+			ExpectedArguments.Add(TEXT("3"));
+			TEST(TestText, ActualArguments, ExpectedArguments);
+
+			TestText = INVTEXT("My name is {Name}. My age is {Age}. My gender is {Gender}.");
+			ExpectedArguments.Empty(3);
+			ExpectedArguments.Add(TEXT("Name"));
+			ExpectedArguments.Add(TEXT("Age"));
+			ExpectedArguments.Add(TEXT("Gender"));
+			TEST(TestText, ActualArguments, ExpectedArguments);
+
+#undef TEST
+
+#if UE_ENABLE_ICU
+			if (I18N.SetCurrentCulture("en-US"))
+			{
+#define TEST(NumBytes, UnitStandard, ExpectedString) \
+			if (!FText::FromString(TEXT(ExpectedString)).EqualTo(FText::AsMemory(NumBytes, &NumberFormattingOptions, nullptr, UnitStandard))) \
+			{ \
+				FAIL_CHECK(FString::Printf(TEXT("FText::AsMemory expected %s bytes in %s to be %s - got %s"), TEXT(#NumBytes), TEXT(#UnitStandard), TEXT(ExpectedString), *FText::AsMemory(NumBytes, &NumberFormattingOptions, nullptr, UnitStandard).ToString())); \
+			} \
+
+			{
+				FNumberFormattingOptions NumberFormattingOptions = FNumberFormattingOptions()
+					.SetRoundingMode(ERoundingMode::HalfFromZero)
+					.SetMinimumFractionalDigits(0)
+					.SetMaximumFractionalDigits(3);
+
+				TEST(0, EMemoryUnitStandard::SI, "0 B");
+				TEST(1, EMemoryUnitStandard::SI, "1 B");
+				TEST(1000, EMemoryUnitStandard::SI, "1 kB");
+				TEST(1000000, EMemoryUnitStandard::SI, "1 MB");
+				TEST(1000000000, EMemoryUnitStandard::SI, "1 GB");
+				TEST(1000000000000, EMemoryUnitStandard::SI, "1 TB");
+				TEST(1000000000000000, EMemoryUnitStandard::SI, "1 PB");
+				TEST(1000000000000000000, EMemoryUnitStandard::SI, "1 EB");
+				TEST(999, EMemoryUnitStandard::SI, "999 B");
+				TEST(999999, EMemoryUnitStandard::SI, "999.999 kB");
+				TEST(999999999, EMemoryUnitStandard::SI, "999.999 MB");
+				TEST(999999999999, EMemoryUnitStandard::SI, "999.999 GB");
+				TEST(999999999999999, EMemoryUnitStandard::SI, "999.999 TB");
+				TEST(999999999999999999, EMemoryUnitStandard::SI, "999.999 PB");
+				TEST(18446744073709551615ULL, EMemoryUnitStandard::SI, "18.446 EB");
+
+				TEST(0, EMemoryUnitStandard::IEC, "0 B");
+				TEST(1, EMemoryUnitStandard::IEC, "1 B");
+				TEST(1024, EMemoryUnitStandard::IEC, "1 KiB");
+				TEST(1048576, EMemoryUnitStandard::IEC, "1 MiB");
+				TEST(1073741824, EMemoryUnitStandard::IEC, "1 GiB");
+				TEST(1099511627776, EMemoryUnitStandard::IEC, "1 TiB");
+				TEST(1125899906842624, EMemoryUnitStandard::IEC, "1 PiB");
+				TEST(1152921504606846976, EMemoryUnitStandard::IEC, "1 EiB");
+				TEST(1023, EMemoryUnitStandard::IEC, "0.999 KiB");
+				TEST(1048575, EMemoryUnitStandard::IEC, "0.999 MiB");
+				TEST(1073741823, EMemoryUnitStandard::IEC, "0.999 GiB");
+				TEST(1099511627775, EMemoryUnitStandard::IEC, "0.999 TiB");
+				TEST(1125899906842623, EMemoryUnitStandard::IEC, "0.999 PiB");
+				TEST(1152921504606846975, EMemoryUnitStandard::IEC, "0.999 EiB");
+				TEST(18446744073709551615ULL, EMemoryUnitStandard::IEC, "15.999 EiB");
+			}
+#undef TEST
+
+#define TEST( A, B, ComparisonLevel ) if( !(FText::FromString(A)).EqualTo(FText::FromString(B), (ComparisonLevel)) ) FAIL_CHECK(FString::Printf(TEXT("Testing comparison of equivalent characters with comparison level (%s). - A=%s B=%s"),TEXT(#ComparisonLevel),(A),(B)))
+
+			// Basic sanity checks
+			TEST(TEXT("a"), TEXT("A"), ETextComparisonLevel::Primary); // Basic sanity check
+			TEST(TEXT("a"), TEXT("a"), ETextComparisonLevel::Tertiary); // Basic sanity check
+			TEST(TEXT("A"), TEXT("A"), ETextComparisonLevel::Tertiary); // Basic sanity check
+
+			// Test equivalence
+			TEST(TEXT("ss"), TEXT("\x00DF"), ETextComparisonLevel::Primary); // Lowercase Sharp s
+			TEST(TEXT("SS"), TEXT("\x1E9E"), ETextComparisonLevel::Primary); // Uppercase Sharp S
+			TEST(TEXT("ae"), TEXT("\x00E6"), ETextComparisonLevel::Primary); // Lowercase ae
+			TEST(TEXT("AE"), TEXT("\x00C6"), ETextComparisonLevel::Primary); // Uppercase AE
+
+			// Test accentuation
+			TEST(TEXT("u"), TEXT("\x00FC"), ETextComparisonLevel::Primary); // Lowercase u with dieresis
+			TEST(TEXT("U"), TEXT("\x00DC"), ETextComparisonLevel::Primary); // Uppercase U with dieresis
+
+#undef TEST
+			}
+			else
+			{
+				WARN(FString::Printf(TEXT("Internationalization data for %s missing - test is partially disabled."), TEXT("en-US")));
+			}
+#else
+			WARN("ICU is disabled thus locale-aware string comparison is disabled.");
+#endif
+
+#if UE_ENABLE_ICU
+			// Sort Testing
+			// French
+			if (I18N.SetCurrentCulture("fr"))
+			{
+				TArray<FText> CorrectlySortedValues;
+				CorrectlySortedValues.Add(INVTEXT("cote"));
+				CorrectlySortedValues.Add(INVTEXT("cot\u00e9"));
+				CorrectlySortedValues.Add(INVTEXT("c\u00f4te"));
+				CorrectlySortedValues.Add(INVTEXT("c\u00f4t\u00e9"));
+
+				{
+					// Make unsorted.
+					TArray<FText> Values;
+					Values.Reserve(CorrectlySortedValues.Num());
+
+					Values.Add(CorrectlySortedValues[1]);
+					Values.Add(CorrectlySortedValues[3]);
+					Values.Add(CorrectlySortedValues[2]);
+					Values.Add(CorrectlySortedValues[0]);
+
+					// Execute sort.
+					Values.Sort(FText::FSortPredicate());
+
+					// Test if sorted.
+					bool Identical = true;
+					for (int32 j = 0; j < Values.Num(); ++j)
+					{
+						Identical = Values[j].EqualTo(CorrectlySortedValues[j]);
+						if (!Identical)
+						{
+							break;
+						}
+					}
+
+					if (!Identical)
+					{
+						//currently failing FAIL_CHECK(FString::Printf(TEXT("Sort order is wrong for culture (%s)."), *FInternationalization::Get().GetCurrentCulture()->GetEnglishName()));
+					}
+				}
+			}
+			else
+			{
+				WARN(FString::Printf(TEXT("Internationalization data for %s missing - test is partially disabled."), TEXT("fr")));
+			}
+
+			// French Canadian
+			if (I18N.SetCurrentCulture("fr-CA"))
+			{
+				TArray<FText> CorrectlySortedValues;
+				CorrectlySortedValues.Add(INVTEXT("cote"));
+				CorrectlySortedValues.Add(INVTEXT("côte"));
+				CorrectlySortedValues.Add(INVTEXT("coté"));
+				CorrectlySortedValues.Add(INVTEXT("côté"));
+
+				{
+					// Make unsorted.
+					TArray<FText> Values;
+					Values.Reserve(CorrectlySortedValues.Num());
+
+					Values.Add(CorrectlySortedValues[1]);
+					Values.Add(CorrectlySortedValues[3]);
+					Values.Add(CorrectlySortedValues[2]);
+					Values.Add(CorrectlySortedValues[0]);
+
+					// Execute sort.
+					Values.Sort(FText::FSortPredicate());
+
+					// Test if sorted.
+					bool Identical = true;
+					for (int32 j = 0; j < Values.Num(); ++j)
+					{
+						Identical = Values[j].EqualTo(CorrectlySortedValues[j]);
+						if (!Identical) break;
+					}
+
+					if (!Identical)
+					{
+						//currently failing FAIL_CHECK(FString::Printf(TEXT("Sort order is wrong for culture (%s)."), *FInternationalization::Get().GetCurrentCulture()->GetEnglishName()));
+					}
+				}
+			}
+			else
+			{
+				WARN(FString::Printf(TEXT("Internationalization data for %s missing - test is partially disabled."), TEXT("fr-CA")));
+			}
+#else
+			WARN("ICU is disabled thus locale-aware string collation is disabled.");
+#endif
+
+#if UE_ENABLE_ICU
+			{
+				I18N.RestoreCultureState(OriginalCultureState);
+
+				TArray<uint8> FormattedHistoryAsEnglish;
+				TArray<uint8> FormattedHistoryAsFrenchCanadian;
+				TArray<uint8> InvariantFTextData;
+
+				FString InvariantString = TEXT("This is a culture invariant string.");
+				FString FormattedTestLayer2_OriginalLanguageSourceString;
+				FText FormattedTestLayer2;
+
+				// Scoping to allow all locals to leave scope after we serialize at the end
+				{
+					FFormatNamedArguments Args;
+					Args.Add(TEXT("String1"), LOCTEXT("RebuildFTextTest1_Lorem", "Lorem"));
+					Args.Add(TEXT("String2"), LOCTEXT("RebuildFTextTest1_Ipsum", "Ipsum"));
+					FText FormattedTest1 = FText::Format(LOCTEXT("RebuildNamedText1", "{String1} \"Lorem Ipsum\" {String2}"), Args);
+
+					FFormatOrderedArguments ArgsOrdered;
+					ArgsOrdered.Add(LOCTEXT("RebuildFTextTest1_Lorem", "Lorem"));
+					ArgsOrdered.Add(LOCTEXT("RebuildFTextTest1_Ipsum", "Ipsum"));
+					FText FormattedTestOrdered1 = FText::Format(LOCTEXT("RebuildOrderedText1", "{0} \"Lorem Ipsum\" {1}"), ArgsOrdered);
+
+					// Will change to 5.542 due to default settings for numbers
+					FText AsNumberTest1 = FText::AsNumber(5.5421);
+
+					FText AsPercentTest1 = FText::AsPercent(0.925);
+					FText AsCurrencyTest1 = FText::AsCurrencyBase(10025, TEXT("USD"));
+
+					FDateTime DateTimeInfo(2080, 8, 20, 9, 33, 22);
+					FText AsDateTimeTest1 = FText::AsDateTime(DateTimeInfo, EDateTimeStyle::Default, EDateTimeStyle::Default, TEXT("UTC"));
+
+					// FormattedTestLayer2 must be updated when adding or removing from this block. Further, below, 
+					// verifying the LEET translated string must be changed to reflect what the new string looks like.
+					FFormatNamedArguments ArgsLayer2;
+					ArgsLayer2.Add("NamedLayer1", FormattedTest1);
+					ArgsLayer2.Add("OrderedLayer1", FormattedTestOrdered1);
+					ArgsLayer2.Add("FTextNumber", AsNumberTest1);
+					ArgsLayer2.Add("Number", 5010.89221);
+					ArgsLayer2.Add("DateTime", AsDateTimeTest1);
+					ArgsLayer2.Add("Percent", AsPercentTest1);
+					ArgsLayer2.Add("Currency", AsCurrencyTest1);
+					FormattedTestLayer2 = FText::Format(LOCTEXT("RebuildTextLayer2", "{NamedLayer1} | {OrderedLayer1} | {FTextNumber} | {Number} | {DateTime} | {Percent} | {Currency}"), ArgsLayer2);
+
+					{
+						// Serialize the full, bulky FText that is a composite of most of the other FTextHistories.
+						FMemoryWriter Ar(FormattedHistoryAsEnglish);
+						Ar << FormattedTestLayer2;
+						Ar.Close();
+					}
+
+					// The original string in the native language.
+					FormattedTestLayer2_OriginalLanguageSourceString = FormattedTestLayer2.BuildSourceString();
+
+#if ENABLE_LOC_TESTING
+					{
+						// Swap to "LEET" culture to check if rebuilding works (verify the whole)
+						I18N.SetCurrentCulture(FLeetCulture::StaticGetName());
+
+						// When changes are made to FormattedTestLayer2, please pull out the newly translated LEET string and update the below if-statement to keep the test passing!
+						FString LEETTranslatedString = FormattedTestLayer2.ToString();
+
+						FString DesiredOutput = FString(TEXT("\x2021") TEXT("\xAB") TEXT("\x2021") TEXT("\xAB") TEXT("\x2021") TEXT("L0r3m") TEXT("\x2021") TEXT("\xBB") TEXT(" \"L0r3m 1p$um\" ") TEXT("\xAB") TEXT("\x2021") TEXT("1p$um") TEXT("\x2021") TEXT("\xBB") TEXT("\x2021") TEXT("\xBB") TEXT(" | ") TEXT("\xAB") TEXT("\x2021") TEXT("\xAB") TEXT("\x2021") TEXT("L0r3m") TEXT("\x2021") TEXT("\xBB") TEXT(" \"L0r3m 1p$um\" ") TEXT("\xAB") TEXT("\x2021") TEXT("1p$um") TEXT("\x2021") TEXT("\xBB") TEXT("\x2021") TEXT("\xBB") TEXT(" | ") TEXT("\xAB") TEXT("5.5421") TEXT("\xBB") TEXT(" | ") TEXT("\xAB") TEXT("5010.89221") TEXT("\xBB") TEXT(" | ") TEXT("\xAB") TEXT("Aug 20, 2080, 9:33:22 AM") TEXT("\xBB") TEXT(" | ") TEXT("\xAB") TEXT("92%") TEXT("\xBB") TEXT(" | ") TEXT("\xAB") TEXT("$") TEXT("\xA0") TEXT("100.25") TEXT("\xBB") TEXT("\x2021"));
+						// Convert the baked string into an FText, which will be leetified, then compare it to the rebuilt FText
+						if (LEETTranslatedString != DesiredOutput)
+						{
+							FAIL_CHECK(TEXT("FormattedTestLayer2 did not rebuild to correctly in LEET!"));
+							FAIL_CHECK(TEXT("Formatted Output=") + LEETTranslatedString);
+							FAIL_CHECK(TEXT("Desired Output=") + DesiredOutput);
+						}					
+					}
+#endif
+
+					// Swap to French-Canadian to check if rebuilding works (verify each numerical component)
+					{
+						I18N.SetCurrentCulture("fr-CA");
+
+						// Need the FText to be rebuilt in fr-CA.
+						FormattedTestLayer2.ToString();
+
+
+						if (AsNumberTest1.CompareTo(FText::AsNumber(5.5421)) != 0)
+						{
+							FAIL_CHECK(TEXT("AsNumberTest1 did not rebuild correctly in French-Canadian"));
+							FAIL_CHECK(TEXT("Number Output=") + AsNumberTest1.ToString());
+						}
+
+						if (AsPercentTest1.CompareTo(FText::AsPercent(0.925)) != 0)
+						{
+							FAIL_CHECK(TEXT("AsPercentTest1 did not rebuild correctly in French-Canadian"));
+							FAIL_CHECK(TEXT("Percent Output=") + AsPercentTest1.ToString());
+						}
+
+						if (AsCurrencyTest1.CompareTo(FText::AsCurrencyBase(10025, TEXT("USD"))) != 0)
+						{
+							FAIL_CHECK(TEXT("AsCurrencyTest1 did not rebuild correctly in French-Canadian"));
+							FAIL_CHECK(TEXT("Currency Output=") + AsCurrencyTest1.ToString());
+						}
+
+						if (AsDateTimeTest1.CompareTo(FText::AsDateTime(DateTimeInfo, EDateTimeStyle::Default, EDateTimeStyle::Default, TEXT("UTC"))) != 0)
+						{
+							FAIL_CHECK(TEXT("AsDateTimeTest1 did not rebuild correctly in French-Canadian"));
+							FAIL_CHECK(TEXT("DateTime Output=") + AsDateTimeTest1.ToString());
+						}
+
+						{
+							// Serialize the full, bulky FText that is a composite of most of the other FTextHistories.
+							// We don't care how this may be translated, we will be serializing this in as LEET.
+							FMemoryWriter Ar(FormattedHistoryAsFrenchCanadian);
+							Ar << FormattedTestLayer2;
+							Ar.Close();
+						}
+
+						{
+							FText InvariantFText = FText::FromString(InvariantString);
+
+							// Serialize an invariant FText
+							FMemoryWriter Ar(InvariantFTextData);
+							Ar << InvariantFText;
+							Ar.Close();
+						}
+					}
+				}
+
+#if ENABLE_LOC_TESTING
+				{
+					I18N.SetCurrentCulture(FLeetCulture::StaticGetName());
+
+					FText FormattedEnglishTextHistoryAsLeet;
+					FText FormattedFrenchCanadianTextHistoryAsLeet;
+
+					{
+						FMemoryReader Ar(FormattedHistoryAsEnglish);
+						Ar << FormattedEnglishTextHistoryAsLeet;
+						Ar.Close();
+					}
+					{
+						FMemoryReader Ar(FormattedHistoryAsFrenchCanadian);
+						Ar << FormattedFrenchCanadianTextHistoryAsLeet;
+						Ar.Close();
+					}
+
+					// Confirm the two FText's serialize in and get translated into the current (LEET) translation. One originated in English, the other in French-Canadian locales.
+					if (FormattedEnglishTextHistoryAsLeet.CompareTo(FormattedFrenchCanadianTextHistoryAsLeet) != 0)
+					{
+						FAIL_CHECK(TEXT("Serialization of text histories from source English and source French-Canadian to LEET did not produce the same results!"));
+						FAIL_CHECK(TEXT("English Output=") + FormattedEnglishTextHistoryAsLeet.ToString());
+						FAIL_CHECK(TEXT("French-Canadian Output=") + FormattedFrenchCanadianTextHistoryAsLeet.ToString());
+					}
+
+					// Confirm the two FText's source strings for the serialized FTexts are the same.
+					if (FormattedEnglishTextHistoryAsLeet.BuildSourceString() != FormattedFrenchCanadianTextHistoryAsLeet.BuildSourceString())
+					{
+						FAIL_CHECK(TEXT("Serialization of text histories from source English and source French-Canadian to LEET did not produce the same source results!"));
+						FAIL_CHECK(TEXT("English Output=") + FormattedEnglishTextHistoryAsLeet.BuildSourceString());
+						FAIL_CHECK(TEXT("French-Canadian Output=") + FormattedFrenchCanadianTextHistoryAsLeet.BuildSourceString());
+					}
+
+					// Rebuild in LEET so that when we build the source string the DisplayString is still in LEET. 
+					FormattedTestLayer2.ToString();
+
+					{
+						I18N.RestoreCultureState(OriginalCultureState);
+
+						FText InvariantFText;
+
+						FMemoryReader Ar(InvariantFTextData);
+						Ar << InvariantFText;
+						Ar.Close();
+
+						if (InvariantFText.ToString() != InvariantString)
+						{
+							FAIL_CHECK(TEXT("Invariant FText did not match the original FString after serialization!"));
+							FAIL_CHECK(TEXT("Invariant Output=") + InvariantFText.ToString());
+						}
+
+
+						FString FormattedTestLayer2_SourceString = FormattedTestLayer2.BuildSourceString();
+
+						// Compare the source string of the LEETified version of FormattedTestLayer2 to ensure it is correct.
+						if (FormattedTestLayer2_OriginalLanguageSourceString != FormattedTestLayer2_SourceString)
+						{
+							FAIL_CHECK(TEXT("FormattedTestLayer2's source string was incorrect!"));
+							FAIL_CHECK(TEXT("Output=") + FormattedTestLayer2_SourceString);
+							FAIL_CHECK(TEXT("Desired Output=") + FormattedTestLayer2_OriginalLanguageSourceString);
+						}
+					}
+				}
+#endif
+			}
+#else
+			WARN("ICU is disabled thus locale-aware formatting needed in rebuilding source text from history is disabled.");
+#endif
+
+			//**********************************
+			// FromString Test
+			//**********************************
+			TestText = FText::FromString(TEXT("Test String"));
+
+			if (GIsEditor && TestText.IsCultureInvariant())
+			{
+				FAIL_CHECK(TEXT("FromString should not produce a Culture Invariant Text when called inside the editor"));
+			}
+
+			if (!GIsEditor && !TestText.IsCultureInvariant())
+			{
+				FAIL_CHECK(TEXT("FromString should produce a Culture Invariant Text when called outside the editor"));
+			}
+
+			if (TestText.IsTransient())
+			{
+				FAIL_CHECK(TEXT("FromString should never produce a Transient Text"));
+			}
+
+			I18N.RestoreCultureState(OriginalCultureState);
+		}
+};
+
+TEST_CASE_NAMED(FTextTest, "System::Core::Misc::Text", "[.][EditorContext][ClientContext][EngineFilter]")
+{
+	FTextTestClass Instance;
+	Instance.TextTest();
+}
+
+TEST_CASE_NAMED(FTextRoundingTest, "System::Core::Misc::TextRounding", "[.][EditorContext][ClientContext][EngineFilter]")
+{
+	static const TCHAR* RoundingModeNames[] = {
+		TEXT("HalfToEven"),
+		TEXT("HalfFromZero"),
+		TEXT("HalfToZero"),
+		TEXT("FromZero"),
+		TEXT("ToZero"),
+		TEXT("ToNegativeInfinity"),
+		TEXT("ToPositiveInfinity"),
+	};
+
+	static_assert(ERoundingMode::ToPositiveInfinity == UE_ARRAY_COUNT(RoundingModeNames) - 1, "RoundingModeNames array needs updating");
+
+	static const double InputValues[] = {
+		1000.1224,
+		1000.1225,
+		1000.1226,
+		1000.1234,
+		1000.1235,
+		1000.1236,
+		
+		1000.1244,
+		1000.1245,
+		1000.1246,
+		1000.1254,
+		1000.1255,
+		1000.1256,
+
+		-1000.1224,
+		-1000.1225,
+		-1000.1226,
+		-1000.1234,
+		-1000.1235,
+		-1000.1236,
+		
+		-1000.1244,
+		-1000.1245,
+		-1000.1246,
+		-1000.1254,
+		-1000.1255,
+		-1000.1256,
+	};
+
+	static const TCHAR* OutputValues[][UE_ARRAY_COUNT(RoundingModeNames)] = 
+	{
+		// HalfToEven        | HalfFromZero      | HalfToZero        | FromZero          | ToZero            | ToNegativeInfinity | ToPositiveInfinity
+		{  TEXT("1000.122"),   TEXT("1000.122"),   TEXT("1000.122"),   TEXT("1000.123"),   TEXT("1000.122"),   TEXT("1000.122"),    TEXT("1000.123") },
+		{  TEXT("1000.122"),   TEXT("1000.123"),   TEXT("1000.122"),   TEXT("1000.123"),   TEXT("1000.122"),   TEXT("1000.122"),    TEXT("1000.123") },
+		{  TEXT("1000.123"),   TEXT("1000.123"),   TEXT("1000.123"),   TEXT("1000.123"),   TEXT("1000.122"),   TEXT("1000.122"),    TEXT("1000.123") },
+		{  TEXT("1000.123"),   TEXT("1000.123"),   TEXT("1000.123"),   TEXT("1000.124"),   TEXT("1000.123"),   TEXT("1000.123"),    TEXT("1000.124") },
+		{  TEXT("1000.124"),   TEXT("1000.124"),   TEXT("1000.123"),   TEXT("1000.124"),   TEXT("1000.123"),   TEXT("1000.123"),    TEXT("1000.124") },
+		{  TEXT("1000.124"),   TEXT("1000.124"),   TEXT("1000.124"),   TEXT("1000.124"),   TEXT("1000.123"),   TEXT("1000.123"),    TEXT("1000.124") },
+
+		{  TEXT("1000.124"),   TEXT("1000.124"),   TEXT("1000.124"),   TEXT("1000.125"),   TEXT("1000.124"),   TEXT("1000.124"),    TEXT("1000.125") },
+		{  TEXT("1000.124"),   TEXT("1000.125"),   TEXT("1000.124"),   TEXT("1000.125"),   TEXT("1000.124"),   TEXT("1000.124"),    TEXT("1000.125") },
+		{  TEXT("1000.125"),   TEXT("1000.125"),   TEXT("1000.125"),   TEXT("1000.125"),   TEXT("1000.124"),   TEXT("1000.124"),    TEXT("1000.125") },
+		{  TEXT("1000.125"),   TEXT("1000.125"),   TEXT("1000.125"),   TEXT("1000.126"),   TEXT("1000.125"),   TEXT("1000.125"),    TEXT("1000.126") },
+		{  TEXT("1000.126"),   TEXT("1000.126"),   TEXT("1000.125"),   TEXT("1000.126"),   TEXT("1000.125"),   TEXT("1000.125"),    TEXT("1000.126") },
+		{  TEXT("1000.126"),   TEXT("1000.126"),   TEXT("1000.126"),   TEXT("1000.126"),   TEXT("1000.125"),   TEXT("1000.125"),    TEXT("1000.126") },
+
+		{ TEXT("-1000.122"),  TEXT("-1000.122"),  TEXT("-1000.122"),  TEXT("-1000.123"),  TEXT("-1000.122"),  TEXT("-1000.123"),   TEXT("-1000.122") },
+		{ TEXT("-1000.122"),  TEXT("-1000.123"),  TEXT("-1000.122"),  TEXT("-1000.123"),  TEXT("-1000.122"),  TEXT("-1000.123"),   TEXT("-1000.122") },
+		{ TEXT("-1000.123"),  TEXT("-1000.123"),  TEXT("-1000.123"),  TEXT("-1000.123"),  TEXT("-1000.122"),  TEXT("-1000.123"),   TEXT("-1000.122") },
+		{ TEXT("-1000.123"),  TEXT("-1000.123"),  TEXT("-1000.123"),  TEXT("-1000.124"),  TEXT("-1000.123"),  TEXT("-1000.124"),   TEXT("-1000.123") },
+		{ TEXT("-1000.124"),  TEXT("-1000.124"),  TEXT("-1000.123"),  TEXT("-1000.124"),  TEXT("-1000.123"),  TEXT("-1000.124"),   TEXT("-1000.123") },
+		{ TEXT("-1000.124"),  TEXT("-1000.124"),  TEXT("-1000.124"),  TEXT("-1000.124"),  TEXT("-1000.123"),  TEXT("-1000.124"),   TEXT("-1000.123") },
+
+		{ TEXT("-1000.124"),  TEXT("-1000.124"),  TEXT("-1000.124"),  TEXT("-1000.125"),  TEXT("-1000.124"),  TEXT("-1000.125"),   TEXT("-1000.124") },
+		{ TEXT("-1000.124"),  TEXT("-1000.125"),  TEXT("-1000.124"),  TEXT("-1000.125"),  TEXT("-1000.124"),  TEXT("-1000.125"),   TEXT("-1000.124") },
+		{ TEXT("-1000.125"),  TEXT("-1000.125"),  TEXT("-1000.125"),  TEXT("-1000.125"),  TEXT("-1000.124"),  TEXT("-1000.125"),   TEXT("-1000.124") },
+		{ TEXT("-1000.125"),  TEXT("-1000.125"),  TEXT("-1000.125"),  TEXT("-1000.126"),  TEXT("-1000.125"),  TEXT("-1000.126"),   TEXT("-1000.125") },
+		{ TEXT("-1000.126"),  TEXT("-1000.126"),  TEXT("-1000.125"),  TEXT("-1000.126"),  TEXT("-1000.125"),  TEXT("-1000.126"),   TEXT("-1000.125") },
+		{ TEXT("-1000.126"),  TEXT("-1000.126"),  TEXT("-1000.126"),  TEXT("-1000.126"),  TEXT("-1000.125"),  TEXT("-1000.126"),   TEXT("-1000.125") },
+	};
+
+	static_assert(UE_ARRAY_COUNT(InputValues) == UE_ARRAY_COUNT(OutputValues), "The size of InputValues does not match OutputValues");
+
+	FInternationalization& I18N = FInternationalization::Get();
+
+	FInternationalization::FCultureStateSnapshot OriginalCultureState;
+	I18N.BackupCultureState(OriginalCultureState);
+	
+	// This test needs to be run using an English culture
+	I18N.SetCurrentCulture(TEXT("en"));
+
+	// Test to make sure that the decimal formatter is rounding fractional numbers correctly (to 3 decimal places)
+	FNumberFormattingOptions FormattingOptions = FNumberFormattingOptions()
+		.SetUseGrouping(false)
+		.SetMaximumFractionalDigits(3);
+
+	auto DoSingleTest = [&](const double InNumber, const FString& InExpectedString, const FString& InDescription)
+	{
+		const FText ResultText = FText::AsNumber(InNumber, &FormattingOptions);
+		if (ResultText.ToString() != InExpectedString)
+		{
+			FAIL_CHECK(FString::Printf(TEXT("Text rounding failure: source '%f' - expected '%s' - result '%s'. %s."), InNumber, *InExpectedString, *ResultText.ToString(), *InDescription));
+		}
+	};
+
+	auto DoAllTests = [&](const ERoundingMode InRoundingMode)
+	{
+		FormattingOptions.SetRoundingMode(InRoundingMode);
+
+		for (int32 TestValueIndex = 0; TestValueIndex < UE_ARRAY_COUNT(InputValues); ++TestValueIndex)
+		{
+			DoSingleTest(InputValues[TestValueIndex], OutputValues[TestValueIndex][InRoundingMode], RoundingModeNames[InRoundingMode]);
+		}
+	};
+
+	DoAllTests(ERoundingMode::HalfToEven);
+	DoAllTests(ERoundingMode::HalfFromZero);
+	DoAllTests(ERoundingMode::HalfToZero);
+	DoAllTests(ERoundingMode::FromZero);
+	DoAllTests(ERoundingMode::ToZero);
+	DoAllTests(ERoundingMode::ToNegativeInfinity);
+	DoAllTests(ERoundingMode::ToPositiveInfinity);
+
+	// HalfToEven - Rounds to the nearest place, equidistant ties go to the value which is closest to an even value: 1.5 becomes 2, 0.5 becomes 0
+	{
+		FormattingOptions.SetRoundingMode(ERoundingMode::HalfToEven);
+
+		DoSingleTest(1000.12459, TEXT("1000.125"), TEXT("HalfToEven"));
+		DoSingleTest(1000.124549, TEXT("1000.125"), TEXT("HalfToEven"));
+		DoSingleTest(1000.124551, TEXT("1000.125"), TEXT("HalfToEven"));
+		DoSingleTest(1000.12451, TEXT("1000.125"), TEXT("HalfToEven"));
+		DoSingleTest(1000.1245000001, TEXT("1000.125"), TEXT("HalfToEven"));
+		DoSingleTest(1000.12450000000001, TEXT("1000.124"), TEXT("HalfToEven"));
+
+		DoSingleTest(512.9999, TEXT("513"), TEXT("HalfToEven"));
+		DoSingleTest(-512.9999, TEXT("-513"), TEXT("HalfToEven"));
+	}
+
+	// Restore original culture
+	I18N.RestoreCultureState(OriginalCultureState);
+}
+
+
+TEST_CASE_NAMED(FTextPaddingTest, "System::Core::Misc::TextPadding", "[.][EditorContext][ClientContext][EngineFilter]")
+{
+	FInternationalization& I18N = FInternationalization::Get();
+
+	FInternationalization::FCultureStateSnapshot OriginalCultureState;
+	I18N.BackupCultureState(OriginalCultureState);
+
+	// This test needs to be run using an English culture
+	I18N.SetCurrentCulture(TEXT("en"));
+
+	// Test to make sure that the decimal formatter is padding integral numbers correctly
+	FNumberFormattingOptions FormattingOptions;
+
+	auto DoSingleIntTest = [&](const int32 InNumber, const FString& InExpectedString, const FString& InDescription)
+	{
+		const FText ResultText = FText::AsNumber(InNumber, &FormattingOptions);
+		if (ResultText.ToString() != InExpectedString)
+		{
+			FAIL_CHECK(FString::Printf(TEXT("Text padding failure: source '%d' - expected '%s' - result '%s'. %s."), InNumber, *InExpectedString, *ResultText.ToString(), *InDescription));
+		}
+	};
+
+	auto DoSingleDoubleTest = [&](const double InNumber, const FString& InExpectedString, const FString& InDescription)
+	{
+		const FText ResultText = FText::AsNumber(InNumber, &FormattingOptions);
+		if (ResultText.ToString() != InExpectedString)
+		{
+			FAIL_CHECK(FString::Printf(TEXT("Text padding failure: source '%f' - expected '%s' - result '%s'. %s."), InNumber, *InExpectedString, *ResultText.ToString(), *InDescription));
+		}
+	};
+
+	// Test with a max limit of 3
+	{
+		FormattingOptions = FNumberFormattingOptions()
+			.SetUseGrouping(false)
+			.SetMaximumIntegralDigits(3);
+
+		DoSingleIntTest(123456,  TEXT("456"),  TEXT("Truncating '123456' to a max of 3 integral digits"));
+		DoSingleIntTest(-123456, TEXT("-456"), TEXT("Truncating '-123456' to a max of 3 integral digits"));
+	}
+
+	// Test with a min limit of 6
+	{
+		FormattingOptions = FNumberFormattingOptions()
+			.SetUseGrouping(false)
+			.SetMinimumIntegralDigits(6);
+
+		DoSingleIntTest(123,  TEXT("000123"),  TEXT("Padding '123' to a min of 6 integral digits"));
+		DoSingleIntTest(-123, TEXT("-000123"), TEXT("Padding '-123' to a min of 6 integral digits"));
+	}
+
+	// Test with forced fractional digits
+	{
+		FormattingOptions = FNumberFormattingOptions()
+			.SetUseGrouping(false)
+			.SetMinimumFractionalDigits(3);
+
+		DoSingleIntTest(123,  TEXT("123.000"),  TEXT("Padding '123' to a min of 3 fractional digits"));
+		DoSingleIntTest(-123, TEXT("-123.000"), TEXT("Padding '-123' to a min of 3 fractional digits"));
+	}
+
+	// Testing with leading zeros on a real number
+	{
+		FormattingOptions = FNumberFormattingOptions()
+			.SetUseGrouping(false)
+			.SetMaximumFractionalDigits(4);
+
+		DoSingleDoubleTest(0.00123,  TEXT("0.0012"),  TEXT("Padding '0.00123' to a max of 4 fractional digits"));
+		DoSingleDoubleTest(-0.00123, TEXT("-0.0012"), TEXT("Padding '-0.00123' to a max of 4 fractional digits"));
+	}
+
+	// Testing with leading zeros on a real number
+	{
+		FormattingOptions = FNumberFormattingOptions()
+			.SetUseGrouping(false)
+			.SetMaximumFractionalDigits(8);
+
+		DoSingleDoubleTest(0.00123,  TEXT("0.00123"),  TEXT("Padding '0.00123' to a max of 8 fractional digits"));
+		DoSingleDoubleTest(-0.00123, TEXT("-0.00123"), TEXT("Padding '-0.00123' to a max of 8 fractional digits"));
+	}
+
+	// Test with forced fractional digits on a real number
+	{
+		FormattingOptions = FNumberFormattingOptions()
+			.SetUseGrouping(false)
+			.SetMinimumFractionalDigits(8)
+			.SetMaximumFractionalDigits(8);
+
+		DoSingleDoubleTest(0.00123,  TEXT("0.00123000"),  TEXT("Padding '0.00123' to a min of 8 fractional digits"));
+		DoSingleDoubleTest(-0.00123, TEXT("-0.00123000"), TEXT("Padding '-0.00123' to a min of 8 fractional digits"));
+	}
+
+	// Restore original culture
+	I18N.RestoreCultureState(OriginalCultureState);
+}
+
+
+class FTextNumericParsingTestClass {
+public:
+	struct FTextNumericParsingTestUtil
+	{
+		template <typename T>
+		static void DoTest(const TCHAR* InStr, const int32 InStrLen, const FDecimalNumberFormattingRules& InFormattingRules, const FNumberParsingOptions& InParsingOptions, const T InExpectedValue, const bool bExpectedToParse, const TCHAR* InDescription)
+		{
+			T Value;
+			const bool bDidParse = FastDecimalFormat::StringToNumber(InStr, InStrLen, InFormattingRules, InParsingOptions, Value);
+
+			if (bDidParse != bExpectedToParse)
+			{
+				FAIL_CHECK(FString::Printf(TEXT("Text parsing failure: source '%s' - expected to parse '%s' - result '%s'. %s."), InStr, bExpectedToParse ? TEXT("true") : TEXT("false"), bDidParse ? TEXT("true") : TEXT("false"), InDescription));
+				return;
+			}
+
+			if (bDidParse && Value != InExpectedValue)
+			{
+				FAIL_CHECK(FString::Printf(TEXT("Text parsing failure: source '%s' - expected value '%f' - result '%f'. %s."), InStr, (double)InExpectedValue, (double)Value, InDescription));
+				return;
+			}
+		}
+
+		template <typename T>
+		static void DoGroupingTest(const TCHAR* InStr, const int32 InStrLen, const FDecimalNumberFormattingRules& InFormattingRules, const T InExpectedValue, const bool bExpectedToParse, const TCHAR* InDescription)
+		{
+			DoTest(InStr, InStrLen, InFormattingRules, FNumberParsingOptions::DefaultWithGrouping(), InExpectedValue, bExpectedToParse, InDescription);
+		}
+
+		template <typename T>
+		static void DoGroupingTest(const TCHAR* InStr, const FDecimalNumberFormattingRules& InFormattingRules, const T InExpectedValue, const bool bExpectedToParse, const TCHAR* InDescription)
+		{
+			DoGroupingTest(InStr, FCString::Strlen(InStr), InFormattingRules, InExpectedValue, bExpectedToParse, InDescription);
+		}
+
+		template <typename T>
+		static void DoLimitsTest(const TCHAR* InStr, const FDecimalNumberFormattingRules& InFormattingRules, const T InExpectedValue, const bool bExpectedToParse, const TCHAR* InDescription)
+		{
+			DoTest(InStr, FCString::Strlen(InStr), InFormattingRules, FNumberParsingOptions().SetUseGrouping(true).SetInsideLimits(true), InExpectedValue, bExpectedToParse, InDescription);
+		}
+
+		template <typename T>
+		static void DoClampTest(const TCHAR* InStr, const FDecimalNumberFormattingRules& InFormattingRules, const T InExpectedValue, const bool bExpectedToParse, const TCHAR* InDescription)
+		{
+			DoTest(InStr, FCString::Strlen(InStr), InFormattingRules, FNumberParsingOptions().SetUseGrouping(true).SetUseClamping(true), InExpectedValue, bExpectedToParse, InDescription);
+		}
+
+		template <typename T>
+		static void DoAllTests(const TCHAR* InStr, const FDecimalNumberFormattingRules& InFormattingRules, const T InExpectedValue, const T InExpectedClampedValue, const bool bExpectedToParse, const bool bExpectedToParseStrict, const TCHAR* InDescription)
+		{
+			DoGroupingTest(InStr, InFormattingRules, InExpectedValue, bExpectedToParse, InDescription);
+			DoClampTest(InStr, InFormattingRules, InExpectedClampedValue, bExpectedToParse, InDescription);
+			DoLimitsTest(InStr, InFormattingRules, InExpectedValue, bExpectedToParseStrict, InDescription);
+		}
+	};
+
+	void TextNumericParsingTest() {
+		FInternationalization& I18N = FInternationalization::Get();
+
+		auto DoTests = [this](const FString& InCulture)
+		{
+			FCulturePtr Culture = FInternationalization::Get().GetCulture(InCulture);
+			if (Culture.IsValid())
+			{
+				const FDecimalNumberFormattingRules& FormattingRules = Culture->GetDecimalNumberFormattingRules();
+
+				auto BuildDescription = [&InCulture](const TCHAR* InTestStr, const TCHAR* InTypeStr) -> FString
+				{
+					return FString::Printf(TEXT("[%s] Parsing '%s' as '%s'"), *InCulture, InTestStr, InTypeStr);
+				};
+
+				const FString UnsignedString = FString::Printf(TEXT("135%c456"), FormattingRules.DecimalSeparatorCharacter);
+				const FString PositiveString = FString::Printf(TEXT("%s135%c456"), *FormattingRules.PlusString, FormattingRules.DecimalSeparatorCharacter);
+				const FString NegativeString = FString::Printf(TEXT("%s135%c456"), *FormattingRules.MinusString, FormattingRules.DecimalSeparatorCharacter);
+				const FString PositiveASCIIString = FString::Printf(TEXT("+135%c456"), FormattingRules.DecimalSeparatorCharacter);
+				const FString NegativeASCIIString = FString::Printf(TEXT("-135%c456"), FormattingRules.DecimalSeparatorCharacter);
+				const FString GroupSeparatedString = FString::Printf(TEXT("1%c234%c5"), FormattingRules.GroupingSeparatorCharacter, FormattingRules.DecimalSeparatorCharacter);
+
+				int32 Number135 = 135;
+
+				FTextNumericParsingTestUtil::DoAllTests<int8>(*UnsignedString, FormattingRules, static_cast<int8>(Number135), TNumericLimits<int8>::Max(), true, false, *BuildDescription(*UnsignedString, TEXT("int8")));
+				FTextNumericParsingTestUtil::DoAllTests<uint8>(*UnsignedString, FormattingRules, static_cast<uint8>(Number135), 135, true, true, *BuildDescription(*UnsignedString, TEXT("uint8")));
+				FTextNumericParsingTestUtil::DoGroupingTest<int16>(*UnsignedString, FormattingRules, static_cast<int16>(Number135), true, *BuildDescription(*UnsignedString, TEXT("int16")));
+				FTextNumericParsingTestUtil::DoGroupingTest<uint16>(*UnsignedString, FormattingRules, static_cast<uint16>(Number135), true, *BuildDescription(*UnsignedString, TEXT("uint16")));
+				FTextNumericParsingTestUtil::DoGroupingTest<int32>(*UnsignedString, FormattingRules, Number135, true, *BuildDescription(*UnsignedString, TEXT("int32")));
+				FTextNumericParsingTestUtil::DoGroupingTest<uint32>(*UnsignedString, FormattingRules, static_cast<uint32>(Number135), true, *BuildDescription(*UnsignedString, TEXT("uint32")));
+				FTextNumericParsingTestUtil::DoGroupingTest<int64>(*UnsignedString, FormattingRules, 135, true, *BuildDescription(*UnsignedString, TEXT("int64")));
+				FTextNumericParsingTestUtil::DoGroupingTest<uint64>(*UnsignedString, FormattingRules, static_cast<uint64>(Number135), true, *BuildDescription(*UnsignedString, TEXT("uint64")));
+				FTextNumericParsingTestUtil::DoGroupingTest<float>(*UnsignedString, FormattingRules, 135.456f, true, *BuildDescription(*UnsignedString, TEXT("float")));
+				FTextNumericParsingTestUtil::DoGroupingTest<double>(*UnsignedString, FormattingRules, 135.456, true, *BuildDescription(*UnsignedString, TEXT("double")));
+
+				FTextNumericParsingTestUtil::DoGroupingTest<int32>(*PositiveString, FormattingRules, 135, true, *BuildDescription(*PositiveString, TEXT("int32")));
+				FTextNumericParsingTestUtil::DoGroupingTest<uint32>(*PositiveString, FormattingRules, 135, true, *BuildDescription(*PositiveString, TEXT("uint32")));
+				FTextNumericParsingTestUtil::DoGroupingTest<float>(*PositiveString, FormattingRules, 135.456f, true, *BuildDescription(*PositiveString, TEXT("float")));
+				FTextNumericParsingTestUtil::DoGroupingTest<double>(*PositiveString, FormattingRules, 135.456, true, *BuildDescription(*PositiveString, TEXT("double")));
+
+				FTextNumericParsingTestUtil::DoAllTests<int8>(*NegativeString, FormattingRules, static_cast<int8>(-Number135), TNumericLimits<int8>::Lowest(), true, false, *BuildDescription(*NegativeString, TEXT("int8")));
+				FTextNumericParsingTestUtil::DoAllTests<int32>(*NegativeString, FormattingRules, -135, -135, true, true, *BuildDescription(*NegativeString, TEXT("int32")));
+				FTextNumericParsingTestUtil::DoAllTests<uint32>(*NegativeString, FormattingRules, static_cast<uint32>(-Number135), 0, true, false, *BuildDescription(*NegativeString, TEXT("uint32")));
+				FTextNumericParsingTestUtil::DoGroupingTest<float>(*NegativeString, FormattingRules, -135.456f, true, *BuildDescription(*NegativeString, TEXT("float")));
+				FTextNumericParsingTestUtil::DoGroupingTest<double>(*NegativeString, FormattingRules, -135.456, true, *BuildDescription(*NegativeString, TEXT("double")));
+
+				FTextNumericParsingTestUtil::DoGroupingTest<int32>(*PositiveASCIIString, FormattingRules, 135, true, *BuildDescription(*PositiveASCIIString, TEXT("int32")));
+				FTextNumericParsingTestUtil::DoGroupingTest<int32>(*NegativeASCIIString, FormattingRules, -135, true, *BuildDescription(*NegativeASCIIString, TEXT("int32")));
+				FTextNumericParsingTestUtil::DoGroupingTest<float>(*PositiveASCIIString, FormattingRules, 135.456f, true, *BuildDescription(*PositiveASCIIString, TEXT("float")));
+				FTextNumericParsingTestUtil::DoGroupingTest<float>(*NegativeASCIIString, FormattingRules, -135.456f, true, *BuildDescription(*PositiveASCIIString, TEXT("float")));
+
+				int32 Number1234 = 1234;
+				FTextNumericParsingTestUtil::DoAllTests<uint8>(*GroupSeparatedString, FormattingRules, static_cast<uint8>(Number1234), TNumericLimits<uint8>::Max(), true, false, *BuildDescription(*GroupSeparatedString, TEXT("int32")));
+				FTextNumericParsingTestUtil::DoGroupingTest<int32>(*GroupSeparatedString, FormattingRules, Number1234, true, *BuildDescription(*GroupSeparatedString, TEXT("int32")));
+				FTextNumericParsingTestUtil::DoGroupingTest<uint32>(*GroupSeparatedString, FormattingRules, static_cast<uint32>(Number1234), true, *BuildDescription(*GroupSeparatedString, TEXT("uint32")));
+				FTextNumericParsingTestUtil::DoGroupingTest<float>(*GroupSeparatedString, FormattingRules, 1234.5f, true, *BuildDescription(*GroupSeparatedString, TEXT("float")));
+
+				uint64 BigNumber = 9223372036854775809ull; // (double)9223372036854775809 == 9223372036854775808.0, (last digic is not the same) && 9223372036854775809 > int64::max()
+				const FString BigUnsignedString = FString::Printf(TEXT("9223372036854775809"));
+				const FString BigPositiveString = FString::Printf(TEXT("%s9223372036854775809"), *FormattingRules.PlusString);
+				const FString BigNegativeString = FString::Printf(TEXT("%s9223372036854775809"), *FormattingRules.MinusString);
+				int64 BigGroupedNumber = 9223372036854775800ll;
+				const FString BigGroupSeparatedString = FString::Printf(TEXT("9%c223%c372%c036%c854%c775%c800")
+					, FormattingRules.GroupingSeparatorCharacter, FormattingRules.GroupingSeparatorCharacter, FormattingRules.GroupingSeparatorCharacter
+					, FormattingRules.GroupingSeparatorCharacter, FormattingRules.GroupingSeparatorCharacter, FormattingRules.GroupingSeparatorCharacter);
+
+				FTextNumericParsingTestUtil::DoAllTests<int32>(*BigUnsignedString, FormattingRules, static_cast<int32>(BigNumber), TNumericLimits<int32>::Max(), true, false, *BuildDescription(*BigUnsignedString, TEXT("int32")));
+				FTextNumericParsingTestUtil::DoAllTests<uint32>(*BigUnsignedString, FormattingRules, static_cast<uint32>(BigNumber), TNumericLimits<uint32>::Max(), true, false, *BuildDescription(*BigUnsignedString, TEXT("uint32")));
+				FTextNumericParsingTestUtil::DoAllTests<int64>(*BigUnsignedString, FormattingRules, static_cast<int64>(BigNumber), TNumericLimits<int64>::Max(), true, false, *BuildDescription(*BigUnsignedString, TEXT("int64")));
+				FTextNumericParsingTestUtil::DoAllTests<uint64>(*BigUnsignedString, FormattingRules, BigNumber, BigNumber, true, true, *BuildDescription(*BigUnsignedString, TEXT("uint64")));
+				FTextNumericParsingTestUtil::DoAllTests<float>(*BigUnsignedString, FormattingRules, static_cast<float>(BigNumber), static_cast<float>(BigNumber), true, true, *BuildDescription(*BigUnsignedString, TEXT("float")));
+				FTextNumericParsingTestUtil::DoAllTests<double>(*BigUnsignedString, FormattingRules, static_cast<double>(BigNumber), static_cast<double>(BigNumber), true, true, *BuildDescription(*BigUnsignedString, TEXT("double")));
+
+				FTextNumericParsingTestUtil::DoAllTests<int64>(*BigPositiveString, FormattingRules, static_cast<int64>(BigNumber), TNumericLimits<int64>::Max(), true, false, *BuildDescription(*BigPositiveString, TEXT("int64")));
+				FTextNumericParsingTestUtil::DoAllTests<uint64>(*BigPositiveString, FormattingRules, BigNumber, BigNumber, true, true, *BuildDescription(*BigPositiveString, TEXT("uint64")));
+				FTextNumericParsingTestUtil::DoAllTests<float>(*BigPositiveString, FormattingRules, static_cast<float>(BigNumber), static_cast<float>(BigNumber), true, true, *BuildDescription(*BigPositiveString, TEXT("float")));
+				FTextNumericParsingTestUtil::DoAllTests<double>(*BigPositiveString, FormattingRules, static_cast<double>(BigNumber), static_cast<double>(BigNumber), true, true, *BuildDescription(*BigPositiveString, TEXT("double")));
+
+				FTextNumericParsingTestUtil::DoAllTests<int64>(*BigNegativeString, FormattingRules, -static_cast<int64>(BigNumber), TNumericLimits<int64>::Lowest(), true, false, *BuildDescription(*NegativeString, TEXT("int64")));
+				FTextNumericParsingTestUtil::DoAllTests<uint64>(*BigNegativeString, FormattingRules, -static_cast<int64>(BigNumber), TNumericLimits<uint64>::Lowest(), true, false, *BuildDescription(*NegativeString, TEXT("uint64")));
+				FTextNumericParsingTestUtil::DoAllTests<float>(*BigNegativeString, FormattingRules, -static_cast<float>(BigNumber), -static_cast<float>(BigNumber), true, true, *BuildDescription(*NegativeString, TEXT("float")));
+				FTextNumericParsingTestUtil::DoAllTests<double>(*BigNegativeString, FormattingRules, -static_cast<double>(BigNumber), -static_cast<double>(BigNumber), true, true, *BuildDescription(*NegativeString, TEXT("double")));
+
+				FTextNumericParsingTestUtil::DoAllTests<int32>(*BigGroupSeparatedString, FormattingRules, static_cast<int32>(BigGroupedNumber), TNumericLimits<int32>::Max(), true, false, *BuildDescription(*BigGroupSeparatedString, TEXT("int32")));
+				FTextNumericParsingTestUtil::DoAllTests<uint32>(*BigGroupSeparatedString, FormattingRules, static_cast<int32>(BigGroupedNumber), TNumericLimits<uint32>::Max(), true, false, *BuildDescription(*BigGroupSeparatedString, TEXT("uint32")));
+				FTextNumericParsingTestUtil::DoAllTests<int64>(*BigGroupSeparatedString, FormattingRules, BigGroupedNumber, BigGroupedNumber, true, true, *BuildDescription(*BigGroupSeparatedString, TEXT("int64")));
+				FTextNumericParsingTestUtil::DoAllTests<uint64>(*BigGroupSeparatedString, FormattingRules, static_cast<uint64>(BigGroupedNumber), BigGroupedNumber, true, true, *BuildDescription(*BigGroupSeparatedString, TEXT("uint64")));
+				FTextNumericParsingTestUtil::DoAllTests<float>(*BigGroupSeparatedString, FormattingRules, static_cast<float>(BigGroupedNumber), static_cast<float>(BigGroupedNumber), true, true, *BuildDescription(*BigGroupSeparatedString, TEXT("float")));
+				FTextNumericParsingTestUtil::DoAllTests<double>(*BigGroupSeparatedString, FormattingRules, static_cast<double>(BigGroupedNumber), static_cast<double>(BigGroupedNumber), true, true, *BuildDescription(*BigGroupSeparatedString, TEXT("double")));
+			}
+		};
+
+		DoTests(TEXT("en"));
+		DoTests(TEXT("fr"));
+		DoTests(TEXT("ar"));
+
+		{
+			const FDecimalNumberFormattingRules& AgnosticFormattingRules = FastDecimalFormat::GetCultureAgnosticFormattingRules();
+
+			FTextNumericParsingTestUtil::DoGroupingTest<int32>(TEXT("10a"), AgnosticFormattingRules, 0, false, TEXT("Parsing '10a' as 'int32'"));
+			FTextNumericParsingTestUtil::DoGroupingTest<uint32>(TEXT("10a"), AgnosticFormattingRules, 0, false, TEXT("Parsing '10a' as 'uint32'"));
+
+			FTextNumericParsingTestUtil::DoGroupingTest<int32>(TEXT("10a"), 2, AgnosticFormattingRules, 10, true, TEXT("Parsing '10a' (len 2) as 'int32'"));
+			FTextNumericParsingTestUtil::DoGroupingTest<uint32>(TEXT("10a"), 2, AgnosticFormattingRules, 10, true, TEXT("Parsing '10a' (len 2) as 'uint32'"));
+		}
+
+		{
+			const FDecimalNumberFormattingRules& AgnosticFormattingRules = FastDecimalFormat::GetCultureAgnosticFormattingRules();
+
+			// test limits
+			FTextNumericParsingTestUtil::DoAllTests<int8>(TEXT("-128"), AgnosticFormattingRules, TNumericLimits<int8>::Lowest(), TNumericLimits<int8>::Lowest(), true, true, TEXT("Parsing int8 lowest"));
+			FTextNumericParsingTestUtil::DoAllTests<int8>(TEXT("127"), AgnosticFormattingRules, TNumericLimits<int8>::Max(), TNumericLimits<int8>::Max(), true, true, TEXT("Parsing int8 max"));
+			FTextNumericParsingTestUtil::DoAllTests<uint8>(TEXT("0"), AgnosticFormattingRules, TNumericLimits<uint8>::Lowest(), TNumericLimits<uint8>::Lowest(), true, true, TEXT("Parsing uint8 lowest"));
+			FTextNumericParsingTestUtil::DoAllTests<uint8>(TEXT("255"), AgnosticFormattingRules, TNumericLimits<uint8>::Max(), TNumericLimits<uint8>::Max(), true, true, TEXT("Parsing uint8 max"));
+			FTextNumericParsingTestUtil::DoAllTests<int64>(TEXT("-9223372036854775808"), AgnosticFormattingRules, TNumericLimits<int64>::Lowest(), TNumericLimits<int64>::Lowest(), true, true, TEXT("Parsing int64 lowest"));
+			FTextNumericParsingTestUtil::DoAllTests<int64>(TEXT("9223372036854775807"), AgnosticFormattingRules, TNumericLimits<int64>::Max(), TNumericLimits<int64>::Max(), true, true, TEXT("Parsing int64 max"));
+			FTextNumericParsingTestUtil::DoAllTests<uint64>(TEXT("0"), AgnosticFormattingRules, TNumericLimits<uint64>::Lowest(), TNumericLimits<uint64>::Lowest(), true, true, TEXT("Parsing uint64 lowest"));
+			FTextNumericParsingTestUtil::DoAllTests<uint64>(TEXT("18446744073709551615"), AgnosticFormattingRules, TNumericLimits<uint64>::Max(), TNumericLimits<uint64>::Max(), true, true, TEXT("Parsing uint64 max"));
+
+			// test limits +- 1
+			int32 Number129 = 129;
+			int32 Number128 = 128;
+			int32 Number1 = 1;
+			int32 Number256 = 256;
+			FTextNumericParsingTestUtil::DoAllTests<int8>(TEXT("-129"), AgnosticFormattingRules, static_cast<int8>(-Number129), TNumericLimits<int8>::Lowest(), true, false, TEXT("Parsing int8 +/-1 lowest"));
+			FTextNumericParsingTestUtil::DoAllTests<int8>(TEXT("128"), AgnosticFormattingRules, static_cast<int8>(Number128), TNumericLimits<int8>::Max(), true, false, TEXT("Parsing int8 +/-1 max"));
+			FTextNumericParsingTestUtil::DoAllTests<uint8>(TEXT("-1"), AgnosticFormattingRules, static_cast<uint8>(-Number1), TNumericLimits<uint8>::Lowest(), true, false, TEXT("Parsing uint8 +/-1 lowest"));
+			FTextNumericParsingTestUtil::DoAllTests<uint8>(TEXT("256"), AgnosticFormattingRules, static_cast<uint8>(Number256), TNumericLimits<uint8>::Max(), true, false, TEXT("Parsing uint8 +/-1 max"));
+
+			int64 Number9223372036854775809 = TNumericLimits<int64>::Lowest();
+			--Number9223372036854775809;
+			int64 Number9223372036854775808 = TNumericLimits<int64>::Max();
+			++Number9223372036854775808;
+			uint64 Number18446744073709551616 = TNumericLimits<uint64>::Max();
+			++Number18446744073709551616;
+			FTextNumericParsingTestUtil::DoAllTests<int64>(TEXT("-9223372036854775809"), AgnosticFormattingRules, Number9223372036854775809, TNumericLimits<int64>::Lowest(), true, false, TEXT("Parsing int64 +/-1 lowest"));
+			FTextNumericParsingTestUtil::DoAllTests<int64>(TEXT("9223372036854775808"), AgnosticFormattingRules, Number9223372036854775808, TNumericLimits<int64>::Max(), true, false, TEXT("Parsing +/-1 int64 max"));
+			FTextNumericParsingTestUtil::DoAllTests<uint64>(TEXT("-1"), AgnosticFormattingRules, static_cast<uint64>(-Number1), TNumericLimits<uint64>::Lowest(), true, false, TEXT("Parsing uint64 +/-1 lowest"));
+			FTextNumericParsingTestUtil::DoAllTests<uint64>(TEXT("18446744073709551616"), AgnosticFormattingRules, Number18446744073709551616, TNumericLimits<uint64>::Max(), true, false, TEXT("Parsing +/-1 uint64 max"));
+
+			FTextNumericParsingTestUtil::DoGroupingTest<int64>(TEXT("-18446744073709551616"), AgnosticFormattingRules, 0, true, TEXT("Parsing negative overflow int64 max"));
+			FTextNumericParsingTestUtil::DoLimitsTest<int64>(TEXT("-18446744073709551616"), AgnosticFormattingRules, 0, false, TEXT("Parsing negative overflow int64 max"));
+			FTextNumericParsingTestUtil::DoClampTest<int64>(TEXT("-18446744073709551616"), AgnosticFormattingRules, TNumericLimits<int64>::Lowest(), true, TEXT("Parsing negative overflow int64 max"));
+		}
+	}
+};
+
+TEST_CASE_NAMED(FTextNumericParsingTest, "System::Core::Misc::TextNumericParsing", "[.][EditorContext][ClientContext][EngineFilter]")
+{
+	FTextNumericParsingTestClass Instance;
+	Instance.TextNumericParsingTest();
+}
+
+TEST_CASE_NAMED(FTextStringificationTest, "System::Core::Misc::TextStringification", "[.][EditorContext][ClientContext][EngineFilter]")
+{
+	FInternationalization& I18N = FInternationalization::Get();
+
+	FInternationalization::FCultureStateSnapshot OriginalCultureState;
+	I18N.BackupCultureState(OriginalCultureState);
+
+	// This test needs to be run using the English (US) culture to ensure the time formatting has a valid timezone to work with
+	I18N.SetCurrentCulture(TEXT("en-US"));
+
+	auto DoSingleTest = [](const FText& InExpectedText, const FString& InExpectedString, const FString& InCppString, const bool bImportCppString)
+	{
+		// Validate that the text produces the string we expect
+		FString ActualString;
+		FTextStringHelper::WriteToBuffer(ActualString, InExpectedText);
+		if (!ActualString.Equals(InExpectedString, ESearchCase::CaseSensitive))
+		{
+			FAIL_CHECK(FString::Printf(TEXT("Text export failure (from text): Text '%s' was expected to export as '%s', but produced '%s'."), *InExpectedText.ToString(), *InExpectedString, *ActualString));
+		}
+		// Validate that the string produces the text we expect
+		FText ActualText;
+		if (!FTextStringHelper::ReadFromBuffer(*InExpectedString, ActualText))
+		{
+			FAIL_CHECK(FString::Printf(TEXT("Text import failure (from string): String '%s' failed to import."), *InExpectedString));
+		}
+
+		if (!InExpectedText.ToString().Equals(ActualText.ToString(), ESearchCase::CaseSensitive))
+		{
+			FAIL_CHECK(FString::Printf(TEXT("Text import failure (from string): String '%s' was expected to import as '%s', but produced '%s'."), *InExpectedString, *InExpectedText.ToString(), *ActualText.ToString()));
+		}
+		// Validate that the C++ string produces the text we expect
+		if (bImportCppString)
+		{
+			FText ActualCppText;
+			if (!FTextStringHelper::ReadFromBuffer(*InCppString, ActualCppText))
+			{
+				FAIL_CHECK(FString::Printf(TEXT("Text import failure (from C++): String '%s' failed to import."), *InCppString));
+			}
+			if (!InExpectedText.ToString().Equals(ActualCppText.ToString(), ESearchCase::CaseSensitive))
+			{
+				FAIL_CHECK(FString::Printf(TEXT("Text import failure (from C++): String '%s' was expected to import as '%s', but produced '%s'."), *InCppString, *InExpectedText.ToString(), *ActualCppText.ToString()));
+			}
+		}
+	};
+
+#define TEST(Text, Str) DoSingleTest(Text, TEXT(Str), TEXT(#Text), true)
+#define TEST_EX(Text, Str, ImportCpp) DoSingleTest(Text, TEXT(Str), TEXT(#Text), ImportCpp)
+
+	// Add the test string table, but only if it isn't already!
+	if (!FStringTableRegistry::Get().FindStringTable("Core.Tests.TextFormatTest"))
+	{
+		LOCTABLE_NEW("Core.Tests.TextFormatTest", "Core.Tests.TextFormatTest");
+		LOCTABLE_SETSTRING("Core.Tests.TextFormatTest", "TextStringificationTest_Lorem", "Lorem");
+	}
+
+	TEST(NSLOCTEXT("Core.Tests.TextFormatTest", "TextStringificationTest_Lorem", "Lorem"), "NSLOCTEXT(\"Core.Tests.TextFormatTest\", \"TextStringificationTest_Lorem\", \"Lorem\")");
+	TEST(LOCTEXT("TextStringificationTest_Lorem", "Lorem"), "NSLOCTEXT(\"Core.Tests.TextFormatTest\", \"TextStringificationTest_Lorem\", \"Lorem\")");
+	TEST(LOCTABLE("Core.Tests.TextFormatTest", "TextStringificationTest_Lorem"), "LOCTABLE(\"Core.Tests.TextFormatTest\", \"TextStringificationTest_Lorem\")");
+	TEST(INVTEXT("DummyText"), "INVTEXT(\"DummyText\")");
+	if (GIsEditor)
+	{
+		TEST_EX(FText::FromString(TEXT("DummyString")), "DummyString", false);
+	}
+	else
+	{
+		TEST_EX(FText::FromString(TEXT("DummyString")), "INVTEXT(\"DummyString\")", false);
+	}
+
+	TEST(LOCGEN_NUMBER(10, ""), "LOCGEN_NUMBER(10, \"\")");
+	TEST(LOCGEN_NUMBER_GROUPED(12.5f, ""), "LOCGEN_NUMBER_GROUPED(12.500000f, \"\")");
+	TEST(LOCGEN_NUMBER_UNGROUPED(12.5f, ""), "LOCGEN_NUMBER_UNGROUPED(12.500000f, \"\")");
+	TEST(LOCGEN_NUMBER_CUSTOM(+10, SetAlwaysSign(true).SetRoundingMode(ERoundingMode::ToZero).SetMinimumFractionalDigits(2), ""), "LOCGEN_NUMBER_CUSTOM(10, SetAlwaysSign(true).SetRoundingMode(ERoundingMode::ToZero).SetMinimumFractionalDigits(2), \"\")");
+	TEST(LOCGEN_NUMBER(-10, "en"), "LOCGEN_NUMBER(-10, \"en\")");
+
+	TEST(LOCGEN_PERCENT(0.1f, ""), "LOCGEN_PERCENT(0.100000f, \"\")");
+	TEST(LOCGEN_PERCENT_GROUPED(0.1f, ""), "LOCGEN_PERCENT_GROUPED(0.100000f, \"\")");
+	TEST(LOCGEN_PERCENT_UNGROUPED(0.1f, ""), "LOCGEN_PERCENT_UNGROUPED(0.100000f, \"\")");
+	TEST(LOCGEN_PERCENT_CUSTOM(0.1f, SetAlwaysSign(true).SetRoundingMode(ERoundingMode::ToZero).SetMinimumFractionalDigits(2), ""), "LOCGEN_PERCENT_CUSTOM(0.100000f, SetAlwaysSign(true).SetRoundingMode(ERoundingMode::ToZero).SetMinimumFractionalDigits(2), \"\")");
+	TEST(LOCGEN_PERCENT(0.1, "en"), "LOCGEN_PERCENT(0.100000, \"en\")");
+
+	TEST(LOCGEN_CURRENCY(125, "USD", ""), "LOCGEN_CURRENCY(125, \"USD\", \"\")");
+	TEST_EX(FText::AsCurrency(1.25f, TEXT("USD"), nullptr, FInternationalization::Get().GetCulture(TEXT("en"))), "LOCGEN_CURRENCY(125, \"USD\", \"en\")", false);
+
+	TEST(LOCGEN_DATE_UTC(1526342400, EDateTimeStyle::Short, "", "en-GB"), "LOCGEN_DATE_UTC(1526342400, EDateTimeStyle::Short, \"\", \"en-GB\")");
+	TEST(LOCGEN_DATE_LOCAL(1526342400, EDateTimeStyle::Medium, ""), "LOCGEN_DATE_LOCAL(1526342400, EDateTimeStyle::Medium, \"\")");
+
+	TEST(LOCGEN_TIME_UTC(1526342400, EDateTimeStyle::Long, "", "en-GB"), "LOCGEN_TIME_UTC(1526342400, EDateTimeStyle::Long, \"\", \"en-GB\")");
+	TEST(LOCGEN_TIME_LOCAL(1526342400, EDateTimeStyle::Full, ""), "LOCGEN_TIME_LOCAL(1526342400, EDateTimeStyle::Full, \"\")");
+
+	TEST(LOCGEN_DATETIME_UTC(1526342400, EDateTimeStyle::Short, EDateTimeStyle::Medium, "", "en-GB"), "LOCGEN_DATETIME_UTC(1526342400, EDateTimeStyle::Short, EDateTimeStyle::Medium, \"\", \"en-GB\")");
+	TEST(LOCGEN_DATETIME_LOCAL(1526342400, EDateTimeStyle::Long, EDateTimeStyle::Full, ""), "LOCGEN_DATETIME_LOCAL(1526342400, EDateTimeStyle::Long, EDateTimeStyle::Full, \"\")");
+
+	TEST(LOCGEN_DATETIME_CUSTOM_UTC(1526342400, "%A, %B %e, %Y", "", "en-GB"), "LOCGEN_DATETIME_CUSTOM_UTC(1526342400, \"%A, %B %e, %Y\", \"\", \"en-GB\")");
+	TEST(LOCGEN_DATETIME_CUSTOM_LOCAL(1526342400, "%A, %B %e, %Y", ""), "LOCGEN_DATETIME_CUSTOM_LOCAL(1526342400, \"%A, %B %e, %Y\", \"\")");
+
+	TEST(LOCGEN_TOUPPER(LOCTEXT("TextStringificationTest_Lorem", "Lorem")), "LOCGEN_TOUPPER(NSLOCTEXT(\"Core.Tests.TextFormatTest\", \"TextStringificationTest_Lorem\", \"Lorem\"))");
+	TEST(LOCGEN_TOLOWER(LOCTEXT("TextStringificationTest_Lorem", "Lorem")), "LOCGEN_TOLOWER(NSLOCTEXT(\"Core.Tests.TextFormatTest\", \"TextStringificationTest_Lorem\", \"Lorem\"))");
+
+	TEST(LOCGEN_FORMAT_ORDERED(LOCTEXT("TextStringificationTest_FmtO", "{0} weighs {1}kg"), LOCTEXT("TextStringificationTest_Bear", "Bear"), 227), "LOCGEN_FORMAT_ORDERED(NSLOCTEXT(\"Core.Tests.TextFormatTest\", \"TextStringificationTest_FmtO\", \"{0} weighs {1}kg\"), NSLOCTEXT(\"Core.Tests.TextFormatTest\", \"TextStringificationTest_Bear\", \"Bear\"), 227)");
+	TEST(LOCGEN_FORMAT_NAMED(LOCTEXT("TextStringificationTest_FmtN", "{Animal} weighs {Weight}kg"), TEXT("Animal"), LOCTEXT("TextStringificationTest_Bear", "Bear"), TEXT("Weight"), 227), "LOCGEN_FORMAT_NAMED(NSLOCTEXT(\"Core.Tests.TextFormatTest\", \"TextStringificationTest_FmtN\", \"{Animal} weighs {Weight}kg\"), \"Animal\", NSLOCTEXT(\"Core.Tests.TextFormatTest\", \"TextStringificationTest_Bear\", \"Bear\"), \"Weight\", 227)");
+
+#undef TEST
+#undef TEST_EX
+
+	// Restore original culture
+	I18N.RestoreCultureState(OriginalCultureState);
+}
+
+TEST_CASE_NAMED(FTextFormatArgModifierTest, "System::Core::Misc::TextFormatArgModifiers", "[.][EditorContext][ClientContext][EngineFilter]")
+{
+	auto EnsureValidResult = [&](const FString& InResult, const FString& InExpected, const FString& InName, const FString& InDescription)
+	{
+		if (!InResult.Equals(InExpected, ESearchCase::CaseSensitive))
+		{
+			FAIL_CHECK(FString::Printf(TEXT("%s failure: result '%s' (expected '%s'). %s."), *InName, *InResult, *InExpected, *InDescription));
+		}
+	};
+
+	FInternationalization& I18N = FInternationalization::Get();
+
+	FInternationalization::FCultureStateSnapshot OriginalCultureState;
+	I18N.BackupCultureState(OriginalCultureState);
+
+	// This test needs to be run using an English culture
+	I18N.SetCurrentCulture(TEXT("en"));
+
+	{
+		const FTextFormat CardinalFormatText = INVTEXT("There {NumCats}|plural(one=is,other=are) {NumCats} {NumCats}|plural(one=cat,other=cats)");
+		EnsureValidResult(FText::FormatNamed(CardinalFormatText, TEXT("NumCats"), 0).ToString(), TEXT("There are 0 cats"), TEXT("CardinalResult0"), CardinalFormatText.GetSourceText().ToString());
+		EnsureValidResult(FText::FormatNamed(CardinalFormatText, TEXT("NumCats"), 1).ToString(), TEXT("There is 1 cat"), TEXT("CardinalResult1"), CardinalFormatText.GetSourceText().ToString());
+		EnsureValidResult(FText::FormatNamed(CardinalFormatText, TEXT("NumCats"), 2).ToString(), TEXT("There are 2 cats"), TEXT("CardinalResult2"), CardinalFormatText.GetSourceText().ToString());
+		EnsureValidResult(FText::FormatNamed(CardinalFormatText, TEXT("NumCats"), 3).ToString(), TEXT("There are 3 cats"), TEXT("CardinalResult3"), CardinalFormatText.GetSourceText().ToString());
+		EnsureValidResult(FText::FormatNamed(CardinalFormatText, TEXT("NumCats"), 4).ToString(), TEXT("There are 4 cats"), TEXT("CardinalResult4"), CardinalFormatText.GetSourceText().ToString());
+	}
+
+	{
+		const FTextFormat OrdinalFormatText = INVTEXT("You came {Place}{Place}|ordinal(one=st,two=nd,few=rd,other=th)!");
+		EnsureValidResult(FText::FormatNamed(OrdinalFormatText, TEXT("Place"), 0).ToString(), TEXT("You came 0th!"), TEXT("OrdinalResult0"), OrdinalFormatText.GetSourceText().ToString());
+		EnsureValidResult(FText::FormatNamed(OrdinalFormatText, TEXT("Place"), 1).ToString(), TEXT("You came 1st!"), TEXT("OrdinalResult1"), OrdinalFormatText.GetSourceText().ToString());
+		EnsureValidResult(FText::FormatNamed(OrdinalFormatText, TEXT("Place"), 2).ToString(), TEXT("You came 2nd!"), TEXT("OrdinalResult2"), OrdinalFormatText.GetSourceText().ToString());
+		EnsureValidResult(FText::FormatNamed(OrdinalFormatText, TEXT("Place"), 3).ToString(), TEXT("You came 3rd!"), TEXT("OrdinalResult3"), OrdinalFormatText.GetSourceText().ToString());
+		EnsureValidResult(FText::FormatNamed(OrdinalFormatText, TEXT("Place"), 4).ToString(), TEXT("You came 4th!"), TEXT("OrdinalResult4"), OrdinalFormatText.GetSourceText().ToString());
+	}
+
+	{
+		const FTextFormat GenderFormatText = INVTEXT("{Gender}|gender(Le,La) {Gender}|gender(guerrier,guerrière) est {Gender}|gender(fort,forte)");
+		EnsureValidResult(FText::FormatNamed(GenderFormatText, TEXT("Gender"), ETextGender::Masculine).ToString(), TEXT("Le guerrier est fort"), TEXT("GenderResultM"), GenderFormatText.GetSourceText().ToString());
+		EnsureValidResult(FText::FormatNamed(GenderFormatText, TEXT("Gender"), ETextGender::Feminine).ToString(), TEXT("La guerrière est forte"), TEXT("GenderResultF"), GenderFormatText.GetSourceText().ToString());
+	}
+
+	{
+		const FTextFormat GenderFormatText = INVTEXT("{Gender}|gender(Le guerrier est fort,La guerrière est forte)");
+		EnsureValidResult(FText::FormatNamed(GenderFormatText, TEXT("Gender"), ETextGender::Masculine).ToString(), TEXT("Le guerrier est fort"), TEXT("GenderResultM"), GenderFormatText.GetSourceText().ToString());
+		EnsureValidResult(FText::FormatNamed(GenderFormatText, TEXT("Gender"), ETextGender::Feminine).ToString(), TEXT("La guerrière est forte"), TEXT("GenderResultF"), GenderFormatText.GetSourceText().ToString());
+	}
+
+	{
+		const FText Consonant = INVTEXT("\uC0AC\uB78C");/* 사람 */
+		const FText ConsonantRieul = INVTEXT("\uC11C\uC6B8");/* 서울 */
+		const FText Vowel = INVTEXT("\uC0AC\uC790");/* 사자 */
+
+		{
+			const FTextFormat HppFormatText = INVTEXT("{Arg}|hpp(\uC740,\uB294)");/* 은/는 */
+			EnsureValidResult(FText::FormatNamed(HppFormatText, TEXT("Arg"), Consonant).ToString(), TEXT("\uC0AC\uB78C\uC740"), TEXT("HppResultConsonant"), HppFormatText.GetSourceText().ToString());
+			EnsureValidResult(FText::FormatNamed(HppFormatText, TEXT("Arg"), ConsonantRieul).ToString(), TEXT("\uC11C\uC6B8\uC740"), TEXT("HppResultConsonantRieul"), HppFormatText.GetSourceText().ToString());
+			EnsureValidResult(FText::FormatNamed(HppFormatText, TEXT("Arg"), Vowel).ToString(), TEXT("\uC0AC\uC790\uB294"), TEXT("HppResultVowel"), HppFormatText.GetSourceText().ToString());
+		}
+
+		{
+			const FTextFormat HppFormatText = INVTEXT("{Arg}|hpp(\uC774,\uAC00)");/* 이/가 */
+			EnsureValidResult(FText::FormatNamed(HppFormatText, TEXT("Arg"), Consonant).ToString(), TEXT("\uC0AC\uB78C\uC774"), TEXT("HppResultConsonant"), HppFormatText.GetSourceText().ToString());
+			EnsureValidResult(FText::FormatNamed(HppFormatText, TEXT("Arg"), ConsonantRieul).ToString(), TEXT("\uC11C\uC6B8\uC774"), TEXT("HppResultConsonantRieul"), HppFormatText.GetSourceText().ToString());
+			EnsureValidResult(FText::FormatNamed(HppFormatText, TEXT("Arg"), Vowel).ToString(), TEXT("\uC0AC\uC790\uAC00"), TEXT("HppResultVowel"), HppFormatText.GetSourceText().ToString());
+		}
+
+		{
+			const FTextFormat HppFormatText = INVTEXT("{Arg}|hpp(\uC744,\uB97C)");/* 을/를 */
+			EnsureValidResult(FText::FormatNamed(HppFormatText, TEXT("Arg"), Consonant).ToString(), TEXT("\uC0AC\uB78C\uC744"), TEXT("HppResultConsonant"), HppFormatText.GetSourceText().ToString());
+			EnsureValidResult(FText::FormatNamed(HppFormatText, TEXT("Arg"), ConsonantRieul).ToString(), TEXT("\uC11C\uC6B8\uC744"), TEXT("HppResultConsonantRieul"), HppFormatText.GetSourceText().ToString());
+			EnsureValidResult(FText::FormatNamed(HppFormatText, TEXT("Arg"), Vowel).ToString(), TEXT("\uC0AC\uC790\uB97C"), TEXT("HppResultVowel"), HppFormatText.GetSourceText().ToString());
+		}
+
+		{
+			const FTextFormat HppFormatText = INVTEXT("{Arg}|hpp(\uACFC,\uC640)");/* 과/와 */
+			EnsureValidResult(FText::FormatNamed(HppFormatText, TEXT("Arg"), Consonant).ToString(), TEXT("\uC0AC\uB78C\uACFC"), TEXT("HppResultConsonant"), HppFormatText.GetSourceText().ToString());
+			EnsureValidResult(FText::FormatNamed(HppFormatText, TEXT("Arg"), ConsonantRieul).ToString(), TEXT("\uC11C\uC6B8\uACFC"), TEXT("HppResultConsonantRieul"), HppFormatText.GetSourceText().ToString());
+			EnsureValidResult(FText::FormatNamed(HppFormatText, TEXT("Arg"), Vowel).ToString(), TEXT("\uC0AC\uC790\uC640"), TEXT("HppResultVowel"), HppFormatText.GetSourceText().ToString());
+		}
+
+		{
+			const FTextFormat HppFormatText = INVTEXT("{Arg}|hpp(\uC544,\uC57C)");/* 아/야 */
+			EnsureValidResult(FText::FormatNamed(HppFormatText, TEXT("Arg"), Consonant).ToString(), TEXT("\uC0AC\uB78C\uC544"), TEXT("HppResultConsonant"), HppFormatText.GetSourceText().ToString());
+			EnsureValidResult(FText::FormatNamed(HppFormatText, TEXT("Arg"), ConsonantRieul).ToString(), TEXT("\uC11C\uC6B8\uC544"), TEXT("HppResultConsonantRieul"), HppFormatText.GetSourceText().ToString());
+			EnsureValidResult(FText::FormatNamed(HppFormatText, TEXT("Arg"), Vowel).ToString(), TEXT("\uC0AC\uC790\uC57C"), TEXT("HppResultVowel"), HppFormatText.GetSourceText().ToString());
+		}
+
+		{
+			const FTextFormat HppFormatText = INVTEXT("{Arg}|hpp(\uC774\uC5B4,\uC5EC)");/* 이어/여 */
+			EnsureValidResult(FText::FormatNamed(HppFormatText, TEXT("Arg"), Consonant).ToString(), TEXT("\uC0AC\uB78C\uC774\uC5B4"), TEXT("HppResultConsonant"), HppFormatText.GetSourceText().ToString());
+			EnsureValidResult(FText::FormatNamed(HppFormatText, TEXT("Arg"), ConsonantRieul).ToString(), TEXT("\uC11C\uC6B8\uC774\uC5B4"), TEXT("HppResultConsonantRieul"), HppFormatText.GetSourceText().ToString());
+			EnsureValidResult(FText::FormatNamed(HppFormatText, TEXT("Arg"), Vowel).ToString(), TEXT("\uC0AC\uC790\uC5EC"), TEXT("HppResultVowel"), HppFormatText.GetSourceText().ToString());
+		}
+
+		{
+			const FTextFormat HppFormatText = INVTEXT("{Arg}|hpp(\uC774\uC5D0,\uC608)");/* 이에/예 */
+			EnsureValidResult(FText::FormatNamed(HppFormatText, TEXT("Arg"), Consonant).ToString(), TEXT("\uC0AC\uB78C\uC774\uC5D0"), TEXT("HppResultConsonant"), HppFormatText.GetSourceText().ToString());
+			EnsureValidResult(FText::FormatNamed(HppFormatText, TEXT("Arg"), ConsonantRieul).ToString(), TEXT("\uC11C\uC6B8\uC774\uC5D0"), TEXT("HppResultConsonantRieul"), HppFormatText.GetSourceText().ToString());
+			EnsureValidResult(FText::FormatNamed(HppFormatText, TEXT("Arg"), Vowel).ToString(), TEXT("\uC0AC\uC790\uC608"), TEXT("HppResultVowel"), HppFormatText.GetSourceText().ToString());
+		}
+
+		{
+			const FTextFormat HppFormatText = INVTEXT("{Arg}|hpp(\uC774\uC5C8,\uC601)");/* 이었/영*/
+			EnsureValidResult(FText::FormatNamed(HppFormatText, TEXT("Arg"), Consonant).ToString(), TEXT("\uC0AC\uB78C\uC774\uC5C8"), TEXT("HppResultConsonant"), HppFormatText.GetSourceText().ToString());
+			EnsureValidResult(FText::FormatNamed(HppFormatText, TEXT("Arg"), ConsonantRieul).ToString(), TEXT("\uC11C\uC6B8\uC774\uC5C8"), TEXT("HppResultConsonantRieul"), HppFormatText.GetSourceText().ToString());
+			EnsureValidResult(FText::FormatNamed(HppFormatText, TEXT("Arg"), Vowel).ToString(), TEXT("\uC0AC\uC790\uC601"), TEXT("HppResultVowel"), HppFormatText.GetSourceText().ToString());
+		}
+
+		{
+			const FTextFormat HppFormatText = INVTEXT("{Arg}|hpp(\uC73C\uB85C,\uB85C)");/* 으로/로 */
+			EnsureValidResult(FText::FormatNamed(HppFormatText, TEXT("Arg"), Consonant).ToString(), TEXT("\uC0AC\uB78C\uC73C\uB85C"), TEXT("HppResultConsonant"), HppFormatText.GetSourceText().ToString());
+			EnsureValidResult(FText::FormatNamed(HppFormatText, TEXT("Arg"), ConsonantRieul).ToString(), TEXT("\uC11C\uC6B8\uB85C"), TEXT("HppResultConsonantRieul"), HppFormatText.GetSourceText().ToString());
+			EnsureValidResult(FText::FormatNamed(HppFormatText, TEXT("Arg"), Vowel).ToString(), TEXT("\uC0AC\uC790\uB85C"), TEXT("HppResultVowel"), HppFormatText.GetSourceText().ToString());
+		}
+	}
+
+
+	// Restore original culture
+	I18N.RestoreCultureState(OriginalCultureState);
+}
+
+#if UE_ENABLE_ICU
+
+TEST_CASE_NAMED(FICUSanitizationTest, "System::Core::Misc::ICUSanitization", "[.][EditorContext][ClientContext][EngineFilter]")
+{
+	// Validate culture code sanitization
+	{
+		auto TestCultureCodeSanitization = [](const FString& InCode, const FString& InExpectedCode)
+		{
+			const FString SanitizedCode = ICUUtilities::SanitizeCultureCode(InCode);
+			if (!SanitizedCode.Equals(InExpectedCode, ESearchCase::CaseSensitive))
+			{
+				FAIL_CHECK(FString::Printf(TEXT("SanitizeCultureCode did not produce the expected result (got '%s', expected '%s')"), *SanitizedCode, *InExpectedCode));
+			}
+		};
+
+		TestCultureCodeSanitization(TEXT("en-US"), TEXT("en-US"));
+		TestCultureCodeSanitization(TEXT("en_US_POSIX"), TEXT("en_US_POSIX"));
+		TestCultureCodeSanitization(TEXT("en-US{}%"), TEXT("en-US"));
+		TestCultureCodeSanitization(TEXT("en{}%-US"), TEXT("en-US"));
+	}
+
+	// Validate timezone code sanitization
+	{
+		auto TestTimezoneCodeSanitization = [](const FString& InCode, const FString& InExpectedCode)
+		{
+			const FString SanitizedCode = ICUUtilities::SanitizeTimezoneCode(InCode);
+			if (!SanitizedCode.Equals(InExpectedCode, ESearchCase::CaseSensitive))
+			{
+				FAIL_CHECK(FString::Printf(TEXT("SanitizeTimezoneCode did not produce the expected result (got '%s', expected '%s')"), *SanitizedCode, *InExpectedCode));
+			}
+		};
+
+		TestTimezoneCodeSanitization(TEXT("Etc/Unknown"), TEXT("Etc/Unknown"));
+		TestTimezoneCodeSanitization(TEXT("America/Sao_Paulo"), TEXT("America/Sao_Paulo"));
+		TestTimezoneCodeSanitization(TEXT("America/Sao_Paulo{}%"), TEXT("America/Sao_Paulo"));
+		TestTimezoneCodeSanitization(TEXT("America/Sao{}%_Paulo"), TEXT("America/Sao_Paulo"));
+		TestTimezoneCodeSanitization(TEXT("Antarctica/DumontDUrville"), TEXT("Antarctica/DumontDUrville"));
+		TestTimezoneCodeSanitization(TEXT("Antarctica/DumontDUrville{}%"), TEXT("Antarctica/DumontDUrville"));
+		TestTimezoneCodeSanitization(TEXT("Antarctica/Dumont{}%DUrville"), TEXT("Antarctica/DumontDUrville"));
+		TestTimezoneCodeSanitization(TEXT("Antarctica/DumontD'Urville"), TEXT("Antarctica/DumontDUrville"));
+		TestTimezoneCodeSanitization(TEXT("Antarctica/DumontDUrville_Dumont"), TEXT("Antarctica/DumontDUrville"));
+		TestTimezoneCodeSanitization(TEXT("GMT-8:00"), TEXT("GMT-8:00"));
+		TestTimezoneCodeSanitization(TEXT("GMT-8:00{}%"), TEXT("GMT-8:00"));
+		TestTimezoneCodeSanitization(TEXT("GMT-{}%8:00"), TEXT("GMT-8:00"));
+	}
+
+	// Validate currency code sanitization
+	{
+		auto TestCurrencyCodeSanitization = [](const FString& InCode, const FString& InExpectedCode)
+		{
+			const FString SanitizedCode = ICUUtilities::SanitizeCurrencyCode(InCode);
+			if (!SanitizedCode.Equals(InExpectedCode, ESearchCase::CaseSensitive))
+			{
+				FAIL_CHECK(FString::Printf(TEXT("SanitizeCurrencyCode did not produce the expected result (got '%s', expected '%s')"), *SanitizedCode, *InExpectedCode));
+			}		
+		};
+
+		TestCurrencyCodeSanitization(TEXT("USD"), TEXT("USD"));
+		TestCurrencyCodeSanitization(TEXT("USD{}%"), TEXT("USD"));
+		TestCurrencyCodeSanitization(TEXT("U{}%SD"), TEXT("USD"));
+		TestCurrencyCodeSanitization(TEXT("USDUSD"), TEXT("USD"));
+	}
+
+	// Validate canonization of culture names
+	{
+		auto TestCultureCodeCanonization = [](const FString& InCode, const FString& InExpectedCode)
+		{
+			const FString CanonizedCode = FCulture::GetCanonicalName(InCode);
+			if (!CanonizedCode.Equals(InExpectedCode, ESearchCase::CaseSensitive))
+			{
+				FAIL_CHECK(FString::Printf(TEXT("GetCanonicalName did not produce the expected result (got '%s', expected '%s')"), *CanonizedCode, *InExpectedCode));
+			}
+		};
+
+		// Valid codes
+		TestCultureCodeCanonization(TEXT(""), TEXT("en-US-POSIX"));
+		TestCultureCodeCanonization(TEXT("en"), TEXT("en"));
+		TestCultureCodeCanonization(TEXT("en_US"), TEXT("en-US"));
+		TestCultureCodeCanonization(TEXT("en_US_POSIX"), TEXT("en-US-POSIX"));
+		TestCultureCodeCanonization(TEXT("en_US@POSIX"), TEXT("en-US-POSIX"));
+		TestCultureCodeCanonization(TEXT("en_US.utf8"), TEXT("en-US"));
+		TestCultureCodeCanonization(TEXT("en_US.utf8@posix"), TEXT("en-US-POSIX"));
+		TestCultureCodeCanonization(TEXT("en_IE_PREEURO"), TEXT("en-IE@currency=IEP"));
+		TestCultureCodeCanonization(TEXT("en_IE@CURRENCY=IEP"), TEXT("en-IE@currency=IEP"));
+		TestCultureCodeCanonization(TEXT("fr@collation=phonebook;calendar=islamic-civil"), TEXT("fr@calendar=islamic-civil;collation=phonebook"));
+		TestCultureCodeCanonization(TEXT("sr_Latn_RS_REVISED@currency=USD"), TEXT("sr-Latn-RS-REVISED@currency=USD"));
+		
+		// Invalid codes
+		TestCultureCodeCanonization(TEXT("%%%"), TEXT("en-US-POSIX"));
+		TestCultureCodeCanonization(TEXT("en____US_POSIX"), TEXT("en-US-POSIX"));
+		TestCultureCodeCanonization(TEXT("en_POSIX"), TEXT("en--POSIX"));
+		TestCultureCodeCanonization(TEXT("en__POSIX"), TEXT("en--POSIX"));
+		TestCultureCodeCanonization(TEXT("en_US@wooble=USD"), TEXT("en-US"));
+		TestCultureCodeCanonization(TEXT("fred_wooble_bob_wibble"), TEXT("en-US-POSIX"));
+	}
+}
+
+TEST_CASE_NAMED(FICUTextTest, "System::Core::Misc::ICUText", "[.][EditorContext][ClientContext][EngineFilter]")
+{
+	// Test to make sure that ICUUtilities converts strings correctly
+
+	const FString SourceString(TEXT("This is a test"));
+	const FString SourceString2(TEXT("This is another test"));
+	icu::UnicodeString ICUString;
+	FString ConversionBackStr;
+
+	// ---------------------------------------------------------------------
+
+	ICUUtilities::ConvertString(SourceString, ICUString);
+	if (SourceString.Len() != ICUString.countChar32())
+	{
+		FAIL_CHECK(FString::Printf(TEXT("icu::UnicodeString is the incorrect length (%d; expected %d)."), ICUString.countChar32(), SourceString.Len()));
+	}
+
+	ICUUtilities::ConvertString(ICUString, ConversionBackStr);
+	if (ICUString.length() != ConversionBackStr.Len())
+	{
+		FAIL_CHECK(FString::Printf(TEXT("FString is the incorrect length (%d; expected %d)."), ConversionBackStr.Len(), ICUString.countChar32()));
+	}
+	if (SourceString != ConversionBackStr)
+	{
+		FAIL_CHECK(FString::Printf(TEXT("FString is has the incorrect converted value ('%s'; expected '%s')."), *ConversionBackStr, *SourceString));
+	}
+
+	// ---------------------------------------------------------------------
+
+	ICUUtilities::ConvertString(SourceString2, ICUString);
+	if (SourceString2.Len() != ICUString.countChar32())
+	{
+		FAIL_CHECK(FString::Printf(TEXT("icu::UnicodeString is the incorrect length (%d; expected %d)."), ICUString.countChar32(), SourceString2.Len()));
+	}
+
+	ICUUtilities::ConvertString(ICUString, ConversionBackStr);
+	if (ICUString.length() != ConversionBackStr.Len())
+	{
+		FAIL_CHECK(FString::Printf(TEXT("FString is the incorrect length (%d; expected %d)."), ConversionBackStr.Len(), ICUString.countChar32()));
+	}
+	if (SourceString2 != ConversionBackStr)
+	{
+		FAIL_CHECK(FString::Printf(TEXT("FString is has the incorrect converted value ('%s'; expected '%s')."), *ConversionBackStr, *SourceString2));
+	}
+
+	// ---------------------------------------------------------------------
+
+	ICUUtilities::ConvertString(SourceString, ICUString);
+	if (SourceString.Len() != ICUString.countChar32())
+	{
+		FAIL_CHECK(FString::Printf(TEXT("icu::UnicodeString is the incorrect length (%d; expected %d)."), ICUString.countChar32(), SourceString.Len()));
+	}
+
+	ICUUtilities::ConvertString(ICUString, ConversionBackStr);
+	if (ICUString.length() != ConversionBackStr.Len())
+	{
+		FAIL_CHECK(FString::Printf(TEXT("FString is the incorrect length (%d; expected %d)."), ConversionBackStr.Len(), ICUString.countChar32()));
+	}
+	if (SourceString != ConversionBackStr)
+	{
+		FAIL_CHECK(FString::Printf(TEXT("FString is has the incorrect converted value ('%s'; expected '%s')."), *ConversionBackStr, *SourceString));
+	}
+}
+
+#endif // #if UE_ENABLE_ICU
+
+#undef LOCTEXT_NAMESPACE 
+
+#endif //WITH_TESTS
\ No newline at end of file
