Fix potential race on weaker memory model when tying up a ParallelFor without any events
  - This fixes a lot of TSAN warnings occuring when a thread consumes data produced by a ParallelFor

#rnx

diff --git a/Engine/Source/Runtime/Core/Public/Async/ParallelFor.h b/Engine/Source/Runtime/Core/Public/Async/ParallelFor.h
--- a/Engine/Source/Runtime/Core/Public/Async/ParallelFor.h
+++ b/Engine/Source/Runtime/Core/Public/Async/ParallelFor.h
@@ -336,9 +336,13 @@ namespace ParallelForImpl
 						CallBody(Body, Contexts, WorkerIndex, Index);
 					}
 
-					// we need to decrement IncompleteBatches when processing a Batch because we need to know if we are the last one
+					// We need to decrement IncompleteBatches when processing a Batch because we need to know if we are the last one
 					// so that if the main thread is the last one we can avoid an FEvent call.
-					if (StartIndex < Num && Data->IncompleteBatches.fetch_sub(1, std::memory_order_relaxed) == 1)
+
+					// Memory ordering is also very important here as it is what's making sure memory manipulated
+					// by the parallelfor is properly published before exiting so that it's safe to be read
+					// without other synchronization mechanism.
+					if (StartIndex < Num && Data->IncompleteBatches.fetch_sub(1, std::memory_order_acq_rel) == 1)
 					{
 						if (!bIsMaster)
 						{
