Silence TSAN warning for prefetch in FNameEntryAllocator::Allocate

#rnx

diff --git a/Engine/Source/Runtime/Core/Private/UObject/UnrealNames.cpp b/Engine/Source/Runtime/Core/Private/UObject/UnrealNames.cpp
--- a/Engine/Source/Runtime/Core/Private/UObject/UnrealNames.cpp
+++ b/Engine/Source/Runtime/Core/Private/UObject/UnrealNames.cpp
@@ -489,7 +489,10 @@ public:
 	template<class ScopeLock>
 	FNameEntryHandle Create(FNameStringView Name, TOptional<FNameEntryId> ComparisonId, FNameEntryHeader Header)
 	{
+// Silence the TSAN warning as we don't care for prefetch if we fetch the old or new value
+#if !USING_THREAD_SANITISER
 		FPlatformMisc::Prefetch(Blocks[CurrentBlock]);
+#endif
 		FNameEntryHandle Handle = Allocate<ScopeLock>(FNameEntry::GetDataOffset() + Name.BytesWithoutTerminator());
 		FNameEntry& Entry = Resolve(Handle);
 
