Fix TSAN warning by using relaxed atomic read in thread singleton

#rnx

diff --git a/Engine/Source/Runtime/Core/Private/HAL/ThreadingBase.cpp b/Engine/Source/Runtime/Core/Private/HAL/ThreadingBase.cpp
--- a/Engine/Source/Runtime/Core/Private/HAL/ThreadingBase.cpp
+++ b/Engine/Source/Runtime/Core/Private/HAL/ThreadingBase.cpp
@@ -1338,16 +1338,22 @@ FQueuedThread::Run()
 	FThreadSingletonInitializer
 -----------------------------------------------------------------------------*/
 
-FTlsAutoCleanup* FThreadSingletonInitializer::Get( TFunctionRef<FTlsAutoCleanup*()> CreateInstance, uint32& TlsSlot )
+FTlsAutoCleanup* FThreadSingletonInitializer::Get( TFunctionRef<FTlsAutoCleanup*()> CreateInstance, uint32& InOutTlsSlot )
 {
+	uint32 TlsSlot = (uint32)FPlatformAtomics::AtomicRead_Relaxed((int32*)&InOutTlsSlot);
 	if (TlsSlot == 0xFFFFFFFF)
 	{
 		const uint32 ThisTlsSlot = FPlatformTLS::AllocTlsSlot();
 		check(FPlatformTLS::IsValidTlsSlot(ThisTlsSlot));
-		const uint32 PrevTlsSlot = FPlatformAtomics::InterlockedCompareExchange( (int32*)&TlsSlot, (int32)ThisTlsSlot, 0xFFFFFFFF );
+		const uint32 PrevTlsSlot = FPlatformAtomics::InterlockedCompareExchange( (int32*)&InOutTlsSlot, (int32)ThisTlsSlot, 0xFFFFFFFF );
 		if (PrevTlsSlot != 0xFFFFFFFF)
 		{
 			FPlatformTLS::FreeTlsSlot( ThisTlsSlot );
+			TlsSlot = PrevTlsSlot;
+		}
+		else
+		{
+			TlsSlot = ThisTlsSlot;
 		}
 	}
 	FTlsAutoCleanup* ThreadSingleton = (FTlsAutoCleanup*)FPlatformTLS::GetTlsValue( TlsSlot );
