Fix C++20 compile error. Add user-provided copy and move assignment operators and a move constructor for FStringFormatArg. There is already a user-provided copy constructor which prevented an implicit copy assignment operator from being generated.

diff --git a/Engine/Source/Runtime/Core/Private/Misc/StringFormatter.cpp b/Engine/Source/Runtime/Core/Private/Misc/StringFormatter.cpp
--- a/Engine/Source/Runtime/Core/Private/Misc/StringFormatter.cpp
+++ b/Engine/Source/Runtime/Core/Private/Misc/StringFormatter.cpp
@@ -6,32 +6,55 @@
 
 #define LOCTEXT_NAMESPACE "StringFormatter"
 
-FStringFormatArg::FStringFormatArg( const int32 Value ) : Type(Int), IntValue(Value) {}
-FStringFormatArg::FStringFormatArg( const uint32 Value ) : Type(UInt), UIntValue(Value) {}
-FStringFormatArg::FStringFormatArg( const int64 Value ) : Type(Int), IntValue(Value) {}
-FStringFormatArg::FStringFormatArg( const uint64 Value ) : Type(UInt), UIntValue(Value) {}
-FStringFormatArg::FStringFormatArg( const float Value ) : Type(Double), DoubleValue(Value) {}
-FStringFormatArg::FStringFormatArg( const double Value ) : Type(Double), DoubleValue(Value) {}
-FStringFormatArg::FStringFormatArg( FString Value ) : Type(String), StringValue(MoveTemp(Value)) {}
-FStringFormatArg::FStringFormatArg( FStringView Value ) : Type(String), StringValue(Value) {}
-FStringFormatArg::FStringFormatArg( const ANSICHAR* Value ) : Type(StringLiteralANSI), StringLiteralANSIValue(Value) {}
-FStringFormatArg::FStringFormatArg( const WIDECHAR* Value ) : Type(StringLiteralWIDE), StringLiteralWIDEValue(Value) {}
-FStringFormatArg::FStringFormatArg( const UCS2CHAR* Value ) : Type(StringLiteralUCS2), StringLiteralUCS2Value(Value) {}
-FStringFormatArg::FStringFormatArg( const UTF8CHAR* Value ) : Type(StringLiteralUTF8), StringLiteralUTF8Value(Value) {}
-FStringFormatArg::FStringFormatArg( const FStringFormatArg& RHS )
-{
-	Type = RHS.Type;
-	switch (Type)
+FStringFormatArg::FStringFormatArg(const int32 Value) : Type(Int), IntValue(Value) {}
+FStringFormatArg::FStringFormatArg(const uint32 Value) : Type(UInt), UIntValue(Value) {}
+FStringFormatArg::FStringFormatArg(const int64 Value) : Type(Int), IntValue(Value) {}
+FStringFormatArg::FStringFormatArg(const uint64 Value) : Type(UInt), UIntValue(Value) {}
+FStringFormatArg::FStringFormatArg(const float Value) : Type(Double), DoubleValue(Value) {}
+FStringFormatArg::FStringFormatArg(const double Value) : Type(Double), DoubleValue(Value) {}
+FStringFormatArg::FStringFormatArg(FString Value) : Type(String), StringValue(MoveTemp(Value)) {}
+FStringFormatArg::FStringFormatArg(FStringView Value) : Type(String), StringValue(Value) {}
+FStringFormatArg::FStringFormatArg(const ANSICHAR* Value) : Type(StringLiteralANSI), StringLiteralANSIValue(Value) {}
+FStringFormatArg::FStringFormatArg(const WIDECHAR* Value) : Type(StringLiteralWIDE), StringLiteralWIDEValue(Value) {}
+FStringFormatArg::FStringFormatArg(const UCS2CHAR* Value) : Type(StringLiteralUCS2), StringLiteralUCS2Value(Value) {}
+FStringFormatArg::FStringFormatArg(const UTF8CHAR* Value) : Type(StringLiteralUTF8), StringLiteralUTF8Value(Value) {}
+FStringFormatArg& FStringFormatArg::operator=(const FStringFormatArg& Other)
+{
+	if (this != &Other)
 	{
-		case Int: 				IntValue = RHS.IntValue; break;
-		case UInt: 				UIntValue = RHS.UIntValue; break;
-		case Double: 			IntValue = RHS.IntValue; break;
-		case String: 			StringValue = RHS.StringValue; break;
-		case StringLiteralANSI: StringLiteralANSIValue = RHS.StringLiteralANSIValue; break;
-		case StringLiteralWIDE: StringLiteralWIDEValue = RHS.StringLiteralWIDEValue; break;
-		case StringLiteralUCS2: StringLiteralUCS2Value = RHS.StringLiteralUCS2Value; break;
-		case StringLiteralUTF8: StringLiteralUTF8Value = RHS.StringLiteralUTF8Value; break;
+		Type = Other.Type;
+		switch (Type)
+		{
+			case Int: 				IntValue = Other.IntValue; break;
+			case UInt: 				UIntValue = Other.UIntValue; break;
+			case Double: 			IntValue = Other.IntValue; break;
+			case String: 			StringValue = Other.StringValue; break;
+			case StringLiteralANSI: StringLiteralANSIValue = Other.StringLiteralANSIValue; break;
+			case StringLiteralWIDE: StringLiteralWIDEValue = Other.StringLiteralWIDEValue; break;
+			case StringLiteralUCS2: StringLiteralUCS2Value = Other.StringLiteralUCS2Value; break;
+			case StringLiteralUTF8: StringLiteralUTF8Value = Other.StringLiteralUTF8Value; break;
+		}
+	}
+	return *this;
+}
+FStringFormatArg& FStringFormatArg::operator=(FStringFormatArg&& Other)
+{
+	if (this != &Other)
+	{
+		Type = Other.Type;
+		switch (Type)
+		{
+			case Int: 				IntValue = Other.IntValue; break;
+			case UInt: 				UIntValue = Other.UIntValue; break;
+			case Double: 			IntValue = Other.IntValue; break;
+			case String: 			StringValue = MoveTemp(Other.StringValue); break;
+			case StringLiteralANSI: StringLiteralANSIValue = Other.StringLiteralANSIValue; break;
+			case StringLiteralWIDE: StringLiteralWIDEValue = Other.StringLiteralWIDEValue; break;
+			case StringLiteralUCS2: StringLiteralUCS2Value = Other.StringLiteralUCS2Value; break;
+			case StringLiteralUTF8: StringLiteralUTF8Value = Other.StringLiteralUTF8Value; break;
+		}
 	}
+	return *this;
 }
 
 void AppendToString(const FStringFormatArg& Arg, FString& StringToAppendTo)
diff --git a/Engine/Source/Runtime/Core/Public/Misc/StringFormatArg.h b/Engine/Source/Runtime/Core/Public/Misc/StringFormatArg.h
--- a/Engine/Source/Runtime/Core/Public/Misc/StringFormatArg.h
+++ b/Engine/Source/Runtime/Core/Public/Misc/StringFormatArg.h
@@ -36,24 +36,33 @@ struct CORE_API FStringFormatArg
 	/** Value as an FString */
 	FString StringValue;
 
-	FStringFormatArg( const int32 Value );
-	FStringFormatArg( const uint32 Value );
-	FStringFormatArg( const int64 Value );
-	FStringFormatArg( const uint64 Value );
-	FStringFormatArg( const float Value );
-	FStringFormatArg( const double Value );
-	FStringFormatArg( FString Value );
-	FStringFormatArg( FStringView Value );
-	FStringFormatArg( const ANSICHAR* Value );
-	FStringFormatArg( const WIDECHAR* Value );
-	FStringFormatArg( const UCS2CHAR* Value );
-	FStringFormatArg( const UTF8CHAR* Value );
-
-	/** Copyable */
-	FStringFormatArg( const FStringFormatArg& RHS );
+	/** Not default constructible */
+	FStringFormatArg() = delete;
+	~FStringFormatArg() = default;
+
+	FStringFormatArg(const int32 Value);
+	FStringFormatArg(const uint32 Value);
+	FStringFormatArg(const int64 Value);
+	FStringFormatArg(const uint64 Value);
+	FStringFormatArg(const float Value);
+	FStringFormatArg(const double Value);
+	FStringFormatArg(FString Value);
+	FStringFormatArg(FStringView Value);
+	FStringFormatArg(const ANSICHAR* Value);
+	FStringFormatArg(const WIDECHAR* Value);
+	FStringFormatArg(const UCS2CHAR* Value);
+	FStringFormatArg(const UTF8CHAR* Value);
+
+	FStringFormatArg(const FStringFormatArg& Other)
+	{
+		*this = Other;
+	}
+	FStringFormatArg(FStringFormatArg&& Other)
+	{
+		*this = MoveTemp(Other);
+	}
 	
-private:
+	FStringFormatArg& operator=(const FStringFormatArg& Other);
+	FStringFormatArg& operator=(FStringFormatArg&& Other);
 
-	/** Not default constructible */
-	FStringFormatArg();
 };
