FStringFormatter: replaced uses of auto with explicit types, and move the results to be returned.

diff --git a/Engine/Source/Runtime/Core/Private/Misc/StringFormatter.cpp b/Engine/Source/Runtime/Core/Private/Misc/StringFormatter.cpp
--- a/Engine/Source/Runtime/Core/Private/Misc/StringFormatter.cpp
+++ b/Engine/Source/Runtime/Core/Private/Misc/StringFormatter.cpp
@@ -374,20 +374,20 @@ TValueOrError<FString, FExpressionError> FStringFormatter::FormatInternal(const
 	// This code deliberately tries to reallocate as little as possible
 	FString Formatted;
 	Formatted.Reserve(UE_PTRDIFF_TO_INT32(Tokens.Last().Context.GetTokenEndPos() - InExpression));
-	for (auto& Token : Tokens)
+	for (const FExpressionToken& Token : Tokens)
 	{
-		if (const auto* Literal = Token.Node.Cast<FStringLiteral>())
+		if (const FStringLiteral* Literal = Token.Node.Cast<FStringLiteral>())
 		{
 			Formatted.AppendChars(Literal->String.GetTokenStartPos(), Literal->Len);
 		}
-		else if (auto* Escaped = Token.Node.Cast<FEscapedCharacter>())
+		else if (const FEscapedCharacter* Escaped = Token.Node.Cast<FEscapedCharacter>())
 		{
 			Formatted.AppendChar(Escaped->Character);
 		}
-		else if (const auto* FormatToken = Token.Node.Cast<FFormatSpecifier>())
+		else if (const FFormatSpecifier* FormatToken = Token.Node.Cast<FFormatSpecifier>())
 		{
 			const FStringFormatArg* Arg = nullptr;
-			for (auto& Pair : Args)
+			for (const TPair<FString, FStringFormatArg>& Pair : Args)
 			{
 				if (Pair.Key.Len() == FormatToken->Len && FCString::Strnicmp(FormatToken->Identifier.GetTokenStartPos(), *Pair.Key, FormatToken->Len) == 0)
 				{
@@ -413,7 +413,7 @@ TValueOrError<FString, FExpressionError> FStringFormatter::FormatInternal(const
 		}
 	}
 
-	return MakeValue(Formatted);
+	return MakeValue(MoveTemp(Formatted));
 }
 
 TValueOrError<FString, FExpressionError> FStringFormatter::FormatInternal(const TCHAR* InExpression, const TArray<FStringFormatArg>& Args, bool bStrict) const
@@ -433,17 +433,17 @@ TValueOrError<FString, FExpressionError> FStringFormatter::FormatInternal(const
 	// This code deliberately tries to reallocate as little as possible
 	FString Formatted;
 	Formatted.Reserve(UE_PTRDIFF_TO_INT32(Tokens.Last().Context.GetTokenEndPos() - InExpression));
-	for (auto& Token : Tokens)
+	for (const FExpressionToken& Token : Tokens)
 	{
-		if (const auto* Literal = Token.Node.Cast<FStringLiteral>())
+		if (const FStringLiteral* Literal = Token.Node.Cast<FStringLiteral>())
 		{
 			Formatted.AppendChars(Literal->String.GetTokenStartPos(), Literal->Len);
 		}
-		else if (auto* Escaped = Token.Node.Cast<FEscapedCharacter>())
+		else if (const FEscapedCharacter* Escaped = Token.Node.Cast<FEscapedCharacter>())
 		{
 			Formatted.AppendChar(Escaped->Character);
 		}
-		else if (const auto* IndexToken = Token.Node.Cast<FIndexSpecifier>())
+		else if (const FIndexSpecifier* IndexToken = Token.Node.Cast<FIndexSpecifier>())
 		{
 			if (Args.IsValidIndex(IndexToken->Index))
 			{
@@ -462,7 +462,7 @@ TValueOrError<FString, FExpressionError> FStringFormatter::FormatInternal(const
 		}
 	}
 
-	return MakeValue(Formatted);
+	return MakeValue(MoveTemp(Formatted));
 }
 
 /** Default formatter for string formatting - thread safe since all formatting is const */
