[Insights] Fixed "-trace=stats" not supporting late connect. Time values are incremental and needs to be reset when stats data is discarded / not traced (i.e. when the "stats" channel is turned on/off).
Also reduced the overhead of stats traceing when the "stats" chnanel is not enbaled.

diff --git a/Engine/Source/Runtime/Core/Private/Stats/StatsTrace.cpp b/Engine/Source/Runtime/Core/Private/Stats/StatsTrace.cpp
--- a/Engine/Source/Runtime/Core/Private/Stats/StatsTrace.cpp
+++ b/Engine/Source/Runtime/Core/Private/Stats/StatsTrace.cpp
@@ -49,6 +49,7 @@ public:
 
 	struct FThreadState
 	{
+		uint64 TracedBufferLastCycle;
 		uint64 LastCycle;
 		uint16 BufferSize;
 		uint8 Buffer[MaxBufferSize];
@@ -70,15 +71,26 @@ FStatsTraceInternal::FThreadState* FStatsTraceInternal::InitThreadState()
 {
 	LLM_SCOPE_BYNAME(TEXT("Trace/Stats"));
 	ThreadLocalThreadState = new FThreadState();
-	ThreadLocalThreadState->BufferSize = 0;
+	ThreadLocalThreadState->TracedBufferLastCycle = 0;
 	ThreadLocalThreadState->LastCycle = 0;
+	ThreadLocalThreadState->BufferSize = 0;
 	return ThreadLocalThreadState;
 }
 
 void FStatsTraceInternal::FlushThreadBuffer(FThreadState* ThreadState)
 {
-	UE_TRACE_LOG(Stats, EventBatch, StatsChannel)
-		<< EventBatch.Data(ThreadState->Buffer, ThreadState->BufferSize);
+	if (UE_TRACE_CHANNELEXPR_IS_ENABLED(StatsChannel))
+	{
+		UE_TRACE_LOG(Stats, EventBatch, StatsChannel)
+			<< EventBatch.Data(ThreadState->Buffer, ThreadState->BufferSize);
+
+		ThreadState->TracedBufferLastCycle = ThreadState->LastCycle;
+	}
+	else
+	{
+		// If the current buffer was not actually traced, we need to reset the base timestamp.
+		ThreadState->LastCycle = ThreadState->TracedBufferLastCycle;
+	}
 	ThreadState->BufferSize = 0;
 }
 
@@ -108,7 +120,7 @@ void FStatsTraceInternal::EndEncodeOp(FThreadState* ThreadState, uint8* BufferPt
 
 void FStatsTrace::DeclareStat(const FName& Stat, const ANSICHAR* Name, const TCHAR* Description, const ANSICHAR* Group, bool IsFloatingPoint, bool IsMemory, bool ShouldClearEveryFrame)
 {
-	if (!Stat.IsNone())
+	if (!Stat.IsNone() && UE_TRACE_CHANNELEXPR_IS_ENABLED(StatsChannel))
 	{
 		uint32 NameLen = FCStringAnsi::Strlen(Name);
 		uint32 DescriptionLen = FCString::Strlen(Description);
@@ -127,31 +139,29 @@ void FStatsTrace::DeclareStat(const FName& Stat, const ANSICHAR* Name, const TCH
 
 void FStatsTrace::Increment(const FName& Stat)
 {
-	if (!Stat.IsNone())
+	if (!Stat.IsNone() && UE_TRACE_CHANNELEXPR_IS_ENABLED(StatsChannel))
 	{
 		FStatsTraceInternal::FThreadState* ThreadState;
 		uint8* BufferPtr;
 		FStatsTraceInternal::BeginEncodeOp(Stat, FStatsTraceInternal::Increment, ThreadState, BufferPtr);
-		ThreadState->BufferSize = uint16(BufferPtr - ThreadState->Buffer);
 		FStatsTraceInternal::EndEncodeOp(ThreadState, BufferPtr);
 	}
 }
 
 void FStatsTrace::Decrement(const FName& Stat)
 {
-	if (!Stat.IsNone())
+	if (!Stat.IsNone() && UE_TRACE_CHANNELEXPR_IS_ENABLED(StatsChannel))
 	{
 		FStatsTraceInternal::FThreadState* ThreadState;
 		uint8* BufferPtr;
 		FStatsTraceInternal::BeginEncodeOp(Stat, FStatsTraceInternal::Decrement, ThreadState, BufferPtr);
-		ThreadState->BufferSize = uint16(BufferPtr - ThreadState->Buffer);
 		FStatsTraceInternal::EndEncodeOp(ThreadState, BufferPtr);
 	}
 }
 
 void FStatsTrace::Add(const FName& Stat, int64 Amount)
 {
-	if (!Stat.IsNone())
+	if (!Stat.IsNone() && UE_TRACE_CHANNELEXPR_IS_ENABLED(StatsChannel))
 	{
 		FStatsTraceInternal::FThreadState* ThreadState;
 		uint8* BufferPtr;
@@ -163,7 +173,7 @@ void FStatsTrace::Add(const FName& Stat, int64 Amount)
 
 void FStatsTrace::Add(const FName& Stat, double Amount)
 {
-	if (!Stat.IsNone())
+	if (!Stat.IsNone() && UE_TRACE_CHANNELEXPR_IS_ENABLED(StatsChannel))
 	{
 		FStatsTraceInternal::FThreadState* ThreadState;
 		uint8* BufferPtr;
@@ -176,7 +186,7 @@ void FStatsTrace::Add(const FName& Stat, double Amount)
 
 void FStatsTrace::Set(const FName& Stat, int64 Value)
 {
-	if (!Stat.IsNone())
+	if (!Stat.IsNone() && UE_TRACE_CHANNELEXPR_IS_ENABLED(StatsChannel))
 	{
 		FStatsTraceInternal::FThreadState* ThreadState;
 		uint8* BufferPtr;
@@ -188,7 +198,7 @@ void FStatsTrace::Set(const FName& Stat, int64 Value)
 
 void FStatsTrace::Set(const FName& Stat, double Value)
 {
-	if (!Stat.IsNone())
+	if (!Stat.IsNone() && UE_TRACE_CHANNELEXPR_IS_ENABLED(StatsChannel))
 	{
 		FStatsTraceInternal::FThreadState* ThreadState;
 		uint8* BufferPtr;
