Return a default value for autortfm_current_nest_throw if not using RTFM

diff --git a/Engine/Source/Runtime/Core/Public/AutoRTFM/AutoRTFM.h b/Engine/Source/Runtime/Core/Public/AutoRTFM/AutoRTFM.h
--- a/Engine/Source/Runtime/Core/Public/AutoRTFM/AutoRTFM.h
+++ b/Engine/Source/Runtime/Core/Public/AutoRTFM/AutoRTFM.h
@@ -162,6 +162,7 @@ bool autortfm_current_nest_throw();
 #else
 inline bool autortfm_current_nest_throw()
 {
+	return false;
 }
 #endif
 
