Add some missing return statements

diff --git a/Engine/Source/Runtime/Core/Private/AutoRTFM/API.cpp b/Engine/Source/Runtime/Core/Private/AutoRTFM/API.cpp
--- a/Engine/Source/Runtime/Core/Private/AutoRTFM/API.cpp
+++ b/Engine/Source/Runtime/Core/Private/AutoRTFM/API.cpp
@@ -111,6 +111,7 @@ extern "C" bool autortfm_is_aborting()
 extern "C" bool autortfm_current_nest_throw()
 {
 	FContext::Get()->Throw();
+	return true;
 }
 
 extern "C" void autortfm_abort_if_transactional()
@@ -368,6 +369,7 @@ UE_AUTORTFM_REGISTER_OPEN_FUNCTION(autortfm_is_aborting);
 bool STM_autortfm_current_nest_throw(FContext* Context)
 {
 	Context->Throw();
+	return true;
 }
 UE_AUTORTFM_REGISTER_OPEN_FUNCTION(autortfm_current_nest_throw);
 
