Translate tests from StringFormatter.cpp

diff --git a/Engine/Source/Runtime/Core/Private/Misc/StringFormatter.cpp b/Engine/Source/Runtime/Core/Private/Misc/StringFormatter.cpp
--- a/Engine/Source/Runtime/Core/Private/Misc/StringFormatter.cpp
+++ b/Engine/Source/Runtime/Core/Private/Misc/StringFormatter.cpp
@@ -484,206 +484,4 @@ FString FString::Format(const TCHAR* InFormatString, const FStringFormatOrderedA
 	return DefaultFormatter.Format(InFormatString, InOrderedArguments);
 }
 
-// Don't include this code in shipping or perf test builds
-#if !(UE_BUILD_SHIPPING || UE_BUILD_TEST)
-
-bool TestStrings(FAutomationTestBase& Test, const TCHAR* Format, const TMap<FString, FStringFormatArg>& Arguments, const TCHAR* ExpectedResult)
-{
-	FString Result = FString::Format(Format, Arguments);
-	if (FCString::Strcmp(*Result, ExpectedResult) != 0)
-	{
-		Test.AddError(FText::Format(LOCTEXT("TestError", "FString::Format failed. \"{0}\" != \"{1}\""), FText::FromString(Result), FText::FromString(ExpectedResult)).ToString());
-		return false;
-	}
-	return true;
-}
-
-bool TestStrings(FAutomationTestBase& Test, const TCHAR* Format, const TArray<FStringFormatArg>& Arguments, const TCHAR* ExpectedResult)
-{
-	FString Result = FString::Format(Format, Arguments);
-	if (FCString::Strcmp(*Result, ExpectedResult) != 0)
-	{
-		Test.AddError(FText::Format(LOCTEXT("TestError", "FString::Format failed. \"{0}\" != \"{1}\""), FText::FromString(Result), FText::FromString(ExpectedResult)).ToString());
-		return false;
-	}
-	return true;
-}
-
-IMPLEMENT_SIMPLE_AUTOMATION_TEST(FStringFormattingTestNamedSimple, "System.Core.String Formatting.Simple (Named)", EAutomationTestFlags::EditorContext | EAutomationTestFlags::SmokeFilter)
-bool FStringFormattingTestNamedSimple::RunTest( const FString& Parameters )
-{
-	TMap<FString, FStringFormatArg> Args;
-	Args.Add(TEXT("Argument1"), TEXT("Replacement 1"));
-	Args.Add(TEXT("Argument2"), TEXT("Replacement 2"));
-
-	const TCHAR* Pattern =	TEXT("This is some text containing two arguments, { Argument1 } and {Argument2}.");
-	const TCHAR* Result =	TEXT("This is some text containing two arguments, Replacement 1 and Replacement 2.");
-
-	return TestStrings(*this, Pattern, Args, Result);
-}
-
-IMPLEMENT_SIMPLE_AUTOMATION_TEST(FStringFormattingTestNamedMultiple, "System.Core.String Formatting.Multiple (Named)", EAutomationTestFlags::EditorContext | EAutomationTestFlags::SmokeFilter)
-bool FStringFormattingTestNamedMultiple::RunTest( const FString& Parameters )
-{
-	TMap<FString, FStringFormatArg> Args;
-	Args.Add(TEXT("Argument1"), TEXT("Replacement 1"));
-
-	const TCHAR* Pattern =	TEXT("This is some text containing the same argument, { Argument1 } and {Argument1}.");
-	const TCHAR* Result =	TEXT("This is some text containing the same argument, Replacement 1 and Replacement 1.");
-
-	return TestStrings(*this, Pattern, Args, Result);
-}
-
-IMPLEMENT_SIMPLE_AUTOMATION_TEST(FStringFormattingTestNamedEscaped, "System.Core.String Formatting.Escaped (Named)", EAutomationTestFlags::EditorContext | EAutomationTestFlags::SmokeFilter)
-bool FStringFormattingTestNamedEscaped::RunTest( const FString& Parameters )
-{
-	TMap<FString, FStringFormatArg> Args;
-	Args.Add(TEXT("Argument1"), TEXT("Replacement 1"));
-
-	const TCHAR* Pattern =	TEXT("This is some text containing brace and an arg, `{ Argument1 } and {Argument1}.");
-	const TCHAR* Result =	TEXT("This is some text containing brace and an arg, { Argument1 } and Replacement 1.");
-
-	return TestStrings(*this, Pattern, Args, Result);
-}
-
-IMPLEMENT_SIMPLE_AUTOMATION_TEST(FStringFormattingTestNamedUnbounded, "System.Core.String Formatting.Unbounded (Named)", EAutomationTestFlags::EditorContext | EAutomationTestFlags::SmokeFilter)
-bool FStringFormattingTestNamedUnbounded::RunTest( const FString& Parameters )
-{
-	TMap<FString, FStringFormatArg> Args;
-	Args.Add(TEXT("Argument1"), TEXT("Replacement 1"));
-
-	const TCHAR* Pattern =	TEXT("This is some text containing an unbounded arg, { Argument1");
-	const TCHAR* Result =	TEXT("This is some text containing an unbounded arg, { Argument1");
-
-	return TestStrings(*this, Pattern, Args, Result);
-}
-
-IMPLEMENT_SIMPLE_AUTOMATION_TEST(FStringFormattingTestNamedNonExistent, "System.Core.String Formatting.Non Existent (Named)", EAutomationTestFlags::EditorContext | EAutomationTestFlags::SmokeFilter)
-bool FStringFormattingTestNamedNonExistent::RunTest( const FString& Parameters )
-{
-	TMap<FString, FStringFormatArg> Args;
-	Args.Add(TEXT("Argument1"), TEXT("Replacement 1"));
-
-	const TCHAR* Pattern =	TEXT("This is some text containing a non-existent arg { Argument2 }");
-	const TCHAR* Result =	TEXT("This is some text containing a non-existent arg { Argument2 }");
-
-	return TestStrings(*this, Pattern, Args, Result);
-}
-
-IMPLEMENT_SIMPLE_AUTOMATION_TEST(FStringFormattingTestNamedInvalid, "System.Core.String Formatting.Invalid (Named)", EAutomationTestFlags::EditorContext | EAutomationTestFlags::SmokeFilter)
-bool FStringFormattingTestNamedInvalid::RunTest( const FString& Parameters )
-{
-	TMap<FString, FStringFormatArg> Args;
-	Args.Add(TEXT("Argument1"), TEXT("Replacement 1"));
-
-	const TCHAR* Pattern =	TEXT("This is some text containing an invalid arg {Argument1 1}  { a8f7690-23\\ {} }");
-	const TCHAR* Result =	TEXT("This is some text containing an invalid arg {Argument1 1}  { a8f7690-23\\ {} }");
-
-	return TestStrings(*this, Pattern, Args, Result);
-}
-
-IMPLEMENT_SIMPLE_AUTOMATION_TEST(FStringFormattingTestNamedEmpty, "System.Core.String Formatting.Empty (Named)", EAutomationTestFlags::EditorContext | EAutomationTestFlags::SmokeFilter)
-bool FStringFormattingTestNamedEmpty::RunTest( const FString& Parameters )
-{
-	TMap<FString, FStringFormatArg> Args;
-	Args.Add(TEXT("Argument1"), TEXT("Replacement 1"));
-	Args.Add(TEXT("Argument2"), TEXT("Replacement 2"));
-
-	const TCHAR* Pattern = TEXT("");
-	const TCHAR* Result = TEXT("");
-
-	return TestStrings(*this, Pattern, Args, Result);
-}
-
-IMPLEMENT_SIMPLE_AUTOMATION_TEST(FStringFormattingTestOrderedSimple, "System.Core.String Formatting.Simple (Ordered)", EAutomationTestFlags::EditorContext | EAutomationTestFlags::SmokeFilter)
-bool FStringFormattingTestOrderedSimple::RunTest( const FString& Parameters )
-{
-	TArray<FStringFormatArg> Args;
-	Args.Add(TEXT("Replacement 1"));
-	Args.Add(TEXT("Replacement 2"));
-
-	const TCHAR* Pattern =	TEXT("This is some text containing two arguments, { 0 } and {1}.");
-	const TCHAR* Result =	TEXT("This is some text containing two arguments, Replacement 1 and Replacement 2.");
-
-	return TestStrings(*this, Pattern, Args, Result);
-}
-
-IMPLEMENT_SIMPLE_AUTOMATION_TEST(FStringFormattingTestOrderedMultiple, "System.Core.String Formatting.Multiple (Ordered)", EAutomationTestFlags::EditorContext | EAutomationTestFlags::SmokeFilter)
-bool FStringFormattingTestOrderedMultiple::RunTest( const FString& Parameters )
-{
-	TArray<FStringFormatArg> Args;
-	Args.Add(TEXT("Replacement 1"));
-
-	const TCHAR* Pattern =	TEXT("This is some text containing the same argument, { 0 } and {0}.");
-	const TCHAR* Result =	TEXT("This is some text containing the same argument, Replacement 1 and Replacement 1.");
-
-	return TestStrings(*this, Pattern, Args, Result);
-}
-
-IMPLEMENT_SIMPLE_AUTOMATION_TEST(FStringFormattingTestOrderedEscaped, "System.Core.String Formatting.Escaped (Ordered)", EAutomationTestFlags::EditorContext | EAutomationTestFlags::SmokeFilter)
-bool FStringFormattingTestOrderedEscaped::RunTest( const FString& Parameters )
-{
-	TArray<FStringFormatArg> Args;
-	Args.Add(TEXT("Replacement 1"));
-
-	const TCHAR* Pattern =	TEXT("This is some text containing brace and an arg, `{ 0 } and {0}.");
-	const TCHAR* Result =	TEXT("This is some text containing brace and an arg, { 0 } and Replacement 1.");
-
-	return TestStrings(*this, Pattern, Args, Result);
-}
-
-IMPLEMENT_SIMPLE_AUTOMATION_TEST(FStringFormattingTestOrderedUnbounded, "System.Core.String Formatting.Unbounded (Ordered)", EAutomationTestFlags::EditorContext | EAutomationTestFlags::SmokeFilter)
-bool FStringFormattingTestOrderedUnbounded::RunTest( const FString& Parameters )
-{
-	TArray<FStringFormatArg> Args;
-	Args.Add(TEXT("Replacement 1"));
-
-	const TCHAR* Pattern =	TEXT("This is some text containing an unbounded arg, { 0");
-	const TCHAR* Result =	TEXT("This is some text containing an unbounded arg, { 0");
-
-	return TestStrings(*this, Pattern, Args, Result);
-}
-
-IMPLEMENT_SIMPLE_AUTOMATION_TEST(FStringFormattingTestOrderedNonExistent, "System.Core.String Formatting.Non Existent (Ordered)", EAutomationTestFlags::EditorContext | EAutomationTestFlags::SmokeFilter)
-bool FStringFormattingTestOrderedNonExistent::RunTest( const FString& Parameters )
-{
-	TArray<FStringFormatArg> Args;
-	Args.Add(TEXT("Replacement 1"));
-
-	const TCHAR* Pattern =	TEXT("This is some text containing a non-existent arg { 5 }");
-	const TCHAR* Result =	TEXT("This is some text containing a non-existent arg { 5 }");
-
-	return TestStrings(*this, Pattern, Args, Result);
-}
-
-IMPLEMENT_SIMPLE_AUTOMATION_TEST(FStringFormattingTestOrderedAllTypes, "System.Core.String Formatting.All Types (Ordered)", EAutomationTestFlags::EditorContext | EAutomationTestFlags::SmokeFilter)
-bool FStringFormattingTestOrderedAllTypes::RunTest( const FString& Parameters )
-{
-	const int32 Value1 = 100;
-	const uint32 Value2 = 200;
-	const int64 Value3 = 300;
-	const uint64 Value4 = 400;
-	const float Value5 = 500.0;
-	const double Value6 = 600.0;
-	const FString Value7 = TEXT("Text");
-	const TCHAR* Value8 = TEXT("Text");
-
-	TArray<FStringFormatArg> Args;
-	Args.Add(Value1);
-	Args.Add(Value2);
-	Args.Add(Value3);
-	Args.Add(Value4);
-	Args.Add(Value5);
-	Args.Add(Value6);
-	Args.Add(Value7);
-	Args.Add(Value8);
-
-	const TCHAR* Pattern =	TEXT("{0}, {1}, {2}, {3}, {4}, {5}, {6}, {7}");
-	const TCHAR* Result =	TEXT("100, 200, 300, 400, 500.000000, 600.000000, Text, Text");
-
-	return TestStrings(*this, Pattern, Args, Result);
-}
-
-#endif //!(UE_BUILD_SHIPPING || UE_BUILD_TEST)
-
 #undef LOCTEXT_NAMESPACE
diff --git a/Engine/Source/Runtime/Core/Tests/String/FormatTest.cpp b/Engine/Source/Runtime/Core/Tests/String/FormatTest.cpp
new file mode 100644
--- /dev/null
+++ b/Engine/Source/Runtime/Core/Tests/String/FormatTest.cpp
@@ -0,0 +1,205 @@
+// Copyright Epic Games, Inc. All Rights Reserved.
+
+#if WITH_TESTS
+
+#include "Containers/UnrealString.h"
+#include "Tests/TestHarnessAdapter.h"
+
+TEST_CASE_NAMED(FStringFormattingTestStaticRequirements, "System::Core::String Formatting::Static Requirements", "[Core][String][EditorContext][SmokeFilter]")
+{
+	STATIC_REQUIRE(!std::is_constructible_v<FStringFormatArg>);
+	STATIC_REQUIRE(std::is_copy_constructible_v<FStringFormatArg>);
+	STATIC_REQUIRE(std::is_move_constructible_v<FStringFormatArg>);
+	STATIC_REQUIRE(std::is_copy_assignable_v<FStringFormatArg>);
+	STATIC_REQUIRE(std::is_move_assignable_v<FStringFormatArg>);
+}
+
+TEST_CASE_NAMED(FStringFormattingTestNamedSimple, "System::Core::String Formatting::Simple (Named)", "[Core][String][EditorContext][SmokeFilter]")
+{
+	TMap<FString, FStringFormatArg> Args;
+	Args.Add(TEXT("Argument1"), TEXT("Replacement 1"));
+	Args.Add(TEXT("Argument2"), TEXT("Replacement 2"));
+
+	const TCHAR* Pattern        = TEXT("This is some text containing two arguments, { Argument1 } and {Argument2}.");
+	const TCHAR* ExpectedResult = TEXT("This is some text containing two arguments, Replacement 1 and Replacement 2.");
+
+	const FString ActualResult = FString::Format(Pattern, Args);
+
+	CHECK_MESSAGE(*FString::Printf(TEXT("FString::Format failed. Expected: \"%s\", Actual: \"%s\""), ExpectedResult, *ActualResult), FCString::Strcmp(ExpectedResult, *ActualResult) == 0);
+}
+
+TEST_CASE_NAMED(FStringFormattingTestNamedMultiple, "System::Core::String Formatting::Multiple (Named)", "[Core][String][EditorContext][SmokeFilter]")
+{
+	TMap<FString, FStringFormatArg> Args;
+	Args.Add(TEXT("Argument1"), TEXT("Replacement 1"));
+
+	const TCHAR* Pattern        = TEXT("This is some text containing the same argument, { Argument1 } and {Argument1}.");
+	const TCHAR* ExpectedResult = TEXT("This is some text containing the same argument, Replacement 1 and Replacement 1.");
+
+	const FString ActualResult = FString::Format(Pattern, Args);
+
+	CHECK_MESSAGE(*FString::Printf(TEXT("FString::Format failed. Expected: \"%s\", Actual: \"%s\""), ExpectedResult, *ActualResult), FCString::Strcmp(ExpectedResult, *ActualResult) == 0);
+}
+
+TEST_CASE_NAMED(FStringFormattingTestNamedEscaped, "System::Core::String Formatting::Escaped (Named)", "[Core][String][EditorContext][SmokeFilter]")
+{
+	TMap<FString, FStringFormatArg> Args;
+	Args.Add(TEXT("Argument1"), TEXT("Replacement 1"));
+
+	const TCHAR* Pattern        = TEXT("This is some text containing brace and an arg, `{ Argument1 } and {Argument1}.");
+	const TCHAR* ExpectedResult = TEXT("This is some text containing brace and an arg, { Argument1 } and Replacement 1.");
+
+	const FString ActualResult = FString::Format(Pattern, Args);
+
+	CHECK_MESSAGE(*FString::Printf(TEXT("FString::Format failed. Expected: \"%s\", Actual: \"%s\""), ExpectedResult, *ActualResult), FCString::Strcmp(ExpectedResult, *ActualResult) == 0);
+}
+
+TEST_CASE_NAMED(FStringFormattingTestNamedUnbounded, "System::Core::String Formatting::Unbounded (Named)", "[Core][String][EditorContext][SmokeFilter]")
+{
+	TMap<FString, FStringFormatArg> Args;
+	Args.Add(TEXT("Argument1"), TEXT("Replacement 1"));
+
+	const TCHAR* Pattern        = TEXT("This is some text containing an unbounded arg, { Argument1");
+	const TCHAR* ExpectedResult = TEXT("This is some text containing an unbounded arg, { Argument1");
+
+	const FString ActualResult = FString::Format(Pattern, Args);
+
+	CHECK_MESSAGE(*FString::Printf(TEXT("FString::Format failed. Expected: \"%s\", Actual: \"%s\""), ExpectedResult, *ActualResult), FCString::Strcmp(ExpectedResult, *ActualResult) == 0);
+}
+
+TEST_CASE_NAMED(FStringFormattingTestNamedNonExistent, "System::Core::String Formatting::Non Existent (Named)", "[Core][String][EditorContext][SmokeFilter]")
+{
+	TMap<FString, FStringFormatArg> Args;
+	Args.Add(TEXT("Argument1"), TEXT("Replacement 1"));
+
+	const TCHAR* Pattern        = TEXT("This is some text containing a non-existent arg { Argument2 }");
+	const TCHAR* ExpectedResult = TEXT("This is some text containing a non-existent arg { Argument2 }");
+
+	const FString ActualResult = FString::Format(Pattern, Args);
+
+	CHECK_MESSAGE(*FString::Printf(TEXT("FString::Format failed. Expected: \"%s\", Actual: \"%s\""), ExpectedResult, *ActualResult), FCString::Strcmp(ExpectedResult, *ActualResult) == 0);
+}
+
+TEST_CASE_NAMED(FStringFormattingTestNamedInvalid, "System::Core::String Formatting::Invalid (Named)", "[Core][String][EditorContext][SmokeFilter]")
+{
+	TMap<FString, FStringFormatArg> Args;
+	Args.Add(TEXT("Argument1"), TEXT("Replacement 1"));
+
+	const TCHAR* Pattern        = TEXT("This is some text containing an invalid arg {Argument1 1}  { a8f7690-23\\ {} }");
+	const TCHAR* ExpectedResult = TEXT("This is some text containing an invalid arg {Argument1 1}  { a8f7690-23\\ {} }");
+
+	const FString ActualResult = FString::Format(Pattern, Args);
+
+	CHECK_MESSAGE(*FString::Printf(TEXT("FString::Format failed. Expected: \"%s\", Actual: \"%s\""), ExpectedResult, *ActualResult), FCString::Strcmp(ExpectedResult, *ActualResult) == 0);
+}
+
+TEST_CASE_NAMED(FStringFormattingTestNamedEmpty, "System::Core::String Formatting::Empty (Named)", "[Core][String][EditorContext][SmokeFilter]")
+{
+	TMap<FString, FStringFormatArg> Args;
+	Args.Add(TEXT("Argument1"), TEXT("Replacement 1"));
+	Args.Add(TEXT("Argument2"), TEXT("Replacement 2"));
+
+	const TCHAR* Pattern        = TEXT("");
+	const TCHAR* ExpectedResult = TEXT("");
+
+	const FString ActualResult = FString::Format(Pattern, Args);
+
+	CHECK_MESSAGE(*FString::Printf(TEXT("FString::Format failed. Expected: \"%s\", Actual: \"%s\""), ExpectedResult, *ActualResult), FCString::Strcmp(ExpectedResult, *ActualResult) == 0);
+}
+
+TEST_CASE_NAMED(FStringFormattingTestOrderedSimple, "System::Core::String Formatting::Simple (Ordered)", "[Core][String][EditorContext][SmokeFilter]")
+{
+	TArray<FStringFormatArg> Args;
+	Args.Add(TEXT("Replacement 1"));
+	Args.Add(TEXT("Replacement 2"));
+
+	const TCHAR* Pattern        = TEXT("This is some text containing two arguments, { 0 } and {1}.");
+	const TCHAR* ExpectedResult = TEXT("This is some text containing two arguments, Replacement 1 and Replacement 2.");
+
+	const FString ActualResult = FString::Format(Pattern, Args);
+
+	CHECK_MESSAGE(*FString::Printf(TEXT("FString::Format failed. Expected: \"%s\", Actual: \"%s\""), ExpectedResult, *ActualResult), FCString::Strcmp(ExpectedResult, *ActualResult) == 0);
+}
+
+TEST_CASE_NAMED(FStringFormattingTestOrderedMultiple, "System::Core::String Formatting::Multiple (Ordered)", "[Core][String][EditorContext][SmokeFilter]")
+{
+	TArray<FStringFormatArg> Args;
+	Args.Add(TEXT("Replacement 1"));
+
+	const TCHAR* Pattern        = TEXT("This is some text containing the same argument, { 0 } and {0}.");
+	const TCHAR* ExpectedResult = TEXT("This is some text containing the same argument, Replacement 1 and Replacement 1.");
+
+	const FString ActualResult = FString::Format(Pattern, Args);
+
+	CHECK_MESSAGE(*FString::Printf(TEXT("FString::Format failed. Expected: \"%s\", Actual: \"%s\""), ExpectedResult, *ActualResult), FCString::Strcmp(ExpectedResult, *ActualResult) == 0);
+}
+
+TEST_CASE_NAMED(FStringFormattingTestOrderedEscaped, "System::Core::String Formatting::Escaped (Ordered)", "[Core][String][EditorContext][SmokeFilter]")
+{
+	TArray<FStringFormatArg> Args;
+	Args.Add(TEXT("Replacement 1"));
+
+	const TCHAR* Pattern        = TEXT("This is some text containing brace and an arg, `{ 0 } and {0}.");
+	const TCHAR* ExpectedResult = TEXT("This is some text containing brace and an arg, { 0 } and Replacement 1.");
+
+	const FString ActualResult = FString::Format(Pattern, Args);
+
+	CHECK_MESSAGE(*FString::Printf(TEXT("FString::Format failed. Expected: \"%s\", Actual: \"%s\""), ExpectedResult, *ActualResult), FCString::Strcmp(ExpectedResult, *ActualResult) == 0);
+}
+
+TEST_CASE_NAMED(FStringFormattingTestOrderedUnbounded, "System::Core::String Formatting::Unbounded (Ordered)", "[Core][String][EditorContext][SmokeFilter]")
+{
+	TArray<FStringFormatArg> Args;
+	Args.Add(TEXT("Replacement 1"));
+
+	const TCHAR* Pattern        = TEXT("This is some text containing an unbounded arg, { 0");
+	const TCHAR* ExpectedResult = TEXT("This is some text containing an unbounded arg, { 0");
+
+	const FString ActualResult = FString::Format(Pattern, Args);
+
+	CHECK_MESSAGE(*FString::Printf(TEXT("FString::Format failed. Expected: \"%s\", Actual: \"%s\""), ExpectedResult, *ActualResult), FCString::Strcmp(ExpectedResult, *ActualResult) == 0);
+}
+
+TEST_CASE_NAMED(FStringFormattingTestOrderedNonExistent, "System::Core::String Formatting::Non Existent (Ordered)", "[Core][String][EditorContext][SmokeFilter]")
+{
+	TArray<FStringFormatArg> Args;
+	Args.Add(TEXT("Replacement 1"));
+
+	const TCHAR* Pattern        = TEXT("This is some text containing a non-existent arg { 5 }");
+	const TCHAR* ExpectedResult = TEXT("This is some text containing a non-existent arg { 5 }");
+
+	const FString ActualResult = FString::Format(Pattern, Args);
+
+	CHECK_MESSAGE(*FString::Printf(TEXT("FString::Format failed. Expected: \"%s\", Actual: \"%s\""), ExpectedResult, *ActualResult), FCString::Strcmp(ExpectedResult, *ActualResult) == 0);
+}
+
+TEST_CASE_NAMED(FStringFormattingTestOrderedAllTypes, "System::Core::String Formatting::All Types (Ordered)", "[Core][String][EditorContext][SmokeFilter]")
+{
+	const int32 Value1 = 100;
+	const uint32 Value2 = 200;
+	const int64 Value3 = 300;
+	const uint64 Value4 = 400;
+	const float Value5 = 500.0;
+	const double Value6 = 600.0;
+	const FString Value7 = TEXT("Text");
+	const TCHAR* Value8 = TEXT("Text");
+
+	TArray<FStringFormatArg> Args;
+	Args.Add(Value1);
+	Args.Add(Value2);
+	Args.Add(Value3);
+	Args.Add(Value4);
+	Args.Add(Value5);
+	Args.Add(Value6);
+	Args.Add(Value7);
+	Args.Add(Value8);
+
+	const TCHAR* Pattern        = TEXT("{0}, {1}, {2}, {3}, {4}, {5}, {6}, {7}");
+	const TCHAR* ExpectedResult = TEXT("100, 200, 300, 400, 500.000000, 600.000000, Text, Text");
+
+	const FString ActualResult = FString::Format(Pattern, Args);
+
+	CHECK_MESSAGE(*FString::Printf(TEXT("FString::Format failed. Expected: \"%s\", Actual: \"%s\""), ExpectedResult, *ActualResult), FCString::Strcmp(ExpectedResult, *ActualResult) == 0);
+}
+
+#endif //WITH_TESTS
\ No newline at end of file
