fix for natvis not working for TObjectPtr
symbols need to be publicly available to all dll's for the natvis to find the type

diff --git a/Engine/Source/Runtime/Core/Private/Misc/CoreGlobals.cpp b/Engine/Source/Runtime/Core/Private/Misc/CoreGlobals.cpp
--- a/Engine/Source/Runtime/Core/Private/Misc/CoreGlobals.cpp
+++ b/Engine/Source/Runtime/Core/Private/Misc/CoreGlobals.cpp
@@ -242,15 +242,15 @@ FChunkedFixedUObjectArray* GCoreObjectArrayForDebugVisualizers = nullptr;
 
 namespace UE::CoreUObject::Private
 {
-	struct FStoredObjectPath;
+	struct FStoredObjectPathDebug;
 	struct FObjectHandlePackageDebugData;
 }
-UE::CoreUObject::Private::FStoredObjectPath* GCoreComplexObjectPathDebug = nullptr;
+UE::CoreUObject::Private::FStoredObjectPathDebug* GCoreComplexObjectPathDebug = nullptr;
 UE::CoreUObject::Private::FObjectHandlePackageDebugData* GCoreObjectHandlePackageDebug = nullptr;
 #if PLATFORM_UNIX
 uint8** CORE_API GNameBlocksDebug = FNameDebugVisualizer::GetBlocks();
 FChunkedFixedUObjectArray*& CORE_API GObjectArrayForDebugVisualizers = GCoreObjectArrayForDebugVisualizers;
-UE::CoreUObject::Private::FStoredObjectPath*& GComplexObjectPathDebug = GCoreComplexObjectPathDebug;
+UE::CoreUObject::Private::FStoredObjectPathDebug*& GComplexObjectPathDebug = GCoreComplexObjectPathDebug;
 UE::CoreUObject::Private::FObjectHandlePackageDebugData*& CORE_API GObjectHandlePackageDebug = GCoreObjectHandlePackageDebug;
 #endif
 
diff --git a/Engine/Source/Runtime/Core/Public/CoreGlobals.h b/Engine/Source/Runtime/Core/Public/CoreGlobals.h
--- a/Engine/Source/Runtime/Core/Public/CoreGlobals.h
+++ b/Engine/Source/Runtime/Core/Public/CoreGlobals.h
@@ -758,12 +758,12 @@ extern CORE_API FChunkedFixedUObjectArray* GCoreObjectArrayForDebugVisualizers;
 namespace UE::CoreUObject::Private
 {
 	/** Array to help visualize object paths in the debugger */
-	struct FStoredObjectPath;
+	struct FStoredObjectPathDebug;
 	
 	/** Array to help visualize object handles in the debugger */
 	struct FObjectHandlePackageDebugData;
 }
-extern CORE_API UE::CoreUObject::Private::FStoredObjectPath* GCoreComplexObjectPathDebug;
+extern CORE_API UE::CoreUObject::Private::FStoredObjectPathDebug* GCoreComplexObjectPathDebug;
 extern CORE_API UE::CoreUObject::Private::FObjectHandlePackageDebugData* GCoreObjectHandlePackageDebug;
 
 /** @return True if running cook-on-the-fly. */
diff --git a/Engine/Source/Runtime/Core/Public/Modules/Boilerplate/ModuleBoilerplate.h b/Engine/Source/Runtime/Core/Public/Modules/Boilerplate/ModuleBoilerplate.h
--- a/Engine/Source/Runtime/Core/Public/Modules/Boilerplate/ModuleBoilerplate.h
+++ b/Engine/Source/Runtime/Core/Public/Modules/Boilerplate/ModuleBoilerplate.h
@@ -72,7 +72,7 @@ class FChunkedFixedUObjectArray;
 	#define UE4_VISUALIZERS_HELPERS \
 		uint8** GNameBlocksDebug = FNameDebugVisualizer::GetBlocks(); \
 		FChunkedFixedUObjectArray*& GObjectArrayForDebugVisualizers = GCoreObjectArrayForDebugVisualizers; \
-		UE::CoreUObject::Private::FStoredObjectPath*& GComplexObjectPathDebug = GCoreComplexObjectPathDebug; \
+		UE::CoreUObject::Private::FStoredObjectPathDebug*& GComplexObjectPathDebug = GCoreComplexObjectPathDebug; \
 		UE::CoreUObject::Private::FObjectHandlePackageDebugData*& GObjectHandlePackageDebug = GCoreObjectHandlePackageDebug;
 #endif
 
diff --git a/Engine/Source/Runtime/CoreUObject/Private/UObject/ObjectPathId.cpp b/Engine/Source/Runtime/CoreUObject/Private/UObject/ObjectPathId.cpp
--- a/Engine/Source/Runtime/CoreUObject/Private/UObject/ObjectPathId.cpp
+++ b/Engine/Source/Runtime/CoreUObject/Private/UObject/ObjectPathId.cpp
@@ -193,7 +193,7 @@ namespace UE::CoreUObject::Private
 		GComplexPathHashToId.Add(Key, NewId);
 		Index = NewId;
 		Number = 0;
-		GCoreComplexObjectPathDebug = GComplexPaths.GetData();
+		GCoreComplexObjectPathDebug = (FStoredObjectPathDebug*)GComplexPaths.GetData();
 	}
 
 	FObjectPathId::FObjectPathId(const UObject* Object)
diff --git a/Engine/Source/Runtime/CoreUObject/Public/UObject/ObjectHandle.h b/Engine/Source/Runtime/CoreUObject/Public/UObject/ObjectHandle.h
--- a/Engine/Source/Runtime/CoreUObject/Public/UObject/ObjectHandle.h
+++ b/Engine/Source/Runtime/CoreUObject/Public/UObject/ObjectHandle.h
@@ -311,6 +311,27 @@ namespace UE::CoreUObject::Private
 		FMinimalName ClassName;
 	};
 
+	struct FObjectPathIdDebug
+	{
+		uint32 Index = 0;
+		uint32 Number = 0;
+
+		static constexpr uint32 WeakObjectMask = ~((~0u) >> 1);       //most significant bit
+		static constexpr uint32 SimpleNameMask = WeakObjectMask >> 1; //second most significant bits
+	};
+
+	struct FStoredObjectPathDebug
+	{
+		static constexpr const int32 NumInlineElements = 3;
+		int32 NumElements;
+
+		union
+		{
+			FMinimalName Short[NumInlineElements];
+			FMinimalName* Long;
+		};
+	};
+
 	constexpr uint32 PackageIdShift = 33;
 	constexpr uint32 ObjectPathIdShift = 1;
 	constexpr uint32 ObjectPathIdMask = 0x00FF'FFFF;
