Added FManualResetEvent::IsNotified() and optimized FManualResetEvent::WaitFor(Zero)

#rnx

diff --git a/Engine/Source/Runtime/Core/Public/Async/ManualResetEvent.h b/Engine/Source/Runtime/Core/Public/Async/ManualResetEvent.h
--- a/Engine/Source/Runtime/Core/Public/Async/ManualResetEvent.h
+++ b/Engine/Source/Runtime/Core/Public/Async/ManualResetEvent.h
@@ -19,6 +19,14 @@ public:
 	FManualResetEvent(const FManualResetEvent&) = delete;
 	FManualResetEvent& operator=(const FManualResetEvent&) = delete;
 
+	/**
+	 * Returns true if the event is notified, otherwise false.
+	 */
+	inline bool IsNotified() const
+	{
+		return !!(State.load(std::memory_order_acquire) & IsNotifiedFlag);
+	}
+
 	/**
 	 * Wait until the event is notified.
 	 */
@@ -35,7 +43,7 @@ public:
 	 */
 	inline bool WaitFor(FMonotonicTimeSpan WaitTime)
 	{
-		return WaitUntil(FMonotonicTimePoint::Now() + WaitTime);
+		return WaitTime.IsZero() ? IsNotified() : WaitUntil(FMonotonicTimePoint::Now() + WaitTime);
 	}
 
 	/**
diff --git a/Engine/Source/Runtime/Core/Public/Misc/MonotonicTime.h b/Engine/Source/Runtime/Core/Public/Misc/MonotonicTime.h
--- a/Engine/Source/Runtime/Core/Public/Misc/MonotonicTime.h
+++ b/Engine/Source/Runtime/Core/Public/Misc/MonotonicTime.h
@@ -37,6 +37,7 @@ public:
 	constexpr double ToSeconds() const { return Time; }
 	constexpr double ToMilliseconds() const { return Time * 1000.0; }
 
+	constexpr bool IsZero() const { return *this == Zero(); }
 	constexpr bool IsInfinity() const { return *this == Infinity() || *this == -Infinity(); }
 	bool IsNaN() const { return FPlatformMath::IsNaN(Time); }
 
diff --git a/Engine/Source/Runtime/Core/Tests/Async/ManualResetEventTest.cpp b/Engine/Source/Runtime/Core/Tests/Async/ManualResetEventTest.cpp
--- a/Engine/Source/Runtime/Core/Tests/Async/ManualResetEventTest.cpp
+++ b/Engine/Source/Runtime/Core/Tests/Async/ManualResetEventTest.cpp
@@ -14,21 +14,29 @@ TEST_CASE("Core::Async::ManualResetEvent", "[Core][Async]")
 {
 	FManualResetEvent Event;
 
+	CHECK_FALSE(Event.IsNotified());
 	CHECK_FALSE(Event.WaitFor(FMonotonicTimeSpan::Zero()));
+	CHECK_FALSE(Event.WaitFor(FMonotonicTimeSpan::FromMilliseconds(1.0)));
+	CHECK_FALSE(Event.WaitUntil(FMonotonicTimePoint::Now() + FMonotonicTimeSpan::FromMilliseconds(1.0)));
 	Event.Notify();
 	Event.Wait();
+	CHECK(Event.IsNotified());
 	CHECK(Event.WaitFor(FMonotonicTimeSpan::Zero()));
-	CHECK(Event.WaitUntil(FMonotonicTimePoint::Now() - FMonotonicTimeSpan::FromSeconds(1.0)));
+	CHECK(Event.WaitFor(FMonotonicTimeSpan::FromSeconds(60.0)));
+	CHECK(Event.WaitUntil(FMonotonicTimePoint::Now() + FMonotonicTimeSpan::FromSeconds(60.0)));
 	Event.Reset();
+	CHECK_FALSE(Event.IsNotified());
 	CHECK_FALSE(Event.WaitFor(FMonotonicTimeSpan::Zero()));
+	CHECK_FALSE(Event.WaitFor(FMonotonicTimeSpan::FromMilliseconds(1.0)));
+	CHECK_FALSE(Event.WaitUntil(FMonotonicTimePoint::Now() + FMonotonicTimeSpan::FromMilliseconds(1.0)));
 
 	Tasks::Launch(UE_SOURCE_LOCATION, [&Event] { Event.Notify(); });
 	Event.Wait();
 
 	Event.Notify();
-	CHECK(Event.WaitFor(FMonotonicTimeSpan::Zero()));
+	CHECK(Event.IsNotified());
 	Event.Reset();
-	CHECK_FALSE(Event.WaitFor(FMonotonicTimeSpan::Zero()));
+	CHECK_FALSE(Event.IsNotified());
 }
 
 } // UE
