Expose access to FSetElementId on TMap as well as TSet.
- GetMaxIndex returns max usable index for underlying sparse storage.
- IsValidId returns whether anything is stored at given index.
- Get returns stored element reference. (Requires IsValidId to return true). (Not using operator[] to prevent confusion with access by map key).
- Iterators have GetId function to return current index for later use with Get.
- Added Get functions to Set to match new interface on Map.

diff --git a/Engine/Source/Runtime/Core/Private/Containers/ContainersTest.cpp b/Engine/Source/Runtime/Core/Private/Containers/ContainersTest.cpp
--- a/Engine/Source/Runtime/Core/Private/Containers/ContainersTest.cpp
+++ b/Engine/Source/Runtime/Core/Private/Containers/ContainersTest.cpp
@@ -232,6 +232,13 @@ namespace
 			check(*It  == E);
 			check(*CIt == E);
 
+			FSetElementId Id = It.GetId();
+			FSetElementId CId = It.GetId();
+			check(Cont.IsValidId(Id));
+			check(Cont.IsValidId(CId));
+			check(Cont.Get(Id) == E);
+			check(Cont.Get(CId) == E);
+
 			++It;
 			++CIt;
 		}
diff --git a/Engine/Source/Runtime/Core/Public/Containers/Map.h b/Engine/Source/Runtime/Core/Public/Containers/Map.h
--- a/Engine/Source/Runtime/Core/Public/Containers/Map.h
+++ b/Engine/Source/Runtime/Core/Public/Containers/Map.h
@@ -272,6 +272,34 @@ public:
 		return Pairs.Num();
 	}
 
+	/** @return The max valid index of the elements in the sparse storage. */
+	[[nodiscard]] FORCEINLINE int32 GetMaxIndex() const
+	{
+		return Pairs.GetMaxIndex();
+	}
+
+	/**
+	 * Checks whether an element id is valid.
+	 * @param Id - The element id to check.
+	 * @return true if the element identifier refers to a valid element in this map.
+	 */
+	[[nodiscard]] FORCEINLINE bool IsValidId(FSetElementId Id) const
+	{
+		return Pairs.IsValidId(Id);
+	}
+
+	/** Return a mapped pair by internal identifier. Element must be valid (see @IsValidId). */
+	[[nodiscard]] FORCEINLINE ElementType& Get(FSetElementId Id)
+	{
+		return Pairs[Id];
+	}
+
+	/** Return a mapped pair by internal identifier.  Element must be valid (see @IsValidId).*/
+	[[nodiscard]] FORCEINLINE const ElementType& Get(FSetElementId Id) const
+	{
+		return Pairs[Id];
+	}
+
 	/**
 	 * Get the unique keys contained within this map.
 	 *
@@ -770,6 +798,11 @@ protected:
 		FORCEINLINE ItKeyType&   Key()   const { return PairIt->Key; }
 		FORCEINLINE ItValueType& Value() const { return PairIt->Value; }
 
+		[[nodiscard]] FORCEINLINE FSetElementId GetId() const
+		{
+			return PairIt.GetId();
+		}
+
 		FORCEINLINE PairType& operator* () const { return  *PairIt; }
 		FORCEINLINE PairType* operator->() const { return &*PairIt; }
 
@@ -810,6 +843,10 @@ protected:
 			return !(bool)*this;
 		}
 
+		[[nodiscard]] FORCEINLINE FSetElementId GetId() const
+		{
+			return SetIt.GetId();
+		}
 		FORCEINLINE ItKeyType&   Key() const { return SetIt->Key; }
 		FORCEINLINE ItValueType& Value() const { return SetIt->Value; }
 
@@ -1673,6 +1710,7 @@ public:
 		return Pairs.Num();
 	}
 
+	/** @return The max valid index of the elements in the sparse storage. */
 	int32 GetMaxIndex() const
 	{
 		return Pairs.GetMaxIndex();
diff --git a/Engine/Source/Runtime/Core/Public/Containers/Set.h b/Engine/Source/Runtime/Core/Public/Containers/Set.h
--- a/Engine/Source/Runtime/Core/Public/Containers/Set.h
+++ b/Engine/Source/Runtime/Core/Public/Containers/Set.h
@@ -587,18 +587,30 @@ public:
 			Elements.IsAllocated(Index);
 	}
 
-	/** Accesses the identified element's value. */
+	/** Accesses the identified element's value. Element must be valid (see @IsValidId). */
 	FORCEINLINE ElementType& operator[](FSetElementId Id)
 	{
 		return Elements[Id.Index].Value;
 	}
 
-	/** Accesses the identified element's value. */
+	/** Accesses the identified element's value. Element must be valid (see @IsValidId). */
 	FORCEINLINE const ElementType& operator[](FSetElementId Id) const
 	{
 		return Elements[Id.Index].Value;
 	}
 
+	/** Accesses the identified element's value. Element must be valid (see @IsValidId). */
+	[[nodiscard]] FORCEINLINE ElementType& Get(FSetElementId Id)
+	{
+		return Elements[Id.Index].Value;
+	}
+
+	/** Accesses the identified element's value. Element must be valid (see @IsValidId). */
+	[[nodiscard]] FORCEINLINE const ElementType& Get(FSetElementId Id) const
+	{
+		return Elements[Id.Index].Value;
+	}
+
 	/**
 	 * Adds an element to the set.
 	 *
diff --git a/Engine/Source/Runtime/Core/Public/Containers/SortedMap.h b/Engine/Source/Runtime/Core/Public/Containers/SortedMap.h
--- a/Engine/Source/Runtime/Core/Public/Containers/SortedMap.h
+++ b/Engine/Source/Runtime/Core/Public/Containers/SortedMap.h
@@ -458,6 +458,36 @@ public:
 	FORCEINLINE       ValueType& operator[](KeyConstPointerType Key)       { return this->FindChecked(Key); }
 	FORCEINLINE const ValueType& operator[](KeyConstPointerType Key) const { return this->FindChecked(Key); }
 
+	// Interface functions to match TMap/TSet
+
+	/** @return The max valid index of the elements in the sparse storage. */
+	[[nodiscard]] FORCEINLINE int32 GetMaxIndex() const
+	{
+		return Pairs.Num() - 1;
+	}
+
+	/**
+	 * Checks whether an element id is valid.
+	 * @param Id - The element id to check.
+	 * @return true if the element identifier refers to a valid element in this map.
+	 */
+	[[nodiscard]] FORCEINLINE bool IsValidId(FSetElementId Id) const
+	{
+		return Pairs.IsValidIndex(Id.AsInteger());
+	}
+
+	/** Return a mapped pair by internal identifier. Element must be valid (see @IsValidId). */
+	[[nodiscard]] FORCEINLINE ElementType& Get(FSetElementId Id)
+	{
+		return Pairs[Id.AsInteger()];
+	}
+
+	/** Return a mapped pair by internal identifier.  Element must be valid (see @IsValidId).*/
+	[[nodiscard]] FORCEINLINE const ElementType& Get(FSetElementId Id) const
+	{
+		return Pairs[Id.AsInteger()];
+	}
+
 private:
 	typedef TArray<ElementType, ArrayAllocator> ElementArrayType;
 
@@ -550,6 +580,11 @@ private:
 		FORCEINLINE ItKeyType&   Key()   const { return PairIt->Key; }
 		FORCEINLINE ItValueType& Value() const { return PairIt->Value; }
 
+		[[nodiscard]] FORCEINLINE FSetElementId GetId() const
+		{
+			return FSetElementId::FromInteger(PairIt.GetIndex());
+		}
+
 		FORCEINLINE PairType& operator* () const { return  *PairIt; }
 		FORCEINLINE PairType* operator->() const { return &*PairIt; }
 
@@ -597,6 +632,11 @@ private:
 		FORCEINLINE ItKeyType& Key()   const { return Data[Index].Key; }
 		FORCEINLINE ItValueType& Value() const { return Data[Index].Value; }
 
+		[[nodiscard]] FORCEINLINE FSetElementId GetId() const
+		{
+			return FSetElementId::FromInteger(Index);
+		}
+
 		FORCEINLINE PairType& operator* () const { return  Data[Index]; }
 		FORCEINLINE PairType* operator->() const { return &Data[Index]; }
 
