Made heap functions take any width of index type.
Specifically fixes use of Algo::Sort with TArrayView64
Ushelved from CL by steve.robb

diff --git a/Engine/Source/Runtime/Core/Public/Algo/Impl/BinaryHeap.h b/Engine/Source/Runtime/Core/Public/Algo/Impl/BinaryHeap.h
--- a/Engine/Source/Runtime/Core/Public/Algo/Impl/BinaryHeap.h
+++ b/Engine/Source/Runtime/Core/Public/Algo/Impl/BinaryHeap.h
@@ -6,6 +6,8 @@
 #include "Templates/Invoke.h"
 #include "Templates/ReversePredicate.h"
 
+#include <type_traits>
+
 namespace AlgoImpl
 {
 	/**
@@ -14,7 +16,8 @@ namespace AlgoImpl
 	 * @param	Index Node for which the left child index is to be returned.
 	 * @returns	Index of the left child.
 	 */
-	FORCEINLINE int32 HeapGetLeftChildIndex(int32 Index)
+	template <typename IndexType>
+	FORCEINLINE IndexType HeapGetLeftChildIndex(IndexType Index)
 	{
 		return Index * 2 + 1;
 	}
@@ -25,7 +28,8 @@ namespace AlgoImpl
 	 * @param	Index Node index.
 	 * @returns	true if node is a leaf, false otherwise.
 	 */
-	FORCEINLINE bool HeapIsLeaf(int32 Index, int32 Count)
+	template <typename IndexType>
+	FORCEINLINE bool HeapIsLeaf(IndexType Index, IndexType Count)
 	{
 		return HeapGetLeftChildIndex(Index) >= Count;
 	}
@@ -36,7 +40,8 @@ namespace AlgoImpl
 	 * @param	Index node index.
 	 * @returns	Parent index.
 	 */
-	FORCEINLINE int32 HeapGetParentIndex(int32 Index)
+	template <typename IndexType>
+	FORCEINLINE IndexType HeapGetParentIndex(IndexType Index)
 	{
 		return (Index - 1) / 2;
 	}
@@ -50,15 +55,15 @@ namespace AlgoImpl
 	 * @param	Projection	The projection to apply to the elements.
 	 * @param	Predicate	A binary predicate object used to specify if one element should precede another.
 	 */
-	template <typename RangeValueType, typename ProjectionType, typename PredicateType>
-	FORCEINLINE void HeapSiftDown(RangeValueType* Heap, int32 Index, const int32 Count, const ProjectionType& Projection, const PredicateType& Predicate)
+	template <typename RangeValueType, typename IndexType, typename ProjectionType, typename PredicateType>
+	FORCEINLINE void HeapSiftDown(RangeValueType* Heap, IndexType Index, const IndexType Count, const ProjectionType& Projection, const PredicateType& Predicate)
 	{
 		while (!HeapIsLeaf(Index, Count))
 		{
-			const int32 LeftChildIndex = HeapGetLeftChildIndex(Index);
-			const int32 RightChildIndex = LeftChildIndex + 1;
+			const IndexType LeftChildIndex = HeapGetLeftChildIndex(Index);
+			const IndexType RightChildIndex = LeftChildIndex + 1;
 
-			int32 MinChildIndex = LeftChildIndex;
+			IndexType MinChildIndex = LeftChildIndex;
 			if (RightChildIndex < Count)
 			{
 				MinChildIndex = Predicate( Invoke(Projection, Heap[LeftChildIndex]), Invoke(Projection, Heap[RightChildIndex]) ) ? LeftChildIndex : RightChildIndex;
@@ -85,12 +90,12 @@ namespace AlgoImpl
 	 *
 	 * @return	The new index of the node that was at NodeIndex
 	 */
-	template <class RangeValueType, typename ProjectionType, class PredicateType>
-	FORCEINLINE int32 HeapSiftUp(RangeValueType* Heap, int32 RootIndex, int32 NodeIndex, const ProjectionType& Projection, const PredicateType& Predicate)
+	template <class RangeValueType, typename IndexType, typename ProjectionType, class PredicateType>
+	FORCEINLINE IndexType HeapSiftUp(RangeValueType* Heap, IndexType RootIndex, IndexType NodeIndex, const ProjectionType& Projection, const PredicateType& Predicate)
 	{
 		while (NodeIndex > RootIndex)
 		{
-			int32 ParentIndex = HeapGetParentIndex(NodeIndex);
+			IndexType ParentIndex = HeapGetParentIndex(NodeIndex);
 			if (!Predicate( Invoke(Projection, Heap[NodeIndex]), Invoke(Projection, Heap[ParentIndex]) ))
 			{
 				break;
@@ -112,10 +117,12 @@ namespace AlgoImpl
 	 * @param	Projection	The projection to apply to the elements.
 	 * @param	Predicate	A binary predicate object used to specify if one element should precede another.
 	 */
-	template <typename RangeValueType, typename ProjectionType, typename PredicateType>
-	FORCEINLINE void HeapifyInternal(RangeValueType* First, int32 Num, ProjectionType Projection, PredicateType Predicate)
+	template <typename RangeValueType, typename IndexType, typename ProjectionType, typename PredicateType>
+	FORCEINLINE void HeapifyInternal(RangeValueType* First, IndexType Num, ProjectionType Projection, PredicateType Predicate)
 	{
-		for (int32 Index = HeapGetParentIndex(Num - 1); Index >= 0; Index--)
+		static_assert(std::is_signed_v<IndexType>, "Expected signed index type");
+
+		for (IndexType Index = HeapGetParentIndex(Num - 1); Index >= 0; Index--)
 		{
 			HeapSiftDown(First, Index, Num, Projection, Predicate);
 		}
@@ -129,17 +136,19 @@ namespace AlgoImpl
 	 * @param	Num			the number of elements to sort
 	 * @param	Predicate	predicate class
 	 */
-	template <typename RangeValueType, typename ProjectionType, class PredicateType>
-	void HeapSortInternal(RangeValueType* First, int32 Num, ProjectionType Projection, PredicateType Predicate)
+	template <typename RangeValueType, typename IndexType, typename ProjectionType, class PredicateType>
+	void HeapSortInternal(RangeValueType* First, IndexType Num, ProjectionType Projection, PredicateType Predicate)
 	{
+		static_assert(std::is_signed_v<IndexType>, "Expected signed index type");
+
 		TReversePredicate< PredicateType > ReversePredicateWrapper(Predicate); // Reverse the predicate to build a max-heap instead of a min-heap
 		HeapifyInternal(First, Num, Projection, ReversePredicateWrapper);
 
-		for(int32 Index = Num - 1; Index > 0; Index--)
+		for (IndexType Index = Num - 1; Index > 0; Index--)
 		{
 			Swap(First[0], First[Index]);
 
-			HeapSiftDown(First, 0, Index, Projection, ReversePredicateWrapper);
+			HeapSiftDown(First, (IndexType)0, Index, Projection, ReversePredicateWrapper);
 		}
 	}
 }
diff --git a/Engine/Source/Runtime/Core/Public/Containers/Array.h b/Engine/Source/Runtime/Core/Public/Containers/Array.h
--- a/Engine/Source/Runtime/Core/Public/Containers/Array.h
+++ b/Engine/Source/Runtime/Core/Public/Containers/Array.h
@@ -3177,7 +3177,7 @@ public:
 		// Add at the end, then sift up
 		Add(MoveTempIfPossible(InItem));
 		TDereferenceWrapper<ElementType, PREDICATE_CLASS> PredicateWrapper(Predicate);
-		SizeType Result = AlgoImpl::HeapSiftUp(GetData(), 0, Num() - 1, FIdentityFunctor(), PredicateWrapper);
+		SizeType Result = AlgoImpl::HeapSiftUp(GetData(), (SizeType)0, Num() - 1, FIdentityFunctor(), PredicateWrapper);
 
 		return Result;
 	}
@@ -3199,7 +3199,7 @@ public:
 		// Add at the end, then sift up
 		Add(InItem);
 		TDereferenceWrapper<ElementType, PREDICATE_CLASS> PredicateWrapper(Predicate);
-		SizeType Result = AlgoImpl::HeapSiftUp(GetData(), 0, Num() - 1, FIdentityFunctor(), PredicateWrapper);
+		SizeType Result = AlgoImpl::HeapSiftUp(GetData(), (SizeType)0, Num() - 1, FIdentityFunctor(), PredicateWrapper);
 
 		return Result;
 	}
@@ -3253,7 +3253,7 @@ public:
 		RemoveAtSwap(0, 1, bAllowShrinking);
 
 		TDereferenceWrapper< ElementType, PREDICATE_CLASS> PredicateWrapper(Predicate);
-		AlgoImpl::HeapSiftDown(GetData(), 0, Num(), FIdentityFunctor(), PredicateWrapper);
+		AlgoImpl::HeapSiftDown(GetData(), (SizeType)0, Num(), FIdentityFunctor(), PredicateWrapper);
 	}
 
 	/** 
@@ -3298,7 +3298,7 @@ public:
 	{
 		RemoveAtSwap(0, 1, bAllowShrinking);
 		TDereferenceWrapper< ElementType, PREDICATE_CLASS> PredicateWrapper(Predicate);
-		AlgoImpl::HeapSiftDown(GetData(), 0, Num(), FIdentityFunctor(), PredicateWrapper);
+		AlgoImpl::HeapSiftDown(GetData(), (SizeType)0, Num(), FIdentityFunctor(), PredicateWrapper);
 	}
 
 	/** 
@@ -3357,7 +3357,7 @@ public:
 
 		TDereferenceWrapper< ElementType, PREDICATE_CLASS> PredicateWrapper(Predicate);
 		AlgoImpl::HeapSiftDown(GetData(), Index, Num(), FIdentityFunctor(), PredicateWrapper);
-		AlgoImpl::HeapSiftUp(GetData(), 0, FPlatformMath::Min(Index, Num() - 1), FIdentityFunctor(), PredicateWrapper);
+		AlgoImpl::HeapSiftUp(GetData(), (SizeType)0, FPlatformMath::Min(Index, Num() - 1), FIdentityFunctor(), PredicateWrapper);
 	}
 
 	/**
