Fix malformed checks on array length in tests.

diff --git a/Engine/Source/Runtime/Core/Private/Experimental/Graph/GraphConvert.cpp b/Engine/Source/Runtime/Core/Private/Experimental/Graph/GraphConvert.cpp
--- a/Engine/Source/Runtime/Core/Private/Experimental/Graph/GraphConvert.cpp
+++ b/Engine/Source/Runtime/Core/Private/Experimental/Graph/GraphConvert.cpp
@@ -454,12 +454,12 @@ bool FConvertToGraphTest::RunTest(const FString& Parameters)
 
 	// Check that moving the graph maintains the edge lists 
 	FGraph DestGraph = MoveTemp(Graph);
-	if (!Graph.Buffer.Num() == 0)
+	if (Graph.Buffer.Num() != 0)
 	{
 		AddError(FString::Printf(TEXT("Moving from graph did not empty buffer")));
 		return false;
 	}
-	if (!Graph.EdgeLists.Num() == 0)
+	if (Graph.EdgeLists.Num() != 0)
 	{
 		AddError(FString::Printf(TEXT("Moving from graph did not empty edge lists")));
 		return false ;
