Rename automation test to prevent conflict in monolithic editor build.

diff --git a/Engine/Source/Runtime/Core/Private/Experimental/Graph/GraphConvert.cpp b/Engine/Source/Runtime/Core/Private/Experimental/Graph/GraphConvert.cpp
--- a/Engine/Source/Runtime/Core/Private/Experimental/Graph/GraphConvert.cpp
+++ b/Engine/Source/Runtime/Core/Private/Experimental/Graph/GraphConvert.cpp
@@ -512,9 +512,9 @@ bool FConvertToSingleBufferGraphTest::RunTest(const FString& Parameters)
 	return true;
 }
 
-IMPLEMENT_SIMPLE_AUTOMATION_TEST(FCondensationGraphTest, "System.Core.Graph.CondensationGraph", EAutomationTestFlags::ApplicationContextMask | EAutomationTestFlags::EngineFilter);
+IMPLEMENT_SIMPLE_AUTOMATION_TEST(FCondensationGraphTest_Core, "System.Core.Graph.CondensationGraph", EAutomationTestFlags::ApplicationContextMask | EAutomationTestFlags::EngineFilter);
 
-bool FCondensationGraphTest::RunTest(const FString& Parameters)
+bool FCondensationGraphTest_Core::RunTest(const FString& Parameters)
 {
 	using namespace UE::Graph;
 
