autortfm:
- remove explicit operator new/delete instrumentation. these operators can be instrumented directly by the application by overriding them, otherwise calling one from transactional code will abort when trying to find the instrumented version
- add auto rtfm instrumentation to the FMemory allocator
- remove 'will_deallocate' since it does nothing and required the runtime to be able to get the alloced size which isn't easy to do from the runtime
- Transact now executes the work item when it's compiled away, assuming that we will wrap this around existing code that we still want to execute when non-instrumented

diff --git a/Engine/Source/Programs/AutoRTFMTests/Private/API.cpp b/Engine/Source/Programs/AutoRTFMTests/Private/API.cpp
--- a/Engine/Source/Programs/AutoRTFMTests/Private/API.cpp
+++ b/Engine/Source/Programs/AutoRTFMTests/Private/API.cpp
@@ -596,33 +596,6 @@ TEST_CASE("API.autortfm_did_allocate")
     }
 }
 
-TEST_CASE("API.autortfm_will_deallocate")
-{
-    constexpr unsigned Size = 1024;
-    unsigned BumpAllocator[Size];
-
-    for (unsigned I = 0; I < Size; I++)
-    {
-        BumpAllocator[I] = I;
-    }
-
-    AutoRTFM::Commit([&]
-    {
-        for (unsigned I = 0; I < Size; I++)
-        {
-            BumpAllocator[I] += 1;
-            autortfm_will_deallocate(&BumpAllocator[I], sizeof(unsigned));
-        }
-    });
-
-	// Even though we marked it as going to be deallocated, that doesn't
-    // magically stop our writes going through on commit.
-    for (unsigned I = 0; I < Size; I++)
-    {
-        REQUIRE((I + 1) == BumpAllocator[I]);
-    }
-}
-
 TEST_CASE("API.autortfm_check_consistency_assuming_no_races")
 {
     AutoRTFM::Commit([&]
@@ -1218,33 +1191,6 @@ TEST_CASE("API.DidAllocate")
     }
 }
 
-TEST_CASE("API.WillDeallocate")
-{
-    constexpr unsigned Size = 1024;
-    unsigned BumpAllocator[Size];
-
-    for (unsigned I = 0; I < Size; I++)
-    {
-        BumpAllocator[I] = I;
-    }
-
-    AutoRTFM::Commit([&]
-    {
-        for (unsigned I = 0; I < Size; I++)
-        {
-            BumpAllocator[I] += 1;
-            AutoRTFM::WillDeallocate(&BumpAllocator[I], sizeof(unsigned));
-        }
-    });
-
-    // Even though we marked it as going to be deallocated, that doesn't
-    // magically stop our writes going through on commit.
-    for (unsigned I = 0; I < Size; I++)
-    {
-        REQUIRE((I + 1) == BumpAllocator[I]);
-    }
-}
-
 TEST_CASE("API.CheckConsistencyAssumingNoRaces")
 {
     AutoRTFM::Commit([&]
diff --git a/Engine/Source/Runtime/Core/Private/AutoRTFM/API.cpp b/Engine/Source/Runtime/Core/Private/AutoRTFM/API.cpp
--- a/Engine/Source/Runtime/Core/Private/AutoRTFM/API.cpp
+++ b/Engine/Source/Runtime/Core/Private/AutoRTFM/API.cpp
@@ -222,10 +222,6 @@ extern "C" void* autortfm_did_allocate(void* Ptr, size_t Size)
     return Ptr;
 }
 
-extern "C" void autortfm_will_deallocate(void* Ptr, size_t Size)
-{
-}
-
 extern "C" void autortfm_check_consistency_assuming_no_races()
 {
     if (FContext::IsTransactional())
@@ -479,12 +475,6 @@ void* STM_autortfm_did_allocate(void* Ptr, size_t Size, FContext* Context)
 }
 UE_AUTORTFM_REGISTER_OPEN_FUNCTION(autortfm_did_allocate);
 
-void STM_autortfm_will_deallocate(void* Ptr, size_t Size, FContext* Context)
-{
-    Context->WillDeallocate(Ptr, Size);
-}
-UE_AUTORTFM_REGISTER_OPEN_FUNCTION(autortfm_will_deallocate);
-
 void STM_autortfm_check_consistency_assuming_no_races(FContext* Context)
 {
 }
diff --git a/Engine/Source/Runtime/Core/Private/AutoRTFM/Context.h b/Engine/Source/Runtime/Core/Private/AutoRTFM/Context.h
--- a/Engine/Source/Runtime/Core/Private/AutoRTFM/Context.h
+++ b/Engine/Source/Runtime/Core/Private/AutoRTFM/Context.h
@@ -39,7 +39,6 @@ public:
     void RecordWrite(void* LogicalAddress, size_t Size, bool bIsClosed);
 
     void DidAllocate(void* LogicalAddress, size_t Size);
-    void WillDeallocate(void* LogicalAddress, size_t Size);
 
     // The rest of this is internalish.
     void AbortByLanguageAndThrow();
diff --git a/Engine/Source/Runtime/Core/Private/AutoRTFM/ContextInlines.h b/Engine/Source/Runtime/Core/Private/AutoRTFM/ContextInlines.h
--- a/Engine/Source/Runtime/Core/Private/AutoRTFM/ContextInlines.h
+++ b/Engine/Source/Runtime/Core/Private/AutoRTFM/ContextInlines.h
@@ -18,8 +18,4 @@ inline void FContext::DidAllocate(void* LogicalAddress, size_t Size)
     CurrentTransaction->DidAllocate(LogicalAddress, Size);
 }
 
-inline void FContext::WillDeallocate(void* LogicalAddress, size_t Size)
-{
-}
-
 } // namespace AutoRTFM
diff --git a/Engine/Source/Runtime/Core/Private/AutoRTFM/CppNewSupport.cpp b/Engine/Source/Runtime/Core/Private/AutoRTFM/CppNewSupport.cpp
deleted file mode 100644
diff --git a/Engine/Source/Runtime/Core/Private/AutoRTFM/CppNewSupport.h b/Engine/Source/Runtime/Core/Private/AutoRTFM/CppNewSupport.h
deleted file mode 100644
diff --git a/Engine/Source/Runtime/Core/Private/AutoRTFM/StdLibSupport.cpp b/Engine/Source/Runtime/Core/Private/AutoRTFM/StdLibSupport.cpp
--- a/Engine/Source/Runtime/Core/Private/AutoRTFM/StdLibSupport.cpp
+++ b/Engine/Source/Runtime/Core/Private/AutoRTFM/StdLibSupport.cpp
@@ -11,6 +11,12 @@
 #include <float.h>
 #include <math.h>
 
+#ifdef __APPLE__
+#include <malloc/malloc.h>
+#else
+#include <malloc.h>
+#endif
+
 #ifdef _WIN32
 #define NOMINMAX
 #include <windows.h>
@@ -45,9 +51,6 @@ void STM_free(void* Ptr, FContext* Context)
 {
     if (Ptr)
     {
-		const size_t AllocSize = GetAllocationSize(Ptr);
-        Context->WillDeallocate(Ptr, AllocSize);
-
         Context->GetCurrentTransaction()->DeferUntilCommit([Ptr]
         {
             free(Ptr);
@@ -61,7 +64,13 @@ void* STM_realloc(void* Ptr, size_t Size, FContext* Context)
     void* NewObject = STM_malloc(Size, Context);
     if (Ptr)
     {
-    	const size_t OldSize = GetAllocationSize(Ptr);
+#if defined(__APPLE__)
+		const size_t OldSize = malloc_size(Ptr);
+#elif defined(_WIN32)
+		const size_t OldSize = _msize(Ptr);
+#else
+		const size_t OldSize = malloc_usable_size(Ptr);
+#endif
         MemcpyToNew(NewObject, Ptr, std::min(OldSize, Size), Context);
         STM_free(Ptr, Context);
     }
@@ -148,13 +157,9 @@ FILE* STM___acrt_iob_func(int Index, FContext* Context)
     }
 }
 UE_AUTORTFM_REGISTER_OPEN_FUNCTION(__acrt_iob_func);
-
 UE_AUTORTFM_REGISTER_SELF_FUNCTION(__stdio_common_vfprintf);
-
-UE_AUTORTFM_REGISTER_SELF_FUNCTION(puts);
-
+UE_AUTORTFM_REGISTER_SELF_FUNCTION(__stdio_common_vsprintf);
 UE_AUTORTFM_REGISTER_SELF_FUNCTION(__stdio_common_vswprintf);
-
 UE_AUTORTFM_REGISTER_SELF_FUNCTION(__stdio_common_vfwprintf);
 
 #else // _WIN32 -> so !_WIN32
@@ -172,6 +177,18 @@ UE_AUTORTFM_REGISTER_SELF_FUNCTION(IsDebuggerPresent);
 UE_AUTORTFM_REGISTER_SELF_FUNCTION(EnterCriticalSection);
 UE_AUTORTFM_REGISTER_SELF_FUNCTION(LeaveCriticalSection);
 UE_AUTORTFM_REGISTER_SELF_FUNCTION(QueryPerformanceCounter);
+UE_AUTORTFM_REGISTER_SELF_FUNCTION(QueryPerformanceFrequency);
+UE_AUTORTFM_REGISTER_SELF_FUNCTION(TlsGetValue);
+UE_AUTORTFM_REGISTER_SELF_FUNCTION(GetCurrentThreadId);
+
+// TODO: need to instrument these because in the case of abort, the
+// constructed object will be undone
+extern "C" void* _Init_thread_header();
+extern "C" void* _Init_thread_abort();
+extern "C" void* _Init_thread_footer();
+UE_AUTORTFM_REGISTER_SELF_FUNCTION(_Init_thread_header);
+UE_AUTORTFM_REGISTER_SELF_FUNCTION(_Init_thread_abort);
+UE_AUTORTFM_REGISTER_SELF_FUNCTION(_Init_thread_footer);
 #endif
 
 wchar_t* STM_wcsncpy(wchar_t* Dst, const wchar_t* Src, size_t Count, FContext* Context)
diff --git a/Engine/Source/Runtime/Core/Private/AutoRTFM/Utils.h b/Engine/Source/Runtime/Core/Private/AutoRTFM/Utils.h
--- a/Engine/Source/Runtime/Core/Private/AutoRTFM/Utils.h
+++ b/Engine/Source/Runtime/Core/Private/AutoRTFM/Utils.h
@@ -9,12 +9,6 @@
 #include <stdlib.h>
 #include <string>
 
-#ifdef __APPLE__
-#include <malloc/malloc.h>
-#else
-#include <malloc.h>
-#endif
-
 #define ASSERT(exp) do { \
     if (!(exp)) { \
         fprintf(stderr, "%s:%d:%s: assertion %s failed.\n", __FILE__, __LINE__, __PRETTY_FUNCTION__, #exp); \
@@ -46,17 +40,6 @@ std::string GetFunctionDescription(TReturnType (*FunctionPtr)(TParameterTypes...
     return GetFunctionDescription(reinterpret_cast<void*>(FunctionPtr));
 }
 
-inline size_t GetAllocationSize(void* const Ptr)
-{
-#if defined(__APPLE__)
-    return malloc_size(Ptr);
-#elif defined(_WIN32)
-    return _msize(Ptr);
-#else
-    return malloc_usable_size(Ptr);
-#endif
-}
-
 template<size_t A, size_t B> struct PrettyStaticAssert final
 {
   static_assert(A == B, "Not equal");
diff --git a/Engine/Source/Runtime/Core/Public/AutoRTFM/AutoRTFM.h b/Engine/Source/Runtime/Core/Public/AutoRTFM/AutoRTFM.h
--- a/Engine/Source/Runtime/Core/Public/AutoRTFM/AutoRTFM.h
+++ b/Engine/Source/Runtime/Core/Public/AutoRTFM/AutoRTFM.h
@@ -84,16 +84,18 @@ inline bool autortfm_is_closed(void)
 }
 #endif
 
-// Run the callback in a transaction. None of its effects will be visible until
-// it finishes, assuming it finishes without aborting. It's fine to call this
-// within a transaction. In that case, the nested transaction is only useful if
-// you might choose to abort it (then only the nested transaction aborts).
+// If AutoRTFM is enabled, run the callback in a transaction (otherwise just run the callback).
+// Writes and other effects get instrumented and will be reversed if the transaction aborts.
+// If this begins a nested transaction, the instrumented effects are logged onto the root
+// transaction, so the effects can be reversed later if the root transaction aborts, even
+// if this nested transaction succeeds.
 #if UE_AUTORTFM
 autortfm_result autortfm_transact(void (*work)(void* arg), void* arg);
 #else
 inline autortfm_result autortfm_transact(void (*work)(void* arg), void* arg)
 {
-    return autortfm_aborted_by_language;
+	work(arg);
+    return autortfm_committed;
 }
 #endif
 
@@ -310,14 +312,6 @@ inline void* autortfm_did_allocate(void* ptr, size_t size)
 }
 #endif
 
-// Inform the runtime that we have freed the object (meaning, we deferred it to
-// be freed in abort). This does nothing when called from open code.
-#if UE_AUTORTFM
-void autortfm_will_deallocate(void* ptr, size_t size);
-#else
-inline void autortfm_will_deallocate(void* ptr, size_t size) { }
-#endif
-
 // If running in a transaction, then perform a consistency check of the
 // transaction's read-write set. If possible, this compares the read-write set's
 // expected values with the actual values in global memory. Does nothing when
@@ -531,11 +525,6 @@ inline void* DidAllocate(void* Ptr, size_t Size)
     return autortfm_did_allocate(Ptr, Size);
 }
 
-inline void WillDeallocate(void* Ptr, size_t Size)
-{
-    autortfm_will_deallocate(Ptr, Size);
-}
-
 inline void CheckConsistencyAssumingNoRaces()
 {
     autortfm_check_consistency_assuming_no_races();
diff --git a/Engine/Source/Runtime/Core/Public/HAL/FMemory.inl b/Engine/Source/Runtime/Core/Public/HAL/FMemory.inl
--- a/Engine/Source/Runtime/Core/Public/HAL/FMemory.inl
+++ b/Engine/Source/Runtime/Core/Public/HAL/FMemory.inl
@@ -9,6 +9,7 @@
 #endif
 
 #include "HAL/LowLevelMemTracker.h"
+#include "AutoRTFM/AutoRTFM.h"
 
 struct FMemory;
 struct FScopedMallocTimer;
@@ -16,23 +17,71 @@ struct FScopedMallocTimer;
 FMEMORY_INLINE_FUNCTION_DECORATOR void* FMemory::Malloc(SIZE_T Count, uint32 Alignment)
 {
 	void* Ptr;
-	if (!FMEMORY_INLINE_GMalloc)
-	{
-		Ptr = MallocExternal(Count, Alignment);
-	}
-	else
+
+	// AutoRTFM: For non-transactional code, all of these calls optimize away and the
+	// behavior is the same as it always has been.
+	// For transactional code, we call the allocator in the 'open' as an optimization, so that
+	// we don't end up keeping track of the writes to the allocator's internal data structures.
+	// This is because allocators are already transactional - malloc can be rolled back by
+	// calling free.
+	AutoRTFM::Open([&Ptr, Count, Alignment]
 	{
-		DoGamethreadHook(0);
-		FScopedMallocTimer Timer(0);
-		Ptr = FMEMORY_INLINE_GMalloc->Malloc(Count, Alignment);
-	}
-	// optional tracking of every allocation
-	LLM_IF_ENABLED(FLowLevelMemTracker::Get().OnLowLevelAlloc(ELLMTracker::Default, Ptr, Count, ELLMTag::Untagged, ELLMAllocType::FMalloc));
-	return Ptr;
+		if (!FMEMORY_INLINE_GMalloc)
+		{
+			Ptr = MallocExternal(Count, Alignment);
+		}
+		else
+		{
+			DoGamethreadHook(0);
+			FScopedMallocTimer Timer(0);
+			Ptr = FMEMORY_INLINE_GMalloc->Malloc(Count, Alignment);
+		}
+		// optional tracking of every allocation
+		LLM_IF_ENABLED(FLowLevelMemTracker::Get().OnLowLevelAlloc(ELLMTracker::Default, Ptr, Count, ELLMTag::Untagged, ELLMAllocType::FMalloc));
+	});
+
+	// AutoRTFM: This is a no-op for non-transactional code.
+	// For transactional code, this defers a call to Free if the transaction aborts,
+	// so that rolling back this allocation will end up freeing the memory.
+	AutoRTFM::OpenAbort([Ptr] { Free(Ptr); });
+
+	return AutoRTFM::DidAllocate(Ptr, Count);
 }
 
 FMEMORY_INLINE_FUNCTION_DECORATOR void* FMemory::Realloc(void* Original, SIZE_T Count, uint32 Alignment)
 {
+	if(AutoRTFM::IsClosed())
+	{
+		// AutoRTFM: For transactional code, we have to do a little dance to handle Realloc
+		// properly. We turn realloc into Malloc + Memcpy + Free and never call into the
+		// underlying allocator's realloc implementation. That's required, because if we
+		// were to call into actual realloc, it could end up freeing the old memory. And
+		// if we then aborted our transaction, we would end up rolling back pointers to
+		// point to that old allocation, and it wouldn't be possible to get back that original
+		// allocation at the same address.
+		// 
+		// There is an opportunity here, in that if the original pointer was allocated within
+		// this transaction, then the above doesn't apply since rolling back the transaction
+		// would free that memory and the resulting heap wouldn't point at that memory. So
+		// if we new that the Original pointer was allocated in this transaction, we could
+		// call into the underlying realloc - however we would also have to account for the
+		// malloc deferring a call to free, so we would also have to erase that call to free.
+		void* Ptr = Malloc(Count, Alignment);
+		if (!Ptr)
+		{
+			return nullptr;
+		}
+
+		if (Original)
+		{
+			SIZE_T OriginalCount = GetAllocSize(Original);
+			Memcpy(Ptr, Original, OriginalCount);
+			Free(Original);
+		}
+
+		return Ptr;
+	}
+
 	// optional tracking -- a realloc with an Original pointer of null is equivalent
 	// to malloc() so there's nothing to free
 	LLM_REALLOC_SCOPE(Original);
@@ -66,36 +115,59 @@ FMEMORY_INLINE_FUNCTION_DECORATOR void FMemory::Free(void* Original)
 		return;
 	}
 
-	// optional tracking of every allocation
-	LLM_IF_ENABLED(FLowLevelMemTracker::Get().OnLowLevelFree(ELLMTracker::Default, Original, ELLMAllocType::FMalloc));
-
-	if (!FMEMORY_INLINE_GMalloc)
+	// AutoRTFM: For transactional code, in order to support the transaction 
+	// aborting and needing to 'roll back' the Free, we defer the actual
+	// free until commit time.
+	AutoRTFM::OpenCommit([Original]
 	{
-		FreeExternal(Original);
-		return;
-	}
-	DoGamethreadHook(2);
-	FScopedMallocTimer Timer(2);
-	FMEMORY_INLINE_GMalloc->Free(Original);
+		// optional tracking of every allocation
+		LLM_IF_ENABLED(FLowLevelMemTracker::Get().OnLowLevelFree(ELLMTracker::Default, Original, ELLMAllocType::FMalloc));
+
+		if (!FMEMORY_INLINE_GMalloc)
+		{
+			FreeExternal(Original);
+			return;
+		}
+		DoGamethreadHook(2);
+		FScopedMallocTimer Timer(2);
+		FMEMORY_INLINE_GMalloc->Free(Original);
+	});
 }
 
 FMEMORY_INLINE_FUNCTION_DECORATOR SIZE_T FMemory::GetAllocSize(void* Original)
 {
-	if (!FMEMORY_INLINE_GMalloc)
+	SIZE_T Result;
+	AutoRTFM::Open([&Result, Original]
 	{
-		return GetAllocSizeExternal(Original);
-	}
-
-	SIZE_T Size = 0;
-	return FMEMORY_INLINE_GMalloc->GetAllocationSize(Original, Size) ? Size : 0;
+		if (!FMEMORY_INLINE_GMalloc)
+		{
+			Result = GetAllocSizeExternal(Original);
+		}
+		else
+		{
+			SIZE_T Size = 0;
+			Result = FMEMORY_INLINE_GMalloc->GetAllocationSize(Original, Size) ? Size : 0;
+		}
+	});
+
+	return Result;
 }
 
 FMEMORY_INLINE_FUNCTION_DECORATOR SIZE_T FMemory::QuantizeSize(SIZE_T Count, uint32 Alignment)
 {
-	if (!FMEMORY_INLINE_GMalloc)
+	SIZE_T Result;
+	AutoRTFM::Open([&Result, Count, Alignment]
 	{
-		return Count;
-	}
-	return FMEMORY_INLINE_GMalloc->QuantizeSize(Count, Alignment);
+		if (!FMEMORY_INLINE_GMalloc)
+		{
+			Result = Count;
+		}
+		else
+		{
+			Result = FMEMORY_INLINE_GMalloc->QuantizeSize(Count, Alignment);
+		}
+	});
+
+	return Result;
 }
 
