[Backout] - CLs 24735827, 24746530, 24750104, 24751543, 24752787, 24753773
Original CL Desc
-----------------------------------------------------------------
Added PlatformManualResetEvent

This private abstraction is used to signal a single waiting thread to wake.
This first implementation is only specialized for Microsoft platforms. Other platforms will be specialized later.

diff --git a/Engine/Platforms/Hololens/Source/Runtime/Core/Private/HoloLensPlatformManualResetEvent.h b/Engine/Platforms/Hololens/Source/Runtime/Core/Private/HoloLensPlatformManualResetEvent.h
deleted file mode 100644
diff --git a/Engine/Source/Developer/DerivedDataCache/Private/DerivedDataRequestOwner.cpp b/Engine/Source/Developer/DerivedDataCache/Private/DerivedDataRequestOwner.cpp
--- a/Engine/Source/Developer/DerivedDataCache/Private/DerivedDataRequestOwner.cpp
+++ b/Engine/Source/Developer/DerivedDataCache/Private/DerivedDataRequestOwner.cpp
@@ -2,10 +2,10 @@
 
 #include "DerivedDataRequestOwner.h"
 
-#include "Async/EventCount.h"
 #include "Containers/Array.h"
 #include "DerivedDataRequest.h"
 #include "DerivedDataRequestTypes.h"
+#include "Experimental/Async/LazyEvent.h"
 #include "HAL/CriticalSection.h"
 #include "Misc/ScopeExit.h"
 #include "Misc/ScopeRWLock.h"
@@ -17,6 +17,12 @@
 namespace UE::DerivedData::Private
 {
 
+/** A one-time-use event to work around the lack of condition variables. */
+struct FRequestBarrierEvent : public FRefCountBase
+{
+	FLazyEvent Event{EEventMode::ManualReset};
+};
+
 class FRequestOwnerShared final : public IRequestOwner, public FRequestBase
 {
 public:
@@ -43,7 +49,7 @@ public:
 private:
 	mutable FRWLock Lock;
 	TArray<TRefCountPtr<IRequest>, TInlineAllocator<1>> Requests;
-	FEventCount BarrierEvent;
+	TRefCountPtr<FRequestBarrierEvent> BarrierEvent;
 	uint32 BarrierCount{0};
 	uint32 PriorityBarrierCount{0};
 	EPriority Priority{EPriority::Normal};
@@ -122,7 +128,7 @@ void FRequestOwnerShared::BeginBarrier(ERequestBarrierFlags Flags)
 void FRequestOwnerShared::EndBarrier(ERequestBarrierFlags Flags)
 {
 	ON_SCOPE_EXIT { Release(); };
-	bool bNotifyBarrier = false;
+	TRefCountPtr<FRequestBarrierEvent> LocalBarrierEvent;
 	{
 		FWriteScopeLock WriteLock(Lock);
 		if (EnumHasAnyFlags(Flags, ERequestBarrierFlags::Priority) && --PriorityBarrierCount == 0)
@@ -131,12 +137,13 @@ void FRequestOwnerShared::EndBarrier(ERequestBarrierFlags Flags)
 		}
 		if (--BarrierCount == 0)
 		{
-			bNotifyBarrier = true;
+			LocalBarrierEvent = MoveTemp(BarrierEvent);
+			BarrierEvent = nullptr;
 		}
 	}
-	if (bNotifyBarrier)
+	if (LocalBarrierEvent)
 	{
-		BarrierEvent.Notify();
+		LocalBarrierEvent->Event.Trigger();
 	}
 }
 
@@ -165,7 +172,7 @@ void FRequestOwnerShared::Cancel()
 	for (;;)
 	{
 		TRefCountPtr<IRequest> LocalRequest;
-		FEventCountToken BarrierToken;
+		TRefCountPtr<FRequestBarrierEvent> LocalBarrierEvent;
 
 		{
 			FWriteScopeLock WriteLock(Lock);
@@ -177,7 +184,11 @@ void FRequestOwnerShared::Cancel()
 			}
 			else if (BarrierCount > 0)
 			{
-				BarrierToken = BarrierEvent.PrepareWait();
+				LocalBarrierEvent = BarrierEvent;
+				if (!LocalBarrierEvent)
+				{
+					LocalBarrierEvent = BarrierEvent = new FRequestBarrierEvent;
+				}
 			}
 			else
 			{
@@ -191,7 +202,7 @@ void FRequestOwnerShared::Cancel()
 		}
 		else
 		{
-			BarrierEvent.Wait(BarrierToken);
+			LocalBarrierEvent->Event.Wait();
 		}
 	}
 }
@@ -201,7 +212,7 @@ void FRequestOwnerShared::Wait()
 	for (;;)
 	{
 		TRefCountPtr<IRequest> LocalRequest;
-		FEventCountToken BarrierToken;
+		TRefCountPtr<FRequestBarrierEvent> LocalBarrierEvent;
 
 		{
 			FWriteScopeLock WriteLock(Lock);
@@ -212,7 +223,11 @@ void FRequestOwnerShared::Wait()
 			}
 			else if (BarrierCount > 0)
 			{
-				BarrierToken = BarrierEvent.PrepareWait();
+				LocalBarrierEvent = BarrierEvent;
+				if (!LocalBarrierEvent)
+				{
+					LocalBarrierEvent = BarrierEvent = new FRequestBarrierEvent;
+				}
 			}
 			else
 			{
@@ -226,7 +241,7 @@ void FRequestOwnerShared::Wait()
 		}
 		else
 		{
-			BarrierEvent.Wait(BarrierToken);
+			LocalBarrierEvent->Event.Wait();
 		}
 	}
 }
diff --git a/Engine/Source/Runtime/Core/Private/Android/AndroidPlatformManualResetEvent.h b/Engine/Source/Runtime/Core/Private/Android/AndroidPlatformManualResetEvent.h
deleted file mode 100644
diff --git a/Engine/Source/Runtime/Core/Private/Apple/ApplePlatformManualResetEvent.h b/Engine/Source/Runtime/Core/Private/Apple/ApplePlatformManualResetEvent.h
deleted file mode 100644
diff --git a/Engine/Source/Runtime/Core/Private/Async/ParkingLot.cpp b/Engine/Source/Runtime/Core/Private/Async/ParkingLot.cpp
--- a/Engine/Source/Runtime/Core/Private/Async/ParkingLot.cpp
+++ b/Engine/Source/Runtime/Core/Private/Async/ParkingLot.cpp
@@ -6,11 +6,11 @@
 #include "Async/UniqueLock.h"
 #include "Containers/Array.h"
 #include "Containers/ArrayView.h"
-#include "HAL/PlatformManualResetEvent.h"
 #include "Templates/RefCounting.h"
 #include <atomic>
 
-#define UE_PARKINGLOT_USE_WORDMUTEX 1
+#define UE_PARKINGLOT_USE_WAITONADDRESS PLATFORM_WINDOWS
+#define UE_PARKINGLOT_USE_WORDMUTEX 0
 
 #if UE_PARKINGLOT_USE_WORDMUTEX
 #include "Async/WordMutex.h"
@@ -18,11 +18,18 @@
 #include "HAL/CriticalSection.h"
 #endif
 
+#if UE_PARKINGLOT_USE_WAITONADDRESS
+#include "Microsoft/WindowsHWrapper.h"
+#else
+#include <condition_variable>
+#include <mutex>
+#endif
+
 namespace UE::ParkingLot::Private
 {
 
 #if UE_PARKINGLOT_USE_WORDMUTEX
-using FBucketMutex = FWordMutex;
+using FBucketMutex = UE::Private::FWordMutex;
 #else
 using FBucketMutex = FCriticalSection;
 #endif
@@ -35,9 +42,16 @@ using FBucketMutex = FCriticalSection;
 struct FThread final
 {
 	FThread* Next = nullptr;
+
+#if UE_PARKINGLOT_USE_WAITONADDRESS
+	const void* volatile WaitAddress = nullptr;
+#else
 	const void* WaitAddress = nullptr;
+	std::mutex Lock;
+	std::condition_variable Condition;
+#endif
+
 	uint64 WakeToken = 0;
-	FPlatformManualResetEvent Event;
 	mutable std::atomic<uint32> ReferenceCount = 0;
 
 	inline void AddRef() const
@@ -526,24 +540,58 @@ FWaitState WaitUntil(const void* Address, TFunctionRef<bool()> CanWait, TFunctio
 			return State;
 		}
 		Self.WaitAddress = Address;
-		Self.Event.Reset();
 		Bucket.Enqueue(&Self);
 	}
 
 	// BeforeWait must be invoked after the bucket is unlocked.
 	BeforeWait();
 
-	// Wait until the timeout or until the thread has been dequeued.
-	Self.Event.WaitUntil(WaitTime);
-
-	// WaitAddress is reset when the thread is dequeued.
-	if (!Self.WaitAddress)
+#if UE_PARKINGLOT_USE_WAITONADDRESS
+	for (;;)
 	{
-		State.bDidWake = true;
-		State.WakeToken = Self.WakeToken;
-		Self.WakeToken = 0;
-		return State;
+		const FMonotonicTimeSpan WaitSpan = WaitTime - FMonotonicTimePoint::Now();
+		if (WaitSpan <= FMonotonicTimeSpan::Zero())
+		{
+			break;
+		}
+		const DWORD WaitMs = WaitSpan.IsInfinity() ? INFINITE : DWORD(FMath::CeilToInt64(WaitSpan.ToMilliseconds()));
+		if (!WaitOnAddress(&Self.WaitAddress, &Address, sizeof(Address), WaitMs) && GetLastError() == ERROR_TIMEOUT)
+		{
+			break;
+		}
+		if (!Self.WaitAddress)
+		{
+			State.bDidWake = true;
+			State.WakeToken = Self.WakeToken;
+			Self.WakeToken = 0;
+			return State;
+		}
 	}
+#else
+	{
+		// Wait until the timeout or until the thread has been dequeued.
+		std::unique_lock SelfLock(Self.Lock);
+		const FMonotonicTimeSpan WaitSpan = WaitTime - FMonotonicTimePoint::Now();
+		if (WaitSpan.IsInfinity())
+		{
+			Self.Condition.wait(SelfLock, [&Self] { return !Self.WaitAddress; });
+		}
+		else if (WaitSpan > FMonotonicTimeSpan::Zero())
+		{
+			const int64 WaitMs = FMath::CeilToInt64(WaitSpan.ToMilliseconds());
+			Self.Condition.wait_for(SelfLock, std::chrono::milliseconds(WaitMs), [&Self] { return !Self.WaitAddress; });
+		}
+
+		// WaitAddress is reset when the thread is dequeued.
+		if (!Self.WaitAddress)
+		{
+			State.bDidWake = true;
+			State.WakeToken = Self.WakeToken;
+			Self.WakeToken = 0;
+			return State;
+		}
+	}
+#endif
 
 	// The timeout was reached and the thread needs to dequeue itself.
 	// This can race with a call to wake a thread, which means Self is unsafe to access outside of the lock.
@@ -566,7 +614,15 @@ FWaitState WaitUntil(const void* Address, TFunctionRef<bool()> CanWait, TFunctio
 	// has finished waking this thread by setting its wait address to null.
 	if (!bDequeued)
 	{
-		Self.Event.Wait();
+	#if UE_PARKINGLOT_USE_WAITONADDRESS
+		while (Self.WaitAddress)
+		{
+			WaitOnAddress(&Self.WaitAddress, &Address, sizeof(Address), INFINITE);
+		}
+	#else
+		std::unique_lock SelfLock(Self.Lock);
+		Self.Condition.wait(SelfLock, [&Self] { return !Self.WaitAddress; });
+	#endif
 		State.bDidWake = true;
 		State.WakeToken = Self.WakeToken;
 		Self.WakeToken = 0;
@@ -604,8 +660,16 @@ void WakeOne(const void* Address, TFunctionRef<uint64(FWakeState)> OnWakeState)
 	{
 		checkSlow(WakeThread->WaitAddress == Address);
 		WakeThread->WakeToken = WakeToken;
+	#if UE_PARKINGLOT_USE_WAITONADDRESS
 		WakeThread->WaitAddress = nullptr;
-		WakeThread->Event.Notify();
+		WakeByAddressSingle((void*)&WakeThread->WaitAddress);
+	#else
+		{
+			std::unique_lock WakeThreadLock(WakeThread->Lock);
+			WakeThread->WaitAddress = nullptr;
+		}
+		WakeThread->Condition.notify_one();
+	#endif
 	}
 }
 
@@ -647,8 +711,16 @@ uint32 WakeMultiple(const void* const Address, const uint32 WakeCount)
 	for (FThread* WakeThread : WakeThreads)
 	{
 		checkSlow(WakeThread->WaitAddress == Address);
+	#if UE_PARKINGLOT_USE_WAITONADDRESS
 		WakeThread->WaitAddress = nullptr;
-		WakeThread->Event.Notify();
+		WakeByAddressSingle((void*)&WakeThread->WaitAddress);
+	#else
+		{
+			std::unique_lock WakeThreadLock(WakeThread->Lock);
+			WakeThread->WaitAddress = nullptr;
+		}
+		WakeThread->Condition.notify_one();
+	#endif
 	}
 
 	return uint32(WakeThreads.Num());
diff --git a/Engine/Source/Runtime/Core/Private/Async/WordMutex.cpp b/Engine/Source/Runtime/Core/Private/Async/WordMutex.cpp
deleted file mode 100644
diff --git a/Engine/Source/Runtime/Core/Private/GenericPlatform/GenericPlatformManualResetEvent.h b/Engine/Source/Runtime/Core/Private/GenericPlatform/GenericPlatformManualResetEvent.h
deleted file mode 100644
diff --git a/Engine/Source/Runtime/Core/Private/HAL/PlatformManualResetEvent.h b/Engine/Source/Runtime/Core/Private/HAL/PlatformManualResetEvent.h
deleted file mode 100644
diff --git a/Engine/Source/Runtime/Core/Private/IOS/IOSPlatformManualResetEvent.h b/Engine/Source/Runtime/Core/Private/IOS/IOSPlatformManualResetEvent.h
deleted file mode 100644
diff --git a/Engine/Source/Runtime/Core/Private/Linux/LinuxPlatformManualResetEvent.h b/Engine/Source/Runtime/Core/Private/Linux/LinuxPlatformManualResetEvent.h
deleted file mode 100644
diff --git a/Engine/Source/Runtime/Core/Private/Mac/MacPlatformManualResetEvent.h b/Engine/Source/Runtime/Core/Private/Mac/MacPlatformManualResetEvent.h
deleted file mode 100644
diff --git a/Engine/Source/Runtime/Core/Private/Microsoft/MicrosoftPlatformManualResetEvent.cpp b/Engine/Source/Runtime/Core/Private/Microsoft/MicrosoftPlatformManualResetEvent.cpp
deleted file mode 100644
diff --git a/Engine/Source/Runtime/Core/Private/Microsoft/MicrosoftPlatformManualResetEvent.h b/Engine/Source/Runtime/Core/Private/Microsoft/MicrosoftPlatformManualResetEvent.h
deleted file mode 100644
diff --git a/Engine/Source/Runtime/Core/Private/Misc/OutputDeviceRedirector.cpp b/Engine/Source/Runtime/Core/Private/Misc/OutputDeviceRedirector.cpp
--- a/Engine/Source/Runtime/Core/Private/Misc/OutputDeviceRedirector.cpp
+++ b/Engine/Source/Runtime/Core/Private/Misc/OutputDeviceRedirector.cpp
@@ -2,10 +2,10 @@
 
 #include "Misc/OutputDeviceRedirector.h"
 
-#include "Async/EventCount.h"
 #include "Containers/BitArray.h"
 #include "Containers/DepletableMpmcQueue.h"
 #include "Experimental/ConcurrentLinearAllocator.h"
+#include "HAL/Event.h"
 #include "HAL/PlatformProcess.h"
 #include "HAL/PlatformTime.h"
 #include "HAL/PlatformTLS.h"
@@ -138,11 +138,11 @@ struct FOutputDeviceRedirectorState
 	/** A lock to synchronize access to the thread. */
 	FRWLock ThreadLock;
 
-	/** An event that is notified when the dedicated primary thread is idle. */
-	FEventCount ThreadIdleEvent;
+	/** A queue of events to trigger when the dedicated primary thread is idle. */
+	TDepletableMpscQueue<FEvent*, FOutputDeviceLinearAllocator> ThreadIdleEvents;
 
 	/** An event to wake the dedicated primary thread to process buffered items. */
-	FEventCount ThreadWakeEvent;
+	std::atomic<FEvent*> ThreadWakeEvent = nullptr;
 
 	/** The ID of the thread holding the primary lock. */
 	std::atomic<uint32> LockedThreadId = MAX_uint32;
@@ -153,9 +153,6 @@ struct FOutputDeviceRedirectorState
 	/** The ID of the panic thread, which is only set by Panic(). */
 	std::atomic<uint32> PanicThreadId = MAX_uint32;
 
-	/** Whether a dedicated primary thread has been started. */
-	std::atomic<bool> bThreadStarted = false;
-
 	/** Whether the backlog is enabled. */
 	bool bEnableBacklog = false;
 
@@ -410,8 +407,11 @@ void FOutputDeviceRedirectorState::RemoveOutputDevice(FOutputDevice* OutputDevic
 
 bool FOutputDeviceRedirectorState::TryStartThread()
 {
-	if (FWriteScopeLock ThreadScopeLock(ThreadLock); !bThreadStarted.exchange(true, std::memory_order_relaxed))
+	if (FWriteScopeLock ThreadScopeLock(ThreadLock); !ThreadWakeEvent.load(std::memory_order_relaxed))
 	{
+		FEvent* WakeEvent = FPlatformProcess::GetSynchEventFromPool();
+		WakeEvent->Trigger();
+		ThreadWakeEvent.store(WakeEvent, std::memory_order_release);
 		Thread = FThread(TEXT("OutputDeviceRedirector"), [this] { ThreadLoop(); });
 	}
 	return true;
@@ -419,10 +419,12 @@ bool FOutputDeviceRedirectorState::TryStartThread()
 
 bool FOutputDeviceRedirectorState::TryStopThread()
 {
-	if (FWriteScopeLock ThreadScopeLock(ThreadLock); bThreadStarted.exchange(false, std::memory_order_relaxed))
+	if (FWriteScopeLock ThreadScopeLock(ThreadLock); FEvent* WakeEvent = ThreadWakeEvent.exchange(nullptr, std::memory_order_acquire))
 	{
-		ThreadWakeEvent.Notify();
+		WakeEvent->Trigger();
 		Thread.Join();
+		FOutputDevicesWriteScopeLock Lock(*this);
+		FPlatformProcess::ReturnSynchEventToPool(WakeEvent);
 	}
 	return true;
 }
@@ -436,13 +438,9 @@ void FOutputDeviceRedirectorState::ThreadLoop()
 		PrimaryThreadId.store(ThreadId, std::memory_order_relaxed);
 	}
 
-	for (;;)
+	while (FEvent* WakeEvent = ThreadWakeEvent.load(std::memory_order_acquire))
 	{
-		FEventCountToken Token = ThreadWakeEvent.PrepareWait();
-		if (!bThreadStarted.load(std::memory_order_relaxed))
-		{
-			break;
-		}
+		WakeEvent->Wait();
 		while (!BufferedItems.IsEmpty() && IsPrimaryThread(ThreadId))
 		{
 			if (FOutputDevicesPrimaryScopeLock Lock(*this); Lock.IsLocked())
@@ -450,8 +448,7 @@ void FOutputDeviceRedirectorState::ThreadLoop()
 				FlushBufferedItems();
 			}
 		}
-		ThreadIdleEvent.Notify();
-		ThreadWakeEvent.Wait(Token);
+		ThreadIdleEvents.Deplete([](FEvent* Event) { Event->Trigger(); });
 	}
 }
 
@@ -555,14 +552,17 @@ bool FOutputDeviceRedirector::IsRedirectingTo(FOutputDevice* OutputDevice)
 
 void FOutputDeviceRedirector::FlushThreadedLogs(EOutputDeviceRedirectorFlushOptions Options)
 {
-	if (FReadScopeLock ThreadLock(State->ThreadLock); State->bThreadStarted.load(std::memory_order_relaxed))
+	if (FReadScopeLock ThreadLock(State->ThreadLock); FEvent* WakeEvent = State->ThreadWakeEvent.load(std::memory_order_acquire))
 	{
 		if (!EnumHasAnyFlags(Options, EOutputDeviceRedirectorFlushOptions::Async))
 		{
 			TRACE_CPUPROFILER_EVENT_SCOPE(FOutputDeviceRedirector::FlushThreadedLogs);
-			UE::FEventCountToken Token = State->ThreadIdleEvent.PrepareWait();
-			State->ThreadWakeEvent.Notify();
-			State->ThreadIdleEvent.Wait(Token);
+			FEventRef IdleEvent(EEventMode::ManualReset);
+			if (State->ThreadIdleEvents.EnqueueAndReturnWasEmpty(IdleEvent.Get()))
+			{
+				WakeEvent->Trigger();
+			}
+			IdleEvent->Wait();
 		}
 		return;
 	}
@@ -655,7 +655,10 @@ void FOutputDeviceRedirector::SerializeRecord(const UE::FLogRecord& Record)
 	// Queue the record to serialize to buffered output devices from the primary thread.
 	if (State->BufferedItems.EnqueueAndReturnWasEmpty(Record))
 	{
-		State->ThreadWakeEvent.Notify();
+		if (FEvent* WakeEvent = State->ThreadWakeEvent.load(std::memory_order_acquire))
+		{
+			WakeEvent->Trigger();
+		}
 	}
 }
 
@@ -716,7 +719,10 @@ void FOutputDeviceRedirector::Serialize(const TCHAR* const Data, const ELogVerbo
 	// Queue the line to serialize to buffered output devices from the primary thread.
 	if (State->BufferedItems.EnqueueAndReturnWasEmpty(Data, Category, Verbosity, RealTime))
 	{
-		State->ThreadWakeEvent.Notify();
+		if (FEvent* WakeEvent = State->ThreadWakeEvent.load(std::memory_order_acquire))
+		{
+			WakeEvent->Trigger();
+		}
 	}
 }
 
diff --git a/Engine/Source/Runtime/Core/Private/Unix/UnixPlatformManualResetEvent.h b/Engine/Source/Runtime/Core/Private/Unix/UnixPlatformManualResetEvent.h
deleted file mode 100644
diff --git a/Engine/Source/Runtime/Core/Private/Windows/WindowsPlatformManualResetEvent.h b/Engine/Source/Runtime/Core/Private/Windows/WindowsPlatformManualResetEvent.h
deleted file mode 100644
diff --git a/Engine/Source/Runtime/Core/Public/Async/EventCount.h b/Engine/Source/Runtime/Core/Public/Async/EventCount.h
deleted file mode 100644
diff --git a/Engine/Source/Runtime/Core/Public/Async/WordMutex.h b/Engine/Source/Runtime/Core/Public/Async/WordMutex.h
deleted file mode 100644
