Properly back out / fix build errors as a result of backouts from CL# 24772720
#rnx

diff --git a/Engine/Source/Runtime/Core/Private/Misc/OutputDeviceRedirector.cpp b/Engine/Source/Runtime/Core/Private/Misc/OutputDeviceRedirector.cpp
--- a/Engine/Source/Runtime/Core/Private/Misc/OutputDeviceRedirector.cpp
+++ b/Engine/Source/Runtime/Core/Private/Misc/OutputDeviceRedirector.cpp
@@ -4,6 +4,7 @@
 
 #include "Containers/BitArray.h"
 #include "Containers/DepletableMpmcQueue.h"
+#include "Containers/ConsumeAllMpmcQueue.h"
 #include "Experimental/ConcurrentLinearAllocator.h"
 #include "HAL/Event.h"
 #include "HAL/PlatformProcess.h"
@@ -139,7 +140,7 @@ struct FOutputDeviceRedirectorState
 	FRWLock ThreadLock;
 
 	/** A queue of events to trigger when the dedicated primary thread is idle. */
-	TDepletableMpscQueue<FEvent*, FOutputDeviceLinearAllocator> ThreadIdleEvents;
+	TConsumeAllMpmcQueue<FEvent*, FOutputDeviceLinearAllocator> ThreadIdleEvents;
 
 	/** An event to wake the dedicated primary thread to process buffered items. */
 	std::atomic<FEvent*> ThreadWakeEvent = nullptr;
@@ -448,7 +449,7 @@ void FOutputDeviceRedirectorState::ThreadLoop()
 				FlushBufferedItems();
 			}
 		}
-		ThreadIdleEvents.Deplete([](FEvent* Event) { Event->Trigger(); });
+		ThreadIdleEvents.ConsumeAllLifo([](FEvent* Event) { Event->Trigger(); });
 	}
 }
 
@@ -558,7 +559,7 @@ void FOutputDeviceRedirector::FlushThreadedLogs(EOutputDeviceRedirectorFlushOpti
 		{
 			TRACE_CPUPROFILER_EVENT_SCOPE(FOutputDeviceRedirector::FlushThreadedLogs);
 			FEventRef IdleEvent(EEventMode::ManualReset);
-			if (State->ThreadIdleEvents.EnqueueAndReturnWasEmpty(IdleEvent.Get()))
+			if (State->ThreadIdleEvents.ProduceItem(IdleEvent.Get()) == UE::EConsumeAllMpmcQueueResult::WasEmpty)
 			{
 				WakeEvent->Trigger();
 			}
