fix old style include guard

diff --git a/Engine/Source/Runtime/Core/Public/AutoRTFM/AutoRTFM.h b/Engine/Source/Runtime/Core/Public/AutoRTFM/AutoRTFM.h
--- a/Engine/Source/Runtime/Core/Public/AutoRTFM/AutoRTFM.h
+++ b/Engine/Source/Runtime/Core/Public/AutoRTFM/AutoRTFM.h
@@ -1,7 +1,6 @@
 // Copyright Epic Games, Inc. All Rights Reserved.
 
-#ifndef AUTORTFM_H
-#define AUTORTFM_H
+#pragma once
 
 #if (defined(__AUTORTFM) && __AUTORTFM)
 #define UE_AUTORTFM 1
@@ -553,5 +552,3 @@ struct FRegisterOpenFunction
 
 } // namespace AutoRTFM
 #endif // __cplusplus
-
-#endif /* AUTORTFM_H */
