Windows platform extensions use the default Win64 crash system.

 - new opt-in WINDOWS_USE_WER preprocessor option that will cause exceptions to be passed out of WinMain and picked up by [REDACTED_PERSON].

diff --git a/Engine/Source/Runtime/Core/Private/Windows/WindowsPlatformCrashContext.cpp b/Engine/Source/Runtime/Core/Private/Windows/WindowsPlatformCrashContext.cpp
--- a/Engine/Source/Runtime/Core/Private/Windows/WindowsPlatformCrashContext.cpp
+++ b/Engine/Source/Runtime/Core/Private/Windows/WindowsPlatformCrashContext.cpp
@@ -65,6 +65,10 @@ THIRD_PARTY_INCLUDES_END
 #define HANDLE_ABORT_SIGNALS 1
 #endif
 
+#ifndef WINDOWS_USE_WER
+#define WINDOWS_USE_WER 0
+#endif 
+
 #define CR_CLIENT_MAX_PATH_LEN 265
 #define CR_CLIENT_MAX_ARGS_LEN 256
 
@@ -1145,6 +1149,12 @@ public:
 			::SetUnhandledExceptionFilter(UnhandledStaticInitException);
 		}
 
+#if WINDOWS_USE_WER
+		// disable the internal crash handling when we're using Windows Error Reporting. This means that all exceptions go
+		// to the unhandled exception filter and will be re-thrown out of WinMain to be picked up by Windows Error Reporting
+		FPlatformMisc::SetCrashHandlingType(ECrashHandlingType::Disabled);
+#endif
+
 #if USE_CRASH_REPORTER_MONITOR
 		if (!FPlatformProperties::IsServerOnly())
 		{
@@ -1469,11 +1479,7 @@ private:
 TOptional<FCrashReportingThread> GCrashReportingThread(InPlace);
 #endif
 
-#if WINDOWS_CRASHCONTEXT_WITH_CUSTOM_HANDLERS
-LONG WINAPI DefaultUnhandledStaticInitException(LPEXCEPTION_POINTERS ExceptionInfo)
-#else
 LONG WINAPI UnhandledStaticInitException(LPEXCEPTION_POINTERS ExceptionInfo)
-#endif
 {
 #if !NOINITCRASHREPORTER
 	// If we get an exception during static init we hope that the crash reporting thread
@@ -1550,25 +1556,23 @@ LONG WINAPI UnhandledStaticInitException(LPEXCEPTION_POINTERS ExceptionInfo)
  * The engine hooks itself in the unhandled exception filter. This is the best place to be as it runs after structured exception handlers and
  * it can be easily overriden externally (because there can only be one) to do something else.
  */
-#if WINDOWS_CRASHCONTEXT_WITH_CUSTOM_HANDLERS
-LONG WINAPI DefaultEngineUnhandledExceptionFilter(LPEXCEPTION_POINTERS ExceptionInfo)
-#else
 LONG WINAPI EngineUnhandledExceptionFilter(LPEXCEPTION_POINTERS ExceptionInfo)
-#endif
 {
 	ReportCrash(ExceptionInfo);
+
+#if WINDOWS_USE_WER
+	return EXCEPTION_CONTINUE_SEARCH;
+#else
+
 	GIsCriticalError = true;
 	FPlatformMisc::RequestExit(true, TEXT("WindowsPlatformCrashContext.EngineUnhandledExceptionFilter"));
 
 	return EXCEPTION_CONTINUE_SEARCH; // Not really important, RequestExit() terminates the process just above.
+#endif // WINDOWS_USE_WER
 }
 
 // #CrashReport: 2015-05-28 This should be named EngineCrashHandler
-#if WINDOWS_CRASHCONTEXT_WITH_CUSTOM_HANDLERS
-int32 DefaultReportCrash( LPEXCEPTION_POINTERS ExceptionInfo )
-#else
 int32 ReportCrash( LPEXCEPTION_POINTERS ExceptionInfo )
-#endif
 {
 #if !NOINITCRASHREPORTER
 	// Only create a minidump the first time this function is called.
@@ -1591,7 +1595,11 @@ int32 ReportCrash( LPEXCEPTION_POINTERS ExceptionInfo )
 	}
 #endif
 
+#if WINDOWS_USE_WER
+	return EXCEPTION_CONTINUE_SEARCH;
+#else
 	return EXCEPTION_EXECUTE_HANDLER;
+#endif // WINDOWS_USE_WER
 }
 
 static FCriticalSection EnsureLock;
diff --git a/Engine/Source/Runtime/Core/Public/Windows/WindowsPlatformCrashContext.h b/Engine/Source/Runtime/Core/Public/Windows/WindowsPlatformCrashContext.h
--- a/Engine/Source/Runtime/Core/Public/Windows/WindowsPlatformCrashContext.h
+++ b/Engine/Source/Runtime/Core/Public/Windows/WindowsPlatformCrashContext.h
@@ -5,10 +5,6 @@
 #include "CoreTypes.h"
 #include "Microsoft/MicrosoftPlatformCrashContext.h"
 
-#if !defined(WINDOWS_CRASHCONTEXT_WITH_CUSTOM_HANDLERS)
-	#define WINDOWS_CRASHCONTEXT_WITH_CUSTOM_HANDLERS 0
-#endif
-
 
 struct CORE_API FWindowsPlatformCrashContext : public FMicrosoftPlatformCrashContext
 {
