tracing task priority

diff --git a/Engine/Source/Runtime/Core/Private/Tasks/TaskPrivate.cpp b/Engine/Source/Runtime/Core/Private/Tasks/TaskPrivate.cpp
--- a/Engine/Source/Runtime/Core/Private/Tasks/TaskPrivate.cpp
+++ b/Engine/Source/Runtime/Core/Private/Tasks/TaskPrivate.cpp
@@ -3,9 +3,7 @@
 #include "Tasks/TaskPrivate.h"
 #include "Tasks/Pipe.h"
 
-#if TASKGRAPH_NEW_FRONTEND
 #include "Async/TaskGraphInterfaces.h"
-#endif
 
 namespace UE::Tasks
 {
@@ -445,4 +443,102 @@ namespace UE::Tasks
 			ExtendedPriority = EExtendedTaskPriority::None;
 		}
 	}
+
+	namespace Private
+	{
+		// task priority translation from the old API to the new API
+		void TranslatePriority(ENamedThreads::Type ThreadType, ETaskPriority& OutPriority, EExtendedTaskPriority& OutExtendedPriority)
+		{
+			using namespace UE::Tasks;
+
+#if TASKGRAPH_NEW_FRONTEND
+
+			ENamedThreads::Type ThreadIndex = ENamedThreads::GetThreadIndex(ThreadType);
+			if (ThreadIndex != ENamedThreads::AnyThread)
+			{
+				check(ThreadIndex == ENamedThreads::GameThread || ThreadIndex == ENamedThreads::GetRenderThread() || ThreadIndex == ENamedThreads::RHIThread);
+				EExtendedTaskPriority ConversionMap[] =
+				{
+					EExtendedTaskPriority::RHIThreadNormalPri,
+					EExtendedTaskPriority::GameThreadNormalPri,
+					EExtendedTaskPriority::RenderThreadNormalPri
+				};
+				OutExtendedPriority = ConversionMap[ThreadIndex - ENamedThreads::RHIThread];
+				OutExtendedPriority = (EExtendedTaskPriority)((int32)OutExtendedPriority + (ENamedThreads::GetTaskPriority(ThreadType) != ENamedThreads::NormalTaskPriority ? 1 : 0));
+				OutExtendedPriority = (EExtendedTaskPriority)((int32)OutExtendedPriority + (ENamedThreads::GetQueueIndex(ThreadType) != ENamedThreads::MainQueue ? 2 : 0));
+				OutPriority = ETaskPriority::Count;
+			}
+			else
+			{
+				OutExtendedPriority = EExtendedTaskPriority::None;
+				uint32 ThreadPriority = GetThreadPriorityIndex(ThreadType);
+				check(ThreadPriority < uint32(ENamedThreads::NumThreadPriorities));
+				ETaskPriority ConversionMap[int(ENamedThreads::NumThreadPriorities)] = { ETaskPriority::Normal, ETaskPriority::High, ETaskPriority::BackgroundNormal };
+				OutPriority = ConversionMap[ThreadPriority];
+			}
+
+			if (OutPriority == ETaskPriority::BackgroundNormal && GetTaskPriority(ThreadType))
+			{
+				OutPriority = ETaskPriority::BackgroundHigh;
+			}
+
+#else
+			uint32 ThreadPriority = GetThreadPriorityIndex(ThreadType);
+			check(ThreadPriority < uint32(ENamedThreads::NumThreadPriorities));
+			ETaskPriority Conversion[int(ENamedThreads::NumThreadPriorities)] = { ETaskPriority::Normal, ETaskPriority::High, ETaskPriority::BackgroundNormal };
+			OutPriority = Conversion[ThreadPriority];
+
+			if (OutPriority == ETaskPriority::BackgroundNormal && GetTaskPriority(ThreadType))
+			{
+				OutPriority = ETaskPriority::BackgroundHigh;
+			}
+#endif
+		}
+
+#if TASKGRAPH_NEW_FRONTEND
+
+		// task priority translation from the new API to the old API
+		ENamedThreads::Type TranslatePriority(EExtendedTaskPriority Priority)
+		{
+			checkf(Priority >= EExtendedTaskPriority::GameThreadNormalPri && Priority < EExtendedTaskPriority::Count, TEXT("only named threads can call this method: %d"), Priority);
+
+			int32 ConversionMap[] =
+			{
+					ENamedThreads::GameThread, // GameThreadNormalPri
+					ENamedThreads::GameThread | ENamedThreads::HighTaskPriority, // GameThreadHiPri
+					ENamedThreads::GameThread | ENamedThreads::LocalQueue, // GameThreadNormalPriLocalQueue
+					ENamedThreads::GameThread | ENamedThreads::HighTaskPriority | ENamedThreads::LocalQueue, // GameThreadHiPriLocalQueue
+
+					ENamedThreads::GetRenderThread(), // RenderThreadNormalPri
+					ENamedThreads::GetRenderThread() | ENamedThreads::HighTaskPriority, // RenderThreadHiPri
+					ENamedThreads::GetRenderThread() | ENamedThreads::LocalQueue, // RenderThreadNormalPriLocalQueue
+					ENamedThreads::GetRenderThread() | ENamedThreads::HighTaskPriority | ENamedThreads::LocalQueue, // RenderThreadHiPriLocalQueue
+
+					ENamedThreads::RHIThread, // RHIThreadNormalPri
+					ENamedThreads::RHIThread | ENamedThreads::HighTaskPriority, // RHIThreadHiPri
+					ENamedThreads::RHIThread | ENamedThreads::LocalQueue, // RHIThreadNormalPriLocalQueue
+					ENamedThreads::RHIThread | ENamedThreads::HighTaskPriority | ENamedThreads::LocalQueue // RHIThreadHiPriLocalQueue
+			};
+
+			return (ENamedThreads::Type)ConversionMap[(int32)Priority - (int32)EExtendedTaskPriority::GameThreadNormalPri];
+		}
+
+#endif
+
+		ENamedThreads::Type TranslatePriority(ETaskPriority Priority)
+		{
+			checkSlow(Priority < ETaskPriority::Count);
+
+			ENamedThreads::Type ConversionMap[] =
+			{
+				ENamedThreads::AnyHiPriThreadNormalTask,
+				ENamedThreads::AnyNormalThreadNormalTask,
+				ENamedThreads::AnyBackgroundHiPriTask,
+				ENamedThreads::AnyBackgroundThreadNormalTask,
+				ENamedThreads::AnyBackgroundThreadNormalTask // same as above
+			};
+
+			return ConversionMap[(int32)Priority];
+		}
+	}
 }
diff --git a/Engine/Source/Runtime/Core/Public/Async/TaskGraphInterfaces.h b/Engine/Source/Runtime/Core/Public/Async/TaskGraphInterfaces.h
--- a/Engine/Source/Runtime/Core/Public/Async/TaskGraphInterfaces.h
+++ b/Engine/Source/Runtime/Core/Public/Async/TaskGraphInterfaces.h
@@ -556,71 +556,9 @@ public:
 
 	ENamedThreads::Type GetThreadToExecuteOn() const
 	{
-		return TranslatePriority(GetExtendedPriority());
+		return UE::Tasks::Private::TranslatePriority(GetExtendedPriority());
 	}
 
-	// task priority translation from the old API to the new API
-	static void TranslatePriority(ENamedThreads::Type ThreadType, UE::Tasks::ETaskPriority& OutPriority, UE::Tasks::EExtendedTaskPriority& OutExtendedPriority)
-	{
-		using namespace UE::Tasks;
-
-		ENamedThreads::Type ThreadIndex = ENamedThreads::GetThreadIndex(ThreadType);
-		if (ThreadIndex != ENamedThreads::AnyThread)
-		{
-			check(ThreadIndex == ENamedThreads::GameThread || ThreadIndex == ENamedThreads::GetRenderThread() || ThreadIndex == ENamedThreads::RHIThread);
-			EExtendedTaskPriority ConversionMap[] =
-			{
-				EExtendedTaskPriority::RHIThreadNormalPri,
-				EExtendedTaskPriority::GameThreadNormalPri,
-				EExtendedTaskPriority::RenderThreadNormalPri
-			};
-			OutExtendedPriority = ConversionMap[ThreadIndex - ENamedThreads::RHIThread];
-			OutExtendedPriority = (EExtendedTaskPriority)((int32)OutExtendedPriority + (ENamedThreads::GetTaskPriority(ThreadType) != ENamedThreads::NormalTaskPriority ? 1 : 0));
-			OutExtendedPriority = (EExtendedTaskPriority)((int32)OutExtendedPriority + (ENamedThreads::GetQueueIndex(ThreadType) != ENamedThreads::MainQueue ? 2 : 0));
-			OutPriority = ETaskPriority::Count;
-		}
-		else
-		{
-			OutExtendedPriority = EExtendedTaskPriority::None;
-			uint32 ThreadPriority = GetThreadPriorityIndex(ThreadType);
-			check(ThreadPriority < uint32(ENamedThreads::NumThreadPriorities));
-			ETaskPriority ConversionMap[int(ENamedThreads::NumThreadPriorities)] = { ETaskPriority::Normal, ETaskPriority::High, ETaskPriority::BackgroundNormal };
-			OutPriority = ConversionMap[ThreadPriority];
-		}
-
-		if (OutPriority == ETaskPriority::BackgroundNormal && GetTaskPriority(ThreadType))
-		{
-			OutPriority = ETaskPriority::BackgroundHigh;
-		}
-	}
-
-	// task priority translation from the new API to the old API
-	static ENamedThreads::Type TranslatePriority(UE::Tasks::EExtendedTaskPriority Priority)
-	{
-		using namespace UE::Tasks;
-
-		checkf(Priority >= EExtendedTaskPriority::GameThreadNormalPri && Priority < EExtendedTaskPriority::Count, TEXT("only named threads can call this method: %d"), Priority);
-
-		int32 ConversionMap[] =
-		{
-				ENamedThreads::GameThread, // GameThreadNormalPri
-				ENamedThreads::GameThread | ENamedThreads::HighTaskPriority, // GameThreadHiPri
-				ENamedThreads::GameThread | ENamedThreads::LocalQueue, // GameThreadNormalPriLocalQueue
-				ENamedThreads::GameThread | ENamedThreads::HighTaskPriority | ENamedThreads::LocalQueue, // GameThreadHiPriLocalQueue
-
-				ENamedThreads::GetRenderThread(), // RenderThreadNormalPri
-				ENamedThreads::GetRenderThread() | ENamedThreads::HighTaskPriority, // RenderThreadHiPri
-				ENamedThreads::GetRenderThread() | ENamedThreads::LocalQueue, // RenderThreadNormalPriLocalQueue
-				ENamedThreads::GetRenderThread() | ENamedThreads::HighTaskPriority | ENamedThreads::LocalQueue, // RenderThreadHiPriLocalQueue
-
-				ENamedThreads::RHIThread, // RHIThreadNormalPri
-				ENamedThreads::RHIThread | ENamedThreads::HighTaskPriority, // RHIThreadHiPri
-				ENamedThreads::RHIThread | ENamedThreads::LocalQueue, // RHIThreadNormalPriLocalQueue
-				ENamedThreads::RHIThread | ENamedThreads::HighTaskPriority | ENamedThreads::LocalQueue // RHIThreadHiPriLocalQueue
-		};
-
-		return (ENamedThreads::Type)ConversionMap[(int32)Priority - (int32)EExtendedTaskPriority::GameThreadNormalPri];
-	}
 };
 
 static constexpr int32 SmallTaskSize = 256;
@@ -673,7 +611,7 @@ public:
 
 			UE::Tasks::ETaskPriority Pri;
 			UE::Tasks::EExtendedTaskPriority ExtPri;
-			TranslatePriority(TaskObject->GetDesiredThread(), Pri, ExtPri);
+			UE::Tasks::Private::TranslatePriority(TaskObject->GetDesiredThread(), Pri, ExtPri);
 
 			Task->Init(Pri, ExtPri);
 
diff --git a/Engine/Source/Runtime/Core/Public/Tasks/TaskPrivate.h b/Engine/Source/Runtime/Core/Public/Tasks/TaskPrivate.h
--- a/Engine/Source/Runtime/Core/Public/Tasks/TaskPrivate.h
+++ b/Engine/Source/Runtime/Core/Public/Tasks/TaskPrivate.h
@@ -73,6 +73,15 @@ namespace UE::Tasks
 	const TCHAR* ToString(EExtendedTaskPriority ExtendedPriority);
 	bool ToExtendedTaskPriority(const TCHAR* ExtendedPriorityStr, EExtendedTaskPriority& OutExtendedPriority);
 
+	namespace Private
+	{
+		CORE_API void TranslatePriority(ENamedThreads::Type ThreadType, ETaskPriority& OutPriority, EExtendedTaskPriority& OutExtendedPriority);
+		CORE_API ENamedThreads::Type TranslatePriority(ETaskPriority Priority);
+#if TASKGRAPH_NEW_FRONTEND
+		CORE_API ENamedThreads::Type TranslatePriority(EExtendedTaskPriority Priority);
+#endif
+	}
+
 	class FPipe;
 
 	namespace Private
@@ -303,7 +312,7 @@ namespace UE::Tasks
 			// In this case the task will be automatically scheduled when all dependencies are completed.
 			bool TryLaunch(uint64 TaskSize)
 			{
-				TaskTrace::Launched(GetTraceId(), LowLevelTask.GetDebugName(), true, (ENamedThreads::Type)0xff, TaskSize);
+				TaskTrace::Launched(GetTraceId(), LowLevelTask.GetDebugName(), true, TranslatePriority(LowLevelTask.GetPriority()), TaskSize);
 				return TryUnlock();
 			}
 
