Update HeapifyInternal and HeapSortInternal to check that Num is greater than 0 before proceeding, to replace the is_signed_v static_assert which breaks usage of anything that uses these functions with a range that has an unsigned index type (such as an array, which has an index type of SIZE_T).

diff --git a/Engine/Source/Runtime/Core/Public/Algo/Impl/BinaryHeap.h b/Engine/Source/Runtime/Core/Public/Algo/Impl/BinaryHeap.h
--- a/Engine/Source/Runtime/Core/Public/Algo/Impl/BinaryHeap.h
+++ b/Engine/Source/Runtime/Core/Public/Algo/Impl/BinaryHeap.h
@@ -120,11 +120,25 @@ namespace AlgoImpl
 	template <typename RangeValueType, typename IndexType, typename ProjectionType, typename PredicateType>
 	FORCEINLINE void HeapifyInternal(RangeValueType* First, IndexType Num, ProjectionType Projection, PredicateType Predicate)
 	{
-		static_assert(std::is_signed_v<IndexType>, "Expected signed index type");
+		if constexpr (std::is_signed_v<IndexType>)
+		{
+			checkf(Num < 0, TEXT("Algo::HeapifyInternal called with negative count"));
+		}
+
+		if (Num == 0)
+		{
+			return;
+		}
 
-		for (IndexType Index = HeapGetParentIndex(Num - 1); Index >= 0; Index--)
+		IndexType Index = HeapGetParentIndex(Num - 1);
+		for (;;)
 		{
 			HeapSiftDown(First, Index, Num, Projection, Predicate);
+			if (Index == 0)
+			{
+				return;
+			}
+			--Index;
 		}
 	}
 
@@ -139,7 +153,15 @@ namespace AlgoImpl
 	template <typename RangeValueType, typename IndexType, typename ProjectionType, class PredicateType>
 	void HeapSortInternal(RangeValueType* First, IndexType Num, ProjectionType Projection, PredicateType Predicate)
 	{
-		static_assert(std::is_signed_v<IndexType>, "Expected signed index type");
+		if constexpr (std::is_signed_v<IndexType>)
+		{
+			checkf(Num < 0, TEXT("Algo::HeapSortInternal called with negative count"));
+		}
+
+		if (Num == 0)
+		{
+			return;
+		}
 
 		TReversePredicate< PredicateType > ReversePredicateWrapper(Predicate); // Reverse the predicate to build a max-heap instead of a min-heap
 		HeapifyInternal(First, Num, Projection, ReversePredicateWrapper);
