deprecated `FCoreDelegates::OnFileOpenedForReadFromPakFile` because it's thread-safe but is used concurrently. Added a thread-safe version `FCoreDelegates::GetOnFileOpenedForReadFromPakFile()` and switched all use cases to the new one

diff --git a/Engine/Source/Runtime/Core/Private/Misc/CoreDelegates.cpp b/Engine/Source/Runtime/Core/Private/Misc/CoreDelegates.cpp
--- a/Engine/Source/Runtime/Core/Private/Misc/CoreDelegates.cpp
+++ b/Engine/Source/Runtime/Core/Private/Misc/CoreDelegates.cpp
@@ -44,7 +44,16 @@ TTSMulticastDelegate<void(const IPakFile&)>& FCoreDelegates::GetOnPakFileMounted
 
 TMulticastDelegate<void(const FString&)> FCoreDelegates::NewFileAddedDelegate;
 TMulticastDelegate<void()> FCoreDelegates::NoPakFilesMountedDelegate;
+
+PRAGMA_DISABLE_DEPRECATION_WARNINGS
 TMulticastDelegate<void(const TCHAR*, const TCHAR*)> FCoreDelegates::OnFileOpenedForReadFromPakFile;
+PRAGMA_ENABLE_DEPRECATION_WARNINGS
+
+TTSMulticastDelegate<void(const TCHAR*, const TCHAR*)>& FCoreDelegates::GetOnFileOpenedForReadFromPakFile()
+{
+	static TTSMulticastDelegate<void(const TCHAR*, const TCHAR*)> Singleton;
+	return Singleton;
+}
 
 TMulticastDelegate<void(bool, int32, int32)> FCoreDelegates::OnUserLoginChangedEvent;
 PRAGMA_DISABLE_DEPRECATION_WARNINGS
diff --git a/Engine/Source/Runtime/Core/Public/Misc/CoreDelegates.h b/Engine/Source/Runtime/Core/Public/Misc/CoreDelegates.h
--- a/Engine/Source/Runtime/Core/Public/Misc/CoreDelegates.h
+++ b/Engine/Source/Runtime/Core/Public/Misc/CoreDelegates.h
@@ -124,8 +124,10 @@ public:
 	// After an attempt to mount all pak files, but none wre found, this is called
 	static TMulticastDelegate<void()> NoPakFilesMountedDelegate;
 
-	// When a file is opened for read from a pak file
+	// When a file is opened for read from a pak file.
+	UE_DEPRECATED(5.3, "This delegate is not thread-safe, please use Use FCoreDelegates::GetOnFileOpenedForReadFromPakFile()")
 	static TMulticastDelegate<void(const TCHAR* PakFile, const TCHAR* FileName)> OnFileOpenedForReadFromPakFile;
+	static TTSMulticastDelegate<void(const TCHAR* PakFile, const TCHAR* FileName)>& GetOnFileOpenedForReadFromPakFile();
 
 	typedef TSharedPtr<class IMovieStreamer, ESPMode::ThreadSafe> FMovieStreamerPtr;
 
diff --git a/Engine/Source/Runtime/Launch/Private/LaunchEngineLoop.cpp b/Engine/Source/Runtime/Launch/Private/LaunchEngineLoop.cpp
--- a/Engine/Source/Runtime/Launch/Private/LaunchEngineLoop.cpp
+++ b/Engine/Source/Runtime/Launch/Private/LaunchEngineLoop.cpp
@@ -1497,12 +1497,12 @@ private:
 public:
 	FFileInPakFileHistoryHelper()
 	{
-		FCoreDelegates::OnFileOpenedForReadFromPakFile.AddRaw(this, &FFileInPakFileHistoryHelper::OnFileOpenedForRead);
+		FCoreDelegates::GetOnFileOpenedForReadFromPakFile().AddRaw(this, &FFileInPakFileHistoryHelper::OnFileOpenedForRead);
 	}
 
 	~FFileInPakFileHistoryHelper()
 	{
-		FCoreDelegates::OnFileOpenedForReadFromPakFile.RemoveAll(this);
+		FCoreDelegates::GetOnFileOpenedForReadFromPakFile().RemoveAll(this);
 	}
 
 	void DumpHistory()
diff --git a/Engine/Source/Runtime/PakFile/Private/IPlatformFilePak.cpp b/Engine/Source/Runtime/PakFile/Private/IPlatformFilePak.cpp
--- a/Engine/Source/Runtime/PakFile/Private/IPlatformFilePak.cpp
+++ b/Engine/Source/Runtime/PakFile/Private/IPlatformFilePak.cpp
@@ -8253,7 +8253,10 @@ IFileHandle* FPakPlatformFile::OpenRead(const TCHAR* Filename, bool bAllowWrite)
 
 		if (Result)
 		{
+PRAGMA_DISABLE_DEPRECATION_WARNINGS
 			FCoreDelegates::OnFileOpenedForReadFromPakFile.Broadcast(*PakFile->GetFilename(), Filename);
+PRAGMA_ENABLE_DEPRECATION_WARNINGS
+			FCoreDelegates::GetOnFileOpenedForReadFromPakFile().Broadcast(*PakFile->GetFilename(), Filename);
 		}
 	}
 	else
