Fixed inverted assert conditions in Algo::Heap*.

diff --git a/Engine/Source/Runtime/Core/Public/Algo/Impl/BinaryHeap.h b/Engine/Source/Runtime/Core/Public/Algo/Impl/BinaryHeap.h
--- a/Engine/Source/Runtime/Core/Public/Algo/Impl/BinaryHeap.h
+++ b/Engine/Source/Runtime/Core/Public/Algo/Impl/BinaryHeap.h
@@ -122,7 +122,7 @@ namespace AlgoImpl
 	{
 		if constexpr (std::is_signed_v<IndexType>)
 		{
-			checkf(Num < 0, TEXT("Algo::HeapifyInternal called with negative count"));
+			checkf(Num >= 0, TEXT("Algo::HeapifyInternal called with negative count"));
 		}
 
 		if (Num == 0)
@@ -155,7 +155,7 @@ namespace AlgoImpl
 	{
 		if constexpr (std::is_signed_v<IndexType>)
 		{
-			checkf(Num < 0, TEXT("Algo::HeapSortInternal called with negative count"));
+			checkf(Num >= 0, TEXT("Algo::HeapSortInternal called with negative count"));
 		}
 
 		if (Num == 0)
