Move Core AutoRTFM to UE concepts and drop using std::

diff --git a/Engine/Source/Runtime/Core/Private/AutoRTFM/API.cpp b/Engine/Source/Runtime/Core/Private/AutoRTFM/API.cpp
--- a/Engine/Source/Runtime/Core/Private/AutoRTFM/API.cpp
+++ b/Engine/Source/Runtime/Core/Private/AutoRTFM/API.cpp
@@ -169,11 +169,11 @@ extern "C" void autortfm_register_open_function(void* OriginalFunction, void* Ne
     FunctionMapAdd(OriginalFunction, NewFunction);
 }
 
-void DeferUntilCommit(std::function<void()>&& Work)
+void DeferUntilCommit(TFunction<void()>&& Work)
 {
     if (FContext::IsTransactional())
     {
-        FContext::Get()->GetCurrentTransaction()->DeferUntilCommit(std::move(Work));
+        FContext::Get()->GetCurrentTransaction()->DeferUntilCommit(MoveTemp(Work));
     }
     else
     {
@@ -181,20 +181,20 @@ void DeferUntilCommit(std::function<void()>&& Work)
     }
 }
 
-void DeferUntilAbort(std::function<void()>&& Work)
+void DeferUntilAbort(TFunction<void()>&& Work)
 {
     if (FContext::IsTransactional())
     {
-        FContext::Get()->GetCurrentTransaction()->DeferUntilAbort(std::move(Work));
+        FContext::Get()->GetCurrentTransaction()->DeferUntilAbort(MoveTemp(Work));
     }
 }
 
-void OpenCommit(std::function<void()>&& Work)
+void OpenCommit(TFunction<void()>&& Work)
 {
     Work();
 }
 
-void OpenAbort(std::function<void()>&& Work)
+void OpenAbort(TFunction<void()>&& Work)
 {
 }
 
@@ -254,7 +254,7 @@ extern "C" void autortfm_check_abi(void* const Ptr, const size_t Size)
 	private:
 		auto Tied() const
 		{
-			return std::tie(LogLineBytes, LineBytes, LineTableSize, Offset_Context_CurrentTransaction, Offset_Context_LineTable, Offset_Context_Status, LogSize_LineEntry, Size_LineEntry, Offset_LineEntry_LogicalLine, Offset_LineEntry_ActiveLine, Offset_LineEntry_LoggingTransaction, Offset_LineEntry_AccessMask, Context_Status_OnTrack);
+			return Tie(LogLineBytes, LineBytes, LineTableSize, Offset_Context_CurrentTransaction, Offset_Context_LineTable, Offset_Context_Status, LogSize_LineEntry, Size_LineEntry, Offset_LineEntry_LogicalLine, Offset_LineEntry_ActiveLine, Offset_LineEntry_LoggingTransaction, Offset_LineEntry_AccessMask, Context_Status_OnTrack);
 		}
 
 	public:
@@ -416,31 +416,31 @@ extern "C" void STM_autortfm_record_open_write(void*, size_t, FContext*)
 }
 UE_AUTORTFM_REGISTER_OPEN_FUNCTION(autortfm_record_open_write);
 
-void STM_DeferUntilCommit(std::function<void()>&& Work, FContext* Context)
+void STM_DeferUntilCommit(TFunction<void()>&& Work, FContext* Context)
 {
     ASSERT(Context->GetStatus() == EContextStatus::OnTrack);
-    Context->GetCurrentTransaction()->DeferUntilCommit(std::move(Work));
+    Context->GetCurrentTransaction()->DeferUntilCommit(MoveTemp(Work));
 }
 UE_AUTORTFM_REGISTER_OPEN_FUNCTION(DeferUntilCommit);
 
-void STM_DeferUntilAbort(std::function<void()>&& Work, FContext* Context)
+void STM_DeferUntilAbort(TFunction<void()>&& Work, FContext* Context)
 {
     ASSERT(Context->GetStatus() == EContextStatus::OnTrack);
-    Context->GetCurrentTransaction()->DeferUntilAbort(std::move(Work));
+    Context->GetCurrentTransaction()->DeferUntilAbort(MoveTemp(Work));
 }
 UE_AUTORTFM_REGISTER_OPEN_FUNCTION(DeferUntilAbort);
 
-void STM_OpenCommit(std::function<void()>&& Work, FContext* Context)
+void STM_OpenCommit(TFunction<void()>&& Work, FContext* Context)
 {
     ASSERT(Context->GetStatus() == EContextStatus::OnTrack);
-    Context->GetCurrentTransaction()->DeferUntilCommit(std::move(Work));
+    Context->GetCurrentTransaction()->DeferUntilCommit(MoveTemp(Work));
 }
 UE_AUTORTFM_REGISTER_OPEN_FUNCTION(OpenCommit);
 
-void STM_OpenAbort(std::function<void()>&& Work, FContext* Context)
+void STM_OpenAbort(TFunction<void()>&& Work, FContext* Context)
 {
     ASSERT(Context->GetStatus() == EContextStatus::OnTrack);
-    Context->GetCurrentTransaction()->DeferUntilAbort(std::move(Work));
+    Context->GetCurrentTransaction()->DeferUntilAbort(MoveTemp(Work));
 }
 UE_AUTORTFM_REGISTER_OPEN_FUNCTION(OpenAbort);
 
diff --git a/Engine/Source/Runtime/Core/Private/AutoRTFM/Context.cpp b/Engine/Source/Runtime/Core/Private/AutoRTFM/Context.cpp
--- a/Engine/Source/Runtime/Core/Private/AutoRTFM/Context.cpp
+++ b/Engine/Source/Runtime/Core/Private/AutoRTFM/Context.cpp
@@ -8,13 +8,11 @@
 #include "GlobalData.h"
 #include "ScopedGuard.h"
 #include "TransactionInlines.h"
-#include <algorithm>
-#include <memory>
 
 namespace AutoRTFM
 {
 
-thread_local std::unique_ptr<FContext> ContextTls;
+thread_local TUniquePtr<FContext> ContextTls;
 
 void FContext::InitializeGlobalData()
 {
@@ -22,12 +20,12 @@ void FContext::InitializeGlobalData()
 
 FContext* FContext::TryGet()
 {
-    return ContextTls.get();
+    return ContextTls.Get();
 }
 
 void FContext::Set()
 {
-    ContextTls.reset(this);
+    ContextTls.Reset(this);
 }
 
 FContext* FContext::Get()
@@ -259,11 +257,11 @@ ETransactionResult FContext::Transact(void (*Function)(void* Arg), void* Arg)
     void (*ClonedFunction)(void* Arg, FContext* Context) = FunctionMapTryLookup(Function);
     if (!ClonedFunction)
     {
-        fprintf(stderr, "Could not find function %p (%s) in AutoRTFM::FContext::Transact.\n", Function, GetFunctionDescription(Function).c_str());
+        fprintf(stderr, "Could not find function %p (%s) in AutoRTFM::FContext::Transact.\n", Function, TCHAR_TO_ANSI(*GetFunctionDescription(Function)));
         return ETransactionResult::AbortedByLanguage;
     }
     
-	//std::unique_ptr<FTransaction> NewTransactionUniquePtr(new FTransaction(this));
+	//TUniquePtr<FTransaction> NewTransactionUniquePtr(new FTransaction(this));
 	FTransaction* NewTransaction = new FTransaction(this);
 	FCallNest* NewNest = new FCallNest(this);
 
@@ -398,13 +396,13 @@ void FContext::AbortByLanguageAndThrow()
     CurrentTransaction->AbortAndThrow();
 }
 
-#if _WIN32
+#if PLATFORM_WINDOWS
 extern "C" __declspec(dllimport) void __stdcall GetCurrentThreadStackLimits(void**, void**);
 #endif
 
 FContext::FContext()
 {
-#if _WIN32
+#if PLATFORM_WINDOWS
     GetCurrentThreadStackLimits(&StackBegin, &StackEnd);
 #elif defined(__APPLE__)         
    StackEnd = pthread_get_stackaddr_np(pthread_self());   
diff --git a/Engine/Source/Runtime/Core/Private/AutoRTFM/FunctionMap.cpp b/Engine/Source/Runtime/Core/Private/AutoRTFM/FunctionMap.cpp
--- a/Engine/Source/Runtime/Core/Private/AutoRTFM/FunctionMap.cpp
+++ b/Engine/Source/Runtime/Core/Private/AutoRTFM/FunctionMap.cpp
@@ -91,12 +91,12 @@ void FunctionMapAdd(void* OldFunction, void* NewFunction)
     InitializeGlobalDataIfNecessary();
     
     TLockGuard<FFastLock> LockGuard(GlobalData->FunctionMapLock);
-    std::vector<void*> Functions;
-    Functions.push_back(OldFunction);
-    while (!Functions.empty())
+    TArray<void*> Functions;
+    Functions.Push(OldFunction);
+    while (!Functions.IsEmpty())
     {
-        void* Function = Functions.back();
-        Functions.pop_back();
+        void* Function = Functions.Last();
+        Functions.Pop();
         FunctionMapAddImpl(Function, NewFunction);
     }
 }
diff --git a/Engine/Source/Runtime/Core/Private/AutoRTFM/FunctionMap.h b/Engine/Source/Runtime/Core/Private/AutoRTFM/FunctionMap.h
--- a/Engine/Source/Runtime/Core/Private/AutoRTFM/FunctionMap.h
+++ b/Engine/Source/Runtime/Core/Private/AutoRTFM/FunctionMap.h
@@ -3,7 +3,6 @@
 #pragma once
 
 #include "GlobalData.h"
-#include <cstddef>
 
 namespace AutoRTFM
 {
diff --git a/Engine/Source/Runtime/Core/Private/AutoRTFM/FunctionMapInlines.h b/Engine/Source/Runtime/Core/Private/AutoRTFM/FunctionMapInlines.h
--- a/Engine/Source/Runtime/Core/Private/AutoRTFM/FunctionMapInlines.h
+++ b/Engine/Source/Runtime/Core/Private/AutoRTFM/FunctionMapInlines.h
@@ -13,7 +13,7 @@ inline void* FunctionMapLookup(void* OldFunction, FContext* Context, const char*
     void* Result = FunctionMapTryLookup(OldFunction);
     if (!Result)
     {
-        fprintf(stderr, "Could not find function %p (%s)\n", OldFunction, GetFunctionDescription(OldFunction).c_str());
+        fprintf(stderr, "Could not find function %p (%s)\n", OldFunction, TCHAR_TO_ANSI(*GetFunctionDescription(OldFunction)));
         if (Where)
         {
             fprintf(stderr, "%s\n", Where);
diff --git a/Engine/Source/Runtime/Core/Private/AutoRTFM/GlobalData.cpp b/Engine/Source/Runtime/Core/Private/AutoRTFM/GlobalData.cpp
--- a/Engine/Source/Runtime/Core/Private/AutoRTFM/GlobalData.cpp
+++ b/Engine/Source/Runtime/Core/Private/AutoRTFM/GlobalData.cpp
@@ -3,17 +3,15 @@
 #if (defined(__AUTORTFM) && __AUTORTFM)
 #include "GlobalData.h"
 #include "Context.h"
-#include <mutex>
 
 namespace AutoRTFM
 {
 
-std::once_flag GlobalDataInitializationFlag;
 FGlobalData* GlobalData;
 
 void InitializeGlobalDataIfNecessary()
 {
-    std::call_once(GlobalDataInitializationFlag, []
+    UE_CALL_ONCE([]
     {
         // AutoRTFM is intended to be used primarily from monolithic UE binaries, where having
         // some additional DLL/so/dylib would just be annoying. However, we guard against the
@@ -28,15 +26,6 @@ void InitializeGlobalDataIfNecessary()
         
         constexpr const char* EnvName = "AutoRTFMGlobalData";
 
-#ifdef _MSC_VER
-/*
-   Disable warning about deprecated STD C functions.
-*/
-#pragma warning(disable : 4996)
-
-#pragma warning(push)
-#endif
-
         if (char* EnvString = getenv(EnvName))
         {
             sscanf(EnvString, "%p", &GlobalData);
@@ -50,10 +39,6 @@ void InitializeGlobalDataIfNecessary()
             snprintf(Buffer, sizeof(Buffer), "%s=%p", EnvName, GlobalData);
             putenv(Buffer);
         }
-
-#ifdef _MSC_VER
-#pragma warning(pop)
-#endif
     });
 }
 
diff --git a/Engine/Source/Runtime/Core/Private/AutoRTFM/GlobalData.h b/Engine/Source/Runtime/Core/Private/AutoRTFM/GlobalData.h
--- a/Engine/Source/Runtime/Core/Private/AutoRTFM/GlobalData.h
+++ b/Engine/Source/Runtime/Core/Private/AutoRTFM/GlobalData.h
@@ -3,8 +3,6 @@
 #pragma once
 
 #include "FastLock.h"
-#include <unordered_map>
-#include <vector>
 
 namespace AutoRTFM
 {
diff --git a/Engine/Source/Runtime/Core/Private/AutoRTFM/LongJump.h b/Engine/Source/Runtime/Core/Private/AutoRTFM/LongJump.h
--- a/Engine/Source/Runtime/Core/Private/AutoRTFM/LongJump.h
+++ b/Engine/Source/Runtime/Core/Private/AutoRTFM/LongJump.h
@@ -34,11 +34,11 @@ template<typename TTryFunctor, typename TCatchFunctor>
 void FLongJump::TryCatch(const TTryFunctor& TryFunctor, const TCatchFunctor& CatchFunctor)
 {
     ASSERT(!bIsSet);
-#ifdef _WIN32
+#if PLATFORM_WINDOWS
 	if (!setjmp(JmpBuf))
 #else
 	if (!_setjmp(JmpBuf))
-#endif
+#endif // PLATFORM_WINDOWS
     {
         bIsSet = true;
         TryFunctor();
@@ -55,11 +55,11 @@ void FLongJump::TryCatch(const TTryFunctor& TryFunctor, const TCatchFunctor& Cat
 inline void FLongJump::Throw()
 {
     ASSERT(bIsSet);
-#ifdef _WIN32
+#if PLATFORM_WINDOWS
 	longjmp(JmpBuf, 1);
 #else
 	_longjmp(JmpBuf, 1);
-#endif
+#endif // PLATFORM_WINDOWS
 }
 
 } // namespace AutoRTFM
diff --git a/Engine/Source/Runtime/Core/Private/AutoRTFM/StdLibSupport.cpp b/Engine/Source/Runtime/Core/Private/AutoRTFM/StdLibSupport.cpp
--- a/Engine/Source/Runtime/Core/Private/AutoRTFM/StdLibSupport.cpp
+++ b/Engine/Source/Runtime/Core/Private/AutoRTFM/StdLibSupport.cpp
@@ -9,7 +9,6 @@
 #include <string.h>
 #include <wchar.h>
 #include <float.h>
-#include <math.h>
 
 #ifdef __APPLE__
 #include <malloc/malloc.h>
@@ -17,7 +16,7 @@
 #include <malloc.h>
 #endif
 
-#ifdef _WIN32
+#if PLATFORM_WINDOWS
 #define NOMINMAX
 #include <windows.h>
 #endif
@@ -71,7 +70,7 @@ void* STM_realloc(void* Ptr, size_t Size, FContext* Context)
 #else
 		const size_t OldSize = malloc_usable_size(Ptr);
 #endif
-        MemcpyToNew(NewObject, Ptr, std::min(OldSize, Size), Context);
+        MemcpyToNew(NewObject, Ptr,  FMath::Min(OldSize, Size), Context);
         STM_free(Ptr, Context);
     }
     return NewObject;
@@ -141,7 +140,7 @@ int STM_puts(const char* Str, FContext* Context)
 }
 UE_AUTORTFM_REGISTER_OPEN_FUNCTION(puts);
 
-#ifdef _WIN32
+#if PLATFORM_WINDOWS
 
 FILE* STM___acrt_iob_func(int Index, FContext* Context)
 {
@@ -162,15 +161,15 @@ UE_AUTORTFM_REGISTER_SELF_FUNCTION(__stdio_common_vsprintf);
 UE_AUTORTFM_REGISTER_SELF_FUNCTION(__stdio_common_vswprintf);
 UE_AUTORTFM_REGISTER_SELF_FUNCTION(__stdio_common_vfwprintf);
 
-#else // _WIN32 -> so !_WIN32
+#else // PLATFORM_WINDOWS -> so !PLATFORM_WINDOWS
 extern "C" size_t _ZNSt3__112__next_primeEm(size_t N) __attribute__((weak));
 UE_AUTORTFM_REGISTER_SELF_FUNCTION(_ZNSt3__112__next_primeEm);
-#endif // _WIN32 -> so end of !_WIN32
+#endif // PLATFORM_WINDOWS -> so end of !PLATFORM_WINDOWS
 
 UE_AUTORTFM_REGISTER_SELF_FUNCTION(static_cast<float(*)(float, float)>(&powf));
 UE_AUTORTFM_REGISTER_SELF_FUNCTION(static_cast<double(*)(double, double)>(&pow));
 
-#ifdef _WIN32
+#if PLATFORM_WINDOWS
 UE_AUTORTFM_REGISTER_SELF_FUNCTION(_isnan);
 UE_AUTORTFM_REGISTER_SELF_FUNCTION(_finite);
 UE_AUTORTFM_REGISTER_SELF_FUNCTION(IsDebuggerPresent);
@@ -189,7 +188,7 @@ extern "C" void* _Init_thread_footer();
 UE_AUTORTFM_REGISTER_SELF_FUNCTION(_Init_thread_header);
 UE_AUTORTFM_REGISTER_SELF_FUNCTION(_Init_thread_abort);
 UE_AUTORTFM_REGISTER_SELF_FUNCTION(_Init_thread_footer);
-#endif
+#endif // PLATFORM_WINDOWS
 
 wchar_t* STM_wcsncpy(wchar_t* Dst, const wchar_t* Src, size_t Count, FContext* Context)
 {
diff --git a/Engine/Source/Runtime/Core/Private/AutoRTFM/TaggedPtr.h b/Engine/Source/Runtime/Core/Private/AutoRTFM/TaggedPtr.h
--- a/Engine/Source/Runtime/Core/Private/AutoRTFM/TaggedPtr.h
+++ b/Engine/Source/Runtime/Core/Private/AutoRTFM/TaggedPtr.h
@@ -87,7 +87,7 @@ namespace AutoRTFM
 
 	    void SetTopTag(const TopTagT Tag)
 	    {
-	    	*reinterpret_cast<TopTagT*>(reinterpret_cast<uint8_t*>(&Payload) + OffsetOfTopTag) = std::move(Tag);
+	    	*reinterpret_cast<TopTagT*>(reinterpret_cast<uint8_t*>(&Payload) + OffsetOfTopTag) = MoveTemp(Tag);
 	    }
 
 	    TopTagT GetTopTag() const
@@ -182,7 +182,7 @@ namespace AutoRTFM
 
 	    void SetTopTag(const TopTagT Tag)
 	    {
-	    	*reinterpret_cast<TopTagT*>(reinterpret_cast<uint8_t*>(&Payload) + OffsetOfTopTag) = std::move(Tag);
+	    	*reinterpret_cast<TopTagT*>(reinterpret_cast<uint8_t*>(&Payload) + OffsetOfTopTag) = Tag;
 	    }
 
 	    TopTagT GetTopTag() const
diff --git a/Engine/Source/Runtime/Core/Private/AutoRTFM/TaskArray.h b/Engine/Source/Runtime/Core/Private/AutoRTFM/TaskArray.h
--- a/Engine/Source/Runtime/Core/Private/AutoRTFM/TaskArray.h
+++ b/Engine/Source/Runtime/Core/Private/AutoRTFM/TaskArray.h
@@ -3,7 +3,6 @@
 #pragma once
 
 #include "SharedPtr.h"
-#include <vector>
 
 namespace AutoRTFM
 {
@@ -14,28 +13,28 @@ class TTaskArray
 public:
     TTaskArray() = default;
 
-    bool IsEmpty() const { return Latest.empty() && Stash.empty(); }
+    bool IsEmpty() const { return Latest.IsEmpty() && Stash.IsEmpty(); }
 
     void Add(T&& value)
     {
-        Latest.push_back(std::move(value));
+        Latest.Push(MoveTemp(value));
     }
     void Add(const T& value)
     {
-        Latest.push_back(value);
+        Latest.Push(value);
     }
 
     void AddAll(TTaskArray<T>&& Other)
     {
         Canonicalize();
-        for (TSharedPtr<std::vector<T>>& StashedVectorBox : Other.Stash)
+        for (TSharedPtr<TArray<T>>& StashedVectorBox : Other.Stash)
         {
-            Stash.push_back(std::move(StashedVectorBox));
+            Stash.Push(MoveTemp(StashedVectorBox));
         }
-        Other.Stash.clear(); 
-        if (!Other.Latest.empty())
+        Other.Stash.Empty();
+        if (!Other.Latest.IsEmpty())
         {
-            Stash.push_back(TSharedPtr<std::vector<T>>::New(std::move(Other.Latest)));
+            Stash.Push(TSharedPtr<TArray<T>>::New(MoveTemp(Other.Latest)));
         }
     }
 
@@ -43,16 +42,16 @@ public:
     {
         Canonicalize();
         Other.Canonicalize();
-        for (const TSharedPtr<std::vector<T>>& StashedVectorBox : Other.Stash)
+        for (const TSharedPtr<TArray<T>>& StashedVectorBox : Other.Stash)
         {
-            Stash.push_back(StashedVectorBox);
+            Stash.Push(StashedVectorBox);
         }
     }
 
     template<typename TFunc>
     bool ForEachForward(const TFunc& Func)
     {
-        for (TSharedPtr<std::vector<T>>& StashedVectorBox : Stash)
+        for (TSharedPtr<TArray<T>>& StashedVectorBox : Stash)
         {
             for (const T& Entry : *StashedVectorBox)
             {
@@ -75,17 +74,17 @@ public:
     template<typename TFunc>
     bool ForEachBackward(const TFunc& Func)
     {
-        for (size_t Index = Latest.size(); Index--;)
+        for (size_t Index = Latest.Num(); Index--;)
         {
             if (!Func(Latest[Index]))
             {
                 return false;
             }
         }
-        for (size_t IndexInStash = Stash.size(); IndexInStash--;)
+        for (size_t IndexInStash = Stash.Num(); IndexInStash--;)
         {
-            const std::vector<T>& StashedVector = *Stash[IndexInStash];
-            for (size_t Index = StashedVector.size(); Index--;)
+            const TArray<T>& StashedVector = *Stash[IndexInStash];
+            for (size_t Index = StashedVector.Num(); Index--;)
             {
                 if (!Func(StashedVector[Index]))
                 {
@@ -98,8 +97,8 @@ public:
 
     void Reset()
     {
-        Latest.clear();
-        Stash.clear();
+        Latest.Empty();
+        Stash.Empty();
     }
 
 private:
@@ -108,15 +107,15 @@ private:
     // method is `const`.
     void Canonicalize() const
     {
-        if (!Latest.empty())
+        if (!Latest.IsEmpty())
         {
-            Stash.push_back(TSharedPtr<std::vector<T>>::New(std::move(Latest)));
+            Stash.Push(TSharedPtr<TArray<T>>::New(MoveTemp(Latest)));
         }
     }
     
-    mutable std::vector<T> Latest;
+    mutable TArray<T> Latest;
     
-    mutable std::vector<TSharedPtr<std::vector<T>>> Stash;
+    mutable TArray<TSharedPtr<TArray<T>>> Stash;
 };
 
 } // namespace AutoRTFM
diff --git a/Engine/Source/Runtime/Core/Private/AutoRTFM/Transaction.cpp b/Engine/Source/Runtime/Core/Private/AutoRTFM/Transaction.cpp
--- a/Engine/Source/Runtime/Core/Private/AutoRTFM/Transaction.cpp
+++ b/Engine/Source/Runtime/Core/Private/AutoRTFM/Transaction.cpp
@@ -94,14 +94,14 @@ void FTransaction::AbortNested()
     Undo();
 
     Parent->CommitTasks.AddAll(AbortTasks);
-    Parent->AbortTasks.AddAll(std::move(AbortTasks));
+    Parent->AbortTasks.AddAll(MoveTemp(AbortTasks));
 }
 
 void FTransaction::AbortOuterNest()
 {
     Undo();
 
-    AbortTasks.ForEachBackward([] (const std::function<void()>& Task) -> bool { Task(); return true; });
+    AbortTasks.ForEachBackward([] (const TFunction<void()>& Task) -> bool { Task(); return true; });
 
     switch (Context->GetStatus())
     {
@@ -137,10 +137,10 @@ void FTransaction::CommitNested()
         Parent->HitSet.Insert(Write.OriginalAndSize);
     }
 
-    Parent->WriteLogBumpAllocator.Merge(std::move(WriteLogBumpAllocator));
+    Parent->WriteLogBumpAllocator.Merge(MoveTemp(WriteLogBumpAllocator));
 
-    Parent->CommitTasks.AddAll(std::move(CommitTasks));
-    Parent->AbortTasks.AddAll(std::move(AbortTasks));
+    Parent->CommitTasks.AddAll(MoveTemp(CommitTasks));
+    Parent->AbortTasks.AddAll(MoveTemp(AbortTasks));
 }
 
 bool FTransaction::AttemptToCommitOuterNest()
@@ -154,7 +154,7 @@ bool FTransaction::AttemptToCommitOuterNest()
         fprintf(GetLogFile(), "Running commit tasks...\n");
     }
 
-    CommitTasks.ForEachForward([] (const std::function<void()>& Task) -> bool { Task(); return true; });
+    CommitTasks.ForEachForward([] (const TFunction<void()>& Task) -> bool { Task(); return true; });
 
     return true;
 }
diff --git a/Engine/Source/Runtime/Core/Private/AutoRTFM/Transaction.h b/Engine/Source/Runtime/Core/Private/AutoRTFM/Transaction.h
--- a/Engine/Source/Runtime/Core/Private/AutoRTFM/Transaction.h
+++ b/Engine/Source/Runtime/Core/Private/AutoRTFM/Transaction.h
@@ -8,7 +8,6 @@
 #include "TaskArray.h"
 #include "WriteLog.h"
 #include "WriteLogBumpAllocator.h"
-#include <functional>
 
 namespace AutoRTFM
 {
@@ -46,8 +45,8 @@ public:
 	inline bool IsScopedTransaction() const { return bIsStackScoped; }
 	inline void SetIsScopedTransaction() { bIsStackScoped = true; }
 
-    void DeferUntilCommit(std::function<void()>&&);
-    void DeferUntilAbort(std::function<void()>&&);
+    void DeferUntilCommit(TFunction<void()>&&);
+    void DeferUntilAbort(TFunction<void()>&&);
 
     void AbortAndThrow();
     void AbortWithoutThrowing();
@@ -75,8 +74,8 @@ private:
     FTransaction* Parent{nullptr};
 
     // Commit tasks run on commit in forward order. Abort tasks run on abort in reverse order.
-    TTaskArray<std::function<void()>> CommitTasks;
-    TTaskArray<std::function<void()>> AbortTasks;
+    TTaskArray<TFunction<void()>> CommitTasks;
+    TTaskArray<TFunction<void()>> AbortTasks;
 
     bool bIsDone{false};
 	bool bIsStackScoped{false};
diff --git a/Engine/Source/Runtime/Core/Private/AutoRTFM/TransactionInlines.h b/Engine/Source/Runtime/Core/Private/AutoRTFM/TransactionInlines.h
--- a/Engine/Source/Runtime/Core/Private/AutoRTFM/TransactionInlines.h
+++ b/Engine/Source/Runtime/Core/Private/AutoRTFM/TransactionInlines.h
@@ -81,14 +81,14 @@ inline void FTransaction::DidAllocate(void* LogicalAddress, size_t Size)
     }
 }
 
-inline void FTransaction::DeferUntilCommit(std::function<void()>&& Callback)
+inline void FTransaction::DeferUntilCommit(TFunction<void()>&& Callback)
 {
-    CommitTasks.Add(std::move(Callback));
+    CommitTasks.Add(MoveTemp(Callback));
 }
 
-inline void FTransaction::DeferUntilAbort(std::function<void()>&& Callback)
+inline void FTransaction::DeferUntilAbort(TFunction<void()>&& Callback)
 {
-    AbortTasks.Add(std::move(Callback));
+    AbortTasks.Add(MoveTemp(Callback));
 }
 
 } // namespace AutoRTFM
diff --git a/Engine/Source/Runtime/Core/Private/AutoRTFM/Utils.cpp b/Engine/Source/Runtime/Core/Private/AutoRTFM/Utils.cpp
--- a/Engine/Source/Runtime/Core/Private/AutoRTFM/Utils.cpp
+++ b/Engine/Source/Runtime/Core/Private/AutoRTFM/Utils.cpp
@@ -3,11 +3,10 @@
 #if (defined(__AUTORTFM) && __AUTORTFM)
 #include "Utils.h"
 #include <errno.h>
-#include <mutex>
 #include <stdio.h>
 #include <string.h>
 
-#ifdef _WIN32
+#if PLATFORM_WINDOWS
 #include <windows.h>
 #include <dbghelp.h>
 #else
@@ -23,32 +22,18 @@ constexpr const char* LogFilename = nullptr;
 
 constexpr const char* LogFileMode = "wt";
 
-std::once_flag LogInitializationFlag;
 FILE* LogFile;
 
 FILE* GetLogFile()
 {
-    std::call_once(LogInitializationFlag, []
+    UE_CALL_ONCE([]
     {
         if (LogFilename)
         {
             ASSERT(LogFileMode);
 
-#ifdef _MSC_VER
-/*
-   Disable warning about deprecated STD C functions.
-*/
-#pragma warning(disable : 4996)
-
-#pragma warning(push)
-#endif
-
             LogFile = fopen(LogFilename, LogFileMode);
 
-#ifdef _MSC_VER
-#pragma warning(pop)
-#endif
-
             if (!LogFile)
             {
                 fprintf(stderr, "Could not open %s: %s\n", LogFilename, strerror(errno));
@@ -63,13 +48,13 @@ FILE* GetLogFile()
     return LogFile;
 }
 
-std::string GetFunctionDescription(void* FunctionPtr)
+FString GetFunctionDescription(void* FunctionPtr)
 {
-#ifdef _WIN32
+#if PLATFORM_WINDOWS
     // This is gross, but it works. It's possible for someone to have SymInitialized before. But if they had, then this
     // will just fail. Also, this function is called in cases where we're failing, so it's ok if we do dirty things.
     SymInitialize(GetCurrentProcess(), nullptr, true);
-    
+
     DWORD64 Displacement = 0;
     DWORD64 Address = reinterpret_cast<DWORD64>(FunctionPtr);
     char Buffer[sizeof(SYMBOL_INFO) + MAX_SYM_NAME * sizeof(TCHAR)];
@@ -82,16 +67,16 @@ std::string GetFunctionDescription(void* FunctionPtr)
     }
     else
     {
-        return "<error getting description>";
+        return TEXT("<error getting description>");
     }
-#else // _WIN32 -> so !_WIN32
+#else // PLATFORM_WINDOWS -> so !PLATFORM_WINDOWS
     char** const symbols = backtrace_symbols(&FunctionPtr, 1);
-    std::string name(*symbols);
+    FString Name(ANSI_TO_TCHAR(*symbols));
     free(symbols);
-    return name;
-#endif // _WIN32 -> so end of !_WIN32
+    return Name;
+#endif // PLATFORM_WINDOWS -> so !PLATFORM_WINDOWS
 }
 
 } // namespace AutoRTFM
 
-#endif // defined(__AUTORTFM) && __AUTORTFM
\ No newline at end of file
+#endif // defined(__AUTORTFM) && __AUTORTFM
diff --git a/Engine/Source/Runtime/Core/Private/AutoRTFM/Utils.h b/Engine/Source/Runtime/Core/Private/AutoRTFM/Utils.h
--- a/Engine/Source/Runtime/Core/Private/AutoRTFM/Utils.h
+++ b/Engine/Source/Runtime/Core/Private/AutoRTFM/Utils.h
@@ -7,7 +7,6 @@
 #include <inttypes.h>
 #include <stdio.h>
 #include <stdlib.h>
-#include <string>
 
 #define ASSERT(exp) do { \
     if (!(exp)) { \
@@ -23,19 +22,19 @@ namespace AutoRTFM
 {
     fprintf(stderr, "Unreachable encountered!\n");
 
-#ifdef _WIN32
+#if PLATFORM_WINDOWS
     __assume(false);
 #else
     __builtin_unreachable();
-#endif
+#endif // PLATFORM_WINDOWS
 }
 
 FILE* GetLogFile();
 
-std::string GetFunctionDescription(void* FunctionPtr);
+FString GetFunctionDescription(void* FunctionPtr);
 
 template<typename TReturnType, typename... TParameterTypes>
-std::string GetFunctionDescription(TReturnType (*FunctionPtr)(TParameterTypes...))
+FString GetFunctionDescription(TReturnType (*FunctionPtr)(TParameterTypes...))
 {
     return GetFunctionDescription(reinterpret_cast<void*>(FunctionPtr));
 }
diff --git a/Engine/Source/Runtime/Core/Public/AutoRTFM/AutoRTFM.h b/Engine/Source/Runtime/Core/Public/AutoRTFM/AutoRTFM.h
--- a/Engine/Source/Runtime/Core/Public/AutoRTFM/AutoRTFM.h
+++ b/Engine/Source/Runtime/Core/Public/AutoRTFM/AutoRTFM.h
@@ -337,8 +337,6 @@ inline void autortfm_check_abi(void* ptr, size_t size) { }
 #ifdef __cplusplus
 }
 
-#include <functional>
-
 namespace AutoRTFM
 {
 
@@ -504,10 +502,10 @@ inline void RegisterOpenFunction(void* OriginalFunction, void* NewFunction)
 }
 
 #if UE_AUTORTFM
-void DeferUntilCommit(std::function<void()>&& Work);
-void DeferUntilAbort(std::function<void()>&& Work);
-void OpenCommit(std::function<void()>&& Work);
-void OpenAbort(std::function<void()>&& Work);
+void DeferUntilCommit(TFunction<void()>&& Work);
+void DeferUntilAbort(TFunction<void()>&& Work);
+void OpenCommit(TFunction<void()>&& Work);
+void OpenAbort(TFunction<void()>&& Work);
 #else
 template<typename TFunctor>
 void DeferUntilCommit(const TFunctor& Work) { Work(); }
