fixed non-unity build compilation error that is not detected by CIS non-unity build jobs but popped up during preflighting other changes. e.g. https://horde.devtools.epicgames.com/job/641dd11a7a393e211a12e03b

diff --git a/Engine/Source/Runtime/Core/Public/Delegates/DelegateBase.h b/Engine/Source/Runtime/Core/Public/Delegates/DelegateBase.h
--- a/Engine/Source/Runtime/Core/Public/Delegates/DelegateBase.h
+++ b/Engine/Source/Runtime/Core/Public/Delegates/DelegateBase.h
@@ -8,6 +8,7 @@
 #include "UObject/NameTypes.h"
 #include "Delegates/DelegateSettings.h"
 #include "Delegates/IDelegateInstance.h"
+#include "Delegates/DelegateMacros.h"
 
 struct FWeakObjectPtr;
 
