Added PlatformManualResetEvent

This private abstraction is used to signal a single waiting thread to wake.
This first implementation is only specialized for Microsoft platforms. Other platforms will be specialized later.

#rnx

diff --git a/Engine/Platforms/Hololens/Source/Runtime/Core/Private/HoloLensPlatformManualResetEvent.h b/Engine/Platforms/Hololens/Source/Runtime/Core/Private/HoloLensPlatformManualResetEvent.h
new file mode 100644
--- /dev/null
+++ b/Engine/Platforms/Hololens/Source/Runtime/Core/Private/HoloLensPlatformManualResetEvent.h
@@ -0,0 +1,7 @@
+// Copyright Epic Games, Inc. All Rights Reserved.
+
+#pragma once
+
+#include "Microsoft/MicrosoftPlatformManualResetEvent.h"
+
+namespace UE { using FPlatformManualResetEvent = UE::HAL::Private::FMicrosoftPlatformManualResetEvent; }
diff --git a/Engine/Source/Runtime/Core/Private/Android/AndroidPlatformManualResetEvent.h b/Engine/Source/Runtime/Core/Private/Android/AndroidPlatformManualResetEvent.h
new file mode 100644
--- /dev/null
+++ b/Engine/Source/Runtime/Core/Private/Android/AndroidPlatformManualResetEvent.h
@@ -0,0 +1,7 @@
+// Copyright Epic Games, Inc. All Rights Reserved.
+
+#pragma once
+
+#include "GenericPlatform/GenericPlatformManualResetEvent.h"
+
+namespace UE { using FPlatformManualResetEvent = UE::HAL::Private::FGenericPlatformManualResetEvent; }
diff --git a/Engine/Source/Runtime/Core/Private/Apple/ApplePlatformManualResetEvent.h b/Engine/Source/Runtime/Core/Private/Apple/ApplePlatformManualResetEvent.h
new file mode 100644
--- /dev/null
+++ b/Engine/Source/Runtime/Core/Private/Apple/ApplePlatformManualResetEvent.h
@@ -0,0 +1,7 @@
+// Copyright Epic Games, Inc. All Rights Reserved.
+
+#pragma once
+
+#include "GenericPlatform/GenericPlatformManualResetEvent.h"
+
+namespace UE::HAL::Private { using FApplePlatformManualResetEvent = FGenericPlatformManualResetEvent; }
diff --git a/Engine/Source/Runtime/Core/Private/GenericPlatform/GenericPlatformManualResetEvent.h b/Engine/Source/Runtime/Core/Private/GenericPlatform/GenericPlatformManualResetEvent.h
new file mode 100644
--- /dev/null
+++ b/Engine/Source/Runtime/Core/Private/GenericPlatform/GenericPlatformManualResetEvent.h
@@ -0,0 +1,95 @@
+// Copyright Epic Games, Inc. All Rights Reserved.
+
+#pragma once
+
+#include "HAL/PlatformMath.h"
+#include "Misc/MonotonicTime.h"
+
+#include <atomic>
+#include <chrono>
+#include <condition_variable>
+#include <mutex>
+
+namespace UE::HAL::Private
+{
+
+/**
+ * A manual reset event that supports only one thread waiting and one thread notifying at a time.
+ *
+ * Only one waiting thread may call Reset() or the Wait() functions.
+ * Only one notifying thread may call Notify() once until the event is reset.
+ */
+class FGenericPlatformManualResetEvent
+{
+public:
+	FGenericPlatformManualResetEvent() = default;
+	FGenericPlatformManualResetEvent(const FGenericPlatformManualResetEvent&) = delete;
+	FGenericPlatformManualResetEvent& operator=(const FGenericPlatformManualResetEvent&) = delete;
+
+	/**
+	 * Resets the event to permit another Wait/Notify cycle.
+	 *
+	 * Must only be called by the waiting thread, and only when there is no possibility of waking
+	 * occurring concurrently with the reset.
+	 */
+	void Reset()
+	{
+		bWait = true;
+	}
+
+	/**
+	 * Waits for Notify() to be called.
+	 *
+	 * Notify() may be called prior to Wait(), and this will return immediately in that case.
+	 */
+	void Wait()
+	{
+		WaitUntil(FMonotonicTimePoint::Infinity());
+	}
+
+	/**
+	 * Waits until the wait time for Notify() to be called.
+	 *
+	 * Notify() may be called prior to WaitUntil(), and this will return immediately in that case.
+	 *
+	 * @param WaitTime   Absolute time after which waiting is canceled and the thread wakes.
+	 * @return True if Notify() was called before the wait time elapsed, otherwise false.
+	 */
+	bool WaitUntil(FMonotonicTimePoint WaitTime)
+	{
+		std::unique_lock SelfLock(Lock);
+		if (WaitTime.IsInfinity())
+		{
+			Condition.wait(SelfLock, [this] { return !bWait; });
+			return true;
+		}
+		if (FMonotonicTimeSpan WaitSpan = WaitTime - FMonotonicTimePoint::Now(); WaitSpan > FMonotonicTimeSpan::Zero())
+		{
+			const int64 WaitMs = FPlatformMath::CeilToInt64(WaitSpan.ToMilliseconds());
+			return Condition.wait_for(SelfLock, std::chrono::milliseconds(WaitMs), [this] { return !bWait; });
+		}
+		return !bWait;
+	}
+
+	/**
+	 * Notifies the waiting thread.
+	 *
+	 * Notify() may be called prior to one of the wait functions, and the eventual wait call will
+	 * return immediately when that occurs.
+	 */
+	void Notify()
+	{
+		{
+			std::unique_lock SelfLock(Lock);
+			bWait = false;
+		}
+		Condition.notify_one();
+	}
+
+private:
+	std::mutex Lock;
+	std::condition_variable Condition;
+	std::atomic<bool> bWait = true;
+};
+
+} // UE::HAL::Private
diff --git a/Engine/Source/Runtime/Core/Private/HAL/PlatformManualResetEvent.h b/Engine/Source/Runtime/Core/Private/HAL/PlatformManualResetEvent.h
new file mode 100644
--- /dev/null
+++ b/Engine/Source/Runtime/Core/Private/HAL/PlatformManualResetEvent.h
@@ -0,0 +1,7 @@
+// Copyright Epic Games, Inc. All Rights Reserved.
+
+#pragma once
+
+#include "CoreTypes.h"
+
+#include COMPILED_PLATFORM_HEADER(PlatformManualResetEvent.h)
diff --git a/Engine/Source/Runtime/Core/Private/IOS/IOSPlatformManualResetEvent.h b/Engine/Source/Runtime/Core/Private/IOS/IOSPlatformManualResetEvent.h
new file mode 100644
--- /dev/null
+++ b/Engine/Source/Runtime/Core/Private/IOS/IOSPlatformManualResetEvent.h
@@ -0,0 +1,7 @@
+// Copyright Epic Games, Inc. All Rights Reserved.
+
+#pragma once
+
+#include "Apple/ApplePlatformManualResetEvent.h"
+
+namespace UE { using FPlatformManualResetEvent = UE::HAL::Private::FApplePlatformManualResetEvent; }
diff --git a/Engine/Source/Runtime/Core/Private/Linux/LinuxPlatformManualResetEvent.h b/Engine/Source/Runtime/Core/Private/Linux/LinuxPlatformManualResetEvent.h
new file mode 100644
--- /dev/null
+++ b/Engine/Source/Runtime/Core/Private/Linux/LinuxPlatformManualResetEvent.h
@@ -0,0 +1,7 @@
+// Copyright Epic Games, Inc. All Rights Reserved.
+
+#pragma once
+
+#include "Unix/UnixPlatformManualResetEvent.h"
+
+namespace UE { using FPlatformManualResetEvent = UE::HAL::Private::FUnixPlatformManualResetEvent; }
diff --git a/Engine/Source/Runtime/Core/Private/Mac/MacPlatformManualResetEvent.h b/Engine/Source/Runtime/Core/Private/Mac/MacPlatformManualResetEvent.h
new file mode 100644
--- /dev/null
+++ b/Engine/Source/Runtime/Core/Private/Mac/MacPlatformManualResetEvent.h
@@ -0,0 +1,7 @@
+// Copyright Epic Games, Inc. All Rights Reserved.
+
+#pragma once
+
+#include "Apple/ApplePlatformManualResetEvent.h"
+
+namespace UE { using FPlatformManualResetEvent = UE::HAL::Private::FApplePlatformManualResetEvent; }
diff --git a/Engine/Source/Runtime/Core/Private/Microsoft/MicrosoftPlatformManualResetEvent.cpp b/Engine/Source/Runtime/Core/Private/Microsoft/MicrosoftPlatformManualResetEvent.cpp
new file mode 100644
--- /dev/null
+++ b/Engine/Source/Runtime/Core/Private/Microsoft/MicrosoftPlatformManualResetEvent.cpp
@@ -0,0 +1,55 @@
+// Copyright Epic Games, Inc. All Rights Reserved.
+
+#include "Microsoft/MicrosoftPlatformManualResetEvent.h"
+
+#include "HAL/PlatformMath.h"
+#include "Misc/MonotonicTime.h"
+
+#include "Microsoft/WindowsHWrapper.h"
+
+namespace UE::HAL::Private
+{
+
+void FMicrosoftPlatformManualResetEvent::Wait()
+{
+	WaitUntil(FMonotonicTimePoint::Infinity());
+}
+
+bool FMicrosoftPlatformManualResetEvent::WaitUntil(FMonotonicTimePoint WaitTime)
+{
+	bool bLocalWait = true;
+	if (WaitTime.IsInfinity())
+	{
+		for (;;)
+		{
+			if (WaitOnAddress(&bWait, &bLocalWait, sizeof(bool), INFINITE) && !bWait)
+			{
+				return true;
+			}
+		}
+	}
+	else
+	{
+		for (;;)
+		{
+			FMonotonicTimeSpan WaitSpan = WaitTime - FMonotonicTimePoint::Now();
+			if (WaitSpan <= FMonotonicTimeSpan::Zero())
+			{
+				return false;
+			}
+			const DWORD WaitMs = DWORD(FPlatformMath::CeilToInt64(WaitSpan.ToMilliseconds()));
+			if (WaitOnAddress(&bWait, &bLocalWait, sizeof(bool), WaitMs) && !bWait)
+			{
+				return true;
+			}
+		}
+	}
+}
+
+void FMicrosoftPlatformManualResetEvent::Notify()
+{
+	bWait = false;
+	WakeByAddressSingle((void*)&bWait);
+}
+
+} // UE::HAL::Private
diff --git a/Engine/Source/Runtime/Core/Private/Microsoft/MicrosoftPlatformManualResetEvent.h b/Engine/Source/Runtime/Core/Private/Microsoft/MicrosoftPlatformManualResetEvent.h
new file mode 100644
--- /dev/null
+++ b/Engine/Source/Runtime/Core/Private/Microsoft/MicrosoftPlatformManualResetEvent.h
@@ -0,0 +1,34 @@
+// Copyright Epic Games, Inc. All Rights Reserved.
+
+#pragma once
+
+#include "CoreTypes.h"
+
+#define UE_API CORE_API
+
+namespace UE { struct FMonotonicTimePoint; }
+
+namespace UE::HAL::Private
+{
+
+/** @see FGenericPlatformManualResetEvent */
+class FMicrosoftPlatformManualResetEvent
+{
+public:
+	FMicrosoftPlatformManualResetEvent() = default;
+	FMicrosoftPlatformManualResetEvent(const FMicrosoftPlatformManualResetEvent&) = delete;
+	FMicrosoftPlatformManualResetEvent& operator=(const FMicrosoftPlatformManualResetEvent&) = delete;
+
+	void Reset() { bWait = true; }
+
+	UE_API void Wait();
+	UE_API bool WaitUntil(FMonotonicTimePoint WaitTime);
+	UE_API void Notify();
+
+private:
+	bool bWait = true;
+};
+
+} // UE::HAL::Private
+
+#undef UE_API
diff --git a/Engine/Source/Runtime/Core/Private/Unix/UnixPlatformManualResetEvent.h b/Engine/Source/Runtime/Core/Private/Unix/UnixPlatformManualResetEvent.h
new file mode 100644
--- /dev/null
+++ b/Engine/Source/Runtime/Core/Private/Unix/UnixPlatformManualResetEvent.h
@@ -0,0 +1,7 @@
+// Copyright Epic Games, Inc. All Rights Reserved.
+
+#pragma once
+
+#include "GenericPlatform/GenericPlatformManualResetEvent.h"
+
+namespace UE::HAL::Private { using FUnixPlatformManualResetEvent = FGenericPlatformManualResetEvent; }
diff --git a/Engine/Source/Runtime/Core/Private/Windows/WindowsPlatformManualResetEvent.h b/Engine/Source/Runtime/Core/Private/Windows/WindowsPlatformManualResetEvent.h
new file mode 100644
--- /dev/null
+++ b/Engine/Source/Runtime/Core/Private/Windows/WindowsPlatformManualResetEvent.h
@@ -0,0 +1,7 @@
+// Copyright Epic Games, Inc. All Rights Reserved.
+
+#pragma once
+
+#include "Microsoft/MicrosoftPlatformManualResetEvent.h"
+
+namespace UE { using FPlatformManualResetEvent = UE::HAL::Private::FMicrosoftPlatformManualResetEvent; }
