Added EventCount

This is an event that provides a mechanism for avoiding a race between waiting and signaling through the use of a signal count token.

#rnx

diff --git a/Engine/Source/Runtime/Core/Public/Async/EventCount.h b/Engine/Source/Runtime/Core/Public/Async/EventCount.h
new file mode 100644
--- /dev/null
+++ b/Engine/Source/Runtime/Core/Public/Async/EventCount.h
@@ -0,0 +1,152 @@
+// Copyright Epic Games, Inc. All Rights Reserved.
+
+#pragma once
+
+#include "Async/ParkingLot.h"
+#include <atomic>
+#include <type_traits>
+
+namespace UE
+{
+
+template <typename CounterType>
+class TEventCount;
+
+/** A token used to wait on TEventCount. */
+template <typename CounterType>
+class TEventCountToken
+{
+	static_assert(std::is_unsigned_v<CounterType>);
+
+public:
+	/** Returns true if the token has been assigned by PrepareWait(). */
+	inline explicit operator bool() const { return !(Value & 1); }
+
+private:
+	/** Defaults to an odd value, which is never valid to wait on. */
+	CounterType Value = 1;
+
+	friend TEventCount<CounterType>;
+};
+
+/**
+ * A type of event that avoids missed notifications by maintaining a notification count.
+ *
+ * This type of event is suited to waiting on another thread conditionally.
+ * Typical usage looks similar to this example:
+ *
+ *     FEventCount Event;
+ *     std::atomic<uint32> CurrentValue = 0;
+ *
+ * On the waiting thread:
+ *
+ *     FEventCountToken Token = Event.PrepareWait();
+ *     if (CurrentValue < TargetValue)
+ *     {
+ *         Event.Wait(Token);
+ *     }
+ *
+ * On the notifying thread:
+ *
+ *     ++CurrentValue;
+ *     if (CurrentValue == TargetValue)
+ *     {
+ *         Event.Notify();
+ *     }
+ *
+ * Acquiring a token before checking the condition avoids a race because Wait returns immediately
+ * when the token no longer matches the notification count.
+ */
+template <typename CounterType>
+class TEventCount
+{
+	static_assert(std::is_unsigned_v<CounterType>);
+
+public:
+	constexpr TEventCount() = default;
+
+	TEventCount(const TEventCount&) = delete;
+	TEventCount& operator=(const TEventCount&) = delete;
+
+	/**
+	 * Prepare to wait.
+	 *
+	 * Call this before any logic that must re-execute if the event is notified in the meantime.
+	 *
+	 * @return A token to pass to one of the wait functions to abort the wait if the event has been notified since.
+	 */
+	inline TEventCountToken<CounterType> PrepareWait()
+	{
+		TEventCountToken<CounterType> Token;
+		Token.Value = Count.fetch_or(1, std::memory_order_relaxed) & ~CounterType(1);
+		return Token;
+	}
+
+	/**
+	 * Wait until the event is notified. Returns immediately if notified since the token was acquired.
+	 *
+	 * @param Compare   A token acquired from PrepareWait() before checking the conditions for this wait.
+	 */
+	inline void Wait(TEventCountToken<CounterType> Compare)
+	{
+		if ((Count.load(std::memory_order_acquire) & ~CounterType(1)) == Compare.Value)
+		{
+			ParkingLot::Wait(&Count, [this, Compare] { return (Count.load(std::memory_order_acquire) & ~CounterType(1)) == Compare.Value; }, []{});
+		}
+	}
+
+	/**
+	 * Wait until the event is notified. Returns immediately if notified since the token was acquired.
+	 *
+	 * @param Compare    A token acquired from PrepareWait() before checking the conditions for this wait.
+	 * @param WaitTime   Relative time after which waiting is automatically canceled and the thread wakes.
+	 * @return True if the event was notified before the wait time elapsed, otherwise false.
+	 */
+	inline bool WaitFor(TEventCountToken<CounterType> Compare, FMonotonicTimeSpan WaitTime)
+	{
+		if ((Count.load(std::memory_order_acquire) & ~CounterType(1)) == Compare.Value)
+		{
+			return ParkingLot::WaitFor(&Count, [this, Compare] { return (Count.load(std::memory_order_acquire) & ~CounterType(1)) == Compare.Value; }, []{}, WaitTime).bDidWake;
+		}
+		return true;
+	}
+
+	/**
+	 * Wait until the event is notified. Returns immediately if notified since the token was acquired.
+	 *
+	 * @param Compare    A token acquired from PrepareWait() before checking the conditions for this wait.
+	 * @param WaitTime   Absolute time after which waiting is automatically canceled and the thread wakes.
+	 * @return True if the event was notified before the wait time elapsed, otherwise false.
+	 */
+	inline bool WaitUntil(TEventCountToken<CounterType> Compare, FMonotonicTimePoint WaitTime)
+	{
+		if ((Count.load(std::memory_order_acquire) & ~CounterType(1)) == Compare.Value)
+		{
+			return ParkingLot::WaitUntil(&Count, [this, Compare] { return (Count.load(std::memory_order_acquire) & ~CounterType(1)) == Compare.Value; }, []{}, WaitTime).bDidWake;
+		}
+		return true;
+	}
+
+	/**
+	 * Notifies all waiting threads.
+	 *
+	 * Any threads that have called PrepareWait() and not yet waited will be notified immediately
+	 * if they do wait on a token from a call to PrepareWait() that preceded this call.
+	 */
+	inline void Notify()
+	{
+		CounterType Value = Count.load(std::memory_order_relaxed);
+		if ((Value & 1) && Count.compare_exchange_strong(Value, Value + 1, std::memory_order_release))
+		{
+			ParkingLot::WakeAll(&Count);
+		}
+	}
+
+private:
+	std::atomic<CounterType> Count = 0;
+};
+
+using FEventCount = TEventCount<uint32>;
+using FEventCountToken = TEventCountToken<uint32>;
+
+} // UE
