Windows platform extensions do not need to override platform process or dynamic rhi.

#rnx

diff --git a/Engine/Source/Runtime/Core/Public/Windows/WindowsPlatform.h b/Engine/Source/Runtime/Core/Public/Windows/WindowsPlatform.h
--- a/Engine/Source/Runtime/Core/Public/Windows/WindowsPlatform.h
+++ b/Engine/Source/Runtime/Core/Public/Windows/WindowsPlatform.h
@@ -110,11 +110,10 @@ typedef FWindowsPlatformTypes FPlatformTypes;
 #define PLATFORM_RETURN_ADDRESS_FOR_CALLSTACKTRACING		PLATFORM_RETURN_ADDRESS_POINTER
 
 #define WINDOWS_USE_FEATURE_APPLICATIONMISC_CLASS			1
-#define WINDOWS_USE_FEATURE_PLATFORMPROCESS_CLASS			1
 #define WINDOWS_USE_FEATURE_PLATFORMMISC_CLASS				1
 #define WINDOWS_USE_FEATURE_PLATFORMHTTP_CLASS				1
 #define WINDOWS_USE_FEATURE_LAUNCH							1
-#define WINDOWS_USE_FEATURE_DYNAMIC_RHI						1
+
 
 // Q: Why is there a __nop() before __debugbreak()?
 // A: VS' debug engine has a bug where it will silently swallow explicit
diff --git a/Engine/Source/Runtime/Core/Public/Windows/WindowsPlatformProcess.h b/Engine/Source/Runtime/Core/Public/Windows/WindowsPlatformProcess.h
--- a/Engine/Source/Runtime/Core/Public/Windows/WindowsPlatformProcess.h
+++ b/Engine/Source/Runtime/Core/Public/Windows/WindowsPlatformProcess.h
@@ -263,7 +263,5 @@ private:
 	static void ResolveMissingImportsRecursive(const FString& FileName, const TArray<FString>& SearchPaths, TArray<FString>& ImportFileNames, TSet<FString>& VisitedImportNames);
 };
 
-#if WINDOWS_USE_FEATURE_PLATFORMPROCESS_CLASS
 typedef FWindowsPlatformProcess FPlatformProcess;
-#endif
 
diff --git a/Engine/Source/Runtime/RHI/Private/Windows/WindowsDynamicRHI.cpp b/Engine/Source/Runtime/RHI/Private/Windows/WindowsDynamicRHI.cpp
--- a/Engine/Source/Runtime/RHI/Private/Windows/WindowsDynamicRHI.cpp
+++ b/Engine/Source/Runtime/RHI/Private/Windows/WindowsDynamicRHI.cpp
@@ -9,8 +9,6 @@
 #include "Containers/StaticArray.h"
 #include "DataDrivenShaderPlatformInfo.h"
 
-#if WINDOWS_USE_FEATURE_DYNAMIC_RHI
-
 #include "Windows/WindowsPlatformApplicationMisc.h"
 
 #define LOCTEXT_NAMESPACE "WindowsDynamicRHI"
@@ -1016,4 +1014,3 @@ const TCHAR* GetSelectedDynamicRHIModuleName(bool bCleanup)
 
 #undef LOCTEXT_NAMESPACE
 
-#endif //WINDOWS_USE_FEATURE_DYNAMIC_RHI
