ParkingLot: Re-added use of PlatformManualResetEvent and WordMutex

#rnx

diff --git a/Engine/Source/Runtime/Core/Private/Async/ParkingLot.cpp b/Engine/Source/Runtime/Core/Private/Async/ParkingLot.cpp
--- a/Engine/Source/Runtime/Core/Private/Async/ParkingLot.cpp
+++ b/Engine/Source/Runtime/Core/Private/Async/ParkingLot.cpp
@@ -4,35 +4,17 @@
 
 #include "Algo/Sort.h"
 #include "Async/UniqueLock.h"
+#include "Async/WordMutex.h"
 #include "Containers/Array.h"
 #include "Containers/ArrayView.h"
+#include "HAL/PlatformManualResetEvent.h"
 #include "Templates/RefCounting.h"
 #include <atomic>
 
-#define UE_PARKINGLOT_USE_WAITONADDRESS PLATFORM_WINDOWS
-#define UE_PARKINGLOT_USE_WORDMUTEX 0
-
-#if UE_PARKINGLOT_USE_WORDMUTEX
-#include "Async/WordMutex.h"
-#else
-#include "HAL/CriticalSection.h"
-#endif
-
-#if UE_PARKINGLOT_USE_WAITONADDRESS
-#include "Microsoft/WindowsHWrapper.h"
-#else
-#include <condition_variable>
-#include <mutex>
-#endif
-
 namespace UE::ParkingLot::Private
 {
 
-#if UE_PARKINGLOT_USE_WORDMUTEX
-using FBucketMutex = UE::Private::FWordMutex;
-#else
-using FBucketMutex = FCriticalSection;
-#endif
+using FBucketMutex = FWordMutex;
 
 ///////////////////////////////////////////////////////////////////////////////////////////////////////////////////////
 
@@ -42,16 +24,9 @@ using FBucketMutex = FCriticalSection;
 struct FThread final
 {
 	FThread* Next = nullptr;
-
-#if UE_PARKINGLOT_USE_WAITONADDRESS
-	const void* volatile WaitAddress = nullptr;
-#else
 	const void* WaitAddress = nullptr;
-	std::mutex Lock;
-	std::condition_variable Condition;
-#endif
-
 	uint64 WakeToken = 0;
+	FPlatformManualResetEvent Event;
 	mutable std::atomic<uint32> ReferenceCount = 0;
 
 	inline void AddRef() const
@@ -540,58 +515,24 @@ FWaitState WaitUntil(const void* Address, TFunctionRef<bool()> CanWait, TFunctio
 			return State;
 		}
 		Self.WaitAddress = Address;
+		Self.Event.Reset();
 		Bucket.Enqueue(&Self);
 	}
 
 	// BeforeWait must be invoked after the bucket is unlocked.
 	BeforeWait();
 
-#if UE_PARKINGLOT_USE_WAITONADDRESS
-	for (;;)
-	{
-		const FMonotonicTimeSpan WaitSpan = WaitTime - FMonotonicTimePoint::Now();
-		if (WaitSpan <= FMonotonicTimeSpan::Zero())
-		{
-			break;
-		}
-		const DWORD WaitMs = WaitSpan.IsInfinity() ? INFINITE : DWORD(FMath::CeilToInt64(WaitSpan.ToMilliseconds()));
-		if (!WaitOnAddress(&Self.WaitAddress, &Address, sizeof(Address), WaitMs) && GetLastError() == ERROR_TIMEOUT)
-		{
-			break;
-		}
-		if (!Self.WaitAddress)
-		{
-			State.bDidWake = true;
-			State.WakeToken = Self.WakeToken;
-			Self.WakeToken = 0;
-			return State;
-		}
-	}
-#else
-	{
-		// Wait until the timeout or until the thread has been dequeued.
-		std::unique_lock SelfLock(Self.Lock);
-		const FMonotonicTimeSpan WaitSpan = WaitTime - FMonotonicTimePoint::Now();
-		if (WaitSpan.IsInfinity())
-		{
-			Self.Condition.wait(SelfLock, [&Self] { return !Self.WaitAddress; });
-		}
-		else if (WaitSpan > FMonotonicTimeSpan::Zero())
-		{
-			const int64 WaitMs = FMath::CeilToInt64(WaitSpan.ToMilliseconds());
-			Self.Condition.wait_for(SelfLock, std::chrono::milliseconds(WaitMs), [&Self] { return !Self.WaitAddress; });
-		}
+	// Wait until the timeout or until the thread has been dequeued.
+	Self.Event.WaitUntil(WaitTime);
 
-		// WaitAddress is reset when the thread is dequeued.
-		if (!Self.WaitAddress)
-		{
-			State.bDidWake = true;
-			State.WakeToken = Self.WakeToken;
-			Self.WakeToken = 0;
-			return State;
-		}
+	// WaitAddress is reset when the thread is dequeued.
+	if (!Self.WaitAddress)
+	{
+		State.bDidWake = true;
+		State.WakeToken = Self.WakeToken;
+		Self.WakeToken = 0;
+		return State;
 	}
-#endif
 
 	// The timeout was reached and the thread needs to dequeue itself.
 	// This can race with a call to wake a thread, which means Self is unsafe to access outside of the lock.
@@ -614,15 +555,7 @@ FWaitState WaitUntil(const void* Address, TFunctionRef<bool()> CanWait, TFunctio
 	// has finished waking this thread by setting its wait address to null.
 	if (!bDequeued)
 	{
-	#if UE_PARKINGLOT_USE_WAITONADDRESS
-		while (Self.WaitAddress)
-		{
-			WaitOnAddress(&Self.WaitAddress, &Address, sizeof(Address), INFINITE);
-		}
-	#else
-		std::unique_lock SelfLock(Self.Lock);
-		Self.Condition.wait(SelfLock, [&Self] { return !Self.WaitAddress; });
-	#endif
+		Self.Event.Wait();
 		State.bDidWake = true;
 		State.WakeToken = Self.WakeToken;
 		Self.WakeToken = 0;
@@ -660,16 +593,8 @@ void WakeOne(const void* Address, TFunctionRef<uint64(FWakeState)> OnWakeState)
 	{
 		checkSlow(WakeThread->WaitAddress == Address);
 		WakeThread->WakeToken = WakeToken;
-	#if UE_PARKINGLOT_USE_WAITONADDRESS
 		WakeThread->WaitAddress = nullptr;
-		WakeByAddressSingle((void*)&WakeThread->WaitAddress);
-	#else
-		{
-			std::unique_lock WakeThreadLock(WakeThread->Lock);
-			WakeThread->WaitAddress = nullptr;
-		}
-		WakeThread->Condition.notify_one();
-	#endif
+		WakeThread->Event.Notify();
 	}
 }
 
@@ -711,16 +636,8 @@ uint32 WakeMultiple(const void* const Address, const uint32 WakeCount)
 	for (FThread* WakeThread : WakeThreads)
 	{
 		checkSlow(WakeThread->WaitAddress == Address);
-	#if UE_PARKINGLOT_USE_WAITONADDRESS
 		WakeThread->WaitAddress = nullptr;
-		WakeByAddressSingle((void*)&WakeThread->WaitAddress);
-	#else
-		{
-			std::unique_lock WakeThreadLock(WakeThread->Lock);
-			WakeThread->WaitAddress = nullptr;
-		}
-		WakeThread->Condition.notify_one();
-	#endif
+		WakeThread->Event.Notify();
 	}
 
 	return uint32(WakeThreads.Num());
