Fixing up include-errors in non-unity builds.

diff --git a/Engine/Source/Runtime/Core/Private/AutoRTFM/API.cpp b/Engine/Source/Runtime/Core/Private/AutoRTFM/API.cpp
--- a/Engine/Source/Runtime/Core/Private/AutoRTFM/API.cpp
+++ b/Engine/Source/Runtime/Core/Private/AutoRTFM/API.cpp
@@ -13,6 +13,8 @@
 #include "TransactionInlines.h"
 #include "Utils.h"
 
+#include "Templates/Tuple.h"
+
 // This is the implementation of the AutoRTFM.h API. Ideally, functions here should just delegate to some internal API.
 // For now, I have these functions also perform some error checking.
 
diff --git a/Engine/Source/Runtime/Core/Private/AutoRTFM/Context.cpp b/Engine/Source/Runtime/Core/Private/AutoRTFM/Context.cpp
--- a/Engine/Source/Runtime/Core/Private/AutoRTFM/Context.cpp
+++ b/Engine/Source/Runtime/Core/Private/AutoRTFM/Context.cpp
@@ -9,6 +9,9 @@
 #include "ScopedGuard.h"
 #include "TransactionInlines.h"
 
+#include "Templates/UniquePtr.h"
+#include "Containers/StringConv.h"
+
 namespace AutoRTFM
 {
 
diff --git a/Engine/Source/Runtime/Core/Private/AutoRTFM/FunctionMap.cpp b/Engine/Source/Runtime/Core/Private/AutoRTFM/FunctionMap.cpp
--- a/Engine/Source/Runtime/Core/Private/AutoRTFM/FunctionMap.cpp
+++ b/Engine/Source/Runtime/Core/Private/AutoRTFM/FunctionMap.cpp
@@ -9,6 +9,8 @@
 #include <stdlib.h>
 #include <string.h>
 
+#include "Containers/Array.h"
+
 namespace AutoRTFM
 {
 
diff --git a/Engine/Source/Runtime/Core/Private/AutoRTFM/FunctionMapInlines.h b/Engine/Source/Runtime/Core/Private/AutoRTFM/FunctionMapInlines.h
--- a/Engine/Source/Runtime/Core/Private/AutoRTFM/FunctionMapInlines.h
+++ b/Engine/Source/Runtime/Core/Private/AutoRTFM/FunctionMapInlines.h
@@ -5,6 +5,8 @@
 #include "Context.h"
 #include "FunctionMap.h"
 
+#include "Containers/StringConv.h"
+
 namespace AutoRTFM
 {
 
diff --git a/Engine/Source/Runtime/Core/Private/AutoRTFM/TaskArray.h b/Engine/Source/Runtime/Core/Private/AutoRTFM/TaskArray.h
--- a/Engine/Source/Runtime/Core/Private/AutoRTFM/TaskArray.h
+++ b/Engine/Source/Runtime/Core/Private/AutoRTFM/TaskArray.h
@@ -3,6 +3,7 @@
 #pragma once
 
 #include "SharedPtr.h"
+#include "Containers/Array.h"
 
 namespace AutoRTFM
 {
diff --git a/Engine/Source/Runtime/Core/Private/AutoRTFM/Utils.cpp b/Engine/Source/Runtime/Core/Private/AutoRTFM/Utils.cpp
--- a/Engine/Source/Runtime/Core/Private/AutoRTFM/Utils.cpp
+++ b/Engine/Source/Runtime/Core/Private/AutoRTFM/Utils.cpp
@@ -13,6 +13,8 @@
 #include <execinfo.h>
 #endif
 
+#include "Misc/CoreMiscDefines.h"
+
 namespace AutoRTFM
 {
 
diff --git a/Engine/Source/Runtime/Core/Private/AutoRTFM/Utils.h b/Engine/Source/Runtime/Core/Private/AutoRTFM/Utils.h
--- a/Engine/Source/Runtime/Core/Private/AutoRTFM/Utils.h
+++ b/Engine/Source/Runtime/Core/Private/AutoRTFM/Utils.h
@@ -8,6 +8,8 @@
 #include <stdio.h>
 #include <stdlib.h>
 
+#include "Containers/UnrealString.h"
+
 #define ASSERT(exp) do { \
     if (!(exp)) { \
         fprintf(stderr, "%s:%d:%s: assertion %s failed.\n", __FILE__, __LINE__, __PRETTY_FUNCTION__, #exp); \
diff --git a/Engine/Source/Runtime/Core/Public/AutoRTFM/AutoRTFM.h b/Engine/Source/Runtime/Core/Public/AutoRTFM/AutoRTFM.h
--- a/Engine/Source/Runtime/Core/Public/AutoRTFM/AutoRTFM.h
+++ b/Engine/Source/Runtime/Core/Public/AutoRTFM/AutoRTFM.h
@@ -16,6 +16,9 @@
 
 #include <stdbool.h>
 #include <stdlib.h>
+#include <memory.h>
+
+#include "Templates/Function.h"
 
 #ifdef __cplusplus
 extern "C"
