Fix -nopch include issues and move std::forward to UEs Forward

diff --git a/Engine/Source/Runtime/Core/Private/AutoRTFM/SharedPtr.h b/Engine/Source/Runtime/Core/Private/AutoRTFM/SharedPtr.h
--- a/Engine/Source/Runtime/Core/Private/AutoRTFM/SharedPtr.h
+++ b/Engine/Source/Runtime/Core/Private/AutoRTFM/SharedPtr.h
@@ -71,7 +71,7 @@ public:
     static TSharedPtr New(TArguments&&... Arguments)
     {
         TSharedPtr Result;
-        Result.Ptr = new T(std::forward<TArguments>(Arguments)...);
+        Result.Ptr = new T(Forward<TArguments>(Arguments)...);
         Result.Count = new unsigned(1);
         return Result;
     }
diff --git a/Engine/Source/Runtime/Core/Private/AutoRTFM/Utils.cpp b/Engine/Source/Runtime/Core/Private/AutoRTFM/Utils.cpp
--- a/Engine/Source/Runtime/Core/Private/AutoRTFM/Utils.cpp
+++ b/Engine/Source/Runtime/Core/Private/AutoRTFM/Utils.cpp
@@ -13,6 +13,7 @@
 #include <execinfo.h>
 #endif
 
+#include "Containers/StringConv.h"
 #include "Misc/CoreMiscDefines.h"
 
 namespace AutoRTFM
