Fixed ManualResetEvent waking spuriously

#rnx

diff --git a/Engine/Source/Runtime/Core/Public/Async/ManualResetEvent.h b/Engine/Source/Runtime/Core/Public/Async/ManualResetEvent.h
--- a/Engine/Source/Runtime/Core/Public/Async/ManualResetEvent.h
+++ b/Engine/Source/Runtime/Core/Public/Async/ManualResetEvent.h
@@ -54,21 +54,21 @@ public:
 	 */
 	inline bool WaitUntil(FMonotonicTimePoint WaitTime)
 	{
-		for (uint8 CurrentState = State.load(std::memory_order_acquire);;)
+		for (;;)
 		{
+			uint8 CurrentState = State.load(std::memory_order_acquire);
 			if (CurrentState & IsNotifiedFlag)
 			{
 				return true;
 			}
-
+			if (WaitTime <= FMonotonicTimePoint::Now())
+			{
+				return false;
+			}
 			if ((CurrentState & MaybeWaitingFlag) ||
 				State.compare_exchange_weak(CurrentState, MaybeWaitingFlag, std::memory_order_relaxed, std::memory_order_acquire))
 			{
-				const ParkingLot::FWaitState WaitState = ParkingLot::WaitUntil(&State, [this, &CurrentState]
-				{
-					return !(State.load(std::memory_order_acquire) & IsNotifiedFlag);
-				}, []{}, WaitTime);
-				return !WaitState.bDidWait || WaitState.bDidWake;
+				ParkingLot::WaitUntil(&State, [this] { return !IsNotified(); }, []{}, WaitTime);
 			}
 		}
 	}
@@ -78,7 +78,7 @@ public:
 	 */
 	inline void Notify()
 	{
-		const uint8 CurrentState = State.fetch_or(IsNotifiedFlag, std::memory_order_release);
+		const uint8 CurrentState = State.exchange(IsNotifiedFlag, std::memory_order_release);
 		if (!(CurrentState & IsNotifiedFlag) && (CurrentState & MaybeWaitingFlag))
 		{
 			ParkingLot::WakeAll(&State);
