Added a check for infinity in FManualResetEvent::WaitUntil

This avoids calling FMonotonicTimePoint::Now() in the common case of waiting indefinitely.

#rnx

diff --git a/Engine/Source/Runtime/Core/Public/Async/ManualResetEvent.h b/Engine/Source/Runtime/Core/Public/Async/ManualResetEvent.h
--- a/Engine/Source/Runtime/Core/Public/Async/ManualResetEvent.h
+++ b/Engine/Source/Runtime/Core/Public/Async/ManualResetEvent.h
@@ -61,7 +61,7 @@ public:
 			{
 				return true;
 			}
-			if (WaitTime <= FMonotonicTimePoint::Now())
+			if (!WaitTime.IsInfinity() && WaitTime <= FMonotonicTimePoint::Now())
 			{
 				return false;
 			}
