Extend the functionality of Automation.SkipStackWalk CVar to be any stack walk that FAutomationTestFramework triggers.

diff --git a/Engine/Source/Runtime/Core/Private/Misc/AutomationTest.cpp b/Engine/Source/Runtime/Core/Private/Misc/AutomationTest.cpp
--- a/Engine/Source/Runtime/Core/Private/Misc/AutomationTest.cpp
+++ b/Engine/Source/Runtime/Core/Private/Misc/AutomationTest.cpp
@@ -30,11 +30,11 @@ namespace AutomationTest
 		bCaptureLogEvents,
 		TEXT("Consider warning/error log events during a test as impacting the test itself"));
 
-	static bool bSkipStackWalk = true;
+	static bool bSkipStackWalk = false;
 	static FAutoConsoleVariableRef CVarAutomationSkipStackWalk(
 		TEXT("Automation.SkipStackWalk"),
 		bSkipStackWalk,
-		TEXT("Whether to skip stack walk while iterating for listing the tests"));
+		TEXT("Whether to skip any stack issues that the automation test framework triggers"));
 
 	static bool bLogBPTestMetadata = false;
 	static FAutoConsoleVariableRef CVarAutomationLogBPTestMetadata(
diff --git a/Engine/Source/Runtime/Core/Public/Misc/AutomationTest.h b/Engine/Source/Runtime/Core/Public/Misc/AutomationTest.h
--- a/Engine/Source/Runtime/Core/Public/Misc/AutomationTest.h
+++ b/Engine/Source/Runtime/Core/Public/Misc/AutomationTest.h
@@ -1004,9 +1004,13 @@ public:
 
 	bool GetCaptureStack() const
 	{
-		return bCaptureStack;
+		return bCaptureStack && !NeedSkipStackWalk();
 	}
 
+	/**
+	 * Used to disabled stack capture when an error or warning event is triggered.
+	 * Setting bCapture=true does not guarantees GetCaptureStack()=true because that method also depend on NeedSkipStackWalk(). 
+	 */
 	void SetCaptureStack(bool bCapture)
 	{
 		bCaptureStack = bCapture;
