removed core delegates deprecated in 5.2. while it's a bit too early for a normal deprecation removal, these delegates are known race conditions so it's rather an urgency

diff --git a/Engine/Plugins/Online/OnlineFramework/Source/Hotfix/Private/OnlineHotfixManager.cpp b/Engine/Plugins/Online/OnlineFramework/Source/Hotfix/Private/OnlineHotfixManager.cpp
--- a/Engine/Plugins/Online/OnlineFramework/Source/Hotfix/Private/OnlineHotfixManager.cpp
+++ b/Engine/Plugins/Online/OnlineFramework/Source/Hotfix/Private/OnlineHotfixManager.cpp
@@ -1113,9 +1113,6 @@ bool UOnlineHotfixManager::HotfixIniFile(const FString& FileName, const FString&
 
 	const FString ConfigFileName = ConfigFile->Name.ToString();
 	FCoreDelegates::TSOnConfigSectionsChanged().Broadcast(ConfigFileName, UpdatedSectionNames);
-PRAGMA_DISABLE_DEPRECATION_WARNINGS
-	FCoreDelegates::OnConfigSectionsChanged.Broadcast(ConfigFileName, UpdatedSectionNames);
-PRAGMA_ENABLE_DEPRECATION_WARNINGS
 
 	// Reload log suppression if configs changed
 	if (bUpdateLogSuppression)
diff --git a/Engine/Source/Developer/PakFileUtilities/Private/PakFileUtilities.cpp b/Engine/Source/Developer/PakFileUtilities/Private/PakFileUtilities.cpp
--- a/Engine/Source/Developer/PakFileUtilities/Private/PakFileUtilities.cpp
+++ b/Engine/Source/Developer/PakFileUtilities/Private/PakFileUtilities.cpp
@@ -5415,9 +5415,6 @@ bool MakeBinaryConfig(const TCHAR* CmdLine)
 	// allow delegates to modify the config data with some tagged binary data
 	FCoreDelegates::FExtraBinaryConfigData ExtraData(Config, true);
 	FCoreDelegates::TSAccessExtraBinaryConfigData().Broadcast(ExtraData);
-PRAGMA_DISABLE_DEPRECATION_WARNINGS
-	FCoreDelegates::AccessExtraBinaryConfigData.Broadcast(ExtraData);
-PRAGMA_ENABLE_DEPRECATION_WARNINGS
 
 	// write it all out!
 	TArray<uint8> FileContent;
diff --git a/Engine/Source/Editor/UnrealEd/Private/Commandlets/MakeBinaryConfigCommandlet.cpp b/Engine/Source/Editor/UnrealEd/Private/Commandlets/MakeBinaryConfigCommandlet.cpp
--- a/Engine/Source/Editor/UnrealEd/Private/Commandlets/MakeBinaryConfigCommandlet.cpp
+++ b/Engine/Source/Editor/UnrealEd/Private/Commandlets/MakeBinaryConfigCommandlet.cpp
@@ -100,9 +100,6 @@ int32 UMakeBinaryConfigCommandlet::Main(const FString& Params)
 	// allow delegates to modify the config data with some tagged binary data
 	FCoreDelegates::FExtraBinaryConfigData ExtraData(Config, true);
 	FCoreDelegates::TSAccessExtraBinaryConfigData().Broadcast(ExtraData);
-PRAGMA_DISABLE_DEPRECATION_WARNINGS
-	FCoreDelegates::AccessExtraBinaryConfigData.Broadcast(ExtraData);
-PRAGMA_ENABLE_DEPRECATION_WARNINGS
 
 	// write it all out!
 	TArray<uint8> FileContent;
diff --git a/Engine/Source/Runtime/Core/Private/Misc/ConfigCacheIni.cpp b/Engine/Source/Runtime/Core/Private/Misc/ConfigCacheIni.cpp
--- a/Engine/Source/Runtime/Core/Private/Misc/ConfigCacheIni.cpp
+++ b/Engine/Source/Runtime/Core/Private/Misc/ConfigCacheIni.cpp
@@ -426,9 +426,6 @@ static void FixupArrayOfStructKeysForSection(FConfigSection* Section, const FStr
 	// will any delegates return contents via TSPreLoadConfigFileDelegate()?
 	int32 ResponderCount = 0;
 	FCoreDelegates::TSCountPreLoadConfigFileRespondersDelegate().Broadcast(IniFile, ResponderCount);
-PRAGMA_DISABLE_DEPRECATION_WARNINGS
-	FCoreDelegates::CountPreLoadConfigFileRespondersDelegate.Broadcast(IniFile, ResponderCount);
-PRAGMA_ENABLE_DEPRECATION_WARNINGS
 
 	if (ResponderCount > 0)
 	{
@@ -475,9 +472,6 @@ static bool LoadConfigFileWrapper(const TCHAR* IniFile, FString& Contents, bool
 {
 	// let other systems load the file instead of the standard load below
 	FCoreDelegates::TSPreLoadConfigFileDelegate().Broadcast(IniFile, Contents);
-PRAGMA_DISABLE_DEPRECATION_WARNINGS
-	FCoreDelegates::PreLoadConfigFileDelegate.Broadcast(IniFile, Contents);
-PRAGMA_ENABLE_DEPRECATION_WARNINGS
 
 	// if this loaded any text, we are done, and we won't override the contents with standard ini file data
 	if (Contents.Len())
@@ -511,9 +505,6 @@ static bool SaveConfigFileWrapper(const TCHAR* IniFile, const FString& Contents)
 	// let anyone that needs to save it, do so (counting how many did)
 	int32 SavedCount = 0;
 	FCoreDelegates::TSPreSaveConfigFileDelegate().Broadcast(IniFile, Contents, SavedCount);
-PRAGMA_DISABLE_DEPRECATION_WARNINGS
-	FCoreDelegates::PreSaveConfigFileDelegate.Broadcast(IniFile, Contents, SavedCount);
-PRAGMA_ENABLE_DEPRECATION_WARNINGS
 
 	// save it even if a delegate did as well
 	bool bLocalWriteSucceeded = false;
@@ -551,9 +542,6 @@ FConfigFile::FConfigFile()
 , SourceConfigFile(nullptr)
 {
 	FCoreDelegates::TSOnFConfigCreated().Broadcast(this);
-PRAGMA_DISABLE_DEPRECATION_WARNINGS
-	FCoreDelegates::OnFConfigCreated.Broadcast(this);
-PRAGMA_ENABLE_DEPRECATION_WARNINGS
 }
 
 FConfigFile::~FConfigFile()
@@ -563,9 +551,6 @@ FConfigFile::~FConfigFile()
 	if ( !GExitPurge )
 	{
 		FCoreDelegates::TSOnFConfigDeleted().Broadcast(this);
-PRAGMA_DISABLE_DEPRECATION_WARNINGS
-		FCoreDelegates::OnFConfigDeleted.Broadcast(this);
-PRAGMA_ENABLE_DEPRECATION_WARNINGS
 	}
 
 	delete SourceConfigFile;
@@ -2511,9 +2496,6 @@ bool FConfigCacheIni::GetString( const TCHAR* Section, const TCHAR* Key, FString
 	Value = ConfigValue->GetValue();
 
 	FCoreDelegates::TSOnConfigValueRead().Broadcast(*Filename, Section, Key);
-PRAGMA_DISABLE_DEPRECATION_WARNINGS
-	FCoreDelegates::OnConfigValueRead.Broadcast(*Filename, Section, Key);
-PRAGMA_ENABLE_DEPRECATION_WARNINGS
 
 	return true;
 }
@@ -2545,9 +2527,6 @@ bool FConfigCacheIni::GetText( const TCHAR* Section, const TCHAR* Key, FText& Va
 	}
 
 	FCoreDelegates::TSOnConfigValueRead().Broadcast(*Filename, Section, Key);
-PRAGMA_DISABLE_DEPRECATION_WARNINGS
-	FCoreDelegates::OnConfigValueRead.Broadcast(*Filename, Section, Key);
-PRAGMA_ENABLE_DEPRECATION_WARNINGS
 
 	return true;
 }
@@ -2573,9 +2552,6 @@ bool FConfigCacheIni::GetSection( const TCHAR* Section, TArray<FString>& Result,
 	}
 
 	FCoreDelegates::TSOnConfigSectionRead().Broadcast(*Filename, Section);
-PRAGMA_DISABLE_DEPRECATION_WARNINGS
-	FCoreDelegates::OnConfigSectionRead.Broadcast(*Filename, Section);
-PRAGMA_ENABLE_DEPRECATION_WARNINGS
 
 	return true;
 }
@@ -2618,9 +2594,6 @@ bool FConfigCacheIni::DoesSectionExist(const TCHAR* Section, const FString& File
 	if (bReturnVal)
 	{
 		FCoreDelegates::TSOnConfigSectionNameRead().Broadcast(*Filename, Section);
-PRAGMA_DISABLE_DEPRECATION_WARNINGS
-		FCoreDelegates::OnConfigSectionNameRead.Broadcast(*Filename, Section);
-PRAGMA_ENABLE_DEPRECATION_WARNINGS
 	}
 
 	return bReturnVal;
@@ -2782,9 +2755,6 @@ bool FConfigCacheIni::GetSectionNames( const FString& Filename, TArray<FString>&
 			out_SectionNames.Add(It.Key());
 
 			FCoreDelegates::TSOnConfigSectionNameRead().Broadcast(*Filename, *It.Key());
-PRAGMA_DISABLE_DEPRECATION_WARNINGS
-			FCoreDelegates::OnConfigSectionNameRead.Broadcast(*Filename, *It.Key());
-PRAGMA_ENABLE_DEPRECATION_WARNINGS
 		}
 		bResult = true;
 	}
@@ -2828,9 +2798,6 @@ bool FConfigCacheIni::GetPerObjectConfigSections( const FString& Filename, const
 					bResult = true;
 
 					FCoreDelegates::TSOnConfigSectionNameRead().Broadcast(*Filename, *SectionName);
-PRAGMA_DISABLE_DEPRECATION_WARNINGS
-					FCoreDelegates::OnConfigSectionNameRead.Broadcast(*Filename, *SectionName);
-PRAGMA_ENABLE_DEPRECATION_WARNINGS
 				}
 			}
 		}
@@ -2996,9 +2963,6 @@ int32 FConfigCacheIni::GetArray
 	if (out_Arr.Num())
 	{
 		FCoreDelegates::TSOnConfigValueRead().Broadcast(*Filename, Section, Key);
-PRAGMA_DISABLE_DEPRECATION_WARNINGS
-		FCoreDelegates::OnConfigValueRead.Broadcast(*Filename, Section, Key);
-PRAGMA_ENABLE_DEPRECATION_WARNINGS
 	}
 
 	return out_Arr.Num();
@@ -3719,9 +3683,6 @@ bool FConfigCacheIni::CreateGConfigFromSaved(const TCHAR* Filename)
 	{
 		TRACE_CPUPROFILER_EVENT_SCOPE(ConfigReadyForUseBroadcast);
 		FCoreDelegates::TSConfigReadyForUse().Broadcast();
-PRAGMA_DISABLE_DEPRECATION_WARNINGS
-		FCoreDelegates::ConfigReadyForUse.Broadcast();
-PRAGMA_ENABLE_DEPRECATION_WARNINGS
 	}
 
 	FMemory::Free(PreloadedData);
@@ -3838,9 +3799,6 @@ void FConfigCacheIni::InitializeConfigSystem()
 			GConfig->bIsReadyForUse = true;
 			TRACE_CPUPROFILER_EVENT_SCOPE(ConfigReadyForUseBroadcast);
 			FCoreDelegates::TSConfigReadyForUse().Broadcast();
-PRAGMA_DISABLE_DEPRECATION_WARNINGS
-			FCoreDelegates::ConfigReadyForUse.Broadcast();
-PRAGMA_ENABLE_DEPRECATION_WARNINGS
 			return;
 		}
 		else
@@ -3892,9 +3850,6 @@ PRAGMA_ENABLE_DEPRECATION_WARNINGS
 
 	TRACE_CPUPROFILER_EVENT_SCOPE(ConfigReadyForUseBroadcast);
 	FCoreDelegates::TSConfigReadyForUse().Broadcast();
-PRAGMA_DISABLE_DEPRECATION_WARNINGS
-	FCoreDelegates::ConfigReadyForUse.Broadcast();
-PRAGMA_ENABLE_DEPRECATION_WARNINGS
 }
 
 const FString& FConfigCacheIni::GetCustomConfigString()
diff --git a/Engine/Source/Runtime/Core/Private/Misc/CoreDelegates.cpp b/Engine/Source/Runtime/Core/Private/Misc/CoreDelegates.cpp
--- a/Engine/Source/Runtime/Core/Private/Misc/CoreDelegates.cpp
+++ b/Engine/Source/Runtime/Core/Private/Misc/CoreDelegates.cpp
@@ -146,18 +146,6 @@ TMulticastDelegate<void(FString, int, int)> FCoreDelegates::ApplicationReceivedL
 TMulticastDelegate<void()> FCoreDelegates::ApplicationPerformFetchDelegate;
 TMulticastDelegate<void(FString)> FCoreDelegates::ApplicationBackgroundSessionEventDelegate;
 
-PRAGMA_DISABLE_DEPRECATION_WARNINGS
-TMulticastDelegate<void(const TCHAR*, int32&)> FCoreDelegates::CountPreLoadConfigFileRespondersDelegate;
-TMulticastDelegate<void(const TCHAR*, FString&)> FCoreDelegates::PreLoadConfigFileDelegate;
-TMulticastDelegate<void(const TCHAR*, const FString&, int32&)> FCoreDelegates::PreSaveConfigFileDelegate;
-TMulticastDelegate<void(const FConfigFile*)> FCoreDelegates::OnFConfigCreated;
-TMulticastDelegate<void(const FConfigFile*)> FCoreDelegates::OnFConfigDeleted;
-TMulticastDelegate<void(const TCHAR*, const TCHAR*, const TCHAR*)> FCoreDelegates::OnConfigValueRead;
-TMulticastDelegate<void(const TCHAR*, const TCHAR*)> FCoreDelegates::OnConfigSectionRead;
-TMulticastDelegate<void(const TCHAR*, const TCHAR*)> FCoreDelegates::OnConfigSectionNameRead;
-TMulticastDelegate<void(const FString&, const TSet<FString>&)> FCoreDelegates::OnConfigSectionsChanged;
-PRAGMA_ENABLE_DEPRECATION_WARNINGS
-
 TTSMulticastDelegate<void(const TCHAR*, int32&)>& FCoreDelegates::TSCountPreLoadConfigFileRespondersDelegate()
 {
 	static TTSMulticastDelegate<void(const TCHAR*, int32&)> Singleton;
@@ -253,10 +241,6 @@ TMulticastDelegate<void()> FCoreDelegates::PreRenderingThreadDestroyed;
 
 TMulticastDelegate<void(int32)> FCoreDelegates::ApplicationReceivedScreenOrientationChangedNotificationDelegate;
 
-PRAGMA_DISABLE_DEPRECATION_WARNINGS
-TMulticastDelegate<void()> FCoreDelegates::ConfigReadyForUse;
-PRAGMA_ENABLE_DEPRECATION_WARNINGS
-
 TTSMulticastDelegate<void()>& FCoreDelegates::TSConfigReadyForUse()
 {
 	static TTSMulticastDelegate<void()> Singleton;
@@ -272,10 +256,6 @@ TMulticastDelegate<void(const FString&, FPlatformUserId)> FCoreDelegates::OnActi
 
 TMulticastDelegate<void(const FString&)> FCoreDelegates::OnGCFinishDestroyTimeExtended;
 
-PRAGMA_DISABLE_DEPRECATION_WARNINGS
-TMulticastDelegate<void(FCoreDelegates::FExtraBinaryConfigData&)> FCoreDelegates::AccessExtraBinaryConfigData;
-PRAGMA_ENABLE_DEPRECATION_WARNINGS
-
 TTSMulticastDelegate<void(FCoreDelegates::FExtraBinaryConfigData&)>& FCoreDelegates::TSAccessExtraBinaryConfigData()
 {
 	static TTSMulticastDelegate<void(FCoreDelegates::FExtraBinaryConfigData&)> Singleton;
diff --git a/Engine/Source/Runtime/Core/Public/Misc/CoreDelegates.h b/Engine/Source/Runtime/Core/Public/Misc/CoreDelegates.h
--- a/Engine/Source/Runtime/Core/Public/Misc/CoreDelegates.h
+++ b/Engine/Source/Runtime/Core/Public/Misc/CoreDelegates.h
@@ -254,40 +254,14 @@ public:
 	/** Called when the OS goes into low power mode */
 	static TMulticastDelegate<void(bool)> OnLowPowerMode;
 
-	UE_DEPRECATED(5.2, "This delegate is not thread-safe. Please use FCoreDelegates::TSCountPreLoadConfigFileRespondersDelegate()")
-	static TMulticastDelegate<void(const TCHAR* IniFilename, int32& ResponderCount)> CountPreLoadConfigFileRespondersDelegate;
 	static TTSMulticastDelegate<void(const TCHAR* IniFilename, int32& ResponderCount)>& TSCountPreLoadConfigFileRespondersDelegate();
-
-	UE_DEPRECATED(5.2, "This delegate is not thread-safe. Please use FCoreDelegates::TSPreLoadConfigFileDelegate()")
-	static TMulticastDelegate<void(const TCHAR* IniFilename, FString& LoadedContents)> PreLoadConfigFileDelegate;
 	static TTSMulticastDelegate<void(const TCHAR* IniFilename, FString& LoadedContents)>& TSPreLoadConfigFileDelegate();
-
-	UE_DEPRECATED(5.2, "This delegate is not thread-safe. Please use FCoreDelegates::TSPreSaveConfigFileDelegate()")
-	static TMulticastDelegate<void(const TCHAR* IniFilename, const FString& ContentsToSave, int32& SavedCount)> PreSaveConfigFileDelegate;
 	static TTSMulticastDelegate<void(const TCHAR* IniFilename, const FString& ContentsToSave, int32& SavedCount)>& TSPreSaveConfigFileDelegate();
-
-	UE_DEPRECATED(5.2, "This delegate is not thread-safe. Please use FCoreDelegates::TSOnFConfigCreated()")
-	static TMulticastDelegate<void(const FConfigFile*)> OnFConfigCreated;
 	static TTSMulticastDelegate<void(const FConfigFile*)>& TSOnFConfigCreated();
-
-	UE_DEPRECATED(5.2, "This delegate is not thread-safe. Please use FCoreDelegates::TSOnFConfigDeleted()")
-	static TMulticastDelegate<void(const FConfigFile*)> OnFConfigDeleted;
 	static TTSMulticastDelegate<void(const FConfigFile*)>& TSOnFConfigDeleted();
-
-	UE_DEPRECATED(5.2, "This delegate is not thread-safe. Please use FCoreDelegates::TSOnConfigValueRead()")
-	static TMulticastDelegate<void(const TCHAR* IniFilename, const TCHAR* SectionName, const TCHAR* Key)> OnConfigValueRead;
 	static TTSMulticastDelegate<void(const TCHAR* IniFilename, const TCHAR* SectionName, const TCHAR* Key)>& TSOnConfigValueRead();
-
-	UE_DEPRECATED(5.2, "This delegate is not thread-safe. Please use FCoreDelegates::TSOnConfigSectionRead()")
-	static TMulticastDelegate<void(const TCHAR* IniFilename, const TCHAR* SectionName)> OnConfigSectionRead;
 	static TTSMulticastDelegate<void(const TCHAR* IniFilename, const TCHAR* SectionName)>& TSOnConfigSectionRead();
-
-	UE_DEPRECATED(5.2, "This delegate is not thread-safe. Please use FCoreDelegates::TSOnConfigSectionNameRead()")
-	static TMulticastDelegate<void(const TCHAR* IniFilename, const TCHAR* SectionName)> OnConfigSectionNameRead;
 	static TTSMulticastDelegate<void(const TCHAR* IniFilename, const TCHAR* SectionName)>& TSOnConfigSectionNameRead();
-
-	UE_DEPRECATED(5.2, "This delegate is not thread-safe. Please use FCoreDelegates::TSOnConfigSectionsChanged()")
-	static TMulticastDelegate<void(const FString& IniFilename, const TSet<FString>& SectionNames)> OnConfigSectionsChanged;
 	static TTSMulticastDelegate<void(const FString& IniFilename, const TSet<FString>& SectionNames)>& TSOnConfigSectionsChanged();
 
 	static TMulticastDelegate<void(const TCHAR* SectionName, const TCHAR* IniFilename, uint32 SetBy, bool bAllowCheating)> OnApplyCVarFromIni;
@@ -454,8 +428,6 @@ public:
 	static TMulticastDelegate<void()>& GetPreMainInitDelegate();
 	
 	/** Sent when GConfig is finished initializing */
-	UE_DEPRECATED(5.2, "This delegate is not thread-safe. Please use FCoreDelegates::TSConfigReadyForUse()")
-	static TMulticastDelegate<void()> ConfigReadyForUse;
 	static TTSMulticastDelegate<void()>& TSConfigReadyForUse();
 
 	/** Callback for notifications regarding changes of the rendering thread. */
@@ -549,8 +521,6 @@ public:
 		}
 	};
 	
-	UE_DEPRECATED(5.2, "This delegate is not thread-safe. Please use FCoreDelegates::TSAccessExtraBinaryConfigData()")
-	static TMulticastDelegate<void(FExtraBinaryConfigData&)> AccessExtraBinaryConfigData;
 	static TTSMulticastDelegate<void(FExtraBinaryConfigData&)>& TSAccessExtraBinaryConfigData();
 
 	using FAttachShaderReadRequestFunc = TFunctionRef<class FIoRequest(const class FIoChunkId&, FGraphEventRef)>;
