Allowed TFunction's debugger visualiser to be set at a global level via UE_ENABLE_TFUNCTIONREF_VISUALIZATION.

diff --git a/Engine/Source/Runtime/Core/Public/Templates/Function.h b/Engine/Source/Runtime/Core/Public/Templates/Function.h
--- a/Engine/Source/Runtime/Core/Public/Templates/Function.h
+++ b/Engine/Source/Runtime/Core/Public/Templates/Function.h
@@ -21,10 +21,12 @@
 
 
 // Disable visualization hack for shipping or test builds.
-#if !(UE_BUILD_SHIPPING || UE_BUILD_TEST)
-	#define ENABLE_TFUNCTIONREF_VISUALIZATION 1
-#else
-	#define ENABLE_TFUNCTIONREF_VISUALIZATION 0
+#ifndef UE_ENABLE_TFUNCTIONREF_VISUALIZATION
+	#if !(UE_BUILD_SHIPPING || UE_BUILD_TEST)
+		#define UE_ENABLE_TFUNCTIONREF_VISUALIZATION 1
+	#else
+		#define UE_ENABLE_TFUNCTIONREF_VISUALIZATION 0
+	#endif
 #endif
 
 #if defined(_WIN32) && !defined(_WIN64) && (!defined(ALLOW_TFUNCTION_INLINE_ALLOCATORS_ON_WIN32) || !ALLOW_TFUNCTION_INLINE_ALLOCATORS_ON_WIN32)
@@ -431,7 +433,7 @@ namespace UE::Core::Private::Function
 		return &NewOwned->Obj;
 	}
 
-	#if ENABLE_TFUNCTIONREF_VISUALIZATION
+	#if UE_ENABLE_TFUNCTIONREF_VISUALIZATION
 		/**
 		 * Helper classes to help debugger visualization.
 		 */
@@ -498,7 +500,7 @@ namespace UE::Core::Private::Function
 		{
 			if (Callable)
 			{
-				#if ENABLE_TFUNCTIONREF_VISUALIZATION
+				#if UE_ENABLE_TFUNCTIONREF_VISUALIZATION
 					// Use Memcpy to copy the other DebugPtrStorage, including vptr (because we don't know the bound type
 					// here), and then reseat the underlying pointer.  Possibly even more evil than the Set code.
 					FMemory::Memcpy(&DebugPtrStorage, &Other.DebugPtrStorage, sizeof(DebugPtrStorage)); //-V598
@@ -516,7 +518,7 @@ namespace UE::Core::Private::Function
 		{
 			if (Callable)
 			{
-				#if ENABLE_TFUNCTIONREF_VISUALIZATION
+				#if UE_ENABLE_TFUNCTIONREF_VISUALIZATION
 					// Use Memcpy to copy the other DebugPtrStorage, including vptr (because we don't know the bound type
 					// here), and then reseat the underlying pointer.  Possibly even more evil than the Set code.
 					FMemory::Memcpy(&DebugPtrStorage, &Other.DebugPtrStorage, sizeof(DebugPtrStorage)); //-V598
@@ -535,7 +537,7 @@ namespace UE::Core::Private::Function
 			{
 				void* NewPtr = Storage.BindCopy(Other.Storage);
 
-				#if ENABLE_TFUNCTIONREF_VISUALIZATION
+				#if UE_ENABLE_TFUNCTIONREF_VISUALIZATION
 					// Use Memcpy to copy the other DebugPtrStorage, including vptr (because we don't know the bound type
 					// here), and then reseat the underlying pointer.  Possibly even more evil than the Set code.
 					FMemory::Memcpy(&DebugPtrStorage, &Other.DebugPtrStorage, sizeof(DebugPtrStorage)); //-V598
@@ -551,7 +553,7 @@ namespace UE::Core::Private::Function
 			{
 				void* NewPtr = Storage.BindCopy(Other.Storage);
 
-				#if ENABLE_TFUNCTIONREF_VISUALIZATION
+				#if UE_ENABLE_TFUNCTIONREF_VISUALIZATION
 					// Use Memcpy to copy the other DebugPtrStorage, including vptr (because we don't know the bound type
 					// here), and then reseat the underlying pointer.  Possibly even more evil than the Set code.
 					FMemory::Memcpy(&DebugPtrStorage, &Other.DebugPtrStorage, sizeof(DebugPtrStorage)); //-V598
@@ -572,7 +574,7 @@ namespace UE::Core::Private::Function
 
 				Callable = &TFunctionRefCaller<DecayedFunctorType, Ret (ParamTypes...)>::Call;
 
-				#if ENABLE_TFUNCTIONREF_VISUALIZATION
+				#if UE_ENABLE_TFUNCTIONREF_VISUALIZATION
 					// We placement new over the top of the same object each time.  This is illegal,
 					// but it ensures that the vptr is set correctly for the bound type, and so is
 					// visualizable.  We never depend on the state of this object at runtime, so it's
@@ -618,7 +620,7 @@ namespace UE::Core::Private::Function
 
 		StorageType Storage;
 
-		#if ENABLE_TFUNCTIONREF_VISUALIZATION
+		#if UE_ENABLE_TFUNCTIONREF_VISUALIZATION
 			// To help debug visualizers
 			TDebugHelper<void> DebugPtrStorage;
 		#endif
