Latest Clang toolchain fixes for TIntVector2, TIntVector3, TIntVector4 and TIntPoint. deprecated-copy-with-user-provided-cop.

diff --git a/Engine/Source/Runtime/Core/Public/Math/IntPoint.h b/Engine/Source/Runtime/Core/Public/Math/IntPoint.h
--- a/Engine/Source/Runtime/Core/Public/Math/IntPoint.h
+++ b/Engine/Source/Runtime/Core/Public/Math/IntPoint.h
@@ -93,6 +93,14 @@ struct TIntPoint
 	{
 	}
 
+	// Workaround for clang deprecation warnings for deprecated XY member in implicitly-defined special member functions
+PRAGMA_DISABLE_DEPRECATION_WARNINGS
+	TIntPoint(TIntPoint&&) = default;
+	TIntPoint(const TIntPoint&) = default;
+	TIntPoint& operator=(TIntPoint&&) = default;
+	TIntPoint& operator=(const TIntPoint&) = default;
+PRAGMA_ENABLE_DEPRECATION_WARNINGS
+
 	/**
 	 * Get specific component of a point.
 	 *
@@ -225,20 +233,6 @@ PRAGMA_ENABLE_DEPRECATION_WARNINGS
 		return *this;
 	}
 
-	/**
-	 * Assign another point to this one.
-	 *
-	 * @param Other The point to assign this point from.
-	 * @return Reference to this point after assignment.
-	 */
-	TIntPoint& operator=(const TIntPoint& Other)
-	{
-		X = Other.X;
-		Y = Other.Y;
-
-		return *this;
-	}
-
 	/**
 	 * Get the result of scaling on this point.
 	 *
diff --git a/Engine/Source/Runtime/Core/Public/Math/IntVector.h b/Engine/Source/Runtime/Core/Public/Math/IntVector.h
--- a/Engine/Source/Runtime/Core/Public/Math/IntVector.h
+++ b/Engine/Source/Runtime/Core/Public/Math/IntVector.h
@@ -107,6 +107,14 @@ struct TIntVector3
 	{
 	}
 
+	// Workaround for clang deprecation warnings for deprecated XYZ member in implicitly-defined special member functions
+PRAGMA_DISABLE_DEPRECATION_WARNINGS
+	TIntVector3(TIntVector3&&) = default;
+	TIntVector3(const TIntVector3&) = default;
+	TIntVector3& operator=(TIntVector3&&) = default;
+	TIntVector3& operator=(const TIntVector3&) = default;
+PRAGMA_ENABLE_DEPRECATION_WARNINGS
+
 	/**
 	 * Gets specific component of a point.
 	 *
@@ -603,6 +611,14 @@ struct TIntVector2
 	{
 	}
 
+	// Workaround for clang deprecation warnings for deprecated XY member in implicitly-defined special member functions
+PRAGMA_DISABLE_DEPRECATION_WARNINGS
+	TIntVector2(TIntVector2&&) = default;
+	TIntVector2(const TIntVector2&) = default;
+	TIntVector2& operator=(TIntVector2&&) = default;
+	TIntVector2& operator=(const TIntVector2&) = default;
+PRAGMA_ENABLE_DEPRECATION_WARNINGS
+
 	const IntType& operator[](int32 ComponentIndex) const
 	{
 PRAGMA_DISABLE_DEPRECATION_WARNINGS
@@ -713,6 +729,14 @@ struct TIntVector4
 	{
 	}
 
+	// Workaround for clang deprecation warnings for deprecated XYZW member in implicitly-defined special member functions
+PRAGMA_DISABLE_DEPRECATION_WARNINGS
+	TIntVector4(TIntVector4&&) = default;
+	TIntVector4(const TIntVector4&) = default;
+	TIntVector4& operator=(TIntVector4&&) = default;
+	TIntVector4& operator=(const TIntVector4&) = default;
+PRAGMA_ENABLE_DEPRECATION_WARNINGS
+
 	/**
 	 * Gets specific component of a point.
 	 *
@@ -883,22 +907,6 @@ PRAGMA_ENABLE_DEPRECATION_WARNINGS
 		return *this;
 	}
 
-	/**
-	 * Assigns another point to this one.
-	 *
-	 * @param Other The point to assign this point from.
-	 * @return Reference to this point after assignment.
-	 */
-	TIntVector4& operator=(const TIntVector4& Other)
-	{
-		X = Other.X;
-		Y = Other.Y;
-		Z = Other.Z;
-		W = Other.W;
-
-		return *this;
-	}
-
 	/**
 	 * Gets the result of component-wise multiplication of this point by another.
 	 *
