Changed complex uses of events to use FEventCount

#rnx

diff --git a/Engine/Source/Developer/DerivedDataCache/Private/DerivedDataRequestOwner.cpp b/Engine/Source/Developer/DerivedDataCache/Private/DerivedDataRequestOwner.cpp
--- a/Engine/Source/Developer/DerivedDataCache/Private/DerivedDataRequestOwner.cpp
+++ b/Engine/Source/Developer/DerivedDataCache/Private/DerivedDataRequestOwner.cpp
@@ -2,10 +2,10 @@
 
 #include "DerivedDataRequestOwner.h"
 
+#include "Async/EventCount.h"
 #include "Containers/Array.h"
 #include "DerivedDataRequest.h"
 #include "DerivedDataRequestTypes.h"
-#include "Experimental/Async/LazyEvent.h"
 #include "HAL/CriticalSection.h"
 #include "Misc/ScopeExit.h"
 #include "Misc/ScopeRWLock.h"
@@ -17,12 +17,6 @@
 namespace UE::DerivedData::Private
 {
 
-/** A one-time-use event to work around the lack of condition variables. */
-struct FRequestBarrierEvent : public FRefCountBase
-{
-	FLazyEvent Event{EEventMode::ManualReset};
-};
-
 class FRequestOwnerShared final : public IRequestOwner, public FRequestBase
 {
 public:
@@ -49,7 +43,7 @@ public:
 private:
 	mutable FRWLock Lock;
 	TArray<TRefCountPtr<IRequest>, TInlineAllocator<1>> Requests;
-	TRefCountPtr<FRequestBarrierEvent> BarrierEvent;
+	FEventCount BarrierEvent;
 	uint32 BarrierCount{0};
 	uint32 PriorityBarrierCount{0};
 	EPriority Priority{EPriority::Normal};
@@ -128,7 +122,7 @@ void FRequestOwnerShared::BeginBarrier(ERequestBarrierFlags Flags)
 void FRequestOwnerShared::EndBarrier(ERequestBarrierFlags Flags)
 {
 	ON_SCOPE_EXIT { Release(); };
-	TRefCountPtr<FRequestBarrierEvent> LocalBarrierEvent;
+	bool bNotifyBarrier = false;
 	{
 		FWriteScopeLock WriteLock(Lock);
 		if (EnumHasAnyFlags(Flags, ERequestBarrierFlags::Priority) && --PriorityBarrierCount == 0)
@@ -137,13 +131,12 @@ void FRequestOwnerShared::EndBarrier(ERequestBarrierFlags Flags)
 		}
 		if (--BarrierCount == 0)
 		{
-			LocalBarrierEvent = MoveTemp(BarrierEvent);
-			BarrierEvent = nullptr;
+			bNotifyBarrier = true;
 		}
 	}
-	if (LocalBarrierEvent)
+	if (bNotifyBarrier)
 	{
-		LocalBarrierEvent->Event.Trigger();
+		BarrierEvent.Notify();
 	}
 }
 
@@ -172,7 +165,7 @@ void FRequestOwnerShared::Cancel()
 	for (;;)
 	{
 		TRefCountPtr<IRequest> LocalRequest;
-		TRefCountPtr<FRequestBarrierEvent> LocalBarrierEvent;
+		FEventCountToken BarrierToken;
 
 		{
 			FWriteScopeLock WriteLock(Lock);
@@ -184,11 +177,7 @@ void FRequestOwnerShared::Cancel()
 			}
 			else if (BarrierCount > 0)
 			{
-				LocalBarrierEvent = BarrierEvent;
-				if (!LocalBarrierEvent)
-				{
-					LocalBarrierEvent = BarrierEvent = new FRequestBarrierEvent;
-				}
+				BarrierToken = BarrierEvent.PrepareWait();
 			}
 			else
 			{
@@ -202,7 +191,7 @@ void FRequestOwnerShared::Cancel()
 		}
 		else
 		{
-			LocalBarrierEvent->Event.Wait();
+			BarrierEvent.Wait(BarrierToken);
 		}
 	}
 }
@@ -212,7 +201,7 @@ void FRequestOwnerShared::Wait()
 	for (;;)
 	{
 		TRefCountPtr<IRequest> LocalRequest;
-		TRefCountPtr<FRequestBarrierEvent> LocalBarrierEvent;
+		FEventCountToken BarrierToken;
 
 		{
 			FWriteScopeLock WriteLock(Lock);
@@ -223,11 +212,7 @@ void FRequestOwnerShared::Wait()
 			}
 			else if (BarrierCount > 0)
 			{
-				LocalBarrierEvent = BarrierEvent;
-				if (!LocalBarrierEvent)
-				{
-					LocalBarrierEvent = BarrierEvent = new FRequestBarrierEvent;
-				}
+				BarrierToken = BarrierEvent.PrepareWait();
 			}
 			else
 			{
@@ -241,7 +226,7 @@ void FRequestOwnerShared::Wait()
 		}
 		else
 		{
-			LocalBarrierEvent->Event.Wait();
+			BarrierEvent.Wait(BarrierToken);
 		}
 	}
 }
diff --git a/Engine/Source/Runtime/Core/Private/Misc/OutputDeviceRedirector.cpp b/Engine/Source/Runtime/Core/Private/Misc/OutputDeviceRedirector.cpp
--- a/Engine/Source/Runtime/Core/Private/Misc/OutputDeviceRedirector.cpp
+++ b/Engine/Source/Runtime/Core/Private/Misc/OutputDeviceRedirector.cpp
@@ -2,11 +2,10 @@
 
 #include "Misc/OutputDeviceRedirector.h"
 
+#include "Async/EventCount.h"
 #include "Containers/BitArray.h"
 #include "Containers/DepletableMpmcQueue.h"
-#include "Containers/ConsumeAllMpmcQueue.h"
 #include "Experimental/ConcurrentLinearAllocator.h"
-#include "HAL/Event.h"
 #include "HAL/PlatformProcess.h"
 #include "HAL/PlatformTime.h"
 #include "HAL/PlatformTLS.h"
@@ -139,11 +138,11 @@ struct FOutputDeviceRedirectorState
 	/** A lock to synchronize access to the thread. */
 	FRWLock ThreadLock;
 
-	/** A queue of events to trigger when the dedicated primary thread is idle. */
-	TConsumeAllMpmcQueue<FEvent*, FOutputDeviceLinearAllocator> ThreadIdleEvents;
+	/** An event that is notified when the dedicated primary thread is idle. */
+	FEventCount ThreadIdleEvent;
 
 	/** An event to wake the dedicated primary thread to process buffered items. */
-	std::atomic<FEvent*> ThreadWakeEvent = nullptr;
+	FEventCount ThreadWakeEvent;
 
 	/** The ID of the thread holding the primary lock. */
 	std::atomic<uint32> LockedThreadId = MAX_uint32;
@@ -154,6 +153,9 @@ struct FOutputDeviceRedirectorState
 	/** The ID of the panic thread, which is only set by Panic(). */
 	std::atomic<uint32> PanicThreadId = MAX_uint32;
 
+	/** Whether a dedicated primary thread has been started. */
+	std::atomic<bool> bThreadStarted = false;
+
 	/** Whether the backlog is enabled. */
 	bool bEnableBacklog = false;
 
@@ -408,11 +410,8 @@ void FOutputDeviceRedirectorState::RemoveOutputDevice(FOutputDevice* OutputDevic
 
 bool FOutputDeviceRedirectorState::TryStartThread()
 {
-	if (FWriteScopeLock ThreadScopeLock(ThreadLock); !ThreadWakeEvent.load(std::memory_order_relaxed))
+	if (FWriteScopeLock ThreadScopeLock(ThreadLock); !bThreadStarted.exchange(true, std::memory_order_relaxed))
 	{
-		FEvent* WakeEvent = FPlatformProcess::GetSynchEventFromPool();
-		WakeEvent->Trigger();
-		ThreadWakeEvent.store(WakeEvent, std::memory_order_release);
 		Thread = FThread(TEXT("OutputDeviceRedirector"), [this] { ThreadLoop(); });
 	}
 	return true;
@@ -420,12 +419,10 @@ bool FOutputDeviceRedirectorState::TryStartThread()
 
 bool FOutputDeviceRedirectorState::TryStopThread()
 {
-	if (FWriteScopeLock ThreadScopeLock(ThreadLock); FEvent* WakeEvent = ThreadWakeEvent.exchange(nullptr, std::memory_order_acquire))
+	if (FWriteScopeLock ThreadScopeLock(ThreadLock); bThreadStarted.exchange(false, std::memory_order_relaxed))
 	{
-		WakeEvent->Trigger();
+		ThreadWakeEvent.Notify();
 		Thread.Join();
-		FOutputDevicesWriteScopeLock Lock(*this);
-		FPlatformProcess::ReturnSynchEventToPool(WakeEvent);
 	}
 	return true;
 }
@@ -439,9 +436,13 @@ void FOutputDeviceRedirectorState::ThreadLoop()
 		PrimaryThreadId.store(ThreadId, std::memory_order_relaxed);
 	}
 
-	while (FEvent* WakeEvent = ThreadWakeEvent.load(std::memory_order_acquire))
+	for (;;)
 	{
-		WakeEvent->Wait();
+		FEventCountToken Token = ThreadWakeEvent.PrepareWait();
+		if (!bThreadStarted.load(std::memory_order_relaxed))
+		{
+			break;
+		}
 		while (!BufferedItems.IsEmpty() && IsPrimaryThread(ThreadId))
 		{
 			if (FOutputDevicesPrimaryScopeLock Lock(*this); Lock.IsLocked())
@@ -449,7 +450,8 @@ void FOutputDeviceRedirectorState::ThreadLoop()
 				FlushBufferedItems();
 			}
 		}
-		ThreadIdleEvents.ConsumeAllLifo([](FEvent* Event) { Event->Trigger(); });
+		ThreadIdleEvent.Notify();
+		ThreadWakeEvent.Wait(Token);
 	}
 }
 
@@ -553,17 +555,14 @@ bool FOutputDeviceRedirector::IsRedirectingTo(FOutputDevice* OutputDevice)
 
 void FOutputDeviceRedirector::FlushThreadedLogs(EOutputDeviceRedirectorFlushOptions Options)
 {
-	if (FReadScopeLock ThreadLock(State->ThreadLock); FEvent* WakeEvent = State->ThreadWakeEvent.load(std::memory_order_acquire))
+	if (FReadScopeLock ThreadLock(State->ThreadLock); State->bThreadStarted.load(std::memory_order_relaxed))
 	{
 		if (!EnumHasAnyFlags(Options, EOutputDeviceRedirectorFlushOptions::Async))
 		{
 			TRACE_CPUPROFILER_EVENT_SCOPE(FOutputDeviceRedirector::FlushThreadedLogs);
-			FEventRef IdleEvent(EEventMode::ManualReset);
-			if (State->ThreadIdleEvents.ProduceItem(IdleEvent.Get()) == UE::EConsumeAllMpmcQueueResult::WasEmpty)
-			{
-				WakeEvent->Trigger();
-			}
-			IdleEvent->Wait();
+			UE::FEventCountToken Token = State->ThreadIdleEvent.PrepareWait();
+			State->ThreadWakeEvent.Notify();
+			State->ThreadIdleEvent.Wait(Token);
 		}
 		return;
 	}
@@ -656,10 +655,7 @@ void FOutputDeviceRedirector::SerializeRecord(const UE::FLogRecord& Record)
 	// Queue the record to serialize to buffered output devices from the primary thread.
 	if (State->BufferedItems.EnqueueAndReturnWasEmpty(Record))
 	{
-		if (FEvent* WakeEvent = State->ThreadWakeEvent.load(std::memory_order_acquire))
-		{
-			WakeEvent->Trigger();
-		}
+		State->ThreadWakeEvent.Notify();
 	}
 }
 
@@ -720,10 +716,7 @@ void FOutputDeviceRedirector::Serialize(const TCHAR* const Data, const ELogVerbo
 	// Queue the line to serialize to buffered output devices from the primary thread.
 	if (State->BufferedItems.EnqueueAndReturnWasEmpty(Data, Category, Verbosity, RealTime))
 	{
-		if (FEvent* WakeEvent = State->ThreadWakeEvent.load(std::memory_order_acquire))
-		{
-			WakeEvent->Trigger();
-		}
+		State->ThreadWakeEvent.Notify();
 	}
 }
 
