Fixing some missing #includes in non-unity builds was causing some C++20 template features to creep into builds where they are unwelcome.

diff --git a/Engine/Source/Runtime/Core/Public/AutoRTFM/AutoRTFM.h b/Engine/Source/Runtime/Core/Public/AutoRTFM/AutoRTFM.h
--- a/Engine/Source/Runtime/Core/Public/AutoRTFM/AutoRTFM.h
+++ b/Engine/Source/Runtime/Core/Public/AutoRTFM/AutoRTFM.h
@@ -18,7 +18,9 @@
 #include <stdlib.h>
 #include <memory.h>
 
+#if UE_AUTORTFM
 #include "Templates/Function.h"
+#endif
 
 #ifdef __cplusplus
 extern "C"
