ParkingLot: Moved Wait and WaitFor out of the header

#rnx

diff --git a/Engine/Source/Runtime/Core/Private/Async/ParkingLot.cpp b/Engine/Source/Runtime/Core/Private/Async/ParkingLot.cpp
--- a/Engine/Source/Runtime/Core/Private/Async/ParkingLot.cpp
+++ b/Engine/Source/Runtime/Core/Private/Async/ParkingLot.cpp
@@ -492,7 +492,17 @@ FThreadLocalData::~FThreadLocalData()
 
 ///////////////////////////////////////////////////////////////////////////////////////////////////////////////////////
 
-FWaitState WaitUntil(const void* Address, TFunctionRef<bool()> CanWait, TFunctionRef<void()> BeforeWait, FMonotonicTimePoint WaitTime)
+FWaitState Wait(const void* Address, const TFunctionRef<bool()>& CanWait, const TFunctionRef<void()>& BeforeWait)
+{
+	return WaitUntil(Address, CanWait, BeforeWait, FMonotonicTimePoint::Infinity());
+}
+
+FWaitState WaitFor(const void* Address, const TFunctionRef<bool()>& CanWait, const TFunctionRef<void()>& BeforeWait, FMonotonicTimeSpan WaitTime)
+{
+	return WaitUntil(Address, CanWait, BeforeWait, FMonotonicTimePoint::Now() + WaitTime);
+}
+
+FWaitState WaitUntil(const void* Address, const TFunctionRef<bool()>& CanWait, const TFunctionRef<void()>& BeforeWait, FMonotonicTimePoint WaitTime)
 {
 	using namespace Private;
 
@@ -564,7 +574,7 @@ FWaitState WaitUntil(const void* Address, TFunctionRef<bool()> CanWait, TFunctio
 	return State;
 }
 
-void WakeOne(const void* Address, TFunctionRef<uint64(FWakeState)> OnWakeState)
+void WakeOne(const void* Address, const TFunctionRef<uint64(FWakeState)>& OnWakeState)
 {
 	using namespace Private;
 
diff --git a/Engine/Source/Runtime/Core/Public/Async/ParkingLot.h b/Engine/Source/Runtime/Core/Public/Async/ParkingLot.h
--- a/Engine/Source/Runtime/Core/Public/Async/ParkingLot.h
+++ b/Engine/Source/Runtime/Core/Public/Async/ParkingLot.h
@@ -34,8 +34,10 @@ struct FWakeState final
 
 namespace Private
 {
-UE_API FWaitState WaitUntil(const void* Address, TFunctionRef<bool()> CanWait, TFunctionRef<void()> BeforeWait, FMonotonicTimePoint WaitTime);
-UE_API void WakeOne(const void* Address, TFunctionRef<uint64(FWakeState)> OnWakeState);
+UE_API FWaitState Wait(const void* Address, const TFunctionRef<bool()>& CanWait, const TFunctionRef<void()>& BeforeWait);
+UE_API FWaitState WaitFor(const void* Address, const TFunctionRef<bool()>& CanWait, const TFunctionRef<void()>& BeforeWait, FMonotonicTimeSpan WaitTime);
+UE_API FWaitState WaitUntil(const void* Address, const TFunctionRef<bool()>& CanWait, const TFunctionRef<void()>& BeforeWait, FMonotonicTimePoint WaitTime);
+UE_API void WakeOne(const void* Address, const TFunctionRef<uint64(FWakeState)>& OnWakeState);
 }
 
 /**
@@ -48,7 +50,7 @@ UE_API void WakeOne(const void* Address, TFunctionRef<uint64(FWakeState)> OnWake
 template <typename CanWaitType, typename BeforeWaitType>
 inline FWaitState Wait(const void* Address, CanWaitType&& CanWait, BeforeWaitType&& BeforeWait)
 {
-	return Private::WaitUntil(Address, Forward<CanWaitType>(CanWait), Forward<BeforeWaitType>(BeforeWait), FMonotonicTimePoint::Infinity());
+	return Private::Wait(Address, Forward<CanWaitType>(CanWait), Forward<BeforeWaitType>(BeforeWait));
 }
 
 /**
@@ -62,7 +64,7 @@ inline FWaitState Wait(const void* Address, CanWaitType&& CanWait, BeforeWaitTyp
 template <typename CanWaitType, typename BeforeWaitType>
 inline FWaitState WaitFor(const void* Address, CanWaitType&& CanWait, BeforeWaitType&& BeforeWait, FMonotonicTimeSpan WaitTime)
 {
-	return Private::WaitUntil(Address, Forward<CanWaitType>(CanWait), Forward<BeforeWaitType>(BeforeWait), FMonotonicTimePoint::Now() + WaitTime);
+	return Private::WaitFor(Address, Forward<CanWaitType>(CanWait), Forward<BeforeWaitType>(BeforeWait), WaitTime);
 }
 
 /**
