Windows platform extensions can override the Platform Features Module via config files.

#rnx

diff --git a/Engine/Source/Runtime/Core/Private/Windows/WindowsPlatformMisc.cpp b/Engine/Source/Runtime/Core/Private/Windows/WindowsPlatformMisc.cpp
--- a/Engine/Source/Runtime/Core/Private/Windows/WindowsPlatformMisc.cpp
+++ b/Engine/Source/Runtime/Core/Private/Windows/WindowsPlatformMisc.cpp
@@ -2188,18 +2188,27 @@ int32 FWindowsPlatformMisc::NumberOfCoresIncludingHyperthreads()
 
 const TCHAR* FWindowsPlatformMisc::GetPlatformFeaturesModuleName()
 {
-	bool bModuleExists = FModuleManager::Get().ModuleExists(TEXT("WindowsPlatformFeatures"));
 	// If running a dedicated server then we use the default PlatformFeatures
-	if (bModuleExists && !IsRunningDedicatedServer())
+	if (!IsRunningDedicatedServer())
 	{
-		UE_LOG(LogWindows, Log, TEXT("WindowsPlatformFeatures enabled"));
-		return TEXT("WindowsPlatformFeatures");
-	}
-	else
-	{
-		UE_LOG(LogWindows, Log, TEXT("WindowsPlatformFeatures disabled or dedicated server build"));
-		return nullptr;
+		static FString PlatformFeaturesName = TEXT("WindowsPlatformFeatures");
+		static bool bIniChecked = false;
+		if (!bIniChecked && !GEngineIni.IsEmpty())
+		{
+			GConfig->GetString(TEXT("PlatformFeatures"), TEXT("PlatformFeaturesModule"), PlatformFeaturesName, GEngineIni);
+			bIniChecked = true;
+		}
+
+		bool bModuleExists = FModuleManager::Get().ModuleExists(*PlatformFeaturesName);
+		if (bModuleExists && !PlatformFeaturesName.IsEmpty())
+		{
+			UE_LOG(LogWindows, Log, TEXT("%s enabled"), *PlatformFeaturesName);
+			return *PlatformFeaturesName;
+		}
 	}
+
+	UE_LOG(LogWindows, Log, TEXT("WindowsPlatformFeatures disabled or dedicated server build"));
+	return nullptr;
 }
 
 int32 FWindowsPlatformMisc::NumberOfWorkerThreadsToSpawn()
