Latest Clang toolchain fixes for TScriptDelegate. deprecated-copy-with-user-provided-copy.

diff --git a/Engine/Source/Runtime/Core/Public/UObject/ScriptDelegates.h b/Engine/Source/Runtime/Core/Public/UObject/ScriptDelegates.h
--- a/Engine/Source/Runtime/Core/Public/UObject/ScriptDelegates.h
+++ b/Engine/Source/Runtime/Core/Public/UObject/ScriptDelegates.h
@@ -69,10 +69,16 @@ public:
 
 	/** Default constructor. */
 	TScriptDelegate() 
-		: Object( nullptr ),
-		  FunctionName( NAME_None )
+		: Object( nullptr )
+		, FunctionName( NAME_None )
 	{ }
 
+	TScriptDelegate(const TScriptDelegate& Other)
+		: Object(Other.Object)
+		, FunctionName(Other.FunctionName)
+	{
+	}
+
 	template <
 		typename OtherDummy
 		UE_REQUIRES(UE::Core::Private::BackwardCompatibilityCheck<Dummy, OtherDummy>())
