Broaden android mkdir permissions on !shipping builds to enable adb shell to modify the dir and its contents.

diff --git a/Engine/Source/Runtime/Core/Private/Android/AndroidPlatformFile.cpp b/Engine/Source/Runtime/Core/Private/Android/AndroidPlatformFile.cpp
--- a/Engine/Source/Runtime/Core/Private/Android/AndroidPlatformFile.cpp
+++ b/Engine/Source/Runtime/Core/Private/Android/AndroidPlatformFile.cpp
@@ -1810,6 +1810,12 @@ public:
 			}
 			return FileHandleAndroid;
 		}
+#if LOG_ANDROID_FILE
+		else
+		{
+			FPlatformMisc::LowLevelOutputDebugStringf(TEXT("FAndroidPlatformFile::OpenWrite('%s') - failed = %s"), Filename, UTF8_TO_TCHAR(strerror(errno)));
+		}
+#endif
 		return nullptr;
 	}
 
@@ -1886,8 +1892,13 @@ public:
 		FString LocalPath;
 		FString AssetPath;
 		PathToAndroidPaths(LocalPath, AssetPath, Directory, AllowLocal);
-
-		return (mkdir(TCHAR_TO_UTF8(*LocalPath), 0755) == 0) || (errno == EEXIST);
+		uint32 mkdirperms = 0755;
+#if !UE_BUILD_SHIPPING
+		// some devices prevent ADB (shell user) from modifying files.
+		// To allow adb shell to modify files we give group users all perms to the new dir.
+		mkdirperms = 0775;
+#endif
+		return (mkdir(TCHAR_TO_UTF8(*LocalPath), mkdirperms) == 0) || (errno == EEXIST);
 	}
 
 	// We assert that modifying dirs are in the local file-system.
