switched to safer down-casts in delegates implementation

diff --git a/Engine/Source/Runtime/Core/Public/Delegates/DelegateSignatureImpl.inl b/Engine/Source/Runtime/Core/Public/Delegates/DelegateSignatureImpl.inl
--- a/Engine/Source/Runtime/Core/Public/Delegates/DelegateSignatureImpl.inl
+++ b/Engine/Source/Runtime/Core/Public/Delegates/DelegateSignatureImpl.inl
@@ -602,7 +602,7 @@ protected:
 	 */
 	FORCEINLINE DelegateInstanceInterfaceType* GetDelegateInstanceProtected() const
 	{
-		return (DelegateInstanceInterfaceType*)Super::GetDelegateInstanceProtected();
+		return static_cast<DelegateInstanceInterfaceType*>(Super::GetDelegateInstanceProtected());
 	}
 };
 
diff --git a/Engine/Source/Runtime/Core/Public/Delegates/MulticastDelegateBase.h b/Engine/Source/Runtime/Core/Public/Delegates/MulticastDelegateBase.h
--- a/Engine/Source/Runtime/Core/Public/Delegates/MulticastDelegateBase.h
+++ b/Engine/Source/Runtime/Core/Public/Delegates/MulticastDelegateBase.h
@@ -177,7 +177,7 @@ protected:
 			if (IDelegateInstance* OtherInstance = GetDelegateInstanceProtectedHelper(OtherDelegateRef))
 			{
 				DelegateType TempDelegate;
-				((DelegateInstanceInterfaceType*)OtherInstance)->CreateCopy(TempDelegate);
+				static_cast<DelegateInstanceInterfaceType*>(OtherInstance)->CreateCopy(TempDelegate);
 				AddDelegateInstance(MoveTemp(TempDelegate));
 			}
 		}
@@ -198,10 +198,10 @@ protected:
 			for (int32 InvocationListIndex = LocalInvocationList.Num() - 1; InvocationListIndex >= 0; --InvocationListIndex)
 			{
 				// this down-cast is OK! allows for managing invocation list in the base class without requiring virtual functions
-				const DelegateBaseType& DelegateBase = (const DelegateBaseType&)LocalInvocationList[InvocationListIndex];
+				const DelegateBaseType& DelegateBase = static_cast<const DelegateBaseType&>(LocalInvocationList[InvocationListIndex]);
 
 				IDelegateInstance* DelegateInstanceInterface = GetDelegateInstanceProtectedHelper(DelegateBase);
-				if (DelegateInstanceInterface == nullptr || !((DelegateInstanceInterfaceType*)DelegateInstanceInterface)->ExecuteIfSafe(Params...))
+				if (DelegateInstanceInterface == nullptr || !static_cast<DelegateInstanceInterfaceType*>(DelegateInstanceInterface)->ExecuteIfSafe(Params...))
 				{
 					NeedsCompaction = true;
 				}
