delegates cosmetic changes

diff --git a/Engine/Source/Runtime/Core/Public/Delegates/DelegateSignatureImpl.inl b/Engine/Source/Runtime/Core/Public/Delegates/DelegateSignatureImpl.inl
--- a/Engine/Source/Runtime/Core/Public/Delegates/DelegateSignatureImpl.inl
+++ b/Engine/Source/Runtime/Core/Public/Delegates/DelegateSignatureImpl.inl
@@ -39,10 +39,6 @@ template<typename UserPolicy> class TTSMulticastDelegateBase;
 
 /**
  * Unicast delegate template class.
- *
- * Use the various DECLARE_DELEGATE macros to create the actual delegate type, templated to
- * the function signature the delegate is compatible with. Then, you can create an instance
- * of that class when you want to bind a function to the delegate.
  */
 template <typename DelegateSignature, typename UserPolicy = FDefaultDelegateUserPolicy>
 class TDelegate
@@ -57,13 +53,9 @@ class TDelegate<InRetValType(ParamTypes...), UserPolicy> : public TDelegateBase<
 	using FuncType                      = InRetValType (ParamTypes...);
 	using DelegateInstanceInterfaceType = IBaseDelegateInstance<FuncType, UserPolicy>;
 
-	using DelegateType          = typename UserPolicy::FDelegateExtras;
-	using DelegateInstanceType  = typename UserPolicy::FDelegateInstanceExtras;
-	using MulticastDelegateType = typename UserPolicy::FMulticastDelegateExtras;
-
-	static_assert(std::is_convertible_v<DelegateType*, FDelegateBase*>, "UserPolicy::FDelegateExtras should publicly inherit FDelegateBase");
-	static_assert(std::is_convertible_v<DelegateInstanceType*, IDelegateInstance*>, "UserPolicy::FDelegateInstanceExtras should publicly inherit IDelegateInstance");
-	static_assert(std::is_convertible_v<MulticastDelegateType*, TMulticastDelegateBase<UserPolicy>*> || std::is_convertible_v<MulticastDelegateType*, TTSMulticastDelegateBase<UserPolicy>*>, "UserPolicy::FMulticastDelegateExtras should publicly inherit TMulticastDelegateBase<UserPolicy> or TTSMulticastDelegateBase<UserPolicy>");
+	static_assert(std::is_convertible_v<typename UserPolicy::FDelegateExtras*, FDelegateBase*>, "UserPolicy::FDelegateExtras should publicly inherit FDelegateBase");
+	static_assert(std::is_convertible_v<typename UserPolicy::FDelegateInstanceExtras*, IDelegateInstance*>, "UserPolicy::FDelegateInstanceExtras should publicly inherit IDelegateInstance");
+	static_assert(std::is_convertible_v<typename UserPolicy::FMulticastDelegateExtras*, TMulticastDelegateBase<UserPolicy>*> || std::is_convertible_v<typename UserPolicy::FMulticastDelegateExtras*, TTSMulticastDelegateBase<UserPolicy>*>, "UserPolicy::FMulticastDelegateExtras should publicly inherit TMulticastDelegateBase<UserPolicy> or TTSMulticastDelegateBase<UserPolicy>");
 
 	template <typename>
 	friend class TMulticastDelegateBase;
@@ -308,57 +300,17 @@ public:
 
 public:
 
-	/**
-	 * Default constructor
-	 */
-	inline TDelegate()
-	{
-	}
+	TDelegate() = default;
 
-	/**
-	 * 'Null' constructor
-	 */
 	inline TDelegate(TYPE_OF_NULLPTR)
 	{
 	}
 
-	/**
-	 * Destructor.
-	 */
-	inline ~TDelegate()
-	{
-		Unbind();
-	}
-
-	/**
-	 * Move constructor.
-	 *
-	 * @param Other The delegate object to move from.
-	 */
-	TDelegate(TDelegate&& Other) = default;
-
-	/**
-	 * Creates and initializes a new instance from an existing delegate object.
-	 *
-	 * @param Other The delegate object to copy from.
-	 */
 	inline TDelegate(const TDelegate& Other)
 	{
 		*this = Other;
 	}
 
-	/**
-	 * Move assignment operator.
-	 *
-	 * @param	OtherDelegate	Delegate object to copy from
-	 */
-	TDelegate& operator=(TDelegate&& Other) = default;
-
-	/**
-	 * Assignment operator.
-	 *
-	 * @param	OtherDelegate	Delegate object to copy from
-	 */
 	inline TDelegate& operator=(const TDelegate& Other)
 	{
 		if (&Other != this)
@@ -379,6 +331,14 @@ public:
 		return *this;
 	}
 
+	TDelegate(TDelegate&& Other) = default;
+	TDelegate& operator=(TDelegate&& Other) = default;
+
+	inline ~TDelegate()
+	{
+		Unbind();
+	}
+
 public:
 
 	/**
@@ -654,6 +614,24 @@ private:
 	using Super::GetDelegateInstanceProtectedHelper;
 
 public:
+	TMulticastDelegate() = default;
+
+	TMulticastDelegate(const TMulticastDelegate& Other)
+	{
+		*this = Other;
+	}
+
+	TMulticastDelegate& operator=(const TMulticastDelegate& Other)
+	{
+		if (&Other != this)
+		{
+			Super::template CopyFrom<DelegateInstanceInterfaceType, FDelegate>(Other);
+		}
+		return *this;
+	}
+
+public:
+
 	/**
 	 * Adds a delegate instance to this multicast delegate's invocation list.
 	 *
@@ -890,36 +868,6 @@ public:
 		return bResult;
 	}
 
-	/** 
-	 * Hidden default constructor.
-	 */
-	inline TMulticastDelegate( ) { }
-
-	/**
-	 * Hidden copy constructor (for proper deep copies).
-	 *
-	 * @param Other The multicast delegate to copy from.
-	 */
-	TMulticastDelegate( const TMulticastDelegate& Other )
-	{
-		*this = Other;
-	}
-
-	/**
-	 * Hidden assignment operator (for proper deep copies).
-	 *
-	 * @param Other The delegate to assign from.
-	 * @return This instance.
-	 */
-	TMulticastDelegate& operator=( const TMulticastDelegate& Other )
-	{
-		if (&Other != this)
-		{
-			Super::template CopyFrom<DelegateInstanceInterfaceType, FDelegate>(Other);
-		}
-		return *this;
-	}
-
 	TMulticastDelegate(TMulticastDelegate&&) = default;
 	TMulticastDelegate& operator=(TMulticastDelegate&&) = default;
 
diff --git a/Engine/Source/Runtime/Core/Public/UObject/ScriptDelegates.h b/Engine/Source/Runtime/Core/Public/UObject/ScriptDelegates.h
--- a/Engine/Source/Runtime/Core/Public/UObject/ScriptDelegates.h
+++ b/Engine/Source/Runtime/Core/Public/UObject/ScriptDelegates.h
@@ -90,6 +90,26 @@ public:
 	{
 	}
 
+	void operator=(const TScriptDelegate& Other)
+	{
+		UE_DELEGATES_MT_SCOPED_WRITE_ACCESS(AccessDetector);
+
+		Object = Other.Object;
+		FunctionName = Other.FunctionName;
+	}
+	template <
+		typename OtherDummy
+		UE_REQUIRES(UE::Core::Private::BackwardCompatibilityCheck<Dummy, OtherDummy>())
+	>
+	/* UE_DEPRECATED(5.3, "Deprecated - remove after TScriptDelegateTraits<FWeakObjectPtr> is removed") */
+	FORCEINLINE void operator=(const TScriptDelegate<OtherDummy>& Other)
+	{
+		UE_DELEGATES_MT_SCOPED_WRITE_ACCESS(AccessDetector);
+
+		Object = Other.Object;
+		FunctionName = Other.FunctionName;
+	}
+
 private:
 
 	template <class UObjectTemplate>
@@ -261,26 +281,6 @@ public:
 		return Object != Other.Object || FunctionName != Other.FunctionName;
 	}
 
-	void operator=( const TScriptDelegate& Other )
-	{
-		UE_DELEGATES_MT_SCOPED_WRITE_ACCESS(AccessDetector);
-
-		Object = Other.Object;
-		FunctionName = Other.FunctionName;
-	}
-	template <
-		typename OtherDummy
-		UE_REQUIRES(UE::Core::Private::BackwardCompatibilityCheck<Dummy, OtherDummy>())
-	>
-	/* UE_DEPRECATED(5.3, "Deprecated - remove after TScriptDelegateTraits<FWeakObjectPtr> is removed") */
-	FORCEINLINE void operator=(const TScriptDelegate<OtherDummy>& Other)
-	{
-		UE_DELEGATES_MT_SCOPED_WRITE_ACCESS(AccessDetector);
-
-		Object = Other.Object;
-		FunctionName = Other.FunctionName;
-	}
-
 	/** 
 	 * Gets the object bound to this delegate
 	 *
@@ -422,10 +422,7 @@ class TMulticastScriptDelegate
 
 public:
 
-	/**
-	 * Default constructor
-	 */
-	inline TMulticastScriptDelegate() { }
+	TMulticastScriptDelegate() = default;
 
 public:
 
