Optimized TFunctionRef calling, as it can never be null.
Restored CheckAddress() out-of-lining, removed in CL# 15374582, replaced with a preprocessor guard to compile out the call.

diff --git a/Engine/Source/Runtime/Core/Public/Templates/Function.h b/Engine/Source/Runtime/Core/Public/Templates/Function.h
--- a/Engine/Source/Runtime/Core/Public/Templates/Function.h
+++ b/Engine/Source/Runtime/Core/Public/Templates/Function.h
@@ -293,6 +293,8 @@ namespace UE::Core::Private::Function
 
 	struct FFunctionStorage
 	{
+		constexpr static bool bCanBeNull = true;
+
 		FFunctionStorage()
 			: HeapAllocation(nullptr)
 		{
@@ -489,15 +491,14 @@ namespace UE::Core::Private::Function
 		template <typename OtherStorageType, typename OtherFuncType>
 		friend struct TFunctionRefBase;
 
-		TFunctionRefBase()
-			: Callable(nullptr)
-		{
-		}
+		TFunctionRefBase() = default;
 
 		TFunctionRefBase(TFunctionRefBase&& Other)
 			: Callable(Other.Callable)
 			, Storage (MoveTemp(Other.Storage))
 		{
+			static_assert(StorageType::bCanBeNull, "Unable to move non-nullable storage");
+
 			if (Callable)
 			{
 				#if UE_ENABLE_TFUNCTIONREF_VISUALIZATION
@@ -516,6 +517,9 @@ namespace UE::Core::Private::Function
 			: Callable(Other.Callable)
 			, Storage (MoveTemp(Other.Storage))
 		{
+			static_assert(OtherStorage::bCanBeNull, "Unable to move from non-nullable storage");
+			static_assert(StorageType::bCanBeNull,  "Unable to move into non-nullable storage");
+
 			if (Callable)
 			{
 				#if UE_ENABLE_TFUNCTIONREF_VISUALIZATION
@@ -533,33 +537,45 @@ namespace UE::Core::Private::Function
 		TFunctionRefBase(const TFunctionRefBase<OtherStorage, Ret (ParamTypes...)>& Other)
 			: Callable(Other.Callable)
 		{
-			if (Callable)
+			if constexpr (OtherStorage::bCanBeNull)
 			{
-				void* NewPtr = Storage.BindCopy(Other.Storage);
+				static_assert(StorageType::bCanBeNull, "Unable to copy from nullable storage into non-nullable storage");
 
-				#if UE_ENABLE_TFUNCTIONREF_VISUALIZATION
-					// Use Memcpy to copy the other DebugPtrStorage, including vptr (because we don't know the bound type
-					// here), and then reseat the underlying pointer.  Possibly even more evil than the Set code.
-					FMemory::Memcpy(&DebugPtrStorage, &Other.DebugPtrStorage, sizeof(DebugPtrStorage)); //-V598
-					DebugPtrStorage.Ptr = NewPtr;
-				#endif
+				if (!Callable)
+				{
+					return;
+				}
 			}
+
+			void* NewPtr = Storage.BindCopy(Other.Storage);
+
+			#if UE_ENABLE_TFUNCTIONREF_VISUALIZATION
+				// Use Memcpy to copy the other DebugPtrStorage, including vptr (because we don't know the bound type
+				// here), and then reseat the underlying pointer.  Possibly even more evil than the Set code.
+				FMemory::Memcpy(&DebugPtrStorage, &Other.DebugPtrStorage, sizeof(DebugPtrStorage)); //-V598
+				DebugPtrStorage.Ptr = NewPtr;
+			#endif
 		}
 
 		TFunctionRefBase(const TFunctionRefBase& Other)
 			: Callable(Other.Callable)
 		{
-			if (Callable)
+			if constexpr (StorageType::bCanBeNull)
 			{
-				void* NewPtr = Storage.BindCopy(Other.Storage);
-
-				#if UE_ENABLE_TFUNCTIONREF_VISUALIZATION
-					// Use Memcpy to copy the other DebugPtrStorage, including vptr (because we don't know the bound type
-					// here), and then reseat the underlying pointer.  Possibly even more evil than the Set code.
-					FMemory::Memcpy(&DebugPtrStorage, &Other.DebugPtrStorage, sizeof(DebugPtrStorage)); //-V598
-					DebugPtrStorage.Ptr = NewPtr;
-				#endif
+				if (!Callable)
+				{
+					return;
+				}
 			}
+
+			void* NewPtr = Storage.BindCopy(Other.Storage);
+
+			#if UE_ENABLE_TFUNCTIONREF_VISUALIZATION
+				// Use Memcpy to copy the other DebugPtrStorage, including vptr (because we don't know the bound type
+				// here), and then reseat the underlying pointer.  Possibly even more evil than the Set code.
+				FMemory::Memcpy(&DebugPtrStorage, &Other.DebugPtrStorage, sizeof(DebugPtrStorage)); //-V598
+				DebugPtrStorage.Ptr = NewPtr;
+			#endif
 		}
 
 		template <
@@ -568,49 +584,67 @@ namespace UE::Core::Private::Function
 		>
 		TFunctionRefBase(FunctorType&& InFunc)
 		{
-			if (auto* Binding = Storage.Bind(Forward<FunctorType>(InFunc)))
+			auto* Binding = Storage.Bind(Forward<FunctorType>(InFunc));
+
+			if constexpr (StorageType::bCanBeNull)
 			{
-				using DecayedFunctorType = typename TRemovePointer<decltype(Binding)>::Type;
+				if (!Binding)
+				{
+					return;
+				}
+			}
 
-				Callable = &TFunctionRefCaller<DecayedFunctorType, Ret (ParamTypes...)>::Call;
+			using DecayedFunctorType = typename TRemovePointer<decltype(Binding)>::Type;
 
-				#if UE_ENABLE_TFUNCTIONREF_VISUALIZATION
-					// We placement new over the top of the same object each time.  This is illegal,
-					// but it ensures that the vptr is set correctly for the bound type, and so is
-					// visualizable.  We never depend on the state of this object at runtime, so it's
-					// ok.
-					new ((void*)&DebugPtrStorage) TDebugHelper<DecayedFunctorType>;
-					DebugPtrStorage.Ptr = (void*)Binding;
-				#endif
-			}
+			Callable = &TFunctionRefCaller<DecayedFunctorType, Ret (ParamTypes...)>::Call;
+
+			#if UE_ENABLE_TFUNCTIONREF_VISUALIZATION
+				// We placement new over the top of the same object each time.  This is illegal,
+				// but it ensures that the vptr is set correctly for the bound type, and so is
+				// visualizable.  We never depend on the state of this object at runtime, so it's
+				// ok.
+				new ((void*)&DebugPtrStorage) TDebugHelper<DecayedFunctorType>;
+				DebugPtrStorage.Ptr = (void*)Binding;
+			#endif
 		}
 
 		TFunctionRefBase& operator=(TFunctionRefBase&&) = delete;
 		TFunctionRefBase& operator=(const TFunctionRefBase&) = delete;
 
 		// Move all of the assert code out of line
-		void CheckCallable() const
+		FORCENOINLINE void CheckCallable() const
 		{
 			checkf(Callable, TEXT("Attempting to call an unbound TFunction!"));
 		}
 
 		Ret operator()(ParamTypes... Params) const
 		{
-			CheckCallable();
+#if DO_CHECK
+			if constexpr (StorageType::bCanBeNull)
+			{
+				CheckCallable();
+			}
+#endif
 			return Callable(Storage.GetPtr(), Params...);
 		}
 
 		~TFunctionRefBase()
 		{
-			if (Callable)
+			if constexpr (StorageType::bCanBeNull)
 			{
-				Storage.Unbind();
+				if (!Callable)
+				{
+					return;
+				}
 			}
+			Storage.Unbind();
 		}
 
 	protected:
 		bool IsSet() const
 		{
+			static_assert(StorageType::bCanBeNull, "Nullable storage cannot be unset");
+
 			return !!Callable;
 		}
 
@@ -676,6 +710,8 @@ namespace UE::Core::Private::Function
 
 	struct FFunctionRefStoragePolicy
 	{
+		constexpr static bool bCanBeNull = false;
+
 		template <typename FunctorType>
 		typename TRemoveReference<FunctorType>::Type* Bind(FunctorType&& InFunc)
 		{
