[Insights] Fixed late connect support for the "stats" trace channel. Changed the batch trace event to always start with an absolute timestamp value (instead of being relative to previous event). This resolves also the use case where a batch event is emitted, but not actually traced in current session (either traced in a previous session and/or overwritten in the tail buffer).

diff --git a/Engine/Source/Developer/TraceServices/Private/Analyzers/StatsTraceAnalysis.cpp b/Engine/Source/Developer/TraceServices/Private/Analyzers/StatsTraceAnalysis.cpp
--- a/Engine/Source/Developer/TraceServices/Private/Analyzers/StatsTraceAnalysis.cpp
+++ b/Engine/Source/Developer/TraceServices/Private/Analyzers/StatsTraceAnalysis.cpp
@@ -27,8 +27,9 @@ void FStatsAnalyzer::OnAnalysisBegin(const FOnAnalysisContext& Context)
 
 	Builder.RouteEvent(RouteId_Spec, "Stats", "Spec");
 	Builder.RouteEvent(RouteId_EventBatch, "Stats", "EventBatch");
+	Builder.RouteEvent(RouteId_EventBatch2, "Stats", "EventBatch2");
 	Builder.RouteEvent(RouteId_BeginFrame, "Misc", "BeginFrame");
-	Builder.RouteEvent(RouteId_EndFrame, "Misc", "EndFrame");
+	//Builder.RouteEvent(RouteId_EndFrame, "Misc", "EndFrame");
 }
 
 void FStatsAnalyzer::OnAnalysisEnd()
@@ -116,10 +117,15 @@ bool FStatsAnalyzer::OnEvent(uint16 RouteId, EStyle Style, const FOnEventContext
 			break;
 		}
 
-		case RouteId_EventBatch:
+		case RouteId_EventBatch: // deprecated in UE 5.3
+		case RouteId_EventBatch2: // added in UE 5.3
 		{
 			uint32 ThreadId = FTraceAnalyzerUtils::GetThreadIdField(Context);
 			TSharedRef<FThreadState> ThreadState = GetThreadState(ThreadId);
+			if (RouteId == RouteId_EventBatch2)
+			{
+				ThreadState->LastCycle = 0; // each batch starts with an absolute timestamp value
+			}
 			TArrayView<const uint8> DataView = FTraceAnalyzerUtils::LegacyAttachmentArray("Data", Context);
 			uint64 BufferSize = DataView.Num();
 			const uint8* BufferPtr = DataView.GetData();
diff --git a/Engine/Source/Developer/TraceServices/Private/Analyzers/StatsTraceAnalysis.h b/Engine/Source/Developer/TraceServices/Private/Analyzers/StatsTraceAnalysis.h
--- a/Engine/Source/Developer/TraceServices/Private/Analyzers/StatsTraceAnalysis.h
+++ b/Engine/Source/Developer/TraceServices/Private/Analyzers/StatsTraceAnalysis.h
@@ -31,6 +31,7 @@ private:
 	{
 		RouteId_Spec,
 		RouteId_EventBatch,
+		RouteId_EventBatch2,
 		RouteId_BeginFrame,
 		RouteId_EndFrame
 	};
diff --git a/Engine/Source/Runtime/Core/Private/Stats/StatsTrace.cpp b/Engine/Source/Runtime/Core/Private/Stats/StatsTrace.cpp
--- a/Engine/Source/Runtime/Core/Private/Stats/StatsTrace.cpp
+++ b/Engine/Source/Runtime/Core/Private/Stats/StatsTrace.cpp
@@ -23,7 +23,7 @@ UE_TRACE_EVENT_BEGIN(Stats, Spec, NoSync|Important)
 	UE_TRACE_EVENT_FIELD(UE::Trace::AnsiString, Group)
 UE_TRACE_EVENT_END()
 
-UE_TRACE_EVENT_BEGIN(Stats, EventBatch)
+UE_TRACE_EVENT_BEGIN(Stats, EventBatch2)
 	UE_TRACE_EVENT_FIELD(uint8[], Data)
 UE_TRACE_EVENT_END()
 
@@ -49,7 +49,6 @@ public:
 
 	struct FThreadState
 	{
-		uint64 TracedBufferLastCycle;
 		uint64 LastCycle;
 		uint16 BufferSize;
 		uint8 Buffer[MaxBufferSize];
@@ -71,7 +70,6 @@ FStatsTraceInternal::FThreadState* FStatsTraceInternal::InitThreadState()
 {
 	LLM_SCOPE_BYNAME(TEXT("Trace/Stats"));
 	ThreadLocalThreadState = new FThreadState();
-	ThreadLocalThreadState->TracedBufferLastCycle = 0;
 	ThreadLocalThreadState->LastCycle = 0;
 	ThreadLocalThreadState->BufferSize = 0;
 	return ThreadLocalThreadState;
@@ -79,18 +77,10 @@ FStatsTraceInternal::FThreadState* FStatsTraceInternal::InitThreadState()
 
 void FStatsTraceInternal::FlushThreadBuffer(FThreadState* ThreadState)
 {
-	if (UE_TRACE_CHANNELEXPR_IS_ENABLED(StatsChannel))
-	{
-		UE_TRACE_LOG(Stats, EventBatch, StatsChannel)
-			<< EventBatch.Data(ThreadState->Buffer, ThreadState->BufferSize);
+	UE_TRACE_LOG(Stats, EventBatch2, StatsChannel)
+		<< EventBatch2.Data(ThreadState->Buffer, ThreadState->BufferSize);
 
-		ThreadState->TracedBufferLastCycle = ThreadState->LastCycle;
-	}
-	else
-	{
-		// If the current buffer was not actually traced, we need to reset the base timestamp.
-		ThreadState->LastCycle = ThreadState->TracedBufferLastCycle;
-	}
+	ThreadState->LastCycle = 0; // each batch starts with an absolute timestamp value
 	ThreadState->BufferSize = 0;
 }
 
