[Insights] Fix late connect for the Counters trace channel.

diff --git a/Engine/Source/Runtime/Core/Private/ProfilingDebugging/CountersTrace.cpp b/Engine/Source/Runtime/Core/Private/ProfilingDebugging/CountersTrace.cpp
--- a/Engine/Source/Runtime/Core/Private/ProfilingDebugging/CountersTrace.cpp
+++ b/Engine/Source/Runtime/Core/Private/ProfilingDebugging/CountersTrace.cpp
@@ -27,7 +27,6 @@ UE_TRACE_EVENT_BEGIN(Counters, SetValueFloat)
 	UE_TRACE_EVENT_FIELD(uint16, CounterId)
 UE_TRACE_EVENT_END()
 
-
 uint16 FCountersTrace::OutputInitCounter(const TCHAR* CounterName, ETraceCounterType CounterType, ETraceCounterDisplayHint CounterDisplayHint)
 {
 	if (!UE_TRACE_CHANNELEXPR_IS_ENABLED(CountersChannel))
@@ -35,9 +34,9 @@ uint16 FCountersTrace::OutputInitCounter(const TCHAR* CounterName, ETraceCounter
 		return 0;
 	}
 
-	static TAtomic<uint16> NextId;
+	static std::atomic<uint32> NextId { 0 };
 
-	uint16 CounterId = uint16(NextId++) + 1;
+	uint16 CounterId = uint16(++NextId);
 	uint16 NameLen = uint16(FCString::Strlen(CounterName));
 	UE_TRACE_LOG(Counters, Spec, CountersChannel, NameLen * sizeof(ANSICHAR))
 		<< Spec.Id(CounterId)
@@ -73,4 +72,4 @@ void FCountersTrace::OutputSetValue(uint16 CounterId, double Value)
 		<< SetValueFloat.CounterId(CounterId);
 }
 
-#endif
\ No newline at end of file
+#endif // COUNTERSTRACE_ENABLED
diff --git a/Engine/Source/Runtime/Core/Public/ProfilingDebugging/CountersTrace.h b/Engine/Source/Runtime/Core/Public/ProfilingDebugging/CountersTrace.h
--- a/Engine/Source/Runtime/Core/Public/ProfilingDebugging/CountersTrace.h
+++ b/Engine/Source/Runtime/Core/Public/ProfilingDebugging/CountersTrace.h
@@ -6,12 +6,12 @@
 #include "HAL/PreprocessorHelpers.h"
 #include "Misc/Build.h"
 #include "Trace/Config.h"
+#include "Trace/Detail/Channel.h"
+#include "Trace/Detail/Channel.inl"
 #include "Trace/Trace.h"
 
 #include <atomic>
 
-namespace UE { namespace Trace { class FChannel; } }
-
 #if !defined(COUNTERSTRACE_ENABLED)
 #if UE_TRACE_ENABLED && !UE_BUILD_SHIPPING
 #define COUNTERSTRACE_ENABLED 1
@@ -34,7 +34,7 @@ enum ETraceCounterDisplayHint
 
 #if COUNTERSTRACE_ENABLED
 
-UE_TRACE_CHANNEL_EXTERN(CountersChannel);
+UE_TRACE_CHANNEL_EXTERN(CountersChannel, CORE_API);
 
 struct FCountersTrace
 {
@@ -52,7 +52,17 @@ struct FCountersTrace
 			, CounterId(0)
 			, CounterDisplayHint(InCounterDisplayHint)
 		{
-			CounterId = OutputInitCounter(InCounterName, CounterType, CounterDisplayHint);
+			EnsureInit();
+		}
+
+		void EnsureInit()
+		{
+			uint32 OldId = CounterId.load();
+			if (!OldId)
+			{
+				uint32 NewId = OutputInitCounter(CounterName, CounterType, CounterDisplayHint);
+				CounterId.compare_exchange_weak(OldId, NewId);
+			}
 		}
 
 		const ValueType& Get() const
@@ -65,7 +75,11 @@ struct FCountersTrace
 			if (Value != InValue)
 			{
 				Value = InValue;
-				OutputSetValue(CounterId, Value);
+				if (UE_TRACE_CHANNELEXPR_IS_ENABLED(CountersChannel))
+				{
+					EnsureInit();
+					OutputSetValue(uint16(CounterId), Value);
+				}
 			}
 		}
 
@@ -74,7 +88,11 @@ struct FCountersTrace
 			if (InValue != 0)
 			{
 				Value += InValue;
-				OutputSetValue(CounterId, Value);
+				if (UE_TRACE_CHANNELEXPR_IS_ENABLED(CountersChannel))
+				{
+					EnsureInit();
+					OutputSetValue(uint16(CounterId), Value);
+				}
 			}
 		}
 
@@ -83,26 +101,38 @@ struct FCountersTrace
 			if (InValue != 0)
 			{
 				Value -= InValue;
-				OutputSetValue(CounterId, Value);
+				if (UE_TRACE_CHANNELEXPR_IS_ENABLED(CountersChannel))
+				{
+					EnsureInit();
+					OutputSetValue(uint16(CounterId), Value);
+				}
 			}
 		}
 
 		void Increment()
 		{
 			++Value;
-			OutputSetValue(CounterId, Value);
+			if (UE_TRACE_CHANNELEXPR_IS_ENABLED(CountersChannel))
+			{
+				EnsureInit();
+				OutputSetValue(uint16(CounterId), Value);
+			}
 		}
 		
 		void Decrement()
 		{
 			--Value;
-			OutputSetValue(CounterId, Value);
+			if (UE_TRACE_CHANNELEXPR_IS_ENABLED(CountersChannel))
+			{
+				EnsureInit();
+				OutputSetValue(uint16(CounterId), Value);
+			}
 		}
 
 	private:
 		StoredType Value;
 		const TCHAR* CounterName;
-		uint16 CounterId;
+		std::atomic<uint32> CounterId;
 		ETraceCounterDisplayHint CounterDisplayHint;
 	};
 
