made FPipe::WaitUntilEmpty() to wait only for its own tasks (potentially retracting and executing them) instead of using busy-waiting that can peek also unrelated tasks

diff --git a/Engine/Source/Runtime/Core/Private/Tasks/Pipe.cpp b/Engine/Source/Runtime/Core/Private/Tasks/Pipe.cpp
--- a/Engine/Source/Runtime/Core/Private/Tasks/Pipe.cpp
+++ b/Engine/Source/Runtime/Core/Private/Tasks/Pipe.cpp
@@ -9,36 +9,82 @@ namespace UE::Tasks
 {
 	Private::FTaskBase* FPipe::PushIntoPipe(Private::FTaskBase& Task)
 	{
-		// critical section: between exchanging tasks and adding the previous LastTask to prerequisites.
-		// `LastTask` can be in the process of destruction. Before this it will try to clear itself from the pipe
-		// we need to let it know we're still using it so destruction should wait until we finished
-		PushingThreadsNum.fetch_add(1, std::memory_order_relaxed);
-		ON_SCOPE_EXIT{ PushingThreadsNum.fetch_sub(1, std::memory_order_relaxed); };
-
+		Task.AddRef(); // the pipe holds a ref to the last task, until it's replaced by the next task or cleared on completion
 		Private::FTaskBase* LastTask_Local = LastTask.exchange(&Task, std::memory_order_release);
 		checkf(LastTask_Local != &Task, TEXT("Dependency cycle: adding itself as a prerequisite (or use after destruction)"));
 
-		if (LastTask_Local == nullptr || !LastTask_Local->AddSubsequent(Task))
+		if (LastTask_Local == nullptr)
+		{
+			return nullptr;
+		}
+
+		if (!LastTask_Local->AddSubsequent(Task))
 		{
+			// the last task doesn't accept subsequents anymore because it's already completed (happened concurrently after we replaced it as
+			// the last pipe's task
+			LastTask_Local->Release(); // the pipe doesn't need it anymore
 			return nullptr;
 		}
 
-		LastTask_Local->AddRef(); // keep it alive, must be called before leaving this function, it's the caller's duty to release it
-		return LastTask_Local;
+		return LastTask_Local; // transfer the reference to the caller that must release it
 	}
 
-	void FPipe::ClearTask(Private::FTaskBase& Task)
+	bool FPipe::TryClearTask(Private::FTaskBase& Task)
 	{
-		// clears `Task` from the pipe if it's still the `LastTask`. if it's not, another task already has set itself as the `LastTask` 
-		// and can be in the process of adding this `Task` as a prerequisite. We need to wait for this to complete before returning 
-		// from the call as the `Task` can be destroyed immediately after this
 		Private::FTaskBase* Task_Local = &Task;
-		if (!LastTask.compare_exchange_strong(Task_Local, nullptr, std::memory_order_acquire, std::memory_order_relaxed))
+		// try clearing the task if it's still pipe's "last task". if succeeded, release the ref accounted for pipe's last task. otherwise whoever replaced it
+		// as the last task will do this
+		if (LastTask.compare_exchange_strong(Task_Local, nullptr, std::memory_order_acquire, std::memory_order_relaxed))
+		{
+			Task.Release(); // it was still pipe's last task. now that we cleared it, release the reference
+			return true;
+		}
+
+		return false;
+	}
+
+	bool FPipe::WaitUntilEmpty(FTimespan Timeout/* = FTimespan::FromMilliseconds(MAX_uint32)*/)
+	{
+		struct FPlaceholderTask final : Private::FTaskBase
 		{
-			while (PushingThreadsNum.load(std::memory_order_relaxed) != 0)
+			FPlaceholderTask()
+				// InitRefCount: one for the initial local ref, and one for the internal reference
+				: Private::FTaskBase(/*InitRefCount =*/ 2)
 			{
+				Init(
+					TEXT("Pipe::WaitUntilEmpty() placeholder"),
+					ETaskPriority::Normal, // doesn't matter
+					EExtendedTaskPriority::TaskEvent // no need for scheduling or execution for this dummy task
+				);
 			}
+
+			virtual bool TryExecuteTaskVirtual() final
+			{
+				checkNoEntry(); // the method won't be called because the task was initialized with `EExtendedTaskPriority::TaskEvent`
+				return true;
+			}
+		};
+
+		bool bRes = true;
+
+		TRefCountPtr<FPlaceholderTask> PlaceholderTask{ new FPlaceholderTask, /*bAddRef =*/ false }; // the initial ref was already accounted for
+		Private::FTaskBase* LastTask_Local = PushIntoPipe(*PlaceholderTask);
+		if (LastTask_Local)
+		{
+			LastTask_Local->Release(); // release pipe's ref to the last task
+			// when the placeholder replaced "the last task" in the pipe, it became its subsequent, so the last task will "complete" it
+			bRes = PlaceholderTask->Wait(Timeout);
 		}
+		else
+		{
+			// there was no "last task" in the pipe, so we have to launch the placeholder to complete it to not assert on its destruction
+			PlaceholderTask->TryLaunch(/*TaskSize =*/ 0);
+			checkSlow(PlaceholderTask->IsCompleted());
+		}
+
+		verifyf(TryClearTask(*PlaceholderTask), TEXT("More tasks were launched (concurrently) in the pipe after `WaitUntilEmpty()` was called"));
+		
+		return bRes;
 	}
 
 	// Maintains pipe callstack. Due to busy waiting tasks from multiple pipes can be executed nested.
diff --git a/Engine/Source/Runtime/Core/Private/Tasks/TaskPrivate.cpp b/Engine/Source/Runtime/Core/Private/Tasks/TaskPrivate.cpp
--- a/Engine/Source/Runtime/Core/Private/Tasks/TaskPrivate.cpp
+++ b/Engine/Source/Runtime/Core/Private/Tasks/TaskPrivate.cpp
@@ -283,7 +283,7 @@ namespace UE::Tasks
 
 		void FTaskBase::ClearPipe()
 		{
-			GetPipe()->ClearTask(*this);
+			GetPipe()->TryClearTask(*this);
 		}
 
 		static thread_local FTaskBase* CurrentTask = nullptr;
diff --git a/Engine/Source/Runtime/Core/Private/Tests/Tasks/TasksTest.cpp b/Engine/Source/Runtime/Core/Private/Tests/Tasks/TasksTest.cpp
--- a/Engine/Source/Runtime/Core/Private/Tests/Tasks/TasksTest.cpp
+++ b/Engine/Source/Runtime/Core/Private/Tests/Tasks/TasksTest.cpp
@@ -655,6 +655,54 @@ namespace UE { namespace TasksTests
 		check(TasksExecutedNum == TasksNum);
 	}
 
+	IMPLEMENT_SIMPLE_AUTOMATION_TEST(FTasksPipeWaitUntilEmptyTest, "System.Core.Tasks.Pipe.WaitUntilEmpty", EAutomationTestFlags::ApplicationContextMask | EAutomationTestFlags::EngineFilter);
+
+	bool FTasksPipeWaitUntilEmptyTest::RunTest(const FString& Parameters)
+	{
+		//for (int i = 0; i != 100000; ++i)
+		{	// waiting until an empty pipe is empty
+			FPipe Pipe{ UE_SOURCE_LOCATION };
+			Pipe.WaitUntilEmpty();
+		}
+
+		//for (int i = 0; i != 100000; ++i)
+		{	// waiting until a not empty pipe is empty
+			FPipe Pipe{ UE_SOURCE_LOCATION };
+			Pipe.Launch(UE_SOURCE_LOCATION, [] {});
+			Pipe.WaitUntilEmpty();
+		}
+
+		if (false) // disabled as it's an invalid case, but keep it compiling
+		//for (int i = 0; i != 100000; ++i)
+		{	// waiting until a not empty pipe is empty, while new tasks are piped. just an example, this would assert in `Pipe.WaitUntilEmtpy()`
+			// as no more tasks can be piped after that
+
+			FPipe Pipe{ UE_SOURCE_LOCATION };
+
+			std::atomic<bool> bQuit{ false };
+
+			auto TaskBody = [&Pipe, &bQuit](auto& TaskBodyRef) -> void
+			{
+				if (!bQuit)
+				{
+					Pipe.Launch(UE_SOURCE_LOCATION, [&TaskBodyRef] { TaskBodyRef(TaskBodyRef); });
+				}
+			};
+
+			// start the chain
+			TaskBody(TaskBody);
+
+			check(!Pipe.WaitUntilEmpty(FTimespan::FromMilliseconds(100))); // it's not empty until we break the chain
+
+			// break the chain
+			bQuit = true;
+
+			check(Pipe.WaitUntilEmpty());
+		}
+
+		return true;
+	}
+
 	struct FAutoTlsSlot
 	{
 		uint32 Slot;
diff --git a/Engine/Source/Runtime/Core/Public/Tasks/Pipe.h b/Engine/Source/Runtime/Core/Public/Tasks/Pipe.h
--- a/Engine/Source/Runtime/Core/Public/Tasks/Pipe.h
+++ b/Engine/Source/Runtime/Core/Public/Tasks/Pipe.h
@@ -47,10 +47,7 @@ namespace UE::Tasks
 
 		// waits until the pipe is empty (its last task is executed)
 		// should be used only after no more tasks are launched in the pipe, e.g. preparing for the pipe destruction
-		void WaitUntilEmpty()
-		{
-			LowLevelTasks::FScheduler::Get().BusyWaitUntil([this] { return !HasWork(); });
-		}
+		CORE_API bool WaitUntilEmpty(FTimespan Timeout = FTimespan::FromMilliseconds(MAX_uint32));
 
 		// launches a task in the pipe
 		// @param InDebugName helps to identify the task in debugger and profiler
@@ -110,20 +107,22 @@ namespace UE::Tasks
 		friend class Private::FTaskBase;
 
 		// pushes given task into the pipe: adds the task as a subsequent to the last task if any and sets it as the new last task
-		// returns the previous piped task if we managed to register the given task as its subsequent, otherwise nullptr
-		Private::FTaskBase* PushIntoPipe(Private::FTaskBase& Task);
+		// returns the accounted reference to the previous piped task if we managed to register the given task as its subsequent, otherwise nullptr.
+		// the reference must be released by the caller when not needed anymore
+		[[nodiscard]] Private::FTaskBase* PushIntoPipe(Private::FTaskBase& Task);
+
 		// pipe holds a "weak" reference to a task. the task must be cleared from the pipe when its execution finished before its completion, 
-		// otherwise the next piped task can try to add itself as a subsequent to already destroyed task
-		void ClearTask(Private::FTaskBase& Task);
+		// otherwise the next piped task can try to add itself as a subsequent to already destroyed task.
+		// returns true if the task was still the last one and so was successfully cleared
+		bool TryClearTask(Private::FTaskBase& Task);
 
 		// notifications about pipe's task execution
 		void ExecutionStarted();
 		void ExecutionFinished();
 
 	private:
-		std::atomic<Private::FTaskBase*> LastTask{ nullptr }; // the last task that 
-		// a counter of how many threads trying to push a task into a pipe right now
-		std::atomic<int32> PushingThreadsNum{ 0 };
+		// pipe builds a chain (a linked list) of tasks and so needs to store only the last one. the last task is null if the pipe is not blocked
+		std::atomic<Private::FTaskBase*> LastTask{ nullptr };
 
 	public:
 		FORCENOINLINE const TCHAR* GetDebugName() const
diff --git a/Engine/Source/Runtime/Core/Public/Tasks/TaskPrivate.h b/Engine/Source/Runtime/Core/Public/Tasks/TaskPrivate.h
--- a/Engine/Source/Runtime/Core/Public/Tasks/TaskPrivate.h
+++ b/Engine/Source/Runtime/Core/Public/Tasks/TaskPrivate.h
@@ -540,7 +540,7 @@ namespace UE::Tasks
 							if (PrevPipedTask != nullptr) // the pipe is blocked
 							{
 								// the prev task in pipe's chain becomes this task's prerequisite, to enabled piped task retraction.
-								// no need to AddRef as it's already sorted in `FPipe::PushIntoPipe`
+								// its ref count already accounted for this ref. the ref will be released when the prereq is not needed anymore
 								Prerequisites.Push(PrevPipedTask);
 								return false;
 							}
