Deprecating global Sort and StableSort functions. They only support an index type of int32 and can produce unexpected/undesirable behaviour due to their use of TDereferenceWrapper.

TArray has been updated to call Algo::Sort/Algo::StableSort, but retains the use of TDereferenceWrapper to keep the same behaviour.
TArrayView has been updated to call Algo::Sort/Algo::StableSort, but retains the use of TDereferenceWrapper to keep the same behaviour.
TSparseArray has been updated to call Algo::Sort/Algo::StableSort - TDereferenceWrapper didn't do anything for this since it operated on FElementOrFreeListLink with FElementCompareClass, so I haven't added it here.

diff --git a/Engine/Source/Runtime/Core/Public/Containers/Array.h b/Engine/Source/Runtime/Core/Public/Containers/Array.h
--- a/Engine/Source/Runtime/Core/Public/Containers/Array.h
+++ b/Engine/Source/Runtime/Core/Public/Containers/Array.h
@@ -18,6 +18,7 @@
 #include "Algo/HeapSort.h"
 #include "Algo/IsHeap.h"
 #include "Algo/Impl/BinaryHeap.h"
+#include "Algo/StableSort.h"
 #include "Concepts/GetTypeHashable.h"
 #include "Templates/AndOrNot.h"
 #include "Templates/IdentityFunctor.h"
@@ -2840,12 +2841,12 @@ public:
 	 *
 	 * @note: If your array contains raw pointers, they will be automatically dereferenced during sorting.
 	 *        Therefore, your array will be sorted by the values being pointed to, rather than the pointers' values.
-	 *        If this is not desirable, please use Algo::Sort(MyArray) instead.
+	 *        If this is not desirable, please use Algo::Sort(MyArray) directly instead.
 	 *        The auto-dereferencing behavior does not occur with smart pointers.
 	 */
 	void Sort()
 	{
-		::Sort(GetData(), Num());
+		Algo::Sort(*this, TDereferenceWrapper<ElementType, TLess<>>(TLess<>()));
 	}
 
 	/**
@@ -2855,13 +2856,14 @@ public:
 	 *
 	 * @note: If your array contains raw pointers, they will be automatically dereferenced during sorting.
 	 *        Therefore, your predicate will be passed references rather than pointers.
-	 *        If this is not desirable, please use Algo::Sort(MyArray, Predicate) instead.
+	 *        If this is not desirable, please use Algo::Sort(MyArray, Predicate) directly instead.
 	 *        The auto-dereferencing behavior does not occur with smart pointers.
 	 */
 	template <class PREDICATE_CLASS>
 	void Sort(const PREDICATE_CLASS& Predicate)
 	{
-		::Sort(GetData(), Num(), Predicate);
+		TDereferenceWrapper<ElementType, PREDICATE_CLASS> PredicateWrapper(Predicate);
+		Algo::Sort(*this, PredicateWrapper);
 	}
 
 	/**
@@ -2871,11 +2873,12 @@ public:
 	 *
 	 * @note: If your array contains raw pointers, they will be automatically dereferenced during sorting.
 	 *        Therefore, your array will be sorted by the values being pointed to, rather than the pointers' values.
+	 *        If this is not desirable, please use Algo::StableSort(MyArray) directly instead.
 	 *        The auto-dereferencing behavior does not occur with smart pointers.
 	 */
 	void StableSort()
 	{
-		::StableSort(GetData(), Num());
+		Algo::StableSort(*this, TDereferenceWrapper<ElementType, TLess<>>(TLess<>()));
 	}
 
 	/**
@@ -2887,12 +2890,14 @@ public:
 	 *
 	 * @note: If your array contains raw pointers, they will be automatically dereferenced during sorting.
 	 *        Therefore, your predicate will be passed references rather than pointers.
+	 *        If this is not desirable, please use Algo::StableSort(MyArray, Predicate) directly instead.
 	 *        The auto-dereferencing behavior does not occur with smart pointers.
 	 */
 	template <class PREDICATE_CLASS>
 	void StableSort(const PREDICATE_CLASS& Predicate)
 	{
-		::StableSort(GetData(), Num(), Predicate);
+		TDereferenceWrapper<ElementType, PREDICATE_CLASS> PredicateWrapper(Predicate);
+		Algo::StableSort(*this, PredicateWrapper);
 	}
 
 #if defined(_MSC_VER) && !defined(__clang__)	// Relies on MSVC-specific lazy template instantiation to support arrays of incomplete types
diff --git a/Engine/Source/Runtime/Core/Public/Containers/ArrayView.h b/Engine/Source/Runtime/Core/Public/Containers/ArrayView.h
--- a/Engine/Source/Runtime/Core/Public/Containers/ArrayView.h
+++ b/Engine/Source/Runtime/Core/Public/Containers/ArrayView.h
@@ -692,31 +692,45 @@ public:
 public:
 	/**
 	 * Sorts the array assuming < operator is defined for the item type.
+	 *
+	 * @note: If your array contains raw pointers, they will be automatically dereferenced during sorting.
+	 *        Therefore, your array will be sorted by the values being pointed to, rather than the pointers' values.
+	 *        If this is not desirable, please use Algo::Sort(MyArray) directly instead.
+	 *        The auto-dereferencing behavior does not occur with smart pointers.
 	 */
 	void Sort()
 	{
-		::Sort(GetData(), Num());
+		Algo::Sort(*this, TDereferenceWrapper<ElementType, TLess<>>(TLess<>()));
 	}
 
 	/**
 	 * Sorts the array using user define predicate class.
 	 *
 	 * @param Predicate Predicate class instance.
+	 * @note: If your array contains raw pointers, they will be automatically dereferenced during sorting.
+	 *        Therefore, your predicate will be passed references rather than pointers.
+	 *        If this is not desirable, please use Algo::Sort(MyArray, Predicate) directly instead.
+	 *        The auto-dereferencing behavior does not occur with smart pointers.
 	 */
 	template <class PREDICATE_CLASS>
 	void Sort(const PREDICATE_CLASS& Predicate)
 	{
-		::Sort(GetData(), Num(), Predicate);
+		TDereferenceWrapper<ElementType, PREDICATE_CLASS> PredicateWrapper(Predicate);
+		Algo::Sort(*this, PredicateWrapper);
 	}
 
 	/**
 	 * Stable sorts the array assuming < operator is defined for the item type.
 	 *
 	 * Stable sort is slower than non-stable algorithm.
+	 * @note: If your array contains raw pointers, they will be automatically dereferenced during sorting.
+	 *        Therefore, your array will be sorted by the values being pointed to, rather than the pointers' values.
+	 *        If this is not desirable, please use Algo::StableSort(MyArray) directly instead.
+	 *        The auto-dereferencing behavior does not occur with smart pointers.
 	 */
 	void StableSort()
 	{
-		::StableSort(GetData(), Num());
+		Algo::StableSort(*this, TDereferenceWrapper<ElementType, TLess<>>(TLess<>()));
 	}
 
 	/**
@@ -725,11 +739,16 @@ public:
 	 * Stable sort is slower than non-stable algorithm.
 	 *
 	 * @param Predicate Predicate class instance
+	 * @note: If your array contains raw pointers, they will be automatically dereferenced during sorting.
+	 *        Therefore, your predicate will be passed references rather than pointers.
+	 *        If this is not desirable, please use Algo::StableSort(MyArray, Predicate) directly instead.
+	 *        The auto-dereferencing behavior does not occur with smart pointers.
 	 */
 	template <class PREDICATE_CLASS>
 	void StableSort(const PREDICATE_CLASS& Predicate)
 	{
-		::StableSort(GetData(), Num(), Predicate);
+		TDereferenceWrapper<ElementType, PREDICATE_CLASS> PredicateWrapper(Predicate);
+		Algo::StableSort(*this, PredicateWrapper);
 	}
 
 private:
diff --git a/Engine/Source/Runtime/Core/Public/Containers/SparseArray.h b/Engine/Source/Runtime/Core/Public/Containers/SparseArray.h
--- a/Engine/Source/Runtime/Core/Public/Containers/SparseArray.h
+++ b/Engine/Source/Runtime/Core/Public/Containers/SparseArray.h
@@ -588,7 +588,7 @@ public:
 			Compact();
 
 			// Sort the elements according to the provided comparison class.
-			::Sort( (FElementOrFreeListLink*)Data.GetData(), Num(), FElementCompareClass< PREDICATE_CLASS >( Predicate ) );
+			Algo::Sort(TArrayView<FElementOrFreeListLink>(Data.GetData(), Num()), FElementCompareClass< PREDICATE_CLASS >( Predicate ) );
 		}
 	}
 
@@ -608,7 +608,7 @@ public:
 			CompactStable();
 
 			// Sort the elements according to the provided comparison class.
-			::StableSort((FElementOrFreeListLink*)Data.GetData(), Num(), FElementCompareClass< PREDICATE_CLASS >(Predicate));
+			Algo::StableSort(TArrayView<FElementOrFreeListLink>(Data.GetData(), Num()), FElementCompareClass< PREDICATE_CLASS >(Predicate));
 		}
 	}
 
diff --git a/Engine/Source/Runtime/Core/Public/Templates/Sorting.h b/Engine/Source/Runtime/Core/Public/Templates/Sorting.h
--- a/Engine/Source/Runtime/Core/Public/Templates/Sorting.h
+++ b/Engine/Source/Runtime/Core/Public/Templates/Sorting.h
@@ -74,6 +74,7 @@ struct TIsContiguousContainer< TArrayRange<T> >
  * @param Predicate predicate class
  */
 template<class T, class PREDICATE_CLASS> 
+UE_DEPRECATED(5.3, "Sort is deprecated, please use Algo::Sort. Algo::Sort supports ranges with index types other than int32, and doesn't automatically dereference pointers.")
 void Sort( T* First, const int32 Num, const PREDICATE_CLASS& Predicate )
 {
 	TArrayRange<T> ArrayRange( First, Num );
@@ -88,6 +89,7 @@ void Sort( T* First, const int32 Num, const PREDICATE_CLASS& Predicate )
  * @param Predicate predicate class
  */
 template<class T, class PREDICATE_CLASS> 
+UE_DEPRECATED(5.3, "Sort is deprecated, please use Algo::Sort. Algo::Sort supports ranges with index types other than int32, and doesn't automatically dereference pointers.")
 void Sort( T** First, const int32 Num, const PREDICATE_CLASS& Predicate )
 {
 	TArrayRange<T*> ArrayRange( First, Num );
@@ -102,6 +104,7 @@ void Sort( T** First, const int32 Num, const PREDICATE_CLASS& Predicate )
  * @param	Num		the number of items to sort
  */
 template<class T> 
+UE_DEPRECATED(5.3, "Sort is deprecated, please use Algo::Sort. Algo::Sort supports ranges with index types other than int32, and doesn't automatically dereference pointers.")
 void Sort( T* First, const int32 Num )
 {
 	TArrayRange<T> ArrayRange( First, Num );
@@ -115,6 +118,7 @@ void Sort( T* First, const int32 Num )
  * @param	Num		the number of items to sort
  */
 template<class T> 
+UE_DEPRECATED(5.3, "Sort is deprecated, please use Algo::Sort. Algo::Sort supports ranges with index types other than int32, and doesn't automatically dereference pointers.")
 void Sort( T** First, const int32 Num )
 {
 	TArrayRange<T*> ArrayRange( First, Num );
@@ -369,6 +373,7 @@ void StableSortInternal(T* First, const int32 Num, const PREDICATE_CLASS& Predic
  * @param Predicate predicate class
  */
 template<class T, class PREDICATE_CLASS>
+UE_DEPRECATED(5.3, "StableSort is deprecated, please use Algo::StableSort. Algo::StableSort supports ranges with index types other than int32, and doesn't automatically dereference pointers.")
 void StableSort(T* First, const int32 Num, const PREDICATE_CLASS& Predicate)
 {
 	StableSortInternal(First, Num, TDereferenceWrapper<T, PREDICATE_CLASS>(Predicate));
@@ -383,6 +388,7 @@ void StableSort(T* First, const int32 Num, const PREDICATE_CLASS& Predicate)
  * @param Predicate predicate class
  */
 template<class T, class PREDICATE_CLASS>
+UE_DEPRECATED(5.3, "StableSort is deprecated, please use Algo::StableSort. Algo::StableSort supports ranges with index types other than int32, and doesn't automatically dereference pointers.")
 void StableSort(T** First, const int32 Num, const PREDICATE_CLASS& Predicate)
 {
 	StableSortInternal(First, Num, TDereferenceWrapper<T*, PREDICATE_CLASS>(Predicate));
@@ -398,6 +404,7 @@ void StableSort(T** First, const int32 Num, const PREDICATE_CLASS& Predicate)
  * @param	Num		the number of items to sort
  */
 template<class T>
+UE_DEPRECATED(5.3, "StableSort is deprecated, please use Algo::StableSort. Algo::StableSort supports ranges with index types other than int32, and doesn't automatically dereference pointers.")
 void StableSort(T* First, const int32 Num)
 {
 	StableSortInternal(First, Num, TDereferenceWrapper<T, TLess<T> >(TLess<T>()));
@@ -411,6 +418,7 @@ void StableSort(T* First, const int32 Num)
  * @param	Num		the number of items to sort
  */
 template<class T>
+UE_DEPRECATED(5.3, "StableSort is deprecated, please use Algo::StableSort. Algo::StableSort supports ranges with index types other than int32, and doesn't automatically dereference pointers.")
 void StableSort(T** First, const int32 Num)
 {
 	StableSortInternal(First, Num, TDereferenceWrapper<T*, TLess<T> >(TLess<T>()));
