Modify FWindowsPlatformStackWalk to use RtlVirtualUnwind  rather than StackWalk64.
Makes it safe to walk a callstack on another suspended thread which may be in the middle of LoadLibrary, preventing us from loading the dbghelp library.

diff --git a/Engine/Source/Runtime/Core/Private/Microsoft/MicrosoftPlatformStackWalk.cpp b/Engine/Source/Runtime/Core/Private/Microsoft/MicrosoftPlatformStackWalk.cpp
--- a/Engine/Source/Runtime/Core/Private/Microsoft/MicrosoftPlatformStackWalk.cpp
+++ b/Engine/Source/Runtime/Core/Private/Microsoft/MicrosoftPlatformStackWalk.cpp
@@ -74,4 +74,77 @@ bool FMicrosoftPlatformStackWalk::ExtractInfoFromModule(void* ProcessHandle, voi
 	}
 
 	return true;
+}
+
+int32 FMicrosoftPlatformStackWalk::CaptureStackTraceInternal(uint64* OutBacktrace, uint32 MaxDepth, void* InContext, void* InThreadHandle, uint32* OutDepth)
+{
+	uint32 CurrentDepth = 0;
+	HANDLE ThreadHandle = reinterpret_cast<HANDLE>(InThreadHandle);
+	CONTEXT ContextCopy = *reinterpret_cast<PCONTEXT>(InContext);
+
+	*OutDepth = 0;
+
+#if !PLATFORM_SEH_EXCEPTIONS_DISABLED
+	__try
+#endif
+	{
+		UNWIND_HISTORY_TABLE UnwindHistoryTable;
+		RtlZeroMemory(&UnwindHistoryTable, sizeof(UNWIND_HISTORY_TABLE));
+
+		for (; CurrentDepth < MaxDepth; ++CurrentDepth)
+		{
+			// If we reach an RIP of zero, this means that we've walked off the end
+			// of the call stack and are done.
+			if (!ContextCopy.Rip)
+				break;
+
+			OutBacktrace[CurrentDepth] = ContextCopy.Rip;
+			*OutDepth = CurrentDepth;
+
+			ULONG64 ImageBase = 0;
+			if (PRUNTIME_FUNCTION RuntimeFunction = RtlLookupFunctionEntry(ContextCopy.Rip, &ImageBase, &UnwindHistoryTable); RuntimeFunction)
+			{
+				// Otherwise, call upon RtlVirtualUnwind to execute the unwind for us.
+				KNONVOLATILE_CONTEXT_POINTERS NVContext;
+				RtlZeroMemory(&NVContext, sizeof(KNONVOLATILE_CONTEXT_POINTERS));
+
+				PVOID HandlerData = nullptr;
+				ULONG64 EstablisherFrame = 0;
+
+				RtlVirtualUnwind(
+					0 /*UNW_FLAG_NHANDLER*/,
+					ImageBase,
+					ContextCopy.Rip,
+					RuntimeFunction,
+					&ContextCopy,
+					&HandlerData,
+					&EstablisherFrame,
+					&NVContext);
+			}
+			else
+			{
+				// https://docs.microsoft.com/en-us/cpp/build/x64-calling-convention?view=msvc-170
+				// Some functions are 'leaf' functions and do not have unwind info, they can be unwound
+				// by simulating a return instruction.
+				ContextCopy.Rip = (uint64)(*(uint64*)ContextCopy.Rsp);
+				ContextCopy.Rsp += 8;
+			}
+		}
+	}
+#if !PLATFORM_SEH_EXCEPTIONS_DISABLED
+	__except (EXCEPTION_EXECUTE_HANDLER)
+	{
+		// We need to catch any exceptions within this function so they don't get sent to
+		// the engine's error handler, hence causing an infinite loop.
+		return EXCEPTION_EXECUTE_HANDLER;
+	}
+#endif
+
+	// NULL out remaining entries.
+	for (; CurrentDepth < MaxDepth; CurrentDepth++)
+	{
+		OutBacktrace[CurrentDepth] = 0;
+	}
+
+	return EXCEPTION_EXECUTE_HANDLER;
 }
\ No newline at end of file
diff --git a/Engine/Source/Runtime/Core/Private/Windows/WindowsPlatformStackWalk.cpp b/Engine/Source/Runtime/Core/Private/Windows/WindowsPlatformStackWalk.cpp
--- a/Engine/Source/Runtime/Core/Private/Windows/WindowsPlatformStackWalk.cpp
+++ b/Engine/Source/Runtime/Core/Private/Windows/WindowsPlatformStackWalk.cpp
@@ -104,122 +104,6 @@ struct FWindowsThreadContextWrapper
 	}
 };
 
-/**
- * Helper function performing the actual stack walk. This code relies on the symbols being loaded for best results
- * walking the stack albeit at a significant performance penalty.
- *
- * This helper function is designed to be called within a structured exception handler.
- *
- * @param	BackTrace			Array to write backtrace to
- * @param	MaxDepth			Maximum depth to walk - needs to be less than or equal to array size
- * @param	Context				Thread context information
- * @return	EXCEPTION_EXECUTE_HANDLER
- */
-static int32 CaptureStackTraceHelper(uint64 *BackTrace, uint32 MaxDepth, FWindowsThreadContextWrapper* ContextWapper, uint32* Depth)
-{
-	STACKFRAME64		StackFrame64;
-	HANDLE				ProcessHandle;
-	unsigned long		LastError;
-	bool				bStackWalkSucceeded	= true;
-	uint32				CurrentDepth		= 0;
-	uint32				MachineType			= IMAGE_FILE_MACHINE_I386;
-	ContextWapper->CheckOk();
-	HANDLE				ThreadHandle = ContextWapper->ThreadHandle;
-	CONTEXT				ContextCopy = ContextWapper->Context;
-
-	*Depth = 0;
-
-#if !PLATFORM_SEH_EXCEPTIONS_DISABLED
-	__try
-#endif
-	{
-		// Get context, process and thread information.
-		ProcessHandle = GProcessHandle;
-
-		// Zero out stack frame.
-		FMemory::Memzero( StackFrame64 );
-
-		// Initialize the STACKFRAME structure.
-		StackFrame64.AddrPC.Mode         = AddrModeFlat;
-		StackFrame64.AddrStack.Mode      = AddrModeFlat;
-		StackFrame64.AddrFrame.Mode      = AddrModeFlat;
-#if PLATFORM_64BITS
-#if PLATFORM_CPU_X86_FAMILY || defined(_M_ARM64EC)
-		StackFrame64.AddrPC.Offset       = ContextWapper->Context.Rip;
-		StackFrame64.AddrStack.Offset    = ContextWapper->Context.Rsp;
-		StackFrame64.AddrFrame.Offset    = ContextWapper->Context.Rbp;
-		MachineType                      = IMAGE_FILE_MACHINE_AMD64;
-#elif PLATFORM_CPU_ARM_FAMILY
-		StackFrame64.AddrPC.Offset       = ContextWapper->Context.Pc;
-		StackFrame64.AddrStack.Offset    = ContextWapper->Context.Sp;
-		StackFrame64.AddrFrame.Offset    = ContextWapper->Context.Fp;
-		MachineType                      = IMAGE_FILE_MACHINE_ARM64;
-#endif
-#else   //PLATFORM_64BITS
-		StackFrame64.AddrPC.Offset       = ContextWapper->Context.Eip;
-		StackFrame64.AddrStack.Offset    = ContextWapper->Context.Esp;
-		StackFrame64.AddrFrame.Offset    = ContextWapper->Context.Ebp;
-#endif	//PLATFORM_64BITS
-
-		// Walk the stack one frame at a time.
-		while( bStackWalkSucceeded && (CurrentDepth < MaxDepth) )
-		{
-			bStackWalkSucceeded = !!StackWalk64(MachineType, 
-												ProcessHandle, 
-												ThreadHandle, 
-												&StackFrame64,
-												&ContextCopy,
-												NULL,
-												SymFunctionTableAccess64,
-												SymGetModuleBase64,
-												NULL );
-
-			if( !bStackWalkSucceeded  )
-			{
-				// StackWalk failed! give up.
-				LastError = GetLastError( );
-				break;
-			}
-
-			// Stop if the frame pointer is NULL.
-			// Note that the thread's PC 'StackFrame64.AddrPC.Offset' COULD be 0 in case something calls a nullptr.
-			if( StackFrame64.AddrFrame.Offset == 0 )
-			{
-				break;
-			}
-
-			BackTrace[CurrentDepth++] = StackFrame64.AddrPC.Offset;
-			*Depth = CurrentDepth;
-		}
-	} 
-#if !PLATFORM_SEH_EXCEPTIONS_DISABLED
-	__except( EXCEPTION_EXECUTE_HANDLER )
-	{
-		// We need to catch any exceptions within this function so they don't get sent to 
-		// the engine's error handler, hence causing an infinite loop.
-		return EXCEPTION_EXECUTE_HANDLER;
-	} 
-#endif
-
-	// NULL out remaining entries.
-	for ( ; CurrentDepth<MaxDepth; CurrentDepth++ )
-	{
-		BackTrace[CurrentDepth] = 0;
-	}
-
-	return EXCEPTION_EXECUTE_HANDLER;
-}
-
-int32 CaptureStackTraceHelper(uint64* BackTrace, uint32 MaxDepth, CONTEXT* Context)
-{
-	FWindowsThreadContextWrapper HelperContext;
-	HelperContext.ThreadHandle = GetCurrentThread();
-	HelperContext.Context = *Context;
-	uint32 Depth = 0;
-
-	return CaptureStackTraceHelper(BackTrace, MaxDepth, &HelperContext, &Depth);
-}
-
 #if USE_FAST_STACKTRACE
 NTSYSAPI uint16 NTAPI RtlCaptureStackBackTrace(
 	__in uint32 FramesToSkip,
@@ -306,7 +190,7 @@ void FWindowsPlatformStackWalk::ThreadStackWalkAndDump(ANSICHAR* HumanReadableSt
 
 uint32 FWindowsPlatformStackWalk::CaptureThreadStackBackTrace(uint64 ThreadId, uint64* BackTrace, uint32 MaxDepth, void* Context)
 {
-	InitStackWalking();
+	// Do not call InitStackWalking as we're just capturing addresses without using the dbghelp library etc.
 
 	if (BackTrace == nullptr || MaxDepth == 0)
 	{
@@ -333,17 +217,17 @@ uint32 FWindowsPlatformStackWalk::CaptureThreadStackBackTrace(uint64 ThreadId, u
 	{
 		FWindowsThreadContextWrapper* ContextWrapper = reinterpret_cast<FWindowsThreadContextWrapper*>(Context);
 		ContextWrapper->ThreadHandle = ThreadHandle; // Use the thread handle open above to be sure it has the correct permissions.
-		CaptureStackTraceHelper(BackTrace, MaxDepth, ContextWrapper, &Depth);
+		FMicrosoftPlatformStackWalk::CaptureStackTraceInternal(BackTrace, MaxDepth, &ContextWrapper->Context, ThreadHandle, &Depth);
 	}
 	else
 	{
-		FWindowsThreadContextWrapper ContextWrapper;
-		ContextWrapper.Context.ContextFlags = CONTEXT_CONTROL;
-		ContextWrapper.ThreadHandle = ThreadHandle;
-		
-		if (::GetThreadContext(ThreadHandle, &ContextWrapper.Context))
+		CONTEXT ThreadContext;
+		FMemory::Memzero(ThreadContext);
+		ThreadContext.ContextFlags = CONTEXT_CONTROL;
+
+		if (::GetThreadContext(ThreadHandle, &ThreadContext))
 		{
-			CaptureStackTraceHelper(BackTrace, MaxDepth, &ContextWrapper, &Depth);
+			FMicrosoftPlatformStackWalk::CaptureStackTraceInternal(BackTrace, MaxDepth, &ThreadContext, ThreadHandle, &Depth);
 		}
 	}
 
@@ -374,8 +258,8 @@ uint32 FWindowsPlatformStackWalk::CaptureStackBackTrace( uint64* BackTrace, uint
 	uint32 Depth = 0;
 	if (Context)
 	{
-		InitStackWalking();
-		CaptureStackTraceHelper(BackTrace, MaxDepth, (FWindowsThreadContextWrapper*)Context, &Depth);
+		FWindowsThreadContextWrapper* Wrapper = reinterpret_cast<FWindowsThreadContextWrapper*>(Context);
+		FMicrosoftPlatformStackWalk::CaptureStackTraceInternal(BackTrace, MaxDepth, &Wrapper->Context, Wrapper->ThreadHandle, &Depth);
 	}
 	else
 	{
@@ -409,50 +293,25 @@ uint32 FWindowsPlatformStackWalk::CaptureStackBackTrace( uint64* BackTrace, uint
 		}		
 #elif USE_SLOW_STACKTRACE
 		// NOTE: Make sure to enable Stack Frame pointers: bOmitFramePointers = false, or /Oy-
-		// If GStackWalkingInitialized is true, traces will work anyway but will be much slower.
-		if (!GStackWalkingInitialized)
-		{
-			InitStackWalking();
-		}
-		
 		CONTEXT HelperContext;
 		RtlCaptureContext(&HelperContext);
 
-		// Capture the back trace.
-		CaptureStackTraceHelper(BackTrace, MaxDepth, &HelperContext, &Depth);		
-#elif PLATFORM_64BITS
+		FMicrosoftPlatformStackWalk::CaptureStackTraceInternal(BackTrace, MaxDepth, &HelperContext, GetCurrentThread(), &Depth);
+#else
+		static_assert(PLATFORM_64BITS, "Non-64 bit windows not supported");
+
 		// Raise an exception so CaptureStackBackTraceHelper has access to context record.
 		__try
 		{
-			RaiseException(0,			// Application-defined exception code.
-							0,			// Zero indicates continuable exception.
-							0,			// Number of arguments in args array (ignored if args is NULL)
-				NULL);		// Array of arguments
-			}
-		// Capture the back trace.
-		__except (CaptureStackTraceHelper(BackTrace, MaxDepth, (GetExceptionInformation())->ContextRecord, &Depth))
-		{
+			RaiseException(0, // Application-defined exception code.
+				0,            // Zero indicates continuable exception.
+				0,            // Number of arguments in args array (ignored if args is NULL)
+				NULL);        // Array of arguments
 		}
-#else
-		// Use a bit of inline assembly to capture the information relevant to stack walking which is
-		// basically EIP and EBP.
-		CONTEXT HelperContext;
-		memset(&HelperContext, 0, sizeof(CONTEXT));
-		HelperContext.ContextFlags = CONTEXT_FULL;
-
-		// Use a fake function call to pop the return address and retrieve EIP.
-		__asm
+		// Capture the back trace.
+		__except (FMicrosoftPlatformStackWalk::CaptureStackTraceInternal(BackTrace, MaxDepth, (GetExceptionInformation())->ContextRecord, GetCurrentThread(), &Depth))
 		{
-			call FakeFunctionCall
-			FakeFunctionCall :
-			pop eax
-			mov HelperContext.Eip, eax
-			mov HelperContext.Ebp, ebp
-			mov HelperContext.Esp, esp
 		}
-
-		// Capture the back trace.
-		CaptureStackTraceHelper(BackTrace, MaxDepth, &HelperContext, &Depth);
 #endif
 	}	
 	return Depth;
diff --git a/Engine/Source/Runtime/Core/Public/Microsoft/MicrosoftPlatformStackWalk.h b/Engine/Source/Runtime/Core/Public/Microsoft/MicrosoftPlatformStackWalk.h
--- a/Engine/Source/Runtime/Core/Public/Microsoft/MicrosoftPlatformStackWalk.h
+++ b/Engine/Source/Runtime/Core/Public/Microsoft/MicrosoftPlatformStackWalk.h
@@ -16,4 +16,20 @@ protected:
 	// Extract debug info for a module from the module header in memory. 
 	// Can directly read the information even when the current target can't load the symbols itself or use certain DbgHelp APIs.
 	static bool ExtractInfoFromModule(void* ProcessHandle, void* ModuleHandle, FStackWalkModuleInfo& OutInfo);
+
+	//
+	// If the thread is not the calling thread, it should already have been suspended by the calling code.
+
+	/**
+	 * Helper for Microsoft platforms to capture a stacktrace for a specific thread.
+	 * If the thread is not the calling thread, it should already have been suspended by the calling code.
+	 *
+	 * @param	OutBackTrace		Array to write backtrace to
+	 * @param	MaxDepth			Maximum depth to walk - needs to be less than or equal to array size
+	 * @param	Context				Pointer to a CONTEXT structure containing thread context info (registers, instruction ptr, etc)
+	 * @param	ThreadHandle		HANDLE for the thread whose stack should be walked.
+	 * @param	OutDepth			Number of stack frames returned.
+	 * @return	EXCEPTION_EXECUTE_HANDLER
+	 */
+	static int32 CaptureStackTraceInternal(uint64* OutBacktrace, uint32 MaxDepth, void* Context, void* ThreadHandle, uint32* OutDepth);
 };
\ No newline at end of file
