Logging: Fixed incorrect inclusion of field names in formatted log messages

#rnx

diff --git a/Engine/Source/Runtime/Core/Private/Logging/StructuredLog.cpp b/Engine/Source/Runtime/Core/Private/Logging/StructuredLog.cpp
--- a/Engine/Source/Runtime/Core/Private/Logging/StructuredLog.cpp
+++ b/Engine/Source/Runtime/Core/Private/Logging/StructuredLog.cpp
@@ -107,7 +107,7 @@ static void LogFieldValue(TStringBuilderBase<CharType>& Out, const FCbFieldView&
 	case ECbFieldType::Array:
 	case ECbFieldType::UniformArray:
 	case ECbFieldType::Binary:
-		CompactBinaryToCompactJson(Field, Out);
+		CompactBinaryToCompactJson(Field.RemoveName(), Out);
 		break;
 	case ECbFieldType::String:
 		Out.Append(Accessor.AsString());
@@ -120,7 +120,7 @@ static void LogFieldValue(TStringBuilderBase<CharType>& Out, const FCbFieldView&
 		break;
 	case ECbFieldType::Float32:
 	case ECbFieldType::Float64:
-		CompactBinaryToCompactJson(Field, Out);
+		CompactBinaryToCompactJson(Field.RemoveName(), Out);
 		break;
 	case ECbFieldType::BoolFalse:
 		Out.Append(FConstants::False);
@@ -159,7 +159,7 @@ static void LogFieldValue(TStringBuilderBase<CharType>& Out, const FCbFieldView&
 		break;
 	case ECbFieldType::CustomById:
 	case ECbFieldType::CustomByName:
-		CompactBinaryToCompactJson(Field, Out);
+		CompactBinaryToCompactJson(Field.RemoveName(), Out);
 		break;
 	default:
 		checkNoEntry();
