Logging: Fixed structured logging to work with Live Coding

diff --git a/Engine/Source/Runtime/Core/Private/Logging/StructuredLog.cpp b/Engine/Source/Runtime/Core/Private/Logging/StructuredLog.cpp
--- a/Engine/Source/Runtime/Core/Private/Logging/StructuredLog.cpp
+++ b/Engine/Source/Runtime/Core/Private/Logging/StructuredLog.cpp
@@ -8,6 +8,7 @@
 #include "Containers/StringConv.h"
 #include "Containers/StringView.h"
 #include "CoreGlobals.h"
+#include "HAL/PlatformMath.h"
 #include "HAL/PlatformTime.h"
 #include "Internationalization/Internationalization.h"
 #include "Logging/LogTrace.h"
@@ -228,17 +229,32 @@ public:
 
 	static void Destroy(FLogTemplate* Template);
 
+	const TCHAR* GetFormat() const { return StaticFormat; }
+
+	FLogTemplate* GetNext() { return Next; }
+	void SetNext(FLogTemplate* Template) { Next = Template; }
+
+	uint8* GetOpData() { return (uint8*)(this + 1); }
+	const uint8* GetOpData() const { return (const uint8*)(this + 1); }
+
 	template <typename CharType>
-	void FormatTo(TStringBuilderBase<CharType>& Out, const TCHAR* Format, const FCbObjectView& Fields) const;
+	void FormatTo(TStringBuilderBase<CharType>& Out, const FCbObjectView& Fields) const;
 
 private:
 	template <typename CharType>
 	void FormatLocalizedTo(TStringBuilderBase<CharType>& Out, const FCbObjectView& Fields) const;
 
-	FLogTemplate() = default;
+	inline constexpr explicit FLogTemplate(const TCHAR* Format)
+		: StaticFormat(Format)
+	{
+	}
+
 	~FLogTemplate() = default;
 	FLogTemplate(const FLogTemplate&) = delete;
 	FLogTemplate& operator=(const FLogTemplate&) = delete;
+
+	const TCHAR* StaticFormat = nullptr;
+	FLogTemplate* Next = nullptr;
 };
 
 FLogTemplate* FLogTemplate::Create(const TCHAR* Format, const FLogField* Fields, const int32 FieldCount)
@@ -341,9 +357,9 @@ FLogTemplate* FLogTemplate::Create(const TCHAR* Format, const FLogField* Fields,
 	checkf(!bFindFields || !bPositional || FormatFieldCount == FieldCount,
 		TEXT("Log format requires %d fields and %d were provided. [[%s]]"), FormatFieldCount, FieldCount, Format);
 
-	const uint32 TotalSize = Algo::TransformAccumulate(Ops, FLogTemplateOp::SaveSize, 0);
-	FLogTemplate* const Template = (FLogTemplate*)FMemory::Malloc(TotalSize);
-	uint8* Data = (uint8*)Template;
+	const uint32 TotalSize = sizeof(FLogTemplate) + Algo::TransformAccumulate(Ops, FLogTemplateOp::SaveSize, 0);
+	FLogTemplate* const Template = new(FMemory::Malloc(TotalSize, alignof(FLogTemplate))) FLogTemplate(Format);
+	uint8* Data = Template->GetOpData();
 	for (const FLogTemplateOp& Op : Ops)
 	{
 		FLogTemplateOp::Save(Op, Data);
@@ -386,10 +402,10 @@ FLogTemplate* FLogTemplate::CreateLocalized(const TCHAR* TextNamespace, const TC
 	}
 
 	const FLogTemplateOp Ops[]{{FLogTemplateOp::OpLocText}, {FLogTemplateOp::OpEnd}};
-	const uint32 TotalSize = sizeof(FTextFormat) + Algo::TransformAccumulate(Ops, FLogTemplateOp::SaveSize, 0);
-	FTextFormat* NewTextFormat = new(FMemory::Malloc(TotalSize, alignof(FTextFormat))) FTextFormat(MoveTemp(TextFormat));
-	FLogTemplate* const Template = (FLogTemplate*)(NewTextFormat + 1);
-	uint8* Data = (uint8*)Template;
+	const uint32 TotalSize = sizeof(FTextFormat) + sizeof(FLogTemplate) + Algo::TransformAccumulate(Ops, FLogTemplateOp::SaveSize, 0);
+	FTextFormat* NewTextFormat = new(FMemory::Malloc(TotalSize, FPlatformMath::Max(alignof(FTextFormat), alignof(FLogTemplate)))) FTextFormat(MoveTemp(TextFormat));
+	FLogTemplate* const Template = new(NewTextFormat + 1) FLogTemplate(Format);
+	uint8* Data = Template->GetOpData();
 	for (const FLogTemplateOp& Op : Ops)
 	{
 		FLogTemplateOp::Save(Op, Data);
@@ -401,7 +417,7 @@ void FLogTemplate::Destroy(FLogTemplate* Template)
 {
 	using namespace Logging::Private;
 
-	const uint8* NextOp = (const uint8*)Template;
+	const uint8* NextOp = Template->GetOpData();
 	if (FLogTemplateOp::Load(NextOp).Code == FLogTemplateOp::OpLocText)
 	{
 		FTextFormat* TextFormat = (FTextFormat*)Template - 1;
@@ -415,11 +431,11 @@ void FLogTemplate::Destroy(FLogTemplate* Template)
 }
 
 template <typename CharType>
-void FLogTemplate::FormatTo(TStringBuilderBase<CharType>& Out, const TCHAR* Format, const FCbObjectView& Fields) const
+void FLogTemplate::FormatTo(TStringBuilderBase<CharType>& Out, const FCbObjectView& Fields) const
 {
 	using namespace Logging::Private;
 
-	auto FindField = [&Fields, It = Fields.CreateViewIterator(), Index = 0, Format](FAnsiStringView Name, int32 IndexHint = -1) mutable -> FCbFieldView&
+	auto FindField = [&Fields, It = Fields.CreateViewIterator(), Index = 0, Format = StaticFormat](FAnsiStringView Name, int32 IndexHint = -1) mutable -> FCbFieldView&
 	{
 		if (IndexHint >= 0)
 		{
@@ -457,8 +473,8 @@ void FLogTemplate::FormatTo(TStringBuilderBase<CharType>& Out, const TCHAR* Form
 	};
 
 	int32 FieldIndexHint = -1;
-	const uint8* NextOp = (const uint8*)this;
-	const TCHAR* NextFormat = Format;
+	const uint8* NextOp = GetOpData();
+	const TCHAR* NextFormat = StaticFormat;
 	for (;;)
 	{
 		const FLogTemplateOp Op = FLogTemplateOp::Load(NextOp);
@@ -517,7 +533,7 @@ FDateTime FLogTime::GetUtcTime() const
 ///////////////////////////////////////////////////////////////////////////////////////////////////////////////////////
 
 template <typename CharType>
-static void FormatRecordMessageTo(TStringBuilderBase<CharType>& Out, const FLogRecord& Record)
+FORCENOINLINE static void FormatDynamicRecordMessageTo(TStringBuilderBase<CharType>& Out, const FLogRecord& Record)
 {
 	const TCHAR* Format = Record.GetFormat();
 	if (UNLIKELY(!Format))
@@ -525,22 +541,27 @@ static void FormatRecordMessageTo(TStringBuilderBase<CharType>& Out, const FLogR
 		return;
 	}
 
-	const FLogTemplate* Template = Record.GetTemplate();
-	if (LIKELY(Template))
-	{
-		return Template->FormatTo(Out, Format, Record.GetFields());
-	}
-
 	const TCHAR* TextNamespace = Record.GetTextNamespace();
 	const TCHAR* TextKey = Record.GetTextKey();
 	checkf(!TextNamespace == !TextKey,
 		TEXT("Log record must have both or neither of the text namespace and text key. [[%s]]"), Format);
 
 	FLogTemplate* LocalTemplate = TextKey ? FLogTemplate::CreateLocalized(TextNamespace, TextKey, Format) : FLogTemplate::Create(Format);
-	LocalTemplate->FormatTo(Out, Format, Record.GetFields());
+	LocalTemplate->FormatTo(Out, Record.GetFields());
 	FLogTemplate::Destroy(LocalTemplate);
 }
 
+template <typename CharType>
+static void FormatRecordMessageTo(TStringBuilderBase<CharType>& Out, const FLogRecord& Record)
+{
+	const FLogTemplate* Template = Record.GetTemplate();
+	if (LIKELY(Template))
+	{
+		return Template->FormatTo(Out, Record.GetFields());
+	}
+	FormatDynamicRecordMessageTo(Out, Record);
+}
+
 void FLogRecord::FormatMessageTo(FUtf8StringBuilderBase& Out) const
 {
 	FormatRecordMessageTo(Out, *this);
@@ -563,9 +584,9 @@ namespace UE::Logging::Private
 class FLogTemplateFieldIterator
 {
 public:
-	inline FLogTemplateFieldIterator(const FLogTemplate* Template, const TCHAR* Format)
-		: NextOp((const uint8*)Template)
-		, NextFormat(Format)
+	inline explicit FLogTemplateFieldIterator(const FLogTemplate& Template)
+		: NextOp(Template.GetOpData())
+		, NextFormat(Template.GetFormat())
 	{
 		++*this;
 	}
@@ -607,23 +628,22 @@ FLogTemplateFieldIterator& FLogTemplateFieldIterator::operator++()
 
 ///////////////////////////////////////////////////////////////////////////////////////////////////////////////////////
 
-struct FStaticLogDynamicDataManager
+struct FStaticLogTemplateManager
 {
-	std::atomic<FStaticLogDynamicData*> Head = nullptr;
+	std::atomic<FLogTemplate*> Head = nullptr;
 
-	~FStaticLogDynamicDataManager()
+	~FStaticLogTemplateManager()
 	{
-		for (FStaticLogDynamicData* Data = Head.exchange(nullptr); Data; Data = Data->Next)
+		for (FLogTemplate* Template = Head.exchange(nullptr); Template;)
 		{
-			if (FLogTemplate* Template = Data->Template.exchange(nullptr))
-			{
-				FLogTemplate::Destroy(Template);
-			}
+			FLogTemplate* NextTemplate = Template->GetNext();
+			FLogTemplate::Destroy(Template);
+			Template = NextTemplate;
 		}
 	}
 };
 
-static FStaticLogDynamicDataManager GStaticLogDynamicDataManager;
+static FStaticLogTemplateManager GStaticLogTemplateManager;
 
 // Tracing the log happens in its own function because that allows stack space for the message to
 // be returned before calling into the output devices.
@@ -639,7 +659,6 @@ FORCENOINLINE static void LogToTrace(const void* LogPoint, const FLogRecord& Rec
 // Serializing log fields to compact binary happens in its own function because that allows stack
 // space for the writer to be returned before calling into the output devices.
 FORCENOINLINE static FCbObject SerializeLogFields(
-	const FStaticLogRecord& Log,
 	const FLogTemplate& Template,
 	const FLogField* Fields,
 	const int32 FieldCount)
@@ -655,7 +674,7 @@ FORCENOINLINE static FCbObject SerializeLogFields(
 	// Anonymous. Extract names from Template.
 	if (!Fields->Name)
 	{
-		FLogTemplateFieldIterator It(&Template, Log.Format);
+		FLogTemplateFieldIterator It(Template);
 		for (const FLogField* FieldsEnd = Fields + FieldCount; Fields != FieldsEnd; ++Fields, ++It)
 		{
 			check(It);
@@ -677,29 +696,6 @@ FORCENOINLINE static FCbObject SerializeLogFields(
 	return Writer.Save().AsObject();
 }
 
-// Serializing log fields to compact binary happens in its own function because that allows stack
-// space for the writer to be returned before calling into the output devices.
-FORCENOINLINE static FCbObject SerializeLogFields(
-	const FStaticLocalizedLogRecord& Log,
-	const FLogTemplate& Template,
-	const FLogField* Fields,
-	const int32 FieldCount)
-{
-	if (FieldCount == 0)
-	{
-		return FCbObject();
-	}
-
-	TCbWriter<1024> Writer;
-	Writer.BeginObject();
-	for (const FLogField* FieldsEnd = Fields + FieldCount; Fields != FieldsEnd; ++Fields)
-	{
-		Fields->WriteValue(Writer.SetName(Fields->Name), Fields->Value);
-	}
-	Writer.EndObject();
-	return Writer.Save().AsObject();
-}
-
 template <typename StaticLogRecordType>
 FORCENOINLINE static FLogTemplate& CreateLogTemplate(const FLogCategoryBase& Category, const StaticLogRecordType& Log, const FLogField* Fields, const int32 FieldCount)
 {
@@ -710,31 +706,36 @@ FORCENOINLINE static FLogTemplate& CreateLogTemplate(const FLogCategoryBase& Cat
 	}
 #endif
 
-	FLogTemplate* NewTemplate = FLogTemplate::CreateStatic(Log, Fields, FieldCount);
-	if (FLogTemplate* ExistingTemplate = nullptr;
-		UNLIKELY(!Log.DynamicData.Template.compare_exchange_strong(ExistingTemplate, NewTemplate, std::memory_order_release, std::memory_order_acquire)))
+	for (FLogTemplate* Template = Log.DynamicData.Template.load(std::memory_order_acquire);;)
 	{
-		FLogTemplate::Destroy(NewTemplate);
-		return *ExistingTemplate;
-	}
+		if (Template && Template->GetFormat() == Log.Format)
+		{
+			return *Template;
+		}
 
-	// Register the template to destroy on exit.
-	for (FStaticLogDynamicData* Head = GStaticLogDynamicDataManager.Head.load(std::memory_order_relaxed);;)
-	{
-		Log.DynamicData.Next = Head;
-		if (GStaticLogDynamicDataManager.Head.compare_exchange_weak(Head, &Log.DynamicData, std::memory_order_release, std::memory_order_relaxed))
+		FLogTemplate* NewTemplate = FLogTemplate::CreateStatic(Log, Fields, FieldCount);
+		if (LIKELY(Log.DynamicData.Template.compare_exchange_strong(Template, NewTemplate, std::memory_order_release, std::memory_order_acquire)))
 		{
-			break;
+			// Register the template to destroy on exit.
+			for (FLogTemplate* Head = GStaticLogTemplateManager.Head.load(std::memory_order_relaxed);;)
+			{
+				NewTemplate->SetNext(Head);
+				if (LIKELY(GStaticLogTemplateManager.Head.compare_exchange_weak(Head, NewTemplate, std::memory_order_release, std::memory_order_relaxed)))
+				{
+					break;
+				}
+			}
+			return *NewTemplate;
 		}
+		FLogTemplate::Destroy(NewTemplate);
 	}
-
-	return *NewTemplate;
 }
 
 template <typename StaticLogRecordType>
 inline static FLogTemplate& EnsureLogTemplate(const FLogCategoryBase& Category, const StaticLogRecordType& Log, const FLogField* Fields, const int32 FieldCount)
 {
-	if (FLogTemplate* Template = Log.DynamicData.Template.load(std::memory_order_acquire); LIKELY(Template))
+	// Format can change on a static log record due to Live Coding.
+	if (FLogTemplate* Template = Log.DynamicData.Template.load(std::memory_order_acquire); LIKELY(Template && Template->GetFormat() == Log.Format))
 	{
 		return *Template;
 	}
@@ -749,7 +750,7 @@ inline static FLogRecord CreateLogRecord(const FLogCategoryBase& Category, const
 	FLogRecord Record;
 	Record.SetFormat(Log.Format);
 	Record.SetTemplate(&Template);
-	Record.SetFields(SerializeLogFields(Log, Template, Fields, FieldCount));
+	Record.SetFields(SerializeLogFields(Template, Fields, FieldCount));
 	Record.SetFile(Log.File);
 	Record.SetLine(Log.Line);
 	Record.SetCategory(Category.GetCategoryName());
diff --git a/Engine/Source/Runtime/Core/Public/Logging/StructuredLog.h b/Engine/Source/Runtime/Core/Public/Logging/StructuredLog.h
--- a/Engine/Source/Runtime/Core/Public/Logging/StructuredLog.h
+++ b/Engine/Source/Runtime/Core/Public/Logging/StructuredLog.h
@@ -220,7 +220,6 @@ namespace UE::Logging::Private
 struct FStaticLogDynamicData
 {
 	std::atomic<FLogTemplate*> Template;
-	FStaticLogDynamicData* Next = nullptr;
 };
 
 /** Data about a static log that is constant for every occurrence. */
