CsvProfiler - add logging when categories are enabled/disabled via the commandline and
- call EnableCategoryByName when processing the commandline rather than duplicating that logic

diff --git a/Engine/Source/Runtime/Core/Private/ProfilingDebugging/CsvProfiler.cpp b/Engine/Source/Runtime/Core/Private/ProfilingDebugging/CsvProfiler.cpp
--- a/Engine/Source/Runtime/Core/Private/ProfilingDebugging/CsvProfiler.cpp
+++ b/Engine/Source/Runtime/Core/Private/ProfilingDebugging/CsvProfiler.cpp
@@ -3659,11 +3659,7 @@ void FCsvProfiler::Init()
 		CsvCategoriesStr.ParseIntoArray(CsvCategories, TEXT(","), true);
 		for (int i = 0; i < CsvCategories.Num(); i++)
 		{
-			int32 Index = FCsvCategoryData::Get()->GetCategoryIndex(CsvCategories[i]);
-			if (Index > 0)
-			{
-				GCsvCategoriesEnabled[Index] = true;
-			}
+			EnableCategoryByString(CsvCategories[i]);
 		}
 	}
 
@@ -3818,9 +3814,11 @@ bool FCsvProfiler::EnableCategoryByString(const FString& CategoryName) const
 	int32 Category = FCsvCategoryData::Get()->GetCategoryIndex(CategoryName);
 	if (Category >= 0)
 	{
+		UE_LOG(LogCsvProfiler, Log, TEXT("Enabled category %s"), *CategoryName);
 		GCsvCategoriesEnabled[Category] = true;
 		return true;
 	}
+	UE_LOG(LogCsvProfiler, Warning, TEXT("Error: Can't find category %s"), *CategoryName);
 	return false;
 }
 
