Don't try to compile FMicrosoftPlatformStackWalk::CaptureStackTraceInternal on arm64 microsoft platforms.

diff --git a/Engine/Source/Runtime/Core/Private/Microsoft/MicrosoftPlatformStackWalk.cpp b/Engine/Source/Runtime/Core/Private/Microsoft/MicrosoftPlatformStackWalk.cpp
--- a/Engine/Source/Runtime/Core/Private/Microsoft/MicrosoftPlatformStackWalk.cpp
+++ b/Engine/Source/Runtime/Core/Private/Microsoft/MicrosoftPlatformStackWalk.cpp
@@ -78,6 +78,7 @@ bool FMicrosoftPlatformStackWalk::ExtractInfoFromModule(void* ProcessHandle, voi
 
 int32 FMicrosoftPlatformStackWalk::CaptureStackTraceInternal(uint64* OutBacktrace, uint32 MaxDepth, void* InContext, void* InThreadHandle, uint32* OutDepth)
 {
+#if PLATFORM_CPU_X86_FAMILY || defined(_M_ARM64EC)
 	uint32 CurrentDepth = 0;
 	HANDLE ThreadHandle = reinterpret_cast<HANDLE>(InThreadHandle);
 	CONTEXT ContextCopy = *reinterpret_cast<PCONTEXT>(InContext);
@@ -146,5 +147,8 @@ int32 FMicrosoftPlatformStackWalk::CaptureStackTraceInternal(uint64* OutBacktrac
 		OutBacktrace[CurrentDepth] = 0;
 	}
 
+#else
+	UE_LOG(LogCore, Fatal, TEXT("Stack trace via RtlLookupFunctionEntry only implemented for x86-64 Microsoft platforms."));
+#endif
 	return EXCEPTION_EXECUTE_HANDLER;
 }
\ No newline at end of file
