Logging: Check NO_LOGGING consistently for structured logging

#rnx

diff --git a/Engine/Source/Runtime/Core/Private/Logging/StructuredLog.cpp b/Engine/Source/Runtime/Core/Private/Logging/StructuredLog.cpp
--- a/Engine/Source/Runtime/Core/Private/Logging/StructuredLog.cpp
+++ b/Engine/Source/Runtime/Core/Private/Logging/StructuredLog.cpp
@@ -793,22 +793,28 @@ void LogWithFieldArray(const FLogCategoryBase& Category, const FStaticLogRecord&
 
 void LogWithNoFields(const FLogCategoryBase& Category, const FStaticLogRecord& Log)
 {
+#if !NO_LOGGING
 	// A non-null field pointer enables field validation in FLogTemplate::Create.
 	static constexpr FLogField EmptyField{};
 	LogWithFieldArray(Category, Log, &EmptyField, 0);
+#endif
 }
 
 void LogWithFieldArray(const FLogCategoryBase& Category, const FStaticLocalizedLogRecord& Log, const FLogField* Fields, const int32 FieldCount)
 {
+#if !NO_LOGGING
 	FLogRecord Record = CreateLogRecord(Category, Log, Fields, FieldCount);
 	Record.SetTextNamespace(Log.TextNamespace);
 	Record.SetTextKey(Log.TextKey);
 	DispatchLogRecord(Log, Record);
+#endif
 }
 
 void LogWithNoFields(const FLogCategoryBase& Category, const FStaticLocalizedLogRecord& Log)
 {
+#if !NO_LOGGING
 	LogWithFieldArray(Category, Log, nullptr, 0);
+#endif
 }
 
 void FatalLogWithFieldArray(const FLogCategoryBase& Category, const FStaticLogRecord& Log, const FLogField* Fields, const int32 FieldCount)
@@ -828,9 +834,11 @@ void FatalLogWithFieldArray(const FLogCategoryBase& Category, const FStaticLogRe
 
 void FatalLogWithNoFields(const FLogCategoryBase& Category, const FStaticLogRecord& Log)
 {
+#if !NO_LOGGING
 	// A non-null field pointer enables field validation in FLogTemplate::Create.
 	static constexpr FLogField EmptyField{};
 	FatalLogWithFieldArray(Category, Log, &EmptyField, 0);
+#endif
 }
 
 void FatalLogWithFieldArray(const FLogCategoryBase& Category, const FStaticLocalizedLogRecord& Log, const FLogField* Fields, const int32 FieldCount)
@@ -850,9 +858,11 @@ void FatalLogWithFieldArray(const FLogCategoryBase& Category, const FStaticLocal
 
 void FatalLogWithNoFields(const FLogCategoryBase& Category, const FStaticLocalizedLogRecord& Log)
 {
+#if !NO_LOGGING
 	// A non-null field pointer enables field validation in FLogTemplate::Create.
 	static constexpr FLogField EmptyField{};
 	FatalLogWithFieldArray(Category, Log, &EmptyField, 0);
+#endif
 }
 
 } // UE::Logging::Private
