Logging: Fixed a new compile warning in structured logging

#rnx

diff --git a/Engine/Source/Runtime/Core/Private/Logging/StructuredLog.cpp b/Engine/Source/Runtime/Core/Private/Logging/StructuredLog.cpp
--- a/Engine/Source/Runtime/Core/Private/Logging/StructuredLog.cpp
+++ b/Engine/Source/Runtime/Core/Private/Logging/StructuredLog.cpp
@@ -403,7 +403,7 @@ FLogTemplate* FLogTemplate::CreateLocalized(const TCHAR* TextNamespace, const TC
 
 	const FLogTemplateOp Ops[]{{FLogTemplateOp::OpLocText}, {FLogTemplateOp::OpEnd}};
 	const uint32 TotalSize = sizeof(FTextFormat) + sizeof(FLogTemplate) + Algo::TransformAccumulate(Ops, FLogTemplateOp::SaveSize, 0);
-	FTextFormat* NewTextFormat = new(FMemory::Malloc(TotalSize, FPlatformMath::Max(alignof(FTextFormat), alignof(FLogTemplate)))) FTextFormat(MoveTemp(TextFormat));
+	FTextFormat* NewTextFormat = new(FMemory::Malloc(TotalSize, uint32(FPlatformMath::Max(alignof(FTextFormat), alignof(FLogTemplate))))) FTextFormat(MoveTemp(TextFormat));
 	FLogTemplate* const Template = new(NewTextFormat + 1) FLogTemplate(Format);
 	uint8* Data = Template->GetOpData();
 	for (const FLogTemplateOp& Op : Ops)
