Changed mistaken filter in FKeysCultureTest

diff --git a/Engine/Source/Runtime/Core/Tests/Internationalization/KeysCultureTest.cpp b/Engine/Source/Runtime/Core/Tests/Internationalization/KeysCultureTest.cpp
--- a/Engine/Source/Runtime/Core/Tests/Internationalization/KeysCultureTest.cpp
+++ b/Engine/Source/Runtime/Core/Tests/Internationalization/KeysCultureTest.cpp
@@ -45,7 +45,7 @@ FString CreateExpectedOutputFromPolyglotData(const FPolyglotTextData& InData)
 	return FString::Printf(TEXT("%s, %s"), *Key, *Namespace);
 }
 
-TEST_CASE_NAMED(FKeysCultureTest, "System::Core::Misc::KeysCulture", "[.][EditorContext][ClientContext][SmokeFilter]")
+TEST_CASE_NAMED(FKeysCultureTest, "System::Core::Misc::KeysCulture", "[.][EditorContext][ClientContext][EngineFilter]")
 {
 	FRestoreCultureStateGuard RestoreCultureStateGuard;
 	
