Do not check whether the chunk install module is known before trying to load it.

If it hasn't been previously loaded, it will not be known. Old code only works if the chunk installer comes from a uplugin.

#rnx

diff --git a/Engine/Source/Runtime/Core/Private/Windows/WindowsPlatformMisc.cpp b/Engine/Source/Runtime/Core/Private/Windows/WindowsPlatformMisc.cpp
--- a/Engine/Source/Runtime/Core/Private/Windows/WindowsPlatformMisc.cpp
+++ b/Engine/Source/Runtime/Core/Private/Windows/WindowsPlatformMisc.cpp
@@ -3636,9 +3636,9 @@ IPlatformChunkInstall* FWindowsPlatformMisc::GetPlatformChunkInstall()
 		{
 			FString InstallModule;
 			GConfig->GetString(TEXT("StreamingInstall"), TEXT("DefaultProviderName"), InstallModule, GEngineIni);
-			FModuleStatus Status;
-			if (FModuleManager::Get().QueryModule(*InstallModule, Status))
-			{
+
+			if (!InstallModule.IsEmpty())
+			{			
 				PlatformChunkInstallModule = FModuleManager::LoadModulePtr<IPlatformChunkInstallModule>(*InstallModule);
 				if (PlatformChunkInstallModule != nullptr)
 				{
