Fixed THIRD_PARTY_INCLUDES_END for Clang on Windows.

diff --git a/Engine/Source/Runtime/Core/Public/Windows/WindowsPlatformCompilerPreSetup.h b/Engine/Source/Runtime/Core/Public/Windows/WindowsPlatformCompilerPreSetup.h
--- a/Engine/Source/Runtime/Core/Public/Windows/WindowsPlatformCompilerPreSetup.h
+++ b/Engine/Source/Runtime/Core/Public/Windows/WindowsPlatformCompilerPreSetup.h
@@ -44,8 +44,8 @@
 	#ifndef THIRD_PARTY_INCLUDES_END
 		#define THIRD_PARTY_INCLUDES_END \
 	        PRAGMA_ENABLE_UNINITIALIZED_CONST_REFERENCE_WARNINGS \
-			PRAGMA_DISABLE_MISSING_BRACES_WARNINGS \
-			PRAGMA_DISABLE_OVERLOADED_VIRTUAL_WARNINGS \
+			PRAGMA_ENABLE_MISSING_BRACES_WARNINGS \
+			PRAGMA_ENABLE_OVERLOADED_VIRTUAL_WARNINGS \
 			PRAGMA_ENABLE_DEPRECATION_WARNINGS \
 			PRAGMA_ENABLE_UNDEFINED_IDENTIFIER_WARNINGS \
 			PRAGMA_ENABLE_SHADOW_VARIABLE_WARNINGS \
