Reconcile Microsoft platform CompilerPreSetup.h files

diff --git a/Engine/Platforms/Hololens/Source/Runtime/Core/Public/HoloLensPlatformCompilerPreSetup.h b/Engine/Platforms/Hololens/Source/Runtime/Core/Public/HoloLensPlatformCompilerPreSetup.h
--- a/Engine/Platforms/Hololens/Source/Runtime/Core/Public/HoloLensPlatformCompilerPreSetup.h
+++ b/Engine/Platforms/Hololens/Source/Runtime/Core/Public/HoloLensPlatformCompilerPreSetup.h
@@ -2,184 +2,6 @@
 
 #pragma once
 
-#ifndef DISABLE_DEPRECATION
-	#define PRAGMA_DISABLE_DEPRECATION_WARNINGS \
-		__pragma (warning(push)) \
-		__pragma (warning(disable: 4995)) /* 'function': name was marked as #pragma deprecated */ \
-		__pragma (warning(disable: 4996)) /* The compiler encountered a deprecated declaration. */
+// HEADER_UNIT_SKIP - Cause cirular dependencies for header units
 
-	#define PRAGMA_ENABLE_DEPRECATION_WARNINGS \
-		__pragma (warning(pop))
-#endif // DISABLE_DEPRECATION
-
-#ifndef PRAGMA_DISABLE_SHADOW_VARIABLE_WARNINGS
-	#define PRAGMA_DISABLE_SHADOW_VARIABLE_WARNINGS \
-		__pragma (warning(push)) \
-		__pragma (warning(disable: 4456)) /* declaration of 'LocalVariable' hides previous local declaration */ \
-		__pragma (warning(disable: 4457)) /* declaration of 'LocalVariable' hides function parameter */ \
-		__pragma (warning(disable: 4458)) /* declaration of 'LocalVariable' hides class member */ \
-		__pragma (warning(disable: 4459)) /* declaration of 'LocalVariable' hides global declaration */ \
-		__pragma (warning(disable: 6244)) /* local declaration of <variable> hides previous declaration at <line> of <file> */
-#endif // PRAGMA_DISABLE_SHADOW_VARIABLE_WARNINGS
-
-#ifndef PRAGMA_ENABLE_SHADOW_VARIABLE_WARNINGS
-	#define PRAGMA_ENABLE_SHADOW_VARIABLE_WARNINGS \
-		__pragma(warning(pop))
-#endif // PRAGMA_ENABLE_SHADOW_VARIABLE_WARNINGS
-
-#ifndef PRAGMA_DISABLE_UNSAFE_TYPECAST_WARNINGS
-	#define PRAGMA_DISABLE_UNSAFE_TYPECAST_WARNINGS \
-		__pragma (warning(push)) \
-		__pragma (warning(disable: 4244)) /* 'argument': conversion from 'type1' to 'type2', possible loss of data */	\
-		__pragma (warning(disable: 4838)) /* 'argument': conversion from 'type1' to 'type2' requires a narrowing conversion */
-#endif // PRAGMA_DISABLE_UNSAFE_TYPECAST_WARNINGS
-
-#ifndef PRAGMA_ENABLE_UNSAFE_TYPECAST_WARNINGS
-	#define PRAGMA_ENABLE_UNSAFE_TYPECAST_WARNINGS \
-		UE_DEPRECATED_MACRO(5.0, "The PRAGMA_ENABLE_UNSAFE_TYPECAST_WARNINGS macro has been deprecated in favor of PRAGMA_RESTORE_UNSAFE_TYPECAST_WARNINGS. To force enable warnings use PRAGMA_FORCE_UNSAFE_TYPECAST_WARNINGS.")
-#endif // PRAGMA_ENABLE_UNSAFE_TYPECAST_WARNINGS
-
-#ifndef PRAGMA_FORCE_UNSAFE_TYPECAST_WARNINGS
-	#define PRAGMA_FORCE_UNSAFE_TYPECAST_WARNINGS \
-		__pragma (warning(push)) \
-		__pragma (warning(error: 4244)) /* 'argument': conversion from 'type1' to 'type2', possible loss of data */	\
-		__pragma (warning(error: 4838)) /* 'argument': conversion from 'type1' to 'type2' requires a narrowing conversion */
-#endif // PRAGMA_FORCE_UNSAFE_TYPECAST_WARNINGS
-
-#ifndef PRAGMA_RESTORE_UNSAFE_TYPECAST_WARNINGS
-	#define PRAGMA_RESTORE_UNSAFE_TYPECAST_WARNINGS \
-		__pragma(warning(pop))
-#endif // PRAGMA_RESTORE_UNSAFE_TYPECAST_WARNINGS
-
-#ifndef PRAGMA_DISABLE_UNDEFINED_IDENTIFIER_WARNINGS
-	#define PRAGMA_DISABLE_UNDEFINED_IDENTIFIER_WARNINGS \
-		__pragma(warning(push)) \
-		__pragma(warning(disable: 4668)) /* 'symbol' is not defined as a preprocessor macro, replacing with '0' for 'directives' */
-#endif // PRAGMA_DISABLE_UNDEFINED_IDENTIFIER_WARNINGS
-
-#ifndef PRAGMA_ENABLE_UNDEFINED_IDENTIFIER_WARNINGS
-	#define PRAGMA_ENABLE_UNDEFINED_IDENTIFIER_WARNINGS \
-		__pragma(warning(pop))
-#endif // PRAGMA_ENABLE_UNDEFINED_IDENTIFIER_WARNINGS
-
-#ifndef PRAGMA_DISABLE_MISSING_VIRTUAL_DESTRUCTOR_WARNINGS
-	#define PRAGMA_DISABLE_MISSING_VIRTUAL_DESTRUCTOR_WARNINGS \
-		__pragma(warning(push)) \
-		__pragma(warning(disable: 4265)) /* class' : class has virtual functions, but destructor is not virtual */
-#endif // PRAGMA_DISABLE_MISSING_VIRTUAL_DESTRUCTOR_WARNINGS
-
-#ifndef PRAGMA_ENABLE_MISSING_VIRTUAL_DESTRUCTOR_WARNINGS
-	#define PRAGMA_ENABLE_MISSING_VIRTUAL_DESTRUCTOR_WARNINGS \
-		__pragma(warning(pop))
-#endif // PRAGMA_ENABLE_MISSING_VIRTUAL_DESTRUCTOR_WARNINGS
-
-#ifndef PRAGMA_DISABLE_REORDER_WARNINGS
-	#define PRAGMA_DISABLE_REORDER_WARNINGS \
-		__pragma(warning(push)) \
-		__pragma(warning(disable: 5038)) /* data member 'member1' will be initialized after data member 'member2' data member 'member' will be initialized after base class 'base_class' */
-#endif // PRAGMA_DISABLE_REORDER_WARNINGS
-
-#ifndef PRAGMA_ENABLE_REORDER_WARNINGS
-	#define PRAGMA_ENABLE_REORDER_WARNINGS \
-		__pragma(warning(pop))
-#endif // PRAGMA_ENABLE_REORDER_WARNINGS
-
-#ifndef PRAGMA_DISABLE_REGISTER_WARNINGS
-	#define PRAGMA_DISABLE_REGISTER_WARNINGS \
-		__pragma(warning(push)) \
-		__pragma(warning(disable: 5033)) /* 'register' is no longer a supported storage class */
-#endif // PRAGMA_DISABLE_REGISTER_WARNINGS
-
-#ifndef PRAGMA_ENABLE_REGISTER_WARNINGS
-	#define PRAGMA_ENABLE_REGISTER_WARNINGS \
-		__pragma(warning(pop))
-#endif // PRAGMA_ENABLE_REGISTER_WARNINGS
-
-#ifndef PRAGMA_DISABLE_UNUSED_PRIVATE_FIELDS_WARNINGS
-#define PRAGMA_DISABLE_UNUSED_PRIVATE_FIELDS_WARNINGS
-	 // MSVC doesn't seem to have a warning similar to -Wunused-private-field
-#endif // PRAGMA_DISABLE_UNUSED_PRIVATE_FIELDS_WARNINGS
-
-#ifndef PRAGMA_ENABLE_UNUSED_PRIVATE_FIELDS_WARNINGS
-#define PRAGMA_ENABLE_UNUSED_PRIVATE_FIELDS_WARNINGS
-		// MSVC doesn't seem to have a warning similar to -Wunused-private-field
-#endif // PRAGMA_ENABLE_UNUSED_PRIVATE_FIELDS_WARNINGS
-
-#ifndef PRAGMA_DISABLE_UNINITIALIZED_CONST_REFERENCE_WARNINGS
-#define PRAGMA_DISABLE_UNINITIALIZED_CONST_REFERENCE_WARNINGS
-	// MSVC doesn't seem to have a warning similar to -Wuninitialized-const-reference
-#endif // PRAGMA_DISABLE_UNINITIALIZED_CONST_REFERENCE_WARNINGS
-
-#ifndef PRAGMA_ENABLE_UNINITIALIZED_CONST_REFERENCE_WARNINGS
-#define PRAGMA_ENABLE_UNINITIALIZED_CONST_REFERENCE_WARNINGS
-	// MSVC doesn't seem to have a warning similar to -Wuninitialized-const-reference
-#endif // PRAGMA_ENABLE_UNINITIALIZED_CONST_REFERENCE_WARNINGS
-
-#ifndef PRAGMA_POP
-	#define PRAGMA_POP \
-		__pragma(warning(pop))
-#endif // PRAGMA_POP
-
-// UE4 uses a struct packing of 4 for Win32 and 8 for Win64, 
-// and the default packing is 8 for Win32 and 16 for Win64
-#ifdef _WIN64
-	#ifndef PRAGMA_PUSH_PLATFORM_DEFAULT_PACKING
-		#define PRAGMA_PUSH_PLATFORM_DEFAULT_PACKING \
-			__pragma(pack(push)) \
-			__pragma(pack(16))
-	#endif // PRAGMA_PUSH_PLATFORM_DEFAULT_PACKING
-#else // _WIN64
-	#ifndef PRAGMA_PUSH_PLATFORM_DEFAULT_PACKING
-		#define PRAGMA_PUSH_PLATFORM_DEFAULT_PACKING \
-			__pragma(pack(push)) \
-			__pragma(pack(8))
-	#endif // PRAGMA_PUSH_PLATFORM_DEFAULT_PACKING
-#endif // _WIN64
-	
-#ifndef PRAGMA_POP_PLATFORM_DEFAULT_PACKING
-	#define PRAGMA_POP_PLATFORM_DEFAULT_PACKING \
-		__pragma(pack(pop))
-#endif // PRAGMA_POP_PLATFORM_DEFAULT_PACKING
-
-// Disable common CA warnings around SDK includes
-#ifndef THIRD_PARTY_INCLUDES_START
-	#define THIRD_PARTY_INCLUDES_START \
-		__pragma(warning(push)) \
-		__pragma(warning(disable: 4510))  /* '<class>': default constructor could not be generated. */ \
-		__pragma(warning(disable: 4610))  /* object '<class>' can never be instantiated - user-defined constructor required. */ \
-		__pragma(warning(disable: 4946))  /* reinterpret_cast used between related classes: '<class1>' and '<class1>' */ \
-		__pragma(warning(disable: 4996))  /* '<obj>' was declared deprecated. */ \
-		__pragma(warning(disable: 6011))  /* Dereferencing NULL pointer '<ptr>'. */ \
-		__pragma(warning(disable: 6101))  /* Returning uninitialized memory '<expr>'.  A successful path through the function does not set the named _Out_ parameter. */ \
-		__pragma(warning(disable: 6287))  /* Redundant code:  the left and right sub-expressions are identical. */ \
-		__pragma(warning(disable: 6308))  /* 'realloc' might return null pointer: assigning null pointer to 'X', which is passed as an argument to 'realloc', will cause the original memory block to be leaked. */ \
-		__pragma(warning(disable: 6326))  /* Potential comparison of a constant with another constant. */ \
-		__pragma(warning(disable: 6340))  /* Mismatch on sign: Incorrect type passed as parameter in call to function. */ \
-		__pragma(warning(disable: 6385))  /* Reading invalid data from '<ptr>':  the readable size is '<num1>' bytes, but '<num2>' bytes may be read. */ \
-		__pragma(warning(disable: 6386))  /* Buffer overrun while writing to '<ptr>':  the writable size is '<num1>' bytes, but '<num2>' bytes might be written. */ \
-		__pragma(warning(disable: 28182)) /* Dereferencing NULL pointer. '<ptr1>' contains the same NULL value as '<ptr2>' did. */ \
-		__pragma(warning(disable: 28251)) /* Inconsistent annotation for '<func>': this instance has no annotations. */ \
-		__pragma(warning(disable: 28252)) /* Inconsistent annotation for '<func>': return/function has '<annotation>' on the prior instance. */ \
-		__pragma(warning(disable: 28253)) /* Inconsistent annotation for '<func>': _Param_(<num>) has '<annotation>' on the prior instance. */ \
-		__pragma(warning(disable: 28301)) /* No annotations for first declaration of '<func>'. */ \
-		PRAGMA_DISABLE_REORDER_WARNINGS \
-		PRAGMA_DISABLE_UNDEFINED_IDENTIFIER_WARNINGS \
-		PRAGMA_DISABLE_SHADOW_VARIABLE_WARNINGS \
-		PRAGMA_DISABLE_MISSING_VIRTUAL_DESTRUCTOR_WARNINGS \
-		PRAGMA_DISABLE_DEPRECATION_WARNINGS \
-        PRAGMA_DISABLE_UNINITIALIZED_CONST_REFERENCE_WARNINGS
-#endif // THIRD_PARTY_INCLUDES_START
-
-#ifndef THIRD_PARTY_INCLUDES_END
-	#define THIRD_PARTY_INCLUDES_END \
-        PRAGMA_ENABLE_UNINITIALIZED_CONST_REFERENCE_WARNINGS \
-		PRAGMA_ENABLE_DEPRECATION_WARNINGS \
-		PRAGMA_ENABLE_MISSING_VIRTUAL_DESTRUCTOR_WARNINGS \
-		PRAGMA_ENABLE_SHADOW_VARIABLE_WARNINGS \
-		PRAGMA_ENABLE_UNDEFINED_IDENTIFIER_WARNINGS \
-		PRAGMA_ENABLE_REORDER_WARNINGS \
-		__pragma(warning(pop))
-#endif // THIRD_PARTY_INCLUDES_END
-
-#define EMIT_CUSTOM_WARNING_AT_LINE(Line, Warning) \
-	__pragma(message(WARNING_LOCATION(Line) ": warning C4996: " Warning))
+#include "MSVC/MSVCPlatformCompilerPreSetup.h"
diff --git a/Engine/Source/Runtime/Core/Public/Clang/ClangPlatformCompilerPreSetup.h b/Engine/Source/Runtime/Core/Public/Clang/ClangPlatformCompilerPreSetup.h
--- a/Engine/Source/Runtime/Core/Public/Clang/ClangPlatformCompilerPreSetup.h
+++ b/Engine/Source/Runtime/Core/Public/Clang/ClangPlatformCompilerPreSetup.h
@@ -200,6 +200,23 @@
 		_Pragma("clang diagnostic pop")
 #endif // PRAGMA_ENABLE_UNINITIALIZED_CONST_REFERENCE_WARNINGS
 
+#if __has_warning("-Wenum-constexpr-conversion")
+#define DISABLE_ENUM_CONSTEXPR_CONVERSION _Pragma("clang diagnostic ignored \"-Wenum-constexpr-conversion\"")
+#else
+#define DISABLE_ENUM_CONSTEXPR_CONVERSION
+#endif
+
+#ifndef PRAGMA_DISABLE_ENUM_CONSTEXPR_CONVERSION_WARNINGS
+#define PRAGMA_DISABLE_ENUM_CONSTEXPR_CONVERSION_WARNINGS \
+		_Pragma("clang diagnostic push") \
+        DISABLE_ENUM_CONSTEXPR_CONVERSION
+#endif // PRAGMA_DISABLE_ENUM_CONSTEXPR_CONVERSION_WARNINGS
+
+#ifndef PRAGMA_ENABLE_ENUM_CONSTEXPR_CONVERSION_WARNINGS
+#define PRAGMA_ENABLE_ENUM_CONSTEXPR_CONVERSION_WARNINGS \
+		_Pragma("clang diagnostic pop")
+#endif // PRAGMA_ENABLE_ENUM_CONSTEXPR_CONVERSION_WARNINGS
+
 #ifndef PRAGMA_POP
 	#define PRAGMA_POP \
 		_Pragma("clang diagnostic pop")
diff --git a/Engine/Source/Runtime/Core/Public/MSVC/MSVCPlatformCompilerPreSetup.h b/Engine/Source/Runtime/Core/Public/MSVC/MSVCPlatformCompilerPreSetup.h
new file mode 100644
--- /dev/null
+++ b/Engine/Source/Runtime/Core/Public/MSVC/MSVCPlatformCompilerPreSetup.h
@@ -0,0 +1,259 @@
+// Copyright Epic Games, Inc. All Rights Reserved.
+
+#pragma once
+
+// HEADER_UNIT_SKIP - Cause cirular dependencies for header units
+
+#if defined(__clang__)
+	// UE4 uses a struct packing of 4 for Win32 and 8 for Win64, 
+	// and the default packing is 8 for Win32 and 16 for Win64
+	#ifdef _WIN64
+		#ifndef PRAGMA_PUSH_PLATFORM_DEFAULT_PACKING
+			#define PRAGMA_PUSH_PLATFORM_DEFAULT_PACKING \
+				_Pragma("pack(push)") \
+				_Pragma("pack(16)")
+		#endif // PRAGMA_PUSH_PLATFORM_DEFAULT_PACKING
+	#else // _WIN64
+		#ifndef PRAGMA_PUSH_PLATFORM_DEFAULT_PACKING
+			#define PRAGMA_PUSH_PLATFORM_DEFAULT_PACKING \
+				_Pragma("pack(push)") \
+				_Pragma("pack(8)")
+		#endif // PRAGMA_PUSH_PLATFORM_DEFAULT_PACKING
+	#endif // _WIN64
+	
+	#ifndef PRAGMA_POP_PLATFORM_DEFAULT_PACKING
+		#define PRAGMA_POP_PLATFORM_DEFAULT_PACKING \
+			_Pragma("pack(pop)")
+	#endif // PRAGMA_POP_PLATFORM_DEFAULT_PACKING
+
+	#include "Clang/ClangPlatformCompilerPreSetup.h"
+
+	// Disable common CA warnings around SDK includes
+	#ifndef THIRD_PARTY_INCLUDES_START
+		#define THIRD_PARTY_INCLUDES_START \
+			PRAGMA_DISABLE_REORDER_WARNINGS \
+			PRAGMA_DISABLE_SHADOW_VARIABLE_WARNINGS \
+			PRAGMA_DISABLE_UNDEFINED_IDENTIFIER_WARNINGS \
+			PRAGMA_DISABLE_DEPRECATION_WARNINGS \
+			PRAGMA_DISABLE_OVERLOADED_VIRTUAL_WARNINGS \
+			PRAGMA_DISABLE_MISSING_BRACES_WARNINGS \
+			PRAGMA_DISABLE_UNINITIALIZED_CONST_REFERENCE_WARNINGS \
+			PRAGMA_DISABLE_ORDERED_COMPARE_FUNCTION_POINTERS
+	#endif // THIRD_PARTY_INCLUDES_START
+
+	#ifndef THIRD_PARTY_INCLUDES_END
+		#define THIRD_PARTY_INCLUDES_END \
+			PRAGMA_ENABLE_UNINITIALIZED_CONST_REFERENCE_WARNINGS \
+			PRAGMA_ENABLE_MISSING_BRACES_WARNINGS \
+			PRAGMA_ENABLE_OVERLOADED_VIRTUAL_WARNINGS \
+			PRAGMA_ENABLE_DEPRECATION_WARNINGS \
+			PRAGMA_ENABLE_UNDEFINED_IDENTIFIER_WARNINGS \
+			PRAGMA_ENABLE_SHADOW_VARIABLE_WARNINGS \
+			PRAGMA_ENABLE_REORDER_WARNINGS \
+			PRAGMA_ENABLE_ORDERED_COMPARE_FUNCTION_POINTERS
+	#endif // THIRD_PARTY_INCLUDES_END
+#else
+	#ifndef DISABLE_DEPRECATION
+		#define PRAGMA_DISABLE_DEPRECATION_WARNINGS \
+			__pragma (warning(push)) \
+			__pragma (warning(disable: 4995)) /* 'function': name was marked as #pragma deprecated */ \
+			__pragma (warning(disable: 4996)) /* The compiler encountered a deprecated declaration. */
+
+		#define PRAGMA_ENABLE_DEPRECATION_WARNINGS \
+			__pragma (warning(pop))
+	#endif // DISABLE_DEPRECATION
+
+	#ifndef PRAGMA_DISABLE_BUFFER_OVERRUN_WARNING
+		#define PRAGMA_DISABLE_BUFFER_OVERRUN_WARNING \
+			__pragma (warning(push)) \
+			__pragma (warning(disable:4789)) /* buffer 'identifier' of size N bytes will be overrun; M bytes will be written starting at offset L */
+	#endif
+
+	#ifndef PRAGMA_ENABLE_BUFFER_OVERRUN_WARNING
+		#define PRAGMA_ENABLE_BUFFER_OVERRUN_WARNING \
+			__pragma (warning(pop))
+	#endif
+
+	#ifndef PRAGMA_DISABLE_SHADOW_VARIABLE_WARNINGS
+		#define PRAGMA_DISABLE_SHADOW_VARIABLE_WARNINGS \
+			__pragma (warning(push)) \
+			__pragma (warning(disable: 4456)) /* declaration of 'LocalVariable' hides previous local declaration */ \
+			__pragma (warning(disable: 4457)) /* declaration of 'LocalVariable' hides function parameter */ \
+			__pragma (warning(disable: 4458)) /* declaration of 'LocalVariable' hides class member */ \
+			__pragma (warning(disable: 4459)) /* declaration of 'LocalVariable' hides global declaration */ \
+			__pragma (warning(disable: 6244)) /* local declaration of <variable> hides previous declaration at <line> of <file> */
+	#endif // PRAGMA_DISABLE_SHADOW_VARIABLE_WARNINGS
+
+	#ifndef PRAGMA_ENABLE_SHADOW_VARIABLE_WARNINGS
+		#define PRAGMA_ENABLE_SHADOW_VARIABLE_WARNINGS \
+			__pragma(warning(pop))
+	#endif // PRAGMA_ENABLE_SHADOW_VARIABLE_WARNINGS
+
+	#ifndef PRAGMA_DISABLE_UNSAFE_TYPECAST_WARNINGS
+		#define PRAGMA_DISABLE_UNSAFE_TYPECAST_WARNINGS \
+			__pragma (warning(push)) \
+			__pragma (warning(disable: 4244)) /* 'argument': conversion from 'type1' to 'type2', possible loss of data */	\
+			__pragma (warning(disable: 4838)) /* 'argument': conversion from 'type1' to 'type2' requires a narrowing conversion */
+	#endif // PRAGMA_DISABLE_UNSAFE_TYPECAST_WARNINGS
+
+	#ifndef PRAGMA_ENABLE_UNSAFE_TYPECAST_WARNINGS
+		#define PRAGMA_ENABLE_UNSAFE_TYPECAST_WARNINGS \
+			UE_DEPRECATED_MACRO(5.0, "The PRAGMA_ENABLE_UNSAFE_TYPECAST_WARNINGS macro has been deprecated in favor of PRAGMA_RESTORE_UNSAFE_TYPECAST_WARNINGS. To force enable warnings use PRAGMA_FORCE_UNSAFE_TYPECAST_WARNINGS.")
+	#endif // PRAGMA_ENABLE_UNSAFE_TYPECAST_WARNINGS
+
+	#ifndef PRAGMA_FORCE_UNSAFE_TYPECAST_WARNINGS
+		#define PRAGMA_FORCE_UNSAFE_TYPECAST_WARNINGS \
+			__pragma (warning(push)) \
+			__pragma (warning(error: 4244)) /* 'argument': conversion from 'type1' to 'type2', possible loss of data */	\
+			__pragma (warning(error: 4838)) /* 'argument': conversion from 'type1' to 'type2' requires a narrowing conversion */
+	#endif // PRAGMA_FORCE_UNSAFE_TYPECAST_WARNINGS
+
+	#ifndef PRAGMA_RESTORE_UNSAFE_TYPECAST_WARNINGS
+		#define PRAGMA_RESTORE_UNSAFE_TYPECAST_WARNINGS \
+			__pragma(warning(pop))
+	#endif // PRAGMA_RESTORE_UNSAFE_TYPECAST_WARNINGS
+
+	#ifndef PRAGMA_DISABLE_UNDEFINED_IDENTIFIER_WARNINGS
+		#define PRAGMA_DISABLE_UNDEFINED_IDENTIFIER_WARNINGS \
+			__pragma(warning(push)) \
+			__pragma(warning(disable: 4668)) /* 'symbol' is not defined as a preprocessor macro, replacing with '0' for 'directives' */
+	#endif // PRAGMA_DISABLE_UNDEFINED_IDENTIFIER_WARNINGS
+
+	#ifndef PRAGMA_ENABLE_UNDEFINED_IDENTIFIER_WARNINGS
+		#define PRAGMA_ENABLE_UNDEFINED_IDENTIFIER_WARNINGS \
+			__pragma(warning(pop))
+	#endif // PRAGMA_ENABLE_UNDEFINED_IDENTIFIER_WARNINGS
+
+	#ifndef PRAGMA_DISABLE_MISSING_VIRTUAL_DESTRUCTOR_WARNINGS
+		#define PRAGMA_DISABLE_MISSING_VIRTUAL_DESTRUCTOR_WARNINGS \
+			__pragma(warning(push)) \
+			__pragma(warning(disable: 4265)) /* class' : class has virtual functions, but destructor is not virtual */
+	#endif // PRAGMA_DISABLE_MISSING_VIRTUAL_DESTRUCTOR_WARNINGS
+
+	#ifndef PRAGMA_ENABLE_MISSING_VIRTUAL_DESTRUCTOR_WARNINGS
+		#define PRAGMA_ENABLE_MISSING_VIRTUAL_DESTRUCTOR_WARNINGS \
+			__pragma(warning(pop))
+	#endif // PRAGMA_ENABLE_MISSING_VIRTUAL_DESTRUCTOR_WARNINGS
+
+	#ifndef PRAGMA_DISABLE_REORDER_WARNINGS
+		#define PRAGMA_DISABLE_REORDER_WARNINGS \
+			__pragma(warning(push)) \
+			__pragma(warning(disable: 5038)) /* data member 'member1' will be initialized after data member 'member2' data member 'member' will be initialized after base class 'base_class' */
+	#endif // PRAGMA_DISABLE_REORDER_WARNINGS
+
+	#ifndef PRAGMA_ENABLE_REORDER_WARNINGS
+		#define PRAGMA_ENABLE_REORDER_WARNINGS \
+			__pragma(warning(pop))
+	#endif // PRAGMA_ENABLE_REORDER_WARNINGS
+
+	#ifndef PRAGMA_DISABLE_REGISTER_WARNINGS
+		#define PRAGMA_DISABLE_REGISTER_WARNINGS \
+			__pragma(warning(push)) \
+			__pragma(warning(disable: 5033)) /* 'register' is no longer a supported storage class */
+	#endif // PRAGMA_DISABLE_REGISTER_WARNINGS
+
+	#ifndef PRAGMA_ENABLE_REGISTER_WARNINGS
+		#define PRAGMA_ENABLE_REGISTER_WARNINGS \
+			__pragma(warning(pop))
+	#endif // PRAGMA_ENABLE_REGISTER_WARNINGS
+
+	#ifndef PRAGMA_DISABLE_UNUSED_PRIVATE_FIELDS_WARNINGS
+	#define PRAGMA_DISABLE_UNUSED_PRIVATE_FIELDS_WARNINGS
+			// MSVC doesn't seem to have a warning similar to -Wunused-private-field
+	#endif // PRAGMA_DISABLE_UNUSED_PRIVATE_FIELDS_WARNINGS
+
+	#ifndef PRAGMA_ENABLE_UNUSED_PRIVATE_FIELDS_WARNINGS
+	#define PRAGMA_ENABLE_UNUSED_PRIVATE_FIELDS_WARNINGS
+			// MSVC doesn't seem to have a warning similar to -Wunused-private-field
+	#endif // PRAGMA_ENABLE_UNUSED_PRIVATE_FIELDS_WARNINGS
+
+	#ifndef PRAGMA_DISABLE_UNINITIALIZED_CONST_REFERENCE_WARNINGS
+	#define PRAGMA_DISABLE_UNINITIALIZED_CONST_REFERENCE_WARNINGS
+		// MSVC doesn't seem to have a warning similar to -Wuninitialized-const-reference
+	#endif // PRAGMA_DISABLE_UNINITIALIZED_CONST_REFERENCE_WARNINGS
+
+	#ifndef PRAGMA_ENABLE_UNINITIALIZED_CONST_REFERENCE_WARNINGS
+	#define PRAGMA_ENABLE_UNINITIALIZED_CONST_REFERENCE_WARNINGS
+		// MSVC doesn't seem to have a warning similar to -Wuninitialized-const-reference
+	#endif // PRAGMA_ENABLE_UNINITIALIZED_CONST_REFERENCE_WARNINGS
+
+	#ifndef PRAGMA_DISABLE_ENUM_CONSTEXPR_CONVERSION_WARNINGS
+	#define PRAGMA_DISABLE_ENUM_CONSTEXPR_CONVERSION_WARNINGS
+		// MSVC doesn't seem to have a warning similar to -Wenum-constexpr-conversion
+	#endif // PRAGMA_DISABLE_ENUM_CONSTEXPR_CONVERSION_WARNINGS
+
+	#ifndef PRAGMA_ENABLE_ENUM_CONSTEXPR_CONVERSION_WARNINGS
+	#define PRAGMA_ENABLE_ENUM_CONSTEXPR_CONVERSION_WARNINGS
+		// MSVC doesn't seem to have a warning similar to -Wenum-constexpr-conversion
+	#endif // PRAGMA_ENABLE_ENUM_CONSTEXPR_CONVERSION_WARNINGS
+
+	#ifndef PRAGMA_POP
+		#define PRAGMA_POP \
+			__pragma(warning(pop))
+	#endif // PRAGMA_POP
+
+	// UE4 uses a struct packing of 4 for Win32 and 8 for Win64, 
+	// and the default packing is 8 for Win32 and 16 for Win64
+	#ifdef _WIN64
+		#ifndef PRAGMA_PUSH_PLATFORM_DEFAULT_PACKING
+			#define PRAGMA_PUSH_PLATFORM_DEFAULT_PACKING \
+				__pragma(pack(push, 16))
+		#endif // PRAGMA_PUSH_PLATFORM_DEFAULT_PACKING
+	#else // _WIN64
+		#ifndef PRAGMA_PUSH_PLATFORM_DEFAULT_PACKING
+			#define PRAGMA_PUSH_PLATFORM_DEFAULT_PACKING \
+				__pragma(pack(push, 8))
+		#endif // PRAGMA_PUSH_PLATFORM_DEFAULT_PACKING
+	#endif // _WIN64
+	
+	#ifndef PRAGMA_POP_PLATFORM_DEFAULT_PACKING
+		#define PRAGMA_POP_PLATFORM_DEFAULT_PACKING \
+			__pragma(pack(pop))
+	#endif // PRAGMA_POP_PLATFORM_DEFAULT_PACKING
+
+	// Disable common CA warnings around SDK includes
+	#ifndef THIRD_PARTY_INCLUDES_START
+		#define THIRD_PARTY_INCLUDES_START \
+			__pragma(warning(push)) \
+			__pragma(warning(disable: 4125))  /* decimal digit terminates octal escape sequence. */ \
+			__pragma(warning(disable: 4510))  /* '<class>': default constructor could not be generated. */ \
+			__pragma(warning(disable: 4610))  /* object '<class>' can never be instantiated - user-defined constructor required. */ \
+			__pragma(warning(disable: 4800))  /* Implicit conversion from '<type>' to bool. Possible information loss. */ \
+			__pragma(warning(disable: 4946))  /* reinterpret_cast used between related classes: '<class1>' and '<class2>' */ \
+			__pragma(warning(disable: 4996))  /* '<obj>' was declared deprecated. */ \
+			__pragma(warning(disable: 6011))  /* Dereferencing NULL pointer '<ptr>'. */ \
+			__pragma(warning(disable: 6101))  /* Returning uninitialized memory '<expr>'.  A successful path through the function does not set the named _Out_ parameter. */ \
+			__pragma(warning(disable: 6287))  /* Redundant code:  the left and right sub-expressions are identical. */ \
+			__pragma(warning(disable: 6308))  /* 'realloc' might return null pointer: assigning null pointer to 'X', which is passed as an argument to 'realloc', will cause the original memory block to be leaked. */ \
+			__pragma(warning(disable: 6326))  /* Potential comparison of a constant with another constant. */ \
+			__pragma(warning(disable: 6340))  /* Mismatch on sign: Incorrect type passed as parameter in call to function. */ \
+			__pragma(warning(disable: 6385))  /* Reading invalid data from '<ptr>':  the readable size is '<num1>' bytes, but '<num2>' bytes may be read. */ \
+			__pragma(warning(disable: 6386))  /* Buffer overrun while writing to '<ptr>':  the writable size is '<num1>' bytes, but '<num2>' bytes might be written. */ \
+			__pragma(warning(disable: 6553))  /* The annotation for function '<func>' on _Param_(<num>) does not apply to a value type. */ \
+			__pragma(warning(disable: 28182)) /* Dereferencing NULL pointer. '<ptr1>' contains the same NULL value as '<ptr2>' did. */ \
+			__pragma(warning(disable: 28251)) /* Inconsistent annotation for '<func>': this instance has no annotations. */ \
+			__pragma(warning(disable: 28252)) /* Inconsistent annotation for '<func>': return/function has '<annotation>' on the prior instance. */ \
+			__pragma(warning(disable: 28253)) /* Inconsistent annotation for '<func>': _Param_(<num>) has '<annotation>' on the prior instance. */ \
+			__pragma(warning(disable: 28301)) /* No annotations for first declaration of '<func>'. */ \
+			PRAGMA_DISABLE_REORDER_WARNINGS \
+			PRAGMA_DISABLE_UNDEFINED_IDENTIFIER_WARNINGS \
+			PRAGMA_DISABLE_SHADOW_VARIABLE_WARNINGS \
+			PRAGMA_DISABLE_MISSING_VIRTUAL_DESTRUCTOR_WARNINGS \
+			PRAGMA_DISABLE_DEPRECATION_WARNINGS \
+			PRAGMA_DISABLE_UNINITIALIZED_CONST_REFERENCE_WARNINGS
+	#endif // THIRD_PARTY_INCLUDES_START
+
+	#ifndef THIRD_PARTY_INCLUDES_END
+		#define THIRD_PARTY_INCLUDES_END \
+			PRAGMA_ENABLE_UNINITIALIZED_CONST_REFERENCE_WARNINGS \
+			PRAGMA_ENABLE_DEPRECATION_WARNINGS \
+			PRAGMA_ENABLE_MISSING_VIRTUAL_DESTRUCTOR_WARNINGS \
+			PRAGMA_ENABLE_SHADOW_VARIABLE_WARNINGS \
+			PRAGMA_ENABLE_UNDEFINED_IDENTIFIER_WARNINGS \
+			PRAGMA_ENABLE_REORDER_WARNINGS \
+			__pragma(warning(pop))
+	#endif // THIRD_PARTY_INCLUDES_END
+
+	#define EMIT_CUSTOM_WARNING_AT_LINE(Line, Warning) \
+		__pragma(message(WARNING_LOCATION(Line) ": warning C4996: " Warning))
+#endif // defined(__clang__)
diff --git a/Engine/Source/Runtime/Core/Public/Windows/WindowsPlatformCompilerPreSetup.h b/Engine/Source/Runtime/Core/Public/Windows/WindowsPlatformCompilerPreSetup.h
--- a/Engine/Source/Runtime/Core/Public/Windows/WindowsPlatformCompilerPreSetup.h
+++ b/Engine/Source/Runtime/Core/Public/Windows/WindowsPlatformCompilerPreSetup.h
@@ -4,246 +4,4 @@
 
 // HEADER_UNIT_SKIP - Cause cirular dependencies for header units
 
-#if defined(__clang__)
-	// UE4 uses a struct packing of 4 for Win32 and 8 for Win64, 
-	// and the default packing is 8 for Win32 and 16 for Win64
-	#ifdef _WIN64
-		#ifndef PRAGMA_PUSH_PLATFORM_DEFAULT_PACKING
-			#define PRAGMA_PUSH_PLATFORM_DEFAULT_PACKING \
-				_Pragma("pack(push)") \
-				_Pragma("pack(16)")
-		#endif // PRAGMA_PUSH_PLATFORM_DEFAULT_PACKING
-	#else // _WIN64
-		#ifndef PRAGMA_PUSH_PLATFORM_DEFAULT_PACKING
-			#define PRAGMA_PUSH_PLATFORM_DEFAULT_PACKING \
-				_Pragma("pack(push)") \
-				_Pragma("pack(8)")
-		#endif // PRAGMA_PUSH_PLATFORM_DEFAULT_PACKING
-	#endif // _WIN64
-	
-	#ifndef PRAGMA_POP_PLATFORM_DEFAULT_PACKING
-		#define PRAGMA_POP_PLATFORM_DEFAULT_PACKING \
-			_Pragma("pack(pop)")
-	#endif // PRAGMA_POP_PLATFORM_DEFAULT_PACKING
-
-	#include "Clang/ClangPlatformCompilerPreSetup.h"
-
-	// Disable common CA warnings around SDK includes
-	#ifndef THIRD_PARTY_INCLUDES_START
-		#define THIRD_PARTY_INCLUDES_START \
-			PRAGMA_DISABLE_REORDER_WARNINGS \
-			PRAGMA_DISABLE_SHADOW_VARIABLE_WARNINGS \
-			PRAGMA_DISABLE_UNDEFINED_IDENTIFIER_WARNINGS \
-			PRAGMA_DISABLE_DEPRECATION_WARNINGS \
-			PRAGMA_DISABLE_OVERLOADED_VIRTUAL_WARNINGS \
-			PRAGMA_DISABLE_MISSING_BRACES_WARNINGS \
-	        PRAGMA_DISABLE_UNINITIALIZED_CONST_REFERENCE_WARNINGS \
-			PRAGMA_DISABLE_ORDERED_COMPARE_FUNCTION_POINTERS
-	#endif // THIRD_PARTY_INCLUDES_START
-
-	#ifndef THIRD_PARTY_INCLUDES_END
-		#define THIRD_PARTY_INCLUDES_END \
-	        PRAGMA_ENABLE_UNINITIALIZED_CONST_REFERENCE_WARNINGS \
-			PRAGMA_ENABLE_MISSING_BRACES_WARNINGS \
-			PRAGMA_ENABLE_OVERLOADED_VIRTUAL_WARNINGS \
-			PRAGMA_ENABLE_DEPRECATION_WARNINGS \
-			PRAGMA_ENABLE_UNDEFINED_IDENTIFIER_WARNINGS \
-			PRAGMA_ENABLE_SHADOW_VARIABLE_WARNINGS \
-			PRAGMA_ENABLE_REORDER_WARNINGS \
-			PRAGMA_ENABLE_ORDERED_COMPARE_FUNCTION_POINTERS
-	#endif // THIRD_PARTY_INCLUDES_END
-#else
-	#ifndef DISABLE_DEPRECATION
-		#define PRAGMA_DISABLE_DEPRECATION_WARNINGS \
-			__pragma (warning(push)) \
-			__pragma (warning(disable: 4995)) /* 'function': name was marked as #pragma deprecated */ \
-			__pragma (warning(disable: 4996)) /* The compiler encountered a deprecated declaration. */
-
-		#define PRAGMA_ENABLE_DEPRECATION_WARNINGS \
-			__pragma (warning(pop))
-	#endif // DISABLE_DEPRECATION
-
-	#ifndef PRAGMA_DISABLE_BUFFER_OVERRUN_WARNING
-		#define PRAGMA_DISABLE_BUFFER_OVERRUN_WARNING \
-			__pragma (warning(push)) \
-			__pragma (warning(disable:4789)) /* buffer 'identifier' of size N bytes will be overrun; M bytes will be written starting at offset L */
-	#endif
-
-	#ifndef PRAGMA_ENABLE_BUFFER_OVERRUN_WARNING
-		#define PRAGMA_ENABLE_BUFFER_OVERRUN_WARNING \
-			__pragma (warning(pop))
-	#endif
-
-	#ifndef PRAGMA_DISABLE_SHADOW_VARIABLE_WARNINGS
-		#define PRAGMA_DISABLE_SHADOW_VARIABLE_WARNINGS \
-			__pragma (warning(push)) \
-			__pragma (warning(disable: 4456)) /* declaration of 'LocalVariable' hides previous local declaration */ \
-			__pragma (warning(disable: 4457)) /* declaration of 'LocalVariable' hides function parameter */ \
-			__pragma (warning(disable: 4458)) /* declaration of 'LocalVariable' hides class member */ \
-			__pragma (warning(disable: 4459)) /* declaration of 'LocalVariable' hides global declaration */ \
-			__pragma (warning(disable: 6244)) /* local declaration of <variable> hides previous declaration at <line> of <file> */
-	#endif // PRAGMA_DISABLE_SHADOW_VARIABLE_WARNINGS
-
-	#ifndef PRAGMA_ENABLE_SHADOW_VARIABLE_WARNINGS
-		#define PRAGMA_ENABLE_SHADOW_VARIABLE_WARNINGS \
-			__pragma(warning(pop))
-	#endif // PRAGMA_ENABLE_SHADOW_VARIABLE_WARNINGS
-
-	#ifndef PRAGMA_DISABLE_UNSAFE_TYPECAST_WARNINGS
-		#define PRAGMA_DISABLE_UNSAFE_TYPECAST_WARNINGS \
-			__pragma (warning(push)) \
-			__pragma (warning(disable: 4244)) /* 'argument': conversion from 'type1' to 'type2', possible loss of data */	\
-			__pragma (warning(disable: 4838)) /* 'argument': conversion from 'type1' to 'type2' requires a narrowing conversion */
-	#endif // PRAGMA_DISABLE_UNSAFE_TYPECAST_WARNINGS
-
-	#ifndef PRAGMA_ENABLE_UNSAFE_TYPECAST_WARNINGS
-		#define PRAGMA_ENABLE_UNSAFE_TYPECAST_WARNINGS \
-			UE_DEPRECATED_MACRO(5.0, "The PRAGMA_ENABLE_UNSAFE_TYPECAST_WARNINGS macro has been deprecated in favor of PRAGMA_RESTORE_UNSAFE_TYPECAST_WARNINGS. To force enable warnings use PRAGMA_FORCE_UNSAFE_TYPECAST_WARNINGS.")
-	#endif // PRAGMA_ENABLE_UNSAFE_TYPECAST_WARNINGS
-
-	#ifndef PRAGMA_FORCE_UNSAFE_TYPECAST_WARNINGS
-		#define PRAGMA_FORCE_UNSAFE_TYPECAST_WARNINGS \
-			__pragma (warning(push)) \
-			__pragma (warning(error: 4244)) /* 'argument': conversion from 'type1' to 'type2', possible loss of data */	\
-			__pragma (warning(error: 4838)) /* 'argument': conversion from 'type1' to 'type2' requires a narrowing conversion */
-	#endif // PRAGMA_FORCE_UNSAFE_TYPECAST_WARNINGS
-
-	#ifndef PRAGMA_RESTORE_UNSAFE_TYPECAST_WARNINGS
-		#define PRAGMA_RESTORE_UNSAFE_TYPECAST_WARNINGS \
-			__pragma(warning(pop))
-	#endif // PRAGMA_RESTORE_UNSAFE_TYPECAST_WARNINGS
-
-	#ifndef PRAGMA_DISABLE_UNDEFINED_IDENTIFIER_WARNINGS
-		#define PRAGMA_DISABLE_UNDEFINED_IDENTIFIER_WARNINGS \
-			__pragma(warning(push)) \
-			__pragma(warning(disable: 4668)) /* 'symbol' is not defined as a preprocessor macro, replacing with '0' for 'directives' */
-	#endif // PRAGMA_DISABLE_UNDEFINED_IDENTIFIER_WARNINGS
-
-	#ifndef PRAGMA_ENABLE_UNDEFINED_IDENTIFIER_WARNINGS
-		#define PRAGMA_ENABLE_UNDEFINED_IDENTIFIER_WARNINGS \
-			__pragma(warning(pop))
-	#endif // PRAGMA_ENABLE_UNDEFINED_IDENTIFIER_WARNINGS
-
-	#ifndef PRAGMA_DISABLE_MISSING_VIRTUAL_DESTRUCTOR_WARNINGS
-		#define PRAGMA_DISABLE_MISSING_VIRTUAL_DESTRUCTOR_WARNINGS \
-			__pragma(warning(push)) \
-			__pragma(warning(disable: 4265)) /* class' : class has virtual functions, but destructor is not virtual */
-	#endif // PRAGMA_DISABLE_MISSING_VIRTUAL_DESTRUCTOR_WARNINGS
-
-	#ifndef PRAGMA_ENABLE_MISSING_VIRTUAL_DESTRUCTOR_WARNINGS
-		#define PRAGMA_ENABLE_MISSING_VIRTUAL_DESTRUCTOR_WARNINGS \
-			__pragma(warning(pop))
-	#endif // PRAGMA_ENABLE_MISSING_VIRTUAL_DESTRUCTOR_WARNINGS
-
-	#ifndef PRAGMA_DISABLE_REORDER_WARNINGS
-		#define PRAGMA_DISABLE_REORDER_WARNINGS \
-			__pragma(warning(push)) \
-			__pragma(warning(disable: 5038)) /* data member 'member1' will be initialized after data member 'member2' data member 'member' will be initialized after base class 'base_class' */
-	#endif // PRAGMA_DISABLE_REORDER_WARNINGS
-
-	#ifndef PRAGMA_ENABLE_REORDER_WARNINGS
-		#define PRAGMA_ENABLE_REORDER_WARNINGS \
-			__pragma(warning(pop))
-	#endif // PRAGMA_ENABLE_REORDER_WARNINGS
-
-	#ifndef PRAGMA_DISABLE_REGISTER_WARNINGS
-		#define PRAGMA_DISABLE_REGISTER_WARNINGS \
-			__pragma(warning(push)) \
-			__pragma(warning(disable: 5033)) /* 'register' is no longer a supported storage class */
-	#endif // PRAGMA_DISABLE_REGISTER_WARNINGS
-
-	#ifndef PRAGMA_ENABLE_REGISTER_WARNINGS
-		#define PRAGMA_ENABLE_REGISTER_WARNINGS \
-			__pragma(warning(pop))
-	#endif // PRAGMA_ENABLE_REGISTER_WARNINGS
-
-	#ifndef PRAGMA_DISABLE_UNUSED_PRIVATE_FIELDS_WARNINGS
-	#define PRAGMA_DISABLE_UNUSED_PRIVATE_FIELDS_WARNINGS
-			// MSVC doesn't seem to have a warning similar to -Wunused-private-field
-	#endif // PRAGMA_DISABLE_UNUSED_PRIVATE_FIELDS_WARNINGS
-
-	#ifndef PRAGMA_ENABLE_UNUSED_PRIVATE_FIELDS_WARNINGS
-	#define PRAGMA_ENABLE_UNUSED_PRIVATE_FIELDS_WARNINGS
-			// MSVC doesn't seem to have a warning similar to -Wunused-private-field
-	#endif // PRAGMA_ENABLE_UNUSED_PRIVATE_FIELDS_WARNINGS
-
-	#ifndef PRAGMA_DISABLE_UNINITIALIZED_CONST_REFERENCE_WARNINGS
-	#define PRAGMA_DISABLE_UNINITIALIZED_CONST_REFERENCE_WARNINGS
-		// MSVC doesn't seem to have a warning similar to -Wuninitialized-const-reference
-	#endif // PRAGMA_DISABLE_UNINITIALIZED_CONST_REFERENCE_WARNINGS
-
-	#ifndef PRAGMA_ENABLE_UNINITIALIZED_CONST_REFERENCE_WARNINGS
-	#define PRAGMA_ENABLE_UNINITIALIZED_CONST_REFERENCE_WARNINGS
-		// MSVC doesn't seem to have a warning similar to -Wuninitialized-const-reference
-	#endif // PRAGMA_ENABLE_UNINITIALIZED_CONST_REFERENCE_WARNINGS
-
-	#ifndef PRAGMA_POP
-		#define PRAGMA_POP \
-			__pragma(warning(pop))
-	#endif // PRAGMA_POP
-
-	// UE4 uses a struct packing of 4 for Win32 and 8 for Win64, 
-	// and the default packing is 8 for Win32 and 16 for Win64
-	#ifdef _WIN64
-		#ifndef PRAGMA_PUSH_PLATFORM_DEFAULT_PACKING
-			#define PRAGMA_PUSH_PLATFORM_DEFAULT_PACKING \
-				__pragma(pack(push, 16))
-		#endif // PRAGMA_PUSH_PLATFORM_DEFAULT_PACKING
-	#else // _WIN64
-		#ifndef PRAGMA_PUSH_PLATFORM_DEFAULT_PACKING
-			#define PRAGMA_PUSH_PLATFORM_DEFAULT_PACKING \
-				__pragma(pack(push, 8))
-		#endif // PRAGMA_PUSH_PLATFORM_DEFAULT_PACKING
-	#endif // _WIN64
-	
-	#ifndef PRAGMA_POP_PLATFORM_DEFAULT_PACKING
-		#define PRAGMA_POP_PLATFORM_DEFAULT_PACKING \
-			__pragma(pack(pop))
-	#endif // PRAGMA_POP_PLATFORM_DEFAULT_PACKING
-
-	// Disable common CA warnings around SDK includes
-	#ifndef THIRD_PARTY_INCLUDES_START
-		#define THIRD_PARTY_INCLUDES_START \
-			__pragma(warning(push)) \
-			__pragma(warning(disable: 4125))  /* decimal digit terminates octal escape sequence. */ \
-			__pragma(warning(disable: 4510))  /* '<class>': default constructor could not be generated. */ \
-			__pragma(warning(disable: 4610))  /* object '<class>' can never be instantiated - user-defined constructor required. */ \
-			__pragma(warning(disable: 4800))  /* Implicit conversion from '<type>' to bool. Possible information loss. */ \
-			__pragma(warning(disable: 4946))  /* reinterpret_cast used between related classes: '<class1>' and '<class2>' */ \
-			__pragma(warning(disable: 4996))  /* '<obj>' was declared deprecated. */ \
-			__pragma(warning(disable: 6011))  /* Dereferencing NULL pointer '<ptr>'. */ \
-			__pragma(warning(disable: 6101))  /* Returning uninitialized memory '<expr>'.  A successful path through the function does not set the named _Out_ parameter. */ \
-			__pragma(warning(disable: 6287))  /* Redundant code:  the left and right sub-expressions are identical. */ \
-			__pragma(warning(disable: 6308))  /* 'realloc' might return null pointer: assigning null pointer to 'X', which is passed as an argument to 'realloc', will cause the original memory block to be leaked. */ \
-			__pragma(warning(disable: 6326))  /* Potential comparison of a constant with another constant. */ \
-			__pragma(warning(disable: 6340))  /* Mismatch on sign: Incorrect type passed as parameter in call to function. */ \
-			__pragma(warning(disable: 6385))  /* Reading invalid data from '<ptr>':  the readable size is '<num1>' bytes, but '<num2>' bytes may be read. */ \
-			__pragma(warning(disable: 6386))  /* Buffer overrun while writing to '<ptr>':  the writable size is '<num1>' bytes, but '<num2>' bytes might be written. */ \
-			__pragma(warning(disable: 6553))  /* The annotation for function '<func>' on _Param_(<num>) does not apply to a value type. */ \
-			__pragma(warning(disable: 28182)) /* Dereferencing NULL pointer. '<ptr1>' contains the same NULL value as '<ptr2>' did. */ \
-			__pragma(warning(disable: 28251)) /* Inconsistent annotation for '<func>': this instance has no annotations. */ \
-			__pragma(warning(disable: 28252)) /* Inconsistent annotation for '<func>': return/function has '<annotation>' on the prior instance. */ \
-			__pragma(warning(disable: 28253)) /* Inconsistent annotation for '<func>': _Param_(<num>) has '<annotation>' on the prior instance. */ \
-			__pragma(warning(disable: 28301)) /* No annotations for first declaration of '<func>'. */ \
-			PRAGMA_DISABLE_REORDER_WARNINGS \
-			PRAGMA_DISABLE_UNDEFINED_IDENTIFIER_WARNINGS \
-			PRAGMA_DISABLE_SHADOW_VARIABLE_WARNINGS \
-			PRAGMA_DISABLE_MISSING_VIRTUAL_DESTRUCTOR_WARNINGS \
-			PRAGMA_DISABLE_DEPRECATION_WARNINGS \
-			PRAGMA_DISABLE_UNINITIALIZED_CONST_REFERENCE_WARNINGS
-	#endif // THIRD_PARTY_INCLUDES_START
-
-	#ifndef THIRD_PARTY_INCLUDES_END
-		#define THIRD_PARTY_INCLUDES_END \
-			PRAGMA_ENABLE_UNINITIALIZED_CONST_REFERENCE_WARNINGS \
-			PRAGMA_ENABLE_DEPRECATION_WARNINGS \
-			PRAGMA_ENABLE_MISSING_VIRTUAL_DESTRUCTOR_WARNINGS \
-			PRAGMA_ENABLE_SHADOW_VARIABLE_WARNINGS \
-			PRAGMA_ENABLE_UNDEFINED_IDENTIFIER_WARNINGS \
-			PRAGMA_ENABLE_REORDER_WARNINGS \
-			__pragma(warning(pop))
-	#endif // THIRD_PARTY_INCLUDES_END
-
-	#define EMIT_CUSTOM_WARNING_AT_LINE(Line, Warning) \
-		__pragma(message(WARNING_LOCATION(Line) ": warning C4996: " Warning))
-#endif // defined(__clang__)
+#include "MSVC/MSVCPlatformCompilerPreSetup.h"
