Fix race in OutputBeginDynamicEvent

#rnx

diff --git a/Engine/Source/Runtime/Core/Private/ProfilingDebugging/CpuProfilerTrace.cpp b/Engine/Source/Runtime/Core/Private/ProfilingDebugging/CpuProfilerTrace.cpp
--- a/Engine/Source/Runtime/Core/Private/ProfilingDebugging/CpuProfilerTrace.cpp
+++ b/Engine/Source/Runtime/Core/Private/ProfilingDebugging/CpuProfilerTrace.cpp
@@ -212,13 +212,13 @@ void FCpuProfilerTrace::OutputBeginDynamicEvent(const FName& Name, const ANSICHA
 		const FNameEntry* NameEntry = Name.GetDisplayNameEntry();
 		if (NameEntry->IsWide())
 		{
-			static WIDECHAR WideName[NAME_SIZE];
+			WIDECHAR WideName[NAME_SIZE];
 			NameEntry->GetWideName(WideName);
 			SpecId = OutputEventType(WideName, File, Line);
 		}
 		else
 		{
-			static ANSICHAR AnsiName[NAME_SIZE];
+			ANSICHAR AnsiName[NAME_SIZE];
 			NameEntry->GetAnsiName(AnsiName);
 			SpecId = OutputEventType(AnsiName, File, Line);
 		}
