Fix maps not rehasing when the ReplaceReferenceArchieve is being used.

diff --git a/Engine/Source/Runtime/Core/Public/Containers/Set.h b/Engine/Source/Runtime/Core/Public/Containers/Set.h
--- a/Engine/Source/Runtime/Core/Public/Containers/Set.h
+++ b/Engine/Source/Runtime/Core/Public/Containers/Set.h
@@ -2186,7 +2186,7 @@ struct TSetPrivateFriend
 		// Load the set's new elements.
 		Ar << Set.Elements;
 
-		if(Ar.IsLoading())
+		if(Ar.IsLoading() || Ar.IsModifyingWeakAndStrongReferences())
 		{
 			// Free the old hash.
 			Set.Hash.ResizeAllocation(0,0,sizeof(FSetElementId));
@@ -2205,7 +2205,7 @@ struct TSetPrivateFriend
  	{
 		Slot << Set.Elements;
 
-		if (Slot.GetUnderlyingArchive().IsLoading())
+		if (Slot.GetUnderlyingArchive().IsLoading() || Slot.GetUnderlyingArchive().IsModifyingWeakAndStrongReferences())
 		{
 			// Free the old hash.
 			Set.Hash.ResizeAllocation(0, 0, sizeof(FSetElementId));
