Fix compile error
#rnx

diff --git a/Engine/Source/Runtime/Core/Private/HAL/ThreadingBase.cpp b/Engine/Source/Runtime/Core/Private/HAL/ThreadingBase.cpp
--- a/Engine/Source/Runtime/Core/Private/HAL/ThreadingBase.cpp
+++ b/Engine/Source/Runtime/Core/Private/HAL/ThreadingBase.cpp
@@ -1341,9 +1341,9 @@ FQueuedThread::Run()
 
 FTlsAutoCleanup* FThreadSingletonInitializer::Get( TFunctionRef<FTlsAutoCleanup*()> CreateInstance, uint32& InOutTlsSlot )
 {
-	AutoRTFM::Open([&InOutTlsSlot]
+	uint32 TlsSlot = (uint32)FPlatformAtomics::AtomicRead_Relaxed((int32*)&InOutTlsSlot);
+	AutoRTFM::Open([&TlsSlot, &InOutTlsSlot]
 	{
-		uint32 TlsSlot = (uint32)FPlatformAtomics::AtomicRead_Relaxed((int32*)&InOutTlsSlot);
 		if (TlsSlot == 0xFFFFFFFF)
 		{
 			const uint32 ThisTlsSlot = FPlatformTLS::AllocTlsSlot();
