Removed useless DISABLE_DEPRECATION macro.

diff --git a/Engine/Source/Runtime/Core/Public/Clang/ClangPlatformCompilerPreSetup.h b/Engine/Source/Runtime/Core/Public/Clang/ClangPlatformCompilerPreSetup.h
--- a/Engine/Source/Runtime/Core/Public/Clang/ClangPlatformCompilerPreSetup.h
+++ b/Engine/Source/Runtime/Core/Public/Clang/ClangPlatformCompilerPreSetup.h
@@ -4,16 +4,18 @@
 
 // HEADER_UNIT_UNSUPPORTED - Clang not supporting header units
 
-#ifndef DISABLE_DEPRECATION
-	#pragma clang diagnostic warning "-Wdeprecated-declarations"
+#pragma clang diagnostic warning "-Wdeprecated-declarations"
 
+#ifndef PRAGMA_DISABLE_DEPRECATION_WARNINGS
 	#define PRAGMA_DISABLE_DEPRECATION_WARNINGS \
 		_Pragma("clang diagnostic push") \
 		_Pragma("clang diagnostic ignored \"-Wdeprecated-declarations\"")
+#endif // PRAGMA_DISABLE_DEPRECATION_WARNINGS
 
+#ifndef PRAGMA_ENABLE_DEPRECATION_WARNINGS
 	#define PRAGMA_ENABLE_DEPRECATION_WARNINGS \
 		_Pragma("clang diagnostic pop")
-#endif // DISABLE_DEPRECATION
+#endif // PRAGMA_ENABLE_DEPRECATION_WARNINGS
 
 #ifndef PRAGMA_DISABLE_OVERLOADED_VIRTUAL_WARNINGS
 	#define PRAGMA_DISABLE_OVERLOADED_VIRTUAL_WARNINGS \
diff --git a/Engine/Source/Runtime/Core/Public/MSVC/MSVCPlatformCompilerPreSetup.h b/Engine/Source/Runtime/Core/Public/MSVC/MSVCPlatformCompilerPreSetup.h
--- a/Engine/Source/Runtime/Core/Public/MSVC/MSVCPlatformCompilerPreSetup.h
+++ b/Engine/Source/Runtime/Core/Public/MSVC/MSVCPlatformCompilerPreSetup.h
@@ -53,26 +53,28 @@
 			PRAGMA_ENABLE_ORDERED_COMPARE_FUNCTION_POINTERS
 	#endif // THIRD_PARTY_INCLUDES_END
 #else
-	#ifndef DISABLE_DEPRECATION
+	#ifndef PRAGMA_DISABLE_DEPRECATION_WARNINGS
 		#define PRAGMA_DISABLE_DEPRECATION_WARNINGS \
 			__pragma (warning(push)) \
 			__pragma (warning(disable: 4995)) /* 'function': name was marked as #pragma deprecated */ \
 			__pragma (warning(disable: 4996)) /* The compiler encountered a deprecated declaration. */
+	#endif // PRAGMA_DISABLE_DEPRECATION_WARNINGS
 
+	#ifndef PRAGMA_ENABLE_DEPRECATION_WARNINGS
 		#define PRAGMA_ENABLE_DEPRECATION_WARNINGS \
 			__pragma (warning(pop))
-	#endif // DISABLE_DEPRECATION
+	#endif // PRAGMA_ENABLE_DEPRECATION_WARNINGS
 
 	#ifndef PRAGMA_DISABLE_BUFFER_OVERRUN_WARNING
 		#define PRAGMA_DISABLE_BUFFER_OVERRUN_WARNING \
 			__pragma (warning(push)) \
 			__pragma (warning(disable:4789)) /* buffer 'identifier' of size N bytes will be overrun; M bytes will be written starting at offset L */
-	#endif
+	#endif // PRAGMA_DISABLE_BUFFER_OVERRUN_WARNING
 
 	#ifndef PRAGMA_ENABLE_BUFFER_OVERRUN_WARNING
 		#define PRAGMA_ENABLE_BUFFER_OVERRUN_WARNING \
 			__pragma (warning(pop))
-	#endif
+	#endif // PRAGMA_ENABLE_BUFFER_OVERRUN_WARNING
 
 	#ifndef PRAGMA_DISABLE_SHADOW_VARIABLE_WARNINGS
 		#define PRAGMA_DISABLE_SHADOW_VARIABLE_WARNINGS \
