hid `LowLevelTask.TryCancel()` behind `ReleaseInternalReference()` cos it confuses users

diff --git a/Engine/Source/Runtime/Core/Public/Tasks/TaskPrivate.h b/Engine/Source/Runtime/Core/Public/Tasks/TaskPrivate.h
--- a/Engine/Source/Runtime/Core/Public/Tasks/TaskPrivate.h
+++ b/Engine/Source/Runtime/Core/Public/Tasks/TaskPrivate.h
@@ -554,7 +554,7 @@ namespace UE::Tasks
 						// "inline" tasks are not scheduled but executed straight away
 						TryExecuteTask(); // result doesn't matter, this can fail if task retraction jumped in and got execution
 						// permission between this thread unlocked the task and tried to execute it
-						verify(LowLevelTask.TryCancel());
+						ReleaseInternalReference();
 					}
 					else if (ExtendedPriority == EExtendedTaskPriority::TaskEvent)
 					{
@@ -565,7 +565,7 @@ namespace UE::Tasks
 							// task events are used as an empty prerequisites/subsequents
 							ReleasePrerequisites();
 							Close();
-							verify(LowLevelTask.TryCancel()); // releases the internal reference
+							ReleaseInternalReference();
 						}
 					}
 					else
