[Insights] CountersTrace: Added explicit support for counters with dynamic name.

diff --git a/Engine/Source/Runtime/Core/Private/IO/IoDispatcher.cpp b/Engine/Source/Runtime/Core/Private/IO/IoDispatcher.cpp
--- a/Engine/Source/Runtime/Core/Private/IO/IoDispatcher.cpp
+++ b/Engine/Source/Runtime/Core/Private/IO/IoDispatcher.cpp
@@ -135,8 +135,8 @@ private:
 	{
 		FRequestCategory(const TCHAR* Name)
 #if COUNTERSTRACE_ENABLED
-			: TotalLoadedCounter(*FString::Printf(TEXT("IoDispatcher/TotalLoaded (%s)"), Name), TraceCounterDisplayHint_Memory)
-			, AverageDurationCounter(*FString::Printf(TEXT("IoDispatcher/AverageDuration (%s)"), Name), TraceCounterDisplayHint_None)
+			: TotalLoadedCounter(TraceCounterNameType_Dynamic, *FString::Printf(TEXT("IoDispatcher/TotalLoaded (%s)"), Name), TraceCounterDisplayHint_Memory)
+			, AverageDurationCounter(TraceCounterNameType_Dynamic, *FString::Printf(TEXT("IoDispatcher/AverageDuration (%s)"), Name), TraceCounterDisplayHint_None)
 #endif
 		{
 
diff --git a/Engine/Source/Runtime/Core/Private/ProfilingDebugging/CountersTrace.cpp b/Engine/Source/Runtime/Core/Private/ProfilingDebugging/CountersTrace.cpp
--- a/Engine/Source/Runtime/Core/Private/ProfilingDebugging/CountersTrace.cpp
+++ b/Engine/Source/Runtime/Core/Private/ProfilingDebugging/CountersTrace.cpp
@@ -27,6 +27,19 @@ UE_TRACE_EVENT_BEGIN(Counters, SetValueFloat)
 	UE_TRACE_EVENT_FIELD(uint16, CounterId)
 UE_TRACE_EVENT_END()
 
+const TCHAR* FCountersTrace::AllocAndCopyCounterName(const TCHAR* InCounterName)
+{
+	int32 Len = FCString::Strlen(InCounterName);
+	const TCHAR* CounterName = new TCHAR[Len + 1];
+	FCString::Strncpy((TCHAR*)CounterName, InCounterName, Len + 1);
+	return CounterName;
+}
+
+void FCountersTrace::FreeCounterName(const TCHAR* InCounterName)
+{
+	delete InCounterName;
+}
+
 uint16 FCountersTrace::OutputInitCounter(const TCHAR* CounterName, ETraceCounterType CounterType, ETraceCounterDisplayHint CounterDisplayHint)
 {
 	if (!UE_TRACE_CHANNELEXPR_IS_ENABLED(CountersChannel))
@@ -35,14 +48,17 @@ uint16 FCountersTrace::OutputInitCounter(const TCHAR* CounterName, ETraceCounter
 	}
 
 	static std::atomic<uint32> NextId { 0 };
-
 	uint16 CounterId = uint16(++NextId);
+
+	check(CounterName); // a trace counter object (TCounter<>) is used before it is constructed!?
 	uint16 NameLen = uint16(FCString::Strlen(CounterName));
+
 	UE_TRACE_LOG(Counters, Spec, CountersChannel, NameLen * sizeof(ANSICHAR))
 		<< Spec.Id(CounterId)
 		<< Spec.Type(uint8(CounterType))
 		<< Spec.DisplayHint(uint8(CounterDisplayHint))
 		<< Spec.Name(CounterName, NameLen);
+
 	return CounterId;
 }
 
diff --git a/Engine/Source/Runtime/Core/Private/ProfilingDebugging/StallDetector.cpp b/Engine/Source/Runtime/Core/Private/ProfilingDebugging/StallDetector.cpp
--- a/Engine/Source/Runtime/Core/Private/ProfilingDebugging/StallDetector.cpp
+++ b/Engine/Source/Runtime/Core/Private/ProfilingDebugging/StallDetector.cpp
@@ -142,12 +142,14 @@ UE::FStallDetectorStats::FStallDetectorStats(const TCHAR* InName, const double I
 	, ReportingMode(InReportingMode)
 	, bReported(false)
 	, TriggerCount(
+		TraceCounterNameType_Static,
 		(
 			FCString::Strcat(TriggerCountCounterName, TEXT("StallDetector/")),
 			FCString::Strcat(TriggerCountCounterName, InName),
 			FCString::Strcat(TriggerCountCounterName, TEXT(" TriggerCount"))
 		), TraceCounterDisplayHint_None)
 	, OverageSeconds(
+		TraceCounterNameType_Static,
 		(
 			FCString::Strcat(OverageSecondsCounterName, TEXT("StallDetector/")),
 			FCString::Strcat(OverageSecondsCounterName, InName),
diff --git a/Engine/Source/Runtime/Core/Public/ProfilingDebugging/CountersTrace.h b/Engine/Source/Runtime/Core/Public/ProfilingDebugging/CountersTrace.h
--- a/Engine/Source/Runtime/Core/Public/ProfilingDebugging/CountersTrace.h
+++ b/Engine/Source/Runtime/Core/Public/ProfilingDebugging/CountersTrace.h
@@ -22,14 +22,20 @@
 
 enum ETraceCounterType
 {
-	TraceCounterType_Int,
-	TraceCounterType_Float,
+	TraceCounterType_Int = 0,
+	TraceCounterType_Float = 1,
 };
 
 enum ETraceCounterDisplayHint
 {
-	TraceCounterDisplayHint_None,
-	TraceCounterDisplayHint_Memory,
+	TraceCounterDisplayHint_None = 0,
+	TraceCounterDisplayHint_Memory = 1,
+};
+
+enum ETraceCounterNameType
+{
+	TraceCounterNameType_Static = 0, // TCounter is allowed to keep a pointer to InCounterName string
+	TraceCounterNameType_Dynamic = 0x10, // TCounter needs to copy the InCounterName string
 };
 
 #if COUNTERSTRACE_ENABLED
@@ -42,25 +48,55 @@ struct FCountersTrace
 	CORE_API static void OutputSetValue(uint16 CounterId, int64 Value);
 	CORE_API static void OutputSetValue(uint16 CounterId, double Value);
 
+	CORE_API static const TCHAR* AllocAndCopyCounterName(const TCHAR* InCounterName);
+	CORE_API static void FreeCounterName(const TCHAR* InCounterName);
+
 	template<typename ValueType, ETraceCounterType CounterType, typename StoredType = ValueType>
 	class TCounter
 	{
 	public:
-		TCounter(const TCHAR* InCounterName, ETraceCounterDisplayHint InCounterDisplayHint)
+		TCounter() = delete;
+
+		template<int N>
+		TCounter(const TCHAR(&InCounterName)[N], ETraceCounterDisplayHint InCounterDisplayHint)
+			: Value(0)
+			, CounterName(InCounterName) // assumes that InCounterName is a static string
+			, CounterId(0)
+			, CounterDisplayHint(InCounterDisplayHint)
+		{
+			CounterId = OutputInitCounter(InCounterName, CounterType, CounterDisplayHint);
+		}
+
+		TCounter(ETraceCounterNameType InCounterNameType, const TCHAR* InCounterName, ETraceCounterDisplayHint InCounterDisplayHint)
 			: Value(0)
 			, CounterName(InCounterName)
 			, CounterId(0)
 			, CounterDisplayHint(InCounterDisplayHint)
 		{
-			EnsureInit();
+			CounterId = OutputInitCounter(InCounterName, CounterType, CounterDisplayHint);
+
+			if (CounterId == 0 && InCounterNameType == TraceCounterNameType_Dynamic)
+			{
+				// Store counter name for late init. Needs a copy as InCounterName pointer might not be valid later.
+				CounterName = AllocAndCopyCounterName(InCounterName);
+				CounterDisplayHint = ETraceCounterDisplayHint(uint8(CounterDisplayHint) | uint8(TraceCounterNameType_Dynamic));
+			}
+		}
+
+		~TCounter()
+		{
+			if (uint8(CounterDisplayHint) & uint8(TraceCounterNameType_Dynamic))
+			{
+				FreeCounterName(CounterName);
+			}
 		}
 
-		void EnsureInit()
+		void LateInit()
 		{
 			uint32 OldId = CounterId.load();
 			if (!OldId)
 			{
-				uint32 NewId = OutputInitCounter(CounterName, CounterType, CounterDisplayHint);
+				uint32 NewId = OutputInitCounter(CounterName, CounterType, ETraceCounterDisplayHint(uint8(CounterDisplayHint) & 0xF));
 				CounterId.compare_exchange_weak(OldId, NewId);
 			}
 		}
@@ -77,7 +113,7 @@ struct FCountersTrace
 				Value = InValue;
 				if (UE_TRACE_CHANNELEXPR_IS_ENABLED(CountersChannel))
 				{
-					EnsureInit();
+					LateInit();
 					OutputSetValue(uint16(CounterId), Value);
 				}
 			}
@@ -90,7 +126,7 @@ struct FCountersTrace
 				Value += InValue;
 				if (UE_TRACE_CHANNELEXPR_IS_ENABLED(CountersChannel))
 				{
-					EnsureInit();
+					LateInit();
 					OutputSetValue(uint16(CounterId), Value);
 				}
 			}
@@ -103,7 +139,7 @@ struct FCountersTrace
 				Value -= InValue;
 				if (UE_TRACE_CHANNELEXPR_IS_ENABLED(CountersChannel))
 				{
-					EnsureInit();
+					LateInit();
 					OutputSetValue(uint16(CounterId), Value);
 				}
 			}
@@ -114,17 +150,17 @@ struct FCountersTrace
 			++Value;
 			if (UE_TRACE_CHANNELEXPR_IS_ENABLED(CountersChannel))
 			{
-				EnsureInit();
+				LateInit();
 				OutputSetValue(uint16(CounterId), Value);
 			}
 		}
-		
+
 		void Decrement()
 		{
 			--Value;
 			if (UE_TRACE_CHANNELEXPR_IS_ENABLED(CountersChannel))
 			{
-				EnsureInit();
+				LateInit();
 				OutputSetValue(uint16(CounterId), Value);
 			}
 		}
@@ -141,7 +177,12 @@ struct FCountersTrace
 	using FCounterFloat = TCounter<double, TraceCounterType_Float>;
 };
 
+#define __TRACE_CHECK_COUNTER_NAME(CounterDisplayName) \
+	static_assert(std::is_const_v<std::remove_reference_t<decltype(CounterDisplayName)>>, "CounterDisplayName string must be a const TCHAR array."); \
+	static_assert(TIsArrayOrRefOfTypeByPredicate<decltype(CounterDisplayName), TIsCharEncodingCompatibleWithTCHAR>::Value, "CounterDisplayName string must be a TCHAR array.");
+
 #define __TRACE_DECLARE_INLINE_COUNTER(CounterDisplayName, CounterType, CounterDisplayHint) \
+	__TRACE_CHECK_COUNTER_NAME(CounterDisplayName) \
 	static FCountersTrace::CounterType PREPROCESSOR_JOIN(__TraceCounter, __LINE__)(CounterDisplayName, CounterDisplayHint);
 
 #define TRACE_INT_VALUE(CounterDisplayName, Value) \
@@ -157,9 +198,11 @@ struct FCountersTrace
 	PREPROCESSOR_JOIN(__TraceCounter, __LINE__).Set(Value);
 
 #define TRACE_DECLARE_INT_COUNTER(CounterName, CounterDisplayName) \
+	__TRACE_CHECK_COUNTER_NAME(CounterDisplayName) \
 	FCountersTrace::FCounterInt PREPROCESSOR_JOIN(__GTraceCounter, CounterName)(CounterDisplayName, TraceCounterDisplayHint_None);
 
 #define TRACE_DECLARE_ATOMIC_INT_COUNTER(CounterName, CounterDisplayName) \
+	__TRACE_CHECK_COUNTER_NAME(CounterDisplayName) \
 	FCountersTrace::FCounterAtomicInt PREPROCESSOR_JOIN(__GTraceCounter, CounterName)(CounterDisplayName, TraceCounterDisplayHint_None);
 
 #define TRACE_DECLARE_INT_COUNTER_EXTERN(CounterName) \
@@ -169,18 +212,21 @@ struct FCountersTrace
 	extern FCountersTrace::FCounterAtomicInt PREPROCESSOR_JOIN(__GTraceCounter, CounterName);
 
 #define TRACE_DECLARE_FLOAT_COUNTER(CounterName, CounterDisplayName) \
+	__TRACE_CHECK_COUNTER_NAME(CounterDisplayName) \
 	FCountersTrace::FCounterFloat PREPROCESSOR_JOIN(__GTraceCounter, CounterName)(CounterDisplayName, TraceCounterDisplayHint_None);
 
 #define TRACE_DECLARE_FLOAT_COUNTER_EXTERN(CounterName) \
 	extern FCountersTrace::FCounterFloat PREPROCESSOR_JOIN(__GTraceCounter, CounterName);
 
 #define TRACE_DECLARE_MEMORY_COUNTER(CounterName, CounterDisplayName) \
+	__TRACE_CHECK_COUNTER_NAME(CounterDisplayName) \
 	FCountersTrace::FCounterInt PREPROCESSOR_JOIN(__GTraceCounter, CounterName)(CounterDisplayName, TraceCounterDisplayHint_Memory);
 
 #define TRACE_DECLARE_MEMORY_COUNTER_EXTERN(CounterName) \
 	TRACE_DECLARE_INT_COUNTER_EXTERN(CounterName)
 
 #define TRACE_DECLARE_ATOMIC_MEMORY_COUNTER(CounterName, CounterDisplayName) \
+	__TRACE_CHECK_COUNTER_NAME(CounterDisplayName) \
 	FCountersTrace::FCounterAtomicInt PREPROCESSOR_JOIN(__GTraceCounter, CounterName)(CounterDisplayName, TraceCounterDisplayHint_Memory);
 
 #define TRACE_DECLARE_ATOMIC_MEMORY_COUNTER_EXTERN(CounterName) \
