Added comment about expected write size for FTraceUtils::Encode7bit

#rnx

diff --git a/Engine/Source/Runtime/Core/Public/ProfilingDebugging/MiscTrace.h b/Engine/Source/Runtime/Core/Public/ProfilingDebugging/MiscTrace.h
--- a/Engine/Source/Runtime/Core/Public/ProfilingDebugging/MiscTrace.h
+++ b/Engine/Source/Runtime/Core/Public/ProfilingDebugging/MiscTrace.h
@@ -29,6 +29,7 @@ struct FTraceUtils
 {
 	static void Encode7bit(uint64 Value, uint8*& BufferPtr)
 	{
+		// Writes 1 to 10 bytes for uint64 and 1 to 5 bytes for uint32.
 		do
 		{
 			uint8 HasMoreBytes = (uint8)((Value > uint64(0x7F)) << 7);
