Added FRecursiveMutex

This is an 8-byte unfair recursive mutex built on top of ParkingLot.
Provides the same API as FCriticalSection which allows it to be used as a drop-in replacement.

diff --git a/Engine/Source/Runtime/Core/Private/Async/RecursiveMutex.cpp b/Engine/Source/Runtime/Core/Private/Async/RecursiveMutex.cpp
new file mode 100644
--- /dev/null
+++ b/Engine/Source/Runtime/Core/Private/Async/RecursiveMutex.cpp
@@ -0,0 +1,166 @@
+// Copyright Epic Games, Inc. All Rights Reserved.
+
+#include "Async/RecursiveMutex.h"
+
+#include "Async/ParkingLot.h"
+#include "HAL/PlatformProcess.h"
+#include "HAL/PlatformTLS.h"
+
+namespace UE
+{
+
+bool FRecursiveMutex::TryLock()
+{
+	const uint32 CurrentThreadId = FPlatformTLS::GetCurrentThreadId();
+	FState CurrentState(State.load(std::memory_order_acquire));
+
+	// Try to acquire the lock if it was unlocked, even if there are waiting threads.
+	// Acquiring the lock despite the waiting threads means that this lock is not FIFO and thus not fair.
+	if (!CurrentState.bIsLocked)
+	{
+		checkSlow(ThreadId.load(std::memory_order_relaxed) == 0 && CurrentState.RecurseCount == 0);
+		FState NewState = CurrentState;
+		NewState.bIsLocked = true;
+		if (LIKELY(State.compare_exchange_strong(CurrentState.Value, NewState.Value, std::memory_order_acquire)))
+		{
+			ThreadId.store(CurrentThreadId, std::memory_order_relaxed);
+			return true;
+		}
+	}
+
+	if (LIKELY(ThreadId.load(std::memory_order_relaxed) != CurrentThreadId))
+	{
+		return false;
+	}
+
+	// Lock recursively if this is the thread that holds the lock.
+	FState AddState;
+	AddState.RecurseCount = 1;
+	State.fetch_add(AddState.Value, std::memory_order_relaxed);
+	return true;
+}
+
+void FRecursiveMutex::Lock()
+{
+	const uint32 CurrentThreadId = FPlatformTLS::GetCurrentThreadId();
+	FState CurrentState(State.load(std::memory_order_acquire));
+
+	// Try to acquire the lock if it was unlocked, even if there are waiting threads.
+	// Acquiring the lock despite the waiting threads means that this lock is not FIFO and thus not fair.
+	if (!CurrentState.bIsLocked)
+	{
+		checkSlow(ThreadId.load(std::memory_order_relaxed) == 0 && CurrentState.RecurseCount == 0);
+		FState NewState = CurrentState;
+		NewState.bIsLocked = true;
+		if (LIKELY(State.compare_exchange_weak(CurrentState.Value, NewState.Value, std::memory_order_acquire)))
+		{
+			ThreadId.store(CurrentThreadId, std::memory_order_relaxed);
+			return;
+		}
+	}
+
+	if (LIKELY(ThreadId.load(std::memory_order_relaxed) != CurrentThreadId))
+	{
+		return LockSlow(CurrentState, CurrentThreadId);
+	}
+
+	// Lock recursively if this is the thread that holds the lock.
+	FState AddState;
+	AddState.RecurseCount = 1;
+	State.fetch_add(AddState.Value, std::memory_order_relaxed);
+}
+
+FORCENOINLINE void FRecursiveMutex::LockSlow(FState CurrentState, const uint32 CurrentThreadId)
+{
+	constexpr int32 SpinLimit = 40;
+	int32 SpinCount = 0;
+
+	for (;;)
+	{
+		// Try to acquire the lock if it was unlocked, even if there are waiting threads.
+		// Acquiring the lock despite the waiting threads means that this lock is not FIFO and thus not fair.
+		if (!CurrentState.bIsLocked)
+		{
+			checkSlow(ThreadId.load(std::memory_order_relaxed) == 0 && CurrentState.RecurseCount == 0);
+			FState NewState = CurrentState;
+			NewState.bIsLocked = true;
+			if (LIKELY(State.compare_exchange_weak(CurrentState.Value, NewState.Value, std::memory_order_acquire)))
+			{
+				ThreadId.store(CurrentThreadId, std::memory_order_relaxed);
+				return;
+			}
+			continue;
+		}
+
+		// Spin up to the spin limit while there are no waiting threads.
+		if (!CurrentState.bHasWaitingThreads && SpinCount < SpinLimit)
+		{
+			FPlatformProcess::YieldThread();
+			++SpinCount;
+			CurrentState.Value = State.load(std::memory_order_acquire);
+			continue;
+		}
+
+		// Store that there are waiting threads. Restart if the state has changed since it was loaded.
+		if (!CurrentState.bHasWaitingThreads)
+		{
+			FState NewState = CurrentState;
+			NewState.bHasWaitingThreads = true;
+			if (!State.compare_exchange_weak(CurrentState.Value, NewState.Value, std::memory_order_acquire))
+			{
+				continue;
+			}
+			CurrentState = NewState;
+		}
+
+		// Wait if the state has not changed. Either way, loop back and try to acquire the lock after trying to wait.
+		ParkingLot::Wait(&State, [this, CurrentState] { return State.load(std::memory_order_acquire) == CurrentState.Value; }, []{});
+		CurrentState.Value = State.load(std::memory_order_acquire);
+	}
+}
+
+void FRecursiveMutex::Unlock()
+{
+	FState CurrentState(State.load(std::memory_order_relaxed));
+	checkSlow(CurrentState.bIsLocked);
+	checkSlow(ThreadId.load(std::memory_order_relaxed) == FPlatformTLS::GetCurrentThreadId());
+
+	if (LIKELY(CurrentState.RecurseCount == 0))
+	{
+		// Remove the association with this thread before unlocking.
+		ThreadId.store(0, std::memory_order_relaxed);
+
+		// Unlocking with no waiting threads only requires resetting the state.
+		if (!CurrentState.bHasWaitingThreads)
+		{
+			FState NewState;
+			if (LIKELY(State.compare_exchange_strong(CurrentState.Value, NewState.Value, std::memory_order_release, std::memory_order_relaxed)))
+			{
+				return;
+			}
+		}
+
+		UnlockSlow(CurrentState);
+		return;
+	}
+
+	// When locked recursively, decrement the count and return.
+	FState SubState;
+	SubState.RecurseCount = 1;
+	State.fetch_sub(SubState.Value, std::memory_order_relaxed);
+}
+
+FORCENOINLINE void FRecursiveMutex::UnlockSlow(FState CurrentState)
+{
+	// There is at least one thread waiting. Wake one thread and return.
+	ParkingLot::WakeOne(&State, [this](ParkingLot::FWakeState WakeState) -> uint64
+	{
+		FState NewState;
+		NewState.bHasWaitingThreads = WakeState.bHasWaitingThreads;
+		FState OldState(State.exchange(NewState.Value, std::memory_order_release));
+		checkSlow(OldState.bIsLocked && OldState.bHasWaitingThreads);
+		return 0;
+	});
+}
+
+} // UE
diff --git a/Engine/Source/Runtime/Core/Public/Async/RecursiveMutex.h b/Engine/Source/Runtime/Core/Public/Async/RecursiveMutex.h
new file mode 100644
--- /dev/null
+++ b/Engine/Source/Runtime/Core/Public/Async/RecursiveMutex.h
@@ -0,0 +1,60 @@
+// Copyright Epic Games, Inc. All Rights Reserved.
+
+#pragma once
+
+#include "CoreTypes.h"
+#include <atomic>
+
+#define UE_API CORE_API
+
+namespace UE
+{
+
+/**
+ * An eight-byte mutex that is not fair and supports recursive locking.
+ *
+ * Prefer FMutex when recursive locking is not required.
+ */
+class FRecursiveMutex final
+{
+public:
+	constexpr FRecursiveMutex() = default;
+
+	FRecursiveMutex(const FRecursiveMutex&) = delete;
+	FRecursiveMutex& operator=(const FRecursiveMutex&) = delete;
+
+	UE_API bool TryLock();
+	UE_API void Lock();
+	UE_API void Unlock();
+
+private:
+	union FState
+	{
+		struct
+		{
+			uint32 bIsLocked : 1;
+			uint32 bHasWaitingThreads : 1;
+			uint32 RecurseCount : 30;
+		};
+		uint32 Value = 0;
+
+		constexpr FState() = default;
+
+		constexpr explicit FState(uint32 State)
+			: Value(State)
+		{
+		}
+
+		constexpr bool operator==(const FState& Other) const { return Value == Other.Value; }
+	};
+
+	void LockSlow(FState CurrentState, const uint32 CurrentThreadId);
+	void UnlockSlow(FState CurrentState);
+
+	std::atomic<uint32> State = 0;
+	std::atomic<uint32> ThreadId = 0;
+};
+
+} // UE
+
+#undef UE_API
