Moved FRecursiveMutex implementation details out of the header

#rnx

diff --git a/Engine/Source/Runtime/Core/Private/Async/RecursiveMutex.cpp b/Engine/Source/Runtime/Core/Private/Async/RecursiveMutex.cpp
--- a/Engine/Source/Runtime/Core/Private/Async/RecursiveMutex.cpp
+++ b/Engine/Source/Runtime/Core/Private/Async/RecursiveMutex.cpp
@@ -9,6 +9,26 @@
 namespace UE
 {
 
+union FRecursiveMutex::FState
+{
+	struct
+	{
+		uint32 bIsLocked : 1;
+		uint32 bHasWaitingThreads : 1;
+		uint32 RecurseCount : 30;
+	};
+	uint32 Value = 0;
+
+	constexpr FState() = default;
+
+	constexpr explicit FState(uint32 State)
+		: Value(State)
+	{
+	}
+
+	constexpr bool operator==(const FState& Other) const { return Value == Other.Value; }
+};
+
 bool FRecursiveMutex::TryLock()
 {
 	const uint32 CurrentThreadId = FPlatformTLS::GetCurrentThreadId();
diff --git a/Engine/Source/Runtime/Core/Public/Async/RecursiveMutex.h b/Engine/Source/Runtime/Core/Public/Async/RecursiveMutex.h
--- a/Engine/Source/Runtime/Core/Public/Async/RecursiveMutex.h
+++ b/Engine/Source/Runtime/Core/Public/Async/RecursiveMutex.h
@@ -28,25 +28,7 @@ public:
 	UE_API void Unlock();
 
 private:
-	union FState
-	{
-		struct
-		{
-			uint32 bIsLocked : 1;
-			uint32 bHasWaitingThreads : 1;
-			uint32 RecurseCount : 30;
-		};
-		uint32 Value = 0;
-
-		constexpr FState() = default;
-
-		constexpr explicit FState(uint32 State)
-			: Value(State)
-		{
-		}
-
-		constexpr bool operator==(const FState& Other) const { return Value == Other.Value; }
-	};
+	union FState;
 
 	void LockSlow(FState CurrentState, const uint32 CurrentThreadId);
 	void UnlockSlow(FState CurrentState);
