Shared compiler-specific third-party include macros.

diff --git a/Engine/Plugins/Runtime/Database/SQLiteCore/Source/SQLiteCore/Private/SQLiteEmbedded.c b/Engine/Plugins/Runtime/Database/SQLiteCore/Source/SQLiteCore/Private/SQLiteEmbedded.c
--- a/Engine/Plugins/Runtime/Database/SQLiteCore/Source/SQLiteCore/Private/SQLiteEmbedded.c
+++ b/Engine/Plugins/Runtime/Database/SQLiteCore/Source/SQLiteCore/Private/SQLiteEmbedded.c
@@ -17,7 +17,7 @@
 #if !defined(USING_CODE_ANALYSIS) || !USING_CODE_ANALYSIS
 #include "HAL/PreprocessorHelpers.h"
 #include COMPILED_PLATFORM_HEADER(PlatformCompilerPreSetup.h)
-THIRD_PARTY_INCLUDES_START
+UE_COMPILER_THIRD_PARTY_INCLUDES_START
 #include "sqlite/sqlite3.inl"
-THIRD_PARTY_INCLUDES_END
+UE_COMPILER_THIRD_PARTY_INCLUDES_END
 #endif
diff --git a/Engine/Source/Runtime/Core/Public/Android/AndroidPlatformCompilerPreSetup.h b/Engine/Source/Runtime/Core/Public/Android/AndroidPlatformCompilerPreSetup.h
--- a/Engine/Source/Runtime/Core/Public/Android/AndroidPlatformCompilerPreSetup.h
+++ b/Engine/Source/Runtime/Core/Public/Android/AndroidPlatformCompilerPreSetup.h
@@ -4,31 +4,6 @@
 
 #include "Clang/ClangPlatformCompilerPreSetup.h"
 
-// Disable common CA warnings around SDK includes
-#ifndef THIRD_PARTY_INCLUDES_START
-	#define THIRD_PARTY_INCLUDES_START \
-		PRAGMA_DISABLE_REORDER_WARNINGS \
-		PRAGMA_DISABLE_SHADOW_VARIABLE_WARNINGS \
-		PRAGMA_DISABLE_UNDEFINED_IDENTIFIER_WARNINGS \
-		PRAGMA_DISABLE_MISSING_VIRTUAL_DESTRUCTOR_WARNINGS \
-		PRAGMA_DISABLE_DEPRECATION_WARNINGS \
-		PRAGMA_DISABLE_OVERLOADED_VIRTUAL_WARNINGS \
-		PRAGMA_DISABLE_UNINITIALIZED_CONST_REFERENCE_WARNINGS \
-		PRAGMA_DISABLE_ORDERED_COMPARE_FUNCTION_POINTERS
-#endif
-
-#ifndef THIRD_PARTY_INCLUDES_END
-	#define THIRD_PARTY_INCLUDES_END \
-		PRAGMA_ENABLE_UNINITIALIZED_CONST_REFERENCE_WARNINGS \
-		PRAGMA_ENABLE_OVERLOADED_VIRTUAL_WARNINGS \
-		PRAGMA_ENABLE_DEPRECATION_WARNINGS \
-		PRAGMA_ENABLE_MISSING_VIRTUAL_DESTRUCTOR_WARNINGS \
-		PRAGMA_ENABLE_UNDEFINED_IDENTIFIER_WARNINGS \
-		PRAGMA_ENABLE_SHADOW_VARIABLE_WARNINGS \
-		PRAGMA_ENABLE_REORDER_WARNINGS \
-		PRAGMA_ENABLE_ORDERED_COMPARE_FUNCTION_POINTERS
-#endif
-
 // Make certain warnings always be warnings, even despite -Werror.
 // Rationale: we don't want to suppress those as there are plans to address them (e.g. UE-12341), but breaking builds due to these warnings is very expensive
 // since they cannot be caught by all compilers that we support. They are deemed to be relatively safe to be ignored, at least until all SDKs/toolchains start supporting them.
diff --git a/Engine/Source/Runtime/Core/Public/Apple/ApplePlatformCompilerPreSetup.h b/Engine/Source/Runtime/Core/Public/Apple/ApplePlatformCompilerPreSetup.h
--- a/Engine/Source/Runtime/Core/Public/Apple/ApplePlatformCompilerPreSetup.h
+++ b/Engine/Source/Runtime/Core/Public/Apple/ApplePlatformCompilerPreSetup.h
@@ -4,31 +4,6 @@
 
 #include "Clang/ClangPlatformCompilerPreSetup.h"
 
-// Disable common CA warnings around SDK includes
-#ifndef THIRD_PARTY_INCLUDES_START
-	#define THIRD_PARTY_INCLUDES_START \
-		PRAGMA_DISABLE_REORDER_WARNINGS \
-		PRAGMA_DISABLE_SHADOW_VARIABLE_WARNINGS \
-		PRAGMA_DISABLE_UNDEFINED_IDENTIFIER_WARNINGS \
-		PRAGMA_DISABLE_MISSING_VIRTUAL_DESTRUCTOR_WARNINGS \
-		PRAGMA_DISABLE_DEPRECATION_WARNINGS \
-		PRAGMA_DISABLE_OVERLOADED_VIRTUAL_WARNINGS \
-		PRAGMA_DISABLE_UNINITIALIZED_CONST_REFERENCE_WARNINGS \
-		PRAGMA_DISABLE_ORDERED_COMPARE_FUNCTION_POINTERS
-#endif
-
-#ifndef THIRD_PARTY_INCLUDES_END
-	#define THIRD_PARTY_INCLUDES_END \
-		PRAGMA_ENABLE_UNINITIALIZED_CONST_REFERENCE_WARNINGS \
-		PRAGMA_ENABLE_OVERLOADED_VIRTUAL_WARNINGS \
-		PRAGMA_ENABLE_DEPRECATION_WARNINGS \
-		PRAGMA_ENABLE_MISSING_VIRTUAL_DESTRUCTOR_WARNINGS \
-		PRAGMA_ENABLE_UNDEFINED_IDENTIFIER_WARNINGS \
-		PRAGMA_ENABLE_SHADOW_VARIABLE_WARNINGS \
-		PRAGMA_ENABLE_REORDER_WARNINGS \
-		PRAGMA_ENABLE_ORDERED_COMPARE_FUNCTION_POINTERS
-#endif
-
 // Make certain warnings always be warnings, even despite -Werror.
 // Rationale: we don't want to suppress those as there are plans to address them (e.g. UE-12341), but breaking builds due to these warnings is very expensive
 // since they cannot be caught by all compilers that we support. They are deemed to be relatively safe to be ignored, at least until all SDKs/toolchains start supporting them.
diff --git a/Engine/Source/Runtime/Core/Public/Clang/ClangPlatformCompilerPreSetup.h b/Engine/Source/Runtime/Core/Public/Clang/ClangPlatformCompilerPreSetup.h
--- a/Engine/Source/Runtime/Core/Public/Clang/ClangPlatformCompilerPreSetup.h
+++ b/Engine/Source/Runtime/Core/Public/Clang/ClangPlatformCompilerPreSetup.h
@@ -4,6 +4,8 @@
 
 // HEADER_UNIT_UNSUPPORTED - Clang not supporting header units
 
+#ifdef __clang__
+
 #pragma clang diagnostic warning "-Wdeprecated-declarations"
 
 #ifndef PRAGMA_DISABLE_DEPRECATION_WARNINGS
@@ -236,3 +238,31 @@
 	#define EMIT_CUSTOM_WARNING_AT_LINE(Line, Warning) \
 		_Pragma(PREPROCESSOR_TO_STRING(message(Warning)))
 #endif // EMIT_CUSTOM_WARNING_AT_LINE
+
+#ifndef UE_COMPILER_THIRD_PARTY_INCLUDES_START
+	#define UE_COMPILER_THIRD_PARTY_INCLUDES_START \
+		PRAGMA_DISABLE_REORDER_WARNINGS \
+		PRAGMA_DISABLE_SHADOW_VARIABLE_WARNINGS \
+		PRAGMA_DISABLE_UNDEFINED_IDENTIFIER_WARNINGS \
+		PRAGMA_DISABLE_MISSING_VIRTUAL_DESTRUCTOR_WARNINGS \
+		PRAGMA_DISABLE_DEPRECATION_WARNINGS \
+		PRAGMA_DISABLE_OVERLOADED_VIRTUAL_WARNINGS \
+		PRAGMA_DISABLE_MISSING_BRACES_WARNINGS \
+		PRAGMA_DISABLE_UNINITIALIZED_CONST_REFERENCE_WARNINGS \
+		PRAGMA_DISABLE_ORDERED_COMPARE_FUNCTION_POINTERS
+#endif
+
+#ifndef UE_COMPILER_THIRD_PARTY_INCLUDES_END
+	#define UE_COMPILER_THIRD_PARTY_INCLUDES_END \
+		PRAGMA_ENABLE_ORDERED_COMPARE_FUNCTION_POINTERS \
+		PRAGMA_ENABLE_UNINITIALIZED_CONST_REFERENCE_WARNINGS \
+		PRAGMA_ENABLE_MISSING_BRACES_WARNINGS \
+		PRAGMA_ENABLE_OVERLOADED_VIRTUAL_WARNINGS \
+		PRAGMA_ENABLE_DEPRECATION_WARNINGS \
+		PRAGMA_ENABLE_MISSING_VIRTUAL_DESTRUCTOR_WARNINGS \
+		PRAGMA_ENABLE_UNDEFINED_IDENTIFIER_WARNINGS \
+		PRAGMA_ENABLE_SHADOW_VARIABLE_WARNINGS \
+		PRAGMA_ENABLE_REORDER_WARNINGS
+#endif
+
+#endif // __clang__
diff --git a/Engine/Source/Runtime/Core/Public/GenericPlatform/GenericPlatformCompilerPreSetup.h b/Engine/Source/Runtime/Core/Public/GenericPlatform/GenericPlatformCompilerPreSetup.h
--- a/Engine/Source/Runtime/Core/Public/GenericPlatform/GenericPlatformCompilerPreSetup.h
+++ b/Engine/Source/Runtime/Core/Public/GenericPlatform/GenericPlatformCompilerPreSetup.h
@@ -43,4 +43,13 @@
 
 #ifndef PRAGMA_ENABLE_BUFFER_OVERRUN_WARNING
 	#define PRAGMA_ENABLE_BUFFER_OVERRUN_WARNING
-#endif
\ No newline at end of file
+#endif
+
+// Disable CA warnings around SDK includes - default to those disabled at the compiler level
+#ifndef THIRD_PARTY_INCLUDES_START
+	#define THIRD_PARTY_INCLUDES_START UE_COMPILER_THIRD_PARTY_INCLUDES_START
+#endif
+
+#ifndef THIRD_PARTY_INCLUDES_END
+	#define THIRD_PARTY_INCLUDES_END UE_COMPILER_THIRD_PARTY_INCLUDES_END
+#endif
diff --git a/Engine/Source/Runtime/Core/Public/MSVC/MSVCPlatformCompilerPreSetup.h b/Engine/Source/Runtime/Core/Public/MSVC/MSVCPlatformCompilerPreSetup.h
--- a/Engine/Source/Runtime/Core/Public/MSVC/MSVCPlatformCompilerPreSetup.h
+++ b/Engine/Source/Runtime/Core/Public/MSVC/MSVCPlatformCompilerPreSetup.h
@@ -27,31 +27,6 @@
 	#endif // PRAGMA_POP_PLATFORM_DEFAULT_PACKING
 
 	#include "Clang/ClangPlatformCompilerPreSetup.h"
-
-	// Disable common CA warnings around SDK includes
-	#ifndef THIRD_PARTY_INCLUDES_START
-		#define THIRD_PARTY_INCLUDES_START \
-			PRAGMA_DISABLE_REORDER_WARNINGS \
-			PRAGMA_DISABLE_SHADOW_VARIABLE_WARNINGS \
-			PRAGMA_DISABLE_UNDEFINED_IDENTIFIER_WARNINGS \
-			PRAGMA_DISABLE_DEPRECATION_WARNINGS \
-			PRAGMA_DISABLE_OVERLOADED_VIRTUAL_WARNINGS \
-			PRAGMA_DISABLE_MISSING_BRACES_WARNINGS \
-			PRAGMA_DISABLE_UNINITIALIZED_CONST_REFERENCE_WARNINGS \
-			PRAGMA_DISABLE_ORDERED_COMPARE_FUNCTION_POINTERS
-	#endif // THIRD_PARTY_INCLUDES_START
-
-	#ifndef THIRD_PARTY_INCLUDES_END
-		#define THIRD_PARTY_INCLUDES_END \
-			PRAGMA_ENABLE_UNINITIALIZED_CONST_REFERENCE_WARNINGS \
-			PRAGMA_ENABLE_MISSING_BRACES_WARNINGS \
-			PRAGMA_ENABLE_OVERLOADED_VIRTUAL_WARNINGS \
-			PRAGMA_ENABLE_DEPRECATION_WARNINGS \
-			PRAGMA_ENABLE_UNDEFINED_IDENTIFIER_WARNINGS \
-			PRAGMA_ENABLE_SHADOW_VARIABLE_WARNINGS \
-			PRAGMA_ENABLE_REORDER_WARNINGS \
-			PRAGMA_ENABLE_ORDERED_COMPARE_FUNCTION_POINTERS
-	#endif // THIRD_PARTY_INCLUDES_END
 #else
 	#ifndef PRAGMA_DISABLE_DEPRECATION_WARNINGS
 		#define PRAGMA_DISABLE_DEPRECATION_WARNINGS \
@@ -214,8 +189,8 @@
 	#endif // PRAGMA_POP_PLATFORM_DEFAULT_PACKING
 
 	// Disable common CA warnings around SDK includes
-	#ifndef THIRD_PARTY_INCLUDES_START
-		#define THIRD_PARTY_INCLUDES_START \
+	#ifndef UE_COMPILER_THIRD_PARTY_INCLUDES_START
+		#define UE_COMPILER_THIRD_PARTY_INCLUDES_START \
 			__pragma(warning(push)) \
 			__pragma(warning(disable: 4125))  /* decimal digit terminates octal escape sequence. */ \
 			__pragma(warning(disable: 4510))  /* '<class>': default constructor could not be generated. */ \
@@ -245,8 +220,8 @@
 			PRAGMA_DISABLE_UNINITIALIZED_CONST_REFERENCE_WARNINGS
 	#endif // THIRD_PARTY_INCLUDES_START
 
-	#ifndef THIRD_PARTY_INCLUDES_END
-		#define THIRD_PARTY_INCLUDES_END \
+	#ifndef UE_COMPILER_THIRD_PARTY_INCLUDES_END
+		#define UE_COMPILER_THIRD_PARTY_INCLUDES_END \
 			PRAGMA_ENABLE_UNINITIALIZED_CONST_REFERENCE_WARNINGS \
 			PRAGMA_ENABLE_DEPRECATION_WARNINGS \
 			PRAGMA_ENABLE_MISSING_VIRTUAL_DESTRUCTOR_WARNINGS \
diff --git a/Engine/Source/Runtime/Core/Public/Unix/UnixPlatformCompilerPreSetup.h b/Engine/Source/Runtime/Core/Public/Unix/UnixPlatformCompilerPreSetup.h
--- a/Engine/Source/Runtime/Core/Public/Unix/UnixPlatformCompilerPreSetup.h
+++ b/Engine/Source/Runtime/Core/Public/Unix/UnixPlatformCompilerPreSetup.h
@@ -31,35 +31,21 @@
 			_Pragma("clang diagnostic pop")
 #endif // PRAGMA_ENABLE_NULL_DEREFERENCE_WARNINGS
 
-// Disable common CA warnings around SDK includes
+// Disable platform-specific CA warnings around SDK includes in addition to common Clang warnings
 #ifndef THIRD_PARTY_INCLUDES_START
 	#define THIRD_PARTY_INCLUDES_START \
-		PRAGMA_DISABLE_REORDER_WARNINGS \
-		PRAGMA_DISABLE_SHADOW_VARIABLE_WARNINGS \
+		UE_COMPILER_THIRD_PARTY_INCLUDES_START \
 		PRAGMA_DISABLE_NULL_DEREFERENCE_WARNINGS \
-		PRAGMA_DISABLE_UNDEFINED_IDENTIFIER_WARNINGS \
-		PRAGMA_DISABLE_MISSING_VIRTUAL_DESTRUCTOR_WARNINGS \
 		PRAGMA_DISABLE_DEPRECATED_REGISTER_WARNINGS \
-		PRAGMA_DISABLE_DEPRECATION_WARNINGS \
-		PRAGMA_DISABLE_OVERLOADED_VIRTUAL_WARNINGS \
-		PRAGMA_DISABLE_MACRO_REDEFINED_WARNINGS \
-		PRAGMA_DISABLE_UNINITIALIZED_CONST_REFERENCE_WARNINGS \
-		PRAGMA_DISABLE_ORDERED_COMPARE_FUNCTION_POINTERS
+		PRAGMA_DISABLE_MACRO_REDEFINED_WARNINGS
 #endif
 
 #ifndef THIRD_PARTY_INCLUDES_END
 	#define THIRD_PARTY_INCLUDES_END \
-		PRAGMA_ENABLE_UNINITIALIZED_CONST_REFERENCE_WARNINGS \
 		PRAGMA_ENABLE_MACRO_REDEFINED_WARNINGS \
-		PRAGMA_ENABLE_OVERLOADED_VIRTUAL_WARNINGS \
-		PRAGMA_ENABLE_DEPRECATION_WARNINGS \
 		PRAGMA_ENABLE_DEPRECATED_REGISTER_WARNINGS \
-		PRAGMA_ENABLE_MISSING_VIRTUAL_DESTRUCTOR_WARNINGS \
-		PRAGMA_ENABLE_UNDEFINED_IDENTIFIER_WARNINGS \
 		PRAGMA_ENABLE_NULL_DEREFERENCE_WARNINGS \
-		PRAGMA_ENABLE_SHADOW_VARIABLE_WARNINGS \
-		PRAGMA_ENABLE_REORDER_WARNINGS \
-		PRAGMA_ENABLE_ORDERED_COMPARE_FUNCTION_POINTERS
+		UE_COMPILER_THIRD_PARTY_INCLUDES_END
 #endif
 
 // Make certain warnings always be warnings, even despite -Werror.
