Fixed Clang 16 warning about FGenericPlatformMemoryConstants having a copy constructor but no assignment operator.

diff --git a/Engine/Source/Runtime/Core/Public/GenericPlatform/GenericPlatformMemory.h b/Engine/Source/Runtime/Core/Public/GenericPlatform/GenericPlatformMemory.h
--- a/Engine/Source/Runtime/Core/Public/GenericPlatform/GenericPlatformMemory.h
+++ b/Engine/Source/Runtime/Core/Public/GenericPlatform/GenericPlatformMemory.h
@@ -86,57 +86,33 @@ enum class EMemcpyCachePolicy : uint8
 struct FGenericPlatformMemoryConstants
 {
 	/** The amount of actual physical memory, in bytes (needs to handle >4GB for 64-bit devices running 32-bit code). */
-	uint64 TotalPhysical;
+	uint64 TotalPhysical = 0;
 
 	/** The amount of virtual memory, in bytes. */
-	uint64 TotalVirtual;
+	uint64 TotalVirtual = 0;
 
 	/** The size of a physical page, in bytes. This is also the granularity for PageProtect(), commitment and properties (e.g. ability to access) of the physical RAM. */
-	SIZE_T PageSize;
+	SIZE_T PageSize = 0;
 
 	/**
 	 * Some platforms have advantages if memory is allocated in chunks larger than PageSize (e.g. VirtualAlloc() seems to have 64KB granularity as of now).
      * This value is the minimum allocation size that the system will use behind the scenes.
 	 */
-	SIZE_T OsAllocationGranularity;
+	SIZE_T OsAllocationGranularity = 0;
 
 	/** The size of a "page" in Binned2 malloc terms, in bytes. Should be at least 64KB. BinnedMalloc expects memory returned from BinnedAllocFromOS() to be aligned on BinnedPageSize boundary. */
-	SIZE_T BinnedPageSize;
+	SIZE_T BinnedPageSize = 0;
 
 	/** This is the "allocation granularity" in Binned malloc terms, i.e. BinnedMalloc will allocate the memory in increments of this value. If zero, Binned will use BinnedPageSize for this value. */
-	SIZE_T BinnedAllocationGranularity;
+	SIZE_T BinnedAllocationGranularity = 0;
 
 	// AddressLimit - Second parameter is estimate of the range of addresses expected to be returns by BinnedAllocFromOS(). Binned
 	// Malloc will adjust its internal structures to make lookups for memory allocations O(1) for this range. 
 	// It is ok to go outside this range, lookups will just be a little slower
-	uint64 AddressLimit;
+	uint64 AddressLimit = (uint64)0xffffffff + 1;
 
 	/** Approximate physical RAM in GB; 1 on everything except PC. Used for "course tuning", like FPlatformMisc::NumberOfCores(). */
-	uint32 TotalPhysicalGB;
-
-	/** Default constructor, clears all variables. */
-	FGenericPlatformMemoryConstants()
-		: TotalPhysical( 0 )
-		, TotalVirtual( 0 )
-		, PageSize( 0 )
-		, OsAllocationGranularity(0)
-		, BinnedPageSize( 0 )
-		, BinnedAllocationGranularity( 0 )
-		, AddressLimit((uint64)0xffffffff + 1)
-		, TotalPhysicalGB( 1 )
-	{}
-
-	/** Copy constructor, used by the generic platform memory stats. */
-	FGenericPlatformMemoryConstants( const FGenericPlatformMemoryConstants& Other )
-		: TotalPhysical( Other.TotalPhysical )
-		, TotalVirtual( Other.TotalVirtual )
-		, PageSize( Other.PageSize )
-		, OsAllocationGranularity(Other.OsAllocationGranularity)
-		, BinnedPageSize(Other.BinnedPageSize)
-		, BinnedAllocationGranularity(Other.BinnedAllocationGranularity)
-		, AddressLimit(Other.AddressLimit)
-		, TotalPhysicalGB(Other.TotalPhysicalGB)
-	{}
+	uint32 TotalPhysicalGB = 1;
 };
 
 typedef FGenericPlatformMemoryConstants FPlatformMemoryConstants;
