PR #9792: Fix crash in linux event handling in FramePro

diff --git a/Engine/Source/Runtime/Core/Private/FramePro/FramePro.cpp b/Engine/Source/Runtime/Core/Private/FramePro/FramePro.cpp
--- a/Engine/Source/Runtime/Core/Private/FramePro/FramePro.cpp
+++ b/Engine/Source/Runtime/Core/Private/FramePro/FramePro.cpp
@@ -9666,7 +9666,8 @@ namespace FramePro
 
 				if (linux_event.m_Signalled)
 				{
-					linux_event.m_Signalled = false;
+					if (linux_event.m_AutoReset)
+						linux_event.m_Signalled = false;
 					pthread_mutex_unlock(&linux_event.m_Mutex);
 					return true;
 				}
@@ -9676,7 +9677,7 @@ namespace FramePro
 					while (!linux_event.m_Signalled)
 						pthread_cond_wait(&linux_event.m_Cond, &linux_event.m_Mutex);
 
-					if (!linux_event.m_AutoReset)
+					if (linux_event.m_AutoReset)
 						linux_event.m_Signalled = false;
 
 					pthread_mutex_unlock(&linux_event.m_Mutex);
@@ -9704,7 +9705,7 @@ namespace FramePro
 
 					if (linux_event.m_Signalled)
 					{
-						if (!linux_event.m_AutoReset)
+						if (linux_event.m_AutoReset)
 							linux_event.m_Signalled = false;
 
 						pthread_mutex_unlock(&linux_event.m_Mutex);
