Improve code around fixing universal binary symbolication

diff --git a/Engine/Source/Developer/CrashDebugHelper/Private/IOS/CrashDebugHelperIOS.cpp b/Engine/Source/Developer/CrashDebugHelper/Private/IOS/CrashDebugHelperIOS.cpp
--- a/Engine/Source/Developer/CrashDebugHelper/Private/IOS/CrashDebugHelperIOS.cpp
+++ b/Engine/Source/Developer/CrashDebugHelper/Private/IOS/CrashDebugHelperIOS.cpp
@@ -404,7 +404,7 @@ static int32 SymboliseStackInfo(FPlatformSymbolDatabaseSet& SymbolCache, TArray<
 	if(!Db)
 	{
 		FApplePlatformSymbolDatabase Database;
-		if(FPlatformSymbolication::LoadSymbolDatabaseForBinary(TEXT(""), Module.Name, Module.Report, Database))
+		if(FPlatformSymbolication::LoadSymbolDatabaseForBinary(TEXT(""), Module.Name, Module.Report, {}, Database))
 		{
 			SymbolCache.Add(Database);
 			Db = SymbolCache.Find(Module.Report);
@@ -780,4 +780,4 @@ bool FCrashDebugHelperIOS::CreateMinidumpDiagnosticReport( const FString& InCras
 	}
 	
 	return bOK;
-}
\ No newline at end of file
+}
diff --git a/Engine/Source/Developer/CrashDebugHelper/Private/Mac/CrashDebugHelperMac.cpp b/Engine/Source/Developer/CrashDebugHelper/Private/Mac/CrashDebugHelperMac.cpp
--- a/Engine/Source/Developer/CrashDebugHelper/Private/Mac/CrashDebugHelperMac.cpp
+++ b/Engine/Source/Developer/CrashDebugHelper/Private/Mac/CrashDebugHelperMac.cpp
@@ -427,8 +427,7 @@ static int32 SymboliseStackInfo(FPlatformSymbolDatabaseSet& SymbolCache, TArray<
 	if(!Db)
 	{
 		FApplePlatformSymbolDatabase Database;
-		Database.Architecture = Architecture;
-		if(FPlatformSymbolication::LoadSymbolDatabaseForBinary(TEXT(""), Module.Name, Module.Report, Database))
+		if(FPlatformSymbolication::LoadSymbolDatabaseForBinary(TEXT(""), Module.Name, Module.Report, Architecture, Database))
 		{
 			SymbolCache.Add(Database);
 			Db = SymbolCache.Find(Module.Report);
diff --git a/Engine/Source/Programs/Mac/DsymExporter/Private/DsymExporterApp.cpp b/Engine/Source/Programs/Mac/DsymExporter/Private/DsymExporterApp.cpp
--- a/Engine/Source/Programs/Mac/DsymExporter/Private/DsymExporterApp.cpp
+++ b/Engine/Source/Programs/Mac/DsymExporter/Private/DsymExporterApp.cpp
@@ -37,7 +37,7 @@ int32 RunDsymExporter(int32 ArgC, TCHAR* Argv[])
         bHaveSigArg = true;
     }
 	FString Signature = SigArg.RightChop(6);
-	bool bOK = FPlatformSymbolication::LoadSymbolDatabaseForBinary(TEXT(""), PathArg, Signature, Symbols);
+	bool bOK = FPlatformSymbolication::LoadSymbolDatabaseForBinary(TEXT(""), PathArg, Signature, {}, Symbols);
 	if(bOK)
 	{
 		FString OutputFolder = FPaths::GetPath(Argv[1]);
diff --git a/Engine/Source/Programs/Mac/UnrealAtoS/Private/UnrealAtoSApp.cpp b/Engine/Source/Programs/Mac/UnrealAtoS/Private/UnrealAtoSApp.cpp
--- a/Engine/Source/Programs/Mac/UnrealAtoS/Private/UnrealAtoSApp.cpp
+++ b/Engine/Source/Programs/Mac/UnrealAtoS/Private/UnrealAtoSApp.cpp
@@ -83,7 +83,7 @@ int32 RunUnrealAtoS(int32 ArgC, TCHAR* Argv[])
 		
 		FString ModuleSignature = ModuleSig;
 		FPlatformSymbolDatabase SymbolDb;
-		if(FPlatformSymbolication::LoadSymbolDatabaseForBinary(FString(SymbolPath), ModulePath, ModuleSignature, SymbolDb))
+		if(FPlatformSymbolication::LoadSymbolDatabaseForBinary(FString(SymbolPath), ModulePath, ModuleSignature, {}, SymbolDb))
 		{
 			bOK = FPlatformSymbolication::SymbolInfoForStrippedSymbol(SymbolDb, ProgramCounter, ModuleBase, ModuleSignature, Info);
 		}
diff --git a/Engine/Source/Runtime/Core/Private/Apple/ApplePlatformSymbolication.cpp b/Engine/Source/Runtime/Core/Private/Apple/ApplePlatformSymbolication.cpp
--- a/Engine/Source/Runtime/Core/Private/Apple/ApplePlatformSymbolication.cpp
+++ b/Engine/Source/Runtime/Core/Private/Apple/ApplePlatformSymbolication.cpp
@@ -237,17 +237,17 @@ void FApplePlatformSymbolication::EnableCoreSymbolication(bool const bEnable)
 	}
 }
 
-bool FApplePlatformSymbolication::LoadSymbolDatabaseForBinary(FString SourceFolder, FString BinaryPath, FString BinarySignature, FApplePlatformSymbolDatabase& OutDatabase)
+bool FApplePlatformSymbolication::LoadSymbolDatabaseForBinary(FString SourceFolder, FString BinaryPath, FString BinarySignature, TOptional<FString> Architecture, FApplePlatformSymbolDatabase& OutDatabase)
 {
-	bool bOK = FGenericPlatformSymbolication::LoadSymbolDatabaseForBinary(SourceFolder, BinaryPath, BinarySignature, *OutDatabase.GenericDB);
+	bool bOK = FGenericPlatformSymbolication::LoadSymbolDatabaseForBinary(SourceFolder, BinaryPath, BinarySignature, Architecture, *OutDatabase.GenericDB);
 	if(!bOK && GAllowApplePlatformSymbolication && (IFileManager::Get().FileSize(*BinaryPath) > 0))
 	{
 		CSSymbolicatorRef Symbolicator = OutDatabase.AppleDB;
-		if( CSIsNull(Symbolicator) && (OutDatabase.Architecture == TEXT("") || OutDatabase.Architecture == TEXT("X86_64")))
+		if( CSIsNull(Symbolicator) && (!Architecture.IsSet() || *Architecture == TEXT("X86_64")))
 		{
 			Symbolicator = CSSymbolicatorCreateWithPathAndArchitecture(TCHAR_TO_UTF8(*BinaryPath), CPU_TYPE_X86_64);
 		}
-        if (CSIsNull(Symbolicator) && (OutDatabase.Architecture == TEXT("") || OutDatabase.Architecture == TEXT("ARM-64")))
+        if (CSIsNull(Symbolicator) && (!Architecture.IsSet() || *Architecture == TEXT("ARM-64")))
         {
             Symbolicator = CSSymbolicatorCreateWithPathAndArchitecture(TCHAR_TO_UTF8(*BinaryPath), CPU_TYPE_ARM64);
         }
diff --git a/Engine/Source/Runtime/Core/Private/GenericPlatform/GenericPlatformSymbolication.cpp b/Engine/Source/Runtime/Core/Private/GenericPlatform/GenericPlatformSymbolication.cpp
--- a/Engine/Source/Runtime/Core/Private/GenericPlatform/GenericPlatformSymbolication.cpp
+++ b/Engine/Source/Runtime/Core/Private/GenericPlatform/GenericPlatformSymbolication.cpp
@@ -12,7 +12,7 @@
 
 DEFINE_LOG_CATEGORY_STATIC(LogGenericPlatformSymbolication, Log, All);
 
-bool FGenericPlatformSymbolication::LoadSymbolDatabaseForBinary(FString SourceFolder, FString Binary, FString ModuleSignature, FGenericPlatformSymbolDatabase& OutDatabase)
+bool FGenericPlatformSymbolication::LoadSymbolDatabaseForBinary(FString SourceFolder, FString Binary, FString ModuleSignature, TOptional<FString> Architecture, FGenericPlatformSymbolDatabase& OutDatabase)
 {
 	bool bOk = false;
 	
diff --git a/Engine/Source/Runtime/Core/Public/Apple/ApplePlatformSymbolication.h b/Engine/Source/Runtime/Core/Public/Apple/ApplePlatformSymbolication.h
--- a/Engine/Source/Runtime/Core/Public/Apple/ApplePlatformSymbolication.h
+++ b/Engine/Source/Runtime/Core/Public/Apple/ApplePlatformSymbolication.h
@@ -32,7 +32,6 @@ struct FApplePlatformSymbolDatabase
 	
 	TSharedPtr<FGenericPlatformSymbolDatabase> GenericDB;
 	FApplePlatformSymbolCache AppleDB;
-	FString Architecture;
 };
 
 /**
@@ -76,7 +75,10 @@ struct CORE_API FApplePlatformSymbolication
 {
 	static void EnableCoreSymbolication(bool const bEnable);
 	
-	static bool LoadSymbolDatabaseForBinary(FString SourceFolder, FString Binary, FString BinarySignature, FApplePlatformSymbolDatabase& OutDatabase);
+	/**
+	 * @param Architecture Specify the architecture in case of universal binary, if no architecture is given, will try X86_64 then ARM-64
+	 */
+	static bool LoadSymbolDatabaseForBinary(FString SourceFolder, FString Binary, FString BinarySignature, TOptional<FString> Architecture, FApplePlatformSymbolDatabase& OutDatabase);
 	static bool SaveSymbolDatabaseForBinary(FString TargetFolder, FString Name, FString BinarySignature, FApplePlatformSymbolDatabase& Database);
 	
 	static bool SymbolInfoForStrippedSymbol(FApplePlatformSymbolDatabase const& Database, uint64 ProgramCounter, uint64 ModuleOffset, FString ModuleSignature, FProgramCounterSymbolInfo& Info);
diff --git a/Engine/Source/Runtime/Core/Public/GenericPlatform/GenericPlatformSymbolication.h b/Engine/Source/Runtime/Core/Public/GenericPlatform/GenericPlatformSymbolication.h
--- a/Engine/Source/Runtime/Core/Public/GenericPlatform/GenericPlatformSymbolication.h
+++ b/Engine/Source/Runtime/Core/Public/GenericPlatform/GenericPlatformSymbolication.h
@@ -118,7 +118,10 @@ typedef TSet<FGenericPlatformSymbolDatabase, FGenericPlatformSymbolDatabaseKeyFu
 
 struct CORE_API FGenericPlatformSymbolication
 {
-	static bool LoadSymbolDatabaseForBinary(FString SourceFolder, FString BinaryPath, FString BinarySignature, FGenericPlatformSymbolDatabase& OutDatabase);
+	/**
+	 * @param Architecture Specify the architecture in case of universal binary, currently unused
+	 */
+	static bool LoadSymbolDatabaseForBinary(FString SourceFolder, FString BinaryPath, FString BinarySignature, TOptional<FString> Architecture, FGenericPlatformSymbolDatabase& OutDatabase);
 	static bool SaveSymbolDatabaseForBinary(FString TargetFolder, FString Name, FGenericPlatformSymbolDatabase& Database);
 	
 	static bool SymbolInfoForStrippedSymbol(FGenericPlatformSymbolDatabase const& Database, uint64 ProgramCounter, uint64 ModuleOffset, FString ModuleSignature, FProgramCounterSymbolInfo& Info);
