Changed ENullTerminatedString to an enum class so its No enumerator doesn't clash with the global vc_attributes::No.

diff --git a/Engine/Source/Runtime/Core/Public/Containers/StringConv.h b/Engine/Source/Runtime/Core/Public/Containers/StringConv.h
--- a/Engine/Source/Runtime/Core/Public/Containers/StringConv.h
+++ b/Engine/Source/Runtime/Core/Public/Containers/StringConv.h
@@ -328,6 +328,12 @@ namespace UE::Core::Private
 			return FPlatformString::ConvertedLength<IntendedToType>((const FromType*)Source, SourceLen);
 		}
 	};
+
+	enum class ENullTerminatedString
+	{
+		No  = 0,
+		Yes = 1
+	};
 }
 
 using FTCHARToUTF8_Convert /*UE_DEPRECATED(5.1, "FTCHARToUTF8_Convert has been deprecated in favor of FPlatformString::Convert and StringCast")*/ = UE::Core::Private::FTCHARToUTF8_Convert;
@@ -631,15 +637,6 @@ private:
 	}
 };
 
-struct ENullTerminatedString
-{
-	enum Type
-	{
-		No  = 0,
-		Yes = 1
-	};
-};
-
 /**
  * Class takes one type of string and converts it to another. The class includes a
  * chunk of presized memory of the destination type. If the presized array is
@@ -659,16 +656,18 @@ class TStringConversion : private Converter, private TInlineAllocator<DefaultCon
 	/**
 	 * Converts the data by using the Convert() method on the base class
 	 */
-	void Init(const FromType* Source, int32 SourceLen, ENullTerminatedString::Type NullTerminated)
+	void Init(const FromType* Source, int32 SourceLen, UE::Core::Private::ENullTerminatedString NullTerminated)
 	{
+		int32 NullOffset = (NullTerminated == UE::Core::Private::ENullTerminatedString::Yes) ? 1 : 0;
+
 		StringLength = Converter::ConvertedLength(Source, SourceLen);
 
-		int32 BufferSize = StringLength + NullTerminated;
+		int32 BufferSize = StringLength + NullOffset;
 
 		AllocatorType::ResizeAllocation(0, BufferSize, sizeof(ToType));
 
 		Ptr = (ToType*)AllocatorType::GetAllocation();
-		Converter::Convert(Ptr, BufferSize, Source, SourceLen + NullTerminated);
+		Converter::Convert(Ptr, BufferSize, Source, SourceLen + NullOffset);
 	}
 
 public:
@@ -686,7 +685,7 @@ public:
 	{
 		if (Source)
 		{
-			Init((const FromType*)Source, TCString<FromType>::Strlen((const FromType*)Source), ENullTerminatedString::Yes);
+			Init((const FromType*)Source, TCString<FromType>::Strlen((const FromType*)Source), UE::Core::Private::ENullTerminatedString::Yes);
 		}
 		else
 		{
@@ -707,11 +706,11 @@ public:
 	{
 		if (Source)
 		{
-			ENullTerminatedString::Type NullTerminated = ENullTerminatedString::No;
+			UE::Core::Private::ENullTerminatedString NullTerminated = UE::Core::Private::ENullTerminatedString::No;
 			if (SourceLen > 0 && ((const FromType*)Source)[SourceLen-1] == 0)
 			{
 				// Given buffer is null-terminated
-				NullTerminated = ENullTerminatedString::Yes;
+				NullTerminated = UE::Core::Private::ENullTerminatedString::Yes;
 				SourceLen -= 1;
 			}
 
