Manual merge 24896745 from Release-5.2

diff --git a/Engine/Source/Runtime/Core/Private/Apple/ApplePlatformSymbolication.cpp b/Engine/Source/Runtime/Core/Private/Apple/ApplePlatformSymbolication.cpp
--- a/Engine/Source/Runtime/Core/Private/Apple/ApplePlatformSymbolication.cpp
+++ b/Engine/Source/Runtime/Core/Private/Apple/ApplePlatformSymbolication.cpp
@@ -243,7 +243,7 @@ bool FApplePlatformSymbolication::LoadSymbolDatabaseForBinary(FString SourceFold
 	if(!bOK && GAllowApplePlatformSymbolication && (IFileManager::Get().FileSize(*BinaryPath) > 0))
 	{
 		CSSymbolicatorRef Symbolicator = OutDatabase.AppleDB;
-		if( CSIsNull(Symbolicator) && (!Architecture.IsSet() || *Architecture == TEXT("X86_64")))
+		if( CSIsNull(Symbolicator) && (!Architecture.IsSet() || *Architecture == TEXT("X86-64")))
 		{
 			Symbolicator = CSSymbolicatorCreateWithPathAndArchitecture(TCHAR_TO_UTF8(*BinaryPath), CPU_TYPE_X86_64);
 		}
