Don't call va_end when va_start/va_copy was not.

#rnx

diff --git a/Engine/Source/Runtime/Core/Public/Android/AndroidPlatformString.h b/Engine/Source/Runtime/Core/Public/Android/AndroidPlatformString.h
--- a/Engine/Source/Runtime/Core/Public/Android/AndroidPlatformString.h
+++ b/Engine/Source/Runtime/Core/Public/Android/AndroidPlatformString.h
@@ -456,7 +456,6 @@ struct FAndroidPlatformString : public FGenericPlatformString
 		NewFormat[NewIndex] = 0;
 #endif // PLATFORM_USE_LS_SPEC_FOR_WIDECHAR
 		int32 Result = vswprintf( Dest, DestSize, NewFormat, ArgPtr);
-		va_end( ArgPtr );
 		return Result;
 	}
 
@@ -558,7 +557,6 @@ struct FAndroidPlatformString : public FGenericPlatformString
 	static FORCEINLINE int32 GetVarArgs( ANSICHAR* Dest, SIZE_T DestSize, const ANSICHAR*& Fmt, va_list ArgPtr )
 	{
 		int32 Result = vsnprintf(Dest, DestSize, Fmt, ArgPtr);
-		va_end( ArgPtr );
 		return (Result != -1 && Result < (int32)DestSize) ? Result : -1;
 	}
 
diff --git a/Engine/Source/Runtime/Core/Public/GenericPlatform/GenericWidePlatformString.h b/Engine/Source/Runtime/Core/Public/GenericPlatform/GenericWidePlatformString.h
--- a/Engine/Source/Runtime/Core/Public/GenericPlatform/GenericWidePlatformString.h
+++ b/Engine/Source/Runtime/Core/Public/GenericPlatform/GenericWidePlatformString.h
@@ -299,7 +299,6 @@ PRAGMA_ENABLE_DEPRECATION_WARNINGS
 	static CORE_API int32 GetVarArgs( ANSICHAR* Dest, SIZE_T DestSize, const ANSICHAR*& Fmt, va_list ArgPtr )
 	{
 		int32 Result = vsnprintf(Dest, DestSize, Fmt, ArgPtr);
-		va_end( ArgPtr );
 		return (Result != -1 && Result < (int32)DestSize) ? Result : -1;
 	}
 
diff --git a/Engine/Source/Runtime/Core/Public/GenericPlatform/MicrosoftPlatformString.h b/Engine/Source/Runtime/Core/Public/GenericPlatform/MicrosoftPlatformString.h
--- a/Engine/Source/Runtime/Core/Public/GenericPlatform/MicrosoftPlatformString.h
+++ b/Engine/Source/Runtime/Core/Public/GenericPlatform/MicrosoftPlatformString.h
@@ -156,7 +156,6 @@ struct FMicrosoftPlatformString :
 	static FORCEINLINE int32 GetVarArgs( WIDECHAR* Dest, SIZE_T DestSize, const WIDECHAR*& Fmt, va_list ArgPtr )
 	{
 		int32 Result = vswprintf(Dest, DestSize, Fmt, ArgPtr);
-		va_end( ArgPtr );
 		return Result;
 	}
 #endif
@@ -259,7 +258,6 @@ struct FMicrosoftPlatformString :
 	static FORCEINLINE int32 GetVarArgs( ANSICHAR* Dest, SIZE_T DestSize, const ANSICHAR*& Fmt, va_list ArgPtr )
 	{
 		int32 Result = vsnprintf( Dest, DestSize, Fmt, ArgPtr );
-		va_end( ArgPtr );
 		return (Result != -1 && Result < (int32)DestSize) ? Result : -1;
 	}
 
diff --git a/Engine/Source/Runtime/Core/Public/GenericPlatform/StandardPlatformString.h b/Engine/Source/Runtime/Core/Public/GenericPlatform/StandardPlatformString.h
--- a/Engine/Source/Runtime/Core/Public/GenericPlatform/StandardPlatformString.h
+++ b/Engine/Source/Runtime/Core/Public/GenericPlatform/StandardPlatformString.h
@@ -202,7 +202,6 @@ PRAGMA_ENABLE_DEPRECATION_WARNINGS
 #else
 		int32 Result = vswprintf( Dest, DestSize, Fmt, ArgPtr);
 #endif
-		va_end( ArgPtr );
 		return Result;
 	}
 #else // PLATFORM_USE_SYSTEM_VSWPRINTF
@@ -319,7 +318,6 @@ PRAGMA_ENABLE_DEPRECATION_WARNINGS
 	static CORE_API int32 GetVarArgs( ANSICHAR* Dest, SIZE_T DestSize, const ANSICHAR*& Fmt, va_list ArgPtr )
 	{
 		int32 Result = vsnprintf(Dest, DestSize, Fmt, ArgPtr);
-		va_end( ArgPtr );
 		return (Result != -1 && Result < (int32)DestSize) ? Result : -1;
 	}
 
