Fixed ASan alloc-dealloc-mismatch for TCounters with dynamic name. Also fixed the case where id for a counter with dynamic name is correctly initialized in constructor (no string copy is allocated, but destructor was attempting to free the name string).

#rnx

diff --git a/Engine/Source/Runtime/Core/Private/ProfilingDebugging/CountersTrace.cpp b/Engine/Source/Runtime/Core/Private/ProfilingDebugging/CountersTrace.cpp
--- a/Engine/Source/Runtime/Core/Private/ProfilingDebugging/CountersTrace.cpp
+++ b/Engine/Source/Runtime/Core/Private/ProfilingDebugging/CountersTrace.cpp
@@ -37,7 +37,7 @@ const TCHAR* FCountersTrace::AllocAndCopyCounterName(const TCHAR* InCounterName)
 
 void FCountersTrace::FreeCounterName(const TCHAR* InCounterName)
 {
-	delete InCounterName;
+	delete[] InCounterName;
 }
 
 uint16 FCountersTrace::OutputInitCounter(const TCHAR* CounterName, ETraceCounterType CounterType, ETraceCounterDisplayHint CounterDisplayHint)
diff --git a/Engine/Source/Runtime/Core/Public/ProfilingDebugging/CountersTrace.h b/Engine/Source/Runtime/Core/Public/ProfilingDebugging/CountersTrace.h
--- a/Engine/Source/Runtime/Core/Public/ProfilingDebugging/CountersTrace.h
+++ b/Engine/Source/Runtime/Core/Public/ProfilingDebugging/CountersTrace.h
@@ -36,6 +36,7 @@ enum ETraceCounterNameType
 {
 	TraceCounterNameType_Static = 0, // TCounter is allowed to keep a pointer to InCounterName string
 	TraceCounterNameType_Dynamic = 0x10, // TCounter needs to copy the InCounterName string
+	TraceCounterNameType_AllocNameCopy = 0x20, // TCounter has allocated a copy of the InCounterName string
 };
 
 #if COUNTERSTRACE_ENABLED
@@ -79,13 +80,13 @@ struct FCountersTrace
 			{
 				// Store counter name for late init. Needs a copy as InCounterName pointer might not be valid later.
 				CounterName = AllocAndCopyCounterName(InCounterName);
-				CounterDisplayHint = ETraceCounterDisplayHint(uint8(CounterDisplayHint) | uint8(TraceCounterNameType_Dynamic));
+				CounterDisplayHint = ETraceCounterDisplayHint(uint8(CounterDisplayHint) | uint8(TraceCounterNameType_AllocNameCopy));
 			}
 		}
 
 		~TCounter()
 		{
-			if (uint8(CounterDisplayHint) & uint8(TraceCounterNameType_Dynamic))
+			if (uint8(CounterDisplayHint) & uint8(TraceCounterNameType_AllocNameCopy))
 			{
 				FreeCounterName(CounterName);
 			}
