Add configuration for whether to record crash dumps for Stalls and Ensures.
No implementation for Unix because we don't/can't record dumps outside of crash scenarios there.
No implementation for Mac at this time.

diff --git a/Engine/Config/BaseEngine.ini b/Engine/Config/BaseEngine.ini
--- a/Engine/Config/BaseEngine.ini
+++ b/Engine/Config/BaseEngine.ini
@@ -3133,6 +3133,8 @@ ConnectionRetryDelay=2
 bHideLogFilesOption=false
 bIsAllowedToCloseWithoutSending=true
 CrashConfigPurgeDays=2
+Stall.RecordDump=false
+Ensure.RecordDump=true
 
 [SteamVR.Settings]
 HMDWornMovementThreshold = 50.0
diff --git a/Engine/Source/Runtime/Core/Private/Windows/WindowsPlatformCrashContext.cpp b/Engine/Source/Runtime/Core/Private/Windows/WindowsPlatformCrashContext.cpp
--- a/Engine/Source/Runtime/Core/Private/Windows/WindowsPlatformCrashContext.cpp
+++ b/Engine/Source/Runtime/Core/Private/Windows/WindowsPlatformCrashContext.cpp
@@ -206,10 +206,32 @@ void FWindowsPlatformCrashContext::CopyPlatformSpecificFiles(const TCHAR* Output
 {
 	FGenericCrashContext::CopyPlatformSpecificFiles(OutputDirectory, Context);
 
-	// Save minidump
-	LPEXCEPTION_POINTERS ExceptionInfo = (LPEXCEPTION_POINTERS)Context;
-	const FString MinidumpFileName = FPaths::Combine(OutputDirectory, FGenericCrashContext::UEMinidumpName);
-	WriteMinidump(ProcessHandle.Get(), CrashedThreadId, *this, *MinidumpFileName, ExceptionInfo);
+	bool bRecordCrashDump = true;
+	switch (Type)
+	{
+		case ECrashContextType::Stall:
+			GConfig->GetBool(TEXT("CrashReportClient"), TEXT("Stall.RecordDump"), bRecordCrashDump, GEngineIni);
+			break;
+		case ECrashContextType::Ensure:
+			GConfig->GetBool(TEXT("CrashReportClient"), TEXT("Ensure.RecordDump"), bRecordCrashDump, GEngineIni);
+			break;
+		case ECrashContextType::AbnormalShutdown:
+		case ECrashContextType::Assert:
+		case ECrashContextType::Crash:
+		case ECrashContextType::GPUCrash:
+		case ECrashContextType::Hang:
+		case ECrashContextType::OutOfMemory:
+		default:
+			break;
+	}
+
+	if (bRecordCrashDump)
+	{
+		// Save minidump
+		LPEXCEPTION_POINTERS ExceptionInfo = (LPEXCEPTION_POINTERS)Context;
+		const FString MinidumpFileName = FPaths::Combine(OutputDirectory, FGenericCrashContext::UEMinidumpName);
+		WriteMinidump(ProcessHandle.Get(), CrashedThreadId, *this, *MinidumpFileName, ExceptionInfo);
+	}
 
 	// If present, include the crash video
 	const FString CrashVideoPath = FPaths::ProjectLogDir() / TEXT("CrashVideo.avi");
