Stall detector refactor
- Minimize duration of locks around FStallDetector and FStallDetectorStats instances
- Remove some public access to FStallDetector and FStallDetectorStats instances and synchronization around them
- Remove Stat from UE::FStallDetectorStats::TabulatedResult for thread safety in case of stat lifetime ending (e.g. thread shutdown)
- Fire callbacks & call ReportStall outside of locks
- Add backtrace to stall detected callback
- Only fire 'stall detected' from background thread, only fire 'stall completed' from stalling thread
- Change delegates to threadsafe versions
- Add checks for length of stall scope name
- Execute ReportStall in a task

diff --git a/Engine/Source/Runtime/Core/Private/ProfilingDebugging/StallDetector.cpp b/Engine/Source/Runtime/Core/Private/ProfilingDebugging/StallDetector.cpp
--- a/Engine/Source/Runtime/Core/Private/ProfilingDebugging/StallDetector.cpp
+++ b/Engine/Source/Runtime/Core/Private/ProfilingDebugging/StallDetector.cpp
@@ -6,17 +6,19 @@
 #include "HAL/ExceptionHandling.h"
 #include "HAL/PlatformProcess.h"
 #include "HAL/PlatformStackWalk.h"
-#include "HAL/PlatformTime.h"
 #include "HAL/PlatformTLS.h"
+#include "HAL/PlatformTime.h"
 #include "HAL/Runnable.h"
 #include "HAL/RunnableThread.h"
+#include "Tasks/Task.h"
+#include "Templates/Greater.h"
 
 // counters for sending information into trace system
 TRACE_DECLARE_INT_COUNTER(StallCount, TEXT("StallDetector/Count"));
 TRACE_DECLARE_FLOAT_COUNTER(StallTimeSeconds, TEXT("StallDetector/TimeSeconds"));
 
 // force normal behavior in the face of debug configuration and debugger attached
-#define STALL_DETECTOR_DEBUG 0
+#define STALL_DETECTOR_DEBUG 1
 
 // use the heart beat clock to account for process suspend
 #define STALL_DETECTOR_HEART_BEAT_CLOCK 1
@@ -40,12 +42,28 @@ static const double InvalidSeconds = -1.0;
 UE::FOnStallDetected UE::FStallDetector::StallDetected;
 UE::FOnStallCompleted UE::FStallDetector::StallCompleted;
 
+// Scopes currently being tracked across all threads
+static FCriticalSection StallScopesSection;
+static TSet<UE::FStallDetector*> StallScopes;
+
 /**
 * Stall Detector Thread
 **/
 
 namespace UE
 {
+	struct FDetectedStall
+	{
+		EStallDetectorReportingMode ReportingMode;
+		TArray<uint64> Backtrace;
+		uint32 ThreadId;
+		const TCHAR* Name;
+		double BudgetSeconds;
+		double OverageSeconds;
+		uint64 UniqueId;
+		bool bReported;
+	};
+
 	class FStallDetectorRunnable : public FRunnable
 	{
 	public:
@@ -113,11 +131,88 @@ uint32 UE::FStallDetectorRunnable::Run()
 		// Check the detectors
 		if (Seconds != InvalidSeconds)
 		{
-			FScopeLock ScopeLock(&FStallDetector::GetInstancesSection());
-			for (FStallDetector* Detector : FStallDetector::GetInstances())
+			bool bDisableReporting = UE_BUILD_DEBUG;
+#if !STALL_DETECTOR_DEBUG
+			bDisableReporting = bDisableReporting || FPlatformMisc::IsDebuggerPresent(); // Do not generate a report if we detect the debugger mucking with things
+#endif
+
+			TArray<FDetectedStall> DetectedStalls;
+			FScopeLock ScopeLock(&StallScopesSection);
+			for (FStallDetector* Detector : StallScopes)
 			{
-				Detector->Check(false, Seconds);
+				if (Detector->bTriggered)
+				{
+					continue;
+				}
+
+				double OverageSeconds = Detector->Check(false, Seconds);
+				if (OverageSeconds > 0.0)
+				{
+					TArray<uint64> Backtrace;
+					Backtrace.AddUninitialized(StallDetectorMaxBacktraceEntries);
+					int32 NumStackEntries = FPlatformStackWalk::CaptureThreadStackBackTrace(Detector->ThreadId, Backtrace.GetData(), Backtrace.Num());
+					Backtrace.SetNum(NumStackEntries);
+
+					// Store enough info to fire the event when we are no longer holding the lock
+					DetectedStalls.Emplace(FDetectedStall{
+						bDisableReporting ? EStallDetectorReportingMode::Disabled : Detector->Stats.ReportingMode,
+						MoveTemp(Backtrace),
+						Detector->ThreadId,
+						Detector->Stats.Name,
+						Detector->Stats.BudgetSeconds,
+						OverageSeconds,
+						Detector->UniqueID,
+						Detector->Stats.bReported });
+
+					// Triggered can be reset, reported cannot
+					Detector->bTriggered = true;
+					Detector->Stats.bReported = true;
+				}
 			}
+
+			// Do stall reporting in a task to keep heartbeat ticking for FStallDetector::Seconds resolution
+			UE::Tasks::Launch(UE_SOURCE_LOCATION, [DetectedStalls = MoveTemp(DetectedStalls)]() {
+				// Determine reporting requirements for each detected stall
+				for (const FDetectedStall& Stall : DetectedStalls)
+				{
+					bool bSendReport = false;
+					switch (Stall.ReportingMode)
+					{
+						case EStallDetectorReportingMode::First:
+							bSendReport = !Stall.bReported;
+							break;
+
+						case EStallDetectorReportingMode::Always:
+							bSendReport = true;
+							break;
+
+						default:
+							break;
+					}
+
+					if (bSendReport)
+					{
+						FStallDetectorStats::TotalReportedCount.Increment();
+						const int NumStackFramesToIgnore = 0;
+						UE_LOG(LogStall, Log, TEXT("Stall detector '%s' exceeded budget of %fs, reporting..."), Stall.Name, Stall.BudgetSeconds);
+						double ReportSeconds = FPlatformTime::Seconds();
+						ReportStall(Stall.Name, Stall.ThreadId);
+						ReportSeconds = FPlatformTime::Seconds() - ReportSeconds;
+						UE_LOG(LogStall, Log, TEXT("Stall detector '%s' report submitted, and took %fs"), Stall.Name, ReportSeconds);
+					}
+					else if (Stall.ReportingMode != EStallDetectorReportingMode::Disabled)
+					{
+						UE_LOG(LogStall, Log, TEXT("Stall detector '%s' exceeded budget of %fs"), Stall.Name, Stall.BudgetSeconds);
+					}
+				}
+
+				// Fire external notifications
+				//
+				for (const FDetectedStall& Stall : DetectedStalls)
+				{
+					FStallDetector::StallDetected.Broadcast(FStallDetectedParams{ Stall.Name, Stall.ThreadId, Stall.UniqueId });
+				}
+			});
 		}
 
 		// Sleep an interval, the resolution at which we want to detect an overage
@@ -131,8 +226,8 @@ uint32 UE::FStallDetectorRunnable::Run()
 * Stall Detector Stats
 **/
 
-FCriticalSection UE::FStallDetectorStats::InstancesSection;
-TSet<UE::FStallDetectorStats*> UE::FStallDetectorStats::Instances;
+static FCriticalSection StallDetectorStatsCritical;
+static TSet<UE::FStallDetectorStats*> StallDetectorStats;
 FCountersTrace::TCounter<std::atomic<int64>, TraceCounterType_Int> UE::FStallDetectorStats::TotalTriggeredCount (TEXT("StallDetector/TotalTriggeredCount"), TraceCounterDisplayHint_None);
 FCountersTrace::TCounter<std::atomic<int64>, TraceCounterType_Int> UE::FStallDetectorStats::TotalReportedCount (TEXT("StallDetector/TotalReportedCount"), TraceCounterDisplayHint_None);
 
@@ -157,15 +252,15 @@ UE::FStallDetectorStats::FStallDetectorStats(const TCHAR* InName, const double I
 		), TraceCounterDisplayHint_None)
 {
 	// Add at the end of construction
-	FScopeLock ScopeLock(&InstancesSection);
-	Instances.Add(this);
+	FScopeLock ScopeLock(&StallDetectorStatsCritical);
+	StallDetectorStats.Add(this);
 }
 
 UE::FStallDetectorStats::~FStallDetectorStats()
 {
 	// Remove at the beginning of destruction
-	FScopeLock ScopeLock(&InstancesSection);
-	Instances.Remove(this);
+	FScopeLock ScopeLock(&StallDetectorStatsCritical);
+	StallDetectorStats.Remove(this);
 }
 
 void UE::FStallDetectorStats::OnStallCompleted(double InOverageSeconds)
@@ -180,54 +275,28 @@ void UE::FStallDetectorStats::TabulateStats(TArray<TabulatedResult>& TabulatedRe
 {
 	TabulatedResults.Empty();
 
-	struct SortableStallStats
+	TArray<TabulatedResult> StatsArray;
 	{
-		explicit SortableStallStats(const UE::FStallDetectorStats* InStats)
-			: StallStats(InStats)
-			, OverageRatio(0.0)
+		FScopeLock ScopeLock(&StallDetectorStatsCritical);
+		for (const UE::FStallDetectorStats* StallStats : StallDetectorStats)
 		{
-			FScopeLock Lock(&InStats->StatsSection);
-			if (InStats->TriggerCount.Get() && InStats->BudgetSeconds > 0.0)
+			if (StallStats->TriggerCount.Get() && StallStats->ReportingMode != UE::EStallDetectorReportingMode::Disabled)
 			{
-				OverageRatio = (InStats->OverageSeconds.Get() / (double)InStats->TriggerCount.Get()) / InStats->BudgetSeconds;
+				// we sync access around these for coherency reasons, can be polled from another thread (detector or scope)
+				FScopeLock Lock(&StallStats->StatsSection);
+				StatsArray.Emplace(TabulatedResult{ StallStats->Name, StallStats->BudgetSeconds, StallStats->TriggerCount.Get(), StallStats->OverageSeconds.Get(), 0.0 });
+				StatsArray.Last().OverageRatio = (StallStats->OverageSeconds.Get() / (double)StallStats->TriggerCount.Get()) / StallStats->BudgetSeconds;
 			}
 		}
-
-		bool operator<(const SortableStallStats& InRhs) const
-		{
-			// NOTE THIS IS REVERSED TO PUT THE MOST OVERAGE AT THE FRONT
-			return OverageRatio > InRhs.OverageRatio;
-		}
-
-		const UE::FStallDetectorStats* StallStats;
-		double OverageRatio;
-	};
-
-	TArray<SortableStallStats> StatsArray;
-	FScopeLock InstancesLock(&UE::FStallDetectorStats::GetInstancesSection());
-	for (const UE::FStallDetectorStats* StallStats : UE::FStallDetectorStats::GetInstances())
-	{
-		if (StallStats->TriggerCount.Get() && StallStats->ReportingMode != UE::EStallDetectorReportingMode::Disabled)
-		{
-			StatsArray.Emplace(StallStats);
-		}
 	}
 
 	if (!StatsArray.IsEmpty())
 	{
-		StatsArray.Sort();
-
-		for (const SortableStallStats& Stat : StatsArray)
-		{
-			TabulatedResults.Emplace(TabulatedResult());
-			TabulatedResult& Result(TabulatedResults.Last());
-			Result.Stats = Stat.StallStats;
-
-			// we sync access around these for coherency reasons, can be polled from another thread (detector or scope)
-			FScopeLock Lock(&Result.Stats->StatsSection);
-			Result.TriggerCount = Result.Stats->TriggerCount.Get();
-			Result.OverageSeconds = Result.Stats->OverageSeconds.Get();
-		}
+		Algo::SortBy(
+			StatsArray,
+			[](const TabulatedResult& Result) { return Result.OverageRatio; },
+			TGreater<double>{});
+		TabulatedResults = MoveTemp(StatsArray);
 	}
 }
 
@@ -310,10 +379,6 @@ void UE::FStallTimer::Resume(const double InSeconds)
 * Stall Detector
 **/
 
-// statics
-FCriticalSection UE::FStallDetector::InstancesSection;
-TSet<UE::FStallDetector*> UE::FStallDetector::Instances;
-
 // globals
 class ThreadLocalSlotAcquire
 {
@@ -346,11 +411,11 @@ UE::FStallDetector::FStallDetector(FStallDetectorStats& InStats)
 	: Stats(InStats)
 	, Parent(nullptr)
 	, ThreadId(0)
-	, bStarted (false)
+	, bStarted(false)
 	, bPersistent(false)
-	, Triggered(false)
+	, bTriggered(false)
 	, UniqueID(UIDGenerator.IncrementExchange())
-	
+
 {
 	// Capture this thread's current stall detector to our parent member
 	Parent = static_cast<FStallDetector*>(FPlatformTLS::GetTlsValue(ThreadLocalSlot.GetSlot()));
@@ -373,16 +438,16 @@ UE::FStallDetector::FStallDetector(FStallDetectorStats& InStats)
 	}
 
 	// Add at the end of construction
-	FScopeLock ScopeLock(&InstancesSection);
-	Instances.Add(this);
+	FScopeLock ScopeLock(&StallScopesSection);
+	StallScopes.Add(this);
 }
 
 UE::FStallDetector::~FStallDetector()
 {
 	// Remove at the beginning of destruction
 	{
-		FScopeLock ScopeLock(&InstancesSection);
-		Instances.Remove(this);
+		FScopeLock ScopeLock(&StallScopesSection);
+		StallScopes.Remove(this);
 	}
 
 	// This will be invalid if the stall detector thread isn't running
@@ -391,14 +456,28 @@ UE::FStallDetector::~FStallDetector()
 	// If we have a timestamp, and captured a timestamp at contruction time, and we are not a persistent detector, do the final check
 	if (Seconds != InvalidSeconds && bStarted && !bPersistent)
 	{
-		Check(true, Seconds);
+		double OverageSeconds = Check(true, Seconds);
+
+		if (OverageSeconds > 0.0)
+		{
+#if STALL_DETECTOR_DEBUG
+			FString OverageString = FString::Printf(TEXT("[FStallDetector] [%s] Overage of %f\n"), Stats.Name, OverageSeconds);
+			FPlatformMisc::LocalPrint(OverageString.GetCharArray().GetData());
+#endif
+			if (Stats.ReportingMode != EStallDetectorReportingMode::Disabled)
+			{
+				UE_LOG(LogStall, Log, TEXT("Stall detector '%s' complete in %fs (%fs overbudget)"), Stats.Name, Stats.BudgetSeconds + OverageSeconds, OverageSeconds);
+			}
+			Stats.OnStallCompleted(OverageSeconds);
+			StallCompleted.Broadcast(FStallCompletedParams{ ThreadId, Stats.Name, Stats.BudgetSeconds, OverageSeconds, UniqueID, bTriggered });
+		}
 	}
 
 	// We are destructed, so set current thread's detector to our parent
 	FPlatformTLS::SetTlsValue(ThreadLocalSlot.GetSlot(), Parent);
 }
 
-void UE::FStallDetector::Check(bool bFinalCheck, double InCheckSeconds)
+double UE::FStallDetector::Check(bool bFinalCheck, double InCheckSeconds)
 {
 	// all callers need to sort out checking the clock
 	check(InCheckSeconds != InvalidSeconds);
@@ -411,37 +490,11 @@ void UE::FStallDetector::Check(bool bFinalCheck, double InCheckSeconds)
 	double DeltaSeconds;
 	double OverageSeconds;
 	Timer.Check(InCheckSeconds, DeltaSeconds, OverageSeconds);
-
-	if (Triggered)
+	if (OverageSeconds > 0.0)
 	{
-		// bFinalCheck ensures that only the destructor or CheckAndReset run the final OnStallCompleted
-		if (bFinalCheck)
-		{
-#if STALL_DETECTOR_DEBUG
-			FString OverageString = FString::Printf(TEXT("[FStallDetector] [%s] Overage of %f\n"), Stats.Name, OverageSeconds);
-			FPlatformMisc::LocalPrint(OverageString.GetCharArray().GetData());
-#endif
-
-			OnStallCompleted(ThreadId, DeltaSeconds, OverageSeconds);
-		}
-	}
-	else
-	{
-		if (OverageSeconds > 0.0)
-		{
-			// can be called from multiple threads, but we only want the first caller to call OnStallDetected
-			bool PreviousTriggered = false;
-			if (Triggered.compare_exchange_strong(PreviousTriggered, true, std::memory_order_acquire, std::memory_order_relaxed))
-			{
-#if STALL_DETECTOR_DEBUG
-				FString OverageString = FString::Printf(TEXT("[FStallDetector] [%s] Triggered at %f\n"), Stats.Name, InCheckSeconds);
-				FPlatformMisc::LocalPrint(OverageString.GetCharArray().GetData());
-#endif
-				
-				OnStallDetected(ThreadId, DeltaSeconds);
-			}
-		}
+		Stats.TotalTriggeredCount.Increment();
 	}
+	return OverageSeconds;
 }
 
 void UE::FStallDetector::CheckAndReset()
@@ -452,6 +505,10 @@ void UE::FStallDetector::CheckAndReset()
 		return;
 	}
 
+	// Lock to prevent races over bTriggered flag
+	FScopeLock ScopeLock(&StallScopesSection);
+
+	bool bExceeded = false;
 	// If this is the first call to CheckAndReset, flag that we are now in persistent mode (vs. scope mode)
 	if (!bPersistent)
 	{
@@ -460,11 +517,24 @@ void UE::FStallDetector::CheckAndReset()
 	else
 	{
 		// only perform the check on the second call
-		Check(true, CheckSeconds);
+		double OverageSeconds = Check(true, CheckSeconds);
+		if (OverageSeconds > 0.0)
+		{
+#if STALL_DETECTOR_DEBUG
+			FString OverageString = FString::Printf(TEXT("[FStallDetector] [%s] Persistent stall detector overage of %f\n"), Stats.Name, OverageSeconds);
+			FPlatformMisc::LocalPrint(OverageString.GetCharArray().GetData());
+#endif
+			if (Stats.ReportingMode != EStallDetectorReportingMode::Disabled)
+			{
+				UE_LOG(LogStall, Log, TEXT("Stall detector '%s' complete in %fs (%fs overbudget)"), Stats.Name, Stats.BudgetSeconds + OverageSeconds, OverageSeconds);
+			}
+			Stats.OnStallCompleted(OverageSeconds);
+			StallCompleted.Broadcast(FStallCompletedParams{ ThreadId, Stats.Name, Stats.BudgetSeconds, OverageSeconds, UniqueID, bTriggered });
+		}
 	}
 
 	Timer.Reset(CheckSeconds, Stats.BudgetSeconds);
-	Triggered = false;
+	bTriggered = false;
 	UniqueID = UIDGenerator.IncrementExchange();
 }
 
@@ -484,101 +554,6 @@ void UE::FStallDetector::Resume(const double InSeconds)
 	Timer.Resume(InSeconds);
 }
 
-void UE::FStallDetector::OnStallDetected(uint32 InThreadId, const double InElapsedSeconds)
-{
-	Stats.TotalTriggeredCount.Increment();
-
-	//
-	// Determine if we want to undermine the specified reporting mode
-	//
-
-	EStallDetectorReportingMode ReportingMode = Stats.ReportingMode;
-
-	bool bDisableReporting = false;
-
-#if UE_BUILD_DEBUG
-	bDisableReporting |= true; // Do not generate a report in debug configurations due to performance characteristics
-#endif
-
-#if !STALL_DETECTOR_DEBUG
-	bDisableReporting |= FPlatformMisc::IsDebuggerPresent(); // Do not generate a report if we detect the debugger mucking with things
-#endif
-
-	if (bDisableReporting)
-	{
-#if !STALL_DETECTOR_DEBUG
-		ReportingMode = EStallDetectorReportingMode::Disabled;
-#endif
-	}
-
-	//
-	// Resolve reporting mode to whether we should send a report for this call
-	//
-
-	bool bSendReport = false;
-	switch (ReportingMode)
-	{
-	case EStallDetectorReportingMode::First:
-		bSendReport = !Stats.bReported;
-		break;
-
-	case EStallDetectorReportingMode::Always:
-		bSendReport = true;
-		break;
-
-	default:
-		break;
-	}
-
-	//
-	// Send the report
-	//
-
-	if (bSendReport)
-	{
-		Stats.bReported = true;
-		Stats.TotalReportedCount.Increment();
-		const int NumStackFramesToIgnore = FPlatformTLS::GetCurrentThreadId() == InThreadId ? 2 : 0;
-		UE_LOG(LogStall, Log, TEXT("Stall detector '%s' exceeded budget of %fs, reporting..."), Stats.Name, Stats.BudgetSeconds);
-		double ReportSeconds = FStallDetector::Seconds();
-		ReportStall(Stats.Name, InThreadId);
-		ReportSeconds = FStallDetector::Seconds() - ReportSeconds;
-		UE_LOG(LogStall, Log, TEXT("Stall detector '%s' report submitted, and took %fs"), Stats.Name, ReportSeconds);
-	}
-	else
-	{
-		if (ReportingMode != EStallDetectorReportingMode::Disabled)
-		{
-			UE_LOG(LogStall, Log, TEXT("Stall detector '%s' exceeded budget of %fs"), Stats.Name, Stats.BudgetSeconds);
-		}
-	}
-	
-	StallDetected.Broadcast({
-		Stats.Name,
-		ThreadId,
-		UniqueID
-	});
-}
-
-void UE::FStallDetector::OnStallCompleted(uint32 InThreadId, double DeltaSeconds, double OverageSeconds)
-{
-	// this callback allows for tracking the total overrun, once the stalling period is finally done
-	Stats.OnStallCompleted(OverageSeconds);
-			
-	const FStallCompletedParams Params {
-		ThreadId,
-		Stats.Name,
-		OverageSeconds,
-		UniqueID
-	};
-	StallCompleted.Broadcast(Params);
-
-	if (Stats.ReportingMode != EStallDetectorReportingMode::Disabled)
-	{
-		UE_LOG(LogStall, Log, TEXT("Stall detector '%s' complete in %fs (%fs overbudget)"), Stats.Name, DeltaSeconds, OverageSeconds);
-	}
-}
-
 double UE::FStallDetector::Seconds()
 {
 	double Result = InvalidSeconds;
diff --git a/Engine/Source/Runtime/Core/Public/ProfilingDebugging/StallDetector.h b/Engine/Source/Runtime/Core/Public/ProfilingDebugging/StallDetector.h
--- a/Engine/Source/Runtime/Core/Public/ProfilingDebugging/StallDetector.h
+++ b/Engine/Source/Runtime/Core/Public/ProfilingDebugging/StallDetector.h
@@ -27,16 +27,60 @@
 
 #if STALL_DETECTOR
 
-CORE_API DECLARE_LOG_CATEGORY_EXTERN(LogStall, Log, All);
+/**
+ * Counts, but does not report, a stall.
+ **/
+#define SCOPE_STALL_COUNTER(InName, InBudgetSeconds)                                                                                                                                           \
+	static_assert(UE_ARRAY_COUNT(PREPROCESSOR_TO_STRING(InName)) <= UE::StallDetectorScopeMaxNameLength, "Scope name must be less than UE::StallDetectorScopeMaxNameLength characters long "); \
+	static UE::FStallDetectorStats PREPROCESSOR_JOIN(StallDetectorStats, __LINE__)(TEXT(#InName), InBudgetSeconds, UE::EStallDetectorReportingMode::Disabled);                                 \
+	UE::FStallDetector PREPROCESSOR_JOIN(ScopeStallDetector, __LINE__)(PREPROCESSOR_JOIN(StallDetectorStats, __LINE__));
+
+/**
+ * Counts and sends a report of the first occurence of a stall.
+ **/
+#define SCOPE_STALL_REPORTER(InName, InBudgetSeconds)                                                                                                                                          \
+	static_assert(UE_ARRAY_COUNT(PREPROCESSOR_TO_STRING(InName)) <= UE::StallDetectorScopeMaxNameLength, "Scope name must be less than UE::StallDetectorScopeMaxNameLength characters long "); \
+	static UE::FStallDetectorStats PREPROCESSOR_JOIN(StallDetectorStats, __LINE__)(TEXT(#InName), InBudgetSeconds, UE::EStallDetectorReportingMode::First);                                    \
+	UE::FStallDetector PREPROCESSOR_JOIN(ScopeStallDetector, __LINE__)(PREPROCESSOR_JOIN(StallDetectorStats, __LINE__));
+
+/**
+ * Counts and reports every occurence of a stall.
+ * Use with extreme caution as you could overload the report backend!
+ * This mainly exists to help test and exercise the stall detector system.
+ **/
+#define SCOPE_STALL_REPORTER_ALWAYS(InName, InBudgetSeconds)                                                                                                                                   \
+	static_assert(UE_ARRAY_COUNT(PREPROCESSOR_TO_STRING(InName)) <= UE::StallDetectorScopeMaxNameLength, "Scope name must be less than UE::StallDetectorScopeMaxNameLength characters long "); \
+	static UE::FStallDetectorStats PREPROCESSOR_JOIN(StallDetectorStats, __LINE__)(TEXT(#InName), InBudgetSeconds, UE::EStallDetectorReportingMode::Always);                                   \
+	UE::FStallDetector PREPROCESSOR_JOIN(ScopeStallDetector, __LINE__)(PREPROCESSOR_JOIN(StallDetectorStats, __LINE__));
 
 /**
-* This code is meant to instrument code to send reports for slow pathways that create unresponsive conditions
-**/
+ * Pauses stall detectors on the current thread so that time in the current thread spent in this scope will not be counted as a stall.
+ **/
+#define SCOPE_STALL_DETECTOR_PAUSE() \
+	UE::FStallDetectorPause PREPROCESSOR_JOIN(ScopeStallDetectorPause, __LINE__);
+
+#else // STALL_DETECTOR
+
+#define SCOPE_STALL_COUNTER(InName, InBudgetSeconds)
+#define SCOPE_STALL_REPORTER(InName, InBudgetSeconds)
+#define SCOPE_STALL_REPORTER_ALWAYS(InName, InBudgetSeconds)
+#define SCOPE_STALL_DETECTOR_PAUSE()
+
+#endif // STALL_DETECTOR
+
+#if STALL_DETECTOR
+
+CORE_API DECLARE_LOG_CATEGORY_EXTERN(LogStall, Log, All);
 
+/**
+ * This code is meant to instrument code to send reports for slow pathways that create unresponsive conditions
+ **/
 namespace UE
 {
 	class FStallDetector;
 
+	// Stall detected event is fired on a worker thread when it is detected that a stall scope has been active for too long.
+	// Fired from a worker thread.
 	struct FStallDetectedParams
 	{
 		/** Name of Stall Detector instance*/
@@ -45,21 +89,35 @@ namespace UE
 		uint32 ThreadId;
 		/** UniqueID of the stalled instance */
 		uint64 UniqueID;
+		/** Backtrace of stall */
+		TArray<uint64> Backtrace;
 	};
-	DECLARE_EVENT_OneParam(FStallDetector, FOnStallDetected, const FStallDetectedParams&);
-	
+	DECLARE_TS_MULTICAST_DELEGATE_OneParam(FOnStallDetected, const FStallDetectedParams&);
+
+	// Stall completed event is fired synchronously on the stalling thread after it recovers from the stall.
+	// It may be fired without a matching StallDetected event in the case that the background detector's resolution
+	// is too low or that thread is blocked for some reason.
 	struct FStallCompletedParams
 	{
 		/** TID of stalled thread */
 		uint32 ThreadId;
 		/** Name of Stall Detector instance*/
 		FStringView StatName;
-		/** Amount of time StallDetector instance stalled  */
+		/** Budget for this scope. */
+		double BudgetSeconds;
+		/** Amount of time StallDetector instance stalled beyond the budget. */
 		double OverbudgetSeconds;
 		/** UniqueID of the stalled instance */
 		uint64 UniqueID;
+		/** 
+		 * True if the background thread detected a stall in this scope.
+		 * Since the compelted delegate is fired from the thread containing the scope, 
+		 * OnStallDetected and OnStallCompleted are fired on different threads and may be executed in any order.
+		 * If bWasTriggered is true, then OnStallDetected has fired or will fire for this UniqueID.
+		 **/
+		bool bWasTriggered;
 	};
-	DECLARE_EVENT_OneParam(FStallDetector, FOnStallCompleted, const FStallCompletedParams&);
+	DECLARE_TS_MULTICAST_DELEGATE_OneParam(FOnStallCompleted, const FStallCompletedParams&);
 
 	/**
 	* The reporting behavior of the detector
@@ -71,11 +129,18 @@ namespace UE
 		Always
 	};
 
+	static constexpr SIZE_T StallDetectorStatNameBufferSize = 256;
+	static constexpr int32 StallDetectorMaxBacktraceEntries = 32;
+
+	class FStallDetectorRunnable;
+
 	/**
-	* Information to track per-callsite (vs. lifetime of a single interval)
-	**/
+	 * This structure is allocated as a function-local static to persist stats across stall checks.
+	 */
 	class CORE_API FStallDetectorStats
 	{
+		friend class FStallDetectorRunnable;
+
 	public:
 		/**
 		* Construct a stats object for use with a stall detector
@@ -92,36 +157,13 @@ namespace UE
 		*/
 		void OnStallCompleted(double InOverageSeconds);
 
-		/**
-		* Retrieve the critical section (lock) for accessing the instances list
-		*/
-		static FCriticalSection& GetInstancesSection()
-		{
-			return InstancesSection;
-		}
-
-		/**
-		* The container of all constructed instances of this class
-		*/
-		static const TSet<FStallDetectorStats*>& GetInstances()
-		{
-			return Instances;
-		}
-
-		/**
-		* 
-		*/
 		struct TabulatedResult
 		{
-			const UE::FStallDetectorStats* Stats;
-			uint64 TriggerCount;
-			double OverageSeconds;
-			
-			TabulatedResult()
-				: Stats(nullptr)
-				, TriggerCount(0)
-				, OverageSeconds(0.0)
-			{}
+			const TCHAR* Name = TEXT("");
+			double BudgetSeconds = 0.0;
+			int64 TriggerCount = 0;
+			double OverageSeconds = 0.0;
+			double OverageRatio = 0.0;
 		};
 		static void TabulateStats(TArray<TabulatedResult>& TabulatedResults);
 
@@ -146,25 +188,20 @@ namespace UE
 	private:
 		// The number of times this callsite has been triggered
 		FCountersTrace::TCounter<int64, TraceCounterType_Int> TriggerCount;
-		TCHAR TriggerCountCounterName[256];
+		TCHAR TriggerCountCounterName[StallDetectorStatNameBufferSize];
 
 		// The cumulative overage time for this callsite
 		FCountersTrace::TCounter<double, TraceCounterType_Float> OverageSeconds;
-		TCHAR OverageSecondsCounterName[256];
+		TCHAR OverageSecondsCounterName[StallDetectorStatNameBufferSize];
 
 		// Guards access to the stats from multiple threads, for coherency
 		mutable FCriticalSection StatsSection;
-
-		// Guards access to the instances container from multiple threads
-		static FCriticalSection InstancesSection;
-
-		// The pointers to every constructed instance of this class
-		static TSet<FStallDetectorStats*> Instances;
 	};
 
 	/**
-	* FStallTimer tracks remaining time from some budget, and access is protected from multiple threads
-	**/
+	 * FStallTimer tracks remaining time from some budget, and access is protected from multiple threads.
+	 * Owned by FStallDetector
+	 **/
 	class CORE_API FStallTimer
 	{
 	public:
@@ -212,6 +249,8 @@ namespace UE
 	**/
 	class CORE_API FStallDetector
 	{
+		friend class FStallDetectorRunnable;
+
 	public:
 		/**
 		* Construct the stall detector linking it to it's stats object and take a timestamp
@@ -223,13 +262,6 @@ namespace UE
 		*/
 		~FStallDetector();
 
-		/**
-		* Check the detector to see if we are over budget, called from both construction and background thread
-		* @param bFinalCheck			Is this check the one performed at completion of the scope or measurement period?
-		* @param InCheckSeconds			The timestamp at which to compare the start time to
-		**/
-		void Check(bool bFinalCheck, double InCheckSeconds);
-
 		/**
 		* Perform a Check and then reset the timer to the current time
 		* This changes the mode of this class to defer the first check, and skip the destruction check
@@ -261,21 +293,6 @@ namespace UE
 		static double Seconds();
 
 		/**
-		* Retrieve the critical section (lock) for accessing the instances list
-		*/
-		static FCriticalSection& GetInstancesSection()
-		{
-			return InstancesSection;
-		}
-
-		/**
-		* The container of all constructed instances of this class
-		*/
-		static const TSet<FStallDetector*>& GetInstances()
-		{
-			return Instances;
-		}
-
 		/**
 		* Initialize Stall Detector API resources
 		**/
@@ -312,7 +329,7 @@ namespace UE
 		// The time provided, the deadline
 		FStallDetectorStats& Stats;
 
-		// The parent detector, if any
+		// A stall detector on the same thread with a longer lifetime if any
 		FStallDetector* Parent;
 
 		// Thread it was constructed on, and to stack trace when budgeted time is elapsed
@@ -324,26 +341,24 @@ namespace UE
 		// Persistent usage mode
 		bool bPersistent;
 
+		// Track whether a stall was detected from the background thread.
+		// This is protected by the same mutex which protects the registered stall detector instances.
+		// Only set from background thread, can be reset from another thread with CheckAndReset
+		bool bTriggered;
+
 		// The timer data for this detector
 		FStallTimer Timer;
 
-		// Track the triggered state, atomic to mitigate foreground and background thread checking
-		std::atomic<bool> Triggered;
-
 		// Unique ID given to a StallDetector
 		uint64 UniqueID;
 
-		// The pinch point for tracking counts and emitting a report for a triggered hitch detector
-		void OnStallDetected(uint32 InThreadId, const double InElapsedSeconds);
-
-		// Called when the thread that triggered the detector is no longer stalling
-		void OnStallCompleted(uint32 InThreadId, const double DeltaSeconds, const double OverageSeconds);
-
-		// Guards access to the instances container from multiple threads
-		static FCriticalSection InstancesSection;
-
-		// The pointers to every constructed instance of this class
-		static TSet<FStallDetector*> Instances;
+		/**
+		 * Check the detector to see if we are over budget, called from both construction and background thread
+		 * @param bFinalCheck			Is this check the one performed at completion of the scope or measurement period?
+		 * @param InCheckSeconds		The timestamp at which to compare the start time to
+		 * @return The amount the scope's budget has been exceeded by or 0.0
+		 **/
+		double Check(bool bFinalCheck, double InCheckSeconds);
 	};
 
 	/**
@@ -358,42 +373,11 @@ namespace UE
 	private:
 		bool bPaused;
 	};
-}
-
-/**
-* Counts, but does not report, a stall.
-**/
-#define SCOPE_STALL_COUNTER(InName, InBudgetSeconds) \
-static UE::FStallDetectorStats PREPROCESSOR_JOIN(StallDetectorStats, __LINE__) (TEXT(#InName), InBudgetSeconds, UE::EStallDetectorReportingMode::Disabled); \
-UE::FStallDetector PREPROCESSOR_JOIN(ScopeStallDetector, __LINE__)(PREPROCESSOR_JOIN(StallDetectorStats, __LINE__));
-
-/**
-* Counts and sends a report of the first occurence of a stall.
-**/
-#define SCOPE_STALL_REPORTER(InName, InBudgetSeconds) \
-static UE::FStallDetectorStats PREPROCESSOR_JOIN(StallDetectorStats, __LINE__) (TEXT(#InName), InBudgetSeconds, UE::EStallDetectorReportingMode::First); \
-UE::FStallDetector PREPROCESSOR_JOIN(ScopeStallDetector, __LINE__)(PREPROCESSOR_JOIN(StallDetectorStats, __LINE__));
-
-/**
-* Counts and reports every occurence of a stall.
-* Use with extreme caution as you could overload the report backend!
-* This mainly exists to help test and exercise the stall detector system.
-**/
-#define SCOPE_STALL_REPORTER_ALWAYS(InName, InBudgetSeconds) \
-static UE::FStallDetectorStats PREPROCESSOR_JOIN(StallDetectorStats, __LINE__) (TEXT(#InName), InBudgetSeconds, UE::EStallDetectorReportingMode::Always); \
-UE::FStallDetector PREPROCESSOR_JOIN(ScopeStallDetector, __LINE__)(PREPROCESSOR_JOIN(StallDetectorStats, __LINE__));
-
-/**
-* Counts, but does not report, a stall.
-**/
-#define SCOPE_STALL_DETECTOR_PAUSE() \
-UE::FStallDetectorPause PREPROCESSOR_JOIN(ScopeStallDetectorPause, __LINE__);
-
-#else // STALL_DETECTOR
-
-#define SCOPE_STALL_COUNTER(InName, InBudgetSeconds)
-#define SCOPE_STALL_REPORTER(InName, InBudgetSeconds)
-#define SCOPE_STALL_REPORTER_ALWAYS(InName, InBudgetSeconds)
-#define SCOPE_STALL_DETECTOR_PAUSE()
 
+	static constexpr TCHAR StallDetectorCounterPrefix[] = TEXT("StallDetector/");
+	static constexpr TCHAR StallDetectorTriggerCountSuffix[] = TEXT(" TriggerCount");
+	static constexpr TCHAR StallDetectorOverageTimeSuffix[] = TEXT(" OverageSeconds");
+	static constexpr SIZE_T StallDetectorScopeMaxNameLength = StallDetectorStatNameBufferSize - 1 - UE_ARRAY_COUNT(StallDetectorCounterPrefix)
+															+ FMath::Max(UE_ARRAY_COUNT(StallDetectorTriggerCountSuffix), UE_ARRAY_COUNT(StallDetectorOverageTimeSuffix));
+} // namespace UE
 #endif // STALL_DETECTOR
\ No newline at end of file
diff --git a/Engine/Source/Runtime/Engine/Private/Analytics/EditorAnalyticsSessionSummary.cpp b/Engine/Source/Runtime/Engine/Private/Analytics/EditorAnalyticsSessionSummary.cpp
--- a/Engine/Source/Runtime/Engine/Private/Analytics/EditorAnalyticsSessionSummary.cpp
+++ b/Engine/Source/Runtime/Engine/Private/Analytics/EditorAnalyticsSessionSummary.cpp
@@ -113,8 +113,8 @@ bool FEditorAnalyticsSessionSummary::UpdateSessionProgressInternal(bool bCrashin
 	if (!StallResults.IsEmpty())
 	{
 		UE::FStallDetectorStats::TabulatedResult TopResult(StallResults[0]);
-		EditorAnalyticsProperties::TopStallName.Set(GetStore(), TopResult.Stats->Name, /*Capacity*/128);
-		EditorAnalyticsProperties::TopStallBudgetSeconds.Set(GetStore(), TopResult.Stats->BudgetSeconds);
+		EditorAnalyticsProperties::TopStallName.Set(GetStore(), TopResult.Name, /*Capacity*/ 128);
+		EditorAnalyticsProperties::TopStallBudgetSeconds.Set(GetStore(), TopResult.BudgetSeconds);
 		EditorAnalyticsProperties::TopStallOverageSeconds.Set(GetStore(), TopResult.OverageSeconds);
 		EditorAnalyticsProperties::TopStallTriggerCount.Set(GetStore(), TopResult.TriggerCount);
 	}
