Fix mismatched block comment.

diff --git a/Engine/Source/Runtime/Core/Public/ProfilingDebugging/StallDetector.h b/Engine/Source/Runtime/Core/Public/ProfilingDebugging/StallDetector.h
--- a/Engine/Source/Runtime/Core/Public/ProfilingDebugging/StallDetector.h
+++ b/Engine/Source/Runtime/Core/Public/ProfilingDebugging/StallDetector.h
@@ -292,7 +292,6 @@ namespace UE
 		**/
 		static double Seconds();
 
-		/**
 		/**
 		* Initialize Stall Detector API resources
 		**/
