Disable stall detector when debugging.
Fix passing backtrace to stall detected delegate.

diff --git a/Engine/Source/Runtime/Core/Private/ProfilingDebugging/StallDetector.cpp b/Engine/Source/Runtime/Core/Private/ProfilingDebugging/StallDetector.cpp
--- a/Engine/Source/Runtime/Core/Private/ProfilingDebugging/StallDetector.cpp
+++ b/Engine/Source/Runtime/Core/Private/ProfilingDebugging/StallDetector.cpp
@@ -18,7 +18,7 @@ TRACE_DECLARE_INT_COUNTER(StallCount, TEXT("StallDetector/Count"));
 TRACE_DECLARE_FLOAT_COUNTER(StallTimeSeconds, TEXT("StallDetector/TimeSeconds"));
 
 // force normal behavior in the face of debug configuration and debugger attached
-#define STALL_DETECTOR_DEBUG 1
+#define STALL_DETECTOR_DEBUG 0
 
 // use the heart beat clock to account for process suspend
 #define STALL_DETECTOR_HEART_BEAT_CLOCK 1
@@ -210,7 +210,7 @@ uint32 UE::FStallDetectorRunnable::Run()
 				//
 				for (const FDetectedStall& Stall : DetectedStalls)
 				{
-					FStallDetector::StallDetected.Broadcast(FStallDetectedParams{ Stall.Name, Stall.ThreadId, Stall.UniqueId });
+					FStallDetector::StallDetected.Broadcast(FStallDetectedParams{ Stall.Name, Stall.ThreadId, Stall.UniqueId, Stall.Backtrace });
 				}
 			});
 		}
