Disabled compiler-defined-copy-with-user-defined-copy warnings in third party headers.

diff --git a/Engine/Source/Runtime/Core/Public/Clang/ClangPlatformCompilerPreSetup.h b/Engine/Source/Runtime/Core/Public/Clang/ClangPlatformCompilerPreSetup.h
--- a/Engine/Source/Runtime/Core/Public/Clang/ClangPlatformCompilerPreSetup.h
+++ b/Engine/Source/Runtime/Core/Public/Clang/ClangPlatformCompilerPreSetup.h
@@ -221,6 +221,24 @@
 		_Pragma("clang diagnostic pop")
 #endif // PRAGMA_ENABLE_ENUM_CONSTEXPR_CONVERSION_WARNINGS
 
+#if __has_warning("-Wdeprecated-copy-with-user-provided-copy")
+#define DISABLE_DEPRECATED_COPY_WITH_USER_DEFINED_COPY _Pragma("clang diagnostic ignored \"-Wdeprecated-copy-with-user-provided-copy\"")
+#else
+#define DISABLE_DEPRECATED_COPY_WITH_USER_DEFINED_COPY
+#endif
+
+#ifndef PRAGMA_DISABLE_DEPRECATED_COPY_WITH_USER_DEFINED_COPY_WARNINGS
+#define PRAGMA_DISABLE_DEPRECATED_COPY_WITH_USER_DEFINED_COPY_WARNINGS \
+		_Pragma("clang diagnostic push") \
+        DISABLE_DEPRECATED_COPY_WITH_USER_DEFINED_COPY
+#endif // PRAGMA_DISABLE_DEPRECATED_COPY_WITH_USER_DEFINED_COPY_WARNINGS
+
+#ifndef PRAGMA_ENABLE_DEPRECATED_COPY_WITH_USER_DEFINED_COPY_WARNINGS
+#define PRAGMA_ENABLE_DEPRECATED_COPY_WITH_USER_DEFINED_COPY_WARNINGS \
+		_Pragma("clang diagnostic pop")
+#endif // PRAGMA_ENABLE_DEPRECATED_COPY_WITH_USER_DEFINED_COPY_WARNINGS
+
+
 #ifndef PRAGMA_POP
 	#define PRAGMA_POP \
 		_Pragma("clang diagnostic pop")
@@ -249,11 +267,13 @@
 		PRAGMA_DISABLE_OVERLOADED_VIRTUAL_WARNINGS \
 		PRAGMA_DISABLE_MISSING_BRACES_WARNINGS \
 		PRAGMA_DISABLE_UNINITIALIZED_CONST_REFERENCE_WARNINGS \
-		PRAGMA_DISABLE_ORDERED_COMPARE_FUNCTION_POINTERS
+		PRAGMA_DISABLE_ORDERED_COMPARE_FUNCTION_POINTERS \
+		PRAGMA_DISABLE_DEPRECATED_COPY_WITH_USER_DEFINED_COPY_WARNINGS
 #endif
 
 #ifndef UE_COMPILER_THIRD_PARTY_INCLUDES_END
 	#define UE_COMPILER_THIRD_PARTY_INCLUDES_END \
+		PRAGMA_ENABLE_DEPRECATED_COPY_WITH_USER_DEFINED_COPY_WARNINGS \
 		PRAGMA_ENABLE_ORDERED_COMPARE_FUNCTION_POINTERS \
 		PRAGMA_ENABLE_UNINITIALIZED_CONST_REFERENCE_WARNINGS \
 		PRAGMA_ENABLE_MISSING_BRACES_WARNINGS \
