Fixed FThreadSafeCounter64 copy constructor.

diff --git a/Engine/Source/Runtime/Core/Public/HAL/ThreadSafeCounter64.h b/Engine/Source/Runtime/Core/Public/HAL/ThreadSafeCounter64.h
--- a/Engine/Source/Runtime/Core/Public/HAL/ThreadSafeCounter64.h
+++ b/Engine/Source/Runtime/Core/Public/HAL/ThreadSafeCounter64.h
@@ -34,7 +34,7 @@ public:
 	*
 	* @param Other The other thread safe counter to copy
 	*/
-	FThreadSafeCounter64(const FThreadSafeCounter& Other)
+	FThreadSafeCounter64(const FThreadSafeCounter64& Other)
 	{
 		Counter = Other.GetValue();
 	}
