Allow TargetFrameRateOverride to be set early because execcmds come in too late.

diff --git a/Engine/Source/Runtime/Core/Private/ProfilingDebugging/CsvProfiler.cpp b/Engine/Source/Runtime/Core/Private/ProfilingDebugging/CsvProfiler.cpp
--- a/Engine/Source/Runtime/Core/Private/ProfilingDebugging/CsvProfiler.cpp
+++ b/Engine/Source/Runtime/Core/Private/ProfilingDebugging/CsvProfiler.cpp
@@ -2909,10 +2909,15 @@ void FCsvProfiler::BeginFrame()
 					static IConsoleVariable* CsvTargetFrameRateCVar = IConsoleManager::Get().FindConsoleVariable(TEXT("csv.TargetFrameRateOverride"));
 					static IConsoleVariable* MaxFPSCVar = IConsoleManager::Get().FindConsoleVariable(TEXT("t.MaxFPS"));
 					static IConsoleVariable* SyncIntervalCVar = IConsoleManager::Get().FindConsoleVariable(TEXT("rhi.SyncInterval"));
+					int32 CmdLineTargetFPS = TargetFPS;
 					if (CsvTargetFrameRateCVar && CsvTargetFrameRateCVar->GetInt() > 0)
 					{
 						TargetFPS = CsvTargetFrameRateCVar->GetInt();
 					}
+					else if (FParse::Value(FCommandLine::Get(), TEXT("csv.TargetFrameRateOverride"), CmdLineTargetFPS)) // Too early to set CsvTargetFrameRateCVar with execcmds
+					{
+						TargetFPS = CmdLineTargetFPS;
+					}
 					else
 					{
 						// Figure out the target framerate
