Array.h whitespace: Change nested block in Serialize to an early exit.
#rnx

diff --git a/Engine/Source/Runtime/Core/Public/Containers/Array.h b/Engine/Source/Runtime/Core/Public/Containers/Array.h
--- a/Engine/Source/Runtime/Core/Public/Containers/Array.h
+++ b/Engine/Source/Runtime/Core/Public/Containers/Array.h
@@ -3532,38 +3532,27 @@ struct TArrayPrivateFriend
 
 		check(SerializeNum >= 0);
 
-		if (!Ar.IsError() && SerializeNum > 0 && ensure(!Ar.IsNetArchive() || SerializeNum <= MaxNetArraySerialize))
+		if (Ar.IsError() || SerializeNum < 0 || !ensure(!Ar.IsNetArchive() || SerializeNum <= MaxNetArraySerialize))
 		{
-			// if we don't need to perform per-item serialization, just read it in bulk
-			if constexpr (sizeof(ElementType) == 1 || TCanBulkSerialize<ElementType>::Value)
-			{
-				A.ArrayNum = SerializeNum;
+			Ar.SetError();
+			return Ar;
+		}
 
-				// Serialize simple bytes which require no construction or destruction.
-				if ((A.ArrayNum || A.ArrayMax) && Ar.IsLoading())
-				{
-					A.ResizeForCopy(A.ArrayNum, A.ArrayMax);
-				}
+		// if we don't need to perform per-item serialization, just read it in bulk
+		if constexpr (sizeof(ElementType) == 1 || TCanBulkSerialize<ElementType>::Value)
+		{
+			A.ArrayNum = SerializeNum;
 
-				if(TIsUECoreVariant<ElementType, double>::Value && Ar.IsLoading() && Ar.UEVer() < EUnrealEngineObjectUE5Version::LARGE_WORLD_COORDINATES)
-				{
-					// Per item serialization is required for core variant types loaded from pre LWC archives, to enable conversion from float to double.
-					A.Empty(SerializeNum);
-					for (SizeType i=0; i<SerializeNum; i++)
-					{
-						Ar << *::new(A) ElementType;
-					}		
-				}
-				else
-				{
-					Ar.Serialize(A.GetData(), A.Num() * sizeof(ElementType));
-				}
+			// Serialize simple bytes which require no construction or destruction.
+			if ((A.ArrayNum || A.ArrayMax) && Ar.IsLoading())
+			{
+				A.ResizeForCopy(A.ArrayNum, A.ArrayMax);
 			}
-			else if (Ar.IsLoading())
+
+			if(TIsUECoreVariant<ElementType, double>::Value && Ar.IsLoading() && Ar.UEVer() < EUnrealEngineObjectUE5Version::LARGE_WORLD_COORDINATES)
 			{
-				// Required for resetting ArrayNum
+				// Per item serialization is required for core variant types loaded from pre LWC archives, to enable conversion from float to double.
 				A.Empty(SerializeNum);
-
 				for (SizeType i=0; i<SerializeNum; i++)
 				{
 					Ar << *::new(A) ElementType;
@@ -3571,17 +3560,27 @@ struct TArrayPrivateFriend
 			}
 			else
 			{
-				A.ArrayNum = SerializeNum;
+				Ar.Serialize(A.GetData(), A.Num() * sizeof(ElementType));
+			}
+		}
+		else if (Ar.IsLoading())
+		{
+			// Required for resetting ArrayNum
+			A.Empty(SerializeNum);
 
-				for (SizeType i=0; i<A.ArrayNum; i++)
-				{
-					Ar << A[i];
-				}
+			for (SizeType i=0; i<SerializeNum; i++)
+			{
+				Ar << *::new(A) ElementType;
 			}
 		}
 		else
 		{
-			Ar.SetError();
+			A.ArrayNum = SerializeNum;
+
+			for (SizeType i=0; i<A.ArrayNum; i++)
+			{
+				Ar << A[i];
+			}
 		}
 
 		return Ar;
