I/O - local cache interface

diff --git a/Engine/Source/Runtime/Core/Private/IO/IoCache.cpp b/Engine/Source/Runtime/Core/Private/IO/IoCache.cpp
new file mode 100644
--- /dev/null
+++ b/Engine/Source/Runtime/Core/Private/IO/IoCache.cpp
@@ -0,0 +1,57 @@
+// Copyright Epic Games, Inc. All Rights Reserved.
+
+#include "IO/IoCache.h"
+
+namespace UE::IO::Private
+{
+	FIoCacheRequestBase::FIoCacheRequestBase(FIoReadCallback&& ReadCallback)
+		: Callback(MoveTemp(ReadCallback))
+	{
+	}
+
+	void FIoCacheRequestBase::CompleteRequest(FIoBuffer&& Buffer)
+	{
+		if (EIoErrorCode::Unknown == ErrorCode.load(std::memory_order_seq_cst))
+		{
+			Callback(TIoStatusOr<FIoBuffer>(MoveTemp(Buffer)));
+			ErrorCode = EIoErrorCode::Ok;
+		}
+	}
+
+	void FIoCacheRequestBase::CompleteRequest(EIoErrorCode Error)
+	{
+		if (EIoErrorCode::Unknown == ErrorCode.load(std::memory_order_seq_cst))
+		{
+			Callback(FIoStatus(Error));
+			ErrorCode = Error; 
+		}
+	}
+
+} // namespace UE::IO::Private
+
+FIoCacheRequest::FIoCacheRequest(UE::IO::Private::FIoCacheRequestBase* Base)
+	: Pimpl(Base)
+{
+}
+
+FIoCacheRequest::~FIoCacheRequest()
+{
+	if (IsValid())
+	{
+		if (!Status().IsCompleted())
+		{
+			Pimpl->Cancel();
+			Pimpl->Wait();
+		}
+	}
+}
+
+void FIoCacheRequest::Cancel()
+{
+	check(IsValid());
+
+	if (!Status().IsCompleted())
+	{
+		Pimpl->Cancel();
+	}
+}
diff --git a/Engine/Source/Runtime/Core/Public/IO/IoCache.h b/Engine/Source/Runtime/Core/Public/IO/IoCache.h
new file mode 100644
--- /dev/null
+++ b/Engine/Source/Runtime/Core/Public/IO/IoCache.h
@@ -0,0 +1,95 @@
+// Copyright Epic Games, Inc. All Rights Reserved.
+
+#pragma once
+
+#include "IO/IoDispatcher.h"
+#include "Memory/MemoryFwd.h"
+#include "Templates/RefCounting.h"
+#include <atomic>
+
+namespace UE::IO::Private
+{
+	/** An abstract base class for implementing cache requests. */
+	class FIoCacheRequestBase
+	{
+	public:
+		virtual ~FIoCacheRequestBase() = default;
+
+		virtual void Wait() = 0;
+		virtual void Cancel() = 0;
+
+		inline FIoStatus Status()
+		{ 
+			return FIoStatus(ErrorCode.load(std::memory_order_relaxed));
+		}
+
+		inline uint32 GetRefCount() const
+		{
+			return RefCount.load(std::memory_order_relaxed);
+		}
+
+		inline void AddRef()
+		{
+			RefCount.fetch_add(1, std::memory_order_relaxed);
+		}
+
+		inline void Release()
+		{
+			if (1 == RefCount.fetch_sub(1, std::memory_order_acq_rel))
+			{
+				delete this;
+			}
+		}
+
+	protected:
+		explicit FIoCacheRequestBase(FIoReadCallback&& ReadCallback);
+		void CompleteRequest(FIoBuffer&& Buffer);
+		void CompleteRequest(EIoErrorCode Error);
+
+		FIoReadCallback Callback;
+		std::atomic_uint32_t RefCount{0};
+		std::atomic<EIoErrorCode> ErrorCode{EIoErrorCode::Unknown};
+	};
+}
+
+/** Represents a pending I/O cache request. */
+class CORE_API FIoCacheRequest
+{
+public:
+	FIoCacheRequest() = default;
+	FIoCacheRequest(const FIoCacheRequest&) = delete;
+	FIoCacheRequest(FIoCacheRequest&& Other) = default;
+	~FIoCacheRequest();
+	
+	FIoCacheRequest& operator=(const FIoCacheRequest&) = delete;
+	FIoCacheRequest& operator=(FIoCacheRequest&& Other) = default;
+
+	explicit FIoCacheRequest(UE::IO::Private::FIoCacheRequestBase* Base);
+
+	bool IsValid() const
+	{ 
+		return Pimpl.IsValid();
+	}
+
+	FIoStatus Status()
+	{
+		return IsValid() ? Pimpl->Status() : FIoStatus(EIoErrorCode::InvalidCode);
+	}
+
+	void Cancel();
+
+private:
+	TRefCountPtr<UE::IO::Private::FIoCacheRequestBase> Pimpl;
+};
+
+/**
+ * Interface for retrieving and storing I/O chunks identified by a 20 byte cache key.
+ */
+class IIoCache
+{
+public:
+	virtual ~IIoCache() = default;
+	virtual bool ContainsChunk(const FIoHash& Key) const = 0;
+	virtual FIoCacheRequest GetChunk(const FIoHash& Key, const FIoReadOptions& Options, FIoReadCallback&& Callback) = 0;
+	virtual FIoStatus PutChunk(const FIoHash& Key, FMemoryView Data) = 0;
+};
