AutoRTFM OpenCommit needs to capture by value.

diff --git a/Engine/Source/Runtime/Core/Public/AutoRTFM/AutoRTFM.h b/Engine/Source/Runtime/Core/Public/AutoRTFM/AutoRTFM.h
--- a/Engine/Source/Runtime/Core/Public/AutoRTFM/AutoRTFM.h
+++ b/Engine/Source/Runtime/Core/Public/AutoRTFM/AutoRTFM.h
@@ -518,7 +518,7 @@ struct FRegisterOpenFunction
 #define UE_AUTORTFM_OPEN_BEGIN AutoRTFM::Open([&]()
 #define UE_AUTORTFM_OPEN_END );
 
-#define UE_AUTORTFM_OPENCOMMIT_BEGIN AutoRTFM::OpenCommit([&]()
+#define UE_AUTORTFM_OPENCOMMIT_BEGIN AutoRTFM::OpenCommit([=]()
 #define UE_AUTORTFM_OPENCOMMIT_END );
 
 #define UE_AUTORTFM_TRANSACT_BEGIN AutoRTFM::Transact([&]()
