I/O - chunk encoding

diff --git a/Engine/Source/Runtime/Core/Private/IO/IoChunkEncoding.cpp b/Engine/Source/Runtime/Core/Private/IO/IoChunkEncoding.cpp
new file mode 100644
--- /dev/null
+++ b/Engine/Source/Runtime/Core/Private/IO/IoChunkEncoding.cpp
@@ -0,0 +1,367 @@
+// Copyright Epic Games, Inc. All Rights Reserved.
+
+#include "IO/IoChunkEncoding.h"
+
+#include "Async/ParallelFor.h"
+#include "Containers/ArrayView.h"
+#include "IO/IoDispatcher.h"
+#include "Math/UnrealMathUtility.h"
+#include "Misc/AES.h"
+#include "Misc/Compression.h"
+#include <atomic>
+
+///////////////////////////////////////////////////////////////////////////////////////////////////////////////////////
+
+bool FIoChunkEncoding::FHeader::IsValid() const
+{
+	return Magic == FIoChunkEncoding::ExpectedMagic
+		&& RawSize < FIoChunkEncoding::MaxSize
+		&& EncodedSize < FIoChunkEncoding::MaxSize;
+}
+
+uint32 FIoChunkEncoding::FHeader::GetBlockSize() const
+{
+	return uint32(1) << uint32(BlockSizeExponent);
+}
+
+uint32 FIoChunkEncoding::FHeader::GetBlockCount() const
+{
+	check(IsValid());
+	return static_cast<uint32>(FMath::DivideAndRoundUp(RawSize, uint64(GetBlockSize())));
+}
+
+uint64 FIoChunkEncoding::FHeader::GetTotalHeaderSize() const
+{
+	check(IsValid());
+	return sizeof(FHeader) + (GetBlockCount() * sizeof(uint32));
+}
+
+const FIoChunkEncoding::FHeader* FIoChunkEncoding::FHeader::Decode(FMemoryView HeaderData)
+{
+	if (HeaderData.GetSize() < sizeof(FHeader))
+	{
+		return nullptr;
+	}
+
+	const FHeader* Header = reinterpret_cast<const FHeader*>(HeaderData.GetData());
+	return !Header->IsValid() || HeaderData.GetSize() < Header->GetTotalHeaderSize() ? nullptr : Header;
+}
+
+///////////////////////////////////////////////////////////////////////////////////////////////////////////////////////
+
+bool FIoChunkEncoding::Encode(const FIoChunkEncodingParams& Params, FMemoryView RawData, FIoBuffer& OutEncodedData)
+{
+	FIoBuffer Encoding, Blocks;
+	if (Encode(Params, RawData,Encoding, Blocks))
+	{
+		OutEncodedData = FIoBuffer(Encoding.GetSize() + Blocks.GetSize()); 
+		OutEncodedData.GetMutableView().CopyFrom(Encoding.GetView());
+		OutEncodedData.GetMutableView().RightChop(Encoding.GetSize()).CopyFrom(Blocks.GetView());
+		return true;
+	}
+
+	return false;
+}
+
+bool FIoChunkEncoding::Encode(const FIoChunkEncodingParams& Params, FMemoryView RawData, FIoBuffer& OutHeader, FIoBuffer& OutEncodedBlocks)
+{
+	const uint32 BlockSize = Params.BlockSize;
+
+	check(IsAligned(BlockSize, FAES::AESBlockSize));
+	checkf(FMath::IsPowerOfTwo(BlockSize) && BlockSize <= MAX_uint32,
+		TEXT("BlockSize must be a 32-bit power of two but was %" UINT64_FMT "."), BlockSize);
+	checkf(Params.EncryptionKey.IsEmpty() || Params.EncryptionKey.GetSize() == FAES::FAESKey::KeySize, TEXT("Encryption key must be zero or %d bytes (AES)"), FAES::FAESKey::KeySize);
+
+	if (RawData.GetSize() == 0)
+	{
+		return false;
+	}
+
+	FAES::FAESKey AESKey;
+	if (!Params.EncryptionKey.IsEmpty())
+	{
+		FMutableMemoryView(AESKey.Key, FAES::FAESKey::KeySize).CopyFrom(Params.EncryptionKey);
+	}
+
+	const int64 BlockCount = FMath::DivideAndRoundUp(RawData.GetSize(), uint64(BlockSize));
+	check(BlockCount > 0);
+	checkf(BlockCount <= FIoChunkEncoding::MaxBlockCount, TEXT("Raw data of size %" UINT64_FMT " with block size %" UINT64_FMT " requires "
+		"%" UINT64_FMT " blocks, but the limit is %d."), RawData.GetSize(), BlockSize, BlockCount, FIoChunkEncoding::MaxBlockCount);
+
+	struct FBlock
+	{
+		FIoBuffer Buffer;
+		uint32 Size = 0;
+	};
+
+	std::atomic_uint64_t TotalEncodedSize{0};
+	TArray<FBlock> Blocks;
+	Blocks.SetNum(int32(BlockCount));
+
+	ParallelFor(int32(BlockCount),
+		[BlockSize, &RawData, &Params, &AESKey, &Blocks, &TotalEncodedSize]
+		(int32 BlockIndex)
+		{
+			const uint64 RawOffset = BlockIndex * BlockSize;
+			const uint64 RawBlockSize = FMath::Min(RawData.GetSize() - RawOffset, static_cast<uint64>(BlockSize));
+			FMemoryView RawBlock = RawData.Mid(RawOffset, BlockSize);
+
+			const uint64 RequiredBlockSize =
+			Align(FCompression::CompressMemoryBound(Params.CompressionFormat, static_cast<int32>(BlockSize)), FAES::AESBlockSize);
+
+			FBlock& Block = Blocks[BlockIndex];
+			Block.Buffer = FIoBuffer(RequiredBlockSize);
+
+			int32 CompressedSize = int32(Block.Buffer.GetSize());
+			const bool bCompressed = FCompression::CompressMemory(
+				Params.CompressionFormat,
+				Block.Buffer.GetData(),
+				CompressedSize,
+				RawBlock.GetData(),
+				static_cast<int32>(RawBlock.GetSize()),
+				COMPRESS_ForPackaging);
+
+			if (bCompressed == false)
+			{
+				Block.Size = 0;
+				return;
+			}
+
+			Block.Size = CompressedSize;
+			if (CompressedSize >= RawBlockSize)
+			{
+				Block.Buffer.GetMutableView().CopyFrom(RawBlock);
+				Block.Size = static_cast<uint32>(RawBlockSize);
+			}
+
+			const uint32 AESAlignedSize = Align(Block.Size, FAES::AESBlockSize);
+			{
+				uint8* Data = Block.Buffer.Data();
+				for (uint64 FillIndex = Block.Size; FillIndex < AESAlignedSize; ++FillIndex)
+				{
+					check(FillIndex < Block.Buffer.GetSize());
+					Data[FillIndex] = Data[(FillIndex - Block.Size) % Block.Size];
+				}
+			}
+
+			if (AESKey.IsValid())
+			{
+				FAES::EncryptData(Block.Buffer.GetData(), AESAlignedSize, AESKey);
+			}
+
+			TotalEncodedSize.fetch_add(AESAlignedSize);
+		});
+
+	const uint64 RequiredHeaderSize = sizeof(FHeader) + (BlockCount * sizeof(uint32)); 
+	if (OutHeader.GetSize() != RequiredHeaderSize)
+	{
+		OutHeader = FIoBuffer(RequiredHeaderSize);
+	}
+
+	const EIoEncryptionMethod Encryption = AESKey.IsValid() ? EIoEncryptionMethod::AES : EIoEncryptionMethod::None;
+
+	FHeader* Header = reinterpret_cast<FHeader*>(OutHeader.GetData());
+	Header->Magic = ExpectedMagic;
+	Header->RawSize = RawData.GetSize();
+	Header->EncodedSize = TotalEncodedSize;
+	Header->BlockSizeExponent = uint8(FMath::FloorLog2(BlockSize));
+	Header->Flags = static_cast<uint8>(Encryption); 
+	Header->Pad = 0;
+
+	TArrayView<uint32> BlockSizes = MakeArrayView(
+		reinterpret_cast<uint32*>(OutHeader.GetMutableView().RightChop(sizeof(FHeader)).GetData()),
+		int32(BlockCount));
+
+	OutEncodedBlocks = FIoBuffer(TotalEncodedSize);
+	FMutableMemoryView EncodedBlocks = OutEncodedBlocks.GetMutableView();
+
+	for (int32 BlockIndex = 0; BlockIndex < BlockCount; BlockIndex++)
+	{
+		FBlock& Block = Blocks[BlockIndex];
+		if (Block.Size == 0)
+		{
+			return false;
+		}
+
+		const uint32 AlignedBlockSize = Align(Block.Size, FAES::AESBlockSize);
+		BlockSizes[BlockIndex] = AlignedBlockSize;
+		EncodedBlocks.CopyFrom(Block.Buffer.GetView().Left(AlignedBlockSize));
+		EncodedBlocks += AlignedBlockSize;
+	}
+
+	check(EncodedBlocks.GetSize() == 0);
+	check(FHeader::Decode(OutHeader.GetView()) != nullptr);
+
+	return true;
+}
+
+///////////////////////////////////////////////////////////////////////////////////////////////////////////////////////
+
+bool FIoChunkEncoding::Decode(
+	FMemoryView EncodedData,
+	FName CompressionFormat,
+	FMemoryView EncryptionKey,
+	FMutableMemoryView OutRawData,
+	uint64 RawOffset)
+{
+	const FHeader* Header = FHeader::Decode(EncodedData);
+	if (Header == nullptr)
+	{
+		return false;
+	}
+
+	const uint32 BlockCount = Header->GetBlockCount();
+
+	FIoChunkDecodingParams Params;
+	Params.CompressionFormat = CompressionFormat;
+	Params.EncryptionKey = EncryptionKey;
+	Params.RawOffset = RawOffset;
+	Params.BlockSize = Header->GetBlockSize();
+	Params.TotalRawSize = Header->RawSize;
+	Params.EncodedBlockSize = TConstArrayView<uint32>(Header->Blocks, BlockCount);
+
+	const uint64 TotalHeaderSize = sizeof(FHeader) + (BlockCount * sizeof(uint32));
+	FMemoryView EncodedBlocks = EncodedData.RightChop(TotalHeaderSize);
+
+	return Decode(Params, EncodedBlocks, OutRawData);
+}
+
+bool FIoChunkEncoding::Decode(const FIoChunkDecodingParams& Params, FMemoryView EncodedBlocks, FMutableMemoryView OutRawData)
+{
+	if (Params.TotalRawSize < Params.RawOffset + OutRawData.GetSize())
+	{
+		return false;
+	}
+
+	checkf(Params.EncryptionKey.IsEmpty() || Params.EncryptionKey.GetSize() == FAES::FAESKey::KeySize, TEXT("Encryption key must be %d bytes (AES)"), FAES::FAESKey::KeySize);
+	FAES::FAESKey AESKey;
+	if (!Params.EncryptionKey.IsEmpty())
+	{
+		FMutableMemoryView(AESKey.Key, FAES::FAESKey::KeySize).CopyFrom(Params.EncryptionKey);
+	}
+
+	const TConstArrayView<uint32>& EncodedBlockSize = Params.EncodedBlockSize;
+	const uint32 BlockSize = Params.BlockSize;
+	const uint64 RawOffset = Params.RawOffset;
+	const uint32 BlockCount = EncodedBlockSize.Num();
+	const uint64 LastRawBlockSize = BlockSize - (BlockSize * BlockCount - Params.TotalRawSize);
+	const uint32 FirstBlockIndex = uint32(RawOffset / BlockSize);
+	const uint32 LastBlockIndex = uint32((RawOffset + OutRawData.GetSize() - 1) / BlockSize);
+
+	uint64 RawBlockOffset = RawOffset % BlockSize;
+	
+	uint64 EncodedOffset = 0;
+	for (uint32 BlockIndex = 0; BlockIndex < FirstBlockIndex; ++BlockIndex)
+	{
+		EncodedOffset += Align(EncodedBlockSize[BlockIndex], FAES::AESBlockSize);
+	}
+	// Subtract the encoded offset if the encoded blocks is a partial range of all encoded block(s)
+	EncodedBlocks += (EncodedOffset - Params.EncodedOffset);
+
+	for (uint32 BlockIndex = FirstBlockIndex; BlockIndex <= LastBlockIndex; BlockIndex++)
+	{
+		const uint64 RawBlockSize = BlockIndex == BlockCount - 1 ? LastRawBlockSize : BlockSize;
+		const uint64 RawBlockReadSize = FMath::Min(OutRawData.GetSize(), RawBlockSize - RawBlockOffset);
+		const uint32 CompressedBlockSize = EncodedBlockSize[BlockIndex];
+		const uint32 AlignedBlockSize = Align(CompressedBlockSize, FAES::AESBlockSize);
+
+		FIoBuffer Tmp(AlignedBlockSize);
+		Tmp.GetMutableView().CopyFrom(EncodedBlocks.Left(AlignedBlockSize));
+		if (AESKey.IsValid())
+		{
+			FAES::DecryptData(Tmp.GetData(), uint32(Tmp.GetSize()), AESKey);
+		}
+
+		if (CompressedBlockSize < RawBlockSize)
+		{
+			if (RawBlockReadSize == RawBlockSize)
+			{
+				if (!FCompression::UncompressMemory(Params.CompressionFormat, OutRawData.GetData(), int32(RawBlockReadSize), Tmp.GetData(), CompressedBlockSize))
+				{
+					return false;
+				}
+			}
+			else
+			{
+				FIoBuffer RawBlockTmp = FIoBuffer(RawBlockSize);
+				if (!FCompression::UncompressMemory(Params.CompressionFormat, RawBlockTmp.GetData(), int32(RawBlockSize), Tmp.GetData(), CompressedBlockSize))
+				{
+					return false;
+				}
+				OutRawData.CopyFrom(RawBlockTmp.GetView().Mid(RawBlockOffset, RawBlockReadSize));
+			}
+		}
+		else
+		{
+			OutRawData.CopyFrom(Tmp.GetView().Mid(RawBlockOffset, RawBlockReadSize));
+		}
+
+		RawBlockOffset = 0;
+		OutRawData += RawBlockReadSize;
+		EncodedBlocks += AlignedBlockSize;
+	}
+
+	check(OutRawData.GetSize() == 0);
+
+	return true;
+}
+
+///////////////////////////////////////////////////////////////////////////////////////////////////////////////////////
+
+bool FIoChunkEncoding::GetEncodedRange(
+	uint64 TotalRawSize,
+	uint32 RawBlockSize,
+	TConstArrayView<uint32> EncodedBlockSize,
+	uint64 RawOffset,
+	uint64 RawSize,
+	uint64& OutEncodedStart,
+	uint64& OutEncodedEnd)
+{
+	check(RawBlockSize > 0);
+	check(EncodedBlockSize.Num() > 0);
+
+	if (TotalRawSize < RawOffset + RawSize)
+	{
+		return false;
+	}
+
+	const uint64 TotalBlockCount = FMath::DivideAndRoundUp(TotalRawSize, uint64(RawBlockSize));
+	const uint32 FirstBlockIndex = uint32(RawOffset / RawBlockSize);
+	const uint32 LastBlockIndex = uint32((RawOffset + RawSize - 1) / RawBlockSize);
+
+	if (TotalBlockCount < FirstBlockIndex || TotalBlockCount < LastBlockIndex)
+	{
+		return false;
+	}
+
+	OutEncodedStart = 0;
+	uint32 BlockIndex = 0;
+	while (BlockIndex < FirstBlockIndex)
+	{
+		OutEncodedStart += Align(EncodedBlockSize[BlockIndex++], FAES::AESBlockSize);
+	}
+
+	OutEncodedEnd = OutEncodedStart;
+	while (BlockIndex <= LastBlockIndex)
+	{
+		OutEncodedEnd += Align(EncodedBlockSize[BlockIndex++], FAES::AESBlockSize);
+	}
+
+	return true;
+}
+
+bool FIoChunkEncoding::GetEncodedRange(const FIoChunkDecodingParams& Params, uint64 RawSize, uint64& OutEncodedStart, uint64& OutEncodedEnd)
+{
+	return GetEncodedRange(Params.TotalRawSize, Params.BlockSize, Params.EncodedBlockSize, Params.RawOffset, RawSize, OutEncodedStart, OutEncodedEnd);
+}
+
+uint64 FIoChunkEncoding::GetTotalEncodedSize(TConstArrayView<uint32> EncodedBlockSize)
+{
+	uint64 TotalEncodedSize = 0;
+	for (uint64 BlockSize : EncodedBlockSize)
+	{
+		TotalEncodedSize += Align(BlockSize, FAES::AESBlockSize);
+	}
+
+	return TotalEncodedSize;
+}
diff --git a/Engine/Source/Runtime/Core/Public/IO/IoChunkEncoding.h b/Engine/Source/Runtime/Core/Public/IO/IoChunkEncoding.h
new file mode 100644
--- /dev/null
+++ b/Engine/Source/Runtime/Core/Public/IO/IoChunkEncoding.h
@@ -0,0 +1,77 @@
+// Copyright Epic Games, Inc. All Rights Reserved.
+
+#pragma once
+
+#include "Memory/MemoryView.h"
+#include "UObject/NameTypes.h"
+
+class FIoBuffer;
+
+/** I/O chunk encryption method. */
+enum class EIoEncryptionMethod : uint8
+{
+	None	= 0,
+	AES 	= (1 << 0)
+};
+
+/** Defines how the I/O chunk gets encoded into a set of compressed and encrypted block(s). */
+struct FIoChunkEncodingParams
+{
+	FName CompressionFormat = TEXT("Oodle");
+	FMemoryView EncryptionKey;
+	uint32 BlockSize = (64 << 10);
+};
+
+/** Parameters for decoding a set of encoded blocks(s). */
+struct FIoChunkDecodingParams
+	: public FIoChunkEncodingParams
+{
+	uint64 TotalRawSize = 0;
+	uint64 RawOffset = 0;
+	uint64 EncodedOffset = 0;
+	TConstArrayView<uint32> EncodedBlockSize;
+};
+
+/**
+ * Encodes data into a set of encrypted and compressed blocks.
+ * The chunk encoding information is enocoded into a 16 byte
+ * header followed by N block sizes.
+ */
+class CORE_API FIoChunkEncoding
+{
+public:
+	static constexpr uint32 ExpectedMagic = 0x2e696f; // .io
+	static constexpr uint32 DefaultBlockSize = (64 << 10);
+	static constexpr uint32 MaxBlockCount = (1 << 24);
+	static constexpr uint64 MaxSize = (uint64(1) << 40);
+
+	/** Header describing the encoded I/O chunk. */
+	struct CORE_API FHeader
+	{
+		uint64 Magic: 24;
+		uint64 RawSize: 40;
+		uint64 EncodedSize: 40;
+		uint64 BlockSizeExponent: 8;
+		uint64 Flags: 8;
+		uint64 Pad: 8;
+		uint32 Blocks[];
+
+		bool IsValid() const;
+		uint32 GetBlockSize() const;
+		uint32 GetBlockCount() const;
+		uint64 GetTotalHeaderSize() const;
+
+		static const FHeader* Decode(FMemoryView HeaderData);
+	};
+
+	static_assert(sizeof(FHeader) == 16, "I/O chunk header size mismatch");
+
+	static bool Encode(const FIoChunkEncodingParams& Params, FMemoryView RawData, FIoBuffer& OutEncodedData);
+	static bool Encode(const FIoChunkEncodingParams& Params, FMemoryView RawData, FIoBuffer& OutHeader, FIoBuffer& OutEncodedBlocks);
+	static bool Decode(const FIoChunkDecodingParams& Params, FMemoryView EncodedBlocks, FMutableMemoryView OutRawData);
+	static bool Decode(FMemoryView EncodedData, FName CompressionFormat, FMemoryView EncryptionKey, FMutableMemoryView OutRawData, uint64 Offset = 0);
+
+	static bool GetEncodedRange(uint64 TotalRawSize, uint32 RawBlockSize, TConstArrayView<uint32> EncodedBlockSize, uint64 RawOffset, uint64 RawSize, uint64& OutEncodedStart, uint64& OutEncodedEnd);
+	static bool GetEncodedRange(const FIoChunkDecodingParams& Params, uint64 RawSize, uint64& OutEncodedStart, uint64& OutEncodedEnd);
+	static uint64 GetTotalEncodedSize(TConstArrayView<uint32> EncodedBlockSize);
+};
diff --git a/Engine/Source/Runtime/Core/Tests/IO/IoChunkEncodingTest.cpp b/Engine/Source/Runtime/Core/Tests/IO/IoChunkEncodingTest.cpp
new file mode 100644
--- /dev/null
+++ b/Engine/Source/Runtime/Core/Tests/IO/IoChunkEncodingTest.cpp
@@ -0,0 +1,93 @@
+// Copyright Epic Games, Inc. All Rights Reserved.
+
+#include "Memory/MemoryFwd.h"
+#if WITH_LOW_LEVEL_TESTS
+
+#include "Algo/Compare.h"
+#include "IO/IoChunkEncoding.h"
+#include "IO/IoDispatcher.h"
+#include "TestHarness.h"
+#include "TestMacros/Assertions.h"
+#include <catch2/generators/catch_generators.hpp>
+
+namespace UE::IO::ChunkEncodingTest
+{
+
+static FAES::FAESKey GenerateEncryptionKey()
+{
+	FAES::FAESKey Key;
+
+	for (int32 Idx = 0; Idx < FAES::FAESKey::KeySize; ++Idx)
+	{
+		Key.Key[Idx] = uint8(Idx + 1);
+	}
+
+	return Key;
+}
+
+TArray<uint64> GenerateData(int32 N)
+{
+	TArray<uint64> Data;
+	Data.SetNum(N);
+	for (int32 Idx = 0; Idx < Data.Num(); ++Idx)
+	{
+		Data[Idx] = uint64(Idx);
+	}
+	return Data;
+};
+
+TEST_CASE("Core::IO::IoChunk", "[Core][IO]")
+{
+	SECTION("Encode/Decode with offset")
+	{
+		const auto UncompressAndValidate =[](
+			FMemoryView EncodedData,
+			FName CompressionMethod,
+			const FAES::FAESKey& Key,
+			const int32 OffsetCount,
+			const int32 Count,
+			const TConstArrayView<uint64> ExpectedValues)
+		{
+			FIoBuffer RawData(Count * sizeof(uint64));
+
+			FMemoryView KeyView(Key.Key, Key.KeySize);
+			const bool bDecoded = FIoChunkEncoding::Decode(EncodedData, CompressionMethod, KeyView, RawData.GetMutableView(), OffsetCount * sizeof(uint64));
+			CHECK(bDecoded);
+
+			TConstArrayView<uint64> DecodedValues(reinterpret_cast<const uint64*>(RawData.GetData()), int32(RawData.GetSize() / sizeof(uint64)));
+			TConstArrayView<uint64> ExpectedRange = ExpectedValues.Mid(OffsetCount, Count);
+			CHECK(Algo::Compare(DecodedValues, ExpectedRange));
+		};
+
+		const FAES::FAESKey Key = GenerateEncryptionKey();
+		const FName Method = TEXT("Oodle");
+		constexpr uint64 BlockSize = 64 * sizeof(uint64);
+		constexpr int32 N = 5000;
+		const TArray<uint64> ExpectedValues = GenerateData(N);
+		
+		FIoBuffer Encoded;
+		const bool bEncoded = FIoChunkEncoding::Encode(
+			FIoChunkEncodingParams{Method, MakeMemoryView(Key.Key, Key.KeySize), BlockSize},
+			MakeMemoryView(ExpectedValues.GetData(), sizeof(uint64) * N),
+			Encoded);
+		
+		CHECK(bEncoded);
+		CHECK(Encoded.GetSize() > 0);
+
+		UncompressAndValidate(Encoded.GetView(), Method, Key, 0, N, ExpectedValues);
+		UncompressAndValidate(Encoded.GetView(), Method, Key, 1, N - 1, ExpectedValues);
+		UncompressAndValidate(Encoded.GetView(), Method, Key, N - 1, 1, ExpectedValues);
+		UncompressAndValidate(Encoded.GetView(), Method, Key, 0, 1, ExpectedValues);
+		UncompressAndValidate(Encoded.GetView(), Method, Key, 2, 4, ExpectedValues);
+		UncompressAndValidate(Encoded.GetView(), Method, Key, 0, 512, ExpectedValues);
+		UncompressAndValidate(Encoded.GetView(), Method, Key, 3, 514, ExpectedValues);
+		UncompressAndValidate(Encoded.GetView(), Method, Key, 256, 512, ExpectedValues);
+		UncompressAndValidate(Encoded.GetView(), Method, Key, 512, 512, ExpectedValues);
+		UncompressAndValidate(Encoded.GetView(), Method, Key, 512, 512, ExpectedValues);
+		UncompressAndValidate(Encoded.GetView(), Method, Key, 4993, 4, ExpectedValues);
+	}
+}
+
+} // namespace UE::IO::ChunkEncodingTest
+
+#endif // WITH_LOW_LEVEL_TESTS
