Alternative way to have the AutoRTFM macros disappear when not using the custom compiler - scoped for safety.

diff --git a/Engine/Source/Runtime/Core/Private/Delegates/DelegateHandle.cpp b/Engine/Source/Runtime/Core/Private/Delegates/DelegateHandle.cpp
--- a/Engine/Source/Runtime/Core/Private/Delegates/DelegateHandle.cpp
+++ b/Engine/Source/Runtime/Core/Private/Delegates/DelegateHandle.cpp
@@ -15,7 +15,7 @@ uint64 FDelegateHandle::GenerateNewID()
 	// Just increment a counter to generate an ID.
 	uint64 Result;
 	
-	UE_AUTORTFM_OPEN_BEGIN
+	UE_AUTORTFM_OPEN(
 	{
 		Result = ++UE::Delegates::Private::GNextID;
 
@@ -25,8 +25,7 @@ uint64 FDelegateHandle::GenerateNewID()
 			// Increment it again - it might not be zero, so don't just assign it to 1.
 			Result = ++UE::Delegates::Private::GNextID;
 		}
-	}
-	UE_AUTORTFM_OPEN_END
+	});
 
 	return Result;
 }
diff --git a/Engine/Source/Runtime/Core/Private/HAL/ThreadingBase.cpp b/Engine/Source/Runtime/Core/Private/HAL/ThreadingBase.cpp
--- a/Engine/Source/Runtime/Core/Private/HAL/ThreadingBase.cpp
+++ b/Engine/Source/Runtime/Core/Private/HAL/ThreadingBase.cpp
@@ -1342,7 +1342,7 @@ FQueuedThread::Run()
 FTlsAutoCleanup* FThreadSingletonInitializer::Get( TFunctionRef<FTlsAutoCleanup*()> CreateInstance, uint32& InOutTlsSlot )
 {
 	uint32 TlsSlot;
-	UE_AUTORTFM_OPEN_BEGIN
+	UE_AUTORTFM_OPEN(
 	{
 		TlsSlot = (uint32)FPlatformAtomics::AtomicRead_Relaxed((int32*)&InOutTlsSlot);
 		if (TlsSlot == 0xFFFFFFFF)
@@ -1360,8 +1360,7 @@ FTlsAutoCleanup* FThreadSingletonInitializer::Get( TFunctionRef<FTlsAutoCleanup*
 				TlsSlot = ThisTlsSlot;
 			}
 		}
-	}
-	UE_AUTORTFM_OPEN_END
+	});
 	
 	FTlsAutoCleanup* ThreadSingleton = (FTlsAutoCleanup*)FPlatformTLS::GetTlsValue( TlsSlot );
 	if( !ThreadSingleton )
diff --git a/Engine/Source/Runtime/Core/Private/UObject/UnrealNames.cpp b/Engine/Source/Runtime/Core/Private/UObject/UnrealNames.cpp
--- a/Engine/Source/Runtime/Core/Private/UObject/UnrealNames.cpp
+++ b/Engine/Source/Runtime/Core/Private/UObject/UnrealNames.cpp
@@ -1167,14 +1167,13 @@ public:
 	FNameEntryId Find(const FNameValue<Sensitivity>& Value) const
 	{
 		FNameEntryId Result;
-		UE_AUTORTFM_OPEN_BEGIN
+		UE_AUTORTFM_OPEN(
 		{
 			FRWScopeLock _(Lock, FRWScopeLockType::SLT_ReadOnly);
 
 			FNameSlot& Slot = Probe(Value);
 			Result = Slot.GetId();
-		}
-		UE_AUTORTFM_OPEN_END
+		});
 		return Result;
 	}
 
diff --git a/Engine/Source/Runtime/Core/Public/AutoRTFM/AutoRTFM.h b/Engine/Source/Runtime/Core/Public/AutoRTFM/AutoRTFM.h
--- a/Engine/Source/Runtime/Core/Public/AutoRTFM/AutoRTFM.h
+++ b/Engine/Source/Runtime/Core/Public/AutoRTFM/AutoRTFM.h
@@ -14,6 +14,14 @@
 	#define UE_PRAGMA_AUTORTFM
 #endif
 
+#if defined(_MSC_VER)
+#define UE_AUTORTFM_FORCEINLINE __forceinline
+#elif defined(__clang__)
+#define UE_AUTORTFM_FORCEINLINE inline __attribute__((always_inline))
+#else
+#define UE_AUTORTFM_FORCEINLINE inline
+#endif
+
 #include <stdbool.h>
 #include <stdlib.h>
 #include <memory.h>
@@ -62,7 +70,7 @@ typedef enum
 #if UE_AUTORTFM
 bool autortfm_is_transactional(void);
 #else
-inline bool autortfm_is_transactional(void)
+UE_AUTORTFM_FORCEINLINE bool autortfm_is_transactional(void)
 {
     return false;
 }
@@ -83,7 +91,7 @@ inline bool autortfm_is_transactional(void)
 #if UE_AUTORTFM
 bool autortfm_is_closed(void);
 #else
-inline bool autortfm_is_closed(void)
+UE_AUTORTFM_FORCEINLINE bool autortfm_is_closed(void)
 {
     return false;
 }
@@ -97,7 +105,7 @@ inline bool autortfm_is_closed(void)
 #if UE_AUTORTFM
 autortfm_result autortfm_transact(void (*work)(void* arg), void* arg);
 #else
-inline autortfm_result autortfm_transact(void (*work)(void* arg), void* arg)
+UE_AUTORTFM_FORCEINLINE autortfm_result autortfm_transact(void (*work)(void* arg), void* arg)
 {
 	work(arg);
     return autortfm_committed;
@@ -110,7 +118,7 @@ inline autortfm_result autortfm_transact(void (*work)(void* arg), void* arg)
 #if UE_AUTORTFM
 void autortfm_commit(void (*work)(void* arg), void* arg);
 #else
-inline void autortfm_commit(void (*work)(void* arg), void* arg)
+UE_AUTORTFM_FORCEINLINE void autortfm_commit(void (*work)(void* arg), void* arg)
 {
     abort();
 }
@@ -120,7 +128,7 @@ inline void autortfm_commit(void (*work)(void* arg), void* arg)
 #if UE_AUTORTFM
 bool autortfm_start_transaction();
 #else
-inline bool autortfm_start_transaction()
+UE_AUTORTFM_FORCEINLINE bool autortfm_start_transaction()
 {
 	return false;
 }
@@ -130,7 +138,7 @@ inline bool autortfm_start_transaction()
 #if UE_AUTORTFM
 autortfm_result autortfm_commit_transaction();
 #else
-inline autortfm_result autortfm_commit_transaction()
+UE_AUTORTFM_FORCEINLINE autortfm_result autortfm_commit_transaction()
 {
 	return autortfm_aborted_by_language;
 }
@@ -140,7 +148,7 @@ inline autortfm_result autortfm_commit_transaction()
 #if UE_AUTORTFM
 autortfm_result autortfm_abort_transaction();
 #else
-inline autortfm_result autortfm_abort_transaction()
+UE_AUTORTFM_FORCEINLINE autortfm_result autortfm_abort_transaction()
 {
 	return autortfm_aborted_by_request;
 }
@@ -150,7 +158,7 @@ inline autortfm_result autortfm_abort_transaction()
 #if UE_AUTORTFM
 void autortfm_clear_transaction_status();
 #else
-inline void autortfm_clear_transaction_status()
+UE_AUTORTFM_FORCEINLINE void autortfm_clear_transaction_status()
 {
 }
 #endif
@@ -162,7 +170,7 @@ inline void autortfm_clear_transaction_status()
 #if UE_AUTORTFM
 void autortfm_abort_if_transactional(void);
 #else
-inline void autortfm_abort_if_transactional(void) { }
+UE_AUTORTFM_FORCEINLINE void autortfm_abort_if_transactional(void) { }
 #endif
 
 // Abort if running in closed code. This does an abort_by_language, which will
@@ -170,7 +178,7 @@ inline void autortfm_abort_if_transactional(void) { }
 #if UE_AUTORTFM
 void autortfm_abort_if_closed(void);
 #else
-inline void autortfm_abort_if_closed(void) { }
+UE_AUTORTFM_FORCEINLINE void autortfm_abort_if_closed(void) { }
 #endif
 
 // Executes the given code non-transactionally regardless of whether we are in
@@ -178,7 +186,7 @@ inline void autortfm_abort_if_closed(void) { }
 #if UE_AUTORTFM
 void autortfm_open(void (*work)(void* arg), void* arg);
 #else
-inline void autortfm_open(void (*work)(void* arg), void* arg)
+UE_AUTORTFM_FORCEINLINE void autortfm_open(void (*work)(void* arg), void* arg)
 {
     work(arg);
 }
@@ -191,7 +199,7 @@ inline void autortfm_open(void (*work)(void* arg), void* arg)
 #if UE_AUTORTFM
 [[nodiscard]] autortfm_status autortfm_close(void (*work)(void* arg), void* arg);
 #else
-inline autortfm_status autortfm_close(void (*work)(void* arg), void* arg)
+[[nodiscard]] UE_AUTORTFM_FORCEINLINE autortfm_status autortfm_close(void (*work)(void* arg), void* arg)
 {
     abort();
 	return autortfm_status_aborted_by_language;
@@ -205,7 +213,7 @@ inline autortfm_status autortfm_close(void (*work)(void* arg), void* arg)
 #if UE_AUTORTFM
 void autortfm_record_open_write(void* Ptr, size_t Size);
 #else
-inline void autortfm_record_open_write(void* Ptr, size_t Size)
+UE_AUTORTFM_FORCEINLINE void autortfm_record_open_write(void* Ptr, size_t Size)
 {
 }
 #endif
@@ -237,7 +245,7 @@ inline void autortfm_record_open_write(void* Ptr, size_t Size)
 #if UE_AUTORTFM
 void autortfm_register_open_function(void* original_function, void* new_function);
 #else
-inline void autortfm_register_open_function(void* original_function, void* new_function) { }
+UE_AUTORTFM_FORCEINLINE void autortfm_register_open_function(void* original_function, void* new_function) { }
 #endif
 
 // Have some work happen when this transaction commits. For nested transactions,
@@ -247,7 +255,7 @@ inline void autortfm_register_open_function(void* original_function, void* new_f
 #if UE_AUTORTFM
 void autortfm_defer_until_commit(void (*work)(void* arg), void* arg);
 #else
-inline void autortfm_defer_until_commit(void (*work)(void* arg), void* arg)
+UE_AUTORTFM_FORCEINLINE void autortfm_defer_until_commit(void (*work)(void* arg), void* arg)
 {
     work(arg);
 }
@@ -260,7 +268,7 @@ inline void autortfm_defer_until_commit(void (*work)(void* arg), void* arg)
 #if UE_AUTORTFM
 void autortfm_defer_until_abort(void (*work)(void* arg), void* arg);
 #else
-inline void autortfm_defer_until_abort(void (*work)(void* arg), void* arg) { }
+UE_AUTORTFM_FORCEINLINE void autortfm_defer_until_abort(void (*work)(void* arg), void* arg) { }
 #endif
 
 // Have some work happen when this transaction commits. For nested transactions,
@@ -270,7 +278,7 @@ inline void autortfm_defer_until_abort(void (*work)(void* arg), void* arg) { }
 #if UE_AUTORTFM
 void autortfm_open_commit(void (*work)(void* arg), void* arg);
 #else
-inline void autortfm_open_commit(void (*work)(void* arg), void* arg)
+UE_AUTORTFM_FORCEINLINE void autortfm_open_commit(void (*work)(void* arg), void* arg)
 {
     work(arg);
 }
@@ -281,7 +289,7 @@ inline void autortfm_open_commit(void (*work)(void* arg), void* arg)
 #if UE_AUTORTFM
 void autortfm_open_abort(void (*work)(void* arg), void* arg);
 #else
-inline void autortfm_open_abort(void (*work)(void* arg), void* arg) { }
+UE_AUTORTFM_FORCEINLINE void autortfm_open_abort(void (*work)(void* arg), void* arg) { }
 #endif
 
 // Inform the runtime that we have performed a new object allocation. It's only
@@ -293,7 +301,7 @@ inline void autortfm_open_abort(void (*work)(void* arg), void* arg) { }
 #if UE_AUTORTFM
 void* autortfm_did_allocate(void* ptr, size_t size);
 #else
-inline void* autortfm_did_allocate(void* ptr, size_t size)
+UE_AUTORTFM_FORCEINLINE void* autortfm_did_allocate(void* ptr, size_t size)
 {
     return ptr;
 }
@@ -307,7 +315,7 @@ inline void* autortfm_did_allocate(void* ptr, size_t size)
 #if UE_AUTORTFM
 void autortfm_check_consistency_assuming_no_races(void);
 #else
-inline void autortfm_check_consistency_assuming_no_races(void) { }
+UE_AUTORTFM_FORCEINLINE void autortfm_check_consistency_assuming_no_races(void) { }
 #endif
 
 // If running with AutoRTFM enabled, then perform an ABI check between the
@@ -319,7 +327,7 @@ inline void autortfm_check_consistency_assuming_no_races(void) { }
 #if UE_AUTORTFM
 void autortfm_check_abi(void* ptr, size_t size);
 #else
-inline void autortfm_check_abi(void* ptr, size_t size) { }
+UE_AUTORTFM_FORCEINLINE void autortfm_check_abi(void* ptr, size_t size) { }
 #endif
 
 #ifdef __cplusplus
@@ -345,11 +353,11 @@ enum class EContextStatus
 	AbortedByRequest = autortfm_status_aborted_by_request
 };
 
-inline bool IsTransactional() { return autortfm_is_transactional(); }
-inline bool IsClosed() { return autortfm_is_closed(); }
+UE_AUTORTFM_FORCEINLINE bool IsTransactional() { return autortfm_is_transactional(); }
+UE_AUTORTFM_FORCEINLINE bool IsClosed() { return autortfm_is_closed(); }
 
 template<typename TFunctor>
-ETransactionResult Transact(const TFunctor& Functor)
+UE_AUTORTFM_FORCEINLINE ETransactionResult Transact(const TFunctor& Functor)
 {
     return static_cast<ETransactionResult>(autortfm_transact(
         [] (void* Arg) { (*static_cast<const TFunctor*>(Arg))(); },
@@ -357,73 +365,73 @@ ETransactionResult Transact(const TFunctor& Functor)
 }
 
 template<typename TFunctor>
-void Commit(const TFunctor& Functor)
+UE_AUTORTFM_FORCEINLINE void Commit(const TFunctor& Functor)
 {
     autortfm_commit(
         [] (void* Arg) { (*static_cast<const TFunctor*>(Arg))(); },
         const_cast<void*>(static_cast<const void*>(&Functor)));
 }
 
-inline bool StartTransaction()
+UE_AUTORTFM_FORCEINLINE bool StartTransaction()
 {
 	return autortfm_start_transaction();
 }
 
-inline ETransactionResult CommitTransaction()
+UE_AUTORTFM_FORCEINLINE ETransactionResult CommitTransaction()
 {
 	return static_cast<ETransactionResult>(autortfm_commit_transaction());
 }
 
-inline ETransactionResult AbortTransaction()
+UE_AUTORTFM_FORCEINLINE ETransactionResult AbortTransaction()
 {
 	return static_cast<ETransactionResult>(autortfm_abort_transaction());
 }
 
-inline void ClearTransactionStatus()
+UE_AUTORTFM_FORCEINLINE void ClearTransactionStatus()
 {
 	autortfm_clear_transaction_status();
 }
 
 // RecordOpenWrite records the memory span into the current transaction as written.
 //  If this memory is previously unknown to the transaction, the original value is saved.
-inline void RecordOpenWrite(void* Ptr, size_t Size)
+UE_AUTORTFM_FORCEINLINE void RecordOpenWrite(void* Ptr, size_t Size)
 {
 	autortfm_record_open_write(Ptr, Size);
 }
 
 template<typename TTYPE>
-inline void RecordOpenWrite(TTYPE* Ptr)
+UE_AUTORTFM_FORCEINLINE void RecordOpenWrite(TTYPE* Ptr)
 {
 	autortfm_record_open_write(Ptr, sizeof(TTYPE));
 }
 
 // RecordOpenRead does nothing right now, but it is intended as a support stub for the day when we move to full AutoSTM
-inline void RecordOpenRead(void const* Ptr, size_t Size)
+UE_AUTORTFM_FORCEINLINE void RecordOpenRead(void const* Ptr, size_t Size)
 {
 }
 
 template<typename TTYPE>
-inline void RecordOpenRead(TTYPE* Ptr)
+UE_AUTORTFM_FORCEINLINE void RecordOpenRead(TTYPE* Ptr)
 {
 	RecordOpenRead(Ptr, sizeof(TTYPE));
 }
 
 // WriteMemory first records the memory span as written (see RecordOpenWrite) and then copies the specified value into it.
-inline void WriteMemory(void* DestPtr, void const* SrcPtr, size_t Size)
+UE_AUTORTFM_FORCEINLINE void WriteMemory(void* DestPtr, void const* SrcPtr, size_t Size)
 {
 	RecordOpenWrite(DestPtr, Size);
 	memcpy(DestPtr, SrcPtr, Size);
 }
 
 template<typename TTYPE>
-inline void WriteMemoryTrivial(TTYPE* DestPtr, TTYPE const* SrcPtr)
+UE_AUTORTFM_FORCEINLINE void WriteMemoryTrivial(TTYPE* DestPtr, TTYPE const* SrcPtr)
 {
 	RecordOpenWrite(DestPtr, sizeof(TTYPE));
 	*DestPtr = *SrcPtr;
 }
 
 template<typename TTYPE>
-inline void WriteMemory(TTYPE* DestPtr, TTYPE const* SrcPtr)
+UE_AUTORTFM_FORCEINLINE void WriteMemory(TTYPE* DestPtr, TTYPE const* SrcPtr)
 {
 	if constexpr (std::is_trivially_copyable<TTYPE>::value)
 	{
@@ -436,7 +444,7 @@ inline void WriteMemory(TTYPE* DestPtr, TTYPE const* SrcPtr)
 }
 
 template<typename TTYPE>
-inline void WriteMemory(TTYPE* DestPtr, TTYPE const SrcValue)
+UE_AUTORTFM_FORCEINLINE void WriteMemory(TTYPE* DestPtr, TTYPE const SrcValue)
 {
 	if constexpr (std::is_trivially_copyable<TTYPE>::value)
 	{
@@ -448,18 +456,18 @@ inline void WriteMemory(TTYPE* DestPtr, TTYPE const SrcValue)
 	}
 }
 
-inline void AbortIfTransactional()
+UE_AUTORTFM_FORCEINLINE void AbortIfTransactional()
 {
     autortfm_abort_if_transactional();
 }
 
-inline void AbortIfClosed()
+UE_AUTORTFM_FORCEINLINE void AbortIfClosed()
 {
     autortfm_abort_if_closed();
 }
 
 template<typename TFunctor>
-void Open(const TFunctor& Functor)
+UE_AUTORTFM_FORCEINLINE void Open(const TFunctor& Functor)
 {
     autortfm_open(
         [] (void* Arg) { (*static_cast<const TFunctor*>(Arg))(); },
@@ -467,14 +475,14 @@ void Open(const TFunctor& Functor)
 }
 
 template<typename TFunctor>
-[[nodiscard]] EContextStatus Close(const TFunctor& Functor)
+[[nodiscard]] UE_AUTORTFM_FORCEINLINE EContextStatus Close(const TFunctor& Functor)
 {
     return static_cast<EContextStatus>(autortfm_close(
         [] (void* Arg) { (*static_cast<const TFunctor*>(Arg))(); },
         const_cast<void*>(static_cast<const void*>(&Functor))));
 }
 
-inline void RegisterOpenFunction(void* OriginalFunction, void* NewFunction)
+UE_AUTORTFM_FORCEINLINE void RegisterOpenFunction(void* OriginalFunction, void* NewFunction)
 {
     autortfm_register_open_function(OriginalFunction, NewFunction);
 }
@@ -486,21 +494,21 @@ void OpenCommit(TFunction<void()>&& Work);
 void OpenAbort(TFunction<void()>&& Work);
 #else
 template<typename TFunctor>
-void DeferUntilCommit(const TFunctor& Work) { Work(); }
+UE_AUTORTFM_FORCEINLINE void DeferUntilCommit(const TFunctor& Work) { Work(); }
 template<typename TFunctor>
-void DeferUntilAbort(const TFunctor& Work) { }
+UE_AUTORTFM_FORCEINLINE void DeferUntilAbort(const TFunctor& Work) { }
 template<typename TFunctor>
-void OpenCommit(const TFunctor& Work) { Work(); }
+UE_AUTORTFM_FORCEINLINE void OpenCommit(const TFunctor& Work) { Work(); }
 template<typename TFunctor>
-void OpenAbort(const TFunctor& Work) { }
+UE_AUTORTFM_FORCEINLINE void OpenAbort(const TFunctor& Work) { }
 #endif
 
-inline void* DidAllocate(void* Ptr, size_t Size)
+UE_AUTORTFM_FORCEINLINE void* DidAllocate(void* Ptr, size_t Size)
 {
     return autortfm_did_allocate(Ptr, Size);
 }
 
-inline void CheckConsistencyAssumingNoRaces()
+UE_AUTORTFM_FORCEINLINE void CheckConsistencyAssumingNoRaces()
 {
     autortfm_check_consistency_assuming_no_races();
 }
@@ -515,23 +523,13 @@ struct FRegisterOpenFunction
 
 // Macro-based variants so we completely compile away when not in use, even in debug builds
 #if UE_AUTORTFM
-#define UE_AUTORTFM_OPEN_BEGIN AutoRTFM::Open([&]()
-#define UE_AUTORTFM_OPEN_END );
-
-#define UE_AUTORTFM_OPENCOMMIT_BEGIN AutoRTFM::OpenCommit([=]()
-#define UE_AUTORTFM_OPENCOMMIT_END );
-
-#define UE_AUTORTFM_TRANSACT_BEGIN AutoRTFM::Transact([&]()
-#define UE_AUTORTFM_TRANSACT_END );
+#define UE_AUTORTFM_OPEN(x) AutoRTFM::Open([&]() { x })
+#define UE_AUTORTFM_OPENCOMMIT(x) AutoRTFM::OpenCommit([=]() { x })
+#define UE_AUTORTFM_TRANSACT(x) AutoRTFM::Transact([&]() { x })
 #else
-#define UE_AUTORTFM_OPEN_BEGIN
-#define UE_AUTORTFM_OPEN_END
-
-#define UE_AUTORTFM_OPENCOMMIT_BEGIN
-#define UE_AUTORTFM_OPENCOMMIT_END
-
-#define UE_AUTORTFM_TRANSACT_BEGIN
-#define UE_AUTORTFM_TRANSACT_END
+#define UE_AUTORTFM_OPEN(x) do { x } while (false)
+#define UE_AUTORTFM_OPENCOMMIT(x) do { x } while (false)
+#define UE_AUTORTFM_TRANSACT(x) do { x } while (false)
 #endif
 
 #define UE_AUTORTFM_CONCAT_IMPL(A, B) A ## B
diff --git a/Engine/Source/Runtime/Core/Public/HAL/FMemory.inl b/Engine/Source/Runtime/Core/Public/HAL/FMemory.inl
--- a/Engine/Source/Runtime/Core/Public/HAL/FMemory.inl
+++ b/Engine/Source/Runtime/Core/Public/HAL/FMemory.inl
@@ -24,7 +24,7 @@ FMEMORY_INLINE_FUNCTION_DECORATOR void* FMemory::Malloc(SIZE_T Count, uint32 Ali
 	// we don't end up keeping track of the writes to the allocator's internal data structures.
 	// This is because allocators are already transactional - malloc can be rolled back by
 	// calling free.
-	UE_AUTORTFM_OPEN_BEGIN
+	UE_AUTORTFM_OPEN(
 	{
 		if (!FMEMORY_INLINE_GMalloc)
 		{
@@ -38,8 +38,7 @@ FMEMORY_INLINE_FUNCTION_DECORATOR void* FMemory::Malloc(SIZE_T Count, uint32 Ali
 		}
 		// optional tracking of every allocation
 		LLM_IF_ENABLED(FLowLevelMemTracker::Get().OnLowLevelAlloc(ELLMTracker::Default, Ptr, Count, ELLMTag::Untagged, ELLMAllocType::FMalloc));
-	}
-	UE_AUTORTFM_OPEN_END
+	});
 
 	// AutoRTFM: This is a no-op for non-transactional code.
 	// For transactional code, this defers a call to Free if the transaction aborts,
@@ -120,7 +119,7 @@ FMEMORY_INLINE_FUNCTION_DECORATOR void FMemory::Free(void* Original)
 	// AutoRTFM: For transactional code, in order to support the transaction 
 	// aborting and needing to 'roll back' the Free, we defer the actual
 	// free until commit time.
-	UE_AUTORTFM_OPENCOMMIT_BEGIN
+	UE_AUTORTFM_OPENCOMMIT(
 	{
 		// optional tracking of every allocation
 		LLM_IF_ENABLED(FLowLevelMemTracker::Get().OnLowLevelFree(ELLMTracker::Default, Original, ELLMAllocType::FMalloc));
@@ -133,14 +132,13 @@ FMEMORY_INLINE_FUNCTION_DECORATOR void FMemory::Free(void* Original)
 		DoGamethreadHook(2);
 		FScopedMallocTimer Timer(2);
 		FMEMORY_INLINE_GMalloc->Free(Original);
-	}
-	UE_AUTORTFM_OPENCOMMIT_END
+	});
 }
 
 FMEMORY_INLINE_FUNCTION_DECORATOR SIZE_T FMemory::GetAllocSize(void* Original)
 {
 	SIZE_T Result;
-	UE_AUTORTFM_OPEN_BEGIN
+	UE_AUTORTFM_OPEN(
 	{
 		if (!FMEMORY_INLINE_GMalloc)
 		{
@@ -151,8 +149,7 @@ FMEMORY_INLINE_FUNCTION_DECORATOR SIZE_T FMemory::GetAllocSize(void* Original)
 			SIZE_T Size = 0;
 			Result = FMEMORY_INLINE_GMalloc->GetAllocationSize(Original, Size) ? Size : 0;
 		}
-	}
-	UE_AUTORTFM_OPEN_END
+	});
 
 	return Result;
 }
@@ -160,7 +157,7 @@ FMEMORY_INLINE_FUNCTION_DECORATOR SIZE_T FMemory::GetAllocSize(void* Original)
 FMEMORY_INLINE_FUNCTION_DECORATOR SIZE_T FMemory::QuantizeSize(SIZE_T Count, uint32 Alignment)
 {
 	SIZE_T Result;
-	UE_AUTORTFM_OPEN_BEGIN
+	UE_AUTORTFM_OPEN(
 	{
 		if (!FMEMORY_INLINE_GMalloc)
 		{
@@ -170,8 +167,7 @@ FMEMORY_INLINE_FUNCTION_DECORATOR SIZE_T FMemory::QuantizeSize(SIZE_T Count, uin
 		{
 			Result = FMEMORY_INLINE_GMalloc->QuantizeSize(Count, Alignment);
 		}
-	}
-	UE_AUTORTFM_OPEN_END
+	});
 
 	return Result;
 }
diff --git a/Engine/Source/Runtime/Core/Public/Templates/SharedPointerInternals.h b/Engine/Source/Runtime/Core/Public/Templates/SharedPointerInternals.h
--- a/Engine/Source/Runtime/Core/Public/Templates/SharedPointerInternals.h
+++ b/Engine/Source/Runtime/Core/Public/Templates/SharedPointerInternals.h
@@ -187,7 +187,7 @@ namespace SharedPointerInternals
 		{
 			if constexpr (Mode == ESPMode::ThreadSafe)
 			{
-				UE_AUTORTFM_OPENCOMMIT_BEGIN
+				UE_AUTORTFM_OPENCOMMIT(
 				{
 					// std::memory_order_acq_rel is used here so that, if we do end up executing the destructor, it's not possible
 					// for side effects from executing the destructor end up being visible before we've determined that the shared
@@ -204,8 +204,7 @@ namespace SharedPointerInternals
 						// reference count reaches zero, this object will be deleted.
 						ReleaseWeakReference();
 					}
-				}
-				UE_AUTORTFM_OPENCOMMIT_END
+				});
 			}
 			else
 			{
diff --git a/Engine/Source/Runtime/CoreUObject/Private/UObject/Class.cpp b/Engine/Source/Runtime/CoreUObject/Private/UObject/Class.cpp
--- a/Engine/Source/Runtime/CoreUObject/Private/UObject/Class.cpp
+++ b/Engine/Source/Runtime/CoreUObject/Private/UObject/Class.cpp
@@ -5869,7 +5869,7 @@ UFunction* UClass::FindFunctionByName(FName InName, EIncludeSuperFlag::Type Incl
 
 	UFunction* Result = nullptr;
 
-	UE_AUTORTFM_OPEN_BEGIN
+	UE_AUTORTFM_OPEN(
 	{
 		{
 			FReadScopeLock ScopeLock(FuncMapLock);
@@ -5912,8 +5912,7 @@ UFunction* UClass::FindFunctionByName(FName InName, EIncludeSuperFlag::Type Incl
 				}
 			}
 		}
-	}
-	UE_AUTORTFM_OPEN_END
+	});
 
 	return Result;
 }
diff --git a/Engine/Source/Runtime/CoreUObject/Private/UObject/ObjectHandleTracking.cpp b/Engine/Source/Runtime/CoreUObject/Private/UObject/ObjectHandleTracking.cpp
--- a/Engine/Source/Runtime/CoreUObject/Private/UObject/ObjectHandleTracking.cpp
+++ b/Engine/Source/Runtime/CoreUObject/Private/UObject/ObjectHandleTracking.cpp
@@ -20,73 +20,68 @@ namespace UE::CoreUObject
 
 			void OnHandleRead(TArrayView<const UObject* const> Objects)
 			{
-				UE_AUTORTFM_OPEN_BEGIN
+				UE_AUTORTFM_OPEN(
 				{
 					FReadScopeLock _(HandleLock);
 					for (auto&& Pair : ReadHandleCallbacks)
 					{
 						Pair.Value(Objects);
 					}
-				}
-				UE_AUTORTFM_OPEN_END
+				});
 			}
 
 			void OnClassReferenceResolved(const FObjectRef& ObjectRef, UPackage* Package, UClass* Class)
 			{
-				UE_AUTORTFM_OPEN_BEGIN
+				UE_AUTORTFM_OPEN(
 				{
 					FReadScopeLock _(HandleLock);
 					for (auto&& Pair : ClassResolvedCallbacks)
 					{
 						Pair.Value(ObjectRef, Package, Class);
 					}
-				}
-				UE_AUTORTFM_OPEN_END
+				});
 			}
 
 			void OnReferenceResolved(const FObjectRef& ObjectRef, UPackage* Package, UObject* Object)
 			{
-				UE_AUTORTFM_OPEN_BEGIN
+				UE_AUTORTFM_OPEN(
 				{
 					FReadScopeLock _(HandleLock);
 					for (auto&& Pair : HandleResolvedCallbacks)
 					{
 						Pair.Value(ObjectRef, Package, Object);
 					}
-				}
-				UE_AUTORTFM_OPEN_END
+				});
 			}
 
 			void OnReferenceLoaded(const FObjectRef& ObjectRef, UPackage* Package, UObject* Object)
 			{
-				UE_AUTORTFM_OPEN_BEGIN
+				UE_AUTORTFM_OPEN(
 				{
 					FReadScopeLock _(HandleLock);
 					for (auto&& Pair : HandleLoadedCallbacks)
 					{
 						Pair.Value(ObjectRef, Package, Object);
 					}
-				}
-				UE_AUTORTFM_OPEN_END
+				});
 			}
 
 			FObjectHandleTrackingCallbackId AddObjectHandleReadCallback(FObjectHandleReadFunc Func)
 			{
 				FObjectHandleTrackingCallbackId Result;
-				UE_AUTORTFM_OPEN_BEGIN
+				UE_AUTORTFM_OPEN(
 				{
 					FWriteScopeLock _(HandleLock);
 					NextHandleId++;
 					ReadHandleCallbacks.Add({ NextHandleId, Func });
 					Result = FObjectHandleTrackingCallbackId{ NextHandleId };
-				}
-				UE_AUTORTFM_OPEN_END
+				});
 				return Result;
 			}
 
 			void RemoveObjectHandleReadCallback(FObjectHandleTrackingCallbackId Handle)
 			{
-				UE_AUTORTFM_OPEN_BEGIN
+				UE_AUTORTFM_OPEN(
 				{
 					FWriteScopeLock _(HandleLock);
 					for (int32 i = ReadHandleCallbacks.Num() - 1; i >= 0; --i)
@@ -97,27 +92,25 @@ namespace UE::CoreUObject
 							ReadHandleCallbacks.RemoveAt(i);
 						}
 					}
-				}
-				UE_AUTORTFM_OPEN_END
+				});
 			}
 
 			FObjectHandleTrackingCallbackId AddObjectHandleClassResolvedCallback(FObjectHandleClassResolvedFunc Func)
 			{
 				FObjectHandleTrackingCallbackId Result;
-				UE_AUTORTFM_OPEN_BEGIN
+				UE_AUTORTFM_OPEN(
 				{
 					FWriteScopeLock _(HandleLock);
 					NextHandleId++;
 					ClassResolvedCallbacks.Add({ NextHandleId, Func });
 					Result = FObjectHandleTrackingCallbackId{ NextHandleId };
-				}
-				UE_AUTORTFM_OPEN_END
+				});
 				return Result;
 			}
 
 			void RemoveObjectHandleClassResolvedCallback(FObjectHandleTrackingCallbackId Handle)
 			{
-				UE_AUTORTFM_OPEN_BEGIN
+				UE_AUTORTFM_OPEN(
 				{
 					FWriteScopeLock _(HandleLock);
 					for (int32 i = ClassResolvedCallbacks.Num() - 1; i >= 0; --i)
@@ -128,27 +121,25 @@ namespace UE::CoreUObject
 							ClassResolvedCallbacks.RemoveAt(i);
 						}
 					}
-				}
-				UE_AUTORTFM_OPEN_END
+				});
 			}
 
 			FObjectHandleTrackingCallbackId AddObjectHandleReferenceResolvedCallback(FObjectHandleReferenceResolvedFunc Func)
 			{
 				FObjectHandleTrackingCallbackId Result;
-				UE_AUTORTFM_OPEN_BEGIN
+				UE_AUTORTFM_OPEN(
 				{
 					FWriteScopeLock _(HandleLock);
 					NextHandleId++;
 					HandleResolvedCallbacks.Add({ NextHandleId, Func });
 					Result = FObjectHandleTrackingCallbackId{ NextHandleId };
-				}
-				UE_AUTORTFM_OPEN_END
+				});
 				return Result;
 			}
 
 			void RemoveObjectHandleReferenceResolvedCallback(FObjectHandleTrackingCallbackId Handle)
 			{
-				UE_AUTORTFM_OPEN_BEGIN
+				UE_AUTORTFM_OPEN(
 				{
 					FWriteScopeLock _(HandleLock);
 					for (int32 i = HandleResolvedCallbacks.Num() - 1; i >= 0; --i)
@@ -159,27 +150,25 @@ namespace UE::CoreUObject
 							HandleResolvedCallbacks.RemoveAt(i);
 						}
 					}
-				}
-				UE_AUTORTFM_OPEN_END
+				});
 			}
 
 			FObjectHandleTrackingCallbackId AddObjectHandleReferenceLoadedCallback(FObjectHandleReferenceLoadedFunc Func)
 			{
 				FObjectHandleTrackingCallbackId Result;
-				UE_AUTORTFM_OPEN_BEGIN
+				UE_AUTORTFM_OPEN(
 				{
 					FWriteScopeLock _(HandleLock);
 					NextHandleId++;
 					HandleLoadedCallbacks.Add({ NextHandleId, Func });
 					Result = FObjectHandleTrackingCallbackId{ NextHandleId };
-				}
-				UE_AUTORTFM_OPEN_END
+				});
 				return Result;
 			}
 
 			void RemoveObjectHandleReferenceLoadedCallback(FObjectHandleTrackingCallbackId Handle)
 			{
-				UE_AUTORTFM_OPEN_BEGIN
+				UE_AUTORTFM_OPEN(
 				{
 					FWriteScopeLock _(HandleLock);
 					for (int32 i = HandleLoadedCallbacks.Num() - 1; i >= 0; --i)
@@ -190,8 +179,7 @@ namespace UE::CoreUObject
 							HandleLoadedCallbacks.RemoveAt(i);
 						}
 					}
-				}
-				UE_AUTORTFM_OPEN_END
+				});
 			}
 
 		private:
diff --git a/Engine/Source/Runtime/CoreUObject/Private/UObject/ScriptCore.cpp b/Engine/Source/Runtime/CoreUObject/Private/UObject/ScriptCore.cpp
--- a/Engine/Source/Runtime/CoreUObject/Private/UObject/ScriptCore.cpp
+++ b/Engine/Source/Runtime/CoreUObject/Private/UObject/ScriptCore.cpp
@@ -189,7 +189,7 @@ FBlueprintContext* FBlueprintContextGetThreadSingletonImpl()
 FBlueprintContext* FBlueprintContext::GetThreadSingleton()
 {
 	FBlueprintContext* Result;
-	UE_AUTORTFM_OPEN_BEGIN { Result = FBlueprintContextGetThreadSingletonImpl(); } UE_AUTORTFM_OPEN_END
+	UE_AUTORTFM_OPEN({ Result = FBlueprintContextGetThreadSingletonImpl(); });
 	return Result;
 }
 
diff --git a/Engine/Source/Runtime/CoreUObject/Private/UObject/UObjectGlobals.cpp b/Engine/Source/Runtime/CoreUObject/Private/UObject/UObjectGlobals.cpp
--- a/Engine/Source/Runtime/CoreUObject/Private/UObject/UObjectGlobals.cpp
+++ b/Engine/Source/Runtime/CoreUObject/Private/UObject/UObjectGlobals.cpp
@@ -3854,11 +3854,10 @@ void FObjectInitializer::PostConstructInit()
 	{
 		// autortfm todo: if this transaction aborts and we are in a transaction's open nest,
 		// we need to have a way of propagating out that abort
-		UE_AUTORTFM_TRANSACT_BEGIN
+		UE_AUTORTFM_TRANSACT(
 		{
 			PropertyInitCallback();
-		}
-		UE_AUTORTFM_TRANSACT_END
+		});
 	}
 	// After the call to `PropertyInitCallback` to allow the callback to modify the instancing graph
 	if (bNeedInstancing || bNeedSubobjectInstancing)
