Refactored FTextId to reduce it's memory footprint

diff --git a/Engine/Source/Runtime/Core/Public/Internationalization/TextKey.h b/Engine/Source/Runtime/Core/Public/Internationalization/TextKey.h
--- a/Engine/Source/Runtime/Core/Public/Internationalization/TextKey.h
+++ b/Engine/Source/Runtime/Core/Public/Internationalization/TextKey.h
@@ -127,6 +127,10 @@ private:
 
 	/** Hash of this text key */
 	uint32 StrHash;
+
+	FTextKey(const TCHAR* Str, uint32 Hash) : StrPtr(Str), StrHash(Hash) {}
+
+	friend class FTextId;
 };
 
 /**
@@ -135,106 +139,152 @@ private:
 class CORE_API FTextId
 {
 public:
-	FTextId() = default;
+	FTextId()
+	{
+		Reset();
+	}
 
 	FTextId(const FTextKey& InNamespace, const FTextKey& InKey)
-		: Namespace(InNamespace)
-		, Key(InKey)
+		: NamespaceStr(InNamespace.StrPtr)
+		, KeyStr(InKey.StrPtr)
+		, NamespaceHash(InNamespace.StrHash)
+		, KeyHash(InKey.StrHash)
 	{
 	}
 
 	/** Get the namespace component of this text identity */
-	FORCEINLINE const FTextKey& GetNamespace() const
+	FORCEINLINE FTextKey GetNamespace() const
 	{
-		return Namespace;
+		return FTextKey(NamespaceStr, NamespaceHash);
 	}
 
 	/** Get the key component of this text identity */
-	FORCEINLINE const FTextKey& GetKey() const
+	FORCEINLINE FTextKey GetKey() const
 	{
-		return Key;
+		return FTextKey(KeyStr, KeyHash);
 	}
 
 	/** Compare for equality */
 	friend FORCEINLINE bool operator==(const FTextId& A, const FTextId& B)
 	{
-		return A.Namespace == B.Namespace && A.Key == B.Key;
+		return A.NamespaceStr == B.NamespaceStr && A.KeyStr == B.KeyStr;
 	}
 
 	/** Compare for inequality */
 	friend FORCEINLINE bool operator!=(const FTextId& A, const FTextId& B)
 	{
-		return A.Namespace != B.Namespace || A.Key != B.Key;
+		return A.NamespaceStr != B.NamespaceStr || A.KeyStr != B.KeyStr;
 	}
 
 	/** Get the hash of this text identity */
 	friend FORCEINLINE uint32 GetTypeHash(const FTextId& A)
 	{
-		return HashCombine(GetTypeHash(A.Namespace), GetTypeHash(A.Key));
+		return HashCombine(A.NamespaceHash, A.KeyHash);
 	}
 
 	/** Serialize this text identity as if it were FStrings */
 	void SerializeAsString(FArchive& Ar)
 	{
+		FTextKey Namespace = FTextKey(NamespaceStr, NamespaceHash);
 		Namespace.SerializeAsString(Ar);
+		NamespaceStr = Namespace.StrPtr;
+		NamespaceHash = Namespace.StrHash;
+
+		FTextKey Key = FTextKey(KeyStr, KeyHash);
 		Key.SerializeAsString(Ar);
+		KeyStr = Key.StrPtr;
+		KeyHash = Key.StrHash;
 	}
 
 	/** Serialize this text identity including its hash values (this method is sensitive to hashing algorithm changes, so only use it for generated files that can be rebuilt from another source) */
 	void SerializeWithHash(FArchive& Ar)
 	{
+		FTextKey Namespace = FTextKey(NamespaceStr, NamespaceHash);
 		Namespace.SerializeWithHash(Ar);
+		NamespaceStr = Namespace.StrPtr;
+		NamespaceHash = Namespace.StrHash;
+
+		FTextKey Key = FTextKey(KeyStr, KeyHash);
 		Key.SerializeWithHash(Ar);
+		KeyStr = Key.StrPtr;
+		KeyHash = Key.StrHash;
 	}
 
 	/** Serialize this text identity including its hash values, discarding the hash on load (to upgrade from an older hashing algorithm) */
 	void SerializeDiscardHash(FArchive& Ar)
 	{
+		FTextKey Namespace = FTextKey(NamespaceStr, NamespaceHash);
 		Namespace.SerializeDiscardHash(Ar);
+		NamespaceStr = Namespace.StrPtr;
+		NamespaceHash = Namespace.StrHash;
+
+		FTextKey Key = FTextKey(KeyStr, KeyHash);
 		Key.SerializeDiscardHash(Ar);
+		KeyStr = Key.StrPtr;
+		KeyHash = Key.StrHash;
 	}
 
 	/** Serialize this text identity as if it were FStrings */
 	void SerializeAsString(FStructuredArchiveSlot Slot)
 	{
 		FStructuredArchiveRecord Record = Slot.EnterRecord();
+		FTextKey Namespace = FTextKey(NamespaceStr, NamespaceHash);
 		Namespace.SerializeAsString(Record.EnterField(TEXT("Namespace")));
+		NamespaceStr = Namespace.StrPtr;
+		NamespaceHash = Namespace.StrHash;
+
+		FTextKey Key = FTextKey(KeyStr, KeyHash);
 		Key.SerializeAsString(Record.EnterField(TEXT("Key")));
+		KeyStr = Key.StrPtr;
+		KeyHash = Key.StrHash;
 	}
 
 	/** Serialize this text identity including its hash values (this method is sensitive to hashing algorithm changes, so only use it for generated files that can be rebuilt from another source) */
 	void SerializeWithHash(FStructuredArchiveSlot Slot)
 	{
 		FStructuredArchiveRecord Record = Slot.EnterRecord();
+		FTextKey Namespace = FTextKey(NamespaceStr, NamespaceHash);
 		Namespace.SerializeWithHash(Record.EnterField(TEXT("Namespace")));
+		NamespaceStr = Namespace.StrPtr;
+		NamespaceHash = Namespace.StrHash;
+
+		FTextKey Key = FTextKey(KeyStr, KeyHash);
 		Key.SerializeWithHash(Record.EnterField(TEXT("Key")));
+		KeyStr = Key.StrPtr;
+		KeyHash = Key.StrHash;
 	}
 
 	/** Serialize this text identity including its hash values, discarding the hash on load (to upgrade from an older hashing algorithm) */
 	void SerializeDiscardHash(FStructuredArchiveSlot Slot)
 	{
 		FStructuredArchiveRecord Record = Slot.EnterRecord();
+		FTextKey Namespace = FTextKey(NamespaceStr, NamespaceHash);
 		Namespace.SerializeDiscardHash(Record.EnterField(TEXT("Namespace")));
+		NamespaceStr = Namespace.StrPtr;
+		NamespaceHash = Namespace.StrHash;
+
+		FTextKey Key = FTextKey(KeyStr, KeyHash);
 		Key.SerializeDiscardHash(Record.EnterField(TEXT("Key")));
+		KeyStr = Key.StrPtr;
+		KeyHash = Key.StrHash;
 	}
 
 	/** Is this text identity empty? */
 	FORCEINLINE bool IsEmpty() const
 	{
-		return Namespace.IsEmpty() && Key.IsEmpty();
+		return *NamespaceStr == 0 && *KeyStr == 0;
 	}
 
 	/** Reset this text identity to be empty */
 	FORCEINLINE void Reset()
 	{
-		Namespace.Reset();
-		Key.Reset();
+		NamespaceStr = KeyStr = TEXT("");
+		NamespaceHash = KeyHash = 0;
 	}
 
 private:
-	/** Namespace component of this text identity */
-	FTextKey Namespace;
-
-	/** Key component of this text identity */
-	FTextKey Key;
+	const TCHAR* NamespaceStr;
+	const TCHAR* KeyStr;
+	uint32 NamespaceHash;
+	uint32 KeyHash;
 };
