Moved #includes around in the shared pchs so the pchs are more heirarchal. This fixes issues where removing or adding a private module dependency would drastically change compile times because the includes in the pch were very different.

Before:
Total CPU Time: 69956.296875 s
Total time in Parallel executor: 1706.01 seconds
Total execution time: 1733.63 seconds
0:28:57

After:
Total time in Parallel executor: 1652.02 seconds
Total execution time: 1679.59 seconds
0:28:02

diff --git a/Engine/Source/Editor/UnrealEd/Public/UnrealEdSharedPCH.h b/Engine/Source/Editor/UnrealEd/Public/UnrealEdSharedPCH.h
--- a/Engine/Source/Editor/UnrealEd/Public/UnrealEdSharedPCH.h
+++ b/Engine/Source/Editor/UnrealEd/Public/UnrealEdSharedPCH.h
@@ -4,49 +4,12 @@
 
 #include "EngineSharedPCH.h"
 
-// From Core:
-#include "Math/BasicMathExpressionEvaluator.h"
-#include "Misc/ExpressionParserTypes.h"
-#include "Misc/FilterCollection.h"
-#include "Misc/IFilter.h"
-
-// From CoreUObject:
-#include "Misc/NotifyHook.h"
-#include "UObject/ObjectKey.h"
-
-// From SlateCore:
-#include "Layout/LayoutUtils.h"
-
-// From Slate:
-#include "Framework/Docking/LayoutService.h"
-#include "Framework/MarqueeRect.h"
-#include "Widgets/Input/NumericTypeInterface.h"
-#include "Widgets/Notifications/SErrorText.h"
-#include "Widgets/SToolTip.h"
-
 // From Messaging:
 #include "IMessageContext.h"
 
 // From AssetRegistry:
 #include "AssetRegistry/ARFilter.h"
 
-// From Engine:
-#include "BoneIndices.h"
-#include "DataTableUtils.h"
-#include "EdGraph/EdGraphNodeUtils.h"
-#include "Engine/CurveTable.h"
-#include "Engine/DataAsset.h"
-#include "Engine/DataTable.h"
-#include "Engine/LevelStreaming.h"
-#include "LatentActions.h"
-#include "NavFilters/NavigationQueryFilter.h"
-#include "PhysicalMaterials/PhysicalMaterial.h"
-#include "PhysicsSettingsEnums.h"
-#include "PreviewScene.h"
-#include "Vehicles/TireType.h"
-#include "VisualLogger/VisualLogger.h"
-#include "VisualLogger/VisualLoggerTypes.h"
-
 // From BlueprintGraph:
 #include "BlueprintNodeSignature.h"
 #include "K2Node.h"
@@ -58,6 +21,9 @@
 #include "GameplayTaskTypes.h"
 #include "GameplayTask.h"
 
+// From NavigationSystem:
+#include "NavFilters/NavigationQueryFilter.h"
+
 // From UnrealEd:
 #include "AssetThumbnail.h"
 #include "Editor.h"
diff --git a/Engine/Source/Runtime/Core/Public/CoreSharedPCH.h b/Engine/Source/Runtime/Core/Public/CoreSharedPCH.h
--- a/Engine/Source/Runtime/Core/Public/CoreSharedPCH.h
+++ b/Engine/Source/Runtime/Core/Public/CoreSharedPCH.h
@@ -5,6 +5,8 @@
 // From Core:
 #include "Algo/Reverse.h"
 #include "Async/AsyncWork.h"
+#include "Async/Future.h"
+#include "Async/ParallelFor.h"
 #include "Async/TaskGraphInterfaces.h"
 #include "Containers/Array.h"
 #include "Containers/ArrayView.h"
@@ -12,6 +14,7 @@
 #include "Containers/ChunkedArray.h"
 #include "Containers/ContainerAllocationPolicies.h"
 #include "Containers/ContainersFwd.h"
+#include "Containers/DynamicRHIResourceArray.h"
 #include "Containers/EnumAsByte.h"
 #include "Containers/IndirectArray.h"
 #include "Containers/List.h"
@@ -19,6 +22,7 @@
 #include "Containers/LockFreeList.h"
 #include "Containers/Map.h"
 #include "Containers/Queue.h"
+#include "Containers/ResourceArray.h"
 #include "Containers/ScriptArray.h"
 #include "Containers/Set.h"
 #include "Containers/SparseArray.h"
@@ -36,6 +40,8 @@
 #include "Delegates/IDelegateInstance.h"
 #include "Delegates/IntegerSequence.h"
 #include "Delegates/MulticastDelegateBase.h"
+#include "Features/IModularFeature.h"
+#include "Features/IModularFeatures.h"
 #include "GenericPlatform/GenericPlatformAffinity.h"
 #include "GenericPlatform/GenericPlatformAtomics.h"
 #include "GenericPlatform/GenericPlatformFile.h"
@@ -53,6 +59,7 @@
 #include "HAL/FileManager.h"
 #include "HAL/IConsoleManager.h"
 #include "HAL/MemoryBase.h"
+#include "HAL/PlatformAffinity.h"
 #include "HAL/PlatformAtomics.h"
 #include "HAL/PlatformCrt.h"
 #include "HAL/PlatformMath.h"
@@ -64,19 +71,25 @@
 #include "HAL/PlatformTime.h"
 #include "HAL/PlatformTLS.h"
 #include "HAL/Runnable.h"
+#include "HAL/RunnableThread.h"
+#include "HAL/ThreadSafeBool.h"
 #include "HAL/ThreadSafeCounter.h"
 #include "HAL/ThreadSingleton.h"
 #include "HAL/TlsAutoCleanup.h"
 #include "HAL/UnrealMemory.h"
 #include "Internationalization/CulturePointer.h"
+#include "Internationalization/GatherableTextData.h"
 #include "Internationalization/Internationalization.h"
+#include "Internationalization/InternationalizationMetadata.h"
 #include "Internationalization/Text.h"
 #include "Internationalization/TextLocalizationManager.h"
 #include "Internationalization/TextNamespaceFwd.h"
 #include "Logging/LogCategory.h"
 #include "Logging/LogMacros.h"
 #include "Logging/LogVerbosity.h"
+#include "Logging/TokenizedMessage.h"
 #include "Math/Axis.h"
+#include "Math/BasicMathExpressionEvaluator.h"
 #include "Math/Box.h"
 #include "Math/Box2D.h"
 #include "Math/BoxSphereBounds.h"
@@ -106,6 +119,7 @@
 #include "Math/Plane.h"
 #include "Math/Quat.h"
 #include "Math/QuatRotationTranslationMatrix.h"
+#include "Math/RandomStream.h"
 #include "Math/Range.h"
 #include "Math/RangeBound.h"
 #include "Math/RangeSet.h"
@@ -115,8 +129,11 @@
 #include "Math/Rotator.h"
 #include "Math/ScaleMatrix.h"
 #include "Math/ScaleRotationTranslationMatrix.h"
+#include "Math/SHMath.h"
 #include "Math/Sphere.h"
 #include "Math/Transform.h"
+#include "Math/TransformCalculus.h"
+#include "Math/TransformCalculus2D.h"
 #include "Math/TranslationMatrix.h"
 #include "Math/TwoVectors.h"
 #include "Math/UnrealMath.h"
@@ -128,9 +145,13 @@
 #include "Math/VectorRegister.h"
 #include "Misc/App.h"
 #include "Misc/AssertionMacros.h"
+#include "Misc/Attribute.h"
+#include "Misc/AutomationTest.h"
+#include "Misc/BufferedOutputDevice.h"
 #include "Misc/ByteSwap.h"
 #include "Misc/Char.h"
 #include "Misc/CommandLine.h"
+#include "Misc/CompilationResult.h"
 #include "Misc/Compression.h"
 #include "Misc/ConfigCacheIni.h"
 #include "Misc/CoreDelegates.h"
@@ -139,23 +160,33 @@
 #include "Misc/Crc.h"
 #include "Misc/CString.h"
 #include "Misc/DateTime.h"
+#include "Misc/EngineVersion.h"
 #include "Misc/EngineVersionBase.h"
 #include "Misc/EnumClassFlags.h"
 #include "Misc/Exec.h"
+#include "Misc/ExpressionParserTypes.h"
 #include "Misc/FeedbackContext.h"
 #include "Misc/FileHelper.h"
+#include "Misc/FilterCollection.h"
 #include "Misc/Guid.h"
+#include "Misc/IFilter.h"
 #include "Misc/IQueuedWork.h"
 #include "Misc/MemStack.h"
+#include "Misc/MessageDialog.h"
+#include "Misc/NetworkGuid.h"
 #include "Misc/NoopCounter.h"
+#include "Misc/ObjectThumbnail.h"
 #include "Misc/Optional.h"
 #include "Misc/OutputDevice.h"
+#include "Misc/OutputDeviceError.h"
 #include "Misc/OutputDeviceRedirector.h"
 #include "Misc/Parse.h"
 #include "Misc/Paths.h"
 #include "Misc/QueuedThreadPool.h"
+#include "Misc/ScopedEvent.h"
 #include "Misc/ScopeLock.h"
 #include "Misc/SecureHash.h"
+#include "Misc/SingleThreadRunnable.h"
 #include "Misc/SlowTask.h"
 #include "Misc/SlowTaskStack.h"
 #include "Misc/StructBuilder.h"
@@ -164,11 +195,20 @@
 #include "Modules/Boilerplate/ModuleBoilerplate.h"
 #include "Modules/ModuleInterface.h"
 #include "Modules/ModuleManager.h"
+#include "ProfilingDebugging/Histogram.h"
+#include "ProfilingDebugging/ProfilingHelpers.h"
+#include "ProfilingDebugging/ResourceSize.h"
 #include "Serialization/Archive.h"
+#include "Serialization/ArchiveProxy.h"
+#include "Serialization/BitReader.h"
+#include "Serialization/BitWriter.h"
 #include "Serialization/BufferReader.h"
+#include "Serialization/CustomVersion.h"
 #include "Serialization/MemoryArchive.h"
+#include "Serialization/MemoryReader.h"
 #include "Serialization/MemoryWriter.h"
 #include "Stats/Stats.h"
+#include "Stats/StatsMisc.h"
 #include "Templates/AlignmentTemplates.h"
 #include "Templates/AndOrNot.h"
 #include "Templates/ChooseClass.h"
@@ -196,6 +236,7 @@
 #include "Templates/RefCounting.h"
 #include "Templates/RemoveExtent.h"
 #include "Templates/RemoveReference.h"
+#include "Templates/ScopedCallback.h"
 #include "Templates/SharedPointer.h"
 #include "Templates/Sorting.h"
 #include "Templates/Tuple.h"
@@ -205,10 +246,13 @@
 #include "Templates/UniquePtr.h"
 #include "Templates/UnrealTemplate.h"
 #include "Templates/UnrealTypeTraits.h"
+#include "Templates/ValueOrError.h"
 #include "Traits/IsContiguousContainer.h"
+#include "UObject/DebugSerializationFlags.h"
 #include "UObject/NameTypes.h"
 #include "UObject/ObjectVersion.h"
+#include "UObject/PropertyPortFlags.h"
 #include "UObject/ScriptDelegates.h"
 #include "UObject/UnrealNames.h"
 #include "UObject/UObjectHierarchyFwd.h"
-#include "UObject/WeakObjectPtrTemplates.h"
+#include "UObject/WeakObjectPtrTemplates.h"
\ No newline at end of file
diff --git a/Engine/Source/Runtime/CoreUObject/Public/CoreUObjectSharedPCH.h b/Engine/Source/Runtime/CoreUObject/Public/CoreUObjectSharedPCH.h
--- a/Engine/Source/Runtime/CoreUObject/Public/CoreUObjectSharedPCH.h
+++ b/Engine/Source/Runtime/CoreUObject/Public/CoreUObjectSharedPCH.h
@@ -4,31 +4,50 @@
 
 #include "CoreSharedPCH.h"
 
-// From Core:
-#include "HAL/PlatformAffinity.h"
-#include "HAL/RunnableThread.h"
-#include "HAL/ThreadSafeBool.h"
-#include "Math/RandomStream.h"
-#include "Math/TransformCalculus.h"
-#include "Math/TransformCalculus2D.h"
-#include "Misc/Attribute.h"
-#include "ProfilingDebugging/ResourceSize.h"
-#include "Serialization/MemoryReader.h"
-
 // From CoreUObject:
+#include "Misc/NotifyHook.h"
+#include "Misc/PackageName.h"
+#include "Misc/WorldCompositionUtility.h"
 #include "Serialization/ArchiveUObject.h"
+#include "Serialization/BulkData.h"
+#include "Serialization/SerializedPropertyScope.h"
+#include "Templates/Casts.h"
+#include "Templates/SubclassOf.h"
 #include "UObject/Class.h"
 #include "UObject/CoreNative.h"
+#include "UObject/CoreNet.h"
+#include "UObject/CoreNetTypes.h"
 #include "UObject/GarbageCollection.h"
 #include "UObject/GCObject.h"
 #include "UObject/GeneratedCppIncludes.h"
 #include "UObject/Interface.h"
+#include "UObject/LazyObjectPtr.h"
+#include "UObject/Linker.h"
+#include "UObject/LinkerLoad.h"
 #include "UObject/Object.h"
+#include "UObject/ObjectKey.h"
 #include "UObject/ObjectMacros.h"
+#include "UObject/ObjectRedirector.h"
+#include "UObject/ObjectResource.h"
+#include "UObject/Package.h"
+#include "UObject/PackageFileSummary.h"
+#include "UObject/PersistentObjectPtr.h"
+#include "UObject/PropertyTag.h"
 #include "UObject/Script.h"
+#include "UObject/ScriptInterface.h"
+#include "UObject/ScriptMacros.h"
+#include "UObject/SoftObjectPtr.h"
+#include "UObject/Stack.h"
+#include "UObject/StructOnScope.h"
+#include "UObject/TextProperty.h"
+#include "UObject/UnrealType.h"
+#include "UObject/UObjectAnnotation.h"
 #include "UObject/UObjectArray.h"
 #include "UObject/UObjectBase.h"
 #include "UObject/UObjectBaseUtility.h"
 #include "UObject/UObjectGlobals.h"
+#include "UObject/UObjectHash.h"
+#include "UObject/UObjectIterator.h"
 #include "UObject/UObjectMarks.h"
-#include "UObject/WeakObjectPtr.h"
+#include "UObject/UObjectThreadContext.h"
+#include "UObject/WeakObjectPtr.h"
\ No newline at end of file
diff --git a/Engine/Source/Runtime/Engine/Public/EngineSharedPCH.h b/Engine/Source/Runtime/Engine/Public/EngineSharedPCH.h
--- a/Engine/Source/Runtime/Engine/Public/EngineSharedPCH.h
+++ b/Engine/Source/Runtime/Engine/Public/EngineSharedPCH.h
@@ -4,36 +4,25 @@
 
 #include "SlateSharedPCH.h"
 
-// From Core:
-#include "Async/Future.h"
-#include "Async/ParallelFor.h"
-#include "Containers/DynamicRHIResourceArray.h"
-#include "Containers/ResourceArray.h"
-#include "Features/IModularFeatures.h"
+// From ApplicationCore:
+#include "GenericPlatform/GenericApplication.h"
+#include "GenericPlatform/GenericApplicationMessageHandler.h"
+#include "GenericPlatform/GenericWindow.h"
+#include "GenericPlatform/GenericWindowDefinition.h"
 #include "GenericPlatform/ICursor.h"
 #include "GenericPlatform/IInputInterface.h"
-#include "HAL/PlatformAffinity.h"
-#include "HAL/RunnableThread.h"
-#include "HAL/ThreadSafeBool.h"
-#include "Internationalization/GatherableTextData.h"
-#include "Internationalization/InternationalizationMetadata.h"
-#include "Math/SHMath.h"
-#include "Misc/AutomationTest.h"
-#include "Misc/BufferedOutputDevice.h"
-#include "Misc/CompilationResult.h"
-#include "Misc/NetworkGuid.h"
-#include "Misc/ScopedEvent.h"
-#include "Misc/SingleThreadRunnable.h"
-#include "ProfilingDebugging/Histogram.h"
-#include "ProfilingDebugging/ProfilingHelpers.h"
-#include "Serialization/ArchiveProxy.h"
-#include "Serialization/BitReader.h"
-#include "Serialization/BitWriter.h"
-#include "Stats/StatsMisc.h"
-#include "Templates/ScopedCallback.h"
-#include "Templates/ValueOrError.h"
-#include "UObject/DebugSerializationFlags.h"
-#include "UObject/PropertyPortFlags.h"
+
+// From AssetRegistry:
+#include "AssetRegistry/AssetData.h"
+
+// From AudioExtensions:
+#include "IAudioExtensionPlugin.h"
+
+// From CoreOnline:
+#include "Online/CoreOnline.h"
+
+// From DeveloperSettings:
+#include "Engine/DeveloperSettings.h"
 
 // From Json:
 #include "Dom/JsonObject.h"
@@ -44,33 +33,13 @@
 #include "Serialization/JsonSerializer.h"
 #include "Serialization/JsonSerializerMacros.h"
 
-// From CoreUObject:
-#include "Online/CoreOnline.h"
-#include "Serialization/BulkData.h"
-#include "Templates/Casts.h"
-#include "Templates/SubclassOf.h"
-#include "UObject/CoreNet.h"
-#include "UObject/Linker.h"
-#include "UObject/LinkerLoad.h"
-#include "UObject/ObjectRedirector.h"
-#include "UObject/ObjectResource.h"
-#include "UObject/PackageFileSummary.h"
-#include "UObject/PersistentObjectPtr.h"
-#include "UObject/ScriptMacros.h"
-#include "UObject/Stack.h"
-#include "UObject/StructOnScope.h"
-#include "UObject/TextProperty.h"
-#include "UObject/UObjectAnnotation.h"
-#include "UObject/UObjectHash.h"
-#include "UObject/UObjectIterator.h"
-#include "UObject/UObjectThreadContext.h"
-
 // From RHI:
 #include "RHIDefinitions.h"
 #include "RHI.h"
 #include "RHIStaticStates.h"
 
 // From RenderCore:
+#include "GlobalShader.h"
 #include "RenderCommandFence.h"
 #include "RenderResource.h"
 #if UE_ENABLE_INCLUDE_ORDER_DEPRECATED_IN_5_2
@@ -85,13 +54,15 @@
 #include "Shader.h"
 #include "VertexFactory.h"
 #include "ShaderParameterUtils.h"
-
-// From AssetRegistry:
-#include "AssetRegistry/AssetData.h"
+#include "StaticBoundShaderState.h"
 
 // From PacketHandler:
 #include "PacketHandler.h"
 
+// From PhysicsCore:
+#include "PhysicalMaterials/PhysicalMaterial.h"
+#include "PhysicsSettingsEnums.h"
+
 // From Engine:
 #include "AI/Navigation/NavAgentInterface.h"
 #include "AI/Navigation/NavigationTypes.h"
@@ -123,6 +94,7 @@
 #include "BlendableManager.h"
 #include "BlueprintUtilities.h"
 #include "BoneContainer.h"
+#include "BoneIndices.h"
 #include "BonePose.h"
 #include "Camera/CameraShakeBase.h"
 #include "Camera/CameraTypes.h"
@@ -147,9 +119,11 @@
 #include "Curves/KeyHandle.h"
 #include "Curves/RichCurve.h"
 #include "CustomBoneIndexArray.h"
+#include "DataTableUtils.h"
 #include "DebugViewModeHelpers.h"
 #include "EdGraph/EdGraph.h"
 #include "EdGraph/EdGraphNode.h"
+#include "EdGraph/EdGraphNodeUtils.h"
 #include "EdGraph/EdGraphPin.h"
 #include "EdGraph/EdGraphSchema.h"
 #include "Engine/BlendableInterface.h"
@@ -159,8 +133,10 @@
 #include "Engine/Brush.h"
 #include "Engine/Channel.h"
 #include "Engine/ChildConnection.h"
+#include "Engine/CurveTable.h"
+#include "Engine/DataAsset.h"
+#include "Engine/DataTable.h"
 #include "Engine/DebugDisplayProperty.h"
-#include "Engine/DeveloperSettings.h"
 #include "Engine/Engine.h"
 #include "Engine/EngineBaseTypes.h"
 #include "Engine/EngineTypes.h"
@@ -169,6 +145,7 @@
 #include "Engine/GameViewportDelegates.h"
 #include "Engine/LatentActionManager.h"
 #include "Engine/Level.h"
+#include "Engine/LevelStreaming.h"
 #include "Engine/LocalPlayer.h"
 #include "Engine/MaterialMerging.h"
 #include "Engine/MemberReference.h"
@@ -207,13 +184,12 @@
 #include "GameFramework/PlayerMuteList.h"
 #include "GameFramework/Volume.h"
 #include "GameFramework/WorldSettings.h"
-#include "GlobalShader.h"
 #include "GPUSkinPublicDefs.h"
 #include "HitProxies.h"
-#include "IAudioExtensionPlugin.h"
 #include "Interfaces/Interface_AssetUserData.h"
 #include "Interfaces/Interface_CollisionDataProvider.h"
 #include "Kismet/BlueprintFunctionLibrary.h"
+#include "LatentActions.h"
 #include "LocalVertexFactory.h"
 #include "MaterialExpressionIO.h"
 #include "Materials/Material.h"
@@ -235,6 +211,7 @@
 #include "PhysicsEngine/BodyInstance.h"
 #include "PhysxUserData.h"
 #include "PixelFormat.h"
+#include "PreviewScene.h"
 #include "PrimitiveSceneProxy.h"
 #include "PrimitiveUniformShaderParameters.h"
 #include "PrimitiveViewRelevance.h"
@@ -253,11 +230,13 @@
 #include "Sound/SoundConcurrency.h"
 #include "Sound/SoundGroups.h"
 #include "Sound/SoundWave.h"
-#include "StaticBoundShaderState.h"
 #include "StaticParameterSet.h"
 #include "TextureResource.h"
 #include "Tickable.h"
 #include "TimerManager.h"
 #include "UnrealClient.h"
 #include "UnrealEngine.h"
-#include "WorldCollision.h"
+#include "Vehicles/TireType.h"
+#include "VisualLogger/VisualLogger.h"
+#include "VisualLogger/VisualLoggerTypes.h"
+#include "WorldCollision.h"
\ No newline at end of file
diff --git a/Engine/Source/Runtime/Slate/Public/SlateSharedPCH.h b/Engine/Source/Runtime/Slate/Public/SlateSharedPCH.h
--- a/Engine/Source/Runtime/Slate/Public/SlateSharedPCH.h
+++ b/Engine/Source/Runtime/Slate/Public/SlateSharedPCH.h
@@ -4,43 +4,12 @@
 
 #include "CoreUObjectSharedPCH.h"
 
-// From Core:
-#include "Features/IModularFeature.h"
-#include "GenericPlatform/GenericApplication.h"
-#include "GenericPlatform/GenericApplicationMessageHandler.h"
-#include "GenericPlatform/GenericWindow.h"
-#include "GenericPlatform/GenericWindowDefinition.h"
-#include "GenericPlatform/ICursor.h"
-#include "Logging/TokenizedMessage.h"
-#include "Math/RandomStream.h"
-#include "Math/TransformCalculus.h"
-#include "Math/TransformCalculus2D.h"
-#include "Misc/Attribute.h"
-#include "Misc/EngineVersion.h"
-#include "Misc/MessageDialog.h"
-#include "Misc/ObjectThumbnail.h"
-#include "Misc/OutputDeviceError.h"
-#include "Serialization/CustomVersion.h"
-
 // From Json:
 #include "Policies/JsonPrintPolicy.h"
 #include "Policies/PrettyJsonPrintPolicy.h"
 #include "Serialization/JsonTypes.h"
 #include "Serialization/JsonWriter.h"
 
-// From CoreUObject:
-#include "Misc/PackageName.h"
-#include "Misc/WorldCompositionUtility.h"
-#include "Serialization/SerializedPropertyScope.h"
-#include "Templates/Casts.h"
-#include "UObject/CoreNetTypes.h"
-#include "UObject/LazyObjectPtr.h"
-#include "UObject/Package.h"
-#include "UObject/PropertyTag.h"
-#include "UObject/ScriptInterface.h"
-#include "UObject/SoftObjectPtr.h"
-#include "UObject/UnrealType.h"
-
 // From InputCore:
 #include "InputCoreTypes.h"
 
@@ -71,6 +40,7 @@
 #include "Layout/Children.h"
 #include "Layout/Geometry.h"
 #include "Layout/LayoutGeometry.h"
+#include "Layout/LayoutUtils.h"
 #include "Layout/Margin.h"
 #include "Layout/PaintGeometry.h"
 #include "Layout/SlateRect.h"
@@ -127,11 +97,13 @@
 #include "Framework/Commands/UIAction.h"
 #include "Framework/Commands/UICommandInfo.h"
 #include "Framework/Commands/UICommandList.h"
+#include "Framework/Docking/LayoutService.h"
 #include "Framework/Docking/TabManager.h"
 #include "Framework/Docking/WorkspaceItem.h"
 #include "Framework/Layout/InertialScrollManager.h"
 #include "Framework/Layout/IScrollableWidget.h"
 #include "Framework/Layout/Overscroll.h"
+#include "Framework/MarqueeRect.h"
 #include "Framework/MultiBox/MultiBoxBuilder.h"
 #include "Framework/MultiBox/MultiBoxDefs.h"
 #include "Framework/MultiBox/MultiBoxExtender.h"
@@ -148,6 +120,7 @@
 #include "Widgets/Docking/SDockTab.h"
 #include "Widgets/Images/SImage.h"
 #include "Widgets/Input/IVirtualKeyboardEntry.h"
+#include "Widgets/Input/NumericTypeInterface.h"
 #include "Widgets/Input/SButton.h"
 #include "Widgets/Input/SCheckBox.h"
 #include "Widgets/Input/SComboBox.h"
@@ -164,6 +137,8 @@
 #include "Widgets/Layout/SSeparator.h"
 #include "Widgets/Layout/SSpacer.h"
 #include "Widgets/Layout/SSplitter.h"
+#include "Widgets/Notifications/SErrorText.h"
+#include "Widgets/SToolTip.h"
 #include "Widgets/Text/ISlateEditableTextWidget.h"
 #include "Widgets/Text/STextBlock.h"
 #include "Widgets/Views/SExpanderArrow.h"
@@ -171,4 +146,4 @@
 #include "Widgets/Views/SListView.h"
 #include "Widgets/Views/STableRow.h"
 #include "Widgets/Views/STableViewBase.h"
-#include "Widgets/Views/STreeView.h"
+#include "Widgets/Views/STreeView.h"
\ No newline at end of file
