#ue5 - Fix TQuat<T>::FindBetweenNormals() and TQuat<T>::FindBetweenVectors() choosing poor basis vectors when input vectors are anti-parallel (opposite directions). Added more test coverage in engine startup tests.

Modified from PR #10029

diff --git a/Engine/Source/Runtime/Core/Private/Math/UnrealMath.cpp b/Engine/Source/Runtime/Core/Private/Math/UnrealMath.cpp
--- a/Engine/Source/Runtime/Core/Private/Math/UnrealMath.cpp
+++ b/Engine/Source/Runtime/Core/Private/Math/UnrealMath.cpp
@@ -1413,7 +1413,7 @@ namespace Math
 
 		if (W >= 1e-6f * NormAB)
 		{
-			//Axis = FVector::CrossProduct(A, B);
+			//Result = FVector::CrossProduct(A, B);
 			Result = TQuat<T>(
 				A.Y * B.Z - A.Z * B.Y,
 				A.Z * B.X - A.X * B.Z,
@@ -1424,9 +1424,19 @@ namespace Math
 		{
 			// A and B point in opposite directions
 			W = 0.f;
-			Result = FMath::Abs(A.X) > FMath::Abs(A.Y)
-				? TQuat<T>(-A.Z, 0.f, A.X, W)
-				: TQuat<T>(0.f, -A.Z, A.Y, W);
+			const T X = FMath::Abs(A.X);
+			const T Y = FMath::Abs(A.Y);
+			const T Z = FMath::Abs(A.Z);
+
+			// Find most orthogonal basis 
+			const TVector<T> Basis = X < Y ? (X < Z ? TVector<T>::XAxisVector : TVector<T>::ZAxisVector) : (Y < Z ? TVector<T>::YAxisVector : TVector<T>::ZAxisVector);
+
+			//Result = FVector::CrossProduct(A, Basis);
+			Result = TQuat<T>(
+				A.Y * Basis.Z - A.Z * Basis.Y,
+				A.Z * Basis.X - A.X * Basis.Z,
+				A.X * Basis.Y - A.Y * Basis.X,
+				W);
 		}
 
 		Result.Normalize();
diff --git a/Engine/Source/Runtime/Core/Tests/Math/UnrealMathTest.cpp b/Engine/Source/Runtime/Core/Tests/Math/UnrealMathTest.cpp
--- a/Engine/Source/Runtime/Core/Tests/Math/UnrealMathTest.cpp
+++ b/Engine/Source/Runtime/Core/Tests/Math/UnrealMathTest.cpp
@@ -3571,7 +3571,9 @@ TEST_CASE_NAMED(FVectorRegisterAbstractionTest, "System::Core::Math::Vector Regi
 			FVector3f(45.0f, -60.0f, 120.0f),
 			FVector3f(-45.0f, 60.0f, -120.0f),
 			FVector3f(0.57735026918962576451f, 0.57735026918962576451f, 0.57735026918962576451f),
-			-FVector3f::ForwardVector,
+			FVector3f( 0.01f,  0.02f,  0.03f),
+			FVector3f(-0.01f,  0.02f, -0.03f),
+			FVector3f( 0.01f, -0.02f,  0.03f),
 		};
 
 		// ... and test within this tolerance.
@@ -3597,34 +3599,42 @@ TEST_CASE_NAMED(FVectorRegisterAbstractionTest, "System::Core::Math::Vector Regi
 			}
 		}
 
-
+		
 		// FindBetween
 		{
-			for (FVector3f const& A : TestVectors)
+			const FVector3f Signs[] = {
+				FVector3f(1.0f, 1.0f, 1.0f),
+				FVector3f(-1.0f, -1.0f, -1.0f),
+			};
+
+			for (FVector3f const& SignVector : Signs)
 			{
-				for (FVector3f const& B : TestVectors)
+				for (FVector3f const& A : TestVectors)
 				{
-					const FVector3f ANorm = A.GetSafeNormal();
-					const FVector3f BNorm = B.GetSafeNormal();
-
-					const FQuat4f Old = FindBetween_Old(ANorm, BNorm);
-					const FQuat4f NewNormal = FQuat4f::FindBetweenNormals(ANorm, BNorm);
-					const FQuat4f NewVector = FQuat4f::FindBetweenVectors(A, B);
-
-					const FVector3f RotAOld = Old.RotateVector(ANorm);
-					const FVector3f RotANewNormal = NewNormal.RotateVector(ANorm);
-					const FVector3f RotANewVector = NewVector.RotateVector(ANorm);
-
-					if (A.IsZero() || B.IsZero())
-					{
-						LogTest<float>(TEXT("FindBetween: Old == New (normal)"), TestQuatsEqual(Old, NewNormal, 1e-6f));
-						LogTest<float>(TEXT("FindBetween: Old == New (vector)"), TestQuatsEqual(Old, NewVector, 1e-6f));
-					}
-					else
+					for (FVector3f const& B : TestVectors)
 					{
-						LogTest<float>(TEXT("FindBetween: Old A->B"), TestFVector3Equal(RotAOld, BNorm, UE_KINDA_SMALL_NUMBER));
-						LogTest<float>(TEXT("FindBetween: New A->B (normal)"), TestFVector3Equal(RotANewNormal, BNorm, UE_KINDA_SMALL_NUMBER));
-						LogTest<float>(TEXT("FindBetween: New A->B (vector)"), TestFVector3Equal(RotANewVector, BNorm, UE_KINDA_SMALL_NUMBER));
+						const FVector3f ANorm = A.GetSafeNormal();
+						const FVector3f BNorm = (B * SignVector).GetSafeNormal();
+
+						const FQuat4f Old = FindBetween_Old(ANorm, BNorm);
+						const FQuat4f NewNormal = FQuat4f::FindBetweenNormals(ANorm, BNorm);
+						const FQuat4f NewVector = FQuat4f::FindBetweenVectors(A, B * SignVector);
+
+						const FVector3f RotAOld = Old.RotateVector(ANorm);
+						const FVector3f RotANewNormal = NewNormal.RotateVector(ANorm);
+						const FVector3f RotANewVector = NewVector.RotateVector(ANorm);
+
+						if (A.IsZero() || B.IsZero())
+						{
+							LogTest<float>(TEXT("FindBetween: Old == New (normal)"), TestQuatsEqual(Old, NewNormal, 1e-6f));
+							LogTest<float>(TEXT("FindBetween: Old == New (vector)"), TestQuatsEqual(Old, NewVector, 1e-6f));
+						}
+						else
+						{
+							LogTest<float>(TEXT("FindBetween: Old A->B"), TestFVector3Equal(RotAOld, BNorm, UE_KINDA_SMALL_NUMBER));
+							LogTest<float>(TEXT("FindBetween: New A->B (normal)"), TestFVector3Equal(RotANewNormal, BNorm, UE_KINDA_SMALL_NUMBER));
+							LogTest<float>(TEXT("FindBetween: New A->B (vector)"), TestFVector3Equal(RotANewVector, BNorm, UE_KINDA_SMALL_NUMBER));
+						}
 					}
 				}
 			}
