- Added DumpCVars and DumpCCmds console commands, which expand on 'DumpConsoleCommands', but in a much less confusing, and also more useful, way

diff --git a/Engine/Source/Runtime/Core/Private/HAL/ConsoleManager.cpp b/Engine/Source/Runtime/Core/Private/HAL/ConsoleManager.cpp
--- a/Engine/Source/Runtime/Core/Private/HAL/ConsoleManager.cpp
+++ b/Engine/Source/Runtime/Core/Private/HAL/ConsoleManager.cpp
@@ -12,6 +12,7 @@ ConsoleManager.cpp: console command handling
 #include "Misc/ConfigCacheIni.h"
 #include "Misc/ConfigUtilities.h"
 #include "Misc/CoreDelegates.h"
+#include "Misc/OutputDeviceFile.h"
 #include "Modules/ModuleManager.h"
 #include "HAL/PlatformProcess.h"
 #include "Misc/RemoteConfigIni.h"
@@ -1822,11 +1823,120 @@ void FConsoleManager::ForEachConsoleObjectThatContains(const FConsoleObjectVisit
 	}
 }
 
+static void DumpObjects(const TMap<FString, IConsoleObject*>& ConsoleObjects, const TCHAR* Params, FOutputDevice& InAr, bool bDisplayCommands)
+{
+	bool bShowHelp = FParse::Param(Params, TEXT("showhelp"));
+	FString CSVFilename;
+	bool bWriteToCSV = FParse::Value(Params, TEXT("-csv="), CSVFilename);
+	bWriteToCSV = bWriteToCSV || FParse::Param(Params, TEXT("csv"));
+	FString Prefix = FParse::Token(Params, false);
+	if (Prefix.StartsWith(TEXT("-")))
+	{
+		Prefix = TEXT("");
+	}
+
+	// get sorted list of keys of all console objects
+	TArray<FString> SortedKeys;
+	ConsoleObjects.GetKeys(SortedKeys);
+	SortedKeys.Sort();
+
+	FOutputDeviceFile* CSV = nullptr;
+	if (bWriteToCSV)
+	{
+		if (CSVFilename.IsEmpty())
+		{
+			CSVFilename = FPaths::Combine(FPaths::ProjectLogDir(), bDisplayCommands ? TEXT("ConsoleCommands.csv") : TEXT("ConsoleVars.csv"));
+		}
+		CSV = new FOutputDeviceFile(*CSVFilename, false, false, false);
+		if (!CSV->IsOpened())
+		{
+			InAr.Logf(TEXT("Unable to create CSV file for writing: '%s'"), *CSVFilename);
+			return;
+		}
+
+		InAr.Logf(TEXT("Dumping to CSV file: '%s'"), *CSVFilename);
+		if (bDisplayCommands)
+		{
+			CSV->Logf(TEXT("NAME%s"), bShowHelp ? TEXT(",HELP") : TEXT(""));
+		}
+		else
+		{
+			CSV->Logf(TEXT("NAME,VALUE,SETBY%s"), bShowHelp ? TEXT(",HELP") : TEXT(""));
+		}
+	}
+	FOutputDevice& Ar = bWriteToCSV ? *CSV : InAr;
+
+	for (const FString& Key : SortedKeys)
+	{
+		if (Prefix.IsEmpty() || Key.StartsWith(Prefix))
+		{
+			IConsoleObject* Obj = ConsoleObjects[Key];
+			IConsoleVariable* CVar = Obj->AsVariable();
+			IConsoleCommand* CCmd = Obj->AsCommand();
+			if (bDisplayCommands && CCmd != nullptr)
+			{
+				if (bWriteToCSV)
+				{
+					FString Help;
+					if (bShowHelp)
+					{
+						Help = FString(TEXT(",")) + FString(CCmd->GetHelp()).TrimStartAndEnd().Replace(TEXT("\n"), TEXT("/"));
+					}
+					Ar.Logf(TEXT("%s,%s,%s%s"), *Key, *Help);
+				}
+				else
+				{
+					Ar.Logf(TEXT("%s"), *Key);
+					if (bShowHelp)
+					{
+						Ar.Logf(TEXT("%s"), CCmd->GetHelp());
+						Ar.Logf(TEXT(" "));
+					}
+				}
+			}
+			if (!bDisplayCommands && CVar != nullptr)
+			{
+				if (bWriteToCSV)
+				{
+					FString Help;
+					if (bShowHelp)
+					{
+						Help = FString(TEXT(",")) + FString(CVar->GetHelp()).TrimStartAndEnd().Replace(TEXT("\n"), TEXT("/"));
+					}
+					Ar.Logf(TEXT("%s,%s,%s%s"), *Key, *CVar->GetString(), GetSetByTCHAR(CVar->GetFlags()), *Help);
+				}
+				else
+				{
+					Ar.Logf(TEXT("%s = \"%s\"      LastSetBy: %s"), *Key, *CVar->GetString(), GetSetByTCHAR(CVar->GetFlags()));
+					if (bShowHelp)
+					{
+						Ar.Logf(TEXT("%s"), CVar->GetHelp());
+						Ar.Logf(TEXT(" "));
+					}
+				}
+			}
+		}
+	}
+
+	delete CSV;
+}
+
 bool FConsoleManager::ProcessUserConsoleInput(const TCHAR* InInput, FOutputDevice& Ar, UWorld* InWorld)
 {
 	check(InInput);
 	CSV_EVENT_GLOBAL(TEXT("Cmd: %s"), InInput);
 
+	if (FParse::Command(&InInput, TEXT("dumpcvars")))
+	{
+		DumpObjects(ConsoleObjects, InInput, Ar, false);
+		return true;
+	}
+	if (FParse::Command(&InInput, TEXT("dumpccmds")))
+	{
+		DumpObjects(ConsoleObjects, InInput, Ar, true);
+		return true;
+	}
+
 	const TCHAR* It = InInput;
 
 	FString Param1 = GetTextSection(It);
@@ -2537,6 +2647,14 @@ void CreateConsoleVariables()
 
 #if	!UE_BUILD_SHIPPING
 	IConsoleManager::Get().RegisterConsoleCommand( TEXT( "DumpConsoleCommands" ), TEXT( "Dumps all console vaiables and commands and all exec that can be discovered to the log/console" ), ECVF_Default );
+
+	IConsoleManager::Get().RegisterConsoleCommand(TEXT("DumpCVars"),
+		TEXT("Lists all CVars (or a subset) and their values. Can also show help, and can save to .csv.\nUsage: DumpCVars [Prefix] [-showhelp] [-csv=[path]]\nIf -csv does not have a file specified, it will create a file in the Project Logs directory"),
+		ECVF_Default);
+	IConsoleManager::Get().RegisterConsoleCommand(TEXT("DumpCCmds"),
+		TEXT("Lists all CVars (or a subset) and their values. Can also show help, and can save to .csv.\nUsage: DumpCCmds [Prefix] [-showhelp] [-csv=[path]]\nIf -csv does not have a file specified, it will create a file in the Project Logs directory"),
+		ECVF_Default);
+
 #endif // !UE_BUILD_SHIPPING
 
 	// testing code
