- Fixed CSV printing to not have timestamps, and CCmd garbage from uninitialized data

diff --git a/Engine/Source/Runtime/Core/Private/HAL/ConsoleManager.cpp b/Engine/Source/Runtime/Core/Private/HAL/ConsoleManager.cpp
--- a/Engine/Source/Runtime/Core/Private/HAL/ConsoleManager.cpp
+++ b/Engine/Source/Runtime/Core/Private/HAL/ConsoleManager.cpp
@@ -17,6 +17,7 @@ ConsoleManager.cpp: console command handling
 #include "HAL/PlatformProcess.h"
 #include "Misc/RemoteConfigIni.h"
 #include "ProfilingDebugging/CsvProfiler.h"
+#include "HAL/FileManager.h"
 
 DEFINE_LOG_CATEGORY(LogConsoleResponse);
 DEFINE_LOG_CATEGORY_STATIC(LogConsoleManager, Log, All);
@@ -1823,6 +1824,19 @@ void FConsoleManager::ForEachConsoleObjectThatContains(const FConsoleObjectVisit
 	}
 }
 
+template <typename FmtType, typename... Types>
+static void MultiLogf(FOutputDevice* Device, FArchive* File, const FmtType& Fmt, Types... Args)
+{
+	if (Device != nullptr)
+	{
+		Device->Logf(Fmt, Args...);
+	}
+	if (File != nullptr)
+	{
+		File->Logf(Fmt, Args...);
+	}
+}
+
 static void DumpObjects(const TMap<FString, IConsoleObject*>& ConsoleObjects, const TCHAR* Params, FOutputDevice& InAr, bool bDisplayCommands)
 {
 	bool bShowHelp = FParse::Param(Params, TEXT("showhelp"));
@@ -1840,15 +1854,16 @@ static void DumpObjects(const TMap<FString, IConsoleObject*>& ConsoleObjects, co
 	ConsoleObjects.GetKeys(SortedKeys);
 	SortedKeys.Sort();
 
-	FOutputDeviceFile* CSV = nullptr;
+	FOutputDevice* Log = nullptr;
+	FArchive* CSV = nullptr;
 	if (bWriteToCSV)
 	{
 		if (CSVFilename.IsEmpty())
 		{
 			CSVFilename = FPaths::Combine(FPaths::ProjectLogDir(), bDisplayCommands ? TEXT("ConsoleCommands.csv") : TEXT("ConsoleVars.csv"));
 		}
-		CSV = new FOutputDeviceFile(*CSVFilename, false, false, false);
-		if (!CSV->IsOpened())
+		CSV = IFileManager::Get().CreateFileWriter(*CSVFilename, FILEWRITE_AllowRead);
+		if (CSV == nullptr)
 		{
 			InAr.Logf(TEXT("Unable to create CSV file for writing: '%s'"), *CSVFilename);
 			return;
@@ -1864,7 +1879,11 @@ static void DumpObjects(const TMap<FString, IConsoleObject*>& ConsoleObjects, co
 			CSV->Logf(TEXT("NAME,VALUE,SETBY%s"), bShowHelp ? TEXT(",HELP") : TEXT(""));
 		}
 	}
-	FOutputDevice& Ar = bWriteToCSV ? *CSV : InAr;
+	else
+	{
+		// only write to the log if CSV is not used
+		Log = &InAr;
+	}
 
 	for (const FString& Key : SortedKeys)
 	{
@@ -1873,46 +1892,36 @@ static void DumpObjects(const TMap<FString, IConsoleObject*>& ConsoleObjects, co
 			IConsoleObject* Obj = ConsoleObjects[Key];
 			IConsoleVariable* CVar = Obj->AsVariable();
 			IConsoleCommand* CCmd = Obj->AsCommand();
-			if (bDisplayCommands && CCmd != nullptr)
+
+			// process optional help
+			FString Help;
+			if (bShowHelp)
 			{
+				Help = FString(Obj->GetHelp()).TrimStartAndEnd();
 				if (bWriteToCSV)
 				{
-					FString Help;
-					if (bShowHelp)
-					{
-						Help = FString(TEXT(",")) + FString(CCmd->GetHelp()).TrimStartAndEnd().Replace(TEXT("\n"), TEXT("/"));
-					}
-					Ar.Logf(TEXT("%s,%s,%s%s"), *Key, *Help);
+					// newlines and commas in help will throw off the csv
+					Help = FString::Printf(TEXT(",\"%s\""), *Help.Replace(TEXT("\n"), TEXT("\\n")));
 				}
 				else
 				{
-					Ar.Logf(TEXT("%s"), *Key);
-					if (bShowHelp)
-					{
-						Ar.Logf(TEXT("%s"), CCmd->GetHelp());
-						Ar.Logf(TEXT(" "));
-					}
+					Help = FString::Printf(TEXT("\n%s\n "), *Help);
 				}
 			}
+
+			if (bDisplayCommands && CCmd != nullptr)
+			{
+				MultiLogf(Log, CSV, TEXT("%s%s"), *Key, *Help);
+			}
 			if (!bDisplayCommands && CVar != nullptr)
 			{
 				if (bWriteToCSV)
 				{
-					FString Help;
-					if (bShowHelp)
-					{
-						Help = FString(TEXT(",")) + FString(CVar->GetHelp()).TrimStartAndEnd().Replace(TEXT("\n"), TEXT("/"));
-					}
-					Ar.Logf(TEXT("%s,%s,%s%s"), *Key, *CVar->GetString(), GetSetByTCHAR(CVar->GetFlags()), *Help);
+					MultiLogf(Log, CSV, TEXT("%s,%s,%s%s"), *Key, *CVar->GetString(), GetSetByTCHAR(CVar->GetFlags()), *Help);
 				}
 				else
 				{
-					Ar.Logf(TEXT("%s = \"%s\"      LastSetBy: %s"), *Key, *CVar->GetString(), GetSetByTCHAR(CVar->GetFlags()));
-					if (bShowHelp)
-					{
-						Ar.Logf(TEXT("%s"), CVar->GetHelp());
-						Ar.Logf(TEXT(" "));
-					}
+					MultiLogf(Log, CSV, TEXT("%s = \"%s\"      LastSetBy: %s%s"), *Key, *CVar->GetString(), GetSetByTCHAR(CVar->GetFlags()), *Help);
 				}
 			}
 		}
