autortfm fixes:
- change autortfm_is_transactional to return true when 'not idle', so it returns true if we're in an aborted state. This lets us propagate abort through Open code

diff --git a/Engine/Source/Runtime/Core/Private/AutoRTFM/API.cpp b/Engine/Source/Runtime/Core/Private/AutoRTFM/API.cpp
--- a/Engine/Source/Runtime/Core/Private/AutoRTFM/API.cpp
+++ b/Engine/Source/Runtime/Core/Private/AutoRTFM/API.cpp
@@ -23,7 +23,7 @@ namespace AutoRTFM
 
 extern "C" bool autortfm_is_transactional()
 {
-    return FContext::Get()->GetStatus() == EContextStatus::OnTrack;
+    return FContext::Get()->IsTransactional();
 }
 
 extern "C" bool autortfm_is_closed()
diff --git a/Engine/Source/Runtime/Core/Private/AutoRTFM/Context.cpp b/Engine/Source/Runtime/Core/Private/AutoRTFM/Context.cpp
--- a/Engine/Source/Runtime/Core/Private/AutoRTFM/Context.cpp
+++ b/Engine/Source/Runtime/Core/Private/AutoRTFM/Context.cpp
@@ -52,13 +52,12 @@ bool FContext::IsTransactional()
         return false;
     }
 
-    if (Context->GetStatus() == EContextStatus::OnTrack)
+    if (Context->GetStatus() != EContextStatus::Idle)
     {
         return true;
     }
     else
     {
-        ASSERT(Context->GetStatus() == EContextStatus::Idle);
         return false;
     }
 }
