Silence bogus static analysis warning.

diff --git a/Engine/Source/Runtime/Core/Private/Delegates/DelegateHandle.cpp b/Engine/Source/Runtime/Core/Private/Delegates/DelegateHandle.cpp
--- a/Engine/Source/Runtime/Core/Private/Delegates/DelegateHandle.cpp
+++ b/Engine/Source/Runtime/Core/Private/Delegates/DelegateHandle.cpp
@@ -13,7 +13,7 @@ namespace UE::Delegates::Private
 uint64 FDelegateHandle::GenerateNewID()
 {
 	// Just increment a counter to generate an ID.
-	uint64 Result;
+	uint64 Result = 0; // Initialize just to silence static analysis.
 	
 	UE_AUTORTFM_OPEN(
 	{
