Silence bogus static analysis warning.

diff --git a/Engine/Source/Runtime/Core/Public/HAL/FMemory.inl b/Engine/Source/Runtime/Core/Public/HAL/FMemory.inl
--- a/Engine/Source/Runtime/Core/Public/HAL/FMemory.inl
+++ b/Engine/Source/Runtime/Core/Public/HAL/FMemory.inl
@@ -16,7 +16,7 @@ struct FScopedMallocTimer;
 
 FMEMORY_INLINE_FUNCTION_DECORATOR void* FMemory::Malloc(SIZE_T Count, uint32 Alignment)
 {
-	void* Ptr;
+	void* Ptr = nullptr; // Silence bogus static analysis warnings.
 
 	// AutoRTFM: For non-transactional code, all of these calls optimize away and the
 	// behavior is the same as it always has been.
