ParkingLot: Use memory from FPlatformMemory to avoid GMalloc

ParkingLot cannot be used to synchronize allocator implementations if ParkingLot itself depends on those allocators.

#rnx

diff --git a/Engine/Source/Runtime/Core/Private/Async/ParkingLot.cpp b/Engine/Source/Runtime/Core/Private/Async/ParkingLot.cpp
--- a/Engine/Source/Runtime/Core/Private/Async/ParkingLot.cpp
+++ b/Engine/Source/Runtime/Core/Private/Async/ParkingLot.cpp
@@ -8,6 +8,7 @@
 #include "Containers/Array.h"
 #include "Containers/ArrayView.h"
 #include "HAL/PlatformManualResetEvent.h"
+#include "Templates/AlignmentTemplates.h"
 #include "Templates/RefCounting.h"
 #include <atomic>
 
@@ -18,39 +19,291 @@ using FBucketMutex = FWordMutex;
 
 ///////////////////////////////////////////////////////////////////////////////////////////////////////////////////////
 
+/**
+ * A low-level linear allocator that bypasses FMemory/GMalloc.
+ *
+ * Allocations made by ParkingLot must use this allocator to ensure that primitives that build on
+ * ParkingLot can be used in the implementation of the allocators used by FMemory/GMalloc.
+ */
+class FLowLevelLinearAllocator
+{
+public:
+	void* Malloc(SIZE_T Size, SIZE_T Alignment);
+	void* Realloc(void* OldMem, SIZE_T OldSize, SIZE_T NewSize, SIZE_T Alignment);
+	void Free(void* Mem, SIZE_T Size, SIZE_T Alignment);
+
+	constexpr FLowLevelLinearAllocator() = default;
+	FLowLevelLinearAllocator(const FLowLevelLinearAllocator&) = delete;
+	FLowLevelLinearAllocator& operator=(const FLowLevelLinearAllocator&) = delete;
+
+private:
+	struct FBlockHeader
+	{
+		uint64 UsedSize;
+		uint64 TotalSize;
+		uint64 ReferenceCount;
+		FBlockHeader* Next;
+	};
+
+	FWordMutex Mutex;
+	FBlockHeader* ActiveBlock = nullptr;
+};
+
+void* FLowLevelLinearAllocator::Malloc(SIZE_T Size, SIZE_T Alignment)
+{
+	const SIZE_T SizeWithAlignment = Align(Size, Alignment);
+
+	TUniqueLock Lock(Mutex);
+
+	for (;;)
+	{
+		if (LIKELY(ActiveBlock))
+		{
+			const SIZE_T MemOffset = Align(ActiveBlock->UsedSize, Alignment);
+			const SIZE_T EndOffset = MemOffset + SizeWithAlignment;
+			if (LIKELY(ActiveBlock->TotalSize >= EndOffset))
+			{
+				ActiveBlock->UsedSize = EndOffset;
+				++ActiveBlock->ReferenceCount;
+				void* Mem = (uint8*)ActiveBlock + MemOffset;
+				checkSlow(IsAligned(Mem, Alignment));
+				return Mem;
+			}
+		}
+
+		using FVirtualMemoryBlock = FPlatformMemory::FPlatformVirtualMemoryBlock;
+		const SIZE_T VirtualSizeAlignment = FVirtualMemoryBlock::GetVirtualSizeAlignment();
+
+		const SIZE_T BlockSizeWithHeader = Align(sizeof(FBlockHeader), Alignment) + SizeWithAlignment;
+		const SIZE_T BlockSizeWithAlignment = Align(BlockSizeWithHeader, VirtualSizeAlignment);
+
+		FVirtualMemoryBlock MemoryBlock = FVirtualMemoryBlock::AllocateVirtual(BlockSizeWithAlignment, VirtualSizeAlignment);
+		MemoryBlock.Commit();
+
+		FBlockHeader* Block = (FBlockHeader*)MemoryBlock.GetVirtualPointer();
+		Block->UsedSize = sizeof(FBlockHeader);
+		Block->TotalSize = BlockSizeWithAlignment;
+		Block->ReferenceCount = 0;
+		Block->Next = ActiveBlock;
+		ActiveBlock = Block;
+	}
+}
+
+void* FLowLevelLinearAllocator::Realloc(void* OldMem, SIZE_T OldSize, SIZE_T NewSize, SIZE_T Alignment)
+{
+	if (!OldMem)
+	{
+		return Malloc(NewSize, Alignment);
+	}
+
+	if (NewSize == 0)
+	{
+		Free(OldMem, NewSize, Alignment);
+		return nullptr;
+	}
+
+	checkSlow(IsAligned(OldMem, Alignment));
+
+	const SIZE_T OldSizeWithAlignment = Align(OldSize, Alignment);
+	const SIZE_T NewSizeWithAlignment = Align(NewSize, Alignment);
+
+	if (OldSizeWithAlignment >= NewSizeWithAlignment)
+	{
+		return OldMem;
+	}
+
+	TUniqueLock Lock(Mutex);
+
+	FBlockHeader* Block = nullptr;
+	FBlockHeader** BlockPtr = &ActiveBlock;
+	for (;;)
+	{
+		Block = *BlockPtr;
+		checkSlow(Block);
+		if (LIKELY(UPTRINT(OldMem) >= UPTRINT(Block) && UPTRINT(OldMem) < UPTRINT(Block) + Block->UsedSize))
+		{
+			break;
+		}
+		BlockPtr = &Block->Next;
+	}
+
+	const SIZE_T MemOffset = SIZE_T(UPTRINT(OldMem) - UPTRINT(Block));
+	const SIZE_T OldEndOffset = MemOffset + OldSizeWithAlignment;
+	const SIZE_T NewEndOffset = MemOffset + NewSizeWithAlignment;
+
+	if (Block->UsedSize == OldEndOffset && Block->TotalSize >= NewEndOffset)
+	{
+		Block->UsedSize = NewEndOffset;
+		return OldMem;
+	}
+
+	void* NewMem = Malloc(NewSize, Alignment);
+	FMemory::Memcpy(NewMem, OldMem, FPlatformMath::Min(NewSize, OldSize));
+	Free(OldMem, OldSize, Alignment);
+	return NewMem;
+}
+
+void FLowLevelLinearAllocator::Free(void* Mem, SIZE_T Size, SIZE_T Alignment)
+{
+	TDynamicUniqueLock Lock(Mutex);
+
+	FBlockHeader* Block = nullptr;
+	FBlockHeader** BlockPtr = &ActiveBlock;
+	for (;;)
+	{
+		Block = *BlockPtr;
+		checkSlow(Block);
+		if (LIKELY(UPTRINT(Mem) >= UPTRINT(Block) && UPTRINT(Mem) < UPTRINT(Block) + Block->UsedSize))
+		{
+			break;
+		}
+		BlockPtr = &Block->Next;
+	}
+
+	if (--Block->ReferenceCount == 0)
+	{
+		*BlockPtr = Block->Next;
+
+		Lock.Unlock();
+		using FVirtualMemoryBlock = FPlatformMemory::FPlatformVirtualMemoryBlock;
+		const SIZE_T VirtualSizeAlignment = FVirtualMemoryBlock::GetVirtualSizeAlignment();
+		FVirtualMemoryBlock(Block, uint32(Block->TotalSize / VirtualSizeAlignment)).FreeVirtual();
+	}
+}
+
+///////////////////////////////////////////////////////////////////////////////////////////////////////////////////////
+
+/** A low-level linear allocator for use with short-lived allocations. */
+class FLowLevelTemporaryAllocator
+{
+public:
+	static void* Malloc(SIZE_T Size)
+	{
+		const SIZE_T Alignment = GetAlignmentForSize(Size);
+		void* Mem = (uint8*)Allocator.Malloc(Size + Alignment, Alignment) + Alignment;
+		((SIZE_T*)Mem)[-1] = Size;
+		return Mem;
+	}
+
+	static void* Realloc(void* Mem, SIZE_T NewSize)
+	{
+		if (Mem)
+		{
+			const SIZE_T OldSize = ((SIZE_T*)Mem)[-1];
+			const SIZE_T Alignment = GetAlignmentForSize(OldSize);
+			Mem = (uint8*)Allocator.Realloc((uint8*)Mem - Alignment, OldSize + Alignment, NewSize + Alignment, Alignment) + Alignment;
+			((SIZE_T*)Mem)[-1] = NewSize;
+			return Mem;
+		}
+		return Malloc(NewSize);
+	}
+
+	static void Free(void* Mem)
+	{
+		const SIZE_T Size = ((SIZE_T*)Mem)[-1];
+		const SIZE_T Alignment = GetAlignmentForSize(Size);
+		return Allocator.Free((uint8*)Mem - Alignment, Size, Alignment);
+	}
+
+private:
+	inline static SIZE_T GetAlignmentForSize(SIZE_T Size)
+	{
+		return Size >= 16 ? 16 : sizeof(SIZE_T);
+	}
+
+	inline static FLowLevelLinearAllocator Allocator;
+};
+
+using FLowLevelContainerAllocator = TSizedHeapAllocator<32, FLowLevelTemporaryAllocator>;
+
+template <uint32 InlineElementCount>
+using TInlineLowLevelContainerAllocator = TSizedInlineAllocator<InlineElementCount, 32, FLowLevelContainerAllocator>;
+
+///////////////////////////////////////////////////////////////////////////////////////////////////////////////////////
+
+/** A low-level allocator that maintains a free list for a given object type. Memory is never released. */
+template <typename ObjectType>
+class TLowLevelObjectAllocator
+{
+	static_assert(sizeof(ObjectType) >= sizeof(void*));
+
+public:
+	static void* Malloc(SIZE_T Size = sizeof(ObjectType))
+	{
+		if (TUniqueLock Lock(FreeListMutex); void* Mem = FreeListHead)
+		{
+			FreeListHead = *(void**)Mem;
+			return Mem;
+		}
+		return Allocator.Malloc(Size, alignof(ObjectType));
+	}
+
+	static void Free(void* Mem)
+	{
+		TUniqueLock Lock(FreeListMutex);
+		void* Next = FreeListHead;
+		*(void**)Mem = Next;
+		FreeListHead = Mem;
+	}
+
+private:
+	inline static FWordMutex FreeListMutex;
+	inline static void* FreeListHead = nullptr;
+	inline static FLowLevelLinearAllocator Allocator;
+};
+
+///////////////////////////////////////////////////////////////////////////////////////////////////////////////////////
+
+/**
+ * Stride between instances of FThread.
+ *
+ * Performance is very sensitive to this stride. It is likely better to find a range within which
+ * to randomize the stride, to maintain fairly consistent performance across a range of hardware.
+ */
+constexpr SIZE_T ThreadStride = 2048;
+
 /**
  * A thread as stored in the wait queue.
  */
-struct FThread final
+struct alignas(PLATFORM_CACHE_LINE_SIZE) FThread final
 {
 	FThread* Next = nullptr;
 	const void* WaitAddress = nullptr;
 	uint64 WakeToken = 0;
 	FPlatformManualResetEvent Event;
-	mutable std::atomic<uint32> ReferenceCount = 0;
+	std::atomic<uint32> ReferenceCount = 0;
 
-	inline void AddRef() const
+	inline void AddRef()
 	{
 		ReferenceCount.fetch_add(1, std::memory_order_relaxed);
 	}
 
-	inline void Release() const
+	inline void Release()
 	{
 		if (ReferenceCount.fetch_sub(1, std::memory_order_acq_rel) == 1)
 		{
-			delete this;
+			this->~FThread();
+			Allocator.Free(this);
 		}
 	}
 
-	FThread() = default;
+	inline static FThread* New()
+	{
+		return new(Allocator.Malloc(FPlatformMath::Max(ThreadStride, sizeof(FThread)))) FThread;
+	}
 
 private:
+	FThread() = default;
 	~FThread() = default;
 
 	FThread(const FThread&) = delete;
 	FThread& operator=(const FThread&) = delete;
+
+	static TLowLevelObjectAllocator<FThread> Allocator;
 };
 
+TLowLevelObjectAllocator<FThread> FThread::Allocator;
+
 class FThreadLocalData
 {
 public:
@@ -64,7 +317,7 @@ public:
 		FThreadLocalData& LocalThreadLocalData = ThreadLocalData;
 		if (!LocalThreadLocalData.Thread)
 		{
-			LocalThreadLocalData.Thread = new FThread;
+			LocalThreadLocalData.Thread = FThread::New();
 		}
 		return *LocalThreadLocalData.Thread;
 	}
@@ -102,9 +355,16 @@ enum class EQueueAction
 class alignas(PLATFORM_CACHE_LINE_SIZE) FBucket final
 {
 public:
-	FBucket() = default;
-	FBucket(const FBucket&) = delete;
-	FBucket& operator=(const FBucket&) = delete;
+	static FBucket* Create()
+	{
+		return new(Allocator.Malloc()) FBucket;
+	}
+
+	static void Destroy(FBucket* Bucket)
+	{
+		Bucket->~FBucket();
+		Allocator.Free(Bucket);
+	}
 
 	[[nodiscard]] inline TDynamicUniqueLock<FBucketMutex> LockDynamic() { return TDynamicUniqueLock(Mutex); }
 
@@ -128,11 +388,21 @@ public:
 	void DequeueIf(VisitorType&& Visitor);
 
 private:
+	FBucket() = default;
+	~FBucket() = default;
+
+	FBucket(const FBucket&) = delete;
+	FBucket& operator=(const FBucket&) = delete;
+
 	FBucketMutex Mutex;
 	FThread* Head = nullptr;
 	FThread* Tail = nullptr;
+
+	static TLowLevelObjectAllocator<FBucket> Allocator;
 };
 
+TLowLevelObjectAllocator<FBucket> FBucket::Allocator;
+
 void FBucket::Enqueue(FThread* Thread)
 {
 	checkSlow(Thread);
@@ -235,7 +505,7 @@ private:
 	static void Destroy(FTable& Table);
 
 	/** Try to lock the whole table by locking every one of its buckets. */
-	static bool TryLock(FTable& Table, TArray<FBucket*>& OutBuckets);
+	static bool TryLock(FTable& Table, TArray<FBucket*, FLowLevelContainerAllocator>& OutBuckets);
 	/** Unlock an array of buckets that was filled by TryLock. */
 	static void Unlock(TConstArrayView<FBucket*> LockedBuckets);
 
@@ -245,14 +515,18 @@ private:
 	/** Pointer to the current global table. Previous global tables are leaked. */
 	inline static std::atomic<FTable*> GlobalTable;
 
+	/** Allocator used exclusively for allocating these tables. */
+	inline static FLowLevelLinearAllocator Allocator;
+
 	FTable() = default;
 	~FTable() = default;
+
 	FTable(const FTable&) = delete;
 	FTable& operator=(const FTable&) = delete;
 
 	/** Find or create the bucket at the index. */
 	template <typename AllocatorType>
-	FBucket& FindOrCreateBucket(uint32 Index, AllocatorType&& Allocator);
+	FBucket& FindOrCreateBucket(uint32 Index, AllocatorType&& BucketAllocator);
 
 	uint32 BucketCount = 0;
 	std::atomic<FBucket*> Buckets[0];
@@ -266,7 +540,7 @@ FBucket& FTable::FindOrCreateBucket(const void* Address, TDynamicUniqueLock<FBuc
 	{
 		FTable& Table = CreateOrGet();
 		const uint32 Index = Hash % Table.BucketCount;
-		FBucket& Bucket = Table.FindOrCreateBucket(Index, [] { return new FBucket; });
+		FBucket& Bucket = Table.FindOrCreateBucket(Index, [] { return FBucket::Create(); });
 		OutLock = Bucket.LockDynamic();
 
 		if (LIKELY(&Table == GlobalTable.load(std::memory_order_acquire)))
@@ -307,20 +581,20 @@ FBucket* FTable::FindBucket(const void* Address, TDynamicUniqueLock<FBucketMutex
 }
 
 template <typename AllocatorType>
-FBucket& FTable::FindOrCreateBucket(uint32 Index, AllocatorType&& Allocator)
+FBucket& FTable::FindOrCreateBucket(uint32 Index, AllocatorType&& BucketAllocator)
 {
 	std::atomic<FBucket*>& BucketPtr = Buckets[Index];
 	FBucket* Bucket = BucketPtr.load(std::memory_order_acquire);
 	if (UNLIKELY(!Bucket))
 	{
-		FBucket* NewBucket = Allocator();
+		FBucket* NewBucket = BucketAllocator();
 		if (BucketPtr.compare_exchange_strong(Bucket, NewBucket, std::memory_order_release, std::memory_order_acquire))
 		{
 			Bucket = NewBucket;
 		}
 		else
 		{
-			delete NewBucket;
+			FBucket::Destroy(NewBucket);
 		}
 		checkSlow(Bucket);
 	}
@@ -352,7 +626,7 @@ FTable& FTable::CreateOrGet()
 FTable& FTable::Create(const uint32 Size)
 {
 	const SIZE_T MemorySize = sizeof(FTable) + sizeof(FBucket*) * SIZE_T(Size);
-	void* const Memory = FMemory::Malloc(MemorySize, alignof(FTable));
+	void* const Memory = Allocator.Malloc(MemorySize, alignof(FTable));
 	FMemory::Memzero(Memory, MemorySize);
 	FTable& Table = *new(Memory) FTable;
 	Table.BucketCount = Size;
@@ -361,14 +635,15 @@ FTable& FTable::Create(const uint32 Size)
 
 void FTable::Destroy(FTable& Table)
 {
+	const SIZE_T MemorySize = sizeof(FTable) + sizeof(FBucket*) * SIZE_T(Table.BucketCount);
 	Table.~FTable();
-	FMemory::Free(&Table);
+	Allocator.Free(&Table, MemorySize, alignof(FTable));
 }
 
 void FTable::Reserve(const uint32 ThreadCount)
 {
 	const uint32 TargetBucketCount = FMath::RoundUpToPowerOfTwo(ThreadCount);
-	TArray<FBucket*> ExistingBuckets;
+	TArray<FBucket*, FLowLevelContainerAllocator> ExistingBuckets;
 
 	for (;;)
 	{
@@ -388,7 +663,7 @@ void FTable::Reserve(const uint32 ThreadCount)
 
 		// Gather waiting threads to be redistributed into the buckets of the new table.
 		// Threads with the same address remain in the same relative order as they were queued.
-		TArray<FThread*> Threads;
+		TArray<FThread*, FLowLevelContainerAllocator> Threads;
 		for (FBucket* Bucket : ExistingBuckets)
 		{
 			while (FThread* Thread = Bucket->Dequeue())
@@ -400,10 +675,10 @@ void FTable::Reserve(const uint32 ThreadCount)
 		FTable& NewTable = Create(TargetBucketCount);
 
 		// Reuse existing now-empty buckets when populating the new table.
-		TArray<FBucket*> AvailableBuckets = ExistingBuckets;
+		TArray<FBucket*, FLowLevelContainerAllocator> AvailableBuckets = ExistingBuckets;
 		const auto AllocateBucket = [&AvailableBuckets]() -> FBucket*
 		{
-			return !AvailableBuckets.IsEmpty() ? AvailableBuckets.Pop() : new FBucket;
+			return !AvailableBuckets.IsEmpty() ? AvailableBuckets.Pop(/*bAllowShrinking*/ false) : FBucket::Create();
 		};
 
 		// Add waiting threads to the new table.
@@ -432,14 +707,14 @@ void FTable::Reserve(const uint32 ThreadCount)
 	}
 }
 
-bool FTable::TryLock(FTable& Table, TArray<FBucket*>& OutBuckets)
+bool FTable::TryLock(FTable& Table, TArray<FBucket*, FLowLevelContainerAllocator>& OutBuckets)
 {
 	OutBuckets.Reset();
 
 	// Gather buckets from the table, creating them as needed because the lock is on the bucket.
 	for (uint32 Index = 0; Index < Table.BucketCount; ++Index)
 	{
-		OutBuckets.Add(&Table.FindOrCreateBucket(Index, [] { return new FBucket; }));
+		OutBuckets.Add(&Table.FindOrCreateBucket(Index, [] { return FBucket::Create(); }));
 	}
 
 	// Lock the buckets in order by address to ensure consistent ordering regardless of the table being locked.
@@ -628,7 +903,7 @@ uint32 WakeMultiple(const void* const Address, const uint32 WakeCount)
 {
 	using namespace Private;
 
-	TArray<TRefCountPtr<FThread>> WakeThreads;
+	TArray<TRefCountPtr<FThread>, TInlineLowLevelContainerAllocator<128>> WakeThreads;
 
 	if (TDynamicUniqueLock<FBucketMutex> BucketLock; FBucket* Bucket = FTable::FindBucket(Address, BucketLock))
 	{
