Fix static analysis warning in I/O chunk encoding

diff --git a/Engine/Source/Runtime/Core/Private/IO/IoChunkEncoding.cpp b/Engine/Source/Runtime/Core/Private/IO/IoChunkEncoding.cpp
--- a/Engine/Source/Runtime/Core/Private/IO/IoChunkEncoding.cpp
+++ b/Engine/Source/Runtime/Core/Private/IO/IoChunkEncoding.cpp
@@ -30,6 +30,13 @@ uint32 FIoChunkEncoding::FHeader::GetBlockCount() const
 	return static_cast<uint32>(FMath::DivideAndRoundUp(RawSize, uint64(GetBlockSize())));
 }
 
+TConstArrayView<uint32> FIoChunkEncoding::FHeader::GetBlocks() const
+{
+	check(IsValid());
+	const uint8* Self = reinterpret_cast<const uint8*>(this);
+	return TConstArrayView<uint32>(reinterpret_cast<const uint32*>(Self + sizeof(FHeader)), GetBlockCount());
+}
+
 uint64 FIoChunkEncoding::FHeader::GetTotalHeaderSize() const
 {
 	check(IsValid());
@@ -218,7 +225,7 @@ bool FIoChunkEncoding::Decode(
 	Params.RawOffset = RawOffset;
 	Params.BlockSize = Header->GetBlockSize();
 	Params.TotalRawSize = Header->RawSize;
-	Params.EncodedBlockSize = TConstArrayView<uint32>(Header->Blocks, BlockCount);
+	Params.EncodedBlockSize = Header->GetBlocks(); 
 
 	const uint64 TotalHeaderSize = sizeof(FHeader) + (BlockCount * sizeof(uint32));
 	FMemoryView EncodedBlocks = EncodedData.RightChop(TotalHeaderSize);
diff --git a/Engine/Source/Runtime/Core/Public/IO/IoChunkEncoding.h b/Engine/Source/Runtime/Core/Public/IO/IoChunkEncoding.h
--- a/Engine/Source/Runtime/Core/Public/IO/IoChunkEncoding.h
+++ b/Engine/Source/Runtime/Core/Public/IO/IoChunkEncoding.h
@@ -54,11 +54,11 @@ public:
 		uint64 BlockSizeExponent: 8;
 		uint64 Flags: 8;
 		uint64 Pad: 8;
-		uint32 Blocks[];
 
 		bool IsValid() const;
 		uint32 GetBlockSize() const;
 		uint32 GetBlockCount() const;
+		TConstArrayView<uint32> GetBlocks() const;
 		uint64 GetTotalHeaderSize() const;
 
 		static const FHeader* Decode(FMemoryView HeaderData);
