Added const-ref overloads of the FMessageDialog functions that take an optional title

This avoids needing to create a temporary FText on the stack so that it can be passed by-pointer as the title override

diff --git a/Engine/Source/Runtime/Core/Private/Misc/MessageDialog.cpp b/Engine/Source/Runtime/Core/Private/Misc/MessageDialog.cpp
--- a/Engine/Source/Runtime/Core/Private/Misc/MessageDialog.cpp
+++ b/Engine/Source/Runtime/Core/Private/Misc/MessageDialog.cpp
@@ -32,6 +32,11 @@ namespace
 }
 
 void FMessageDialog::Debugf( const FText& Message, const FText* OptTitle )
+{
+	Debugf(Message, OptTitle ? *OptTitle : GetDefaultMessageTitle());
+}
+
+void FMessageDialog::Debugf( const FText& Message, const FText& Title )
 {
 	if( FApp::IsUnattended() == true )
 	{
@@ -39,7 +44,6 @@ void FMessageDialog::Debugf( const FText& Message, const FText* OptTitle )
 	}
 	else
 	{
-		FText Title = OptTitle ? *OptTitle : GetDefaultMessageTitle();
 		if ( GIsEditor && FCoreDelegates::ModalErrorMessage.IsBound() )
 		{
 			FCoreDelegates::ModalErrorMessage.Execute(EAppMsgType::Ok, Message, Title);
@@ -68,6 +72,11 @@ void FMessageDialog::ShowLastError()
 }
 
 EAppReturnType::Type FMessageDialog::Open( EAppMsgType::Type MessageType, const FText& Message, const FText* OptTitle )
+{
+	return Open(MessageType, Message, OptTitle ? *OptTitle : GetDefaultMessageTitle());
+}
+
+EAppReturnType::Type FMessageDialog::Open(EAppMsgType::Type MessageType, const FText& Message, const FText& Title)
 {
 	EAppReturnType::Type DefaultValue = EAppReturnType::Yes;
 	switch(MessageType)
@@ -113,13 +122,17 @@ EAppReturnType::Type FMessageDialog::Open( EAppMsgType::Type MessageType, const
 		}
 	}
 
-	return Open(MessageType, DefaultValue, Message, OptTitle);
+	return Open(MessageType, DefaultValue, Message, Title);
 }
 
 EAppReturnType::Type FMessageDialog::Open(EAppMsgType::Type MessageType, EAppReturnType::Type DefaultValue, const FText& Message, const FText* OptTitle)
+{
+	return Open(MessageType, DefaultValue, Message, OptTitle ? *OptTitle : GetDefaultMessageTitle());
+}
+
+EAppReturnType::Type FMessageDialog::Open(EAppMsgType::Type MessageType, EAppReturnType::Type DefaultValue, const FText& Message, const FText& Title)
 {
 	EAppReturnType::Type Result = DefaultValue;
-	const FText Title = OptTitle ? *OptTitle : GetDefaultMessageTitle();
 
 	if (!FApp::IsUnattended() && !GIsRunningUnattendedScript)
 	{
diff --git a/Engine/Source/Runtime/Core/Public/Misc/MessageDialog.h b/Engine/Source/Runtime/Core/Public/Misc/MessageDialog.h
--- a/Engine/Source/Runtime/Core/Public/Misc/MessageDialog.h
+++ b/Engine/Source/Runtime/Core/Public/Misc/MessageDialog.h
@@ -20,6 +20,7 @@ struct CORE_API FMessageDialog
 	 * @param OptTitle Optional title to use (defaults to "Message")
 	*/
 	static void Debugf( const FText& Message, const FText* OptTitle = nullptr );
+	static void Debugf( const FText& Message, const FText& Title );
 
 	/** Pops up a message dialog box containing the last system error code in string form. */
 	static void ShowLastError();
@@ -31,6 +32,7 @@ struct CORE_API FMessageDialog
 	 * @param OptTitle Optional title to use (defaults to "Message")
 	*/
 	static EAppReturnType::Type Open( EAppMsgType::Type MessageType, const FText& Message, const FText* OptTitle = nullptr);
+	static EAppReturnType::Type Open( EAppMsgType::Type MessageType, const FText& Message, const FText& Title);
 
 	/**
 	 * Open a modal message box dialog
@@ -40,4 +42,5 @@ struct CORE_API FMessageDialog
 	 * @param OptTitle Optional title to use (defaults to "Message")
 	*/
 	static EAppReturnType::Type Open(EAppMsgType::Type MessageType, EAppReturnType::Type DefaultValue, const FText& Message, const FText* OptTitle = nullptr);
+	static EAppReturnType::Type Open(EAppMsgType::Type MessageType, EAppReturnType::Type DefaultValue, const FText& Message, const FText& Title);
 };
