Optimized code generation for FMutex and FRecursiveMutex

Avoiding duplicate conditions, restructuring loops, and using LIKELY/UNLIKELY allows the compiler to generate shorter and faster code.

#rnx

diff --git a/Engine/Source/Runtime/Core/Private/Async/Mutex.cpp b/Engine/Source/Runtime/Core/Private/Async/Mutex.cpp
--- a/Engine/Source/Runtime/Core/Private/Async/Mutex.cpp
+++ b/Engine/Source/Runtime/Core/Private/Async/Mutex.cpp
@@ -18,27 +18,27 @@ void FMutex::LockSlow()
 
 		// Try to acquire the lock if it was unlocked, even if there are waiting threads.
 		// Acquiring the lock despite the waiting threads means that this lock is not FIFO and thus not fair.
-		if (!(CurrentState & IsLockedFlag))
+		if (LIKELY(!(CurrentState & IsLockedFlag)))
 		{
-			if (State.compare_exchange_weak(CurrentState, CurrentState | IsLockedFlag, std::memory_order_acquire))
+			if (LIKELY(State.compare_exchange_weak(CurrentState, CurrentState | IsLockedFlag, std::memory_order_acquire)))
 			{
 				return;
 			}
 			continue;
 		}
 
-		// Spin up to the spin limit while there are no waiting threads.
-		if (!(CurrentState & HasWaitingThreadsFlag) && SpinCount < SpinLimit)
+		if (LIKELY(!(CurrentState & HasWaitingThreadsFlag)))
 		{
-			FPlatformProcess::YieldThread();
-			++SpinCount;
-			continue;
-		}
+			// Spin up to the spin limit while there are no waiting threads.
+			if (LIKELY(SpinCount < SpinLimit))
+			{
+				FPlatformProcess::YieldThread();
+				++SpinCount;
+				continue;
+			}
 
-		// Store that there are waiting threads. Restart if the state has changed since it was loaded.
-		if (!(CurrentState & HasWaitingThreadsFlag))
-		{
-			if (!State.compare_exchange_weak(CurrentState, CurrentState | HasWaitingThreadsFlag, std::memory_order_acquire))
+			// Store that there are waiting threads. Restart if the state has changed since it was loaded.
+			if (UNLIKELY(!State.compare_exchange_weak(CurrentState, CurrentState | HasWaitingThreadsFlag, std::memory_order_acquire)))
 			{
 				continue;
 			}
@@ -52,36 +52,28 @@ void FMutex::LockSlow()
 
 void FMutex::UnlockSlow()
 {
-	uint8 CurrentState;
+	uint8 CurrentState = State.load(std::memory_order_acquire);
+	checkSlow(CurrentState & IsLockedFlag);
 
-	for (;;)
+	// Spin on the fast path because there may be spurious failures.
+	while (LIKELY(CurrentState == IsLockedFlag))
 	{
-		CurrentState = State.load(std::memory_order_acquire);
-
-		checkSlow(CurrentState & IsLockedFlag);
-
-		// Spin on the fast path because there may be spurious failures.
-		if (CurrentState == IsLockedFlag)
+		if (LIKELY(State.compare_exchange_weak(CurrentState, 0, std::memory_order_release)))
 		{
-			if (State.compare_exchange_weak(CurrentState, 0, std::memory_order_release))
-			{
-				return;
-			}
-			continue;
+			return;
 		}
-
-		// There is at least one thread waiting. Wake one thread and return.
-		ParkingLot::WakeOne(&State, [this](ParkingLot::FWakeState WakeState) -> uint64
-		{
-			constexpr uint8 ExpectedState = IsLockedFlag | HasWaitingThreadsFlag;
-			const uint8 NewState = WakeState.bHasWaitingThreads ? HasWaitingThreadsFlag : 0;
-			uint8 OldState = ExpectedState;
-			State.compare_exchange_strong(OldState, NewState, std::memory_order_release);
-			checkSlow(OldState == ExpectedState);
-			return 0;
-		});
-		return;
 	}
+
+	// There is at least one thread waiting. Wake one thread and return.
+	ParkingLot::WakeOne(&State, [this](ParkingLot::FWakeState WakeState) -> uint64
+	{
+		constexpr uint8 ExpectedState = IsLockedFlag | HasWaitingThreadsFlag;
+		const uint8 NewState = WakeState.bHasWaitingThreads ? HasWaitingThreadsFlag : 0;
+		uint8 OldState = ExpectedState;
+		State.compare_exchange_strong(OldState, NewState, std::memory_order_release);
+		checkSlow(OldState == ExpectedState);
+		return 0;
+	});
 }
 
 } // UE
diff --git a/Engine/Source/Runtime/Core/Private/Async/RecursiveMutex.cpp b/Engine/Source/Runtime/Core/Private/Async/RecursiveMutex.cpp
--- a/Engine/Source/Runtime/Core/Private/Async/RecursiveMutex.cpp
+++ b/Engine/Source/Runtime/Core/Private/Async/RecursiveMutex.cpp
@@ -36,13 +36,13 @@ bool FRecursiveMutex::TryLock()
 
 	// Try to acquire the lock if it was unlocked, even if there are waiting threads.
 	// Acquiring the lock despite the waiting threads means that this lock is not FIFO and thus not fair.
-	if (!CurrentState.bIsLocked)
+	if (LIKELY(!CurrentState.bIsLocked))
 	{
-		checkSlow(ThreadId.load(std::memory_order_relaxed) == 0 && CurrentState.RecurseCount == 0);
 		FState NewState = CurrentState;
 		NewState.bIsLocked = true;
 		if (LIKELY(State.compare_exchange_strong(CurrentState.Value, NewState.Value, std::memory_order_acquire)))
 		{
+			checkSlow(ThreadId.load(std::memory_order_relaxed) == 0 && CurrentState.RecurseCount == 0);
 			ThreadId.store(CurrentThreadId, std::memory_order_relaxed);
 			return true;
 		}
@@ -67,13 +67,13 @@ void FRecursiveMutex::Lock()
 
 	// Try to acquire the lock if it was unlocked, even if there are waiting threads.
 	// Acquiring the lock despite the waiting threads means that this lock is not FIFO and thus not fair.
-	if (!CurrentState.bIsLocked)
+	if (LIKELY(!CurrentState.bIsLocked))
 	{
-		checkSlow(ThreadId.load(std::memory_order_relaxed) == 0 && CurrentState.RecurseCount == 0);
 		FState NewState = CurrentState;
 		NewState.bIsLocked = true;
 		if (LIKELY(State.compare_exchange_weak(CurrentState.Value, NewState.Value, std::memory_order_acquire)))
 		{
+			checkSlow(ThreadId.load(std::memory_order_relaxed) == 0 && CurrentState.RecurseCount == 0);
 			ThreadId.store(CurrentThreadId, std::memory_order_relaxed);
 			return;
 		}
@@ -99,34 +99,34 @@ FORCENOINLINE void FRecursiveMutex::LockSlow(FState CurrentState, const uint32 C
 	{
 		// Try to acquire the lock if it was unlocked, even if there are waiting threads.
 		// Acquiring the lock despite the waiting threads means that this lock is not FIFO and thus not fair.
-		if (!CurrentState.bIsLocked)
+		if (LIKELY(!CurrentState.bIsLocked))
 		{
-			checkSlow(ThreadId.load(std::memory_order_relaxed) == 0 && CurrentState.RecurseCount == 0);
 			FState NewState = CurrentState;
 			NewState.bIsLocked = true;
 			if (LIKELY(State.compare_exchange_weak(CurrentState.Value, NewState.Value, std::memory_order_acquire)))
 			{
+				checkSlow(ThreadId.load(std::memory_order_relaxed) == 0 && CurrentState.RecurseCount == 0);
 				ThreadId.store(CurrentThreadId, std::memory_order_relaxed);
 				return;
 			}
 			continue;
 		}
 
-		// Spin up to the spin limit while there are no waiting threads.
-		if (!CurrentState.bHasWaitingThreads && SpinCount < SpinLimit)
+		if (LIKELY(!CurrentState.bHasWaitingThreads))
 		{
-			FPlatformProcess::YieldThread();
-			++SpinCount;
-			CurrentState.Value = State.load(std::memory_order_acquire);
-			continue;
-		}
+			// Spin up to the spin limit while there are no waiting threads.
+			if (LIKELY(SpinCount < SpinLimit))
+			{
+				FPlatformProcess::YieldThread();
+				++SpinCount;
+				CurrentState.Value = State.load(std::memory_order_acquire);
+				continue;
+			}
 
-		// Store that there are waiting threads. Restart if the state has changed since it was loaded.
-		if (!CurrentState.bHasWaitingThreads)
-		{
+			// Store that there are waiting threads. Restart if the state has changed since it was loaded.
 			FState NewState = CurrentState;
 			NewState.bHasWaitingThreads = true;
-			if (!State.compare_exchange_weak(CurrentState.Value, NewState.Value, std::memory_order_acquire))
+			if (UNLIKELY(!State.compare_exchange_weak(CurrentState.Value, NewState.Value, std::memory_order_acquire)))
 			{
 				continue;
 			}
@@ -151,7 +151,7 @@ void FRecursiveMutex::Unlock()
 		ThreadId.store(0, std::memory_order_relaxed);
 
 		// Unlocking with no waiting threads only requires resetting the state.
-		if (!CurrentState.bHasWaitingThreads)
+		if (LIKELY(!CurrentState.bHasWaitingThreads))
 		{
 			FState NewState;
 			if (LIKELY(State.compare_exchange_strong(CurrentState.Value, NewState.Value, std::memory_order_release, std::memory_order_relaxed)))
diff --git a/Engine/Source/Runtime/Core/Public/Async/RecursiveMutex.h b/Engine/Source/Runtime/Core/Public/Async/RecursiveMutex.h
--- a/Engine/Source/Runtime/Core/Public/Async/RecursiveMutex.h
+++ b/Engine/Source/Runtime/Core/Public/Async/RecursiveMutex.h
@@ -30,7 +30,7 @@ public:
 private:
 	union FState;
 
-	void LockSlow(FState CurrentState, const uint32 CurrentThreadId);
+	void LockSlow(FState CurrentState, uint32 CurrentThreadId);
 	void UnlockSlow(FState CurrentState);
 
 	std::atomic<uint32> State = 0;
