making  `FTimespan::MaxValue()` a special "infinity" value for `FTimeout`. this also happen to be MAX_UINT and so plays well with OS waiting with timeout functions

diff --git a/Engine/Source/Runtime/Core/Public/Misc/Timeout.h b/Engine/Source/Runtime/Core/Public/Misc/Timeout.h
--- a/Engine/Source/Runtime/Core/Public/Misc/Timeout.h
+++ b/Engine/Source/Runtime/Core/Public/Misc/Timeout.h
@@ -24,7 +24,7 @@ namespace UE
 
 		explicit operator bool() const
 		{
-			return GetElapsedTime() >= Timeout;
+			return GetRemainingTime() <= FTimespan::Zero();
 		}
 
 		FTimespan GetElapsedTime() const
@@ -34,7 +34,22 @@ namespace UE
 
 		FTimespan GetRemainingTime() const
 		{
-			return Timeout - GetElapsedTime();
+			return Timeout == FTimespan::MaxValue() ? FTimespan::MaxValue() : Timeout - GetElapsedTime();
+		}
+
+		static FTimeout Never()
+		{
+			return FTimeout{ FTimespan::MaxValue() };
+		}
+
+		FTimespan GetTimeoutValue() const
+		{
+			return Timeout;
+		}
+
+		friend bool operator==(FTimeout Left, FTimeout Right)
+		{
+			return Left.Timeout == Right.Timeout && (Left.Timeout == FTimespan::MaxValue() || Left.Start == Right.Start);
 		}
 
 	private:
