some delegates code shuffling split from a bigger CL to unclotter it

diff --git a/Engine/Source/Runtime/Core/Public/Delegates/DelegateBase.h b/Engine/Source/Runtime/Core/Public/Delegates/DelegateBase.h
--- a/Engine/Source/Runtime/Core/Public/Delegates/DelegateBase.h
+++ b/Engine/Source/Runtime/Core/Public/Delegates/DelegateBase.h
@@ -41,7 +41,7 @@ ALIAS_TEMPLATE_TYPE_LAYOUT(template<typename ElementType>, FDelegateAllocatorTyp
 /**
  * Base class for unicast delegates.
  */
-class FDelegateBase //-V1062
+class FDelegateBase
 {
 	template <typename>
 	friend class TMulticastDelegateBase;
@@ -59,19 +59,12 @@ protected:
 
 	FDelegateBase(FDelegateBase&& Other)
 	{
-		DelegateAllocator.MoveToEmpty(Other.DelegateAllocator);
-		DelegateSize = Other.DelegateSize;
-		Other.DelegateSize = 0;
+		MoveConstruct(MoveTemp(Other));
 	}
 
 	FDelegateBase& operator=(FDelegateBase&& Other)
 	{
-		UE_DELEGATES_MT_SCOPED_WRITE_ACCESS(AccessDetector);
-
-		Unbind();
-		DelegateAllocator.MoveToEmpty(Other.DelegateAllocator);
-		DelegateSize = Other.DelegateSize;
-		Other.DelegateSize = 0;
+		MoveAssign(MoveTemp(Other));
 		return *this;
 	}
 
@@ -111,9 +104,25 @@ protected:
 		return DelegateAllocator.GetAllocatedSize(DelegateSize, sizeof(FAlignedInlineDelegateType));
 	}
 
-private:
-	friend void* operator new(size_t Size, FDelegateBase& Base);
+public:
+	/**
+	 * "emplacement" of delegate instance of the given type
+	 */
+	template<typename DelegateInstanceType, typename... DelegateInstanceParams>
+	void CreateDelegateInstance(DelegateInstanceParams&&... Params)
+	{
+		UE_DELEGATES_MT_SCOPED_READ_ACCESS(AccessDetector);
+
+		IDelegateInstance* DelegateInstance = GetDelegateInstanceProtected();
+		if (DelegateInstance)
+		{
+			DelegateInstance->~IDelegateInstance();
+		}
+
+		new(Allocate(sizeof(DelegateInstanceType))) DelegateInstanceType(Forward<DelegateInstanceParams>(Params)...);
+	}
 
+private:
 	void* Allocate(int32 Size)
 	{
 		UE_DELEGATES_MT_SCOPED_WRITE_ACCESS(AccessDetector);
@@ -133,6 +142,23 @@ private:
 		return DelegateAllocator.GetAllocation();
 	}
 
+	void MoveConstruct(FDelegateBase&& Other)
+	{
+		DelegateAllocator.MoveToEmpty(Other.DelegateAllocator);
+		DelegateSize = Other.DelegateSize;
+		Other.DelegateSize = 0;
+	}
+
+	void MoveAssign(FDelegateBase&& Other)
+	{
+		UE_DELEGATES_MT_SCOPED_WRITE_ACCESS(AccessDetector);
+
+		Unbind();
+		DelegateAllocator.MoveToEmpty(Other.DelegateAllocator);
+		DelegateSize = Other.DelegateSize;
+		Other.DelegateSize = 0;
+	}
+
 private:
 	FDelegateAllocatorType::ForElementType<FAlignedInlineDelegateType> DelegateAllocator;
 	int32 DelegateSize = 0;
@@ -140,11 +166,6 @@ private:
 	UE_DELEGATES_MT_ACCESS_DETECTOR(AccessDetector);
 };
 
-inline void* operator new(size_t Size, FDelegateBase& Base)
-{
-	return Base.Allocate((int32)Size);
-}
-
 struct FDefaultDelegateUserPolicy
 {
 	// To extend delegates, you should implement a policy struct like this and pass it as the second template
diff --git a/Engine/Source/Runtime/Core/Public/Delegates/DelegateInstancesImpl.h b/Engine/Source/Runtime/Core/Public/Delegates/DelegateInstancesImpl.h
--- a/Engine/Source/Runtime/Core/Public/Delegates/DelegateInstancesImpl.h
+++ b/Engine/Source/Runtime/Core/Public/Delegates/DelegateInstancesImpl.h
@@ -12,6 +12,7 @@
 
 #include "CoreTypes.h"
 #include "Delegates/DelegateInstanceInterface.h"
+#include "Delegates/IDelegateInstance.h"
 #include "Misc/AssertionMacros.h"
 #include "Templates/RemoveReference.h"
 #include "Templates/SharedPointer.h"
@@ -20,8 +21,6 @@
 #include "UObject/NameTypes.h"
 #include "UObject/WeakObjectPtrTemplates.h"
 
-class FDelegateBase;
-class FDelegateHandle;
 enum class ESPMode : uint8;
 
 namespace UE::Delegates::Private
@@ -134,7 +133,7 @@ public:
 
 	void CreateCopy(DelegateBaseType& Base) const final
 	{
-		new (Base) TBaseUFunctionDelegateInstance(*this);
+		Base.template CreateDelegateInstance<TBaseUFunctionDelegateInstance>(*this);
 	}
 
 	RetValType Execute(ParamTypes... Params) const final
@@ -254,7 +253,7 @@ public:
 
 	void CreateCopy(DelegateBaseType& Base) const final
 	{
-		new (Base) TBaseSPMethodDelegateInstance(*this);
+		Base.template CreateDelegateInstance<TBaseSPMethodDelegateInstance>(*this);
 	}
 
 	RetValType Execute(ParamTypes... Params) const final
@@ -390,7 +389,7 @@ public:
 
 	void CreateCopy(DelegateBaseType& Base) const final
 	{
-		new (Base) TBaseRawMethodDelegateInstance(*this);
+		Base.template CreateDelegateInstance<TBaseRawMethodDelegateInstance>(*this);
 	}
 
 	RetValType Execute(ParamTypes... Params) const final
@@ -514,7 +513,7 @@ public:
 
 	void CreateCopy(DelegateBaseType& Base) const final
 	{
-		new (Base) TBaseUObjectMethodDelegateInstance(*this);
+		Base.template CreateDelegateInstance<TBaseUObjectMethodDelegateInstance>(*this);
 	}
 
 	RetValType Execute(ParamTypes... Params) const final
@@ -637,7 +636,7 @@ public:
 
 	void CreateCopy(DelegateBaseType& Base) const final
 	{
-		new (Base) TBaseStaticDelegateInstance(*this);
+		Base.template CreateDelegateInstance<TBaseStaticDelegateInstance>(*this);
 	}
 
 	RetValType Execute(ParamTypes... Params) const final
@@ -730,7 +729,7 @@ public:
 	// IBaseDelegateInstance interface
 	void CreateCopy(DelegateBaseType& Base) const final
 	{
-		new (Base) TBaseFunctorDelegateInstance(*this);
+		Base.template CreateDelegateInstance<TBaseFunctorDelegateInstance>(*this);
 	}
 
 	RetValType Execute(ParamTypes... Params) const final
@@ -824,7 +823,7 @@ public:
 	// IBaseDelegateInstance interface
 	void CreateCopy(DelegateBaseType& Base) const final
 	{
-		new (Base) TWeakBaseFunctorDelegateInstance(*this);
+		Base.template CreateDelegateInstance<TWeakBaseFunctorDelegateInstance>(*this);
 	}
 
 	RetValType Execute(ParamTypes... Params) const final
diff --git a/Engine/Source/Runtime/Core/Public/Delegates/DelegateSignatureImpl.inl b/Engine/Source/Runtime/Core/Public/Delegates/DelegateSignatureImpl.inl
--- a/Engine/Source/Runtime/Core/Public/Delegates/DelegateSignatureImpl.inl
+++ b/Engine/Source/Runtime/Core/Public/Delegates/DelegateSignatureImpl.inl
@@ -91,7 +91,7 @@ public:
 	UE_NODISCARD inline static TDelegate<RetValType(ParamTypes...), UserPolicy> CreateStatic(typename TIdentity<RetValType (*)(ParamTypes..., std::decay_t<VarTypes>...)>::Type InFunc, VarTypes&&... Vars)
 	{
 		TDelegate<RetValType(ParamTypes...), UserPolicy> Result;
-		new (Result) TBaseStaticDelegateInstance<FuncType, UserPolicy, std::decay_t<VarTypes>...>(InFunc, Forward<VarTypes>(Vars)...);
+		Result.template CreateDelegateInstance<TBaseStaticDelegateInstance<FuncType, UserPolicy, std::decay_t<VarTypes>...>>(InFunc, Forward<VarTypes>(Vars)...);
 		return Result;
 	}
 
@@ -103,7 +103,7 @@ public:
 	UE_NODISCARD inline static TDelegate<RetValType(ParamTypes...), UserPolicy> CreateLambda(FunctorType&& InFunctor, VarTypes&&... Vars)
 	{
 		TDelegate<RetValType(ParamTypes...), UserPolicy> Result;
-		new (Result) TBaseFunctorDelegateInstance<FuncType, UserPolicy, typename TRemoveReference<FunctorType>::Type, std::decay_t<VarTypes>...>(Forward<FunctorType>(InFunctor), Forward<VarTypes>(Vars)...);
+		Result.template CreateDelegateInstance<TBaseFunctorDelegateInstance<FuncType, UserPolicy, typename TRemoveReference<FunctorType>::Type, std::decay_t<VarTypes>...>>(Forward<FunctorType>(InFunctor), Forward<VarTypes>(Vars)...);
 		return Result;
 	}
 
@@ -115,7 +115,7 @@ public:
 	UE_NODISCARD inline static TDelegate<RetValType(ParamTypes...), UserPolicy> CreateWeakLambda(UserClass* InUserObject, FunctorType&& InFunctor, VarTypes&&... Vars)
 	{
 		TDelegate<RetValType(ParamTypes...), UserPolicy> Result;
-		new (Result) TWeakBaseFunctorDelegateInstance<UserClass, FuncType, UserPolicy, typename TRemoveReference<FunctorType>::Type, std::decay_t<VarTypes>...>(InUserObject, Forward<FunctorType>(InFunctor), Forward<VarTypes>(Vars)...);
+		Result.template CreateDelegateInstance<TWeakBaseFunctorDelegateInstance<UserClass, FuncType, UserPolicy, typename TRemoveReference<FunctorType>::Type, std::decay_t<VarTypes>...>>(InUserObject, Forward<FunctorType>(InFunctor), Forward<VarTypes>(Vars)...);
 		return Result;
 	}
 
@@ -131,14 +131,14 @@ public:
 		static_assert(!TIsConst<UserClass>::Value, "Attempting to bind a delegate with a const object pointer and non-const member function.");
 
 		TDelegate<RetValType(ParamTypes...), UserPolicy> Result;
-		new (Result) TBaseRawMethodDelegateInstance<false, UserClass, FuncType, UserPolicy, std::decay_t<VarTypes>...>(InUserObject, InFunc, Forward<VarTypes>(Vars)...);
+		Result.template CreateDelegateInstance<TBaseRawMethodDelegateInstance<false, UserClass, FuncType, UserPolicy, std::decay_t<VarTypes>...>>(InUserObject, InFunc, Forward<VarTypes>(Vars)...);
 		return Result;
 	}
 	template <typename UserClass, typename... VarTypes>
 	UE_NODISCARD inline static TDelegate<RetValType(ParamTypes...), UserPolicy> CreateRaw(const UserClass* InUserObject, typename TMemFunPtrType<true, UserClass, RetValType (ParamTypes..., std::decay_t<VarTypes>...)>::Type InFunc, VarTypes&&... Vars)
 	{
 		TDelegate<RetValType(ParamTypes...), UserPolicy> Result;
-		new (Result) TBaseRawMethodDelegateInstance<true, const UserClass, FuncType, UserPolicy, std::decay_t<VarTypes>...>(InUserObject, InFunc, Forward<VarTypes>(Vars)...);
+		Result.template CreateDelegateInstance<TBaseRawMethodDelegateInstance<true, const UserClass, FuncType, UserPolicy, std::decay_t<VarTypes>...>>(InUserObject, InFunc, Forward<VarTypes>(Vars)...);
 		return Result;
 	}
 
@@ -154,14 +154,14 @@ public:
 		static_assert(!TIsConst<UserClass>::Value, "Attempting to bind a delegate with a const object pointer and non-const member function.");
 
 		TDelegate<RetValType(ParamTypes...), UserPolicy> Result;
-		new (Result) TBaseSPMethodDelegateInstance<false, UserClass, Mode, FuncType, UserPolicy, std::decay_t<VarTypes>...>(InUserObjectRef, InFunc, Forward<VarTypes>(Vars)...);
+		Result.template CreateDelegateInstance<TBaseSPMethodDelegateInstance<false, UserClass, Mode, FuncType, UserPolicy, std::decay_t<VarTypes>...>>(InUserObjectRef, InFunc, Forward<VarTypes>(Vars)...);
 		return Result;
 	}
 	template <typename UserClass, ESPMode Mode, typename... VarTypes>
 	UE_NODISCARD inline static TDelegate<RetValType(ParamTypes...), UserPolicy> CreateSP(const TSharedRef<UserClass, Mode>& InUserObjectRef, typename TMemFunPtrType<true, UserClass, RetValType (ParamTypes..., std::decay_t<VarTypes>...)>::Type InFunc, VarTypes&&... Vars)
 	{
 		TDelegate<RetValType(ParamTypes...), UserPolicy> Result;
-		new (Result) TBaseSPMethodDelegateInstance<true, const UserClass, Mode, FuncType, UserPolicy, std::decay_t<VarTypes>...>(InUserObjectRef, InFunc, Forward<VarTypes>(Vars)...);
+		Result.template CreateDelegateInstance<TBaseSPMethodDelegateInstance<true, const UserClass, Mode, FuncType, UserPolicy, std::decay_t<VarTypes>...>>(InUserObjectRef, InFunc, Forward<VarTypes>(Vars)...);
 		return Result;
 	}
 
@@ -177,14 +177,14 @@ public:
 		static_assert(!TIsConst<UserClass>::Value, "Attempting to bind a delegate with a const object pointer and non-const member function.");
 
 		TDelegate<RetValType(ParamTypes...), UserPolicy> Result;
-		new (Result) TBaseSPMethodDelegateInstance<false, UserClass, decltype(InUserObject->AsShared())::Mode, FuncType, UserPolicy, std::decay_t<VarTypes>...>(StaticCastSharedRef<UserClass>(InUserObject->AsShared()), InFunc, Forward<VarTypes>(Vars)...);
+		Result.template CreateDelegateInstance<TBaseSPMethodDelegateInstance<false, UserClass, decltype(InUserObject->AsShared())::Mode, FuncType, UserPolicy, std::decay_t<VarTypes>...>>(StaticCastSharedRef<UserClass>(InUserObject->AsShared()), InFunc, Forward<VarTypes>(Vars)...);
 		return Result;
 	}
 	template <typename UserClass, typename... VarTypes>
 	UE_NODISCARD inline static TDelegate<RetValType(ParamTypes...), UserPolicy> CreateSP(const UserClass* InUserObject, typename TMemFunPtrType<true, UserClass, RetValType (ParamTypes..., std::decay_t<VarTypes>...)>::Type InFunc, VarTypes&&... Vars)
 	{
 		TDelegate<RetValType(ParamTypes...), UserPolicy> Result;
-		new (Result) TBaseSPMethodDelegateInstance<true, const UserClass, decltype(InUserObject->AsShared())::Mode, FuncType, UserPolicy, std::decay_t<VarTypes>...>(StaticCastSharedRef<const UserClass>(InUserObject->AsShared()), InFunc, Forward<VarTypes>(Vars)...);
+		Result.template CreateDelegateInstance<TBaseSPMethodDelegateInstance<true, const UserClass, decltype(InUserObject->AsShared())::Mode, FuncType, UserPolicy, std::decay_t<VarTypes>...>>(StaticCastSharedRef<const UserClass>(InUserObject->AsShared()), InFunc, Forward<VarTypes>(Vars)...);
 		return Result;
 	}
 
@@ -202,14 +202,14 @@ public:
 		static_assert(!TIsConst<UserClass>::Value, "Attempting to bind a delegate with a const object pointer and non-const member function.");
 
 		TDelegate<RetValType(ParamTypes...), UserPolicy> Result;
-		new (Result) TBaseSPMethodDelegateInstance<false, UserClass, ESPMode::ThreadSafe, FuncType, UserPolicy, std::decay_t<VarTypes>...>(InUserObjectRef, InFunc, Forward<VarTypes>(Vars)...);
+		Result.template CreateDelegateInstance<TBaseSPMethodDelegateInstance<false, UserClass, ESPMode::ThreadSafe, FuncType, UserPolicy, std::decay_t<VarTypes>...>>(InUserObjectRef, InFunc, Forward<VarTypes>(Vars)...);
 		return Result;
 	}
 	template <typename UserClass, typename... VarTypes>
 	UE_NODISCARD inline static TDelegate<RetValType(ParamTypes...), UserPolicy> CreateThreadSafeSP(const TSharedRef<UserClass, ESPMode::ThreadSafe>& InUserObjectRef, typename TMemFunPtrType<true, UserClass, RetValType (ParamTypes..., std::decay_t<VarTypes>...)>::Type InFunc, VarTypes&&... Vars)
 	{
 		TDelegate<RetValType(ParamTypes...), UserPolicy> Result;
-		new (Result) TBaseSPMethodDelegateInstance<true, const UserClass, ESPMode::ThreadSafe, FuncType, UserPolicy, std::decay_t<VarTypes>...>(InUserObjectRef, InFunc, Forward<VarTypes>(Vars)...);
+		Result.template CreateDelegateInstance<TBaseSPMethodDelegateInstance<true, const UserClass, ESPMode::ThreadSafe, FuncType, UserPolicy, std::decay_t<VarTypes>...>>(InUserObjectRef, InFunc, Forward<VarTypes>(Vars)...);
 		return Result;
 	}
 
@@ -227,14 +227,14 @@ public:
 		static_assert(!TIsConst<UserClass>::Value, "Attempting to bind a delegate with a const object pointer and non-const member function.");
 
 		TDelegate<RetValType(ParamTypes...), UserPolicy> Result;
-		new (Result) TBaseSPMethodDelegateInstance<false, UserClass, ESPMode::ThreadSafe, FuncType, UserPolicy, std::decay_t<VarTypes>...>(StaticCastSharedRef<UserClass>(InUserObject->AsShared()), InFunc, Forward<VarTypes>(Vars)...);
+		Result.template CreateDelegateInstance<TBaseSPMethodDelegateInstance<false, UserClass, ESPMode::ThreadSafe, FuncType, UserPolicy, std::decay_t<VarTypes>...>>(StaticCastSharedRef<UserClass>(InUserObject->AsShared()), InFunc, Forward<VarTypes>(Vars)...);
 		return Result;
 	}
 	template <typename UserClass, typename... VarTypes>
 	UE_NODISCARD inline static TDelegate<RetValType(ParamTypes...), UserPolicy> CreateThreadSafeSP(const UserClass* InUserObject, typename TMemFunPtrType<true, UserClass, RetValType (ParamTypes..., std::decay_t<VarTypes>...)>::Type InFunc, VarTypes&&... Vars)
 	{
 		TDelegate<RetValType(ParamTypes...), UserPolicy> Result;
-		new (Result) TBaseSPMethodDelegateInstance<true, const UserClass, ESPMode::ThreadSafe, FuncType, UserPolicy, std::decay_t<VarTypes>...>(StaticCastSharedRef<const UserClass>(InUserObject->AsShared()), InFunc, Forward<VarTypes>(Vars)...);
+		Result.template CreateDelegateInstance<TBaseSPMethodDelegateInstance<true, const UserClass, ESPMode::ThreadSafe, FuncType, UserPolicy, std::decay_t<VarTypes>...>>(StaticCastSharedRef<const UserClass>(InUserObject->AsShared()), InFunc, Forward<VarTypes>(Vars)...);
 		return Result;
 	}
 
@@ -248,14 +248,14 @@ public:
 	UE_NODISCARD inline static TDelegate<RetValType(ParamTypes...), UserPolicy> CreateUFunction(UObjectTemplate* InUserObject, const FName& InFunctionName, VarTypes&&... Vars)
 	{
 		TDelegate<RetValType(ParamTypes...), UserPolicy> Result;
-		new (Result) TBaseUFunctionDelegateInstance<UObjectTemplate, FuncType, UserPolicy, std::decay_t<VarTypes>...>(InUserObject, InFunctionName, Forward<VarTypes>(Vars)...);
+		Result.template CreateDelegateInstance<TBaseUFunctionDelegateInstance<UObjectTemplate, FuncType, UserPolicy, std::decay_t<VarTypes>...>>(InUserObject, InFunctionName, Forward<VarTypes>(Vars)...);
 		return Result;
 	}
 	template <typename UObjectTemplate, typename... VarTypes>
 	UE_NODISCARD inline static TDelegate<RetValType(ParamTypes...), UserPolicy> CreateUFunction(TObjectPtr<UObjectTemplate> InUserObject, const FName& InFunctionName, VarTypes&&... Vars)
 	{
 		TDelegate<RetValType(ParamTypes...), UserPolicy> Result;
-		new (Result) TBaseUFunctionDelegateInstance<UObjectTemplate, FuncType, UserPolicy, std::decay_t<VarTypes>...>(ToRawPtr(InUserObject), InFunctionName, Forward<VarTypes>(Vars)...);
+		Result.template CreateDelegateInstance<TBaseUFunctionDelegateInstance<UObjectTemplate, FuncType, UserPolicy, std::decay_t<VarTypes>...>>(ToRawPtr(InUserObject), InFunctionName, Forward<VarTypes>(Vars)...);
 		return Result;
 	}
 
@@ -271,14 +271,14 @@ public:
 		static_assert(!TIsConst<UserClass>::Value, "Attempting to bind a delegate with a const object pointer and non-const member function.");
 
 		TDelegate<RetValType(ParamTypes...), UserPolicy> Result;
-		new (Result) TBaseUObjectMethodDelegateInstance<false, UserClass, FuncType, UserPolicy, std::decay_t<VarTypes>...>(InUserObject, InFunc, Forward<VarTypes>(Vars)...);
+		Result.template CreateDelegateInstance<TBaseUObjectMethodDelegateInstance<false, UserClass, FuncType, UserPolicy, std::decay_t<VarTypes>...>>(InUserObject, InFunc, Forward<VarTypes>(Vars)...);
 		return Result;
 	}
 	template <typename UserClass, typename... VarTypes>
 	UE_NODISCARD inline static TDelegate<RetValType(ParamTypes...), UserPolicy> CreateUObject(const UserClass* InUserObject, typename TMemFunPtrType<true, UserClass, RetValType (ParamTypes..., std::decay_t<VarTypes>...)>::Type InFunc, VarTypes&&... Vars)
 	{
 		TDelegate<RetValType(ParamTypes...), UserPolicy> Result;
-		new (Result) TBaseUObjectMethodDelegateInstance<true, const UserClass, FuncType, UserPolicy, std::decay_t<VarTypes>...>(InUserObject, InFunc, Forward<VarTypes>(Vars)...);
+		Result.template CreateDelegateInstance<TBaseUObjectMethodDelegateInstance<true, const UserClass, FuncType, UserPolicy, std::decay_t<VarTypes>...>>(InUserObject, InFunc, Forward<VarTypes>(Vars)...);
 		return Result;
 	}
 	template <typename UserClass, typename... VarTypes>
@@ -287,14 +287,14 @@ public:
 		static_assert(!TIsConst<UserClass>::Value, "Attempting to bind a delegate with a const object pointer and non-const member function.");
 
 		TDelegate<RetValType(ParamTypes...), UserPolicy> Result;
-		new (Result) TBaseUObjectMethodDelegateInstance<false, UserClass, FuncType, UserPolicy, std::decay_t<VarTypes>...>(ToRawPtr(InUserObject), InFunc, Forward<VarTypes>(Vars)...);
+		Result.template CreateDelegateInstance<TBaseUObjectMethodDelegateInstance<false, UserClass, FuncType, UserPolicy, std::decay_t<VarTypes>...>>(ToRawPtr(InUserObject), InFunc, Forward<VarTypes>(Vars)...);
 		return Result;
 	}
 	template <typename UserClass, typename... VarTypes>
 	UE_NODISCARD inline static TDelegate<RetValType(ParamTypes...), UserPolicy> CreateUObject(TObjectPtr<UserClass> InUserObject, typename TMemFunPtrType<true, UserClass, RetValType (ParamTypes..., std::decay_t<VarTypes>...)>::Type InFunc, VarTypes&&... Vars)
 	{
 		TDelegate<RetValType(ParamTypes...), UserPolicy> Result;
-		new (Result) TBaseUObjectMethodDelegateInstance<true, const UserClass, FuncType, UserPolicy, std::decay_t<VarTypes>...>(ToRawPtr(InUserObject), InFunc, Forward<VarTypes>(Vars)...);
+		Result.template CreateDelegateInstance<TBaseUObjectMethodDelegateInstance<true, const UserClass, FuncType, UserPolicy, std::decay_t<VarTypes>...>>(ToRawPtr(InUserObject), InFunc, Forward<VarTypes>(Vars)...);
 		return Result;
 	}
 
@@ -347,7 +347,7 @@ public:
 	template <typename... VarTypes>
 	inline void BindStatic(typename TBaseStaticDelegateInstance<FuncType, UserPolicy, std::decay_t<VarTypes>...>::FFuncPtr InFunc, VarTypes&&... Vars)
 	{
-		new (*this) TBaseStaticDelegateInstance<FuncType, UserPolicy, std::decay_t<VarTypes>...>(InFunc, Forward<VarTypes>(Vars)...);
+		Super::template CreateDelegateInstance<TBaseStaticDelegateInstance<FuncType, UserPolicy, std::decay_t<VarTypes>...>>(InFunc, Forward<VarTypes>(Vars)...);
 	}
 	
 	/**
@@ -357,7 +357,7 @@ public:
 	template<typename FunctorType, typename... VarTypes>
 	inline void BindLambda(FunctorType&& InFunctor, VarTypes&&... Vars)
 	{
-		new (*this) TBaseFunctorDelegateInstance<FuncType, UserPolicy, typename TRemoveReference<FunctorType>::Type, std::decay_t<VarTypes>...>(Forward<FunctorType>(InFunctor), Forward<VarTypes>(Vars)...);
+		Super::template CreateDelegateInstance<TBaseFunctorDelegateInstance<FuncType, UserPolicy, typename TRemoveReference<FunctorType>::Type, std::decay_t<VarTypes>...>>(Forward<FunctorType>(InFunctor), Forward<VarTypes>(Vars)...);
 	}
 
 	/**
@@ -367,7 +367,7 @@ public:
 	template<typename UserClass, typename FunctorType, typename... VarTypes>
 	inline void BindWeakLambda(UserClass* InUserObject, FunctorType&& InFunctor, VarTypes&&... Vars)
 	{
-		new (*this) TWeakBaseFunctorDelegateInstance<UserClass, FuncType, UserPolicy, typename TRemoveReference<FunctorType>::Type, std::decay_t<VarTypes>...>(InUserObject, Forward<FunctorType>(InFunctor), Forward<VarTypes>(Vars)...);
+		Super::template CreateDelegateInstance<TWeakBaseFunctorDelegateInstance<UserClass, FuncType, UserPolicy, typename TRemoveReference<FunctorType>::Type, std::decay_t<VarTypes>...>>(InUserObject, Forward<FunctorType>(InFunctor), Forward<VarTypes>(Vars)...);
 	}
 
 	/**
@@ -381,12 +381,12 @@ public:
 	{
 		static_assert(!TIsConst<UserClass>::Value, "Attempting to bind a delegate with a const object pointer and non-const member function.");
 
-		new (*this) TBaseRawMethodDelegateInstance<false, UserClass, FuncType, UserPolicy, std::decay_t<VarTypes>...>(InUserObject, InFunc, Forward<VarTypes>(Vars)...);
+		Super::template CreateDelegateInstance<TBaseRawMethodDelegateInstance<false, UserClass, FuncType, UserPolicy, std::decay_t<VarTypes>...>>(InUserObject, InFunc, Forward<VarTypes>(Vars)...);
 	}
 	template <typename UserClass, typename... VarTypes>
 	inline void BindRaw(const UserClass* InUserObject, typename TMemFunPtrType<true, UserClass, RetValType (ParamTypes..., std::decay_t<VarTypes>...)>::Type InFunc, VarTypes&&... Vars)
 	{
-		new (*this) TBaseRawMethodDelegateInstance<true, const UserClass, FuncType, UserPolicy, std::decay_t<VarTypes>...>(InUserObject, InFunc, Forward<VarTypes>(Vars)...);
+		Super::template CreateDelegateInstance<TBaseRawMethodDelegateInstance<true, const UserClass, FuncType, UserPolicy, std::decay_t<VarTypes>...>>(InUserObject, InFunc, Forward<VarTypes>(Vars)...);
 	}
 
 	/**
@@ -397,12 +397,12 @@ public:
 	{
 		static_assert(!TIsConst<UserClass>::Value, "Attempting to bind a delegate with a const object pointer and non-const member function.");
 
-		new (*this) TBaseSPMethodDelegateInstance<false, UserClass, Mode, FuncType, UserPolicy, std::decay_t<VarTypes>...>(InUserObjectRef, InFunc, Forward<VarTypes>(Vars)...);
+		Super::template CreateDelegateInstance<TBaseSPMethodDelegateInstance<false, UserClass, Mode, FuncType, UserPolicy, std::decay_t<VarTypes>...>>(InUserObjectRef, InFunc, Forward<VarTypes>(Vars)...);
 	}
 	template <typename UserClass, ESPMode Mode, typename... VarTypes>
 	inline void BindSP(const TSharedRef<UserClass, Mode>& InUserObjectRef, typename TMemFunPtrType<true, UserClass, RetValType (ParamTypes..., std::decay_t<VarTypes>...)>::Type InFunc, VarTypes&&... Vars)
 	{
-		new (*this) TBaseSPMethodDelegateInstance<true, const UserClass, Mode, FuncType, UserPolicy, std::decay_t<VarTypes>...>(InUserObjectRef, InFunc, Forward<VarTypes>(Vars)...);
+		Super::template CreateDelegateInstance<TBaseSPMethodDelegateInstance<true, const UserClass, Mode, FuncType, UserPolicy, std::decay_t<VarTypes>...>>(InUserObjectRef, InFunc, Forward<VarTypes>(Vars)...);
 	}
 
 	/**
@@ -416,12 +416,12 @@ public:
 	{
 		static_assert(!TIsConst<UserClass>::Value, "Attempting to bind a delegate with a const object pointer and non-const member function.");
 
-		new (*this) TBaseSPMethodDelegateInstance<false, UserClass, decltype(InUserObject->AsShared())::Mode, FuncType, UserPolicy, std::decay_t<VarTypes>...>(StaticCastSharedRef<UserClass>(InUserObject->AsShared()), InFunc, Forward<VarTypes>(Vars)...);
+		Super::template CreateDelegateInstance<TBaseSPMethodDelegateInstance<false, UserClass, decltype(InUserObject->AsShared())::Mode, FuncType, UserPolicy, std::decay_t<VarTypes>...>>(StaticCastSharedRef<UserClass>(InUserObject->AsShared()), InFunc, Forward<VarTypes>(Vars)...);
 	}
 	template <typename UserClass, typename... VarTypes>
 	inline void BindSP(const UserClass* InUserObject, typename TMemFunPtrType<true, UserClass, RetValType (ParamTypes..., std::decay_t<VarTypes>...)>::Type InFunc, VarTypes&&... Vars)
 	{
-		new (*this) TBaseSPMethodDelegateInstance<true, const UserClass, decltype(InUserObject->AsShared())::Mode, FuncType, UserPolicy, std::decay_t<VarTypes>...>(StaticCastSharedRef<const UserClass>(InUserObject->AsShared()), InFunc, Forward<VarTypes>(Vars)...);
+		Super::template CreateDelegateInstance<TBaseSPMethodDelegateInstance<true, const UserClass, decltype(InUserObject->AsShared())::Mode, FuncType, UserPolicy, std::decay_t<VarTypes>...>>(StaticCastSharedRef<const UserClass>(InUserObject->AsShared()), InFunc, Forward<VarTypes>(Vars)...);
 	}
 
 	/**
@@ -437,12 +437,12 @@ public:
 	{
 		static_assert(!TIsConst<UserClass>::Value, "Attempting to bind a delegate with a const object pointer and non-const member function.");
 
-		new (*this) TBaseSPMethodDelegateInstance<false, UserClass, ESPMode::ThreadSafe, FuncType, UserPolicy, std::decay_t<VarTypes>...>(InUserObjectRef, InFunc, Forward<VarTypes>(Vars)...);
+		Super::template CreateDelegateInstance<TBaseSPMethodDelegateInstance<false, UserClass, ESPMode::ThreadSafe, FuncType, UserPolicy, std::decay_t<VarTypes>...>>(InUserObjectRef, InFunc, Forward<VarTypes>(Vars)...);
 	}
 	template <typename UserClass, typename... VarTypes>
 	inline void BindThreadSafeSP(const TSharedRef<UserClass, ESPMode::ThreadSafe>& InUserObjectRef, typename TMemFunPtrType<true, UserClass, RetValType (ParamTypes..., std::decay_t<VarTypes>...)>::Type InFunc, VarTypes&&... Vars)
 	{
-		new (*this) TBaseSPMethodDelegateInstance<true, const UserClass, ESPMode::ThreadSafe, FuncType, UserPolicy, std::decay_t<VarTypes>...>(InUserObjectRef, InFunc, Forward<VarTypes>(Vars)...);
+		Super::template CreateDelegateInstance<TBaseSPMethodDelegateInstance<true, const UserClass, ESPMode::ThreadSafe, FuncType, UserPolicy, std::decay_t<VarTypes>...>>(InUserObjectRef, InFunc, Forward<VarTypes>(Vars)...);
 	}
 
 	/**
@@ -458,12 +458,12 @@ public:
 	{
 		static_assert(!TIsConst<UserClass>::Value, "Attempting to bind a delegate with a const object pointer and non-const member function.");
 
-		new (*this) TBaseSPMethodDelegateInstance<false, UserClass, ESPMode::ThreadSafe, FuncType, UserPolicy, std::decay_t<VarTypes>...>(StaticCastSharedRef<UserClass>(InUserObject->AsShared()), InFunc, Forward<VarTypes>(Vars)...);
+		Super::template CreateDelegateInstance<TBaseSPMethodDelegateInstance<false, UserClass, ESPMode::ThreadSafe, FuncType, UserPolicy, std::decay_t<VarTypes>...>>(StaticCastSharedRef<UserClass>(InUserObject->AsShared()), InFunc, Forward<VarTypes>(Vars)...);
 	}
 	template <typename UserClass, typename... VarTypes>
 	inline void BindThreadSafeSP(const UserClass* InUserObject, typename TMemFunPtrType<true, UserClass, RetValType (ParamTypes..., std::decay_t<VarTypes>...)>::Type InFunc, VarTypes&&... Vars)
 	{
-		new (*this) TBaseSPMethodDelegateInstance<true, const UserClass, ESPMode::ThreadSafe, FuncType, UserPolicy, std::decay_t<VarTypes>...>(StaticCastSharedRef<const UserClass>(InUserObject->AsShared()), InFunc, Forward<VarTypes>(Vars)...);
+		Super::template CreateDelegateInstance<TBaseSPMethodDelegateInstance<true, const UserClass, ESPMode::ThreadSafe, FuncType, UserPolicy, std::decay_t<VarTypes>...>>(StaticCastSharedRef<const UserClass>(InUserObject->AsShared()), InFunc, Forward<VarTypes>(Vars)...);
 	}
 
 	/**
@@ -475,12 +475,12 @@ public:
 	template <typename UObjectTemplate, typename... VarTypes>
 	inline void BindUFunction(UObjectTemplate* InUserObject, const FName& InFunctionName, VarTypes&&... Vars)
 	{
-		new (*this) TBaseUFunctionDelegateInstance<UObjectTemplate, FuncType, UserPolicy, std::decay_t<VarTypes>...>(InUserObject, InFunctionName, Forward<VarTypes>(Vars)...);
+		Super::template CreateDelegateInstance<TBaseUFunctionDelegateInstance<UObjectTemplate, FuncType, UserPolicy, std::decay_t<VarTypes>...>>(InUserObject, InFunctionName, Forward<VarTypes>(Vars)...);
 	}
 	template <typename UObjectTemplate, typename... VarTypes>
 	inline void BindUFunction(TObjectPtr<UObjectTemplate> InUserObject, const FName& InFunctionName, VarTypes&&... Vars)
 	{
-		new (*this) TBaseUFunctionDelegateInstance<UObjectTemplate, FuncType, UserPolicy, std::decay_t<VarTypes>...>(ToRawPtr(InUserObject), InFunctionName, Forward<VarTypes>(Vars)...);
+		Super::template CreateDelegateInstance<TBaseUFunctionDelegateInstance<UObjectTemplate, FuncType, UserPolicy, std::decay_t<VarTypes>...>>(ToRawPtr(InUserObject), InFunctionName, Forward<VarTypes>(Vars)...);
 	}
 
 	/**
@@ -494,24 +494,24 @@ public:
 	{
 		static_assert(!TIsConst<UserClass>::Value, "Attempting to bind a delegate with a const object pointer and non-const member function.");
 
-		new (*this) TBaseUObjectMethodDelegateInstance<false, UserClass, FuncType, UserPolicy, std::decay_t<VarTypes>...>(InUserObject, InFunc, Forward<VarTypes>(Vars)...);
+		Super::template CreateDelegateInstance<TBaseUObjectMethodDelegateInstance<false, UserClass, FuncType, UserPolicy, std::decay_t<VarTypes>...>>(InUserObject, InFunc, Forward<VarTypes>(Vars)...);
 	}
 	template <typename UserClass, typename... VarTypes>
 	inline void BindUObject(const UserClass* InUserObject, typename TMemFunPtrType<true, UserClass, RetValType (ParamTypes..., std::decay_t<VarTypes>...)>::Type InFunc, VarTypes&&... Vars)
 	{
-		new (*this) TBaseUObjectMethodDelegateInstance<true, const UserClass, FuncType, UserPolicy, std::decay_t<VarTypes>...>(InUserObject, InFunc, Forward<VarTypes>(Vars)...);
+		Super::template CreateDelegateInstance<TBaseUObjectMethodDelegateInstance<true, const UserClass, FuncType, UserPolicy, std::decay_t<VarTypes>...>>(InUserObject, InFunc, Forward<VarTypes>(Vars)...);
 	}
 	template <typename UserClass, typename... VarTypes>
 	inline void BindUObject(TObjectPtr<UserClass> InUserObject, typename TMemFunPtrType<false, UserClass, RetValType (ParamTypes..., std::decay_t<VarTypes>...)>::Type InFunc, VarTypes&&... Vars)
 	{
 		static_assert(!TIsConst<UserClass>::Value, "Attempting to bind a delegate with a const object pointer and non-const member function.");
 
-		new (*this) TBaseUObjectMethodDelegateInstance<false, UserClass, FuncType, UserPolicy, std::decay_t<VarTypes>...>(ToRawPtr(InUserObject), InFunc, Forward<VarTypes>(Vars)...);
+		Super::template CreateDelegateInstance<TBaseUObjectMethodDelegateInstance<false, UserClass, FuncType, UserPolicy, std::decay_t<VarTypes>...>>(ToRawPtr(InUserObject), InFunc, Forward<VarTypes>(Vars)...);
 	}
 	template <typename UserClass, typename... VarTypes>
 	inline void BindUObject(TObjectPtr<UserClass> InUserObject, typename TMemFunPtrType<true, UserClass, RetValType (ParamTypes..., std::decay_t<VarTypes>...)>::Type InFunc, VarTypes&&... Vars)
 	{
-		new (*this) TBaseUObjectMethodDelegateInstance<true, const UserClass, FuncType, UserPolicy, std::decay_t<VarTypes>...>(ToRawPtr(InUserObject), InFunc, Forward<VarTypes>(Vars)...);
+		Super::template CreateDelegateInstance<TBaseUObjectMethodDelegateInstance<true, const UserClass, FuncType, UserPolicy, std::decay_t<VarTypes>...>>(ToRawPtr(InUserObject), InFunc, Forward<VarTypes>(Vars)...);
 	}
 
 public:
