Made FPlatformProcess::WritePipe() use proper UTF-8 conversion rather than truncation.

diff --git a/Engine/Source/Runtime/Core/Private/Mac/MacPlatformProcess.cpp b/Engine/Source/Runtime/Core/Private/Mac/MacPlatformProcess.cpp
--- a/Engine/Source/Runtime/Core/Private/Mac/MacPlatformProcess.cpp
+++ b/Engine/Source/Runtime/Core/Private/Mac/MacPlatformProcess.cpp
@@ -1404,19 +1404,17 @@ bool FMacPlatformProcess::ReadPipeToArray(void* ReadPipe, TArray<uint8>& Output)
 bool FMacPlatformProcess::WritePipe(void* WritePipe, const FString& Message, FString* OutWritten)
 {
 	// if there is not a message or WritePipe is nullptr
-	if ((Message.Len() == 0) || (WritePipe == nullptr))
+	int32 MessageLen = Message.Len();
+	if ((MessageLen == 0) || (WritePipe == nullptr))
 	{
 		return false;
 	}
 
 	// Convert input to UTF8CHAR
-	uint32 BytesAvailable = Message.Len();
-	UTF8CHAR * Buffer = new UTF8CHAR[BytesAvailable + 2];
-	for (uint32 i = 0; i < BytesAvailable; i++)
-	{
-		Buffer[i] = (UTF8CHAR)Message[i];
-	}
-	Buffer[BytesAvailable] = (UTF8CHAR)'\n';
+	const TCHAR* MessagePtr = *Message;
+	int32 BytesAvailable = FPlatformString::ConvertedLength<UTF8CHAR>(MessagePtr, MessageLen);
+	UTF8CHAR* Buffer = new UTF8CHAR[BytesAvailable + 2];
+	*FPlatformString::Convert(Buffer, BytesAvailable, MessagePtr, MessageLen) = (UTF8CHAR)'\n';
 
 	// Write to pipe
 	uint32 BytesWritten = write([(NSFileHandle*)WritePipe fileDescriptor], Buffer, BytesAvailable + 1);
@@ -1424,8 +1422,7 @@ bool FMacPlatformProcess::WritePipe(void* WritePipe, const FString& Message, FSt
 	// Get written message
 	if (OutWritten)
 	{
-		Buffer[BytesWritten] = (UTF8CHAR)'\0';
-		*OutWritten = FUTF8ToTCHAR((const ANSICHAR*)Buffer).Get();
+		*OutWritten = StringCast<TCHAR>(Buffer, BytesWritten).Get();
 	}
 
 	delete[] Buffer;
diff --git a/Engine/Source/Runtime/Core/Private/Unix/UnixPlatformProcess.cpp b/Engine/Source/Runtime/Core/Private/Unix/UnixPlatformProcess.cpp
--- a/Engine/Source/Runtime/Core/Private/Unix/UnixPlatformProcess.cpp
+++ b/Engine/Source/Runtime/Core/Private/Unix/UnixPlatformProcess.cpp
@@ -582,19 +582,17 @@ bool FUnixPlatformProcess::ReadPipeToArray(void* ReadPipe, TArray<uint8> & Outpu
 bool FUnixPlatformProcess::WritePipe(void* WritePipe, const FString& Message, FString* OutWritten)
 {
 	// if there is not a message or WritePipe is null
-	if ((Message.Len() == 0) || (WritePipe == nullptr))
+	int32 MessageLen = Message.Len();
+	if ((MessageLen == 0) || (WritePipe == nullptr))
 	{
 		return false;
 	}
 
 	// Convert input to UTF8CHAR
-	uint32 BytesAvailable = Message.Len();
-	UTF8CHAR * Buffer = new UTF8CHAR[BytesAvailable + 2];
-	for (uint32 i = 0; i < BytesAvailable; i++)
-	{
-		Buffer[i] = (UTF8CHAR)Message[i];
-	}
-	Buffer[BytesAvailable] = (UTF8CHAR)'\n';
+	const TCHAR* MessagePtr = *Message;
+	int32 BytesAvailable = FPlatformString::ConvertedLength<UTF8CHAR>(MessagePtr, MessageLen);
+	UTF8CHAR* Buffer = new UTF8CHAR[BytesAvailable + 2];
+	*FPlatformString::Convert(Buffer, BytesAvailable, MessagePtr, MessageLen) = (UTF8CHAR)'\n';
 
 	// write to pipe
 	uint32 BytesWritten = write(*(int*)WritePipe, Buffer, BytesAvailable + 1);
@@ -602,8 +600,7 @@ bool FUnixPlatformProcess::WritePipe(void* WritePipe, const FString& Message, FS
 	// Get written message
 	if (OutWritten)
 	{
-		Buffer[BytesWritten] = (UTF8CHAR)'\0';
-		*OutWritten = FUTF8ToTCHAR((const ANSICHAR*)Buffer).Get();
+		*OutWritten = StringCast<TCHAR>(Buffer, BytesWritten).Get();
 	}
 
 	delete[] Buffer;
diff --git a/Engine/Source/Runtime/Core/Private/Windows/WindowsPlatformProcess.cpp b/Engine/Source/Runtime/Core/Private/Windows/WindowsPlatformProcess.cpp
--- a/Engine/Source/Runtime/Core/Private/Windows/WindowsPlatformProcess.cpp
+++ b/Engine/Source/Runtime/Core/Private/Windows/WindowsPlatformProcess.cpp
@@ -1682,19 +1682,17 @@ bool FWindowsPlatformProcess::ReadPipeToArray(void* ReadPipe, TArray<uint8> & Ou
 bool FWindowsPlatformProcess::WritePipe(void* WritePipe, const FString& Message, FString* OutWritten)
 {
 	// If there is not a message or WritePipe is null
-	if (Message.Len() == 0 || WritePipe == nullptr)
+	int32 MessageLen = Message.Len();
+	if (MessageLen == 0 || WritePipe == nullptr)
 	{
 		return false;
 	}
 
 	// Convert input to UTF8CHAR
-	uint32 BytesAvailable = Message.Len();
-	UTF8CHAR * Buffer = new UTF8CHAR[BytesAvailable + 2];
-	for (uint32 i = 0; i < BytesAvailable; i++)
-	{
-		Buffer[i] = (UTF8CHAR)Message[i];
-	}
-	Buffer[BytesAvailable] = (UTF8CHAR)'\n';
+	const TCHAR* MessagePtr = *Message;
+	int32 BytesAvailable = FPlatformString::ConvertedLength<UTF8CHAR>(MessagePtr, MessageLen);
+	UTF8CHAR* Buffer = new UTF8CHAR[BytesAvailable + 2];
+	*FPlatformString::Convert(Buffer, BytesAvailable, MessagePtr, MessageLen) = (UTF8CHAR)'\n';
 
 	// Write to pipe
 	uint32 BytesWritten = 0;
@@ -1703,8 +1701,7 @@ bool FWindowsPlatformProcess::WritePipe(void* WritePipe, const FString& Message,
 	// Get written message
 	if (OutWritten)
 	{
-		Buffer[BytesWritten] = (UTF8CHAR)'\0';
-		*OutWritten = FUTF8ToTCHAR((const ANSICHAR*)Buffer).Get();
+		*OutWritten = StringCast<TCHAR>(Buffer, BytesWritten).Get();
 	}
 
 	delete[] Buffer;
