Removed some includes of Function.h in Core

#rnx

diff --git a/Engine/Plugins/Developer/Concert/ConcertSync/ConcertSyncCore/Source/ConcertSyncCore/Private/HistoryEdition/ActivityDependencyGraph.cpp b/Engine/Plugins/Developer/Concert/ConcertSync/ConcertSyncCore/Source/ConcertSyncCore/Private/HistoryEdition/ActivityDependencyGraph.cpp
--- a/Engine/Plugins/Developer/Concert/ConcertSync/ConcertSyncCore/Source/ConcertSyncCore/Private/HistoryEdition/ActivityDependencyGraph.cpp
+++ b/Engine/Plugins/Developer/Concert/ConcertSync/ConcertSyncCore/Source/ConcertSyncCore/Private/HistoryEdition/ActivityDependencyGraph.cpp
@@ -6,6 +6,7 @@
 #include "CoreGlobals.h"
 #include "HistoryEdition/ActivityDependencyEdge.h"
 #include "HistoryEdition/ActivityNode.h"
+#include "Templates/Function.h"
 
 #include "Algo/ForEach.h"
 
diff --git a/Engine/Plugins/Developer/Concert/ConcertSync/ConcertSyncCore/Source/ConcertSyncCore/Public/HistoryEdition/ActivityDependencyGraph.h b/Engine/Plugins/Developer/Concert/ConcertSync/ConcertSyncCore/Source/ConcertSyncCore/Public/HistoryEdition/ActivityDependencyGraph.h
--- a/Engine/Plugins/Developer/Concert/ConcertSync/ConcertSyncCore/Source/ConcertSyncCore/Public/HistoryEdition/ActivityDependencyGraph.h
+++ b/Engine/Plugins/Developer/Concert/ConcertSync/ConcertSyncCore/Source/ConcertSyncCore/Public/HistoryEdition/ActivityDependencyGraph.h
@@ -4,6 +4,7 @@
 
 #include "ActivityNode.h"
 
+template <typename FuncType> class TFunctionRef;
 template <typename OptionalType> struct TOptional;
 
 namespace UE::ConcertSyncCore { class FActivityDependencyEdge; }
diff --git a/Engine/Plugins/Editor/AssetSearch/Source/Private/AssetSearchDatabase.h b/Engine/Plugins/Editor/AssetSearch/Source/Private/AssetSearchDatabase.h
--- a/Engine/Plugins/Editor/AssetSearch/Source/Private/AssetSearchDatabase.h
+++ b/Engine/Plugins/Editor/AssetSearch/Source/Private/AssetSearchDatabase.h
@@ -12,6 +12,8 @@ enum class ESQLiteDatabaseOpenMode : uint8;
 class FAssetSearchDatabaseStatements;
 struct FAssetData;
 
+template <typename FuncType> class TFunctionRef;
+
 DECLARE_LOG_CATEGORY_EXTERN(LogAssetSearch, Log, All);
 
 class FAssetSearchDatabase
diff --git a/Engine/Plugins/Enterprise/MDLImporter/Source/MDLImporter/Private/mdl/MaterialDistiller.h b/Engine/Plugins/Enterprise/MDLImporter/Source/MDLImporter/Private/mdl/MaterialDistiller.h
--- a/Engine/Plugins/Enterprise/MDLImporter/Source/MDLImporter/Private/mdl/MaterialDistiller.h
+++ b/Engine/Plugins/Enterprise/MDLImporter/Source/MDLImporter/Private/mdl/MaterialDistiller.h
@@ -2,8 +2,8 @@
 
 #pragma once
 
-
 #include "Math/IntPoint.h"
+#include "Templates/Function.h"
 
 #ifdef USE_MDLSDK
 
diff --git a/Engine/Plugins/Experimental/LearningAgents/Source/Learning/Public/LearningArray.h b/Engine/Plugins/Experimental/LearningAgents/Source/Learning/Public/LearningArray.h
--- a/Engine/Plugins/Experimental/LearningAgents/Source/Learning/Public/LearningArray.h
+++ b/Engine/Plugins/Experimental/LearningAgents/Source/Learning/Public/LearningArray.h
@@ -104,6 +104,9 @@
 #endif
 
 
+template <typename FuncType> class TFunctionRef;
+
+
 // Learning Multi-Dimensional Array Types
 
 template<uint32 DimNum>
diff --git a/Engine/Plugins/Experimental/LearningAgents/Source/LearningTraining/Private/LearningExperience.cpp b/Engine/Plugins/Experimental/LearningAgents/Source/LearningTraining/Private/LearningExperience.cpp
--- a/Engine/Plugins/Experimental/LearningAgents/Source/LearningTraining/Private/LearningExperience.cpp
+++ b/Engine/Plugins/Experimental/LearningAgents/Source/LearningTraining/Private/LearningExperience.cpp
@@ -4,6 +4,7 @@
 
 #include "LearningCompletion.h"
 #include "LearningLog.h"
+#include "Templates/Function.h"
 
 namespace UE::Learning
 {
diff --git a/Engine/Plugins/Runtime/GeometryProcessing/Source/DynamicMesh/Public/Selections/QuadGridPatch.h b/Engine/Plugins/Runtime/GeometryProcessing/Source/DynamicMesh/Public/Selections/QuadGridPatch.h
--- a/Engine/Plugins/Runtime/GeometryProcessing/Source/DynamicMesh/Public/Selections/QuadGridPatch.h
+++ b/Engine/Plugins/Runtime/GeometryProcessing/Source/DynamicMesh/Public/Selections/QuadGridPatch.h
@@ -3,7 +3,7 @@
 #pragma once
 
 #include "IndexTypes.h"
-
+#include "Templates/Function.h"
 
 namespace UE
 {
diff --git a/Engine/Plugins/Runtime/NetworkPrediction/Source/NetworkPrediction/Public/NetworkPredictionBuffer.h b/Engine/Plugins/Runtime/NetworkPrediction/Source/NetworkPrediction/Public/NetworkPredictionBuffer.h
--- a/Engine/Plugins/Runtime/NetworkPrediction/Source/NetworkPrediction/Public/NetworkPredictionBuffer.h
+++ b/Engine/Plugins/Runtime/NetworkPrediction/Source/NetworkPrediction/Public/NetworkPredictionBuffer.h
@@ -2,6 +2,7 @@
 
 #pragma once
 #include "Containers/Array.h"
+#include "Templates/Function.h"
 
 // Helper struct to encapsulate optional, delayed writing of new element to TNetworkSimAuxBuffer buffer
 template<typename ElementType>
diff --git a/Engine/Source/Developer/DerivedDataCache/Public/DerivedDataBuildTypes.h b/Engine/Source/Developer/DerivedDataCache/Public/DerivedDataBuildTypes.h
--- a/Engine/Source/Developer/DerivedDataCache/Public/DerivedDataBuildTypes.h
+++ b/Engine/Source/Developer/DerivedDataCache/Public/DerivedDataBuildTypes.h
@@ -10,6 +10,9 @@
 
 #define UE_API DERIVEDDATACACHE_API
 
+template <typename FuncType> class TFunctionRef;
+template <typename FuncType> class TUniqueFunction;
+
 namespace UE::DerivedData { class FBuildOutput; }
 namespace UE::DerivedData { struct FBuildCompleteParams; }
 namespace UE::DerivedData { struct FCacheKey; }
diff --git a/Engine/Source/Developer/TraceServices/Public/Common/PagedArray.h b/Engine/Source/Developer/TraceServices/Public/Common/PagedArray.h
--- a/Engine/Source/Developer/TraceServices/Public/Common/PagedArray.h
+++ b/Engine/Source/Developer/TraceServices/Public/Common/PagedArray.h
@@ -6,6 +6,7 @@
 #include "Containers/Array.h"
 #include "Containers/ArrayView.h"
 #include "Math/Interval.h"
+#include "Templates/Function.h"
 
 #ifndef TRACESERVICES_PAGED_ARRAY_ITERATOR_V2
 #define TRACESERVICES_PAGED_ARRAY_ITERATOR_V2 0 // enables a simpler implementation of TPagedArrayIterator, for debug purposes
diff --git a/Engine/Source/Runtime/Core/Private/Memory/CompositeBuffer.cpp b/Engine/Source/Runtime/Core/Private/Memory/CompositeBuffer.cpp
--- a/Engine/Source/Runtime/Core/Private/Memory/CompositeBuffer.cpp
+++ b/Engine/Source/Runtime/Core/Private/Memory/CompositeBuffer.cpp
@@ -73,6 +73,14 @@ FCompositeBuffer FCompositeBuffer::Mid(uint64 Offset, uint64 Size) const
 	return Buffer;
 }
 
+FMemoryView FCompositeBuffer::ViewOrCopyRange(
+	const uint64 Offset,
+	const uint64 Size,
+	FUniqueBuffer& CopyBuffer) const
+{
+	return ViewOrCopyRange(Offset, Size, CopyBuffer, FUniqueBuffer::Alloc);
+}
+
 FMemoryView FCompositeBuffer::ViewOrCopyRange(
 	const uint64 Offset,
 	const uint64 Size,
diff --git a/Engine/Source/Runtime/Core/Private/Misc/DelayedAutoRegister.cpp b/Engine/Source/Runtime/Core/Private/Misc/DelayedAutoRegister.cpp
--- a/Engine/Source/Runtime/Core/Private/Misc/DelayedAutoRegister.cpp
+++ b/Engine/Source/Runtime/Core/Private/Misc/DelayedAutoRegister.cpp
@@ -8,6 +8,7 @@
 #include "Delegates/DelegateInstancesImpl.h"
 #include "Delegates/IDelegateInstance.h"
 #include "Misc/AssertionMacros.h"
+#include "Templates/Function.h"
 
 DECLARE_MULTICAST_DELEGATE(FDelayedAutoRegisterDelegate);
 
diff --git a/Engine/Source/Runtime/Core/Public/Memory/CompositeBuffer.h b/Engine/Source/Runtime/Core/Public/Memory/CompositeBuffer.h
--- a/Engine/Source/Runtime/Core/Public/Memory/CompositeBuffer.h
+++ b/Engine/Source/Runtime/Core/Public/Memory/CompositeBuffer.h
@@ -4,14 +4,12 @@
 
 #include "Containers/Array.h"
 #include "Containers/ArrayView.h"
-#include "Containers/ContainerAllocationPolicies.h"
 #include "Containers/ContainersFwd.h"
 #include "HAL/Platform.h"
 #include "Math/NumericLimits.h"
 #include "Memory/MemoryFwd.h"
 #include "Memory/MemoryView.h"
 #include "Memory/SharedBuffer.h"
-#include "Templates/Function.h"
 #include "Templates/UnrealTemplate.h"
 
 #include <type_traits>
@@ -91,11 +89,15 @@ public:
 	 * @param CopyBuffer   The buffer to write the copy into if a copy is required.
 	 * @param Allocator    The optional allocator to use when the copy buffer is required.
 	 */
+	[[nodiscard]] CORE_API FMemoryView ViewOrCopyRange(
+		uint64 Offset,
+		uint64 Size,
+		FUniqueBuffer& CopyBuffer) const;
 	[[nodiscard]] CORE_API FMemoryView ViewOrCopyRange(
 		uint64 Offset,
 		uint64 Size,
 		FUniqueBuffer& CopyBuffer,
-		TFunctionRef<FUniqueBuffer (uint64 Size)> Allocator = FUniqueBuffer::Alloc) const;
+		TFunctionRef<FUniqueBuffer (uint64 Size)> Allocator) const;
 
 	/**
 	 * Copies a range of the buffer to a contiguous region of memory.
@@ -142,3 +144,8 @@ private:
 };
 
 inline const FCompositeBuffer FCompositeBuffer::Null;
+
+#if UE_ENABLE_INCLUDE_ORDER_DEPRECATED_IN_5_3
+#include "Containers/ContainerAllocationPolicies.h"
+#include "Templates/Function.h"
+#endif
diff --git a/Engine/Source/Runtime/Core/Public/Misc/DelayedAutoRegister.h b/Engine/Source/Runtime/Core/Public/Misc/DelayedAutoRegister.h
--- a/Engine/Source/Runtime/Core/Public/Misc/DelayedAutoRegister.h
+++ b/Engine/Source/Runtime/Core/Public/Misc/DelayedAutoRegister.h
@@ -3,7 +3,8 @@
 #pragma once
 
 #include "HAL/Platform.h"
-#include "Templates/Function.h"
+
+template <typename FuncType> class TFunction;
 
 enum class EDelayedRegisterRunPhase : uint8
 {
@@ -29,3 +30,7 @@ struct CORE_API FDelayedAutoRegisterHelper
 
 	static void RunAndClearDelayedAutoRegisterDelegates(EDelayedRegisterRunPhase RunPhase);
 };
+
+#if UE_ENABLE_INCLUDE_ORDER_DEPRECATED_IN_5_3
+#include "Templates/Function.h"
+#endif
diff --git a/Engine/Source/Runtime/Core/Public/Serialization/Archive.h b/Engine/Source/Runtime/Core/Public/Serialization/Archive.h
--- a/Engine/Source/Runtime/Core/Public/Serialization/Archive.h
+++ b/Engine/Source/Runtime/Core/Public/Serialization/Archive.h
@@ -15,7 +15,6 @@
 #include "Serialization/ArchiveCookData.h"
 #include "Templates/AndOrNot.h"
 #include "Templates/EnableIf.h"
-#include "Templates/Function.h"
 #include "Templates/IsArrayOrRefOfTypeByPredicate.h"
 #include "Templates/IsEnumClass.h"
 #include "Templates/IsSigned.h"
@@ -36,7 +35,8 @@ class ITargetPlatform;
 class UObject;
 namespace UE::Serialization{ class FEditorBulkData; }
 struct FArchiveSerializedPropertyChain;
-template<class TEnum> class TEnumAsByte;
+template <typename TEnum> class TEnumAsByte;
+template <typename FuncType> class TFunction;
 template <typename T> struct TObjectPtr;
 
 typedef TFunction<bool (double RemainingTime)> FExternalReadCallback;
@@ -2260,3 +2260,7 @@ FORCEINLINE FArchive& operator<<(FArchive& Ar, EnumType& Value)
 {
 	return Ar << (__underlying_type(EnumType)&)Value;
 }
+
+#if UE_ENABLE_INCLUDE_ORDER_DEPRECATED_IN_5_3
+#include "Templates/Function.h"
+#endif
diff --git a/Engine/Source/Runtime/Core/Public/Serialization/MemoryLayout.h b/Engine/Source/Runtime/Core/Public/Serialization/MemoryLayout.h
--- a/Engine/Source/Runtime/Core/Public/Serialization/MemoryLayout.h
+++ b/Engine/Source/Runtime/Core/Public/Serialization/MemoryLayout.h
@@ -6,12 +6,14 @@
 #include "Concepts/StaticStructProvider.h"
 #include "Containers/EnumAsByte.h"
 #include "Containers/StringFwd.h"
+#include "HAL/UnrealMemory.h"
 #include "Misc/DelayedAutoRegister.h"
 #include "Templates/EnableIf.h"
 #include "Templates/IsAbstract.h"
 #include "Templates/IsPolymorphic.h"
 #include "Templates/IsTriviallyDestructible.h"
 #include "Templates/Models.h"
+#include "Templates/UnrealTemplate.h"
 
 class FHashedName;
 class FSHA1;
