- Fixed an outdated comment

diff --git a/Engine/Source/Runtime/Core/Private/Misc/ConfigContext.cpp b/Engine/Source/Runtime/Core/Private/Misc/ConfigContext.cpp
--- a/Engine/Source/Runtime/Core/Private/Misc/ConfigContext.cpp
+++ b/Engine/Source/Runtime/Core/Private/Misc/ConfigContext.cpp
@@ -636,13 +636,10 @@ static bool LoadIniFileHierarchy(const FConfigFileHierarchy& HierarchyToLoad, FC
 }
 
 /**
- * This will load up two .ini files and then determine if the destination one is outdated.
+ * This will load up two .ini files and then determine if the destination one is outdated by comparing
+ * version number in [CurrentIniVersion] section, Version key against the version in the Default*.ini
  * Outdatedness is determined by the following mechanic:
  *
- * When a generated .ini is written out it will store the timestamps of the files it was generated
- * from.  This way whenever the Default__.inis are modified the Generated .ini will view itself as
- * outdated and regenerate it self.
- *
  * Outdatedness also can be affected by commandline params which allow one to delete all .ini, have
  * automated build system etc.
  */
