[Backout] - CL25005038
Original CL Desc
-----------------------------------------------------------------
some delegates code shuffling split from a bigger CL to unclotter it

diff --git a/Engine/Source/Runtime/Core/Public/Delegates/DelegateBase.h b/Engine/Source/Runtime/Core/Public/Delegates/DelegateBase.h
--- a/Engine/Source/Runtime/Core/Public/Delegates/DelegateBase.h
+++ b/Engine/Source/Runtime/Core/Public/Delegates/DelegateBase.h
@@ -41,7 +41,7 @@ ALIAS_TEMPLATE_TYPE_LAYOUT(template<typename ElementType>, FDelegateAllocatorTyp
 /**
  * Base class for unicast delegates.
  */
-class FDelegateBase
+class FDelegateBase //-V1062
 {
 	template <typename>
 	friend class TMulticastDelegateBase;
@@ -59,12 +59,19 @@ protected:
 
 	FDelegateBase(FDelegateBase&& Other)
 	{
-		MoveConstruct(MoveTemp(Other));
+		DelegateAllocator.MoveToEmpty(Other.DelegateAllocator);
+		DelegateSize = Other.DelegateSize;
+		Other.DelegateSize = 0;
 	}
 
 	FDelegateBase& operator=(FDelegateBase&& Other)
 	{
-		MoveAssign(MoveTemp(Other));
+		UE_DELEGATES_MT_SCOPED_WRITE_ACCESS(AccessDetector);
+
+		Unbind();
+		DelegateAllocator.MoveToEmpty(Other.DelegateAllocator);
+		DelegateSize = Other.DelegateSize;
+		Other.DelegateSize = 0;
 		return *this;
 	}
 
@@ -104,25 +111,9 @@ protected:
 		return DelegateAllocator.GetAllocatedSize(DelegateSize, sizeof(FAlignedInlineDelegateType));
 	}
 
-public:
-	/**
-	 * "emplacement" of delegate instance of the given type
-	 */
-	template<typename DelegateInstanceType, typename... DelegateInstanceParams>
-	void CreateDelegateInstance(DelegateInstanceParams&&... Params)
-	{
-		UE_DELEGATES_MT_SCOPED_READ_ACCESS(AccessDetector);
-
-		IDelegateInstance* DelegateInstance = GetDelegateInstanceProtected();
-		if (DelegateInstance)
-		{
-			DelegateInstance->~IDelegateInstance();
-		}
-
-		new(Allocate(sizeof(DelegateInstanceType))) DelegateInstanceType(Forward<DelegateInstanceParams>(Params)...);
-	}
-
 private:
+	friend void* operator new(size_t Size, FDelegateBase& Base);
+
 	void* Allocate(int32 Size)
 	{
 		UE_DELEGATES_MT_SCOPED_WRITE_ACCESS(AccessDetector);
@@ -142,23 +133,6 @@ private:
 		return DelegateAllocator.GetAllocation();
 	}
 
-	void MoveConstruct(FDelegateBase&& Other)
-	{
-		DelegateAllocator.MoveToEmpty(Other.DelegateAllocator);
-		DelegateSize = Other.DelegateSize;
-		Other.DelegateSize = 0;
-	}
-
-	void MoveAssign(FDelegateBase&& Other)
-	{
-		UE_DELEGATES_MT_SCOPED_WRITE_ACCESS(AccessDetector);
-
-		Unbind();
-		DelegateAllocator.MoveToEmpty(Other.DelegateAllocator);
-		DelegateSize = Other.DelegateSize;
-		Other.DelegateSize = 0;
-	}
-
 private:
 	FDelegateAllocatorType::ForElementType<FAlignedInlineDelegateType> DelegateAllocator;
 	int32 DelegateSize = 0;
@@ -166,6 +140,11 @@ private:
 	UE_DELEGATES_MT_ACCESS_DETECTOR(AccessDetector);
 };
 
+inline void* operator new(size_t Size, FDelegateBase& Base)
+{
+	return Base.Allocate((int32)Size);
+}
+
 struct FDefaultDelegateUserPolicy
 {
 	// To extend delegates, you should implement a policy struct like this and pass it as the second template
diff --git a/Engine/Source/Runtime/Core/Public/Delegates/DelegateInstancesImpl.h b/Engine/Source/Runtime/Core/Public/Delegates/DelegateInstancesImpl.h
--- a/Engine/Source/Runtime/Core/Public/Delegates/DelegateInstancesImpl.h
+++ b/Engine/Source/Runtime/Core/Public/Delegates/DelegateInstancesImpl.h
@@ -12,7 +12,6 @@
 
 #include "CoreTypes.h"
 #include "Delegates/DelegateInstanceInterface.h"
-#include "Delegates/IDelegateInstance.h"
 #include "Misc/AssertionMacros.h"
 #include "Templates/RemoveReference.h"
 #include "Templates/SharedPointer.h"
@@ -21,6 +20,8 @@
 #include "UObject/NameTypes.h"
 #include "UObject/WeakObjectPtrTemplates.h"
 
+class FDelegateBase;
+class FDelegateHandle;
 enum class ESPMode : uint8;
 
 namespace UE::Delegates::Private
@@ -133,7 +134,7 @@ public:
 
 	void CreateCopy(DelegateBaseType& Base) const final
 	{
-		Base.template CreateDelegateInstance<TBaseUFunctionDelegateInstance>(*this);
+		new (Base) TBaseUFunctionDelegateInstance(*this);
 	}
 
 	RetValType Execute(ParamTypes... Params) const final
@@ -253,7 +254,7 @@ public:
 
 	void CreateCopy(DelegateBaseType& Base) const final
 	{
-		Base.template CreateDelegateInstance<TBaseSPMethodDelegateInstance>(*this);
+		new (Base) TBaseSPMethodDelegateInstance(*this);
 	}
 
 	RetValType Execute(ParamTypes... Params) const final
@@ -389,7 +390,7 @@ public:
 
 	void CreateCopy(DelegateBaseType& Base) const final
 	{
-		Base.template CreateDelegateInstance<TBaseRawMethodDelegateInstance>(*this);
+		new (Base) TBaseRawMethodDelegateInstance(*this);
 	}
 
 	RetValType Execute(ParamTypes... Params) const final
@@ -513,7 +514,7 @@ public:
 
 	void CreateCopy(DelegateBaseType& Base) const final
 	{
-		Base.template CreateDelegateInstance<TBaseUObjectMethodDelegateInstance>(*this);
+		new (Base) TBaseUObjectMethodDelegateInstance(*this);
 	}
 
 	RetValType Execute(ParamTypes... Params) const final
@@ -636,7 +637,7 @@ public:
 
 	void CreateCopy(DelegateBaseType& Base) const final
 	{
-		Base.template CreateDelegateInstance<TBaseStaticDelegateInstance>(*this);
+		new (Base) TBaseStaticDelegateInstance(*this);
 	}
 
 	RetValType Execute(ParamTypes... Params) const final
@@ -729,7 +730,7 @@ public:
 	// IBaseDelegateInstance interface
 	void CreateCopy(DelegateBaseType& Base) const final
 	{
-		Base.template CreateDelegateInstance<TBaseFunctorDelegateInstance>(*this);
+		new (Base) TBaseFunctorDelegateInstance(*this);
 	}
 
 	RetValType Execute(ParamTypes... Params) const final
@@ -823,7 +824,7 @@ public:
 	// IBaseDelegateInstance interface
 	void CreateCopy(DelegateBaseType& Base) const final
 	{
-		Base.template CreateDelegateInstance<TWeakBaseFunctorDelegateInstance>(*this);
+		new (Base) TWeakBaseFunctorDelegateInstance(*this);
 	}
 
 	RetValType Execute(ParamTypes... Params) const final
diff --git a/Engine/Source/Runtime/Core/Public/Delegates/DelegateSignatureImpl.inl b/Engine/Source/Runtime/Core/Public/Delegates/DelegateSignatureImpl.inl
--- a/Engine/Source/Runtime/Core/Public/Delegates/DelegateSignatureImpl.inl
+++ b/Engine/Source/Runtime/Core/Public/Delegates/DelegateSignatureImpl.inl
@@ -91,7 +91,7 @@ public:
 	UE_NODISCARD inline static TDelegate<RetValType(ParamTypes...), UserPolicy> CreateStatic(typename TIdentity<RetValType (*)(ParamTypes..., std::decay_t<VarTypes>...)>::Type InFunc, VarTypes&&... Vars)
 	{
 		TDelegate<RetValType(ParamTypes...), UserPolicy> Result;
-		Result.template CreateDelegateInstance<TBaseStaticDelegateInstance<FuncType, UserPolicy, std::decay_t<VarTypes>...>>(InFunc, Forward<VarTypes>(Vars)...);
+		new (Result) TBaseStaticDelegateInstance<FuncType, UserPolicy, std::decay_t<VarTypes>...>(InFunc, Forward<VarTypes>(Vars)...);
 		return Result;
 	}
 
@@ -103,7 +103,7 @@ public:
 	UE_NODISCARD inline static TDelegate<RetValType(ParamTypes...), UserPolicy> CreateLambda(FunctorType&& InFunctor, VarTypes&&... Vars)
 	{
 		TDelegate<RetValType(ParamTypes...), UserPolicy> Result;
-		Result.template CreateDelegateInstance<TBaseFunctorDelegateInstance<FuncType, UserPolicy, typename TRemoveReference<FunctorType>::Type, std::decay_t<VarTypes>...>>(Forward<FunctorType>(InFunctor), Forward<VarTypes>(Vars)...);
+		new (Result) TBaseFunctorDelegateInstance<FuncType, UserPolicy, typename TRemoveReference<FunctorType>::Type, std::decay_t<VarTypes>...>(Forward<FunctorType>(InFunctor), Forward<VarTypes>(Vars)...);
 		return Result;
 	}
 
@@ -115,7 +115,7 @@ public:
 	UE_NODISCARD inline static TDelegate<RetValType(ParamTypes...), UserPolicy> CreateWeakLambda(UserClass* InUserObject, FunctorType&& InFunctor, VarTypes&&... Vars)
 	{
 		TDelegate<RetValType(ParamTypes...), UserPolicy> Result;
-		Result.template CreateDelegateInstance<TWeakBaseFunctorDelegateInstance<UserClass, FuncType, UserPolicy, typename TRemoveReference<FunctorType>::Type, std::decay_t<VarTypes>...>>(InUserObject, Forward<FunctorType>(InFunctor), Forward<VarTypes>(Vars)...);
+		new (Result) TWeakBaseFunctorDelegateInstance<UserClass, FuncType, UserPolicy, typename TRemoveReference<FunctorType>::Type, std::decay_t<VarTypes>...>(InUserObject, Forward<FunctorType>(InFunctor), Forward<VarTypes>(Vars)...);
 		return Result;
 	}
 
@@ -131,14 +131,14 @@ public:
 		static_assert(!TIsConst<UserClass>::Value, "Attempting to bind a delegate with a const object pointer and non-const member function.");
 
 		TDelegate<RetValType(ParamTypes...), UserPolicy> Result;
-		Result.template CreateDelegateInstance<TBaseRawMethodDelegateInstance<false, UserClass, FuncType, UserPolicy, std::decay_t<VarTypes>...>>(InUserObject, InFunc, Forward<VarTypes>(Vars)...);
+		new (Result) TBaseRawMethodDelegateInstance<false, UserClass, FuncType, UserPolicy, std::decay_t<VarTypes>...>(InUserObject, InFunc, Forward<VarTypes>(Vars)...);
 		return Result;
 	}
 	template <typename UserClass, typename... VarTypes>
 	UE_NODISCARD inline static TDelegate<RetValType(ParamTypes...), UserPolicy> CreateRaw(const UserClass* InUserObject, typename TMemFunPtrType<true, UserClass, RetValType (ParamTypes..., std::decay_t<VarTypes>...)>::Type InFunc, VarTypes&&... Vars)
 	{
 		TDelegate<RetValType(ParamTypes...), UserPolicy> Result;
-		Result.template CreateDelegateInstance<TBaseRawMethodDelegateInstance<true, const UserClass, FuncType, UserPolicy, std::decay_t<VarTypes>...>>(InUserObject, InFunc, Forward<VarTypes>(Vars)...);
+		new (Result) TBaseRawMethodDelegateInstance<true, const UserClass, FuncType, UserPolicy, std::decay_t<VarTypes>...>(InUserObject, InFunc, Forward<VarTypes>(Vars)...);
 		return Result;
 	}
 
@@ -154,14 +154,14 @@ public:
 		static_assert(!TIsConst<UserClass>::Value, "Attempting to bind a delegate with a const object pointer and non-const member function.");
 
 		TDelegate<RetValType(ParamTypes...), UserPolicy> Result;
-		Result.template CreateDelegateInstance<TBaseSPMethodDelegateInstance<false, UserClass, Mode, FuncType, UserPolicy, std::decay_t<VarTypes>...>>(InUserObjectRef, InFunc, Forward<VarTypes>(Vars)...);
+		new (Result) TBaseSPMethodDelegateInstance<false, UserClass, Mode, FuncType, UserPolicy, std::decay_t<VarTypes>...>(InUserObjectRef, InFunc, Forward<VarTypes>(Vars)...);
 		return Result;
 	}
 	template <typename UserClass, ESPMode Mode, typename... VarTypes>
 	UE_NODISCARD inline static TDelegate<RetValType(ParamTypes...), UserPolicy> CreateSP(const TSharedRef<UserClass, Mode>& InUserObjectRef, typename TMemFunPtrType<true, UserClass, RetValType (ParamTypes..., std::decay_t<VarTypes>...)>::Type InFunc, VarTypes&&... Vars)
 	{
 		TDelegate<RetValType(ParamTypes...), UserPolicy> Result;
-		Result.template CreateDelegateInstance<TBaseSPMethodDelegateInstance<true, const UserClass, Mode, FuncType, UserPolicy, std::decay_t<VarTypes>...>>(InUserObjectRef, InFunc, Forward<VarTypes>(Vars)...);
+		new (Result) TBaseSPMethodDelegateInstance<true, const UserClass, Mode, FuncType, UserPolicy, std::decay_t<VarTypes>...>(InUserObjectRef, InFunc, Forward<VarTypes>(Vars)...);
 		return Result;
 	}
 
@@ -177,14 +177,14 @@ public:
 		static_assert(!TIsConst<UserClass>::Value, "Attempting to bind a delegate with a const object pointer and non-const member function.");
 
 		TDelegate<RetValType(ParamTypes...), UserPolicy> Result;
-		Result.template CreateDelegateInstance<TBaseSPMethodDelegateInstance<false, UserClass, decltype(InUserObject->AsShared())::Mode, FuncType, UserPolicy, std::decay_t<VarTypes>...>>(StaticCastSharedRef<UserClass>(InUserObject->AsShared()), InFunc, Forward<VarTypes>(Vars)...);
+		new (Result) TBaseSPMethodDelegateInstance<false, UserClass, decltype(InUserObject->AsShared())::Mode, FuncType, UserPolicy, std::decay_t<VarTypes>...>(StaticCastSharedRef<UserClass>(InUserObject->AsShared()), InFunc, Forward<VarTypes>(Vars)...);
 		return Result;
 	}
 	template <typename UserClass, typename... VarTypes>
 	UE_NODISCARD inline static TDelegate<RetValType(ParamTypes...), UserPolicy> CreateSP(const UserClass* InUserObject, typename TMemFunPtrType<true, UserClass, RetValType (ParamTypes..., std::decay_t<VarTypes>...)>::Type InFunc, VarTypes&&... Vars)
 	{
 		TDelegate<RetValType(ParamTypes...), UserPolicy> Result;
-		Result.template CreateDelegateInstance<TBaseSPMethodDelegateInstance<true, const UserClass, decltype(InUserObject->AsShared())::Mode, FuncType, UserPolicy, std::decay_t<VarTypes>...>>(StaticCastSharedRef<const UserClass>(InUserObject->AsShared()), InFunc, Forward<VarTypes>(Vars)...);
+		new (Result) TBaseSPMethodDelegateInstance<true, const UserClass, decltype(InUserObject->AsShared())::Mode, FuncType, UserPolicy, std::decay_t<VarTypes>...>(StaticCastSharedRef<const UserClass>(InUserObject->AsShared()), InFunc, Forward<VarTypes>(Vars)...);
 		return Result;
 	}
 
@@ -202,14 +202,14 @@ public:
 		static_assert(!TIsConst<UserClass>::Value, "Attempting to bind a delegate with a const object pointer and non-const member function.");
 
 		TDelegate<RetValType(ParamTypes...), UserPolicy> Result;
-		Result.template CreateDelegateInstance<TBaseSPMethodDelegateInstance<false, UserClass, ESPMode::ThreadSafe, FuncType, UserPolicy, std::decay_t<VarTypes>...>>(InUserObjectRef, InFunc, Forward<VarTypes>(Vars)...);
+		new (Result) TBaseSPMethodDelegateInstance<false, UserClass, ESPMode::ThreadSafe, FuncType, UserPolicy, std::decay_t<VarTypes>...>(InUserObjectRef, InFunc, Forward<VarTypes>(Vars)...);
 		return Result;
 	}
 	template <typename UserClass, typename... VarTypes>
 	UE_NODISCARD inline static TDelegate<RetValType(ParamTypes...), UserPolicy> CreateThreadSafeSP(const TSharedRef<UserClass, ESPMode::ThreadSafe>& InUserObjectRef, typename TMemFunPtrType<true, UserClass, RetValType (ParamTypes..., std::decay_t<VarTypes>...)>::Type InFunc, VarTypes&&... Vars)
 	{
 		TDelegate<RetValType(ParamTypes...), UserPolicy> Result;
-		Result.template CreateDelegateInstance<TBaseSPMethodDelegateInstance<true, const UserClass, ESPMode::ThreadSafe, FuncType, UserPolicy, std::decay_t<VarTypes>...>>(InUserObjectRef, InFunc, Forward<VarTypes>(Vars)...);
+		new (Result) TBaseSPMethodDelegateInstance<true, const UserClass, ESPMode::ThreadSafe, FuncType, UserPolicy, std::decay_t<VarTypes>...>(InUserObjectRef, InFunc, Forward<VarTypes>(Vars)...);
 		return Result;
 	}
 
@@ -227,14 +227,14 @@ public:
 		static_assert(!TIsConst<UserClass>::Value, "Attempting to bind a delegate with a const object pointer and non-const member function.");
 
 		TDelegate<RetValType(ParamTypes...), UserPolicy> Result;
-		Result.template CreateDelegateInstance<TBaseSPMethodDelegateInstance<false, UserClass, ESPMode::ThreadSafe, FuncType, UserPolicy, std::decay_t<VarTypes>...>>(StaticCastSharedRef<UserClass>(InUserObject->AsShared()), InFunc, Forward<VarTypes>(Vars)...);
+		new (Result) TBaseSPMethodDelegateInstance<false, UserClass, ESPMode::ThreadSafe, FuncType, UserPolicy, std::decay_t<VarTypes>...>(StaticCastSharedRef<UserClass>(InUserObject->AsShared()), InFunc, Forward<VarTypes>(Vars)...);
 		return Result;
 	}
 	template <typename UserClass, typename... VarTypes>
 	UE_NODISCARD inline static TDelegate<RetValType(ParamTypes...), UserPolicy> CreateThreadSafeSP(const UserClass* InUserObject, typename TMemFunPtrType<true, UserClass, RetValType (ParamTypes..., std::decay_t<VarTypes>...)>::Type InFunc, VarTypes&&... Vars)
 	{
 		TDelegate<RetValType(ParamTypes...), UserPolicy> Result;
-		Result.template CreateDelegateInstance<TBaseSPMethodDelegateInstance<true, const UserClass, ESPMode::ThreadSafe, FuncType, UserPolicy, std::decay_t<VarTypes>...>>(StaticCastSharedRef<const UserClass>(InUserObject->AsShared()), InFunc, Forward<VarTypes>(Vars)...);
+		new (Result) TBaseSPMethodDelegateInstance<true, const UserClass, ESPMode::ThreadSafe, FuncType, UserPolicy, std::decay_t<VarTypes>...>(StaticCastSharedRef<const UserClass>(InUserObject->AsShared()), InFunc, Forward<VarTypes>(Vars)...);
 		return Result;
 	}
 
@@ -248,14 +248,14 @@ public:
 	UE_NODISCARD inline static TDelegate<RetValType(ParamTypes...), UserPolicy> CreateUFunction(UObjectTemplate* InUserObject, const FName& InFunctionName, VarTypes&&... Vars)
 	{
 		TDelegate<RetValType(ParamTypes...), UserPolicy> Result;
-		Result.template CreateDelegateInstance<TBaseUFunctionDelegateInstance<UObjectTemplate, FuncType, UserPolicy, std::decay_t<VarTypes>...>>(InUserObject, InFunctionName, Forward<VarTypes>(Vars)...);
+		new (Result) TBaseUFunctionDelegateInstance<UObjectTemplate, FuncType, UserPolicy, std::decay_t<VarTypes>...>(InUserObject, InFunctionName, Forward<VarTypes>(Vars)...);
 		return Result;
 	}
 	template <typename UObjectTemplate, typename... VarTypes>
 	UE_NODISCARD inline static TDelegate<RetValType(ParamTypes...), UserPolicy> CreateUFunction(TObjectPtr<UObjectTemplate> InUserObject, const FName& InFunctionName, VarTypes&&... Vars)
 	{
 		TDelegate<RetValType(ParamTypes...), UserPolicy> Result;
-		Result.template CreateDelegateInstance<TBaseUFunctionDelegateInstance<UObjectTemplate, FuncType, UserPolicy, std::decay_t<VarTypes>...>>(ToRawPtr(InUserObject), InFunctionName, Forward<VarTypes>(Vars)...);
+		new (Result) TBaseUFunctionDelegateInstance<UObjectTemplate, FuncType, UserPolicy, std::decay_t<VarTypes>...>(ToRawPtr(InUserObject), InFunctionName, Forward<VarTypes>(Vars)...);
 		return Result;
 	}
 
@@ -271,14 +271,14 @@ public:
 		static_assert(!TIsConst<UserClass>::Value, "Attempting to bind a delegate with a const object pointer and non-const member function.");
 
 		TDelegate<RetValType(ParamTypes...), UserPolicy> Result;
-		Result.template CreateDelegateInstance<TBaseUObjectMethodDelegateInstance<false, UserClass, FuncType, UserPolicy, std::decay_t<VarTypes>...>>(InUserObject, InFunc, Forward<VarTypes>(Vars)...);
+		new (Result) TBaseUObjectMethodDelegateInstance<false, UserClass, FuncType, UserPolicy, std::decay_t<VarTypes>...>(InUserObject, InFunc, Forward<VarTypes>(Vars)...);
 		return Result;
 	}
 	template <typename UserClass, typename... VarTypes>
 	UE_NODISCARD inline static TDelegate<RetValType(ParamTypes...), UserPolicy> CreateUObject(const UserClass* InUserObject, typename TMemFunPtrType<true, UserClass, RetValType (ParamTypes..., std::decay_t<VarTypes>...)>::Type InFunc, VarTypes&&... Vars)
 	{
 		TDelegate<RetValType(ParamTypes...), UserPolicy> Result;
-		Result.template CreateDelegateInstance<TBaseUObjectMethodDelegateInstance<true, const UserClass, FuncType, UserPolicy, std::decay_t<VarTypes>...>>(InUserObject, InFunc, Forward<VarTypes>(Vars)...);
+		new (Result) TBaseUObjectMethodDelegateInstance<true, const UserClass, FuncType, UserPolicy, std::decay_t<VarTypes>...>(InUserObject, InFunc, Forward<VarTypes>(Vars)...);
 		return Result;
 	}
 	template <typename UserClass, typename... VarTypes>
@@ -287,14 +287,14 @@ public:
 		static_assert(!TIsConst<UserClass>::Value, "Attempting to bind a delegate with a const object pointer and non-const member function.");
 
 		TDelegate<RetValType(ParamTypes...), UserPolicy> Result;
-		Result.template CreateDelegateInstance<TBaseUObjectMethodDelegateInstance<false, UserClass, FuncType, UserPolicy, std::decay_t<VarTypes>...>>(ToRawPtr(InUserObject), InFunc, Forward<VarTypes>(Vars)...);
+		new (Result) TBaseUObjectMethodDelegateInstance<false, UserClass, FuncType, UserPolicy, std::decay_t<VarTypes>...>(ToRawPtr(InUserObject), InFunc, Forward<VarTypes>(Vars)...);
 		return Result;
 	}
 	template <typename UserClass, typename... VarTypes>
 	UE_NODISCARD inline static TDelegate<RetValType(ParamTypes...), UserPolicy> CreateUObject(TObjectPtr<UserClass> InUserObject, typename TMemFunPtrType<true, UserClass, RetValType (ParamTypes..., std::decay_t<VarTypes>...)>::Type InFunc, VarTypes&&... Vars)
 	{
 		TDelegate<RetValType(ParamTypes...), UserPolicy> Result;
-		Result.template CreateDelegateInstance<TBaseUObjectMethodDelegateInstance<true, const UserClass, FuncType, UserPolicy, std::decay_t<VarTypes>...>>(ToRawPtr(InUserObject), InFunc, Forward<VarTypes>(Vars)...);
+		new (Result) TBaseUObjectMethodDelegateInstance<true, const UserClass, FuncType, UserPolicy, std::decay_t<VarTypes>...>(ToRawPtr(InUserObject), InFunc, Forward<VarTypes>(Vars)...);
 		return Result;
 	}
 
@@ -347,7 +347,7 @@ public:
 	template <typename... VarTypes>
 	inline void BindStatic(typename TBaseStaticDelegateInstance<FuncType, UserPolicy, std::decay_t<VarTypes>...>::FFuncPtr InFunc, VarTypes&&... Vars)
 	{
-		Super::template CreateDelegateInstance<TBaseStaticDelegateInstance<FuncType, UserPolicy, std::decay_t<VarTypes>...>>(InFunc, Forward<VarTypes>(Vars)...);
+		new (*this) TBaseStaticDelegateInstance<FuncType, UserPolicy, std::decay_t<VarTypes>...>(InFunc, Forward<VarTypes>(Vars)...);
 	}
 	
 	/**
@@ -357,7 +357,7 @@ public:
 	template<typename FunctorType, typename... VarTypes>
 	inline void BindLambda(FunctorType&& InFunctor, VarTypes&&... Vars)
 	{
-		Super::template CreateDelegateInstance<TBaseFunctorDelegateInstance<FuncType, UserPolicy, typename TRemoveReference<FunctorType>::Type, std::decay_t<VarTypes>...>>(Forward<FunctorType>(InFunctor), Forward<VarTypes>(Vars)...);
+		new (*this) TBaseFunctorDelegateInstance<FuncType, UserPolicy, typename TRemoveReference<FunctorType>::Type, std::decay_t<VarTypes>...>(Forward<FunctorType>(InFunctor), Forward<VarTypes>(Vars)...);
 	}
 
 	/**
@@ -367,7 +367,7 @@ public:
 	template<typename UserClass, typename FunctorType, typename... VarTypes>
 	inline void BindWeakLambda(UserClass* InUserObject, FunctorType&& InFunctor, VarTypes&&... Vars)
 	{
-		Super::template CreateDelegateInstance<TWeakBaseFunctorDelegateInstance<UserClass, FuncType, UserPolicy, typename TRemoveReference<FunctorType>::Type, std::decay_t<VarTypes>...>>(InUserObject, Forward<FunctorType>(InFunctor), Forward<VarTypes>(Vars)...);
+		new (*this) TWeakBaseFunctorDelegateInstance<UserClass, FuncType, UserPolicy, typename TRemoveReference<FunctorType>::Type, std::decay_t<VarTypes>...>(InUserObject, Forward<FunctorType>(InFunctor), Forward<VarTypes>(Vars)...);
 	}
 
 	/**
@@ -381,12 +381,12 @@ public:
 	{
 		static_assert(!TIsConst<UserClass>::Value, "Attempting to bind a delegate with a const object pointer and non-const member function.");
 
-		Super::template CreateDelegateInstance<TBaseRawMethodDelegateInstance<false, UserClass, FuncType, UserPolicy, std::decay_t<VarTypes>...>>(InUserObject, InFunc, Forward<VarTypes>(Vars)...);
+		new (*this) TBaseRawMethodDelegateInstance<false, UserClass, FuncType, UserPolicy, std::decay_t<VarTypes>...>(InUserObject, InFunc, Forward<VarTypes>(Vars)...);
 	}
 	template <typename UserClass, typename... VarTypes>
 	inline void BindRaw(const UserClass* InUserObject, typename TMemFunPtrType<true, UserClass, RetValType (ParamTypes..., std::decay_t<VarTypes>...)>::Type InFunc, VarTypes&&... Vars)
 	{
-		Super::template CreateDelegateInstance<TBaseRawMethodDelegateInstance<true, const UserClass, FuncType, UserPolicy, std::decay_t<VarTypes>...>>(InUserObject, InFunc, Forward<VarTypes>(Vars)...);
+		new (*this) TBaseRawMethodDelegateInstance<true, const UserClass, FuncType, UserPolicy, std::decay_t<VarTypes>...>(InUserObject, InFunc, Forward<VarTypes>(Vars)...);
 	}
 
 	/**
@@ -397,12 +397,12 @@ public:
 	{
 		static_assert(!TIsConst<UserClass>::Value, "Attempting to bind a delegate with a const object pointer and non-const member function.");
 
-		Super::template CreateDelegateInstance<TBaseSPMethodDelegateInstance<false, UserClass, Mode, FuncType, UserPolicy, std::decay_t<VarTypes>...>>(InUserObjectRef, InFunc, Forward<VarTypes>(Vars)...);
+		new (*this) TBaseSPMethodDelegateInstance<false, UserClass, Mode, FuncType, UserPolicy, std::decay_t<VarTypes>...>(InUserObjectRef, InFunc, Forward<VarTypes>(Vars)...);
 	}
 	template <typename UserClass, ESPMode Mode, typename... VarTypes>
 	inline void BindSP(const TSharedRef<UserClass, Mode>& InUserObjectRef, typename TMemFunPtrType<true, UserClass, RetValType (ParamTypes..., std::decay_t<VarTypes>...)>::Type InFunc, VarTypes&&... Vars)
 	{
-		Super::template CreateDelegateInstance<TBaseSPMethodDelegateInstance<true, const UserClass, Mode, FuncType, UserPolicy, std::decay_t<VarTypes>...>>(InUserObjectRef, InFunc, Forward<VarTypes>(Vars)...);
+		new (*this) TBaseSPMethodDelegateInstance<true, const UserClass, Mode, FuncType, UserPolicy, std::decay_t<VarTypes>...>(InUserObjectRef, InFunc, Forward<VarTypes>(Vars)...);
 	}
 
 	/**
@@ -416,12 +416,12 @@ public:
 	{
 		static_assert(!TIsConst<UserClass>::Value, "Attempting to bind a delegate with a const object pointer and non-const member function.");
 
-		Super::template CreateDelegateInstance<TBaseSPMethodDelegateInstance<false, UserClass, decltype(InUserObject->AsShared())::Mode, FuncType, UserPolicy, std::decay_t<VarTypes>...>>(StaticCastSharedRef<UserClass>(InUserObject->AsShared()), InFunc, Forward<VarTypes>(Vars)...);
+		new (*this) TBaseSPMethodDelegateInstance<false, UserClass, decltype(InUserObject->AsShared())::Mode, FuncType, UserPolicy, std::decay_t<VarTypes>...>(StaticCastSharedRef<UserClass>(InUserObject->AsShared()), InFunc, Forward<VarTypes>(Vars)...);
 	}
 	template <typename UserClass, typename... VarTypes>
 	inline void BindSP(const UserClass* InUserObject, typename TMemFunPtrType<true, UserClass, RetValType (ParamTypes..., std::decay_t<VarTypes>...)>::Type InFunc, VarTypes&&... Vars)
 	{
-		Super::template CreateDelegateInstance<TBaseSPMethodDelegateInstance<true, const UserClass, decltype(InUserObject->AsShared())::Mode, FuncType, UserPolicy, std::decay_t<VarTypes>...>>(StaticCastSharedRef<const UserClass>(InUserObject->AsShared()), InFunc, Forward<VarTypes>(Vars)...);
+		new (*this) TBaseSPMethodDelegateInstance<true, const UserClass, decltype(InUserObject->AsShared())::Mode, FuncType, UserPolicy, std::decay_t<VarTypes>...>(StaticCastSharedRef<const UserClass>(InUserObject->AsShared()), InFunc, Forward<VarTypes>(Vars)...);
 	}
 
 	/**
@@ -437,12 +437,12 @@ public:
 	{
 		static_assert(!TIsConst<UserClass>::Value, "Attempting to bind a delegate with a const object pointer and non-const member function.");
 
-		Super::template CreateDelegateInstance<TBaseSPMethodDelegateInstance<false, UserClass, ESPMode::ThreadSafe, FuncType, UserPolicy, std::decay_t<VarTypes>...>>(InUserObjectRef, InFunc, Forward<VarTypes>(Vars)...);
+		new (*this) TBaseSPMethodDelegateInstance<false, UserClass, ESPMode::ThreadSafe, FuncType, UserPolicy, std::decay_t<VarTypes>...>(InUserObjectRef, InFunc, Forward<VarTypes>(Vars)...);
 	}
 	template <typename UserClass, typename... VarTypes>
 	inline void BindThreadSafeSP(const TSharedRef<UserClass, ESPMode::ThreadSafe>& InUserObjectRef, typename TMemFunPtrType<true, UserClass, RetValType (ParamTypes..., std::decay_t<VarTypes>...)>::Type InFunc, VarTypes&&... Vars)
 	{
-		Super::template CreateDelegateInstance<TBaseSPMethodDelegateInstance<true, const UserClass, ESPMode::ThreadSafe, FuncType, UserPolicy, std::decay_t<VarTypes>...>>(InUserObjectRef, InFunc, Forward<VarTypes>(Vars)...);
+		new (*this) TBaseSPMethodDelegateInstance<true, const UserClass, ESPMode::ThreadSafe, FuncType, UserPolicy, std::decay_t<VarTypes>...>(InUserObjectRef, InFunc, Forward<VarTypes>(Vars)...);
 	}
 
 	/**
@@ -458,12 +458,12 @@ public:
 	{
 		static_assert(!TIsConst<UserClass>::Value, "Attempting to bind a delegate with a const object pointer and non-const member function.");
 
-		Super::template CreateDelegateInstance<TBaseSPMethodDelegateInstance<false, UserClass, ESPMode::ThreadSafe, FuncType, UserPolicy, std::decay_t<VarTypes>...>>(StaticCastSharedRef<UserClass>(InUserObject->AsShared()), InFunc, Forward<VarTypes>(Vars)...);
+		new (*this) TBaseSPMethodDelegateInstance<false, UserClass, ESPMode::ThreadSafe, FuncType, UserPolicy, std::decay_t<VarTypes>...>(StaticCastSharedRef<UserClass>(InUserObject->AsShared()), InFunc, Forward<VarTypes>(Vars)...);
 	}
 	template <typename UserClass, typename... VarTypes>
 	inline void BindThreadSafeSP(const UserClass* InUserObject, typename TMemFunPtrType<true, UserClass, RetValType (ParamTypes..., std::decay_t<VarTypes>...)>::Type InFunc, VarTypes&&... Vars)
 	{
-		Super::template CreateDelegateInstance<TBaseSPMethodDelegateInstance<true, const UserClass, ESPMode::ThreadSafe, FuncType, UserPolicy, std::decay_t<VarTypes>...>>(StaticCastSharedRef<const UserClass>(InUserObject->AsShared()), InFunc, Forward<VarTypes>(Vars)...);
+		new (*this) TBaseSPMethodDelegateInstance<true, const UserClass, ESPMode::ThreadSafe, FuncType, UserPolicy, std::decay_t<VarTypes>...>(StaticCastSharedRef<const UserClass>(InUserObject->AsShared()), InFunc, Forward<VarTypes>(Vars)...);
 	}
 
 	/**
@@ -475,12 +475,12 @@ public:
 	template <typename UObjectTemplate, typename... VarTypes>
 	inline void BindUFunction(UObjectTemplate* InUserObject, const FName& InFunctionName, VarTypes&&... Vars)
 	{
-		Super::template CreateDelegateInstance<TBaseUFunctionDelegateInstance<UObjectTemplate, FuncType, UserPolicy, std::decay_t<VarTypes>...>>(InUserObject, InFunctionName, Forward<VarTypes>(Vars)...);
+		new (*this) TBaseUFunctionDelegateInstance<UObjectTemplate, FuncType, UserPolicy, std::decay_t<VarTypes>...>(InUserObject, InFunctionName, Forward<VarTypes>(Vars)...);
 	}
 	template <typename UObjectTemplate, typename... VarTypes>
 	inline void BindUFunction(TObjectPtr<UObjectTemplate> InUserObject, const FName& InFunctionName, VarTypes&&... Vars)
 	{
-		Super::template CreateDelegateInstance<TBaseUFunctionDelegateInstance<UObjectTemplate, FuncType, UserPolicy, std::decay_t<VarTypes>...>>(ToRawPtr(InUserObject), InFunctionName, Forward<VarTypes>(Vars)...);
+		new (*this) TBaseUFunctionDelegateInstance<UObjectTemplate, FuncType, UserPolicy, std::decay_t<VarTypes>...>(ToRawPtr(InUserObject), InFunctionName, Forward<VarTypes>(Vars)...);
 	}
 
 	/**
@@ -494,24 +494,24 @@ public:
 	{
 		static_assert(!TIsConst<UserClass>::Value, "Attempting to bind a delegate with a const object pointer and non-const member function.");
 
-		Super::template CreateDelegateInstance<TBaseUObjectMethodDelegateInstance<false, UserClass, FuncType, UserPolicy, std::decay_t<VarTypes>...>>(InUserObject, InFunc, Forward<VarTypes>(Vars)...);
+		new (*this) TBaseUObjectMethodDelegateInstance<false, UserClass, FuncType, UserPolicy, std::decay_t<VarTypes>...>(InUserObject, InFunc, Forward<VarTypes>(Vars)...);
 	}
 	template <typename UserClass, typename... VarTypes>
 	inline void BindUObject(const UserClass* InUserObject, typename TMemFunPtrType<true, UserClass, RetValType (ParamTypes..., std::decay_t<VarTypes>...)>::Type InFunc, VarTypes&&... Vars)
 	{
-		Super::template CreateDelegateInstance<TBaseUObjectMethodDelegateInstance<true, const UserClass, FuncType, UserPolicy, std::decay_t<VarTypes>...>>(InUserObject, InFunc, Forward<VarTypes>(Vars)...);
+		new (*this) TBaseUObjectMethodDelegateInstance<true, const UserClass, FuncType, UserPolicy, std::decay_t<VarTypes>...>(InUserObject, InFunc, Forward<VarTypes>(Vars)...);
 	}
 	template <typename UserClass, typename... VarTypes>
 	inline void BindUObject(TObjectPtr<UserClass> InUserObject, typename TMemFunPtrType<false, UserClass, RetValType (ParamTypes..., std::decay_t<VarTypes>...)>::Type InFunc, VarTypes&&... Vars)
 	{
 		static_assert(!TIsConst<UserClass>::Value, "Attempting to bind a delegate with a const object pointer and non-const member function.");
 
-		Super::template CreateDelegateInstance<TBaseUObjectMethodDelegateInstance<false, UserClass, FuncType, UserPolicy, std::decay_t<VarTypes>...>>(ToRawPtr(InUserObject), InFunc, Forward<VarTypes>(Vars)...);
+		new (*this) TBaseUObjectMethodDelegateInstance<false, UserClass, FuncType, UserPolicy, std::decay_t<VarTypes>...>(ToRawPtr(InUserObject), InFunc, Forward<VarTypes>(Vars)...);
 	}
 	template <typename UserClass, typename... VarTypes>
 	inline void BindUObject(TObjectPtr<UserClass> InUserObject, typename TMemFunPtrType<true, UserClass, RetValType (ParamTypes..., std::decay_t<VarTypes>...)>::Type InFunc, VarTypes&&... Vars)
 	{
-		Super::template CreateDelegateInstance<TBaseUObjectMethodDelegateInstance<true, const UserClass, FuncType, UserPolicy, std::decay_t<VarTypes>...>>(ToRawPtr(InUserObject), InFunc, Forward<VarTypes>(Vars)...);
+		new (*this) TBaseUObjectMethodDelegateInstance<true, const UserClass, FuncType, UserPolicy, std::decay_t<VarTypes>...>(ToRawPtr(InUserObject), InFunc, Forward<VarTypes>(Vars)...);
 	}
 
 public:
