Added CanExecute functor to FActionToken

diff --git a/Engine/Source/Runtime/Core/Private/Logging/TokenizedMessage.cpp b/Engine/Source/Runtime/Core/Private/Logging/TokenizedMessage.cpp
--- a/Engine/Source/Runtime/Core/Private/Logging/TokenizedMessage.cpp
+++ b/Engine/Source/Runtime/Core/Private/Logging/TokenizedMessage.cpp
@@ -230,9 +230,14 @@ FDocumentationToken::FDocumentationToken(const FString& InDocumentationLink, con
 	}
 }
 
-TSharedRef<FActionToken> FActionToken::Create(const FText& ActionName, const FText& ActionDescription, const FOnActionTokenExecuted& Action, bool bInSingleUse)
+TSharedRef<FActionToken> FActionToken::Create(const FText& InActionName, const FText& InActionDescription, const FOnActionTokenExecuted& InAction, bool bInSingleUse)
 {
-	return MakeShared<FActionToken>(FPrivateToken(), ActionName, ActionDescription, Action, bInSingleUse);
+	return MakeShared<FActionToken>(FPrivateToken(), InActionName, InActionDescription, InAction, bInSingleUse);
+}
+
+TSharedRef<FActionToken> FActionToken::Create(const FText& InActionName, const FText& InActionDescription, const FOnActionTokenExecuted& InAction, const FCanExecuteActionToken& InCanExecuteAction, bool bInSingleUse)
+{
+	return MakeShared<FActionToken>(FPrivateToken(), InActionName, InActionDescription, InAction, InCanExecuteAction, bInSingleUse);
 }
 
 TSharedRef<FTutorialToken> FTutorialToken::Create(const FString& TutorialAssetName)
diff --git a/Engine/Source/Runtime/Core/Public/Logging/TokenizedMessage.h b/Engine/Source/Runtime/Core/Public/Logging/TokenizedMessage.h
--- a/Engine/Source/Runtime/Core/Public/Logging/TokenizedMessage.h
+++ b/Engine/Source/Runtime/Core/Public/Logging/TokenizedMessage.h
@@ -523,6 +523,7 @@ private:
 
 
 DECLARE_DELEGATE(FOnActionTokenExecuted);
+DECLARE_DELEGATE_RetVal(bool, FCanExecuteActionToken);
 
 /**
  * Message token that performs an action when activated.
@@ -535,13 +536,14 @@ class FActionToken
 
 public:
 
-	/** Factory method, tokens can only be constructed as shared refs */
-	CORE_API static TSharedRef<FActionToken> Create(const FText& ActionName, const FText& ActionDescription, const FOnActionTokenExecuted& Action, bool bInSingleUse = false);
+	/** Factory methods, tokens can only be constructed as shared refs */
+	CORE_API static TSharedRef<FActionToken> Create(const FText& InActionName, const FText& InActionDescription, const FOnActionTokenExecuted& InAction, bool bInSingleUse = false);
+	CORE_API static TSharedRef<FActionToken> Create(const FText& InActionName, const FText& InActionDescription, const FOnActionTokenExecuted& InAction, const FCanExecuteActionToken& InCanExecuteAction, bool bInSingleUse = false);
 
 	/** Executes the assigned action delegate. */
 	void ExecuteAction()
 	{
-		Action.ExecuteIfBound();
+		ActionDelegate.ExecuteIfBound();
 		bActionExecuted = true;
 	}
 
@@ -554,7 +556,13 @@ public:
 	/** Returns true if the action can be activated */
 	bool CanExecuteAction() const
 	{
-		return Action.IsBound() && (!bSingleUse || !bActionExecuted);
+		return ActionDelegate.IsBound() && (!bSingleUse || !bActionExecuted) && (!CanExecuteActionDelegate.IsBound() || CanExecuteActionDelegate.Execute());
+	}
+
+	/** Returns true if the action is properly set */
+	bool IsValidAction() const
+	{
+		return ActionDelegate.IsBound();
 	}
 
 	// IMessageToken interface
@@ -568,11 +576,26 @@ public:
 	{
 	}
 
+	inline FActionToken(FPrivateToken, const FText& InActionName, const FText& InActionDescription, const FOnActionTokenExecuted& InAction, const FCanExecuteActionToken& InCanExecuteAction, bool bInSingleUse)
+		: FActionToken(InActionName, InActionDescription, InAction, InCanExecuteAction, bInSingleUse)
+	{
+	}
+
 protected:
 
-	/** Hidden constructor. */
+	/** Hidden constructors. */
 	FActionToken(const FText& InActionName, const FText& InActionDescription, const FOnActionTokenExecuted& InAction, bool bInSingleUse)
-		: Action(InAction)
+		: ActionDelegate(InAction)
+		, ActionDescription(InActionDescription)
+		, bSingleUse(bInSingleUse)
+		, bActionExecuted(false)
+	{
+		CachedText = InActionName;
+	}
+
+	FActionToken(const FText& InActionName, const FText& InActionDescription, const FOnActionTokenExecuted& InAction, const FCanExecuteActionToken& InCanExecuteAction, bool bInSingleUse)
+		: ActionDelegate(InAction)
+		, CanExecuteActionDelegate(InCanExecuteAction)
 		, ActionDescription(InActionDescription)
 		, bSingleUse(bInSingleUse)
 		, bActionExecuted(false)
@@ -583,7 +606,10 @@ protected:
 private:
 
 	/** Holds a delegate that is executed when this token is activated. */
-	FOnActionTokenExecuted Action;
+	FOnActionTokenExecuted ActionDelegate;
+
+	/** Holds a delegate that is executed to know whether this token's action can actually execute. */
+	FCanExecuteActionToken CanExecuteActionDelegate;
 
 	/** The action's description text. */
 	const FText ActionDescription;
