Fixed incorrect comment for PerlinNoise2D

diff --git a/Engine/Source/Runtime/Core/Public/Math/UnrealMathUtility.h b/Engine/Source/Runtime/Core/Public/Math/UnrealMathUtility.h
--- a/Engine/Source/Runtime/Core/Public/Math/UnrealMathUtility.h
+++ b/Engine/Source/Runtime/Core/Public/Math/UnrealMathUtility.h
@@ -2433,7 +2433,7 @@ public:
 	UE_NODISCARD static CORE_API float PerlinNoise1D(float Value);
 
 	/**
-	* Generates a 1D Perlin noise sample at the given location.  Returns a continuous random value between -1.0 and 1.0.
+	* Generates a 2D Perlin noise sample at the given location.  Returns a continuous random value between -1.0 and 1.0.
 	*
 	* @param	Location	Where to sample
 	*
