CSVProfiler - duplicate start/stop capture hardening
- Detect duplicate start/stop capture requests. If a Start or Stop request is already pending as the next command in the queue, disregard repeat requests and emit a warning. This avoids an edge case where duplicate EndCapture calls would cause old metadata to be emitted
- Add IsEndCapturePending() method so calling code can check if there's a pending EndCapture before calling EndCapture (instead of just checking IsCapturing, which will be true if an EndCapture is pending).

diff --git a/Engine/Source/Runtime/Core/Private/ProfilingDebugging/CsvProfiler.cpp b/Engine/Source/Runtime/Core/Private/ProfilingDebugging/CsvProfiler.cpp
--- a/Engine/Source/Runtime/Core/Private/ProfilingDebugging/CsvProfiler.cpp
+++ b/Engine/Source/Runtime/Core/Private/ProfilingDebugging/CsvProfiler.cpp
@@ -3205,17 +3205,38 @@ void FCsvProfiler::BeginCapture(int InNumFramesToCapture,
 	ECsvProfilerFlags InFlags)
 {
 	LLM_SCOPE(ELLMTag::CsvProfiler);
-
 	check(IsInGameThread());
+
+	// If there's already a start command in flight for this capture, warn and continue
+	FCsvCaptureCommand CurrentCommand;
+	if (CommandQueue.Peek(CurrentCommand) && CurrentCommand.CommandType == ECsvCommandType::Start)
+	{
+		UE_LOG(LogCsvProfiler, Warning, TEXT("BeginCapture() called, but there is already a pending start command. Ignoring!"));
+		return;
+	}
+
+
 	CommandQueue.Enqueue(FCsvCaptureCommand(ECsvCommandType::Start, GCsvProfilerFrameNumber, InNumFramesToCapture, InDestinationFolder, InFilename, InFlags));
 }
 
 TSharedFuture<FString> FCsvProfiler::EndCapture(FGraphEventRef EventToSignal)
 {
 	LLM_SCOPE(ELLMTag::CsvProfiler);
-
 	check(IsInGameThread());
 
+	// If there's already a stop command in flight for this capture, warn and continue
+	FCsvCaptureCommand CurrentCommand;
+	if (CommandQueue.Peek(CurrentCommand) && CurrentCommand.CommandType == ECsvCommandType::Stop)
+	{
+		UE_LOG(LogCsvProfiler, Warning, TEXT("EndCapture() called, but there is already a pending stop command. Ignoring!"));
+		return CurrentCommand.Future;
+	}
+
+	if (CommandQueue.Peek(CurrentCommand) && CurrentCommand.CommandType == ECsvCommandType::Start)
+	{
+		UE_LOG(LogCsvProfiler, Warning, TEXT("EndCapture() called, but there is already a pending start command!"));
+	}
+
 	// Fire before we copy the metadata so it gives other systems a chance to write any final information.
 	OnCSVProfileEndRequestedDelegate.Broadcast();
 
@@ -3777,30 +3798,39 @@ void FCsvProfiler::Init()
 	}
 }
 
-bool FCsvProfiler::IsCapturing()
+bool FCsvProfiler::IsCapturing() const
 {
 	check(IsInGameThread());
 	return GCsvProfilerIsCapturing;
 }
 
-bool FCsvProfiler::IsWritingFile()
+bool FCsvProfiler::IsWritingFile() const
 {
 	check(IsInGameThread());
 	return GCsvProfilerIsWritingFile;
 }
 
+bool FCsvProfiler::IsEndCapturePending() const
+{
+	check(IsInGameThread());
+	// Return true if the next command is Stop. If the next command is Start then we ignore, since any further Stop command corresponds to a different capture
+	FCsvCaptureCommand CurrentCommand;
+	return CommandQueue.Peek(CurrentCommand) && CurrentCommand.CommandType == ECsvCommandType::Stop;
+}
+
+
 bool FCsvProfiler::IsWaitTrackingEnabledOnCurrentThread()
 {
 	return GCsvTrackWaitsOnAllThreads || GCsvThreadLocalWaitsEnabled;
 }
 
 /*Get the current frame capture count*/
-int32 FCsvProfiler::GetCaptureFrameNumber()
+int32 FCsvProfiler::GetCaptureFrameNumber() const
 {
 	return CaptureFrameNumber;
 }
 
-int32 FCsvProfiler::GetCaptureFrameNumberRT()
+int32 FCsvProfiler::GetCaptureFrameNumberRT() const
 {
 	return CaptureFrameNumberRT;
 }
@@ -3809,7 +3839,7 @@ int32 FCsvProfiler::GetCaptureFrameNumberRT()
 //Get the total frame to capture when we are capturing on event. 
 //Example:  -csvStartOnEvent="My Event"
 //			-csvCaptureOnEventFrameCount=2500
-int32 FCsvProfiler::GetNumFrameToCaptureOnEvent()
+int32 FCsvProfiler::GetNumFrameToCaptureOnEvent() const
 {
 	return CaptureOnEventFrameCount;
 }
@@ -3839,7 +3869,7 @@ bool FCsvProfiler::IsCategoryEnabled(uint32 CategoryIndex) const
 	return GCsvCategoriesEnabled[CategoryIndex];
 }
 
-bool FCsvProfiler::IsCapturing_Renderthread()
+bool FCsvProfiler::IsCapturing_Renderthread() const
 {
 	check(IsInParallelRenderingThread());
 	return GCsvProfilerIsCapturingRT;
@@ -3984,4 +4014,4 @@ void CSVTest()
 
 #endif // CSV_PROFILER_ALLOW_DEBUG_FEATURES
 
-#endif // CSV_PROFILER
\ No newline at end of file
+#endif // CSV_PROFILER
diff --git a/Engine/Source/Runtime/Core/Public/ProfilingDebugging/CsvProfiler.h b/Engine/Source/Runtime/Core/Public/ProfilingDebugging/CsvProfiler.h
--- a/Engine/Source/Runtime/Core/Public/ProfilingDebugging/CsvProfiler.h
+++ b/Engine/Source/Runtime/Core/Public/ProfilingDebugging/CsvProfiler.h
@@ -385,13 +385,14 @@ public:
 	CORE_API static void EndWait();
 
 	/** Singleton interface */
-	CORE_API bool IsCapturing();
-	CORE_API bool IsCapturing_Renderthread();
-	CORE_API bool IsWritingFile();
-
-	CORE_API int32 GetCaptureFrameNumber();
-	CORE_API int32 GetCaptureFrameNumberRT();
-	CORE_API int32 GetNumFrameToCaptureOnEvent();
+	CORE_API bool IsCapturing() const;
+	CORE_API bool IsCapturing_Renderthread() const;
+	CORE_API bool IsWritingFile() const;
+	CORE_API bool IsEndCapturePending() const;
+
+	CORE_API int32 GetCaptureFrameNumber() const;
+	CORE_API int32 GetCaptureFrameNumberRT() const;
+	CORE_API int32 GetNumFrameToCaptureOnEvent() const;
 
 	CORE_API bool EnableCategoryByString(const FString& CategoryName) const;
 	CORE_API void EnableCategoryByIndex(uint32 CategoryIndex, bool bEnable) const;
