Cleanup InitializeKnownConfigFiles API

Legacy function forward decl being removed, plus doc comment update. Probably happened as part of 20980862 rollback.

diff --git a/Engine/Source/Runtime/Core/Public/Misc/ConfigCacheIni.h b/Engine/Source/Runtime/Core/Public/Misc/ConfigCacheIni.h
--- a/Engine/Source/Runtime/Core/Public/Misc/ConfigCacheIni.h
+++ b/Engine/Source/Runtime/Core/Public/Misc/ConfigCacheIni.h
@@ -1272,13 +1272,10 @@ public:
 	/**
 	 * Load the standard (used on all platforms) ini files, like Engine, Input, etc
 	 *
-	 * @param PlatformName Ini name of the platform this is loadenig for (or nullptr for current platform)
-	 * @param bDefaultEngineIniRequired True if the engine ini is required
-	 * @param OverrideProjectDir Used to generate configs for another project (like from a program making binary configs)
+	 * @param FConfigContext The loading context that controls the destination of the loaded ini files
 	 *
 	 * return True if the engine ini was loaded
 	 */
-	bool InitializeKnownConfigFiles(const TCHAR* PlatformName, bool bDefaultEngineIniRequired, const TCHAR* OverrideProjectDir=nullptr);
 	static bool InitializeKnownConfigFiles(FConfigContext& Context);
 
 	/**
