Fixed TSortedMap's key iterators.
Added operator* and operator-> to TMap's key iterators.
Added key iterator tests.

diff --git a/Engine/Source/Runtime/Core/Private/Containers/ContainersTest.cpp b/Engine/Source/Runtime/Core/Private/Containers/ContainersTest.cpp
--- a/Engine/Source/Runtime/Core/Private/Containers/ContainersTest.cpp
+++ b/Engine/Source/Runtime/Core/Private/Containers/ContainersTest.cpp
@@ -401,6 +401,66 @@ namespace
 			check(Cont.IsEmpty());
 			check(Cont.GetAllocatedSize() == InitialAllocatedSize);
 		}
+
+		// Test key iterators
+		{
+			static_assert(std::is_same_v<decltype((Cont.CreateKeyIterator     (DeclVal<KeyType&>())->Key)),         KeyType&>);
+			static_assert(std::is_same_v<decltype((Cont.CreateKeyIterator     (DeclVal<KeyType&>())->Value)),       FContainerTestValueType&>);
+			static_assert(std::is_same_v<decltype((Cont.CreateConstKeyIterator(DeclVal<KeyType&>())->Key)),   const KeyType&>);
+			static_assert(std::is_same_v<decltype((Cont.CreateConstKeyIterator(DeclVal<KeyType&>())->Value)), const FContainerTestValueType&>);
+
+			const TCHAR* RegularValue  = TEXT("Regular");
+			const TCHAR* ReplacedValue = TEXT("Replaced");
+
+			for (int32 Count = 0; Count < MAX_TEST_OBJECTS - 1; Count += MAX_TEST_OBJECTS_STEP)
+			{
+				Cont.Empty();
+				CheckContainer();
+
+				for (int32 N = 0; N != Count; ++N)
+				{
+					Cont.Add(GenerateTestKey<KeyType>(N), FContainerTestValueType(RegularValue));
+					CheckContainer();
+				}
+
+				// Iterate over all possible keys, and some before/after the range [0, Count) that won't exist
+				for (int32 KeyValue = -2; KeyValue < Count + 2; ++KeyValue)
+				{
+					KeyType Key = GenerateTestKey<KeyType>(KeyValue);
+
+					// Check that at most one key is found by the const key iterator
+					const KeyType* FoundConstKey = nullptr;
+					for (auto It = Cont.CreateConstKeyIterator(Key); It; ++It)
+					{
+						check(!FoundConstKey);
+						check(It->Key == Key);
+						FoundConstKey = &It->Key;
+					}
+
+					// Check that at most one key is found by the key iterator, and that we can mutate the value via one
+					KeyType* FoundKey = nullptr;
+					for (auto It = Cont.CreateKeyIterator(Key); It; ++It)
+					{
+						check(!FoundKey);
+						check(It->Key == Key);
+						FoundKey = &It->Key;
+						It->Value.Str = ReplacedValue;
+					}
+
+					// Check that the key iterators found the right element, if any
+					check(FoundKey == FoundConstKey);
+					if (FoundConstKey)
+					{
+						check(KeyValue >= 0 && KeyValue < Count);
+						check(Cont[Key].Str == ReplacedValue);
+					}
+					else
+					{
+						check(KeyValue < 0 || KeyValue >= Count);
+					}
+				}
+			}
+		}
 	}
 
 	// Test container element address consistency when using SortFreeList 
diff --git a/Engine/Source/Runtime/Core/Public/Containers/Map.h b/Engine/Source/Runtime/Core/Public/Containers/Map.h
--- a/Engine/Source/Runtime/Core/Public/Containers/Map.h
+++ b/Engine/Source/Runtime/Core/Public/Containers/Map.h
@@ -850,6 +850,9 @@ protected:
 		FORCEINLINE ItKeyType&   Key() const { return SetIt->Key; }
 		FORCEINLINE ItValueType& Value() const { return SetIt->Value; }
 
+		FORCEINLINE decltype(auto) operator* () const { return SetIt.operator* (); }
+		FORCEINLINE decltype(auto) operator->() const { return SetIt.operator->(); }
+
 	protected:
 		SetItType SetIt;
 	};
diff --git a/Engine/Source/Runtime/Core/Public/Containers/SortedMap.h b/Engine/Source/Runtime/Core/Public/Containers/SortedMap.h
--- a/Engine/Source/Runtime/Core/Public/Containers/SortedMap.h
+++ b/Engine/Source/Runtime/Core/Public/Containers/SortedMap.h
@@ -504,7 +504,7 @@ private:
 	}
 
 	/** Find index of key */
-	FORCEINLINE int32 FindIndex(KeyConstPointerType Key)
+	FORCEINLINE int32 FindIndex(KeyConstPointerType Key) const
 	{
 		return Algo::BinarySearchBy(Pairs, Key, FKeyForward(), SortPredicate());
 	}
@@ -714,25 +714,27 @@ public:
 	/** Iterates over values associated with a specified key in a const map. This will be at most one value because keys must be unique */
 	class TConstKeyIterator : public TBaseIterator<true>
 	{
+		using Super = TBaseIterator<true>;
+
 	public:
 		FORCEINLINE TConstKeyIterator(const TSortedMap& InMap, KeyInitType InKey)
-			: TBaseIterator<true>(InMap.Pairs.CreateIterator())
+			: Super(InMap.Pairs.CreateConstIterator())
 		{
-			int32 NewIndex = FindIndex(InKey);
+			int32 NewIndex = InMap.FindIndex(InKey);
 		
 			if (NewIndex != INDEX_NONE)
 			{
-				TBaseIterator<true>::PairIt += NewIndex;
+				Super::PairIt += NewIndex;
 			}
 			else
 			{
-				TBaseIterator<true>::PairIt.SetToEnd();
+				Super::PairIt.SetToEnd();
 			}
 		}
 
 		FORCEINLINE TConstKeyIterator& operator++()
 		{
-			TBaseIterator<true>::PairIt.SetToEnd();
+			Super::PairIt.SetToEnd();
 			return *this;
 		}
 	};
@@ -740,33 +742,35 @@ public:
 	/** Iterates over values associated with a specified key in a map. This will be at most one value because keys must be unique */
 	class TKeyIterator : public TBaseIterator<false>
 	{
+		using Super = TBaseIterator<false>;
+
 	public:
 		FORCEINLINE TKeyIterator(TSortedMap& InMap, KeyInitType InKey)
-			: TBaseIterator<false>(InMap.Pairs.CreateConstIterator())
+			: Super(InMap.Pairs.CreateIterator())
 		{
-			int32 NewIndex = FindIndex(InKey);
+			int32 NewIndex = InMap.FindIndex(InKey);
 
 			if (NewIndex != INDEX_NONE)
 			{
-				TBaseIterator<true>::PairIt += NewIndex;
+				Super::PairIt += NewIndex;
 			}
 			else
 			{
-				TBaseIterator<true>::PairIt.SetToEnd();
+				Super::PairIt.SetToEnd();
 			}
 		}
 
 		FORCEINLINE TKeyIterator& operator++()
 		{
-			TBaseIterator<false>::PairIt.SetToEnd();
+			Super::PairIt.SetToEnd();
 			return *this;
 		}
 
 		/** Removes the current key-value pair from the map. */
 		FORCEINLINE void RemoveCurrent()
 		{
-			TBaseIterator<false>::PairIt.RemoveCurrent();
-			TBaseIterator<false>::PairIt.SetToEnd();
+			Super::PairIt.RemoveCurrent();
+			Super::PairIt.SetToEnd();
 		}
 	};
 
