Added better comments to TBox and TIntRect to indicate how their Intersect() and Contains()/IsInsideOf() functions handle overlaps on their borders.
Fixed FVector parameter in TBox::IsInsideOrOnXY().

diff --git a/Engine/Source/Runtime/Core/Public/Math/Box.h b/Engine/Source/Runtime/Core/Public/Math/Box.h
--- a/Engine/Source/Runtime/Core/Public/Math/Box.h
+++ b/Engine/Source/Runtime/Core/Public/Math/Box.h
@@ -368,6 +368,9 @@ public:
 	 *
 	 * @param Other The bounding box to intersect with.
 	 * @return true if the boxes intersect, false otherwise.
+	 *
+	 * @note  This function assumes boxes have closed bounds, i.e. boxes with
+	 *        coincident borders on any edge will overlap.
 	 */
 	FORCEINLINE bool Intersect( const TBox<T>& Other ) const;
 
@@ -376,6 +379,9 @@ public:
 	 *
 	 * @param Other The bounding box to test intersection.
 	 * @return true if the boxes intersect in the XY Plane, false otherwise.
+	 *
+	 * @note  This function assumes boxes have closed bounds, i.e. boxes with
+	 *        coincident borders on any edge will overlap.
 	 */
 	FORCEINLINE bool IntersectXY( const TBox<T>& Other ) const;
 
@@ -401,6 +407,9 @@ public:
 	 * @param In The location to test for inside the bounding volume.
 	 * @return true if location is inside this volume.
 	 * @see IsInsideXY
+	 *
+	 * @note  This function assumes boxes have open bounds, i.e. points lying on the border of the box are not inside.
+	 *        Use IsInsideOrOn to include borders in the test.
 	 */
 	FORCEINLINE bool IsInside( const TVector<T>& In ) const
 	{
@@ -410,9 +419,12 @@ public:
 	/** 
 	 * Checks whether the given location is inside or on this box.
 	 * 
-	 * @param In The location to test for inside the bounding volume.
-	 * @return true if location is inside this volume.
-	 * @see IsInsideXY
+	 * @param In The location to test for inside or on the bounding volume.
+	 * @return true if location is inside or on this volume.
+	 * @see IsInsideOrOnXY
+	 *
+	 * @note  This function assumes boxes have closed bounds, i.e. points lying on the border of the box are inside.
+	 *        Use IsInside to exclude borders from the test.
 	 */
 	FORCEINLINE bool IsInsideOrOn( const TVector<T>& In ) const
 	{
@@ -424,6 +436,10 @@ public:
 	 * 
 	 * @param Other The box to test for encapsulation within the bounding volume.
 	 * @return true if box is inside this volume.
+	 * @see IsInsideXY
+	 *
+	 * @note  This function assumes boxes have open bounds, i.e. boxes with
+	 *        coincident borders on any edge are not encapsulated.
 	 */
 	FORCEINLINE bool IsInside( const TBox<T>& Other ) const
 	{
@@ -436,6 +452,9 @@ public:
 	 * @param In The location to test for inside the bounding box.
 	 * @return true if location is inside this box in the XY plane.
 	 * @see IsInside
+	 *
+	 * @note  This function assumes boxes have open bounds, i.e. points lying on the border of the box are not inside.
+	 *        Use IsInsideOrOnXY to include borders in the test.
 	 */
 	FORCEINLINE bool IsInsideXY( const TVector<T>& In ) const
 	{
@@ -445,11 +464,14 @@ public:
 	/**
 	 * Checks whether the given location is inside or on this box in the XY plane.
 	 *
-	 * @param In The location to test for inside the bounding volume.
-	 * @return true if location is inside this box in the XY plane.
+	 * @param In The location to test for inside or on the bounding volume.
+	 * @return true if location is inside or on this box in the XY plane.
 	 * @see IsInsideOrOn
+	 *
+	 * @note  This function assumes boxes have closed bounds, i.e. points lying on the border of the box are not inside.
+	 *        Use IsInsideXY to exclude borders from the test.
 	 */
-	FORCEINLINE bool IsInsideOrOnXY(const FVector& In) const
+	FORCEINLINE bool IsInsideOrOnXY(const TVector<T>& In) const
 	{
 		return ((In.X >= Min.X) && (In.X <= Max.X) && (In.Y >= Min.Y) && (In.Y <= Max.Y));
 	}
@@ -459,6 +481,10 @@ public:
 	 * 
 	 * @param Other The box to test for encapsulation within the bounding box.
 	 * @return true if box is inside this box in the XY plane.
+	 * @see IsInside
+	 *
+	 * @note  This function assumes boxes have open bounds, i.e. boxes with
+	 *        coincident borders on any edge are not encapsulated.
 	 */
 	FORCEINLINE bool IsInsideXY( const TBox<T>& Other ) const
 	{
diff --git a/Engine/Source/Runtime/Core/Public/Math/IntRect.h b/Engine/Source/Runtime/Core/Public/Math/IntRect.h
--- a/Engine/Source/Runtime/Core/Public/Math/IntRect.h
+++ b/Engine/Source/Runtime/Core/Public/Math/IntRect.h
@@ -307,7 +307,15 @@ PRAGMA_ENABLE_DEPRECATION_WARNINGS
 		Max.Y = FMath::Max<IntType>(Max.Y, R.Max.Y);
 	}
 
-	/** Returns true if the two rects have any overlap. */
+	/**
+	 * Returns true if the two rects have any overlap.
+	 *
+	 * @param Other The rect to compare with.
+	 * @return true if the rectangle overlaps the other rectangle, false otherwise.
+	 *
+	 * @note  This function assumes rects have open bounds, i.e. rects with
+	 *        coincident borders on any edge will not overlap.
+	 */
 	bool Intersect(const TIntRect& Other) const
 	{
 		return Other.Min.X < Max.X&& Other.Max.X > Min.X && Other.Min.Y < Max.Y&& Other.Max.Y > Min.Y;
@@ -317,7 +325,10 @@ PRAGMA_ENABLE_DEPRECATION_WARNINGS
 	 * Test whether this rectangle contains a point.
 	 *
 	 * @param Point The point to test against.
-	 * @return true if the rectangle contains the specified point,, false otherwise..
+	 * @return true if the rectangle contains the specified point, false otherwise.
+	 *
+	 * @note  This function assumes rects have half-open bounds, i.e. points are contained
+	 *        by the minimum border of the box, but not the maximum border.
 	 */
 	bool Contains(IntPointType P) const
 	{
