Add project field seperate to app name

Previously we used project name rather than app name if it was available. However app name is needed to look up target binary is certain cases.

diff --git a/Engine/Source/Developer/TraceServices/Private/Analyzers/DiagnosticsAnalysis.cpp b/Engine/Source/Developer/TraceServices/Private/Analyzers/DiagnosticsAnalysis.cpp
--- a/Engine/Source/Developer/TraceServices/Private/Analyzers/DiagnosticsAnalysis.cpp
+++ b/Engine/Source/Developer/TraceServices/Private/Analyzers/DiagnosticsAnalysis.cpp
@@ -78,6 +78,7 @@ bool FDiagnosticsAnalyzer::OnEvent(uint16 RouteId, EStyle Style, const FOnEventC
 
 		EventData.GetString("Platform", SessionInfo.Platform);
 		EventData.GetString("AppName", SessionInfo.AppName);
+		EventData.GetString("ProjectName", SessionInfo.ProjectName);	
 		EventData.GetString("CommandLine", SessionInfo.CommandLine);
 		EventData.GetString("Branch", SessionInfo.Branch);
 		EventData.GetString("BuildVersion", SessionInfo.BuildVersion);
diff --git a/Engine/Source/Developer/TraceServices/Public/TraceServices/Model/Diagnostics.h b/Engine/Source/Developer/TraceServices/Public/TraceServices/Model/Diagnostics.h
--- a/Engine/Source/Developer/TraceServices/Public/TraceServices/Model/Diagnostics.h
+++ b/Engine/Source/Developer/TraceServices/Public/TraceServices/Model/Diagnostics.h
@@ -15,6 +15,7 @@ struct FSessionInfo
 {
 	FString Platform;
 	FString AppName;
+	FString ProjectName;
 	FString CommandLine;
 	FString Branch;
 	FString BuildVersion;
diff --git a/Engine/Source/Runtime/Core/Private/ProfilingDebugging/TraceAuxiliary.cpp b/Engine/Source/Runtime/Core/Private/ProfilingDebugging/TraceAuxiliary.cpp
--- a/Engine/Source/Runtime/Core/Private/ProfilingDebugging/TraceAuxiliary.cpp
+++ b/Engine/Source/Runtime/Core/Private/ProfilingDebugging/TraceAuxiliary.cpp
@@ -1232,6 +1232,7 @@ static void LaunchUnrealTraceInternal(const TCHAR* CommandLine)
 UE_TRACE_EVENT_BEGIN(Diagnostics, Session2, NoSync|Important)
 	UE_TRACE_EVENT_FIELD(UE::Trace::AnsiString, Platform)
 	UE_TRACE_EVENT_FIELD(UE::Trace::AnsiString, AppName)
+	UE_TRACE_EVENT_FIELD(UE::Trace::WideString, ProjectName)
 	UE_TRACE_EVENT_FIELD(UE::Trace::WideString, CommandLine)
 	UE_TRACE_EVENT_FIELD(UE::Trace::WideString, Branch)
 	UE_TRACE_EVENT_FIELD(UE::Trace::WideString, BuildVersion)
@@ -1500,11 +1501,11 @@ void FTraceAuxiliary::Initialize(const TCHAR* CommandLine)
 	}
 
 	const TCHAR* AppName = TEXT(UE_APP_NAME);
+	const TCHAR* ProjectName = TEXT("");
 #if IS_MONOLITHIC && !IS_PROGRAM
-	extern TCHAR GInternalProjectName[];
-	if (GInternalProjectName[0] != '\0')
+	if (FApp::HasProjectName())
 	{
-		AppName = GInternalProjectName;
+		ProjectName = FApp::GetProjectName();
 	}
 #endif
 
@@ -1514,18 +1515,21 @@ void FTraceAuxiliary::Initialize(const TCHAR* CommandLine)
 	const TCHAR* BuildVersion = BuildSettings::GetBuildVersion();
 	constexpr uint32 PlatformLen = UE_ARRAY_COUNT(PREPROCESSOR_TO_STRING(UBT_COMPILED_PLATFORM)) - 1;
 	const uint32 AppNameLen = FCString::Strlen(AppName);
+	const uint32 ProjectNameLen = FCString::Strlen(ProjectName);
 	const uint32 CommandLineLen = FCString::Strlen(CommandLine);
 	const uint32 BranchNameLen = FCString::Strlen(BranchName);
 	const uint32 BuildVersionLen = FCString::Strlen(BuildVersion);
-	uint32 DataSize =
+	const uint32 DataSize =
 		(PlatformLen * sizeof(ANSICHAR)) +
 		(AppNameLen * sizeof(ANSICHAR)) +
+		(ProjectNameLen * sizeof(TCHAR)) +
 		(CommandLineLen * sizeof(TCHAR)) +
 		(BranchNameLen * sizeof(TCHAR)) +
 		(BuildVersionLen * sizeof(TCHAR));
 	UE_TRACE_LOG(Diagnostics, Session2, UE::Trace::TraceLogChannel, DataSize)
 		<< Session2.Platform(PREPROCESSOR_TO_STRING(UBT_COMPILED_PLATFORM), PlatformLen)
 		<< Session2.AppName(AppName, AppNameLen)
+		<< Session2.ProjectName(ProjectName, ProjectNameLen)
 		<< Session2.CommandLine(CommandLine, CommandLineLen)
 		<< Session2.Branch(BranchName, BranchNameLen)
 		<< Session2.BuildVersion(BuildVersion, BuildVersionLen)
