fixed a race in DepletableMpmcQueue

diff --git a/Engine/Source/Runtime/Core/Public/Containers/DepletableMpmcQueue.h b/Engine/Source/Runtime/Core/Public/Containers/DepletableMpmcQueue.h
--- a/Engine/Source/Runtime/Core/Public/Containers/DepletableMpmcQueue.h
+++ b/Engine/Source/Runtime/Core/Public/Containers/DepletableMpmcQueue.h
@@ -70,7 +70,10 @@ namespace UE
 
 			FNode* Prev = Tail.exchange(New, std::memory_order_acq_rel); // "acquire" to sync with `Prev->Next` initialisation from a concurrent calls,
 			// `release` to make sure the new node is fully constructed before it becomes visible to the consumer or a concurrent enqueueing
-			check(Prev->Next.load(std::memory_order_relaxed) == nullptr); // `Tail` is assigned before its Next
+			
+			// the following `check` is commented out because it can be reordered after the following `Prev->Next.store` which unlocks 
+			// the consumer that will free `Prev`. left commented as a documentation
+			// check(Prev->Next.load(std::memory_order_relaxed) == nullptr); // `Tail` is assigned before its Next
 
 			Prev->Next.store(New, std::memory_order_relaxed);
 
