Made FGenericPlatformTLS::InvalidTlsSlot constant so it can be shared between platforms.

#tests Did a Horde Incremental CI.

diff --git a/Engine/Source/Runtime/Core/Public/Android/AndroidPlatformTLS.h b/Engine/Source/Runtime/Core/Public/Android/AndroidPlatformTLS.h
--- a/Engine/Source/Runtime/Core/Public/Android/AndroidPlatformTLS.h
+++ b/Engine/Source/Runtime/Core/Public/Android/AndroidPlatformTLS.h
@@ -33,7 +33,7 @@ struct CORE_API FAndroidTLS : public FGenericPlatformTLS
 		pthread_key_t Key = 0;
 		if (pthread_key_create(&Key, NULL) != 0)
 		{
-			Key = 0xFFFFFFFF;  // matches the Windows TlsAlloc() retval //@todo android: should probably check for this below, or assert out instead
+			Key = InvalidTlsSlot;  // matches the Windows TlsAlloc() retval //@todo android: should probably check for this below, or assert out instead
 		}
 		return Key;
 	}
diff --git a/Engine/Source/Runtime/Core/Public/Apple/ApplePlatformTLS.h b/Engine/Source/Runtime/Core/Public/Apple/ApplePlatformTLS.h
--- a/Engine/Source/Runtime/Core/Public/Apple/ApplePlatformTLS.h
+++ b/Engine/Source/Runtime/Core/Public/Apple/ApplePlatformTLS.h
@@ -35,7 +35,7 @@ struct CORE_API FApplePlatformTLS : public FGenericPlatformTLS
 		pthread_key_t SlotKey = 0;
 		if (pthread_key_create(&SlotKey, NULL) != 0)
 		{
-			SlotKey = 0xFFFFFFFF;  // matches the Windows TlsAlloc() retval.
+			SlotKey = InvalidTlsSlot;  // matches the Windows TlsAlloc() retval.
 		}
 		return SlotKey;
 	}
diff --git a/Engine/Source/Runtime/Core/Public/GenericPlatform/GenericPlatformTLS.h b/Engine/Source/Runtime/Core/Public/GenericPlatform/GenericPlatformTLS.h
--- a/Engine/Source/Runtime/Core/Public/GenericPlatform/GenericPlatformTLS.h
+++ b/Engine/Source/Runtime/Core/Public/GenericPlatform/GenericPlatformTLS.h
@@ -10,6 +10,8 @@
  */
 struct FGenericPlatformTLS
 {
+	static const uint32 InvalidTlsSlot = 0xFFFFFFFF;
+
 	/**
 	 * Return false if this is an invalid TLS slot
 	 * @param SlotIndex the TLS index to check
@@ -17,7 +19,7 @@ struct FGenericPlatformTLS
 	 */
 	static FORCEINLINE bool IsValidTlsSlot(uint32 SlotIndex)
 	{
-		return SlotIndex != 0xFFFFFFFF;
+		return SlotIndex != InvalidTlsSlot;
 	}
 
 #if 0 // provided for reference
