Assertion failure in FVirtualStackAllocator::Allocate

Incorrect check for out of memory.

#rnx

diff --git a/Engine/Source/Runtime/Core/Private/Memory/VirtualStackAllocator.cpp b/Engine/Source/Runtime/Core/Private/Memory/VirtualStackAllocator.cpp
--- a/Engine/Source/Runtime/Core/Private/Memory/VirtualStackAllocator.cpp
+++ b/Engine/Source/Runtime/Core/Private/Memory/VirtualStackAllocator.cpp
@@ -58,14 +58,13 @@ void* FVirtualStackAllocator::Allocate(size_t Size, size_t Alignment)
 	if (Size > 0)
 	{
 		void* const AllocationEnd = OffsetPointer(AllocationStart, Size);
+		void* const UsableMemoryEnd = OffsetPointer(VirtualMemory.GetVirtualPointer(), TotalReservationSize - PageSize);
 
-		if (Size > (TotalReservationSize - PageSize))
+		if (AllocationEnd > UsableMemoryEnd)
 		{
 			FPlatformMemory::OnOutOfMemory(Size, Alignment);
 		}
 
-		void* const UsableMemoryEnd = OffsetPointer(VirtualMemory.GetVirtualPointer(), TotalReservationSize - PageSize);
-
 		// After the high water mark is established, needing to commit pages should be rare
 		if (UNLIKELY(AllocationEnd > NextUncommittedPage))
 		{
