Added OnMallocInitialized implementation to FMallocLeakDetectionProxy

#rnx

diff --git a/Engine/Source/Runtime/Core/Private/HAL/MallocLeakDetectionProxy.h b/Engine/Source/Runtime/Core/Private/HAL/MallocLeakDetectionProxy.h
--- a/Engine/Source/Runtime/Core/Private/HAL/MallocLeakDetectionProxy.h
+++ b/Engine/Source/Runtime/Core/Private/HAL/MallocLeakDetectionProxy.h
@@ -121,6 +121,11 @@ public:
 	{
 		return UsedMalloc->GetDescriptiveName();
 	}
+
+	virtual void OnMallocInitialized() override
+	{
+		UsedMalloc->OnMallocInitialized();
+	}
 };
 
 #endif // MALLOC_LEAKDETECTIONMALLOC_LEAKDETECTION
