Avoided indirection in multicast delegate adds.

diff --git a/Engine/Source/Runtime/Core/Public/Delegates/DelegateSignatureImpl.inl b/Engine/Source/Runtime/Core/Public/Delegates/DelegateSignatureImpl.inl
--- a/Engine/Source/Runtime/Core/Public/Delegates/DelegateSignatureImpl.inl
+++ b/Engine/Source/Runtime/Core/Public/Delegates/DelegateSignatureImpl.inl
@@ -663,7 +663,7 @@ public:
 	template <typename... VarTypes>
 	inline FDelegateHandle AddStatic(typename TBaseStaticDelegateInstance<void (ParamTypes...), UserPolicy, std::decay_t<VarTypes>...>::FFuncPtr InFunc, VarTypes&&... Vars)
 	{
-		return Add(FDelegate::CreateStatic(InFunc, Forward<VarTypes>(Vars)...));
+		return Super::AddDelegateInstance(FDelegate::CreateStatic(InFunc, Forward<VarTypes>(Vars)...));
 	}
 
 	/**
@@ -675,7 +675,7 @@ public:
 	template<typename FunctorType, typename... VarTypes>
 	inline FDelegateHandle AddLambda(FunctorType&& InFunctor, VarTypes&&... Vars)
 	{
-		return Add(FDelegate::CreateLambda(Forward<FunctorType>(InFunctor), Forward<VarTypes>(Vars)...));
+		return Super::AddDelegateInstance(FDelegate::CreateLambda(Forward<FunctorType>(InFunctor), Forward<VarTypes>(Vars)...));
 	}
 
 	/**
@@ -688,7 +688,7 @@ public:
 	template<typename UserClass, typename FunctorType, typename... VarTypes>
 	inline FDelegateHandle AddWeakLambda(UserClass* InUserObject, FunctorType&& InFunctor, VarTypes&&... Vars)
 	{
-		return Add(FDelegate::CreateWeakLambda(InUserObject, Forward<FunctorType>(InFunctor), Forward<VarTypes>(Vars)...));
+		return Super::AddDelegateInstance(FDelegate::CreateWeakLambda(InUserObject, Forward<FunctorType>(InFunctor), Forward<VarTypes>(Vars)...));
 	}
 
 	/**
@@ -705,12 +705,12 @@ public:
 	{
 		static_assert(!TIsConst<UserClass>::Value, "Attempting to bind a delegate with a const object pointer and non-const member function.");
 
-		return Add(FDelegate::CreateRaw(InUserObject, InFunc, Forward<VarTypes>(Vars)...));
+		return Super::AddDelegateInstance(FDelegate::CreateRaw(InUserObject, InFunc, Forward<VarTypes>(Vars)...));
 	}
 	template <typename UserClass, typename... VarTypes>
 	inline FDelegateHandle AddRaw(const UserClass* InUserObject, typename TMemFunPtrType<true, UserClass, void (ParamTypes..., std::decay_t<VarTypes>...)>::Type InFunc, VarTypes&&... Vars)
 	{
-		return Add(FDelegate::CreateRaw(InUserObject, InFunc, Forward<VarTypes>(Vars)...));
+		return Super::AddDelegateInstance(FDelegate::CreateRaw(InUserObject, InFunc, Forward<VarTypes>(Vars)...));
 	}
 
 	/**
@@ -726,12 +726,12 @@ public:
 	{
 		static_assert(!TIsConst<UserClass>::Value, "Attempting to bind a delegate with a const object pointer and non-const member function.");
 
-		return Add(FDelegate::CreateSP(InUserObjectRef, InFunc, Forward<VarTypes>(Vars)...));
+		return Super::AddDelegateInstance(FDelegate::CreateSP(InUserObjectRef, InFunc, Forward<VarTypes>(Vars)...));
 	}
 	template <typename UserClass, ESPMode Mode, typename... VarTypes>
 	inline FDelegateHandle AddSP(const TSharedRef<UserClass, Mode>& InUserObjectRef, typename TMemFunPtrType<true, UserClass, void (ParamTypes..., std::decay_t<VarTypes>...)>::Type InFunc, VarTypes&&... Vars)
 	{
-		return Add(FDelegate::CreateSP(InUserObjectRef, InFunc, Forward<VarTypes>(Vars)...));
+		return Super::AddDelegateInstance(FDelegate::CreateSP(InUserObjectRef, InFunc, Forward<VarTypes>(Vars)...));
 	}
 
 	/**
@@ -747,12 +747,12 @@ public:
 	{
 		static_assert(!TIsConst<UserClass>::Value, "Attempting to bind a delegate with a const object pointer and non-const member function.");
 
-		return Add(FDelegate::CreateSP(InUserObject, InFunc, Forward<VarTypes>(Vars)...));
+		return Super::AddDelegateInstance(FDelegate::CreateSP(InUserObject, InFunc, Forward<VarTypes>(Vars)...));
 	}
 	template <typename UserClass, typename... VarTypes>
 	inline FDelegateHandle AddSP(const UserClass* InUserObject, typename TMemFunPtrType<true, UserClass, void (ParamTypes..., std::decay_t<VarTypes>...)>::Type InFunc, VarTypes&&... Vars)
 	{
-		return Add(FDelegate::CreateSP(InUserObject, InFunc, Forward<VarTypes>(Vars)...));
+		return Super::AddDelegateInstance(FDelegate::CreateSP(InUserObject, InFunc, Forward<VarTypes>(Vars)...));
 	}
 
 	/**
@@ -768,12 +768,12 @@ public:
 	{
 		static_assert(!TIsConst<UserClass>::Value, "Attempting to bind a delegate with a const object pointer and non-const member function.");
 
-		return Add(FDelegate::CreateThreadSafeSP(InUserObjectRef, InFunc, Forward<VarTypes>(Vars)...));
+		return Super::AddDelegateInstance(FDelegate::CreateThreadSafeSP(InUserObjectRef, InFunc, Forward<VarTypes>(Vars)...));
 	}
 	template <typename UserClass, typename... VarTypes>
 	inline FDelegateHandle AddThreadSafeSP(const TSharedRef<UserClass, ESPMode::ThreadSafe>& InUserObjectRef, typename TMemFunPtrType<true, UserClass, void (ParamTypes..., std::decay_t<VarTypes>...)>::Type InFunc, VarTypes&&... Vars)
 	{
-		return Add(FDelegate::CreateThreadSafeSP(InUserObjectRef, InFunc, Forward<VarTypes>(Vars)...));
+		return Super::AddDelegateInstance(FDelegate::CreateThreadSafeSP(InUserObjectRef, InFunc, Forward<VarTypes>(Vars)...));
 	}
 
 	/**
@@ -791,12 +791,12 @@ public:
 	{
 		static_assert(!TIsConst<UserClass>::Value, "Attempting to bind a delegate with a const object pointer and non-const member function.");
 
-		return Add(FDelegate::CreateThreadSafeSP(InUserObject, InFunc, Forward<VarTypes>(Vars)...));
+		return Super::AddDelegateInstance(FDelegate::CreateThreadSafeSP(InUserObject, InFunc, Forward<VarTypes>(Vars)...));
 	}
 	template <typename UserClass, typename... VarTypes>
 	inline FDelegateHandle AddThreadSafeSP(const UserClass* InUserObject, typename TMemFunPtrType<true, UserClass, void (ParamTypes..., std::decay_t<VarTypes>...)>::Type InFunc, VarTypes&&... Vars)
 	{
-		return Add(FDelegate::CreateThreadSafeSP(InUserObject, InFunc, Forward<VarTypes>(Vars)...));
+		return Super::AddDelegateInstance(FDelegate::CreateThreadSafeSP(InUserObject, InFunc, Forward<VarTypes>(Vars)...));
 	}
 
 	/**
@@ -810,12 +810,12 @@ public:
 	template <typename UObjectTemplate, typename... VarTypes>
 	inline FDelegateHandle AddUFunction(UObjectTemplate* InUserObject, const FName& InFunctionName, VarTypes&&... Vars)
 	{
-		return Add(FDelegate::CreateUFunction(InUserObject, InFunctionName, Forward<VarTypes>(Vars)...));
+		return Super::AddDelegateInstance(FDelegate::CreateUFunction(InUserObject, InFunctionName, Forward<VarTypes>(Vars)...));
 	}
 	template <typename UObjectTemplate, typename... VarTypes>
 	inline FDelegateHandle AddUFunction(TObjectPtr<UObjectTemplate> InUserObject, const FName& InFunctionName, VarTypes&&... Vars)
 	{
-		return Add(FDelegate::CreateUFunction(InUserObject, InFunctionName, Forward<VarTypes>(Vars)...));
+		return Super::AddDelegateInstance(FDelegate::CreateUFunction(InUserObject, InFunctionName, Forward<VarTypes>(Vars)...));
 	}
 
 	/**
@@ -831,24 +831,24 @@ public:
 	{
 		static_assert(!TIsConst<UserClass>::Value, "Attempting to bind a delegate with a const object pointer and non-const member function.");
 
-		return Add(FDelegate::CreateUObject(InUserObject, InFunc, Forward<VarTypes>(Vars)...));
+		return Super::AddDelegateInstance(FDelegate::CreateUObject(InUserObject, InFunc, Forward<VarTypes>(Vars)...));
 	}
 	template <typename UserClass, typename... VarTypes>
 	inline FDelegateHandle AddUObject(const UserClass* InUserObject, typename TMemFunPtrType<true, UserClass, void (ParamTypes..., std::decay_t<VarTypes>...)>::Type InFunc, VarTypes&&... Vars)
 	{
-		return Add(FDelegate::CreateUObject(InUserObject, InFunc, Forward<VarTypes>(Vars)...));
+		return Super::AddDelegateInstance(FDelegate::CreateUObject(InUserObject, InFunc, Forward<VarTypes>(Vars)...));
 	}
 	template <typename UserClass, typename... VarTypes>
 	inline FDelegateHandle AddUObject(TObjectPtr<UserClass> InUserObject, typename TMemFunPtrType<false, UserClass, void (ParamTypes..., std::decay_t<VarTypes>...)>::Type InFunc, VarTypes&&... Vars)
 	{
 		static_assert(!TIsConst<UserClass>::Value, "Attempting to bind a delegate with a const object pointer and non-const member function.");
 
-		return Add(FDelegate::CreateUObject(InUserObject, InFunc, Forward<VarTypes>(Vars)...));
+		return Super::AddDelegateInstance(FDelegate::CreateUObject(InUserObject, InFunc, Forward<VarTypes>(Vars)...));
 	}
 	template <typename UserClass, typename... VarTypes>
 	inline FDelegateHandle AddUObject(TObjectPtr<UserClass> InUserObject, typename TMemFunPtrType<true, UserClass, void (ParamTypes..., std::decay_t<VarTypes>...)>::Type InFunc, VarTypes&&... Vars)
 	{
-		return Add(FDelegate::CreateUObject(InUserObject, InFunc, Forward<VarTypes>(Vars)...));
+		return Super::AddDelegateInstance(FDelegate::CreateUObject(InUserObject, InFunc, Forward<VarTypes>(Vars)...));
 	}
 
 public:
