Allowed TMulticastDelegateBase::AddDelegateInstance to accept any type of delegate, as long as it's explicitly convertible to the invocation list element type.

diff --git a/Engine/Source/Runtime/Core/Public/Delegates/MulticastDelegateBase.h b/Engine/Source/Runtime/Core/Public/Delegates/MulticastDelegateBase.h
--- a/Engine/Source/Runtime/Core/Public/Delegates/MulticastDelegateBase.h
+++ b/Engine/Source/Runtime/Core/Public/Delegates/MulticastDelegateBase.h
@@ -220,7 +220,8 @@ protected:
 	 *
 	 * @param NewDelegateBaseRef The delegate instance to add.
 	 */
-	inline FDelegateHandle AddDelegateInstance(TDelegateBase<UserPolicy>&& NewDelegateBaseRef)
+	template <typename NewDelegateType>
+	inline FDelegateHandle AddDelegateInstance(NewDelegateType&& NewDelegateBaseRef)
 	{
 		UE_DELEGATES_MT_SCOPED_WRITE_ACCESS(AccessDetector);
 
@@ -230,7 +231,7 @@ protected:
 			// compact but obey threshold of when this will trigger
 			CompactInvocationList(true);
 			Result = NewDelegateBaseRef.GetHandle();
-			InvocationList.Add(MoveTemp(NewDelegateBaseRef));
+			InvocationList.Emplace(Forward<NewDelegateType>(NewDelegateBaseRef));
 		}
 		return Result;
 	}
diff --git a/Engine/Source/Runtime/Core/Public/Delegates/TSMulticastDelegateBase.h b/Engine/Source/Runtime/Core/Public/Delegates/TSMulticastDelegateBase.h
--- a/Engine/Source/Runtime/Core/Public/Delegates/TSMulticastDelegateBase.h
+++ b/Engine/Source/Runtime/Core/Public/Delegates/TSMulticastDelegateBase.h
@@ -84,10 +84,11 @@ protected:
 	 *
 	 * @param NewDelegateBaseRef The delegate instance to add.
 	 */
-	inline FDelegateHandle AddDelegateInstance(TDelegateBase<UserPolicy>&& NewDelegateBaseRef)
+	template <typename NewDelegateType>
+	inline FDelegateHandle AddDelegateInstance(NewDelegateType&& NewDelegateBaseRef)
 	{
 		FScopeLock Lock(&CS);
-		return Super::AddDelegateInstance(MoveTemp(NewDelegateBaseRef));
+		return Super::AddDelegateInstance(Forward<NewDelegateType>(NewDelegateBaseRef));
 	}
 
 	/**
