Fixed some const correctness and inconsistencies in FDelegateBase implementation.

diff --git a/Engine/Source/Runtime/Core/Public/Delegates/DelegateBase.h b/Engine/Source/Runtime/Core/Public/Delegates/DelegateBase.h
--- a/Engine/Source/Runtime/Core/Public/Delegates/DelegateBase.h
+++ b/Engine/Source/Runtime/Core/Public/Delegates/DelegateBase.h
@@ -207,6 +207,9 @@ class TDelegateBase : public UserPolicy::FDelegateExtras
 
 	using Super = typename UserPolicy::FDelegateExtras;
 
+protected:
+	using Super::GetDelegateInstanceProtected;
+
 public:
 #if USE_DELEGATE_TRYGETBOUNDFUNCTIONNAME
 
@@ -220,12 +223,8 @@ public:
 	 */
 	FName TryGetBoundFunctionName() const
 	{
-		if (IDelegateInstance* Ptr = Super::GetDelegateInstanceProtected())
-		{
-			return Ptr->TryGetBoundFunctionName();
-		}
-
-		return NAME_None;
+		const IDelegateInstance* DelegateInstance = GetDelegateInstanceProtected();
+		return DelegateInstance ? DelegateInstance->TryGetBoundFunctionName() : NAME_None;
 	}
 
 #endif
@@ -237,12 +236,8 @@ public:
 	 */
 	FORCEINLINE class UObject* GetUObject( ) const
 	{
-		if (IDelegateInstance* Ptr = Super::GetDelegateInstanceProtected())
-		{
-			return Ptr->GetUObject();
-		}
-
-		return nullptr;
+		const IDelegateInstance* DelegateInstance = GetDelegateInstanceProtected();
+		return DelegateInstance ? DelegateInstance->GetUObject() : nullptr;
 	}
 
 	/**
@@ -252,9 +247,8 @@ public:
 	 */
 	FORCEINLINE bool IsBound( ) const
 	{
-		IDelegateInstance* Ptr = Super::GetDelegateInstanceProtected();
-
-		return Ptr && Ptr->IsSafeToExecute();
+		const IDelegateInstance* DelegateInstance = GetDelegateInstanceProtected();
+		return DelegateInstance && DelegateInstance->IsSafeToExecute();
 	}
 
 	/** 
@@ -264,10 +258,8 @@ public:
 	 */
 	FORCEINLINE const void* GetObjectForTimerManager() const
 	{
-		IDelegateInstance* Ptr = Super::GetDelegateInstanceProtected();
-
-		const void* Result = Ptr ? Ptr->GetObjectForTimerManager() : nullptr;
-		return Result;
+		const IDelegateInstance* DelegateInstance = GetDelegateInstanceProtected();
+		return DelegateInstance ? DelegateInstance->GetObjectForTimerManager() : nullptr;
 	}
 
 	/**
@@ -280,12 +272,8 @@ public:
 	 */
 	uint64 GetBoundProgramCounterForTimerManager() const
 	{
-		if (IDelegateInstance* Ptr = Super::GetDelegateInstanceProtected())
-		{
-			return Ptr->GetBoundProgramCounterForTimerManager();
-		}
-
-		return 0;
+		const IDelegateInstance* DelegateInstance = GetDelegateInstanceProtected();
+		return DelegateInstance ? DelegateInstance->GetBoundProgramCounterForTimerManager() : 0;
 	}
 
 	/** 
@@ -300,9 +288,8 @@ public:
 			return false;
 		}
 
-		IDelegateInstance* Ptr = Super::GetDelegateInstanceProtected();
-
-		return Ptr && Ptr->HasSameObject(InUserObject);
+		const IDelegateInstance* DelegateInstance = GetDelegateInstanceProtected();
+		return DelegateInstance && DelegateInstance->HasSameObject(InUserObject);
 	}
 
 	/**
@@ -312,12 +299,7 @@ public:
 	 */
 	FORCEINLINE FDelegateHandle GetHandle() const
 	{
-		FDelegateHandle Result;
-		if (IDelegateInstance* Ptr = Super::GetDelegateInstanceProtected())
-		{
-			Result = Ptr->GetHandle();
-		}
-
-		return Result;
+		const IDelegateInstance* DelegateInstance = GetDelegateInstanceProtected();
+		return DelegateInstance ? DelegateInstance->GetHandle() : FDelegateHandle{};
 	}
 };
