Add a OnCSVProfileEndFrame delegate, called only when the CSV profiler is running.

diff --git a/Engine/Source/Runtime/Core/Private/ProfilingDebugging/CsvProfiler.cpp b/Engine/Source/Runtime/Core/Private/ProfilingDebugging/CsvProfiler.cpp
--- a/Engine/Source/Runtime/Core/Private/ProfilingDebugging/CsvProfiler.cpp
+++ b/Engine/Source/Runtime/Core/Private/ProfilingDebugging/CsvProfiler.cpp
@@ -2975,6 +2975,8 @@ void FCsvProfiler::EndFrame()
 	check(IsInGameThread());
 	if (GCsvProfilerIsCapturing)
 	{
+		OnCSVProfileEndFrameDelegate.Broadcast();
+
 		GCsvPersistentCustomStats.RecordStats();
 
 		QUICK_SCOPE_CYCLE_COUNTER(STAT_FCsvProfiler_EndFrame_Capturing);
diff --git a/Engine/Source/Runtime/Core/Public/ProfilingDebugging/CsvProfiler.h b/Engine/Source/Runtime/Core/Public/ProfilingDebugging/CsvProfiler.h
--- a/Engine/Source/Runtime/Core/Public/ProfilingDebugging/CsvProfiler.h
+++ b/Engine/Source/Runtime/Core/Public/ProfilingDebugging/CsvProfiler.h
@@ -445,10 +445,14 @@ public:
 
 	DECLARE_MULTICAST_DELEGATE(FOnCSVProfileEnd);
 	FOnCSVProfileEnd& OnCSVProfileEnd() { return OnCSVProfileEndDelegate; }
-	
+
 	DECLARE_MULTICAST_DELEGATE_OneParam(FOnCSVProfileFinished, const FString& /*Filename */);
 	FOnCSVProfileFinished& OnCSVProfileFinished() { return OnCSVProfileFinishedDelegate; }
 
+	/** Called at the end of the frame when capturing, before end frame processing completes */
+	DECLARE_MULTICAST_DELEGATE(FOnCSVProfileEndFrame);
+	FOnCSVProfileEndFrame& OnCSVProfileEndFrame() { return OnCSVProfileEndFrameDelegate; }
+
 	CORE_API void SetRenderThreadId(uint32 InRenderThreadId)
 	{
 		RenderThreadId = InRenderThreadId;
@@ -509,7 +513,7 @@ private:
 	FOnCSVProfileFirstFrame OnCSVProfileFirstFrameDelegate;
 	FOnCSVProfileEndRequested OnCSVProfileEndRequestedDelegate;
 	FOnCSVProfileEnd OnCSVProfileEndDelegate;
-	
+	FOnCSVProfileEndFrame OnCSVProfileEndFrameDelegate;
 	FOnCSVProfileFinished OnCSVProfileFinishedDelegate;
 
 	std::atomic<uint32> RenderThreadId{ 0 };
