Always use a space and a comma as the group and decimal separator when use RAE-style Spanish numbers

diff --git a/Engine/Source/Runtime/Core/Private/Internationalization/ICUCulture.cpp b/Engine/Source/Runtime/Core/Private/Internationalization/ICUCulture.cpp
--- a/Engine/Source/Runtime/Core/Private/Internationalization/ICUCulture.cpp
+++ b/Engine/Source/Runtime/Core/Private/Internationalization/ICUCulture.cpp
@@ -691,6 +691,7 @@ FDecimalNumberFormattingRules ExtractNumberFormattingRulesFromICUDecimalFormatte
 	if (FCStringAnsi::Strcmp(InICULocale.getLanguage(), "es") == 0 && CVarSpanishUsesRAENumberFormat.AsVariable()->GetInt())
 	{
 		NewUEDecimalNumberFormattingRules.GroupingSeparatorCharacter = TEXT('\u00A0'); // No-Break Space
+		NewUEDecimalNumberFormattingRules.DecimalSeparatorCharacter = TEXT(',');
 	}
 
 	return NewUEDecimalNumberFormattingRules;
