Additional safety checks for Windows stack walking.
Bail if we can't suspend the thread we want to trace.
If we record null values for instruction, stack or base pointers, bail.

diff --git a/Engine/Source/Runtime/Core/Private/Microsoft/MicrosoftPlatformStackWalk.cpp b/Engine/Source/Runtime/Core/Private/Microsoft/MicrosoftPlatformStackWalk.cpp
--- a/Engine/Source/Runtime/Core/Private/Microsoft/MicrosoftPlatformStackWalk.cpp
+++ b/Engine/Source/Runtime/Core/Private/Microsoft/MicrosoftPlatformStackWalk.cpp
@@ -85,6 +85,11 @@ int32 FMicrosoftPlatformStackWalk::CaptureStackTraceInternal(uint64* OutBacktrac
 
 	*OutDepth = 0;
 
+	if (!ContextCopy.Rbp || !ContextCopy.Rip || !ContextCopy.Rsp)
+	{
+		return EXCEPTION_EXECUTE_HANDLER;
+	}
+
 #if !PLATFORM_SEH_EXCEPTIONS_DISABLED
 	__try
 #endif
diff --git a/Engine/Source/Runtime/Core/Private/Windows/WindowsPlatformStackWalk.cpp b/Engine/Source/Runtime/Core/Private/Windows/WindowsPlatformStackWalk.cpp
--- a/Engine/Source/Runtime/Core/Private/Windows/WindowsPlatformStackWalk.cpp
+++ b/Engine/Source/Runtime/Core/Private/Windows/WindowsPlatformStackWalk.cpp
@@ -210,7 +210,10 @@ uint32 FWindowsPlatformStackWalk::CaptureThreadStackBackTrace(uint64 ThreadId, u
 	}
 
 	// Suspend the thread before grabbing its context
-	SuspendThread(ThreadHandle);
+	if (SuspendThread(ThreadHandle) < 0)
+	{
+		return 0;
+	}
 
 	uint32 Depth = 0;
 	if (Context)
