LogLocalization now reports Log and above by default (it previously only reported Errors by default)

#rnx

diff --git a/Engine/Source/Runtime/Core/Public/CoreGlobals.h b/Engine/Source/Runtime/Core/Public/CoreGlobals.h
--- a/Engine/Source/Runtime/Core/Public/CoreGlobals.h
+++ b/Engine/Source/Runtime/Core/Public/CoreGlobals.h
@@ -45,7 +45,7 @@ CORE_API DECLARE_LOG_CATEGORY_EXTERN(LogInit, Log, All);
 CORE_API DECLARE_LOG_CATEGORY_EXTERN(LogExit, Log, All);
 CORE_API DECLARE_LOG_CATEGORY_EXTERN(LogExec, Warning, All);
 CORE_API DECLARE_LOG_CATEGORY_EXTERN(LogScript, Warning, All);
-CORE_API DECLARE_LOG_CATEGORY_EXTERN(LogLocalization, Error, All);
+CORE_API DECLARE_LOG_CATEGORY_EXTERN(LogLocalization, Log, All);
 CORE_API DECLARE_LOG_CATEGORY_EXTERN(LogLongPackageNames, Log, All);
 CORE_API DECLARE_LOG_CATEGORY_EXTERN(LogProcess, Log, All);
 CORE_API DECLARE_LOG_CATEGORY_EXTERN(LogLoad, Log, All);
