Preparing the deprecate the FMessageDialog overloads that take the optional title by-pointer

#rnx

diff --git a/Engine/Source/Runtime/Core/Private/Misc/MessageDialog.cpp b/Engine/Source/Runtime/Core/Private/Misc/MessageDialog.cpp
--- a/Engine/Source/Runtime/Core/Private/Misc/MessageDialog.cpp
+++ b/Engine/Source/Runtime/Core/Private/Misc/MessageDialog.cpp
@@ -36,6 +36,11 @@ void FMessageDialog::Debugf( const FText& Message, const FText* OptTitle )
 	Debugf(Message, OptTitle ? *OptTitle : GetDefaultMessageTitle());
 }
 
+void FMessageDialog::Debugf( const FText& Message )
+{
+	Debugf(Message, GetDefaultMessageTitle());
+}
+
 void FMessageDialog::Debugf( const FText& Message, const FText& Title )
 {
 	if( FApp::IsUnattended() == true )
@@ -76,6 +81,11 @@ EAppReturnType::Type FMessageDialog::Open( EAppMsgType::Type MessageType, const
 	return Open(MessageType, Message, OptTitle ? *OptTitle : GetDefaultMessageTitle());
 }
 
+EAppReturnType::Type FMessageDialog::Open( EAppMsgType::Type MessageType, const FText& Message )
+{
+	return Open(MessageType, Message, GetDefaultMessageTitle());
+}
+
 EAppReturnType::Type FMessageDialog::Open(EAppMsgType::Type MessageType, const FText& Message, const FText& Title)
 {
 	EAppReturnType::Type DefaultValue = EAppReturnType::Yes;
@@ -130,6 +140,11 @@ EAppReturnType::Type FMessageDialog::Open(EAppMsgType::Type MessageType, EAppRet
 	return Open(MessageType, DefaultValue, Message, OptTitle ? *OptTitle : GetDefaultMessageTitle());
 }
 
+EAppReturnType::Type FMessageDialog::Open(EAppMsgType::Type MessageType, EAppReturnType::Type DefaultValue, const FText& Message)
+{
+	return Open(MessageType, DefaultValue, Message, GetDefaultMessageTitle());
+}
+
 EAppReturnType::Type FMessageDialog::Open(EAppMsgType::Type MessageType, EAppReturnType::Type DefaultValue, const FText& Message, const FText& Title)
 {
 	EAppReturnType::Type Result = DefaultValue;
diff --git a/Engine/Source/Runtime/Core/Public/Misc/MessageDialog.h b/Engine/Source/Runtime/Core/Public/Misc/MessageDialog.h
--- a/Engine/Source/Runtime/Core/Public/Misc/MessageDialog.h
+++ b/Engine/Source/Runtime/Core/Public/Misc/MessageDialog.h
@@ -17,10 +17,12 @@ struct CORE_API FMessageDialog
 {
 	/** Pops up a message dialog box containing the input string.
 	 * @param Message Text of message to show
-	 * @param OptTitle Optional title to use (defaults to "Message")
+	 * @param Title Optional title to use (defaults to "Message")
 	*/
-	static void Debugf( const FText& Message, const FText* OptTitle = nullptr );
+	static void Debugf( const FText& Message );
 	static void Debugf( const FText& Message, const FText& Title );
+	// TODO: UE_DEPRECATED(5.3, "Use the overload of Debugf that takes the Title by-value")
+	static void Debugf( const FText& Message, const FText* OptTitle );
 
 	/** Pops up a message dialog box containing the last system error code in string form. */
 	static void ShowLastError();
@@ -29,18 +31,22 @@ struct CORE_API FMessageDialog
 	 * Open a modal message box dialog
 	 * @param MessageType Controls buttons dialog should have
 	 * @param Message Text of message to show
-	 * @param OptTitle Optional title to use (defaults to "Message")
+	 * @param Title Optional title to use (defaults to "Message")
 	*/
-	static EAppReturnType::Type Open( EAppMsgType::Type MessageType, const FText& Message, const FText* OptTitle = nullptr);
+	static EAppReturnType::Type Open( EAppMsgType::Type MessageType, const FText& Message);
 	static EAppReturnType::Type Open( EAppMsgType::Type MessageType, const FText& Message, const FText& Title);
+	// TODO: UE_DEPRECATED(5.3, "Use the overload of Open that takes the Title by-value")
+	static EAppReturnType::Type Open( EAppMsgType::Type MessageType, const FText& Message, const FText* OptTitle);
 
 	/**
 	 * Open a modal message box dialog
 	 * @param MessageType Controls buttons dialog should have
 	 * @param DefaultValue If the application is Unattended, the function will log and return DefaultValue
 	 * @param Message Text of message to show
-	 * @param OptTitle Optional title to use (defaults to "Message")
+	 * @param Title Optional title to use (defaults to "Message")
 	*/
-	static EAppReturnType::Type Open(EAppMsgType::Type MessageType, EAppReturnType::Type DefaultValue, const FText& Message, const FText* OptTitle = nullptr);
+	static EAppReturnType::Type Open(EAppMsgType::Type MessageType, EAppReturnType::Type DefaultValue, const FText& Message);
 	static EAppReturnType::Type Open(EAppMsgType::Type MessageType, EAppReturnType::Type DefaultValue, const FText& Message, const FText& Title);
+	// TODO: UE_DEPRECATED(5.3, "Use the overload of Open that takes the Title by-value")
+	static EAppReturnType::Type Open(EAppMsgType::Type MessageType, EAppReturnType::Type DefaultValue, const FText& Message, const FText* OptTitle);
 };
