Deprecated the FMessageDialog overloads that take the title by-pointer

Use the overload that passes it by-value instead

diff --git a/Engine/Source/Runtime/Core/Public/Misc/MessageDialog.h b/Engine/Source/Runtime/Core/Public/Misc/MessageDialog.h
--- a/Engine/Source/Runtime/Core/Public/Misc/MessageDialog.h
+++ b/Engine/Source/Runtime/Core/Public/Misc/MessageDialog.h
@@ -21,7 +21,7 @@ struct CORE_API FMessageDialog
 	*/
 	static void Debugf( const FText& Message );
 	static void Debugf( const FText& Message, const FText& Title );
-	// TODO: UE_DEPRECATED(5.3, "Use the overload of Debugf that takes the Title by-value")
+	UE_DEPRECATED(5.3, "Use the overload of Debugf that takes the Title by-value.")
 	static void Debugf( const FText& Message, const FText* OptTitle );
 
 	/** Pops up a message dialog box containing the last system error code in string form. */
@@ -35,7 +35,7 @@ struct CORE_API FMessageDialog
 	*/
 	static EAppReturnType::Type Open( EAppMsgType::Type MessageType, const FText& Message);
 	static EAppReturnType::Type Open( EAppMsgType::Type MessageType, const FText& Message, const FText& Title);
-	// TODO: UE_DEPRECATED(5.3, "Use the overload of Open that takes the Title by-value")
+	UE_DEPRECATED(5.3, "Use the overload of Open that takes the Title by-value.")
 	static EAppReturnType::Type Open( EAppMsgType::Type MessageType, const FText& Message, const FText* OptTitle);
 
 	/**
@@ -47,6 +47,6 @@ struct CORE_API FMessageDialog
 	*/
 	static EAppReturnType::Type Open(EAppMsgType::Type MessageType, EAppReturnType::Type DefaultValue, const FText& Message);
 	static EAppReturnType::Type Open(EAppMsgType::Type MessageType, EAppReturnType::Type DefaultValue, const FText& Message, const FText& Title);
-	// TODO: UE_DEPRECATED(5.3, "Use the overload of Open that takes the Title by-value")
+	UE_DEPRECATED(5.3, "Use the overload of Open that takes the Title by-value.")
 	static EAppReturnType::Type Open(EAppMsgType::Type MessageType, EAppReturnType::Type DefaultValue, const FText& Message, const FText* OptTitle);
 };
