Moved the parsing of AutomationControllerSettings out of GetAutomationLogLevel into FAutomationTestBase as a static function EnableDefaultLogSettings.
Adapted bSuppressLogErrors, bSuppressLogWarnings, bElevateLogWarningsToErrors, SuppressedLogCategories for that function

diff --git a/Engine/Source/Runtime/Core/Private/Misc/AutomationTest.cpp b/Engine/Source/Runtime/Core/Private/Misc/AutomationTest.cpp
--- a/Engine/Source/Runtime/Core/Private/Misc/AutomationTest.cpp
+++ b/Engine/Source/Runtime/Core/Private/Misc/AutomationTest.cpp
@@ -43,6 +43,10 @@ namespace AutomationTest
 		TEXT("Whether to output blueprint functional test metadata to the log when test is running"));
 };
 
+bool FAutomationTestBase::bSuppressLogWarnings = false;
+bool FAutomationTestBase::bSuppressLogErrors = false;
+bool FAutomationTestBase::bElevateLogWarningsToErrors = false;
+TArray<FString> FAutomationTestBase::SuppressedLogCategories;
 
 CORE_API const TMap<FString, EAutomationTestFlags::Type>& EAutomationTestFlags::GetTestFlagsMap()
 {
@@ -88,10 +92,6 @@ CORE_API ELogVerbosity::Type GetAutomationLogLevel(ELogVerbosity::Type LogVerbos
 	// agrant-todo: these should be controlled by FAutomationTestBase for 4.27 with the same project-level override that
 	// FunctionalTest has. Now that warnings are correctly associated with tests they need to be something all tests
 	// can leverage, not just functional tests
-	static bool bSuppressLogWarnings = false;
-	static bool bSuppressLogErrors = false;
-	static bool bElevateLogWarningsToErrors = false;
-	static TArray<FString> SuppressedLogCategories;
 	static FAutomationTestBase* LastTest = nullptr;
 
 	if (AutomationTest::bCaptureLogEvents == false)
@@ -99,20 +99,16 @@ CORE_API ELogVerbosity::Type GetAutomationLogLevel(ELogVerbosity::Type LogVerbos
 		return ELogVerbosity::NoLogging;
 	}
 
-	if (CurrentTest != LastTest)
+	if (CurrentTest != LastTest) 
 	{
-		// These can be changed in the editor so can't just be cached for the whole session
-		SuppressedLogCategories.Empty();
-		GConfig->GetBool(TEXT("/Script/AutomationController.AutomationControllerSettings"), TEXT("bSuppressLogErrors"), bSuppressLogErrors, GEngineIni);
-		GConfig->GetBool(TEXT("/Script/AutomationController.AutomationControllerSettings"), TEXT("bSuppressLogWarnings"), bSuppressLogWarnings, GEngineIni);
-		GConfig->GetBool(TEXT("/Script/AutomationController.AutomationControllerSettings"), TEXT("bElevateLogWarningsToErrors"), bElevateLogWarningsToErrors, GEngineIni);
-		GConfig->GetArray(TEXT("/Script/AutomationController.AutomationControllerSettings"), TEXT("SuppressedLogCategories"), SuppressedLogCategories, GEngineIni);
+		FAutomationTestBase::SuppressedLogCategories.Empty();
+		FAutomationTestBase::LoadDefaultLogSettings();
 		LastTest = CurrentTest;
 	}
 
 	if (CurrentTest)
 	{
-		if (CurrentTest->SuppressLogs() || SuppressedLogCategories.Contains(LogCategory.ToString()))
+		if (CurrentTest->SuppressLogs() || CurrentTest->GetSuppressedLogCategories().Contains(LogCategory.ToString()))
 		{
 			EffectiveVerbosity = ELogVerbosity::NoLogging;
 		}
@@ -120,11 +116,11 @@ CORE_API ELogVerbosity::Type GetAutomationLogLevel(ELogVerbosity::Type LogVerbos
 		{
 			if (EffectiveVerbosity == ELogVerbosity::Warning)
 			{
-				if (CurrentTest->SuppressLogWarnings() || bSuppressLogWarnings)
+				if (CurrentTest->SuppressLogWarnings())
 				{
 					EffectiveVerbosity = ELogVerbosity::NoLogging;
 				}
-				else if (CurrentTest->ElevateLogWarningsToErrors() || bElevateLogWarningsToErrors)
+				else if (CurrentTest->ElevateLogWarningsToErrors())
 				{
 					EffectiveVerbosity = ELogVerbosity::Error;
 				}
@@ -132,7 +128,7 @@ CORE_API ELogVerbosity::Type GetAutomationLogLevel(ELogVerbosity::Type LogVerbos
 
 			if (EffectiveVerbosity == ELogVerbosity::Error)
 			{
-				if (CurrentTest->SuppressLogErrors() ||  bSuppressLogErrors)
+				if (CurrentTest->SuppressLogErrors())
 				{
 					EffectiveVerbosity = ELogVerbosity::NoLogging;
 				}
@@ -1457,6 +1453,14 @@ bool FAutomationTestBase::LogCategoryMatchesSeverityInclusive(
 	return Actual == ELogVerbosity::All || MaximumVerbosity == ELogVerbosity::All || Actual <= MaximumVerbosity;
 }
 
+void FAutomationTestBase::LoadDefaultLogSettings()
+{
+	GConfig->GetBool(TEXT("/Script/AutomationController.AutomationControllerSettings"), TEXT("bSuppressLogErrors"), bSuppressLogErrors, GEngineIni);
+	GConfig->GetBool(TEXT("/Script/AutomationController.AutomationControllerSettings"), TEXT("bSuppressLogWarnings"), bSuppressLogWarnings, GEngineIni);
+	GConfig->GetBool(TEXT("/Script/AutomationController.AutomationControllerSettings"), TEXT("bElevateLogWarningsToErrors"), bElevateLogWarningsToErrors, GEngineIni);
+	GConfig->GetArray(TEXT("/Script/AutomationController.AutomationControllerSettings"), TEXT("SuppressedLogCategories"), SuppressedLogCategories, GEngineIni);
+}
+
 // --------------------------------------------------------------------------------------
 
 bool FAutomationTestBase::TestEqual(const TCHAR* What, const int32 Actual, const int32 Expected)
diff --git a/Engine/Source/Runtime/Core/Public/Misc/AutomationTest.h b/Engine/Source/Runtime/Core/Public/Misc/AutomationTest.h
--- a/Engine/Source/Runtime/Core/Public/Misc/AutomationTest.h
+++ b/Engine/Source/Runtime/Core/Public/Misc/AutomationTest.h
@@ -1282,6 +1282,12 @@ public:
 		FAutomationTestFramework::Get().UnregisterAutomationTest( TestName );
 	}
 
+	/** Log flags */
+	static bool bSuppressLogWarnings;
+	static bool bSuppressLogErrors;
+	static bool bElevateLogWarningsToErrors;
+	static TArray<FString> SuppressedLogCategories;
+
 	/**
 	 * Pure virtual method; returns the flags associated with the given automation test
 	 *
@@ -1456,6 +1462,11 @@ public:
 	static bool LogCategoryMatchesSeverityInclusive(ELogVerbosity::Type Actual, ELogVerbosity::Type MaximumVerbosity);
 
 	/**
+	 * Enables log settings from config
+	*/
+	static void LoadDefaultLogSettings();
+	/**
+	
 	* Adds a regex pattern to an internal list that this test will expect to encounter in logs (of the specified verbosity) during its execution. If an expected pattern
 	* is not encountered, it will cause this test to fail.
 	*
@@ -1521,7 +1532,7 @@ public:
 	{
 		return bSuppressLogs;
 	}
-	
+
 	/**
 	 * Should the log category be captured and surfaced as part of the test.
 	 * If true will then go through the SuppressLogWarnings and SuppressLogErrors checks for if this should be suppressed further or not
@@ -1536,21 +1547,27 @@ public:
 	 *
 	 * @return false to make errors errors
 	 */
-	virtual bool SuppressLogErrors() { return false; }
+	virtual bool SuppressLogErrors() { return bSuppressLogErrors; }
 
 	/**
 	 * If returns true then logging with a level of Warning will not be recorded in test results
 	 *
 	 * @return true to make warnings errors
 	 */
-	virtual bool SuppressLogWarnings() { return false; }
+	virtual bool SuppressLogWarnings() { return bSuppressLogWarnings; }
 
 	/**
 	 * If returns true then logging with a level of Warning will be treated as an error
 	 *
 	 * @return true to make warnings errors
 	 */
-	virtual bool ElevateLogWarningsToErrors() { return false; }
+	virtual bool ElevateLogWarningsToErrors() { return bElevateLogWarningsToErrors; }
+
+	/**
+	 * Return suppressed log categories
+	 */
+	virtual TArray<FString> GetSuppressedLogCategories() { return SuppressedLogCategories; }
+
 
 	/**
 	 * Enqueues a new latent command.
