Merged TDelegateBase and FDelegateBase, moving towards the removal of delegate policies.

diff --git a/Engine/Source/Runtime/Core/Public/Delegates/DelegateBase.h b/Engine/Source/Runtime/Core/Public/Delegates/DelegateBase.h
--- a/Engine/Source/Runtime/Core/Public/Delegates/DelegateBase.h
+++ b/Engine/Source/Runtime/Core/Public/Delegates/DelegateBase.h
@@ -36,33 +36,73 @@
 template <typename UserPolicy>
 class TTSMulticastDelegateBase;
 
+template <typename UserPolicy>
+class TDelegateBase;
+
 ALIAS_TEMPLATE_TYPE_LAYOUT(template<typename ElementType>, FDelegateAllocatorType::ForElementType<ElementType>, void*);
 
+struct FDefaultDelegateUserPolicy
+{
+	// To extend delegates, you should implement a policy struct like this and pass it as the second template
+	// argument to TDelegate and TMulticastDelegate.  This policy struct containing three classes called:
+	// 
+	// FDelegateInstanceExtras:
+	//   - Must publicly inherit IDelegateInstance.
+	//   - Should contain any extra data and functions injected into a binding (the object which holds and
+	//     is able to invoke the binding passed to FMyDelegate::CreateSP, FMyDelegate::CreateLambda etc.).
+	//   - This binding is not available through the public API of the delegate, but is accessible to FDelegateExtras.
+	//
+	// FDelegateExtras:
+	//   - Must publicly inherit TDelegateBase<UserPolicy>.
+	//   - Should contain any extra data and functions injected into a delegate (the object which holds an
+	//     FDelegateInstance-derived object, above).
+	//   - Public data members and member functions are accessible directly through the TDelegate object.
+	//   - Typically member functions in this class will forward calls to the inner FDelegateInstanceExtras,
+	//     by downcasting the result of a call to GetDelegateInstanceProtected().
+	//
+	// FMulticastDelegateExtras:
+	//   - Must publicly inherit TMulticastDelegateBase<FYourUserPolicyStruct>.
+	//   - Should contain any extra data and functions injected into a multicast delegate (the object which
+	//     holds an array of FDelegateExtras-derived objects which is the invocation list).
+	//   - Public data members and member functions are accessible directly through the TMulticastDelegate object.
+
+	using FDelegateInstanceExtras  = IDelegateInstance;
+	using FDelegateExtras          = TDelegateBase<FDefaultDelegateUserPolicy>;
+	using FMulticastDelegateExtras = TMulticastDelegateBase<FDefaultDelegateUserPolicy>;
+};
+
+struct FDefaultTSDelegateUserPolicy
+{
+	// see `FDefaultDelegateUserPolicy` for documentation
+	using FDelegateInstanceExtras  = IDelegateInstance;
+	using FDelegateExtras          = TDelegateBase<FDefaultTSDelegateUserPolicy>;
+	using FMulticastDelegateExtras = TTSMulticastDelegateBase<FDefaultTSDelegateUserPolicy>;
+};
+
 /**
  * Base class for unicast delegates.
  */
-class FDelegateBase
+template <typename UserPolicy>
+class TDelegateBase
 {
 	template <typename>
 	friend class TMulticastDelegateBase;
 
-	template <typename>
-	friend class TDelegateBase;
-
 protected:
-	explicit FDelegateBase() = default;
+	explicit TDelegateBase() = default;
 
-	~FDelegateBase()
+public:
+	~TDelegateBase()
 	{
 		Unbind();
 	}
 
-	FDelegateBase(FDelegateBase&& Other)
+	TDelegateBase(TDelegateBase&& Other)
 	{
 		MoveConstruct(MoveTemp(Other));
 	}
 
-	FDelegateBase& operator=(FDelegateBase&& Other)
+	TDelegateBase& operator=(TDelegateBase&& Other)
 	{
 		MoveAssign(MoveTemp(Other));
 		return *this;
@@ -84,18 +124,5 @@
 	}
 
 	/**
-	 * Gets the delegate instance.  Not intended for use by user code.
-	 *
-	 * @return The delegate instance.
-	 * @see SetDelegateInstance
-	 */
-	FORCEINLINE IDelegateInstance* GetDelegateInstanceProtected() const
-	{
-		UE_DELEGATES_MT_SCOPED_READ_ACCESS(AccessDetector);
-
-		return DelegateSize ? (IDelegateInstance*)DelegateAllocator.GetAllocation() : nullptr;
-	}
-
-	/**
 	 * Returns the amount of memory allocated by this delegate, not including sizeof(*this).
 	 */
@@ -104,113 +131,6 @@
 		return DelegateAllocator.GetAllocatedSize(DelegateSize, sizeof(FAlignedInlineDelegateType));
 	}
 
-public:
-	/**
-	 * "emplacement" of delegate instance of the given type
-	 */
-	template<typename DelegateInstanceType, typename... DelegateInstanceParams>
-	void CreateDelegateInstance(DelegateInstanceParams&&... Params)
-	{
-		UE_DELEGATES_MT_SCOPED_READ_ACCESS(AccessDetector);
-
-		IDelegateInstance* DelegateInstance = GetDelegateInstanceProtected();
-		if (DelegateInstance)
-		{
-			DelegateInstance->~IDelegateInstance();
-		}
-
-		new(Allocate(sizeof(DelegateInstanceType))) DelegateInstanceType(Forward<DelegateInstanceParams>(Params)...);
-	}
-
-private:
-	void* Allocate(int32 Size)
-	{
-		UE_DELEGATES_MT_SCOPED_WRITE_ACCESS(AccessDetector);
-
-		int32 NewDelegateSize = FMath::DivideAndRoundUp(Size, (int32)sizeof(FAlignedInlineDelegateType));
-		if (DelegateSize != NewDelegateSize)
-		{
-			DelegateAllocator.ResizeAllocation(0, NewDelegateSize, sizeof(FAlignedInlineDelegateType));
-			DelegateSize = NewDelegateSize;
-		}
-
-		return DelegateAllocator.GetAllocation();
-	}
-
-	void MoveConstruct(FDelegateBase&& Other)
-	{
-		DelegateAllocator.MoveToEmpty(Other.DelegateAllocator);
-		DelegateSize = Other.DelegateSize;
-		Other.DelegateSize = 0;
-	}
-
-	void MoveAssign(FDelegateBase&& Other)
-	{
-		UE_DELEGATES_MT_SCOPED_WRITE_ACCESS(AccessDetector);
-
-		Unbind();
-		DelegateAllocator.MoveToEmpty(Other.DelegateAllocator);
-		DelegateSize = Other.DelegateSize;
-		Other.DelegateSize = 0;
-	}
-
-private:
-	FDelegateAllocatorType::ForElementType<FAlignedInlineDelegateType> DelegateAllocator;
-	int32 DelegateSize = 0;
-
-	UE_DELEGATES_MT_ACCESS_DETECTOR(AccessDetector);
-};
-
-struct FDefaultDelegateUserPolicy
-{
-	// To extend delegates, you should implement a policy struct like this and pass it as the second template
-	// argument to TDelegate and TMulticastDelegate.  This policy struct containing three classes called:
-	// 
-	// FDelegateInstanceExtras:
-	//   - Must publicly inherit IDelegateInstance.
-	//   - Should contain any extra data and functions injected into a binding (the object which holds and
-	//     is able to invoke the binding passed to FMyDelegate::CreateSP, FMyDelegate::CreateLambda etc.).
-	//   - This binding is not available through the public API of the delegate, but is accessible to FDelegateExtras.
-	//
-	// FDelegateExtras:
-	//   - Must publicly inherit FDelegateBase.
-	//   - Should contain any extra data and functions injected into a delegate (the object which holds an
-	//     FDelegateInstance-derived object, above).
-	//   - Public data members and member functions are accessible directly through the TDelegate object.
-	//   - Typically member functions in this class will forward calls to the inner FDelegateInstanceExtras,
-	//     by downcasting the result of a call to GetDelegateInstanceProtected().
-	//
-	// FMulticastDelegateExtras:
-	//   - Must publicly inherit TMulticastDelegateBase<FYourUserPolicyStruct>.
-	//   - Should contain any extra data and functions injected into a multicast delegate (the object which
-	//     holds an array of FDelegateExtras-derived objects which is the invocation list).
-	//   - Public data members and member functions are accessible directly through the TMulticastDelegate object.
-
-	using FDelegateInstanceExtras  = IDelegateInstance;
-	using FDelegateExtras          = FDelegateBase;
-	using FMulticastDelegateExtras = TMulticastDelegateBase<FDefaultDelegateUserPolicy>;
-};
-
-struct FDefaultTSDelegateUserPolicy
-{
-	// see `FDefaultDelegateUserPolicy` for documentation
-	using FDelegateInstanceExtras = IDelegateInstance;
-	using FDelegateExtras = FDelegateBase;
-	using FMulticastDelegateExtras = TTSMulticastDelegateBase<FDefaultTSDelegateUserPolicy>;
-};
-
-template <typename UserPolicy>
-class TDelegateBase : public UserPolicy::FDelegateExtras
-{
-	template <typename>
-	friend class TMulticastDelegateBase;
-
-	using Super = typename UserPolicy::FDelegateExtras;
-
-protected:
-	using Super::GetDelegateInstanceProtected;
-
-public:
 #if USE_DELEGATE_TRYGETBOUNDFUNCTIONNAME
 
 	/**
@@ -302,4 +222,80 @@
 		const IDelegateInstance* DelegateInstance = GetDelegateInstanceProtected();
 		return DelegateInstance ? DelegateInstance->GetHandle() : FDelegateHandle{};
 	}
+
+	/**
+	 * "emplacement" of delegate instance of the given type
+	 */
+	template<typename DelegateInstanceType, typename... DelegateInstanceParams>
+	void CreateDelegateInstance(DelegateInstanceParams&&... Params)
+	{
+		UE_DELEGATES_MT_SCOPED_READ_ACCESS(AccessDetector);
+
+		IDelegateInstance* DelegateInstance = GetDelegateInstanceProtected();
+		if (DelegateInstance)
+		{
+			DelegateInstance->~IDelegateInstance();
+		}
+
+		new(Allocate(sizeof(DelegateInstanceType))) DelegateInstanceType(Forward<DelegateInstanceParams>(Params)...);
+	}
+
+protected:
+	/**
+	 * Gets the delegate instance.  Not intended for use by user code.
+	 *
+	 * @return The delegate instance.
+	 * @see SetDelegateInstance
+	 */
+	FORCEINLINE IDelegateInstance* GetDelegateInstanceProtected()
+	{
+		UE_DELEGATES_MT_SCOPED_READ_ACCESS(AccessDetector);
+
+		return DelegateSize ? (IDelegateInstance*)DelegateAllocator.GetAllocation() : nullptr;
+	}
+
+	FORCEINLINE const IDelegateInstance* GetDelegateInstanceProtected() const
+	{
+		UE_DELEGATES_MT_SCOPED_READ_ACCESS(AccessDetector);
+
+		return DelegateSize ? (const IDelegateInstance*)DelegateAllocator.GetAllocation() : nullptr;
+	}
+
+private:
+	void* Allocate(int32 Size)
+	{
+		UE_DELEGATES_MT_SCOPED_WRITE_ACCESS(AccessDetector);
+
+		int32 NewDelegateSize = FMath::DivideAndRoundUp(Size, (int32)sizeof(FAlignedInlineDelegateType));
+		if (DelegateSize != NewDelegateSize)
+		{
+			DelegateAllocator.ResizeAllocation(0, NewDelegateSize, sizeof(FAlignedInlineDelegateType));
+			DelegateSize = NewDelegateSize;
+		}
+
+		return DelegateAllocator.GetAllocation();
+	}
+
+	void MoveConstruct(TDelegateBase&& Other)
+	{
+		DelegateAllocator.MoveToEmpty(Other.DelegateAllocator);
+		DelegateSize = Other.DelegateSize;
+		Other.DelegateSize = 0;
+	}
+
+	void MoveAssign(TDelegateBase&& Other)
+	{
+		UE_DELEGATES_MT_SCOPED_WRITE_ACCESS(AccessDetector);
+
+		Unbind();
+		DelegateAllocator.MoveToEmpty(Other.DelegateAllocator);
+		DelegateSize = Other.DelegateSize;
+		Other.DelegateSize = 0;
+	}
+
+private:
+	FDelegateAllocatorType::ForElementType<FAlignedInlineDelegateType> DelegateAllocator;
+	int32 DelegateSize = 0;
+
+	UE_DELEGATES_MT_ACCESS_DETECTOR(AccessDetector);
 };
diff --git a/Engine/Source/Runtime/Core/Public/Delegates/DelegateInstanceInterface.h b/Engine/Source/Runtime/Core/Public/Delegates/DelegateInstanceInterface.h
--- a/Engine/Source/Runtime/Core/Public/Delegates/DelegateInstanceInterface.h
+++ b/Engine/Source/Runtime/Core/Public/Delegates/DelegateInstanceInterface.h
@@ -6,7 +6,6 @@
 #include "Templates/TypeCompatibleBytes.h"
 #include "Templates/Tuple.h"
 
-class FDelegateBase;
 class IDelegateInstance;
 
 template <typename FuncType, typename UserPolicy>
diff --git a/Engine/Source/Runtime/Core/Public/Delegates/DelegateSignatureImpl.inl b/Engine/Source/Runtime/Core/Public/Delegates/DelegateSignatureImpl.inl
--- a/Engine/Source/Runtime/Core/Public/Delegates/DelegateSignatureImpl.inl
+++ b/Engine/Source/Runtime/Core/Public/Delegates/DelegateSignatureImpl.inl
@@ -21,7 +21,6 @@
 #include "Templates/UnrealTemplate.h"
 #include <type_traits>
 
-class FDelegateBase;
 class FDelegateHandle;
 class IDelegateInstance;
 struct FWeakObjectPtr;
@@ -53,7 +52,6 @@ class TDelegate<InRetValType(ParamTypes...), UserPolicy> : public TDelegateBase<
 	using FuncType                      = InRetValType (ParamTypes...);
 	using DelegateInstanceInterfaceType = IBaseDelegateInstance<FuncType, UserPolicy>;
 
-	static_assert(std::is_convertible_v<typename UserPolicy::FDelegateExtras*, FDelegateBase*>, "UserPolicy::FDelegateExtras should publicly inherit FDelegateBase");
 	static_assert(std::is_convertible_v<typename UserPolicy::FDelegateInstanceExtras*, IDelegateInstance*>, "UserPolicy::FDelegateInstanceExtras should publicly inherit IDelegateInstance");
 	static_assert(std::is_convertible_v<typename UserPolicy::FMulticastDelegateExtras*, TMulticastDelegateBase<UserPolicy>*> || std::is_convertible_v<typename UserPolicy::FMulticastDelegateExtras*, TTSMulticastDelegateBase<UserPolicy>*>, "UserPolicy::FMulticastDelegateExtras should publicly inherit TMulticastDelegateBase<UserPolicy> or TTSMulticastDelegateBase<UserPolicy>");
 
@@ -316,7 +314,7 @@ public:
 		if (&Other != this)
 		{
 			// this down-cast is OK! allows for managing invocation list in the base class without requiring virtual functions
-			DelegateInstanceInterfaceType* OtherInstance = Other.GetDelegateInstanceProtected();
+			const DelegateInstanceInterfaceType* OtherInstance = Other.GetDelegateInstanceProtected();
 
 			if (OtherInstance != nullptr)
 			{
@@ -525,7 +523,7 @@ public:
 	 */
 	FORCEINLINE RetValType Execute(ParamTypes... Params) const
 	{
-		DelegateInstanceInterfaceType* LocalDelegateInstance = GetDelegateInstanceProtected();
+		const DelegateInstanceInterfaceType* LocalDelegateInstance = GetDelegateInstanceProtected();
 
 		// If this assert goes off, Execute() was called before a function was bound to the delegate.
 		// Consider using ExecuteIfBound() instead.
@@ -548,7 +546,7 @@ public:
 	>
 	inline bool ExecuteIfBound(ParamTypes... Params) const
 	{
-		if (DelegateInstanceInterfaceType* Ptr = GetDelegateInstanceProtected())
+		if (const DelegateInstanceInterfaceType* Ptr = GetDelegateInstanceProtected())
 		{
 			return Ptr->ExecuteIfSafe(Params...);
 		}
@@ -560,10 +558,14 @@ protected:
 	/**
 	 * Returns a pointer to the correctly-typed delegate instance.
 	 */
-	FORCEINLINE DelegateInstanceInterfaceType* GetDelegateInstanceProtected() const
+	FORCEINLINE DelegateInstanceInterfaceType* GetDelegateInstanceProtected()
 	{
 		return static_cast<DelegateInstanceInterfaceType*>(Super::GetDelegateInstanceProtected());
 	}
+	FORCEINLINE const DelegateInstanceInterfaceType* GetDelegateInstanceProtected() const
+	{
+		return static_cast<const DelegateInstanceInterfaceType*>(Super::GetDelegateInstanceProtected());
+	}
 };
 
 /**
diff --git a/Engine/Source/Runtime/Core/Public/Delegates/MulticastDelegateBase.h b/Engine/Source/Runtime/Core/Public/Delegates/MulticastDelegateBase.h
--- a/Engine/Source/Runtime/Core/Public/Delegates/MulticastDelegateBase.h
+++ b/Engine/Source/Runtime/Core/Public/Delegates/MulticastDelegateBase.h
@@ -70,7 +70,7 @@ public:
 
 		for (const TDelegateBase<UserPolicy>& DelegateBaseRef : InvocationList)
 		{
-			IDelegateInstance* DelegateInstance = DelegateBaseRef.GetDelegateInstanceProtected();
+			const IDelegateInstance* DelegateInstance = DelegateBaseRef.GetDelegateInstanceProtected();
 			if ((DelegateInstance != nullptr) && DelegateInstance->HasSameObject(InUserObject))
 			{
 				return true;
@@ -174,10 +174,10 @@ protected:
 
 		for (const TDelegateBase<UserPolicy>& OtherDelegateRef : Other.GetInvocationList())
 		{
-			if (IDelegateInstance* OtherInstance = GetDelegateInstanceProtectedHelper(OtherDelegateRef))
+			if (const IDelegateInstance* OtherInstance = GetDelegateInstanceProtectedHelper(OtherDelegateRef))
 			{
 				DelegateType TempDelegate;
-				static_cast<DelegateInstanceInterfaceType*>(OtherInstance)->CreateCopy(TempDelegate);
+				static_cast<const DelegateInstanceInterfaceType*>(OtherInstance)->CreateCopy(TempDelegate);
 				AddDelegateInstance(MoveTemp(TempDelegate));
 			}
 		}
@@ -200,8 +200,8 @@ protected:
 				// this down-cast is OK! allows for managing invocation list in the base class without requiring virtual functions
 				const DelegateBaseType& DelegateBase = static_cast<const DelegateBaseType&>(LocalInvocationList[InvocationListIndex]);
 
-				IDelegateInstance* DelegateInstanceInterface = GetDelegateInstanceProtectedHelper(DelegateBase);
-				if (DelegateInstanceInterface == nullptr || !static_cast<DelegateInstanceInterfaceType*>(DelegateInstanceInterface)->ExecuteIfSafe(Params...))
+				const IDelegateInstance* DelegateInstanceInterface = GetDelegateInstanceProtectedHelper(DelegateBase);
+				if (DelegateInstanceInterface == nullptr || !static_cast<const DelegateInstanceInterfaceType*>(DelegateInstanceInterface)->ExecuteIfSafe(Params...))
 				{
 					NeedsCompaction = true;
 				}
