I/O - added support for loading ondemand utoc files from CbObject

diff --git a/Engine/Source/Runtime/Core/Internal/IO/IoStore.h b/Engine/Source/Runtime/Core/Internal/IO/IoStore.h
--- a/Engine/Source/Runtime/Core/Internal/IO/IoStore.h
+++ b/Engine/Source/Runtime/Core/Internal/IO/IoStore.h
@@ -16,6 +16,7 @@
 CORE_API DECLARE_LOG_CATEGORY_EXTERN(LogIoStore, Log, All);
 
 class FCbWriter;
+class FCbFieldView;
 
 /**
  * I/O store container format version
@@ -277,6 +278,8 @@ struct FIoStoreOnDemandTocHeader
 	CORE_API friend FCbWriter& operator<<(FCbWriter& Writer, const FIoStoreOnDemandTocHeader& Header);
 };
 
+CORE_API bool LoadFromCompactBinary(FCbFieldView Field, FIoStoreOnDemandTocHeader& OutTocHeader);
+
 struct FIoStoreOnDemandTocEntry
 {
 	FIoHash Hash = FIoHash::Zero;
@@ -289,6 +292,8 @@ struct FIoStoreOnDemandTocEntry
 	CORE_API friend FCbWriter& operator<<(FCbWriter& Writer, const FIoStoreOnDemandTocEntry& Entry);
 };
 
+CORE_API bool LoadFromCompactBinary(FCbFieldView Field, FIoStoreOnDemandTocEntry& OutTocEntry);
+
 struct FIoStoreOnDemandContainerEntry
 {
 	FString ContainerName;
@@ -299,6 +304,8 @@ struct FIoStoreOnDemandContainerEntry
 	CORE_API friend FCbWriter& operator<<(FCbWriter& Writer, const FIoStoreOnDemandContainerEntry& ContainerEntry);
 };
 
+CORE_API bool LoadFromCompactBinary(FCbFieldView Field, FIoStoreOnDemandContainerEntry& OutContainer);
+
 struct FIoStoreOndemandTocResource
 {
 	FIoStoreOnDemandTocHeader Header;
@@ -308,3 +315,5 @@ struct FIoStoreOndemandTocResource
 
 	UE_NODISCARD CORE_API static TIoStatusOr<FString> Save(const TCHAR* Directory, const FIoStoreOndemandTocResource& TocResource);
 };
+
+CORE_API bool LoadFromCompactBinary(FCbFieldView Field, FIoStoreOndemandTocResource& OutToc);
diff --git a/Engine/Source/Runtime/Core/Private/IO/IoStore.cpp b/Engine/Source/Runtime/Core/Private/IO/IoStore.cpp
--- a/Engine/Source/Runtime/Core/Private/IO/IoStore.cpp
+++ b/Engine/Source/Runtime/Core/Private/IO/IoStore.cpp
@@ -3244,7 +3244,7 @@ void FIoStoreReader::GetFilenamesByBlockIndex(const TArray<int32>& InBlockIndexL
 		});
 }
 
-//////////////////////////////////////////////////////////////////////////
+////////////////////////////////////////////////////////////////////////////////
 
 FCbWriter& operator<<(FCbWriter& Writer, const FIoStoreOnDemandTocHeader& Header)
 {
@@ -3259,6 +3259,23 @@ FCbWriter& operator<<(FCbWriter& Writer, const FIoStoreOnDemandTocHeader& Header
 	return Writer;
 }
 
+bool LoadFromCompactBinary(FCbFieldView Field, FIoStoreOnDemandTocHeader& OutTocHeader)
+{
+	if (FCbObjectView Obj = Field.AsObjectView())
+	{
+		OutTocHeader.Magic = Obj["Magic"].AsUInt64();
+		OutTocHeader.Version = Obj["Version"].AsUInt32();
+		OutTocHeader.BlockSize = Obj["BlockSize"].AsUInt32();
+		OutTocHeader.CompressionFormat = FString(Obj["CompressionFormat"].AsString());
+		OutTocHeader.ChunksDirectory = FString(Obj["ChunksDirectory"].AsString());
+
+		return OutTocHeader.Magic == FIoStoreOnDemandTocHeader::ExpectedMagic &&
+			static_cast<EIoOnDemandTocVersion>(OutTocHeader.Version) != EIoOnDemandTocVersion::Invalid;
+	}
+
+	return false;
+}
+
 FCbWriter& operator<<(FCbWriter& Writer, const FIoStoreOnDemandTocEntry& Entry)
 {
 	Writer.BeginObject();
@@ -3273,6 +3290,30 @@ FCbWriter& operator<<(FCbWriter& Writer, const FIoStoreOnDemandTocEntry& Entry)
 	return Writer;
 }
 
+bool LoadFromCompactBinary(FCbFieldView Field, FIoStoreOnDemandTocEntry& OutTocEntry)
+{
+	if (FCbObjectView Obj = Field.AsObjectView())
+	{
+		if (!LoadFromCompactBinary(Obj["ChunkId"], OutTocEntry.ChunkId))
+		{
+			return false;
+		}
+
+		OutTocEntry.Hash = Obj["Hash"].AsHash();
+		OutTocEntry.RawSize = Obj["RawSize"].AsUInt64(~uint64(0));
+		OutTocEntry.EncodedSize = Obj["EncodedSize"].AsUInt64(~uint64(0));
+		OutTocEntry.BlockOffset = Obj["BlockOffset"].AsUInt32(~uint32(0));
+		OutTocEntry.BlockCount = Obj["BlockCount"].AsUInt32();
+
+		return OutTocEntry.Hash != FIoHash::Zero &&
+			OutTocEntry.RawSize != ~uint64(0) &&
+			OutTocEntry.EncodedSize != ~uint64(0) &&
+			OutTocEntry.BlockOffset != ~uint32(0);
+	}
+
+	return false;
+}
+
 FCbWriter& operator<<(FCbWriter& Writer, const FIoStoreOnDemandContainerEntry& ContainerEntry)
 {
 	Writer.BeginObject();
@@ -3298,6 +3339,36 @@ FCbWriter& operator<<(FCbWriter& Writer, const FIoStoreOnDemandContainerEntry& C
 	return Writer;
 }
 
+bool LoadFromCompactBinary(FCbFieldView Field, FIoStoreOnDemandContainerEntry& OutContainer)
+{
+	if (FCbObjectView Obj = Field.AsObjectView())
+	{
+		OutContainer.ContainerName = FString(Obj["Name"].AsString());
+		OutContainer.EncryptionKeyGuid = FString(Obj["EncryptionKeyGuid"].AsString());
+
+		FCbArrayView Entries = Obj["Entries"].AsArrayView();
+		OutContainer.Entries.Reserve(int32(Entries.Num()));
+		for (FCbFieldView ArrayField : Entries)
+		{
+			if (!LoadFromCompactBinary(ArrayField, OutContainer.Entries.AddDefaulted_GetRef()))
+			{
+				return false;
+			}
+		}
+
+		FCbArrayView BlockSizes = Obj["BlockSizes"].AsArrayView();
+		OutContainer.BlockSizes.Reserve(int32(BlockSizes.Num()));
+		for (FCbFieldView ArrayField : BlockSizes)
+		{
+			OutContainer.BlockSizes.Add(ArrayField.AsUInt32());
+		}
+
+		return true;
+	}
+
+	return false;
+}
+
 FCbWriter& operator<<(FCbWriter& Writer, const FIoStoreOndemandTocResource& TocResource)
 {
 	Writer.BeginObject();
@@ -3347,3 +3418,28 @@ TIoStatusOr<FString> FIoStoreOndemandTocResource::Save(const TCHAR* Directory, c
 
 	return FIoStatus(EIoErrorCode::WriteError);
 }
+
+bool LoadFromCompactBinary(FCbFieldView Field, FIoStoreOndemandTocResource& OutToc)
+{
+	if (FCbObjectView Obj = Field.AsObjectView())
+	{
+		if (!LoadFromCompactBinary(Obj["Header"], OutToc.Header))
+		{
+			return false;
+		}
+
+		FCbArrayView Containers = Obj["Containers"].AsArrayView();
+		OutToc.Containers.Reserve(int32(Containers.Num()));
+		for (FCbFieldView ArrayField : Containers)
+		{
+			if (!LoadFromCompactBinary(ArrayField, OutToc.Containers.AddDefaulted_GetRef()))
+			{
+				return false;
+			}
+		}
+
+		return true;
+	}
+
+	return false;
+}
