Removed unused forward declarations of FWeakObjectPtr.

diff --git a/Engine/Source/Runtime/Core/Public/Delegates/DelegateBase.h b/Engine/Source/Runtime/Core/Public/Delegates/DelegateBase.h
--- a/Engine/Source/Runtime/Core/Public/Delegates/DelegateBase.h
+++ b/Engine/Source/Runtime/Core/Public/Delegates/DelegateBase.h
@@ -10,8 +10,6 @@
 #include "Delegates/IDelegateInstance.h"
 #include "Delegates/DelegateMacros.h"
 
-struct FWeakObjectPtr;
-
 #if !defined(_WIN32) || defined(_WIN64) || (defined(ALLOW_DELEGATE_INLINE_ALLOCATORS_ON_WIN32) && ALLOW_DELEGATE_INLINE_ALLOCATORS_ON_WIN32)
 	typedef TAlignedBytes<16, 16> FAlignedInlineDelegateType;
 	#if !defined(NUM_DELEGATE_INLINE_BYTES) || NUM_DELEGATE_INLINE_BYTES == 0
@@ -28,8 +26,6 @@ struct FWeakObjectPtr;
 	typedef FHeapAllocator FDelegateAllocatorType;
 #endif
 
-struct FWeakObjectPtr;
-
 template <typename UserPolicy>
 class TMulticastDelegateBase;
 
