Workaround for VS preview ICE

#rnx

diff --git a/Engine/Source/Runtime/Core/Public/UObject/WeakObjectPtrTemplates.h b/Engine/Source/Runtime/Core/Public/UObject/WeakObjectPtrTemplates.h
--- a/Engine/Source/Runtime/Core/Public/UObject/WeakObjectPtrTemplates.h
+++ b/Engine/Source/Runtime/Core/Public/UObject/WeakObjectPtrTemplates.h
@@ -49,7 +49,7 @@ public:
 	 */
 	template <
 		typename U,
-		decltype(ImplicitConv<T*>(std::declval<U>()))* = nullptr
+		typename = decltype(ImplicitConv<T*>(std::declval<U>()))
 	>
 	FORCEINLINE TWeakObjectPtr(U Object) :
 		TWeakObjectPtrBase((const UObject*)Object)
