Translated tests from Engine\Source\Runtime\Core\Private\Tests\Misc

diff --git a/Engine/Source/Developer/LowLevelTestsRunner/Public/TestCommon/Comparisons.h b/Engine/Source/Developer/LowLevelTestsRunner/Public/TestCommon/Comparisons.h
new file mode 100644
--- /dev/null
+++ b/Engine/Source/Developer/LowLevelTestsRunner/Public/TestCommon/Comparisons.h
@@ -0,0 +1,247 @@
+// Copyright Epic Games, Inc. All Rights Reserved.
+
+#pragma once
+
+#include "CoreTypes.h"
+#include "Misc/CString.h"
+#include "Math/Color.h"
+#include "Math/MathFwd.h"
+#include "Math/Rotator.h"
+#include "Math/Transform.h"
+#include "Math/UnrealMathUtility.h"
+#include "Math/Vector.h"
+#include "Containers/UnrealString.h"
+#include "Containers/StringFwd.h"
+
+#include "../TestHarness.h"
+
+namespace {
+
+	bool TestEqual(const TCHAR* What, const int32 Actual, const int32 Expected)
+	{
+		if (Actual != Expected)
+		{
+			FAIL_CHECK(FString::Printf(TEXT("Expected '%s' to be %d, but it was %d."), What, Expected, Actual));
+			return false;
+		}
+		return true;
+	}
+
+	bool TestEqual(const TCHAR* What, const int64 Actual, const int64 Expected)
+	{
+		if (Actual != Expected)
+		{
+			FAIL_CHECK(FString::Printf(TEXT("Expected '%s' to be %" "lld" ", but it was %" "lld" "."), What, Expected, Actual));
+			return false;
+		}
+		return true;
+	}
+
+#if PLATFORM_64BITS
+	bool TestEqual(const TCHAR* What, const SIZE_T Actual, const SIZE_T Expected)
+	{
+		if (Actual != Expected)
+		{
+			FAIL_CHECK(FString::Printf(TEXT("Expected '%s' to be %" "Iu" ", but it was %" "Iu" "."), What, Expected, Actual));
+			return false;
+		}
+		return true;
+	}
+#endif
+
+	bool TestEqual(const TCHAR* What, const float Actual, const float Expected, float Tolerance)
+	{
+		if (!FMath::IsNearlyEqual(Actual, Expected, Tolerance))
+		{
+			FAIL_CHECK(FString::Printf(TEXT("Expected '%s' to be %f, but it was %f within tolerance %f."), What, Expected, Actual, Tolerance));
+			return false;
+		}
+		return true;
+	}
+
+	bool TestEqual(const TCHAR* What, const double Actual, const double Expected, double Tolerance)
+	{
+		if (!FMath::IsNearlyEqual(Actual, Expected, Tolerance))
+		{
+			FAIL_CHECK(FString::Printf(TEXT("Expected '%s' to be %f, but it was %f within tolerance %f."), What, Expected, Actual, Tolerance));
+			return false;
+		}
+		return true;
+	}
+
+	bool TestEqual(const TCHAR* What, const FVector Actual, const FVector Expected, float Tolerance)
+	{
+		if (!Expected.Equals(Actual, Tolerance))
+		{
+			FAIL_CHECK(FString::Printf(TEXT("Expected '%s' to be %s, but it was %s within tolerance %f."), What, *Expected.ToString(), *Actual.ToString(), Tolerance));
+			return false;
+		}
+		return true;
+	}
+
+	bool TestEqual(const TCHAR* What, const FTransform Actual, const FTransform Expected, float Tolerance)
+	{
+		if (!Expected.Equals(Actual, Tolerance))
+		{
+			FAIL_CHECK(FString::Printf(TEXT("Expected '%s' to be %s, but it was %s within tolerance %f."), What, *Expected.ToString(), *Actual.ToString(), Tolerance));
+			return false;
+		}
+		return true;
+	}
+
+	bool TestEqual(const TCHAR* What, const FRotator Actual, const FRotator Expected, float Tolerance)
+	{
+		if (!Expected.Equals(Actual, Tolerance))
+		{
+			FAIL_CHECK(FString::Printf(TEXT("Expected '%s' to be %s, but it was %s within tolerance %f."), What, *Expected.ToString(), *Actual.ToString(), Tolerance));
+			return false;
+		}
+		return true;
+	}
+
+	bool TestEqual(const TCHAR* What, const FColor Actual, const FColor Expected)
+	{
+		if (Expected != Actual)
+		{
+			FAIL_CHECK(FString::Printf(TEXT("Expected '%s' to be %s, but it was %s."), What, *Expected.ToString(), *Actual.ToString()));
+			return false;
+		}
+		return true;
+	}
+
+	bool TestEqual(const TCHAR* What, const FLinearColor Actual, const FLinearColor Expected)
+	{
+		if (Expected != Actual)
+		{
+			FAIL_CHECK(FString::Printf(TEXT("Expected '%s' to be %s, but it was %s."), What, *Expected.ToString(), *Actual.ToString()));
+			return false;
+		}
+		return true;
+	}
+
+	bool TestEqual(const TCHAR* What, const TCHAR* Actual, const TCHAR* Expected)
+	{
+		if (FCString::Strcmp(Actual, Expected) != 0)
+		{
+			FAIL_CHECK(FString::Printf(TEXT("Expected '%s' to be \"%s\", but it was \"%s\"."), What, Expected, Actual));
+			return false;
+		}
+		return true;
+	}
+
+	bool TestEqualInsensitive(const TCHAR* What, const TCHAR* Actual, const TCHAR* Expected)
+	{
+		if (FCString::Stricmp(Actual, Expected) != 0)
+		{
+			FAIL_CHECK(FString::Printf(TEXT("Expected '%s' to be \"%s\", but it was \"%s\"."), What, Expected, Actual));
+			return false;
+		}
+		return true;
+	}
+
+	bool TestEqual(const FString& What, const int32 Actual, const int32 Expected)
+	{
+		return TestEqual(*What, Actual, Expected);
+	}
+
+	bool TestEqual(const FString& What, const float Actual, const float Expected, float Tolerance = UE_KINDA_SMALL_NUMBER)
+	{
+		return TestEqual(*What, Actual, Expected, Tolerance);
+	}
+
+	bool TestEqual(const FString& What, const double Actual, const double Expected, double Tolerance = UE_KINDA_SMALL_NUMBER)
+	{
+		return TestEqual(*What, Actual, Expected, Tolerance);
+	}
+
+	bool TestEqual(const FString& What, const FVector Actual, const FVector Expected, float Tolerance = UE_KINDA_SMALL_NUMBER)
+	{
+		return TestEqual(*What, Actual, Expected, Tolerance);
+	}
+
+	bool TestEqual(const FString& What, const FTransform Actual, const FTransform Expected, float Tolerance = UE_KINDA_SMALL_NUMBER)
+	{
+		return TestEqual(*What, Actual, Expected, Tolerance);
+	}
+
+	bool TestEqual(const FString& What, const FRotator Actual, const FRotator Expected, float Tolerance = UE_KINDA_SMALL_NUMBER)
+	{
+		return TestEqual(*What, Actual, Expected, Tolerance);
+	}
+
+	bool TestEqual(const FString& What, const FColor Actual, const FColor Expected)
+	{
+		return TestEqual(*What, Actual, Expected);
+	}
+
+	bool TestEqual(const FString& What, const TCHAR* Actual, const TCHAR* Expected)
+	{
+		return TestEqual(*What, Actual, Expected);
+	}
+
+	bool TestEqual(const TCHAR* What, const FString& Actual, const TCHAR* Expected)
+	{
+		return TestEqualInsensitive(What, *Actual, Expected);
+	}
+
+	bool TestEqual(const FString& What, const FString& Actual, const TCHAR* Expected)
+	{
+		return TestEqualInsensitive(*What, *Actual, Expected);
+	}
+
+	bool TestEqual(const TCHAR* What, const TCHAR* Actual, const FString& Expected)
+	{
+		return TestEqualInsensitive(What, Actual, *Expected);
+	}
+
+	bool TestEqual(const FString& What, const TCHAR* Actual, const FString& Expected)
+	{
+		return TestEqualInsensitive(*What, Actual, *Expected);
+	}
+
+	bool TestEqual(const TCHAR* What, const FString& Actual, const FString& Expected)
+	{
+		return TestEqualInsensitive(What, *Actual, *Expected);
+	}
+
+	bool TestEqual(const FString& What, const FString& Actual, const FString& Expected)
+	{
+		return TestEqualInsensitive(*What, *Actual, *Expected);
+	}
+
+	template<typename ValueType>
+	bool TestEqual(const TCHAR* What, const ValueType& Actual, const ValueType& Expected)
+	{
+		if (Actual != Expected)
+		{
+			FAIL_CHECK(FString::Printf(TEXT("%s: The two values are not equal."), What));
+			return false;
+		}
+		return true;
+	}
+
+	template<typename ValueType>
+	bool TestEqual(const FString& What, const ValueType& Actual, const ValueType& Expected)
+	{
+		return TestEqual(*What, Actual, Expected);
+	}
+
+	template<typename ValueType> bool TestNotEqual(const TCHAR* Description, const ValueType& Actual, const ValueType& Expected)
+	{
+		if (Actual == Expected)
+		{
+			FAIL_CHECK(FString::Printf(TEXT("%s: The two values are equal."), Description));
+			return false;
+		}
+		return true;
+	}
+
+	template<typename ValueType> bool TestNotEqual(const FString& Description, const ValueType& Actual, const ValueType& Expected)
+	{
+		return TestNotEqual(*Description, Actual, Expected);
+	}
+
+	#define CHECK_EQUALS(What, X, Y) TestEqual(What, X, Y);
+	#define CHECK_NOT_EQUALS(What, X, Y) TestNotEqual(What, X, Y);
+
+}
\ No newline at end of file
diff --git a/Engine/Source/Runtime/Core/Private/Tests/Misc/AsciiSetTest.cpp b/Engine/Source/Runtime/Core/Private/Tests/Misc/AsciiSetTest.cpp
deleted file mode 100644
diff --git a/Engine/Source/Runtime/Core/Private/Tests/Misc/AtomicTest.cpp b/Engine/Source/Runtime/Core/Private/Tests/Misc/AtomicTest.cpp
deleted file mode 100644
diff --git a/Engine/Source/Runtime/Core/Private/Tests/Misc/ByteSwapTest.cpp b/Engine/Source/Runtime/Core/Private/Tests/Misc/ByteSwapTest.cpp
deleted file mode 100644
diff --git a/Engine/Source/Runtime/Core/Private/Tests/Misc/CharTest.cpp b/Engine/Source/Runtime/Core/Private/Tests/Misc/CharTest.cpp
deleted file mode 100644
diff --git a/Engine/Source/Runtime/Core/Private/Tests/Misc/CircularBufferTest.cpp b/Engine/Source/Runtime/Core/Private/Tests/Misc/CircularBufferTest.cpp
deleted file mode 100644
diff --git a/Engine/Source/Runtime/Core/Private/Tests/Misc/CircularQueueTest.cpp b/Engine/Source/Runtime/Core/Private/Tests/Misc/CircularQueueTest.cpp
deleted file mode 100644
diff --git a/Engine/Source/Runtime/Core/Private/Tests/Misc/DateTimeTest.cpp b/Engine/Source/Runtime/Core/Private/Tests/Misc/DateTimeTest.cpp
deleted file mode 100644
diff --git a/Engine/Source/Runtime/Core/Private/Tests/Misc/DisplayNameTest.cpp b/Engine/Source/Runtime/Core/Private/Tests/Misc/DisplayNameTest.cpp
deleted file mode 100644
diff --git a/Engine/Source/Runtime/Core/Private/Tests/Misc/FileTest.cpp b/Engine/Source/Runtime/Core/Private/Tests/Misc/FileTest.cpp
deleted file mode 100644
diff --git a/Engine/Source/Runtime/Core/Private/Tests/Misc/GuidTest.cpp b/Engine/Source/Runtime/Core/Private/Tests/Misc/GuidTest.cpp
deleted file mode 100644
diff --git a/Engine/Source/Runtime/Core/Private/Tests/Misc/PathTests.h b/Engine/Source/Runtime/Core/Private/Tests/Misc/PathTests.h
deleted file mode 100644
diff --git a/Engine/Source/Runtime/Core/Private/Tests/Misc/PathViewsTest.cpp b/Engine/Source/Runtime/Core/Private/Tests/Misc/PathViewsTest.cpp
deleted file mode 100644
diff --git a/Engine/Source/Runtime/Core/Private/Tests/Misc/PathsTest.cpp b/Engine/Source/Runtime/Core/Private/Tests/Misc/PathsTest.cpp
deleted file mode 100644
diff --git a/Engine/Source/Runtime/Core/Private/Tests/Misc/PlatformStringConvertTests.cpp b/Engine/Source/Runtime/Core/Private/Tests/Misc/PlatformStringConvertTests.cpp
deleted file mode 100644
diff --git a/Engine/Source/Runtime/Core/Private/Tests/Misc/QueueTest.cpp b/Engine/Source/Runtime/Core/Private/Tests/Misc/QueueTest.cpp
deleted file mode 100644
diff --git a/Engine/Source/Runtime/Core/Private/Tests/Misc/TCharToUTF8.cpp b/Engine/Source/Runtime/Core/Private/Tests/Misc/TCharToUTF8.cpp
deleted file mode 100644
diff --git a/Engine/Source/Runtime/Core/Private/Tests/Misc/TVariantTest.cpp b/Engine/Source/Runtime/Core/Private/Tests/Misc/TVariantTest.cpp
deleted file mode 100644
diff --git a/Engine/Source/Runtime/Core/Private/Tests/Misc/TimecodeTest.cpp b/Engine/Source/Runtime/Core/Private/Tests/Misc/TimecodeTest.cpp
deleted file mode 100644
diff --git a/Engine/Source/Runtime/Core/Private/Tests/Misc/TripleBufferTest.cpp b/Engine/Source/Runtime/Core/Private/Tests/Misc/TripleBufferTest.cpp
deleted file mode 100644
diff --git a/Engine/Source/Runtime/Core/Private/Tests/Misc/TypeContainerTest.cpp b/Engine/Source/Runtime/Core/Private/Tests/Misc/TypeContainerTest.cpp
deleted file mode 100644
diff --git a/Engine/Source/Runtime/Core/Public/Misc/LowLevelTestAdapter.h b/Engine/Source/Runtime/Core/Public/Misc/LowLevelTestAdapter.h
--- a/Engine/Source/Runtime/Core/Public/Misc/LowLevelTestAdapter.h
+++ b/Engine/Source/Runtime/Core/Public/Misc/LowLevelTestAdapter.h
@@ -125,6 +125,8 @@ FString CaptureExpressionsAndValues(const FString& InExpressions, ArgTypes&&...
 #define REQUIRE_MESSAGE(Message, Expr) if (!(Expr)) { FAutomationTestFramework::Get().GetCurrentTest()->AddError(Message); return; }
 #define STATIC_REQUIRE(...) static_assert(__VA_ARGS__, #__VA_ARGS__);
 #define STATIC_CHECK(...) static_assert(__VA_ARGS__, #__VA_ARGS__);
+#define CHECK_EQUALS(What, X, Y) FAutomationTestFramework::Get().GetCurrentTest()->TestEqual(What, X, Y);
+#define CHECK_NOT_EQUALS(What, X, Y) FAutomationTestFramework::Get().GetCurrentTest()->TestNotEqual(What, X, Y);
 
 #define SECTION(Text) FAutomationTestFramework::Get().GetCurrentTest()->AddInfo(TEXT(Text));
 #define FAIL_CHECK(Message) FAutomationTestFramework::Get().GetCurrentTest()->AddError(Message);
diff --git a/Engine/Source/Runtime/Core/Tests/Misc/AsciiSetTest.cpp b/Engine/Source/Runtime/Core/Tests/Misc/AsciiSetTest.cpp
new file mode 100644
--- /dev/null
+++ b/Engine/Source/Runtime/Core/Tests/Misc/AsciiSetTest.cpp
@@ -0,0 +1,96 @@
+// Copyright Epic Games, Inc. All Rights Reserved.
+#if WITH_TESTS
+
+#include "Containers/StringView.h"
+#include "Misc/AsciiSet.h"
+#include "Tests/TestHarnessAdapter.h"
+#if WITH_LOW_LEVEL_TESTS
+#include "TestCommon/Comparisons.h"
+#endif
+
+TEST_CASE_NAMED(TAsciiSetTest, "System::Core::Misc::AsciiSet", "[ApplicationContextMask][SmokeFilter]")
+{
+	constexpr FAsciiSet Whitespaces(" \v\f\t\r\n");
+	CHECK_MESSAGE(TEXT("Contains"), Whitespaces.Contains(' '));
+	CHECK_MESSAGE(TEXT("Contains"), Whitespaces.Contains('\n'));
+	CHECK_FALSE_MESSAGE(TEXT("Contains"), Whitespaces.Contains('a'));
+	CHECK_FALSE_MESSAGE(TEXT("Contains no extended ASCII"), Whitespaces.Contains('\x80'));
+	CHECK_FALSE_MESSAGE(TEXT("Contains no extended ASCII"), Whitespaces.Contains('\xA0'));
+	CHECK_FALSE_MESSAGE(TEXT("Contains no extended ASCII"), Whitespaces.Contains('\xFF'));
+	
+	constexpr FAsciiSet Aa("Aa");
+	uint32 ANum = 0;
+	for (char32_t C = 0; C < 512; ++C)
+	{
+		ANum += Aa.Contains(C);
+	}
+	CHECK_EQUALS(TEXT("Contains no wide"), ANum, 2u);
+
+	constexpr FAsciiSet NonWhitespaces = ~Whitespaces;
+	uint32 WhitespaceNum = 0;
+	for (uint8 Char = 0; Char < 128; ++Char)
+	{
+		WhitespaceNum += !!Whitespaces.Test(Char);
+		CHECK_EQUALS(TEXT("Inverse"), !!Whitespaces.Test(Char), !NonWhitespaces.Test(Char));
+	}
+	CHECK_EQUALS(TEXT("Num"), WhitespaceNum, 6);
+
+	CHECK_EQUALS(TEXT("Skip"), FAsciiSet::Skip(TEXT("  \t\tHello world!"), Whitespaces), TEXT("Hello world!"));
+	CHECK_EQUALS(TEXT("Skip"), FAsciiSet::Skip(TEXT("Hello world!"), Whitespaces), TEXT("Hello world!"));
+	CHECK_EQUALS(TEXT("Skip to extended ASCII"), FAsciiSet::Skip(" " "\xA0" " abc", Whitespaces), "\xA0" " abc");
+	CHECK_EQUALS(TEXT("Skip to wide"), FAsciiSet::Skip(TEXT(" 变 abc"), Whitespaces), TEXT("变 abc"));
+	CHECK_EQUALS(TEXT("AdvanceToFirst"),	*FAsciiSet::FindFirstOrEnd("NonWhitespace\t \nNonWhitespace", Whitespaces), '\t');
+	CHECK_EQUALS(TEXT("AdvanceToLast"),	*FAsciiSet::FindLastOrEnd("NonWhitespace\t \nNonWhitespace", Whitespaces), '\n');
+	CHECK_EQUALS(TEXT("AdvanceToLast"),	*FAsciiSet::FindLastOrEnd("NonWhitespace\t NonWhitespace\n", Whitespaces), '\n');
+	CHECK_EQUALS(TEXT("AdvanceToFirst"),	*FAsciiSet::FindFirstOrEnd("NonWhitespaceNonWhitespace", Whitespaces), '\0');
+	CHECK_EQUALS(TEXT("AdvanceToLast"),	*FAsciiSet::FindLastOrEnd("NonWhitespaceNonWhitespace", Whitespaces), '\0');
+
+	constexpr FAsciiSet Lowercase("abcdefghijklmnopqrstuvwxyz");
+	CHECK_EQUALS(TEXT("TrimPrefixWithout"),		FAsciiSet::TrimPrefixWithout(ANSITEXTVIEW("ABcdEF"), Lowercase), ANSITEXTVIEW("cdEF"));
+	CHECK_EQUALS(TEXT("FindPrefixWithout"),		FAsciiSet::FindPrefixWithout(ANSITEXTVIEW("ABcdEF"), Lowercase), ANSITEXTVIEW("AB"));
+	CHECK_EQUALS(TEXT("TrimSuffixWithout"),		FAsciiSet::TrimSuffixWithout(ANSITEXTVIEW("ABcdEF"), Lowercase), ANSITEXTVIEW("ABcd"));
+	CHECK_EQUALS(TEXT("FindSuffixWithout"),		FAsciiSet::FindSuffixWithout(ANSITEXTVIEW("ABcdEF"), Lowercase), ANSITEXTVIEW("EF"));
+	CHECK_EQUALS(TEXT("TrimPrefixWithout none"),	FAsciiSet::TrimPrefixWithout(ANSITEXTVIEW("same"), Lowercase), ANSITEXTVIEW("same"));
+	CHECK_EQUALS(TEXT("FindPrefixWithout none"),	FAsciiSet::FindPrefixWithout(ANSITEXTVIEW("same"), Lowercase), ANSITEXTVIEW(""));
+	CHECK_EQUALS(TEXT("TrimSuffixWithout none"),	FAsciiSet::TrimSuffixWithout(ANSITEXTVIEW("same"), Lowercase), ANSITEXTVIEW("same"));
+	CHECK_EQUALS(TEXT("FindSuffixWithout none"),	FAsciiSet::FindSuffixWithout(ANSITEXTVIEW("same"), Lowercase), ANSITEXTVIEW(""));
+	CHECK_EQUALS(TEXT("TrimPrefixWithout empty"),	FAsciiSet::TrimPrefixWithout(ANSITEXTVIEW(""), Lowercase), ANSITEXTVIEW(""));
+	CHECK_EQUALS(TEXT("FindPrefixWithout empty"),	FAsciiSet::FindPrefixWithout(ANSITEXTVIEW(""), Lowercase), ANSITEXTVIEW(""));
+	CHECK_EQUALS(TEXT("TrimSuffixWithout empty"),	FAsciiSet::TrimSuffixWithout(ANSITEXTVIEW(""), Lowercase), ANSITEXTVIEW(""));
+	CHECK_EQUALS(TEXT("FindSuffixWithout empty"),	FAsciiSet::FindSuffixWithout(ANSITEXTVIEW(""), Lowercase), ANSITEXTVIEW(""));
+
+
+	auto TestHasFunctions = [](auto MakeString)
+	{
+		FAsciiSet XmlEscapeChars("&<>\"'");
+		CHECK_MESSAGE(TEXT("None"),	FAsciiSet::HasNone(MakeString("No escape chars"), XmlEscapeChars));
+		CHECK_FALSE_MESSAGE(TEXT("Any"),	FAsciiSet::HasAny(MakeString("No escape chars"), XmlEscapeChars));
+		CHECK_FALSE_MESSAGE(TEXT("Only"), FAsciiSet::HasOnly(MakeString("No escape chars"), XmlEscapeChars));
+
+		CHECK_MESSAGE(TEXT("None"),	FAsciiSet::HasNone(MakeString(""), XmlEscapeChars));
+		CHECK_FALSE_MESSAGE(TEXT("Any"),	FAsciiSet::HasAny(MakeString(""), XmlEscapeChars));
+		CHECK_MESSAGE(TEXT("Only"),	FAsciiSet::HasOnly(MakeString(""), XmlEscapeChars));
+
+		CHECK_FALSE_MESSAGE(TEXT("None"), FAsciiSet::HasNone(MakeString("&<>\"'"), XmlEscapeChars));
+		CHECK_MESSAGE(TEXT("Any"),	FAsciiSet::HasAny(MakeString("&<>\"'"), XmlEscapeChars));
+		CHECK_MESSAGE(TEXT("Only"),	FAsciiSet::HasOnly(MakeString("&<>\"'"), XmlEscapeChars));
+
+		CHECK_FALSE_MESSAGE(TEXT("None"), FAsciiSet::HasNone(MakeString("&<>\"' and more"), XmlEscapeChars));
+		CHECK_MESSAGE(TEXT("Any"),	FAsciiSet::HasAny(MakeString("&<>\"' and more"), XmlEscapeChars));
+		CHECK_FALSE_MESSAGE(TEXT("Only"), FAsciiSet::HasOnly(MakeString("&<>\"' and more"), XmlEscapeChars));
+	};
+	TestHasFunctions([] (const char* Str) { return Str; });
+	TestHasFunctions([] (const char* Str) { return FAnsiStringView(Str); });
+	TestHasFunctions([] (const char* Str) { return FString(Str); });
+
+
+	constexpr FAsciiSet Abc("abc");
+	constexpr FAsciiSet Abcd = Abc + 'd';
+	CHECK_MESSAGE(TEXT("Add"), Abcd.Contains('a'));
+	CHECK_MESSAGE(TEXT("Add"), Abcd.Contains('b'));
+	CHECK_MESSAGE(TEXT("Add"), Abcd.Contains('c'));
+	CHECK_MESSAGE(TEXT("Add"), Abcd.Contains('d'));
+	CHECK_FALSE_MESSAGE(TEXT("Add"), Abcd.Contains('e'));
+}
+
+#endif //WITH_TESTS
\ No newline at end of file
diff --git a/Engine/Source/Runtime/Core/Tests/Misc/AtomicTest.cpp b/Engine/Source/Runtime/Core/Tests/Misc/AtomicTest.cpp
new file mode 100644
--- /dev/null
+++ b/Engine/Source/Runtime/Core/Tests/Misc/AtomicTest.cpp
@@ -0,0 +1,408 @@
+// Copyright Epic Games, Inc. All Rights Reserved.
+
+#include "CoreTypes.h"
+#include "Misc/AssertionMacros.h"
+#include "Templates/Atomic.h"
+#include "Templates/EnableIf.h"
+#include "Traits/IntType.h"
+
+#if WITH_TESTS
+
+#include "Tests/TestHarnessAdapter.h"
+
+namespace
+{
+	// Returns an element type filled with a repeated byte
+	template <typename ElementType>
+	FORCEINLINE constexpr typename TEnableIf<sizeof(ElementType) != 1, ElementType>::Type GetByteFilledElement(uint8 Byte)
+	{
+		TUnsignedIntType_T<sizeof(ElementType)> Result = 0;
+
+		for (int32 i = 0; i != sizeof(ElementType); ++i)
+		{
+			Result <<= 8;
+			Result |= (TUnsignedIntType_T<sizeof(ElementType)>)Byte;
+		}
+
+		return (const ElementType&)Result;
+	}
+
+	template <typename ElementType>
+	FORCEINLINE constexpr typename TEnableIf<sizeof(ElementType) == 1, ElementType>::Type GetByteFilledElement(uint8 Byte)
+	{
+		return (const ElementType&)Byte;
+	}
+
+	// Returns an element type filled with random bytes
+	template <typename ElementType>
+	FORCEINLINE typename TEnableIf<sizeof(ElementType) != 1, ElementType>::Type GetRandomFilledElement()
+	{
+		TUnsignedIntType_T<sizeof(ElementType)> Result = 0;
+
+		for (int32 i = 0; i != sizeof(ElementType); ++i)
+		{
+			Result <<= 8;
+			Result |= (TUnsignedIntType_T<sizeof(ElementType)>)(uint8)FMath::Rand();
+		}
+
+		return (const ElementType&)Result;
+	}
+
+	template <typename ElementType>
+	FORCEINLINE typename TEnableIf<sizeof(ElementType) == 1, ElementType>::Type GetRandomFilledElement()
+	{
+		return (const ElementType&)(uint8)FMath::Rand();
+	}
+
+	template <typename ElementType>
+	struct TAtomicTestWrapper
+	{
+		static_assert(sizeof(ElementType) == sizeof(TAtomic<ElementType>), "Atomic should be the same size as the underlying type");
+
+		explicit TAtomicTestWrapper(ElementType Init)
+		{
+			ElementType DebugCode = GetByteFilledElement<ElementType>(0xCD);
+
+			NativeArray[0] = DebugCode;
+			NativeArray[1] = Init;
+			NativeArray[2] = DebugCode;
+
+			AtomicArray[0] = DebugCode;
+			AtomicArray[1] = Init;
+			AtomicArray[2] = DebugCode;
+		}
+
+		void Check() const
+		{
+			ElementType DebugCode = GetByteFilledElement<ElementType>(0xCD);
+
+			check(NativeArray[0]        == DebugCode);
+			check(NativeArray[2]        == DebugCode);
+			check(AtomicArray[0].Load() == DebugCode);
+			check(AtomicArray[2].Load() == DebugCode);
+
+			check(NativeArray[1] == AtomicArray[1].Load());
+		}
+
+		ElementType& GetNative()
+		{
+			return NativeArray[1];
+		}
+
+		TAtomic<ElementType>& GetAtomic()
+		{
+			return AtomicArray[1];
+		}
+
+	private:
+		// Arrays with some run-off, where we're modifying the middle element, and we can check the
+		// surrounding elements to see we're not accidentally over-/under-running.
+		ElementType          NativeArray[3];
+		TAtomic<ElementType> AtomicArray[3];
+	};
+
+	template <typename ElementType>
+	void RunBasicAtomicTests()
+	{
+		ElementType Val = GetRandomFilledElement<ElementType>();
+
+		TAtomicTestWrapper<ElementType> Data(Val);
+		Data.Check();
+
+		Val = GetRandomFilledElement<ElementType>();
+		Data.GetNative() = Val;
+		Data.GetAtomic().Store(Val, EMemoryOrder::SequentiallyConsistent);
+		Data.Check();
+
+		Val = GetRandomFilledElement<ElementType>();
+		Data.GetNative() = Val;
+		Data.GetAtomic().Store(Val, EMemoryOrder::Relaxed);
+		Data.Check();
+
+		Val = GetRandomFilledElement<ElementType>();
+		ElementType Old = Data.GetAtomic().Exchange(Val);
+		check(Data.GetNative() == Old);
+		Data.GetNative() = Val;
+		Data.Check();
+	}
+
+	template <typename ElementType>
+	void RunArithmeticAtomicTests(ElementType Init)
+	{
+		typename TRemovePointer<ElementType>::Type Array[100] = {};
+		TAtomicTestWrapper<ElementType> Data(Init);
+		Data.Check();
+
+		ElementType NativeVal = Data.GetNative() += 4;
+		ElementType AtomicVal = Data.GetAtomic() += 4;
+		check(NativeVal == AtomicVal);
+		Data.Check();
+
+		NativeVal = Data.GetNative() += -7;
+		AtomicVal = Data.GetAtomic() += -7;
+		check(NativeVal == AtomicVal);
+		Data.Check();
+
+		NativeVal = Data.GetNative() -= 3;
+		AtomicVal = Data.GetAtomic() -= 3;
+		check(NativeVal == AtomicVal);
+		Data.Check();
+
+		NativeVal = Data.GetNative() -= -8;
+		AtomicVal = Data.GetAtomic() -= -8;
+		check(NativeVal == AtomicVal);
+		Data.Check();
+
+		NativeVal = --Data.GetNative();
+		AtomicVal = --Data.GetAtomic();
+		check(NativeVal == AtomicVal);
+		Data.Check();
+
+		NativeVal = ++Data.GetNative();
+		AtomicVal = ++Data.GetAtomic();
+		check(NativeVal == AtomicVal);
+		Data.Check();
+
+		NativeVal = Data.GetNative()--;
+		AtomicVal = Data.GetAtomic()--;
+		check(NativeVal == AtomicVal);
+		NativeVal = Data.GetNative();
+		AtomicVal = Data.GetAtomic();
+		check(NativeVal == AtomicVal);
+		Data.Check();
+
+		NativeVal = Data.GetNative()++;
+		AtomicVal = Data.GetAtomic()++;
+		check(NativeVal == AtomicVal);
+		NativeVal = Data.GetNative();
+		AtomicVal = Data.GetAtomic();
+		check(NativeVal == AtomicVal);
+		Data.Check();
+
+		NativeVal = Data.GetNative()--;
+		AtomicVal = Data.GetAtomic().DecrementExchange();
+		check(NativeVal == AtomicVal);
+		NativeVal = Data.GetNative();
+		AtomicVal = Data.GetAtomic();
+		check(NativeVal == AtomicVal);
+		Data.Check();
+
+		NativeVal = Data.GetNative()++;
+		AtomicVal = Data.GetAtomic().IncrementExchange();
+		check(NativeVal == AtomicVal);
+		NativeVal = Data.GetNative();
+		AtomicVal = Data.GetAtomic();
+		check(NativeVal == AtomicVal);
+		Data.Check();
+
+		NativeVal = Data.GetNative();
+		AtomicVal = Data.GetAtomic().AddExchange(47);
+		check(NativeVal == AtomicVal);
+		NativeVal = Data.GetNative() += 47;
+		AtomicVal = Data.GetAtomic();
+		check(NativeVal == AtomicVal);
+		Data.Check();
+
+		NativeVal = Data.GetNative();
+		AtomicVal = Data.GetAtomic().AddExchange(-11);
+		check(NativeVal == AtomicVal);
+		NativeVal = Data.GetNative() += -11;
+		AtomicVal = Data.GetAtomic();
+		check(NativeVal == AtomicVal);
+		Data.Check();
+
+		NativeVal = Data.GetNative();
+		AtomicVal = Data.GetAtomic().SubExchange(2);
+		check(NativeVal == AtomicVal);
+		NativeVal = Data.GetNative() -= 2;
+		AtomicVal = Data.GetAtomic();
+		check(NativeVal == AtomicVal);
+		Data.Check();
+
+		NativeVal = Data.GetNative();
+		AtomicVal = Data.GetAtomic().SubExchange(-9);
+		check(NativeVal == AtomicVal);
+		NativeVal = Data.GetNative() -= -9;
+		AtomicVal = Data.GetAtomic();
+		check(NativeVal == AtomicVal);
+		Data.Check();
+	}
+
+	template <typename ElementType>
+	void RunBitwiseOperationsAtomicTests()
+	{
+		ElementType NativeVal;
+		ElementType AtomicVal;
+
+		// And
+		{
+			constexpr ElementType Init = GetByteFilledElement<ElementType>(0x30);
+			constexpr ElementType AndValues[] = {GetByteFilledElement<ElementType>(0x66), GetByteFilledElement<ElementType>(0xFF), ElementType(0), };
+			constexpr ElementType ExpectedValues[] = {GetByteFilledElement<ElementType>(0x20), Init, ElementType(0), };
+
+			// operator &=
+			for (SIZE_T TestIt = 0; TestIt < UE_ARRAY_COUNT(AndValues); ++TestIt)
+			{
+				TAtomicTestWrapper<ElementType> Data(Init);
+
+				AtomicVal = Data.GetAtomic() &= AndValues[TestIt];
+				NativeVal = Data.GetNative() &= AndValues[TestIt];
+				check(NativeVal == AtomicVal);
+
+				check(AtomicVal == ExpectedValues[TestIt]);
+				Data.Check();
+			}
+
+			// AndExchange
+			for (SIZE_T TestIt = 0; TestIt < UE_ARRAY_COUNT(AndValues); ++TestIt)
+			{
+				TAtomicTestWrapper<ElementType> Data(Init);
+
+				AtomicVal = Data.GetAtomic().AndExchange(AndValues[TestIt]);
+				NativeVal = Data.GetNative();
+				check(NativeVal == AtomicVal);
+				NativeVal = Data.GetNative() &= AndValues[TestIt];
+				AtomicVal = Data.GetAtomic().Load();
+				check(NativeVal == AtomicVal);
+
+				check(AtomicVal == ExpectedValues[TestIt]);
+				Data.Check();
+			}
+		}
+
+		// Or
+		{
+			constexpr ElementType Init = GetByteFilledElement<ElementType>(0x30);
+			constexpr ElementType OrValues[] = {GetByteFilledElement<ElementType>(0x66), GetByteFilledElement<ElementType>(0xFF), ElementType(0), };
+			constexpr ElementType ExpectedValues[] = {GetByteFilledElement<ElementType>(0x76), GetByteFilledElement<ElementType>(0xFF), Init, };
+
+			// operator |=
+			for (SIZE_T TestIt = 0; TestIt < UE_ARRAY_COUNT(OrValues); ++TestIt)
+			{
+				TAtomicTestWrapper<ElementType> Data(Init);
+
+				AtomicVal = Data.GetAtomic() |= OrValues[TestIt];
+				NativeVal = Data.GetNative() |= OrValues[TestIt];
+				check(NativeVal == AtomicVal);
+
+				check(AtomicVal == ExpectedValues[TestIt]);
+				Data.Check();
+			}
+
+			// OrExchange
+			for (SIZE_T TestIt = 0; TestIt < UE_ARRAY_COUNT(OrValues); ++TestIt)
+			{
+				TAtomicTestWrapper<ElementType> Data(Init);
+
+				AtomicVal = Data.GetAtomic().OrExchange(OrValues[TestIt]);
+				NativeVal = Data.GetNative();
+				check(NativeVal == AtomicVal);
+				NativeVal = Data.GetNative() |= OrValues[TestIt];
+				AtomicVal = Data.GetAtomic().Load();
+				check(NativeVal == AtomicVal);
+
+				check(AtomicVal == ExpectedValues[TestIt]);
+				Data.Check();
+			}
+		}
+
+		// Xor
+		{
+			constexpr ElementType Init = GetByteFilledElement<ElementType>(0x30);
+			constexpr ElementType XorValues[] = {GetByteFilledElement<ElementType>(0x66), GetByteFilledElement<ElementType>(0xFF), ElementType(0), };
+			constexpr ElementType ExpectedValues[] = {GetByteFilledElement<ElementType>(0x56), GetByteFilledElement<ElementType>(~0x30), Init, };
+
+			// operator ^=
+			for (SIZE_T TestIt = 0; TestIt < UE_ARRAY_COUNT(XorValues); ++TestIt)
+			{
+				TAtomicTestWrapper<ElementType> Data(Init);
+
+				AtomicVal = Data.GetAtomic() ^= XorValues[TestIt];
+				NativeVal = Data.GetNative() ^= XorValues[TestIt];
+				check(NativeVal == AtomicVal);
+
+				check(AtomicVal == ExpectedValues[TestIt]);
+				Data.Check();
+			}
+
+			// XorExchange
+			for (SIZE_T TestIt = 0; TestIt < UE_ARRAY_COUNT(XorValues); ++TestIt)
+			{
+				TAtomicTestWrapper<ElementType> Data(Init);
+
+				AtomicVal = Data.GetAtomic().XorExchange(XorValues[TestIt]);
+				NativeVal = Data.GetNative();
+				check(NativeVal == AtomicVal);
+				NativeVal = Data.GetNative() ^= XorValues[TestIt];
+				AtomicVal = Data.GetAtomic().Load();
+				check(NativeVal == AtomicVal);
+
+				check(AtomicVal == ExpectedValues[TestIt]);
+				Data.Check();
+			}
+		}
+	}
+
+	template <typename ElementType>
+	void RunNumericAtomicTests()
+	{
+		RunBasicAtomicTests<ElementType>();
+
+		RunArithmeticAtomicTests<ElementType>(50);
+
+		RunBitwiseOperationsAtomicTests<ElementType>();
+	}
+
+	template <typename ElementType>
+	void RunPointerAtomicTests()
+	{
+		RunBasicAtomicTests<ElementType>();
+
+		typename TRemovePointer<ElementType>::Type Array[100] = {};
+		RunArithmeticAtomicTests<ElementType>(&Array[50]);
+	}
+}
+
+TEST_CASE_NAMED(FAtomicSmokeTest, "System::Core::Misc::Atomic", "[ApplicationContextMask][SmokeFilter]")
+{
+	RunNumericAtomicTests<int8>  ();
+	RunNumericAtomicTests<uint8> ();
+	RunNumericAtomicTests<int16> ();
+	RunNumericAtomicTests<uint16>();
+	RunNumericAtomicTests<int32> ();
+	RunNumericAtomicTests<uint32>();
+	RunNumericAtomicTests<int64> ();
+	RunNumericAtomicTests<uint64>();
+
+	// Don't run the pointer test on void*, because you can't do arithmetic on void*s.
+	RunBasicAtomicTests<               void*>();
+	RunBasicAtomicTests<const          void*>();
+	RunBasicAtomicTests<      volatile void*>();
+	RunBasicAtomicTests<const volatile void*>();
+
+	RunPointerAtomicTests<int8*>();
+	RunPointerAtomicTests<int16*>();
+	RunPointerAtomicTests<int32*>();
+	RunPointerAtomicTests<int64*>();
+	RunPointerAtomicTests<FString*>();
+
+	RunPointerAtomicTests<const int8*>();
+	RunPointerAtomicTests<const int16*>();
+	RunPointerAtomicTests<const int32*>();
+	RunPointerAtomicTests<const int64*>();
+	RunPointerAtomicTests<const FString*>();
+
+	RunPointerAtomicTests<volatile int8*>();
+	RunPointerAtomicTests<volatile int16*>();
+	RunPointerAtomicTests<volatile int32*>();
+	RunPointerAtomicTests<volatile int64*>();
+	RunPointerAtomicTests<volatile FString*>();
+
+	RunPointerAtomicTests<const volatile int8*>();
+	RunPointerAtomicTests<const volatile int16*>();
+	RunPointerAtomicTests<const volatile int32*>();
+	RunPointerAtomicTests<const volatile int64*>();
+	RunPointerAtomicTests<const volatile FString*>();
+}
+
+#endif // WITH_TESTS
diff --git a/Engine/Source/Runtime/Core/Tests/Misc/ByteSwapTest.cpp b/Engine/Source/Runtime/Core/Tests/Misc/ByteSwapTest.cpp
new file mode 100644
--- /dev/null
+++ b/Engine/Source/Runtime/Core/Tests/Misc/ByteSwapTest.cpp
@@ -0,0 +1,106 @@
+// Copyright Epic Games, Inc. All Rights Reserved.
+
+#include "CoreTypes.h"
+#include "Misc/ByteSwap.h"
+
+#if WITH_TESTS
+
+#include "Tests/TestHarnessAdapter.h"
+
+/**
+ * Test byte swapping algorithms.
+ */
+TEST_CASE_NAMED(FByteSwapTest, "System::Core::Misc::ByteSwap", "[ApplicationContextMask][SmokeFilter]")
+{
+	int16  ValS16 = static_cast<int16>(0x1122);
+	uint16 ValU16 = static_cast<uint16>(0x1122);
+	int32  ValS32 = 0xFFEE1122L;
+	uint32 ValU32 = 0xFFEE1122UL;
+	int64  ValS64 = 0xFFEEDDCC11223344ll;
+	uint64 ValU64 = 0xFFEEDDCC11223344ull;
+	float  ValF   = *reinterpret_cast<const float*>(&ValU32);
+	double ValD   = *reinterpret_cast<const double*>(&ValU64);
+	uint32 ExpectedF = ByteSwap(ValU32);
+	uint64 ExpectedD = ByteSwap(ValU64);
+	char16_t ValCh16 = static_cast<char16_t>(0x2233);
+
+	CHECK_MESSAGE("Swapping singned int16 value",  ByteSwap(ValS16) == 0x2211);
+	CHECK_MESSAGE("Swapping unsigned int16 value", ByteSwap(ValU16) == 0x2211);
+	CHECK_MESSAGE("Swapping singned int32 value",  ByteSwap(ValS32) == 0x2211EEFF);
+	CHECK_MESSAGE("Swapping unsigned int32 value", ByteSwap(ValU32) == 0x2211EEFF);
+	CHECK_MESSAGE("Swapping singned int64 value",  ByteSwap(ValS64) == 0x44332211CCDDEEFFll);
+	CHECK_MESSAGE("Swapping unsigned int64 value", ByteSwap(ValU64) == 0x44332211CCDDEEFFull);
+	CHECK_MESSAGE("Swapping float value",          ByteSwap(ValF)   == *reinterpret_cast<const float*>(&ExpectedF));
+	CHECK_MESSAGE("Swapping double value",         ByteSwap(ValD)   == *reinterpret_cast<const double*>(&ExpectedD));
+	CHECK_MESSAGE("Swapping char16_t value",       ByteSwap(ValCh16)== 0x3322);
+}
+
+// The byte swap benchmarking test are useful to compare the intrinsic implementation vs the generic implementation. Normally, the intrinsic is expected to be faster, but in some
+// cases, the intrinsic vs generic speed is within margin of error in optimized build. So the test result are not always consistent. The code remains there in case a new
+// implementation needs to be testes. Here some observations:
+//     - VC++ 2019 (16.4.3): The intrinsic versions are consistently (and significantly) faster than the generic version. The compiler poorly optimize the generic version.
+//     - Apple Clang (11.0.0): No performance change observed in release. Clang generates the same assembly for the intrinsic and the generic functions.
+//     - Linux Clang (8.0.1): No performance change observed in release. Clang generates the same assembly for the intrinsic and the generic functions.
+#define UE_BYTE_SWAP_BENCHMARK_ENABLED 0
+#if UE_BYTE_SWAP_BENCHMARK_ENABLED
+
+/**
+ * Compare the performance of swapping bytes using the intrinsic vs generic implementation. See comment above UE_BYTE_SWAP_BENCHMARK_ENABLED.
+ */
+
+template<typename T, typename ByteSwapFn>
+TPair<FTimespan, T> BenchmarkByteswapping(T InitialValue, uint64 LoopCount, ByteSwapFn InByteSwapFn)
+{
+	static_assert(TIsIntegral<T>::Value, "The test only works with integers");
+	static_assert(!TIsSigned<T>::Value, "The test must support overflowing T when making the sum, value cannot be signed.");
+
+	T Sum = 0; // The sum value is not very relevant, but it is an observable result the compiler cannot ignore (prevent to optimize away the entire loop).
+	FDateTime StartTime = FDateTime::UtcNow();
+	for (uint64 I = 0; I < LoopCount; ++I)
+	{
+		Sum += InByteSwapFn(InitialValue + static_cast<T>(I)); // Expect InByteSwapFn to be inlined because the compiler knows the exact function to call.
+	}
+	FTimespan Duration = FDateTime::UtcNow() - StartTime;
+	return MakeTuple(Duration, Sum);
+}
+
+TEST_CASE_NAMED(FByteSwapPerformanceTest, "System::Core::Misc::ByteSwapPerf", "[ApplicationContextMask][PerfFilter]")
+{
+#ifdef UE_BYTESWAP_INTRINSIC_PRIVATE_16
+	{
+		uint16 InitialValue = 0xF0F0;
+		TPair<FTimespan, uint16> Intrinsic = BenchmarkByteswapping(InitialValue, 1000000000, [](uint16 Val){ return ByteSwap(Val); });
+		TPair<FTimespan, uint16> Generic   = BenchmarkByteswapping(InitialValue, 1000000000, [](uint16 Val){ return Internal::ByteSwapGeneric16(Val); });
+		CHECK_MESSAGE("Swapping uint16 bytes is faster uing the compiler intrinsic than the generic implementation", Intrinsic.Get<0>() <= Generic.Get<0>());
+		CHECK_MESSAGE("Swapping uint16 bytes using intrinsic and generic algorithm produce the same values", Intrinsic.Get<1>() == Generic.Get<1>());
+		GLog->Logf(TEXT("Swapping 2 bytes using intrinsic is %s faster than generic version"), *LexToString(Generic.Get<0>().GetTotalMicroseconds() / Intrinsic.Get<0>().GetTotalMicroseconds()));
+	}
+#endif
+
+#ifdef UE_BYTESWAP_INTRINSIC_PRIVATE_32
+	{
+		uint32 InitialValue = 0xFF00FF00ul;
+		TPair<FTimespan, uint32> Intrinsic = BenchmarkByteswapping(InitialValue, 1000000000, [](uint32 Val){ return ByteSwap(Val); });
+		TPair<FTimespan, uint32> Generic   = BenchmarkByteswapping(InitialValue, 1000000000, [](uint32 Val){ return Internal::ByteSwapGeneric32(Val); });
+		CHECK_MESSAGE("Swapping uint32 bytes is faster uing the compiler intrinsic than the generic implementation", Intrinsic.Get<0>() <= Generic.Get<0>());
+		CHECK_MESSAGE("Swapping uint32 bytes using intrinsic and generic algorithm produce the same values", Intrinsic.Get<1>() == Generic.Get<1>());
+		GLog->Logf(TEXT("Swapping 4 bytes using intrinsic is %s faster than generic version"), *LexToString(Generic.Get<0>().GetTotalMicroseconds() / Intrinsic.Get<0>().GetTotalMicroseconds()));
+	}
+#endif
+
+#ifdef UE_BYTESWAP_INTRINSIC_PRIVATE_64
+	{
+		uint64 InitialValue = 0xFF00FF00FF00FF00ull;
+		TPair<FTimespan, uint64> Intrinsic = BenchmarkByteswapping(InitialValue, 1000000000, [](uint64 Val){ return ByteSwap(Val); });
+		TPair<FTimespan, uint64> Generic   = BenchmarkByteswapping(InitialValue, 1000000000, [](uint64 Val){ return Internal::ByteSwapGeneric64(Val); });
+		CHECK_MESSAGE("Swapping uint64 bytes is faster uing the compiler intrinsic than the generic implementation", Intrinsic.Get<0>() <= Generic.Get<0>());
+		CHECK_MESSAGE("Swapping uint64 bytes using intrinsic and generic algorithm produce the same values", Intrinsic.Get<1>() == Generic.Get<1>());
+		GLog->Logf(TEXT("Swapping 8 bytes using intrinsic is %s faster than generic version"), *LexToString(Generic.Get<0>().GetTotalMicroseconds() / Intrinsic.Get<0>().GetTotalMicroseconds()));
+	}
+#endif
+
+}
+
+#endif //UE_BYTE_SWAP_BENCHMARK_ENABLED
+
+#endif // WITH_TESTS
diff --git a/Engine/Source/Runtime/Core/Tests/Misc/CharTest.cpp b/Engine/Source/Runtime/Core/Tests/Misc/CharTest.cpp
new file mode 100644
--- /dev/null
+++ b/Engine/Source/Runtime/Core/Tests/Misc/CharTest.cpp
@@ -0,0 +1,53 @@
+// Copyright Epic Games, Inc. All Rights Reserved.
+
+#if WITH_TESTS 
+
+#include "Misc/Char.h"
+#include <locale.h>
+#include <ctype.h>
+#include <wctype.h>
+#include "Tests/TestHarnessAdapter.h"
+#if WITH_LOW_LEVEL_TESTS
+#include "TestCommon/Comparisons.h"
+#endif
+
+namespace crt
+{
+	int tolower(ANSICHAR c) { return ::tolower(c); }
+	int toupper(ANSICHAR c) { return ::toupper(c); }
+
+	int tolower(WIDECHAR c) { return ::towlower(c); }
+	int toupper(WIDECHAR c) { return ::towupper(c); }
+}
+
+template<typename CharType>
+void RunCharTests(uint32 MaxChar)
+{
+	for (uint32 I = 0; I < MaxChar; ++I)
+	{
+		CharType C = (CharType)I;
+		CHECK_EQUALS(TEXT("TChar::ToLower()"), TChar<CharType>::ToLower(C), crt::tolower(C));
+		CHECK_EQUALS(TEXT("TChar::ToUpper()"), TChar<CharType>::ToUpper(C), crt::toupper(C));
+	}
+}
+
+TEST_CASE_NAMED(TCharTest, "System::Core::Misc::Char", "[ApplicationContextMask][SmokeFilter]")
+{
+	const char* CurrentLocale = setlocale(LC_CTYPE, nullptr);
+	if (CurrentLocale == nullptr)
+	{
+		FAIL_CHECK(FString::Printf(TEXT("Locale is null but should be \"C\". Did something call setlocale()?")));
+	}
+	else if (strcmp("C", CurrentLocale))
+	{
+		FAIL_CHECK(FString::Printf(TEXT("Locale is \"%s\" but should be \"C\". Did something call setlocale()?"), ANSI_TO_TCHAR(CurrentLocale)));
+	}
+	else
+	{
+		RunCharTests<ANSICHAR>(128);
+		RunCharTests<WIDECHAR>(65536);
+	}
+
+}
+
+#endif //WITH_TESTS
\ No newline at end of file
diff --git a/Engine/Source/Runtime/Core/Tests/Misc/CircularBufferTest.cpp b/Engine/Source/Runtime/Core/Tests/Misc/CircularBufferTest.cpp
new file mode 100644
--- /dev/null
+++ b/Engine/Source/Runtime/Core/Tests/Misc/CircularBufferTest.cpp
@@ -0,0 +1,49 @@
+// Copyright Epic Games, Inc. All Rights Reserved.
+
+#if WITH_TESTS
+
+#include "CoreTypes.h"
+#include "Containers/UnrealString.h"
+#include "Containers/CircularBuffer.h"
+#include "Tests/TestHarnessAdapter.h"
+#if WITH_LOW_LEVEL_TESTS
+#include "TestCommon/Comparisons.h"
+#endif
+
+TEST_CASE_NAMED(FCircularBufferTest, "System::Core::Misc::CircularBuffer", "[ApplicationContextMask][SmokeFilter]")
+{
+	// buffer capacity
+	TCircularBuffer<int32> b1_1(127);
+	TCircularBuffer<int32> b1_2(128);
+	TCircularBuffer<int32> b1_3(129);
+
+
+	CHECK_EQUALS(TEXT("Buffer capacity of 127 must be rounded up to 128"), b1_1.Capacity(), 128u);
+	CHECK_EQUALS(TEXT("Buffer capacity of 128 must not change"), b1_2.Capacity(), 128u);
+	CHECK_EQUALS(TEXT("Buffer capacity of 129 must be rounded up to 256"), b1_3.Capacity(), 256u);
+
+	// initial values
+	TCircularBuffer<int32> b2_1(64, 666);
+
+	for (uint32 Index = 0; Index < b2_1.Capacity(); ++Index)
+	{
+		CHECK_EQUALS(FString::Printf(TEXT("Initial value must be correct (%i)"), Index), b2_1[Index], 666);
+	}
+
+	// indexing
+	TCircularBuffer<int32> b3_1(64, 0);
+
+	CHECK_EQUALS(TEXT("Next index from 0 must be 1"), b3_1.GetNextIndex(0), 1u);
+	CHECK_EQUALS(TEXT("Next index from 63 must be 0"), b3_1.GetNextIndex(63), 0u);
+	CHECK_EQUALS(TEXT("Next index from 64 must be 1"), b3_1.GetNextIndex(64), 1u);
+
+	b3_1[0] = 42;
+	b3_1[65] = 42;
+
+	CHECK_EQUALS(TEXT("Index 0 must be written and read correctly"), b3_1[0], 42);
+	CHECK_EQUALS(TEXT("Index 1 must be written and read correctly"), b3_1[1], 42);
+	CHECK_EQUALS(TEXT("Index 65 must be written and read correctly"), b3_1[65], 42);
+
+}
+
+#endif //WITH_TESTS
diff --git a/Engine/Source/Runtime/Core/Tests/Misc/CircularQueueTest.cpp b/Engine/Source/Runtime/Core/Tests/Misc/CircularQueueTest.cpp
new file mode 100644
--- /dev/null
+++ b/Engine/Source/Runtime/Core/Tests/Misc/CircularQueueTest.cpp
@@ -0,0 +1,137 @@
+// Copyright Epic Games, Inc. All Rights Reserved.
+
+#if WITH_TESTS
+
+#include "CoreTypes.h"
+#include "Containers/CircularQueue.h"
+#include "Tests/TestHarnessAdapter.h"
+#if WITH_LOW_LEVEL_TESTS
+#include "TestCommon/Comparisons.h"
+#endif
+
+TEST_CASE_NAMED(FCircularQueueTest, "System::Core::Misc::CircularQueue", "[ApplicationContextMask][SmokeFilter]")
+{
+	const uint32 QueueSize = 8;
+
+	// empty queue
+	{
+		TCircularQueue<int32> Queue(QueueSize);
+
+		CHECK_EQUALS(TEXT("Newly created queues must have zero elements"), Queue.Count(), 0u);
+		CHECK_MESSAGE(TEXT("Newly created queues must be empty"), Queue.IsEmpty());
+		CHECK_FALSE_MESSAGE(TEXT("Newly created queues must not be full"), Queue.IsFull());
+
+		int32 Value;
+		CHECK_FALSE_MESSAGE(TEXT("Peek must fail on an empty queue"), Queue.Peek(Value));
+		CHECK_MESSAGE(TEXT("Peek must fail on an empty queue"), Queue.Peek() == nullptr);
+	}
+
+	// partially filled
+	{
+		TCircularQueue<int32> Queue(QueueSize);
+		int32 Value = 0;
+
+		CHECK_MESSAGE(TEXT("Adding to an empty queue must succeed"), Queue.Enqueue(666));
+		CHECK_EQUALS(TEXT("After adding to an empty queue it must have one element"), Queue.Count(), 1u);
+		CHECK_FALSE_MESSAGE(TEXT("Partially filled queues must not be empty"), Queue.IsEmpty());
+		CHECK_FALSE_MESSAGE(TEXT("Partially filled queues must not be full"), Queue.IsFull());
+		CHECK_MESSAGE(TEXT("Peeking at a partially filled queue must succeed"), Queue.Peek(Value));
+		CHECK_EQUALS(TEXT("The peeked at value must be correct"), Value, 666);
+
+		const int32* PeekValue = Queue.Peek();
+		CHECK_MESSAGE(TEXT("Peeking at a partially filled queue must succeed"), PeekValue != nullptr);
+		CHECK_EQUALS(TEXT("The peeked at value must be correct"), *PeekValue, 666);
+	}
+
+	// full queue
+	for (uint32 PeekType = 0; PeekType < 2; PeekType++)
+	{
+		TCircularQueue<int32> Queue(QueueSize);
+
+		for (int32 Index = 0; Index < QueueSize - 1; ++Index)
+		{
+			CHECK_MESSAGE(TEXT("Adding to non-full queue must succeed"), Queue.Enqueue(Index));
+		}
+
+		CHECK_FALSE_MESSAGE(TEXT("Full queues must not be empty"), Queue.IsEmpty());
+		CHECK_MESSAGE(TEXT("Full queues must be full"), Queue.IsFull());
+		CHECK_FALSE_MESSAGE(TEXT("Adding to full queue must fail"), Queue.Enqueue(666));
+
+		int32 Value = 0;
+
+		for (int32 Index = 0; Index < QueueSize - 1; ++Index)
+		{
+			if (PeekType == 0)
+			{
+				CHECK_MESSAGE(TEXT("Peeking at a non-empty queue must succeed"), Queue.Peek(Value));
+				CHECK_EQUALS(TEXT("The peeked at value must be correct"), Value, Index);
+
+				CHECK_MESSAGE(TEXT("Removing from a non-empty queue must succeed"), Queue.Dequeue(Value));
+				CHECK_EQUALS(TEXT("The removed value must be correct"), Value, Index);
+			}
+			else
+			{
+				const int32* PeekValue = Queue.Peek();
+				CHECK_MESSAGE(TEXT("Peeking at a non-empty queue must succeed"), PeekValue != nullptr);
+				CHECK_EQUALS(TEXT("The peeked at value must be correct"), *PeekValue, Index);
+
+				CHECK_MESSAGE(TEXT("Removing from a non-empty queue must succeed"), Queue.Dequeue());
+			}
+		}
+
+		CHECK_MESSAGE(TEXT("A queue that had all items removed must be empty"), Queue.IsEmpty());
+		CHECK_FALSE_MESSAGE(TEXT("A queue that had all items removed must not be full"), Queue.IsFull());
+	}
+
+	// queue with index wrapping around
+	{
+		TCircularQueue<int32> Queue(QueueSize);
+
+		//Fill queue
+		for (int32 Index = 0; Index < QueueSize - 1; ++Index)
+		{
+			CHECK_MESSAGE(TEXT("Adding to non-full queue must succeed"), Queue.Enqueue(Index));
+		}
+
+		int32 Value = 0;
+		const int32 ExpectedSize = QueueSize - 1;
+		for (int32 Index = 0; Index < QueueSize; ++Index)
+		{
+			CHECK_EQUALS(TEXT("Number of elements must be valid for all permutation of Tail and Head"), Queue.Count(), ExpectedSize);
+			CHECK_MESSAGE(TEXT("Removing from a non-empty queue must succeed"), Queue.Dequeue(Value));
+			CHECK_MESSAGE(TEXT("Adding to non-full queue must succeed"), Queue.Enqueue(Index));
+		}
+	}
+
+	// Non-zero initialization - ensures the backing store constructs and destructs non-POD objects properly.
+	{
+		static uint32 Const;
+		static uint32 Dest;
+		static uint32 CopyConst;
+		struct FNonPOD
+		{
+			FNonPOD() { Const++; }
+			~FNonPOD() { Dest++; }
+			FNonPOD(const FNonPOD&) { CopyConst++; }
+		};
+
+		// The current implementation of TCircularQueue doesn't call the held object constructors on but does call the
+		// destructors. Verify that (somewhat surprising, if typically ok for POD) behavior first.
+		Const = 0;
+		Dest = 0;
+		CopyConst = 0;
+		{
+			TCircularQueue<FNonPOD> Queue(QueueSize);
+
+			CHECK_EQUALS(TEXT("The constructor should not run"), Const, 0);
+			CHECK_EQUALS(TEXT("The destructor should not run"), Dest, 0);
+			CHECK_EQUALS(TEXT("The copy constructor should not run"), CopyConst, 0);
+		}
+		CHECK_EQUALS(TEXT("The constructor should not run"), Const, 0);
+		CHECK_EQUALS(TEXT("The destructor should run"), Dest, QueueSize);
+		CHECK_EQUALS(TEXT("The copy constructor should not run"), CopyConst, 0);
+	}
+
+}
+
+#endif //WITH_TESTS
diff --git a/Engine/Source/Runtime/Core/Tests/Misc/DateTimeTest.cpp b/Engine/Source/Runtime/Core/Tests/Misc/DateTimeTest.cpp
new file mode 100644
--- /dev/null
+++ b/Engine/Source/Runtime/Core/Tests/Misc/DateTimeTest.cpp
@@ -0,0 +1,143 @@
+// Copyright Epic Games, Inc. All Rights Reserved.
+
+#if WITH_TESTS
+
+#include "CoreTypes.h"
+#include "Containers/UnrealString.h"
+#include "Misc/DateTime.h"
+#include "Tests/TestHarnessAdapter.h"
+#if WITH_LOW_LEVEL_TESTS
+#include "TestCommon/Comparisons.h"
+#endif
+
+#define TestUnixEquivalent(Desc, A, B) if ((A).ToUnixTimestamp() != (B)) FAIL_CHECK(FString::Printf(TEXT("%s - A=%d B=%d"), Desc, (A).ToUnixTimestamp(), (B)));
+#define TestYear(Desc, A, B) if ((A.GetYear()) != (B)) FAIL_CHECK(FString::Printf(TEXT("%s - A=%d B=%d"), Desc, (A.GetYear()), (B)));
+#define TestMonth(Desc, A, B) if ((A.GetMonth()) != (B)) FAIL_CHECK(FString::Printf(TEXT("%s - A=%d B=%d"), Desc, (A.GetMonth()), B));
+#define TestMonthOfYear(Desc, A, B) if ((A.GetMonthOfYear()) != (B)) FAIL_CHECK(FString::Printf(TEXT("%s - A=%d B=%d"), Desc ,(static_cast<int32>(A.GetMonthOfYear())), static_cast<int32>(B)));
+#define TestDay(Desc, A, B) if ((A.GetDay()) != (B)) FAIL_CHECK(FString::Printf(TEXT("%s - A=%d B=%d"), Desc, (A.GetDay()), (B)));
+#define TestHour(Desc, A, B) if ((A.GetHour()) != (B)) FAIL_CHECK(FString::Printf(TEXT("%s - A=%d B=%d"), Desc, (A.GetHour()), (B)));
+#define TestMinute(Desc, A, B) if ((A.GetMinute()) != (B)) FAIL_CHECK(FString::Printf(TEXT("%s - A=%d B=%d"), Desc, (A.GetMinute()), (B)));
+#define TestSecond(Desc, A, B) if ((A.GetSecond()) != (B)) FAIL_CHECK(FString::Printf(TEXT("%s - A=%d B=%d"), Desc, (A.GetSecond()), (B)));
+#define TestMillisecond(Desc, A, B) if ((A.GetMillisecond()) != (B)) FAIL_CHECK(FString::Printf(TEXT("%s - A=%d B=%d"), Desc, (A.GetMillisecond()) ,(B)));
+
+TEST_CASE_NAMED(FDateTimeTest, "System::Core::Misc::DateTime", "[ApplicationContextMask][SmokeFilter]")
+{
+	const int64 UnixEpochTimestamp = 0;
+	const int64 UnixBillenniumTimestamp = 1000000000;
+	const int64 UnixOnesTimestamp = 1111111111;
+	const int64 UnixDecimalSequenceTimestamp = 1234567890;
+
+	const FDateTime UnixEpoch = FDateTime::FromUnixTimestamp(UnixEpochTimestamp);
+	const FDateTime UnixBillennium = FDateTime::FromUnixTimestamp(UnixBillenniumTimestamp);
+	const FDateTime UnixOnes = FDateTime::FromUnixTimestamp(UnixOnesTimestamp);
+	const FDateTime UnixDecimalSequence = FDateTime::FromUnixTimestamp(UnixDecimalSequenceTimestamp);
+
+	const FDateTime TestDateTime(2013, 8, 14, 12, 34, 56, 789);
+
+	TestUnixEquivalent(TEXT("Testing Unix Epoch Ticks"), UnixEpoch, UnixEpochTimestamp);
+	TestUnixEquivalent(TEXT("Testing Unix Billennium Ticks"), UnixBillennium, UnixBillenniumTimestamp);
+	TestUnixEquivalent(TEXT("Testing Unix Ones Ticks"), UnixOnes, UnixOnesTimestamp);
+	TestUnixEquivalent(TEXT("Testing Unix Decimal Sequence Ticks"), UnixDecimalSequence, UnixDecimalSequenceTimestamp);
+
+	TestYear(TEXT("Testing Unix Epoch Year"), UnixEpoch, 1970);
+	TestMonth(TEXT("Testing Unix Epoch Month"), UnixEpoch, 1);
+	TestMonthOfYear(TEXT("Testing Unix Epoch MonthOfYear"), UnixEpoch, EMonthOfYear::January);
+	TestDay(TEXT("Testing Unix Epoch Day"), UnixEpoch, 1);
+	TestHour(TEXT("Testing Unix Epoch Hour"), UnixEpoch, 0);
+	TestMinute(TEXT("Testing Unix Epoch Minute"), UnixEpoch, 0);
+	TestSecond(TEXT("Testing Unix Epoch Second"), UnixEpoch, 0);
+	TestMillisecond(TEXT("Testing Unix Epoch Millisecond"), UnixEpoch, 0);
+
+	TestYear(TEXT("Testing Unix Billennium Year"), UnixBillennium, 2001);
+	TestMonth(TEXT("Testing Unix Billennium MonthOfYear"), UnixBillennium, 9);
+	TestMonthOfYear(TEXT("Testing Unix Billennium Month"), UnixBillennium, EMonthOfYear::September);
+	TestDay(TEXT("Testing Unix Billennium Day"), UnixBillennium, 9);
+	TestHour(TEXT("Testing Unix Billennium Hour"), UnixBillennium, 1);
+	TestMinute(TEXT("Testing Unix Billennium Minute"), UnixBillennium, 46);
+	TestSecond(TEXT("Testing Unix Billennium Second"), UnixBillennium, 40);
+	TestMillisecond(TEXT("Testing Unix Billennium Millisecond"), UnixBillennium, 0);
+
+	TestYear(TEXT("Testing Unix Ones Year"), UnixOnes, 2005);
+	TestMonth(TEXT("Testing Unix Ones Month"), UnixOnes, 3);
+	TestMonthOfYear(TEXT("Testing Unix Ones MonthOfYear"), UnixOnes, EMonthOfYear::March);
+	TestDay(TEXT("Testing Unix Ones Day"), UnixOnes, 18);
+	TestHour(TEXT("Testing Unix Ones Hour"), UnixOnes, 1);
+	TestMinute(TEXT("Testing Unix Ones Minute"), UnixOnes, 58);
+	TestSecond(TEXT("Testing Unix Ones Second"), UnixOnes, 31);
+	TestMillisecond(TEXT("Testing Unix Ones Millisecond"), UnixOnes, 0);
+
+	TestYear(TEXT("Testing Unix Decimal Sequence Year"), UnixDecimalSequence, 2009);
+	TestMonth(TEXT("Testing Unix Decimal Sequence Month"), UnixDecimalSequence, 2);
+	TestMonthOfYear(TEXT("Testing Unix Decimal Sequence MonthOfYear"), UnixDecimalSequence, EMonthOfYear::February);
+	TestDay(TEXT("Testing Unix Decimal Sequence Day"), UnixDecimalSequence, 13);
+	TestHour(TEXT("Testing Unix Decimal Sequence Hour"), UnixDecimalSequence, 23);
+	TestMinute(TEXT("Testing Unix Decimal Sequence Minute"), UnixDecimalSequence, 31);
+	TestSecond(TEXT("Testing Unix Decimal Sequence Second"), UnixDecimalSequence, 30);
+	TestMillisecond(TEXT("Testing Unix Decimal Sequence Millisecond"), UnixDecimalSequence, 0);
+
+	TestYear(TEXT("Testing Test Date Time Year"), TestDateTime, 2013);
+	TestMonth(TEXT("Testing Test Date Time Month"), TestDateTime, 8);
+	TestMonthOfYear(TEXT("Testing Test Date Time MonthOfYear"), TestDateTime, EMonthOfYear::August);
+	TestDay(TEXT("Testing Test Date Time Day"), TestDateTime, 14);
+	TestHour(TEXT("Testing Test Date Time Hour"), TestDateTime, 12);
+	TestMinute(TEXT("Testing Test Date Time Minute"), TestDateTime, 34);
+	TestSecond(TEXT("Testing Test Date Time Second"), TestDateTime, 56);
+	TestMillisecond(TEXT("Testing Test Date Time Millisecond"), TestDateTime, 789);
+
+	FDateTime ParsedDateTime;
+
+	CHECK_FALSE_MESSAGE(TEXT("Parsing an empty ISO string must fail"), FDateTime::ParseIso8601(TEXT(""), ParsedDateTime));
+
+	FDateTime::ParseIso8601(TEXT("2019-05-22"), ParsedDateTime);
+	CHECK_EQUALS(TEXT("Testing ISO 8601 date"), ParsedDateTime, (FDateTime{ 2019, 5, 22 }));
+
+	FDateTime::ParseIso8601(TEXT("2019-05-20T19:41:38+01:30"), ParsedDateTime);
+	CHECK_EQUALS(TEXT("Testing ISO 8601 with +hh:mm timezone info"), ParsedDateTime, (FDateTime{ 2019, 5, 20, 18, 11, 38 }));
+	FDateTime::ParseIso8601(TEXT("2019-05-20T19:41:38-01:30"), ParsedDateTime);
+	CHECK_EQUALS(TEXT("Testing ISO 8601 with -hh:mm timezone info"), ParsedDateTime, (FDateTime{ 2019, 5, 20, 21, 11, 38 }));
+	FDateTime::ParseIso8601(TEXT("2019-05-20T19:41:38+0030"), ParsedDateTime);
+	CHECK_EQUALS(TEXT("Testing ISO 8601 with +hhmm timezone info"), ParsedDateTime, (FDateTime{ 2019, 5, 20, 19, 11, 38 }));
+	FDateTime::ParseIso8601(TEXT("2019-05-20T19:41:38-01"), ParsedDateTime);
+	CHECK_EQUALS(TEXT("Testing ISO 8601 with -hh timezone info"), ParsedDateTime, (FDateTime{ 2019, 5, 20, 20, 41, 38 }));
+
+	for (double JulianDay : {0., 1000., 1721425.5, 1721425.0, FDateTime::MinValue().GetJulianDay(), FDateTime::MaxValue().GetJulianDay()})
+	{
+		CHECK_EQUALS(TEXT("convertion from/to JulianDay is stable"), FDateTime::FromJulianDay(JulianDay).GetJulianDay(), JulianDay);
+	}
+
+	FDateTime Date{ 2019, 5, 20, 18, 11, 38 };
+	FDateTime DateMidnight = Date.GetDate();
+	CHECK_EQUALS(TEXT("GetDate returns a FDateTime at midnight"), DateMidnight.GetHour(), 0);
+
+	FDateTime DateNoon = DateMidnight + FTimespan::FromHours(12);
+	CHECK_EQUALS(TEXT("GetDate returns a FDateTime at midnight"), DateNoon.GetHour(), 12);
+
+	double JulianDayMidnight = DateMidnight.GetJulianDay();
+	double JulianDayNoon = DateNoon.GetJulianDay();
+	CHECK_EQUALS(TEXT("At midnight, fractionnal part of the JulianDay value should be 0.5"), JulianDayMidnight - (int)JulianDayMidnight, 0.5);
+	CHECK_EQUALS(TEXT("A 12h timespan adds half a julianday"), JulianDayNoon - JulianDayMidnight, 0.5);
+
+	double OffsetDayCount = 12345.;
+	FDateTime OffsetDate = Date + FTimespan::FromDays(OffsetDayCount);
+	CHECK_EQUALS(TEXT("An offset by a given numer of days leads to a similar JulianDay offset"), OffsetDate.GetJulianDay() - Date.GetJulianDay(), OffsetDayCount);
+
+	int32 Year, Month, Day;
+	DateMidnight.GetDate(Year, Month, Day);
+	FDateTime PreviousDay(DateMidnight.GetTicks() - 1);
+	int32 YearPrev, MonthPrev, DayPrev;
+	PreviousDay.GetDate(YearPrev, MonthPrev, DayPrev);
+	CHECK_NOT_EQUALS(TEXT("One tick before a date at midnight leads to the previous date"), Day, DayPrev);
+}
+
+
+#undef TestUnixEquivalent
+#undef TestYear
+#undef TestMonth
+#undef TestDay
+#undef TestHour
+#undef TestMinute
+#undef TestSecond
+#undef TestMillisecond
+
+
+#endif //WITH_TESTS
diff --git a/Engine/Source/Runtime/Core/Tests/Misc/DisplayNameTest.cpp b/Engine/Source/Runtime/Core/Tests/Misc/DisplayNameTest.cpp
new file mode 100644
--- /dev/null
+++ b/Engine/Source/Runtime/Core/Tests/Misc/DisplayNameTest.cpp
@@ -0,0 +1,39 @@
+// Copyright Epic Games, Inc. All Rights Reserved.
+
+#if WITH_TESTS 
+
+#include "UObject/UnrealNames.h"
+#include "Tests/TestHarnessAdapter.h"
+#if WITH_LOW_LEVEL_TESTS
+#include "TestCommon/Comparisons.h"
+#endif
+
+TEST_CASE_NAMED(TDisplayNameTest, "System::Core::Misc::DisplayName", "[ApplicationContextMask][SmokeFilter]")
+{
+	CHECK_EQUALS(TEXT("Boolean"), FName::NameToDisplayString(TEXT("bTest"), true), TEXT("Test"));
+	CHECK_EQUALS(TEXT("Boolean Lower"), FName::NameToDisplayString(TEXT("bTwoWords"), true), TEXT("Two Words"));
+	CHECK_EQUALS(TEXT("Lower Boolean"), FName::NameToDisplayString(TEXT("boolean"), true), TEXT("Boolean"));
+	CHECK_EQUALS(TEXT("Almost Boolean"), FName::NameToDisplayString(TEXT("bNotBoolean"), false), TEXT("B Not Boolean"));
+	CHECK_EQUALS(TEXT("Boolean No Prefix"), FName::NameToDisplayString(TEXT("NonprefixBoolean"), true), TEXT("Nonprefix Boolean"));
+	CHECK_EQUALS(TEXT("Lower Boolean No Prefix"), FName::NameToDisplayString(TEXT("lowerNonprefixBoolean"), true), TEXT("Lower Nonprefix Boolean"));
+	CHECK_EQUALS(TEXT("Lower Camel Case"), FName::NameToDisplayString(TEXT("lowerCase"), false), TEXT("Lower Case"));
+	CHECK_EQUALS(TEXT("With Underscores"), FName::NameToDisplayString(TEXT("With_Underscores"), false), TEXT("With Underscores"));
+	CHECK_EQUALS(TEXT("Lower Underscores"), FName::NameToDisplayString(TEXT("lower_underscores"), false), TEXT("Lower Underscores"));
+	CHECK_EQUALS(TEXT("Mixed Underscores"), FName::NameToDisplayString(TEXT("mixed_Underscores"), false), TEXT("Mixed Underscores"));
+	CHECK_EQUALS(TEXT("Mixed Underscores"), FName::NameToDisplayString(TEXT("Mixed_underscores"), false), TEXT("Mixed Underscores"));
+	CHECK_EQUALS(TEXT("Article in String"), FName::NameToDisplayString(TEXT("ArticleInString"), false), TEXT("Article in String"));
+	CHECK_EQUALS(TEXT("One or Two"), FName::NameToDisplayString(TEXT("OneOrTwo"), false), TEXT("One or Two"));
+	CHECK_EQUALS(TEXT("One and Two"), FName::NameToDisplayString(TEXT("OneAndTwo"), false), TEXT("One and Two"));
+	CHECK_EQUALS(TEXT("-1.5"), FName::NameToDisplayString(TEXT("-1.5"), false), TEXT("-1.5"));
+	CHECK_EQUALS(TEXT("1234"), FName::NameToDisplayString(TEXT("1234"), false), TEXT("1234"));
+	CHECK_EQUALS(TEXT("1234.5"), FName::NameToDisplayString(TEXT("1234.5"), false), TEXT("1234.5"));
+	CHECK_EQUALS(TEXT("-1234.5"), FName::NameToDisplayString(TEXT("-1234.5"), false), TEXT("-1234.5"));
+	CHECK_EQUALS(TEXT("Text (In Parens)"), FName::NameToDisplayString(TEXT("Text (in parens)"), false), TEXT("Text (In Parens)"));
+
+	CHECK_EQUALS(TEXT("Text 3D"), FName::NameToDisplayString(TEXT("Text3D"), false), TEXT("Text 3D"));
+	CHECK_EQUALS(TEXT("Plural CAPs"), FName::NameToDisplayString(TEXT("PluralCAPs"), false), TEXT("Plural CAPs"));
+	CHECK_EQUALS(TEXT("FBXEditor"), FName::NameToDisplayString(TEXT("FBXEditor"), false), TEXT("FBXEditor"));
+	CHECK_EQUALS(TEXT("FBX_Editor"), FName::NameToDisplayString(TEXT("FBX_Editor"), false), TEXT("FBX Editor"));
+}
+
+#endif //WITH_TESTS
\ No newline at end of file
diff --git a/Engine/Source/Runtime/Core/Tests/Misc/FileTest.cpp b/Engine/Source/Runtime/Core/Tests/Misc/FileTest.cpp
new file mode 100644
--- /dev/null
+++ b/Engine/Source/Runtime/Core/Tests/Misc/FileTest.cpp
@@ -0,0 +1,239 @@
+// Copyright Epic Games, Inc. All Rights Reserved.
+
+#if WITH_TESTS
+
+#include "CoreTypes.h"
+#include "Misc/AssertionMacros.h"
+#include "Containers/Array.h"
+#include "Containers/UnrealString.h"
+#include "HAL/PlatformFile.h"
+#include "HAL/PlatformFileManager.h"
+#include "HAL/FileManager.h"
+#include "Misc/FileHelper.h"
+#include "Misc/Paths.h"
+#include "Misc/Guid.h"
+#include "Misc/ScopeExit.h"
+#include "Tests/TestHarnessAdapter.h"
+#if WITH_LOW_LEVEL_TESTS
+#include "TestCommon/Comparisons.h"
+#endif
+
+// These file tests are designed to ensure expected file writing behavior, as well as cross-platform consistency
+
+TEST_CASE_NAMED(FFileTruncateTest, "System::Core::Misc::FileTruncate", "[.][EditorContext][CriticalPriority][EngineFilter]")
+{
+	const FString TempFilename = FPaths::CreateTempFilename(*FPaths::ProjectIntermediateDir());
+	IPlatformFile& PlatformFile = FPlatformFileManager::Get().GetPlatformFile();
+	ON_SCOPE_EXIT
+	{
+		// Delete temp file
+		PlatformFile.DeleteFile(*TempFilename);
+	};
+
+	// Open a test file
+	if (TUniquePtr<IFileHandle> TestFile = TUniquePtr<IFileHandle>(PlatformFile.OpenWrite(*TempFilename, /*bAppend*/false, /*bAllowRead*/true)))
+	{
+		// Append 4 int32 values of incrementing value to this file
+		int32 Val = 1;
+		TestFile->Write((const uint8*)&Val, sizeof(Val));
+
+		++Val;
+		TestFile->Write((const uint8*)&Val, sizeof(Val));
+
+		// Tell here, so we can move back and truncate after writing
+		const int64 ExpectedTruncatePos = TestFile->Tell();
+		++Val;
+		TestFile->Write((const uint8*)&Val, sizeof(Val));
+
+		// Tell here, so we can attempt to read here after truncation
+		const int64 TestReadPos = TestFile->Tell();
+		++Val;
+		TestFile->Write((const uint8*)&Val, sizeof(Val));
+
+		// Validate that the Tell position is at the end of the file, and that the size is reported correctly
+		{
+			const int64 ActualEOFPos = TestFile->Tell();
+			const int64 ExpectedEOFPos = (sizeof(int32) * 4);
+			REQUIRE_MESSAGE(FString::Printf(TEXT("File was not the expected size (got %d, expected %d): %s"), ActualEOFPos, ExpectedEOFPos, *TempFilename), ActualEOFPos == ExpectedEOFPos);
+
+			const int64 ActualFileSize = TestFile->Size();
+			REQUIRE_MESSAGE(FString::Printf(TEXT("File was not the expected size (got %d, expected %d): %s"), ActualFileSize, ExpectedEOFPos, *TempFilename), ActualFileSize == ExpectedEOFPos);
+		}
+
+		// Truncate the file at our test pos
+		REQUIRE_MESSAGE(FString::Printf(TEXT("File truncation request failed: %s"), *TempFilename), TestFile->Truncate(ExpectedTruncatePos));
+
+		// Validate that the size is reported correctly
+		{
+			const int64 ActualFileSize = TestFile->Size();
+			REQUIRE_MESSAGE(FString::Printf(TEXT("File was not the expected size after truncation (got %d, expected %d): %s"), ActualFileSize, ExpectedTruncatePos, *TempFilename), ActualFileSize == ExpectedTruncatePos);
+		}
+
+		// Validate that we can't read past the truncation point
+		{
+			int32 Dummy = 0;
+			REQUIRE_MESSAGE(FString::Printf(TEXT("File read seek outside the truncated range: %s"), *TempFilename), !(TestFile->Seek(TestReadPos) && TestFile->Read((uint8*)&Dummy, sizeof(Dummy))));
+		}
+	}
+	else
+	{
+		REQUIRE_MESSAGE(FString::Printf(TEXT("File failed to open: %s"), *TempFilename), false);
+	}
+}
+
+TEST_CASE_NAMED(FFileAppendTest, "System::Core::Misc::FileAppend", "[.][EditorContext][CriticalPriority][EngineFilter]")
+{
+	const FString TempFilename = FPaths::CreateTempFilename(*FPaths::ProjectIntermediateDir());
+	IPlatformFile& PlatformFile = FPlatformFileManager::Get().GetPlatformFile();
+	ON_SCOPE_EXIT
+	{
+		// Delete temp file
+		PlatformFile.DeleteFile(*TempFilename);
+	};
+
+	// Scratch data for testing
+	uint8 One = 1;
+	TArray<uint8> TestData;
+
+	// Check a new file can be created
+	if (TUniquePtr<IFileHandle> TestFile = TUniquePtr<IFileHandle>(PlatformFile.OpenWrite(*TempFilename, /*bAppend*/false, /*bAllowRead*/true)))
+	{
+		TestData.AddZeroed(64);
+
+		TestFile->Write(TestData.GetData(), TestData.Num());
+	}
+	else
+	{
+		REQUIRE_MESSAGE(FString::Printf(TEXT("File failed to open when new: %s"), *TempFilename), false);
+	}
+
+	// Confirm same data
+	{
+		TArray<uint8> ReadData;
+		REQUIRE_MESSAGE(FString::Printf(TEXT("File failed to load after writing: %s"), *TempFilename), FFileHelper::LoadFileToArray(ReadData, *TempFilename));
+		REQUIRE_MESSAGE(FString::Printf(TEXT("File data was incorrect after writing: %s"), *TempFilename), ReadData == TestData);
+	}
+
+	// Using append flag should open the file, and writing data immediately should append to the end.
+	// We should also be capable of seeking writing.
+	if (TUniquePtr<IFileHandle> TestFile = TUniquePtr<IFileHandle>(PlatformFile.OpenWrite(*TempFilename, /*bAppend*/true, /*bAllowRead*/true)))
+	{
+		// Validate the file actually opened in append mode correctly
+		{
+			const int64 ActualEOFPos = TestFile->Tell();
+			const int64 ExpectedEOFPos = TestFile->Size();
+			REQUIRE_MESSAGE(FString::Printf(TEXT("File did not seek to the end when opening (got %d, expected %d): %s"), ActualEOFPos, ExpectedEOFPos, *TempFilename), ActualEOFPos == ExpectedEOFPos);
+		}
+
+		TestData.Add(One);
+		TestData[10] = One;
+
+		TestFile->Write(&One, 1);
+		TestFile->Seek(10);
+		TestFile->Write(&One, 1);
+	}
+	else
+	{
+		REQUIRE_MESSAGE(FString::Printf(TEXT("File failed to open when appending: %s"), *TempFilename), false);
+	}
+
+	// Confirm same data
+	{
+		TArray<uint8> ReadData;
+		REQUIRE_MESSAGE(FString::Printf(TEXT("File failed to load after appending: %s"), *TempFilename), FFileHelper::LoadFileToArray(ReadData, *TempFilename));
+		REQUIRE_MESSAGE(FString::Printf(TEXT("File data was incorrect after appending: %s"), *TempFilename), ReadData == TestData);
+	}
+
+	// No append should clobber existing file
+	if (TUniquePtr<IFileHandle> TestFile = TUniquePtr<IFileHandle>(PlatformFile.OpenWrite(*TempFilename, /*bAppend*/false, /*bAllowRead*/true)))
+	{
+		TestData.Reset();
+		TestData.Add(One);
+
+		TestFile->Write(&One, 1);
+	}
+	else
+	{
+		REQUIRE_MESSAGE(FString::Printf(TEXT("File failed to open when clobbering: %s"), *TempFilename), false);
+	}
+
+	// Confirm same data
+	{
+		TArray<uint8> ReadData;
+		REQUIRE_MESSAGE(FString::Printf(TEXT("File failed to load after clobbering: %s"), *TempFilename), FFileHelper::LoadFileToArray(ReadData, *TempFilename));
+		REQUIRE_MESSAGE(FString::Printf(TEXT("File data was incorrect after clobbering: %s"), *TempFilename), ReadData == TestData);
+	}
+}
+
+TEST_CASE_NAMED(FFileShrinkBuffersTest, "System::Core::Misc::FileShrinkBuffers", "[.][EditorContext][CriticalPriority][EngineFilter]")
+{
+	const FString TempFilename = FPaths::CreateTempFilename(*FPaths::ProjectIntermediateDir());
+	IPlatformFile& PlatformFile = FPlatformFileManager::Get().GetPlatformFile();
+	ON_SCOPE_EXIT
+	{
+		// Delete temp file
+		PlatformFile.DeleteFile(*TempFilename);
+	};
+
+	// Scratch data for testing
+	uint8 One = 1;
+	TArray<uint8> TestData;
+
+	// Check a new file can be created
+	if (TUniquePtr<IFileHandle> TestFile = TUniquePtr<IFileHandle>(PlatformFile.OpenWrite(*TempFilename, /*bAppend*/false, /*bAllowRead*/true)))
+	{
+		for (uint8 i = 0; i < 64; ++i)
+		{
+			TestData.Add(i);
+		}
+
+		TestFile->Write(TestData.GetData(), TestData.Num());
+	}
+	else
+	{
+		REQUIRE_MESSAGE(FString::Printf(TEXT("File failed to open when new: %s"), *TempFilename), false);
+	}
+
+	// Confirm same data
+	{
+		TArray<uint8> ReadData;
+		REQUIRE_MESSAGE(FString::Printf(TEXT("File failed to load after writing: %s"), *TempFilename), FFileHelper::LoadFileToArray(ReadData, *TempFilename));
+		REQUIRE_MESSAGE(FString::Printf(TEXT("File data was incorrect after writing: %s"), *TempFilename), ReadData == TestData);
+	}
+
+	// Using ShrinkBuffers should not disrupt our read position in the file
+	if (TUniquePtr<IFileHandle> TestFile = TUniquePtr<IFileHandle>(PlatformFile.OpenRead(*TempFilename, /*bAllowWrite*/false)))
+	{
+		// Validate the file actually opened and is of the right size
+		CHECK_EQUALS(TEXT("File not of expected size at time of ShrinkBuffers read test"), TestFile->Size(), static_cast<decltype(TestFile->Size())>(TestData.Num()));
+
+		const int32 FirstHalfSize = TestData.Num() / 2;
+		const int32 SecondHalfSize = TestData.Num() - FirstHalfSize;
+
+		TArray<uint8> FirstHalfReadData;
+		FirstHalfReadData.AddUninitialized(FirstHalfSize);
+		CHECK_MESSAGE(TEXT("Failed to read first half of test file"), TestFile->Read(FirstHalfReadData.GetData(), FirstHalfReadData.Num()));
+
+		for (int32 i = 0; i < FirstHalfSize; ++i)
+		{
+			CHECK_EQUALS(TEXT("Mismatch in data before ShrinkBuffers was called"), FirstHalfReadData[i], TestData[i]);
+		}
+
+		TestFile->ShrinkBuffers();
+
+		TArray<uint8> SecondHalfReadData;
+		SecondHalfReadData.AddUninitialized(SecondHalfSize);
+		CHECK_MESSAGE(TEXT("Failed to read second half of test file"), TestFile->Read(SecondHalfReadData.GetData(), SecondHalfReadData.Num()));
+
+		for (int32 i = 0; i < SecondHalfSize; ++i)
+		{
+			CHECK_EQUALS(TEXT("Mismatch in data after ShrinkBuffers was called"), SecondHalfReadData[i], TestData[FirstHalfSize + i]);
+		}
+	}
+	else
+	{
+		REQUIRE_MESSAGE(FString::Printf(TEXT("File failed to open file for reading: %s"), *TempFilename), false);
+	}
+}
+
+#endif //WITH_TESTS
diff --git a/Engine/Source/Runtime/Core/Tests/Misc/GuidTest.cpp b/Engine/Source/Runtime/Core/Tests/Misc/GuidTest.cpp
new file mode 100644
--- /dev/null
+++ b/Engine/Source/Runtime/Core/Tests/Misc/GuidTest.cpp
@@ -0,0 +1,112 @@
+// Copyright Epic Games, Inc. All Rights Reserved.
+
+#if WITH_TESTS
+
+#include "CoreTypes.h"
+#include "Containers/UnrealString.h"
+#include "Misc/Guid.h"
+#include "Tests/TestHarnessAdapter.h"
+#include "Misc/SecureHash.h"
+#if WITH_LOW_LEVEL_TESTS
+#include "TestCommon/Comparisons.h"
+#endif
+
+TEST_CASE_NAMED(FGuidTest, "System::Core::Misc::Guid", "[ApplicationContextMask][SmokeFilter]")
+{
+	// Base36		12SRDE1NOQERISMPEF0H6KUU9
+	// Base64		EjRWeIdlQyESNFZ4h2VDIQ
+	// Hex			0x12345678 0x87654321  0x12345678 0x87654321
+	FGuid g = FGuid(305419896, 2271560481, 305419896, 2271560481);
+
+	// string conversions
+	CHECK_EQUALS(TEXT("String conversion (Default) must return EGuidFormats::Digits string"), g.ToString(), g.ToString(EGuidFormats::Digits));
+	TestEqual<FString>(TEXT("String conversion (EGuidFormats::Digits)"), g.ToString(EGuidFormats::Digits), TEXT("12345678876543211234567887654321"));
+	TestEqual<FString>(TEXT("String conversion (EGuidFormats::DigitsWithHyphens)"), g.ToString(EGuidFormats::DigitsWithHyphens), TEXT("12345678-8765-4321-1234-567887654321"));
+	TestEqual<FString>(TEXT("String conversion (EGuidFormats::DigitsWithHyphensInBraces)"), g.ToString(EGuidFormats::DigitsWithHyphensInBraces), TEXT("{12345678-8765-4321-1234-567887654321}"));
+	TestEqual<FString>(TEXT("String conversion (EGuidFormats::DigitsWithHyphensInParentheses)"), g.ToString(EGuidFormats::DigitsWithHyphensInParentheses), TEXT("(12345678-8765-4321-1234-567887654321)"));
+	TestEqual<FString>(TEXT("String conversion (EGuidFormats::HexValuesInBraces)"), g.ToString(EGuidFormats::HexValuesInBraces), TEXT("{0x12345678,0x8765,0x4321,{0x12,0x34,0x56,0x78,0x87,0x65,0x43,0x21}}"));
+	TestEqual<FString>(TEXT("String conversion (EGuidFormats::UniqueObjectGuid)"), g.ToString(EGuidFormats::UniqueObjectGuid), TEXT("12345678-87654321-12345678-87654321"));
+	TestEqual<FString>(TEXT("String conversion (EGuidFormats::Short)"), g.ToString(EGuidFormats::Short), TEXT("EjRWeIdlQyESNFZ4h2VDIQ"));
+	TestEqual<FString>(TEXT("String conversion (EGuidFormats::Base36Encoded)"), g.ToString(EGuidFormats::Base36Encoded), TEXT("12SRDE1NOQERISMPEF0H6KUU9"));
+
+	// parsing valid strings (exact)
+	FGuid g2_1;
+
+	CHECK_MESSAGE(TEXT("Parsing valid strings must succeed (EGuidFormats::Digits)"), FGuid::ParseExact(TEXT("12345678876543211234567887654321"), EGuidFormats::Digits, g2_1));
+	CHECK_EQUALS(TEXT("Parsing valid strings must succeed (EGuidFormats::Digits)"), g2_1, g);
+
+	CHECK_MESSAGE(TEXT("Parsing valid strings must succeed (EGuidFormats::DigitsWithHyphens)"), FGuid::ParseExact(TEXT("12345678-8765-4321-1234-567887654321"), EGuidFormats::DigitsWithHyphens, g2_1));
+	CHECK_EQUALS(TEXT("Parsing valid strings must succeed (EGuidFormats::DigitsWithHyphens)"), g2_1, g);
+
+	CHECK_MESSAGE(TEXT("Parsing valid strings must succeed (EGuidFormats::DigitsWithHyphensInBraces)"), FGuid::ParseExact(TEXT("{12345678-8765-4321-1234-567887654321}"), EGuidFormats::DigitsWithHyphensInBraces, g2_1));
+	CHECK_EQUALS(TEXT("Parsing valid strings must succeed (EGuidFormats::DigitsWithHyphensInBraces)"), g2_1, g);
+
+	CHECK_MESSAGE(TEXT("Parsing valid strings must succeed (EGuidFormats::DigitsWithHyphensInParentheses)"), FGuid::ParseExact(TEXT("(12345678-8765-4321-1234-567887654321)"), EGuidFormats::DigitsWithHyphensInParentheses, g2_1));
+	CHECK_EQUALS(TEXT("Parsing valid strings must succeed (EGuidFormats::DigitsWithHyphensInParentheses)"), g2_1, g);
+
+	CHECK_MESSAGE(TEXT("Parsing valid strings must succeed (EGuidFormats::HexValuesInBraces)"), FGuid::ParseExact(TEXT("{0x12345678,0x8765,0x4321,{0x12,0x34,0x56,0x78,0x87,0x65,0x43,0x21}}"), EGuidFormats::HexValuesInBraces, g2_1));
+	CHECK_EQUALS(TEXT("Parsing valid strings must succeed (EGuidFormats::HexValuesInBraces)"), g2_1, g);
+
+	CHECK_MESSAGE(TEXT("Parsing valid strings must succeed (EGuidFormats::UniqueObjectGuid)"), FGuid::ParseExact(TEXT("12345678-87654321-12345678-87654321"), EGuidFormats::UniqueObjectGuid, g2_1));
+	CHECK_EQUALS(TEXT("Parsing valid strings must succeed (EGuidFormats::UniqueObjectGuid)"), g2_1, g);
+
+	CHECK_MESSAGE(TEXT("Parsing valid strings must succeed (EGuidFormats::Short)"), FGuid::ParseExact(TEXT("EjRWeIdlQyESNFZ4h2VDIQ"), EGuidFormats::Short, g2_1));
+	CHECK_EQUALS(TEXT("Parsing valid strings must succeed (EGuidFormats::Short)"), g2_1, g);
+
+	CHECK_MESSAGE(TEXT("Parsing valid strings must succeed (EGuidFormats::Base36Encoded)"), FGuid::ParseExact(TEXT("12SRDE1NOQERISMPEF0H6KUU9"), EGuidFormats::Base36Encoded, g2_1));
+	CHECK_EQUALS(TEXT("Parsing valid strings must succeed (EGuidFormats::Base36Encoded)"), g2_1, g);
+
+	// parsing invalid strings (exact)
+
+
+	// parsing valid strings (automatic)
+	FGuid g3_1;
+
+	CHECK_MESSAGE(TEXT("Parsing valid strings must succeed (12345678876543211234567887654321)"), FGuid::Parse(TEXT("12345678876543211234567887654321"), g3_1));
+	CHECK_EQUALS(TEXT("Parsing valid strings must succeed (12345678876543211234567887654321)"), g3_1, g);
+
+	CHECK_MESSAGE(TEXT("Parsing valid strings must succeed (12345678-8765-4321-1234-567887654321)"), FGuid::Parse(TEXT("12345678-8765-4321-1234-567887654321"), g3_1));
+	CHECK_EQUALS(TEXT("Parsing valid strings must succeed (12345678-8765-4321-1234-567887654321)"), g3_1, g);
+
+	CHECK_MESSAGE(TEXT("Parsing valid strings must succeed ({12345678-8765-4321-1234-567887654321})"), FGuid::Parse(TEXT("{12345678-8765-4321-1234-567887654321}"), g3_1));
+	CHECK_EQUALS(TEXT("Parsing valid strings must succeed ({12345678-8765-4321-1234-567887654321})"), g3_1, g);
+
+	CHECK_MESSAGE(TEXT("Parsing valid strings must succeed ((12345678-8765-4321-1234-567887654321))"), FGuid::Parse(TEXT("(12345678-8765-4321-1234-567887654321)"), g3_1));
+	CHECK_EQUALS(TEXT("Parsing valid strings must succeed ((12345678-8765-4321-1234-567887654321))"), g3_1, g);
+
+	CHECK_MESSAGE(TEXT("Parsing valid strings must succeed ({0x12345678,0x8765,0x4321,{0x12,0x34,0x56,0x78,0x87,0x65,0x43,0x21}})"), FGuid::Parse(TEXT("{0x12345678,0x8765,0x4321,{0x12,0x34,0x56,0x78,0x87,0x65,0x43,0x21}}"), g3_1));
+	CHECK_EQUALS(TEXT("Parsing valid strings must succeed ({0x12345678,0x8765,0x4321,{0x12,0x34,0x56,0x78,0x87,0x65,0x43,0x21}})"), g3_1, g);
+
+	CHECK_MESSAGE(TEXT("Parsing valid strings must succeed (12345678-87654321-12345678-87654321)"), FGuid::Parse(TEXT("12345678-87654321-12345678-87654321"), g3_1));
+	CHECK_EQUALS(TEXT("Parsing valid strings must succeed (12345678-87654321-12345678-87654321)"), g3_1, g);
+
+	CHECK_MESSAGE(TEXT("Parsing valid strings must succeed (EjRWeIdlQyESNFZ4h2VDIQ)"), FGuid::Parse(TEXT("EjRWeIdlQyESNFZ4h2VDIQ"), g3_1));
+	CHECK_EQUALS(TEXT("Parsing valid strings must succeed (EjRWeIdlQyESNFZ4h2VDIQ)"), g3_1, g);
+
+	CHECK_MESSAGE(TEXT("Parsing valid strings must succeed (12SRDE1NOQERISMPEF0H6KUU9)"), FGuid::Parse(TEXT("12SRDE1NOQERISMPEF0H6KUU9"), g3_1));
+	CHECK_EQUALS(TEXT("Parsing valid strings must succeed (12SRDE1NOQERISMPEF0H6KUU9)"), g3_1, g);
+
+	//parsing invalid strings (automatic)
+
+	// GUID validation
+	FGuid g4_1 = FGuid::NewGuid();
+
+	CHECK_MESSAGE(TEXT("New GUIDs must be valid"), g4_1.IsValid());
+
+	g4_1.Invalidate();
+
+	CHECK_FALSE_MESSAGE(TEXT("Invalidated GUIDs must be invalid"), g4_1.IsValid());
+
+	// GUID creation from MD5 hash
+	FMD5 MD5;
+	ANSICHAR String[] = "Construct Guid from MD5 hash";
+	MD5.Update((uint8*)String, sizeof(String) - 1);
+	FMD5Hash Hash;
+	Hash.Set(MD5);
+	CHECK_MESSAGE(TEXT("Expected hash must match"), LexToString(Hash) == TEXT("e500d989a5db77bc696e79ee785f4a69"));
+
+	FGuid g5_1(MD5HashToGuid(Hash));
+	CHECK_MESSAGE(TEXT("Constructing from MD5 hash must be valid"), g5_1.ToString() == TEXT("89d900e5bc77dba5ee796e69694a5f78"));
+}
+
+#endif //WITH_TESTS
diff --git a/Engine/Source/Runtime/Core/Tests/Misc/PathTests.h b/Engine/Source/Runtime/Core/Tests/Misc/PathTests.h
new file mode 100644
--- /dev/null
+++ b/Engine/Source/Runtime/Core/Tests/Misc/PathTests.h
@@ -0,0 +1,138 @@
+// Copyright Epic Games, Inc. All Rights Reserved.
+
+#pragma once
+
+#include "Containers/StringView.h"
+#include "Tests/TestHarnessAdapter.h"
+#if WITH_LOW_LEVEL_TESTS
+#include "TestCommon/Comparisons.h"
+#endif
+
+#if WITH_TESTS
+
+template<class PathType, class StringType>
+void TestCollapseRelativeDirectories()
+{
+	auto Run = [](const TCHAR* Path, const TCHAR* Expected)
+	{
+		// Run test
+		StringType Actual;
+		Actual += Path;
+		const bool bValid = PathType::CollapseRelativeDirectories(Actual);
+
+		if (Expected)
+		{
+			// If we're looking for a result, make sure it was returned correctly
+			if (!bValid || FCString::Strcmp(*Actual, Expected) != 0)
+			{
+				FAIL_CHECK(FString::Printf(TEXT("Path '%s' failed to collapse correctly (got '%s', expected '%s')."), Path, *Actual, Expected));
+			}
+		}
+		else
+		{
+			// Otherwise, make sure it failed
+			if (bValid)
+			{
+				FAIL_CHECK(FString::Printf(TEXT("Path '%s' collapsed unexpectedly."), Path));
+			}
+		}
+	};
+
+	Run(TEXT(".."),                                                   nullptr);
+	Run(TEXT("/.."),                                                  nullptr);
+	Run(TEXT("./"),                                                   TEXT(""));
+	Run(TEXT("./file.txt"),                                           TEXT("file.txt"));
+	Run(TEXT("/."),                                                   TEXT("/."));
+	Run(TEXT("Folder"),                                               TEXT("Folder"));
+	Run(TEXT("/Folder"),                                              TEXT("/Folder"));
+	Run(TEXT("C:/Folder"),                                            TEXT("C:/Folder"));
+	Run(TEXT("C:/Folder/.."),                                         TEXT("C:"));
+	Run(TEXT("C:/Folder/../"),                                        TEXT("C:/"));
+	Run(TEXT("C:/Folder/../file.txt"),                                TEXT("C:/file.txt"));
+	Run(TEXT("Folder/.."),                                            TEXT(""));
+	Run(TEXT("Folder/../"),                                           TEXT("/"));
+	Run(TEXT("Folder/../file.txt"),                                   TEXT("/file.txt"));
+	Run(TEXT("/Folder/.."),                                           TEXT(""));
+	Run(TEXT("/Folder/../"),                                          TEXT("/"));
+	Run(TEXT("/Folder/../file.txt"),                                  TEXT("/file.txt"));
+	Run(TEXT("Folder/../.."),                                         nullptr);
+	Run(TEXT("Folder/../../"),                                        nullptr);
+	Run(TEXT("Folder/../../file.txt"),                                nullptr);
+	Run(TEXT("C:/.."),                                                nullptr);
+	Run(TEXT("C:/."),                                                 TEXT("C:/."));
+	Run(TEXT("C:/./"),                                                TEXT("C:/"));
+	Run(TEXT("C:/./file.txt"),                                        TEXT("C:/file.txt"));
+	Run(TEXT("C:/Folder1/../Folder2"),                                TEXT("C:/Folder2"));
+	Run(TEXT("C:/Folder1/../Folder2/"),                               TEXT("C:/Folder2/"));
+	Run(TEXT("C:/Folder1/../Folder2/file.txt"),                       TEXT("C:/Folder2/file.txt"));
+	Run(TEXT("C:/Folder1/../Folder2/../.."),                          nullptr);
+	Run(TEXT("C:/Folder1/../Folder2/../Folder3"),                     TEXT("C:/Folder3"));
+	Run(TEXT("C:/Folder1/../Folder2/../Folder3/"),                    TEXT("C:/Folder3/"));
+	Run(TEXT("C:/Folder1/../Folder2/../Folder3/file.txt"),            TEXT("C:/Folder3/file.txt"));
+	Run(TEXT("C:/Folder1/Folder2/../../Folder3"),                     TEXT("C:/Folder3"));
+	Run(TEXT("C:/Folder1/Folder2/../../Folder3/"),                    TEXT("C:/Folder3/"));
+	Run(TEXT("C:/Folder1/Folder2/../../Folder3/file.txt"),            TEXT("C:/Folder3/file.txt"));
+	Run(TEXT("C:/Folder1/Folder2/../../Folder3/../Folder4"),          TEXT("C:/Folder4"));
+	Run(TEXT("C:/Folder1/Folder2/../../Folder3/../Folder4/"),         TEXT("C:/Folder4/"));
+	Run(TEXT("C:/Folder1/Folder2/../../Folder3/../Folder4/file.txt"), TEXT("C:/Folder4/file.txt"));
+	Run(TEXT("C:/Folder1/Folder2/../Folder3/../../Folder4"),          TEXT("C:/Folder4"));
+	Run(TEXT("C:/Folder1/Folder2/../Folder3/../../Folder4/"),         TEXT("C:/Folder4/"));
+	Run(TEXT("C:/Folder1/Folder2/../Folder3/../../Folder4/file.txt"), TEXT("C:/Folder4/file.txt"));
+	Run(TEXT("C:/Folder1/Folder2/.././../Folder4"),                   TEXT("C:/Folder4"));
+	Run(TEXT("C:/Folder1/Folder2/.././../Folder4/"),                  TEXT("C:/Folder4/"));
+	Run(TEXT("C:/Folder1/Folder2/.././../Folder4/file.txt"),          TEXT("C:/Folder4/file.txt"));
+	Run(TEXT("C:/A/B/.././../C"),                                     TEXT("C:/C"));
+	Run(TEXT("C:/A/B/.././../C/"),                                    TEXT("C:/C/"));
+	Run(TEXT("C:/A/B/.././../C/file.txt"),                            TEXT("C:/C/file.txt"));
+	Run(TEXT(".svn"),                                                 TEXT(".svn"));
+	Run(TEXT("/.svn"),                                                TEXT("/.svn"));
+	Run(TEXT("./Folder/.svn"),                                        TEXT("Folder/.svn"));
+	Run(TEXT("./.svn/../.svn"),                                       TEXT(".svn"));
+	Run(TEXT(".svn/./.svn/.././../.svn"),                             TEXT("/.svn"));
+	Run(TEXT("C:/Folder1/./Folder2/..Folder3"),						 TEXT("C:/Folder1/Folder2/..Folder3"));
+	Run(TEXT("C:/Folder1/./Folder2/..Folder3/Folder4"),				 TEXT("C:/Folder1/Folder2/..Folder3/Folder4"));
+	Run(TEXT("C:/Folder1/./Folder2/..Folder3/..Folder4"),			 TEXT("C:/Folder1/Folder2/..Folder3/..Folder4"));
+	Run(TEXT("C:/Folder1/./Folder2/..Folder3/Folder4/../Folder5"),	 TEXT("C:/Folder1/Folder2/..Folder3/Folder5"));
+	Run(TEXT("C:/Folder1/..Folder2/Folder3/..Folder4/../Folder5"),	 TEXT("C:/Folder1/..Folder2/Folder3/Folder5"));
+}
+
+template<class PathType, class StringType>
+void TestRemoveDuplicateSlashes()
+{
+	auto Run = [&](const TCHAR* Path, const TCHAR* Expected)
+	{
+		StringType Actual;
+		Actual += Path;	
+		PathType::RemoveDuplicateSlashes(Actual);
+		CHECK_EQUALS(TEXT("RemoveDuplicateSlashes"), *Actual, Expected);
+	};
+
+	Run(TEXT(""), TEXT(""));
+	Run(TEXT("C:/Folder/File.txt"), TEXT("C:/Folder/File.txt"));
+	Run(TEXT("C:/Folder/File/"), TEXT("C:/Folder/File/"));
+	Run(TEXT("/"), TEXT("/"));
+	Run(TEXT("//"), TEXT("/"));
+	Run(TEXT("////"), TEXT("/"));
+	Run(TEXT("/Folder/File"), TEXT("/Folder/File"));
+	Run(TEXT("//Folder/File"), TEXT("/Folder/File")); // Don't use on //UNC paths; it will be stripped!
+	Run(TEXT("/////Folder//////File/////"), TEXT("/Folder/File/"));
+	Run(TEXT("\\\\Folder\\\\File\\\\"), TEXT("\\\\Folder\\\\File\\\\")); // It doesn't strip backslash, and we rely on that in some places
+	Run(TEXT("//\\\\//Folder//\\\\//File//\\\\//"), TEXT("/\\\\/Folder/\\\\/File/\\\\/"));
+}
+
+namespace PathTest
+{
+
+struct FTestPair
+{
+	FStringView Input;
+	FStringView Expected;
+};
+
+extern const FStringView BaseDir;
+
+extern const FTestPair ExpectedRelativeToAbsolutePaths[10];
+
+}
+
+#endif //WITH_TESTS
\ No newline at end of file
diff --git a/Engine/Source/Runtime/Core/Tests/Misc/PathViewsTest.cpp b/Engine/Source/Runtime/Core/Tests/Misc/PathViewsTest.cpp
new file mode 100644
--- /dev/null
+++ b/Engine/Source/Runtime/Core/Tests/Misc/PathViewsTest.cpp
@@ -0,0 +1,800 @@
+// Copyright Epic Games, Inc. All Rights Reserved.
+
+#include "PathTests.h"
+#include "Containers/StringView.h"
+#include "Misc/PathViews.h"
+#include "Tests/TestHarnessAdapter.h"
+#include "Misc/StringBuilder.h"
+#if WITH_LOW_LEVEL_TESTS
+#include "TestCommon/Comparisons.h"
+#endif
+
+#if WITH_TESTS
+
+class FPathViewsTest
+{
+public:
+
+	void TestViewTransform(FStringView (*InFunction)(const FStringView& InPath), const FStringView& InPath, const TCHAR* InExpected)
+	{
+		const FStringView Actual = InFunction(InPath);
+		if (Actual != InExpected)
+		{
+			FAIL_CHECK(FString::Printf(TEXT("Failed on path '%.*s' (got '%.*s', expected '%s')."),
+				InPath.Len(), InPath.GetData(), Actual.Len(), Actual.GetData(), InExpected));
+		}
+	}
+
+	static const TCHAR* BoolToString(bool bValue)
+	{
+		return bValue ? TEXT("true") : TEXT("false");
+	}
+
+	bool GetCleanFilenameTest()
+	{
+		auto RunGetCleanFilenameTest = [this](const TCHAR* InPath, const TCHAR* InExpected)
+		{
+			TestViewTransform(FPathViews::GetCleanFilename, InPath, InExpected);
+		};
+
+		RunGetCleanFilenameTest(TEXT(""), TEXT(""));
+		RunGetCleanFilenameTest(TEXT(".txt"), TEXT(".txt"));
+		RunGetCleanFilenameTest(TEXT(".tar.gz"), TEXT(".tar.gz"));
+		RunGetCleanFilenameTest(TEXT(".tar.gz/"), TEXT(""));
+		RunGetCleanFilenameTest(TEXT(".tar.gz\\"), TEXT(""));
+		RunGetCleanFilenameTest(TEXT("File"), TEXT("File"));
+		RunGetCleanFilenameTest(TEXT("File.tar.gz"), TEXT("File.tar.gz"));
+		RunGetCleanFilenameTest(TEXT("File.tar.gz/"), TEXT(""));
+		RunGetCleanFilenameTest(TEXT("File.tar.gz\\"), TEXT(""));
+		RunGetCleanFilenameTest(TEXT("C:/Folder/"), TEXT(""));
+		RunGetCleanFilenameTest(TEXT("C:/Folder/File"), TEXT("File"));
+		RunGetCleanFilenameTest(TEXT("C:/Folder/File.tar.gz"), TEXT("File.tar.gz"));
+		RunGetCleanFilenameTest(TEXT("C:/Folder/First.Last/File"), TEXT("File"));
+		RunGetCleanFilenameTest(TEXT("C:/Folder/First.Last/File.tar.gz"), TEXT("File.tar.gz"));
+		RunGetCleanFilenameTest(TEXT("C:\\Folder\\"), TEXT(""));
+		RunGetCleanFilenameTest(TEXT("C:\\Folder\\File"), TEXT("File"));
+		RunGetCleanFilenameTest(TEXT("C:\\Folder\\First.Last\\"), TEXT(""));
+		RunGetCleanFilenameTest(TEXT("C:\\Folder\\First.Last\\File"), TEXT("File"));
+		RunGetCleanFilenameTest(TEXT("C:\\Folder\\First.Last\\File.tar.gz"), TEXT("File.tar.gz"));
+		return true;
+	}
+
+	bool GetBaseFilenameTest()
+	{
+		auto RunGetBaseFilenameTest = [this](const TCHAR* InPath, const TCHAR* InExpected, const TCHAR* InExpectedWithPath)
+		{
+			const FStringView Path = InPath;
+			TestViewTransform(FPathViews::GetBaseFilename, Path, InExpected);
+			TestViewTransform(FPathViews::GetBaseFilenameWithPath, Path, InExpectedWithPath);
+		};
+
+		RunGetBaseFilenameTest(TEXT(""), TEXT(""), TEXT(""));
+		RunGetBaseFilenameTest(TEXT(".txt"), TEXT(""), TEXT(""));
+		RunGetBaseFilenameTest(TEXT(".tar.gz"), TEXT(".tar"), TEXT(".tar"));
+		RunGetBaseFilenameTest(TEXT(".tar.gz/"), TEXT(""), TEXT(".tar.gz/"));
+		RunGetBaseFilenameTest(TEXT(".tar.gz\\"), TEXT(""), TEXT(".tar.gz\\"));
+		RunGetBaseFilenameTest(TEXT("File"), TEXT("File"), TEXT("File"));
+		RunGetBaseFilenameTest(TEXT("File.txt"), TEXT("File"), TEXT("File"));
+		RunGetBaseFilenameTest(TEXT("File.tar.gz"), TEXT("File.tar"), TEXT("File.tar"));
+		RunGetBaseFilenameTest(TEXT("File.tar.gz/"), TEXT(""), TEXT("File.tar.gz/"));
+		RunGetBaseFilenameTest(TEXT("File.tar.gz\\"), TEXT(""), TEXT("File.tar.gz\\"));
+		RunGetBaseFilenameTest(TEXT("C:/Folder/"), TEXT(""), TEXT("C:/Folder/"));
+		RunGetBaseFilenameTest(TEXT("C:/Folder/File"), TEXT("File"), TEXT("C:/Folder/File"));
+		RunGetBaseFilenameTest(TEXT("C:/Folder/File.tar.gz"), TEXT("File.tar"), TEXT("C:/Folder/File.tar"));
+		RunGetBaseFilenameTest(TEXT("C:/Folder/First.Last/File"), TEXT("File"), TEXT("C:/Folder/First.Last/File"));
+		RunGetBaseFilenameTest(TEXT("C:/Folder/First.Last/File.txt"), TEXT("File"), TEXT("C:/Folder/First.Last/File"));
+		RunGetBaseFilenameTest(TEXT("C:/Folder/First.Last/File.tar.gz"), TEXT("File.tar"), TEXT("C:/Folder/First.Last/File.tar"));
+		RunGetBaseFilenameTest(TEXT("C:\\Folder\\"), TEXT(""), TEXT("C:\\Folder\\"));
+		RunGetBaseFilenameTest(TEXT("C:\\Folder\\File"), TEXT("File"), TEXT("C:\\Folder\\File"));
+		RunGetBaseFilenameTest(TEXT("C:\\Folder\\First.Last\\"), TEXT(""), TEXT("C:\\Folder\\First.Last\\"));
+		RunGetBaseFilenameTest(TEXT("C:\\Folder\\First.Last\\File"), TEXT("File"), TEXT("C:\\Folder\\First.Last\\File"));
+		RunGetBaseFilenameTest(TEXT("C:\\Folder\\First.Last\\File.txt"), TEXT("File"), TEXT("C:\\Folder\\First.Last\\File"));
+		RunGetBaseFilenameTest(TEXT("C:\\Folder\\First.Last\\File.tar.gz"), TEXT("File.tar"), TEXT("C:\\Folder\\First.Last\\File.tar"));
+
+		return true;
+	}
+
+	bool GetPathTest()
+	{
+		auto RunGetPathTest = [this](const TCHAR* InPath, const TCHAR* InExpected)
+		{
+			TestViewTransform(FPathViews::GetPath, InPath, InExpected);
+		};
+
+		RunGetPathTest(TEXT(""), TEXT(""));
+		RunGetPathTest(TEXT(".txt"), TEXT(""));
+		RunGetPathTest(TEXT(".tar.gz"), TEXT(""));
+		RunGetPathTest(TEXT(".tar.gz/"), TEXT(".tar.gz"));
+		RunGetPathTest(TEXT(".tar.gz\\"), TEXT(".tar.gz"));
+		RunGetPathTest(TEXT("File"), TEXT(""));
+		RunGetPathTest(TEXT("File.txt"), TEXT(""));
+		RunGetPathTest(TEXT("File.tar.gz"), TEXT(""));
+		RunGetPathTest(TEXT("File.tar.gz/"), TEXT("File.tar.gz"));
+		RunGetPathTest(TEXT("File.tar.gz\\"), TEXT("File.tar.gz"));
+		RunGetPathTest(TEXT("C:/Folder/"), TEXT("C:/Folder"));
+		RunGetPathTest(TEXT("C:/Folder/File"), TEXT("C:/Folder"));
+		RunGetPathTest(TEXT("C:/Folder/File.tar.gz"), TEXT("C:/Folder"));
+		RunGetPathTest(TEXT("C:/Folder/First.Last/File"), TEXT("C:/Folder/First.Last"));
+		RunGetPathTest(TEXT("C:/Folder/First.Last/File.tar.gz"), TEXT("C:/Folder/First.Last"));
+		RunGetPathTest(TEXT("C:\\Folder\\"), TEXT("C:\\Folder"));
+		RunGetPathTest(TEXT("C:\\Folder\\File"), TEXT("C:\\Folder"));
+		RunGetPathTest(TEXT("C:\\Folder\\First.Last\\"), TEXT("C:\\Folder\\First.Last"));
+		RunGetPathTest(TEXT("C:\\Folder\\First.Last\\File"), TEXT("C:\\Folder\\First.Last"));
+		RunGetPathTest(TEXT("C:\\Folder\\First.Last\\File.tar.gz"), TEXT("C:\\Folder\\First.Last"));
+
+		return true;
+	}
+
+	bool GetExtensionTest() 
+	{
+		auto RunGetExtensionTest = [this](const TCHAR* InPath, const TCHAR* InExpectedExt, const TCHAR* InExpectedExtDot)
+		{
+			const FStringView Path = InPath;
+			TestViewTransform([](const FStringView& InPath2) { return FPathViews::GetExtension(InPath2, /*bIncludeDot*/ false); }, Path, InExpectedExt);
+			TestViewTransform([](const FStringView& InPath2) { return FPathViews::GetExtension(InPath2, /*bIncludeDot*/ true); }, Path, InExpectedExtDot);
+		};
+
+		RunGetExtensionTest(TEXT(""), TEXT(""), TEXT(""));
+		RunGetExtensionTest(TEXT(".txt"), TEXT("txt"), TEXT(".txt"));
+		RunGetExtensionTest(TEXT(".tar.gz"), TEXT("gz"), TEXT(".gz"));
+		RunGetExtensionTest(TEXT(".tar.gz/"), TEXT(""), TEXT(""));
+		RunGetExtensionTest(TEXT(".tar.gz\\"), TEXT(""), TEXT(""));
+		RunGetExtensionTest(TEXT("File"), TEXT(""), TEXT(""));
+		RunGetExtensionTest(TEXT("File.txt"), TEXT("txt"), TEXT(".txt"));
+		RunGetExtensionTest(TEXT("File.tar.gz"), TEXT("gz"), TEXT(".gz"));
+		RunGetExtensionTest(TEXT("File.tar.gz/"), TEXT(""), TEXT(""));
+		RunGetExtensionTest(TEXT("File.tar.gz\\"), TEXT(""), TEXT(""));
+		RunGetExtensionTest(TEXT("C:/Folder/File"), TEXT(""), TEXT(""));
+		RunGetExtensionTest(TEXT("C:\\Folder\\File"), TEXT(""), TEXT(""));
+		RunGetExtensionTest(TEXT("C:/Folder/File.txt"), TEXT("txt"), TEXT(".txt"));
+		RunGetExtensionTest(TEXT("C:\\Folder\\File.txt"), TEXT("txt"), TEXT(".txt"));
+		RunGetExtensionTest(TEXT("C:/Folder/File.tar.gz"), TEXT("gz"), TEXT(".gz"));
+		RunGetExtensionTest(TEXT("C:\\Folder\\File.tar.gz"), TEXT("gz"), TEXT(".gz"));
+		RunGetExtensionTest(TEXT("C:/Folder/First.Last/File"), TEXT(""), TEXT(""));
+		RunGetExtensionTest(TEXT("C:\\Folder\\First.Last\\File"), TEXT(""), TEXT(""));
+		RunGetExtensionTest(TEXT("C:/Folder/First.Last/File.txt"), TEXT("txt"), TEXT(".txt"));
+		RunGetExtensionTest(TEXT("C:\\Folder\\First.Last\\File.txt"), TEXT("txt"), TEXT(".txt"));
+		RunGetExtensionTest(TEXT("C:/Folder/First.Last/File.tar.gz"), TEXT("gz"), TEXT(".gz"));
+		RunGetExtensionTest(TEXT("C:\\Folder\\First.Last\\File.tar.gz"), TEXT("gz"), TEXT(".gz"));
+
+		return true;
+	}
+
+	bool GetPathLeafTest()
+	{
+		auto RunGetPathLeafTest = [this](const TCHAR* InPath, const TCHAR* InExpected)
+		{
+			TestViewTransform(FPathViews::GetPathLeaf, InPath, InExpected);
+		};
+
+		RunGetPathLeafTest(TEXT(""), TEXT(""));
+		RunGetPathLeafTest(TEXT(".txt"), TEXT(".txt"));
+		RunGetPathLeafTest(TEXT(".tar.gz"), TEXT(".tar.gz"));
+		RunGetPathLeafTest(TEXT(".tar.gz/"), TEXT(".tar.gz"));
+		RunGetPathLeafTest(TEXT(".tar.gz\\"), TEXT(".tar.gz"));
+		RunGetPathLeafTest(TEXT("File"), TEXT("File"));
+		RunGetPathLeafTest(TEXT("File.txt"), TEXT("File.txt"));
+		RunGetPathLeafTest(TEXT("File.tar.gz"), TEXT("File.tar.gz"));
+		RunGetPathLeafTest(TEXT("File.tar.gz/"), TEXT("File.tar.gz"));
+		RunGetPathLeafTest(TEXT("File.tar.gz\\"), TEXT("File.tar.gz"));
+		RunGetPathLeafTest(TEXT("C:/Folder/"), TEXT("Folder"));
+		RunGetPathLeafTest(TEXT("C:/Folder/File"), TEXT("File"));
+		RunGetPathLeafTest(TEXT("C:/Folder/File.tar.gz"), TEXT("File.tar.gz"));
+		RunGetPathLeafTest(TEXT("C:/Folder/First.Last/File"), TEXT("File"));
+		RunGetPathLeafTest(TEXT("C:/Folder/First.Last/File.tar.gz"), TEXT("File.tar.gz"));
+		RunGetPathLeafTest(TEXT("C:\\Folder\\"), TEXT("Folder"));
+		RunGetPathLeafTest(TEXT("C:\\Folder\\File"), TEXT("File"));
+		RunGetPathLeafTest(TEXT("C:\\Folder\\First.Last\\"), TEXT("First.Last"));
+		RunGetPathLeafTest(TEXT("C:\\Folder\\First.Last\\File"), TEXT("File"));
+		RunGetPathLeafTest(TEXT("C:\\Folder\\First.Last\\File.tar.gz"), TEXT("File.tar.gz"));
+
+		return true;
+	}
+
+	bool IsPathLeafTest()
+	{
+		auto RunIsPathLeafTest = [this](const TCHAR* InPath, bool bExpected)
+		{
+			bool bResult = FPathViews::IsPathLeaf(InPath);
+			if (bResult != bExpected)
+			{
+				FAIL_CHECK(FString::Printf(TEXT("GetPathLeaf(\"%s\") == %s, expected %s."), InPath, BoolToString(bResult), BoolToString(bExpected)));
+			}
+		};
+
+		RunIsPathLeafTest(TEXT(""), true);
+		RunIsPathLeafTest(TEXT(".txt"), true);
+		RunIsPathLeafTest(TEXT(".tar.gz"), true);
+		RunIsPathLeafTest(TEXT(".tar.gz/"), true);
+		RunIsPathLeafTest(TEXT(".tar.gz\\"), true);
+		RunIsPathLeafTest(TEXT("File"), true);
+		RunIsPathLeafTest(TEXT("File.txt"), true);
+		RunIsPathLeafTest(TEXT("File.tar.gz"), true);
+		RunIsPathLeafTest(TEXT("File.tar.gz/"), true);
+		RunIsPathLeafTest(TEXT("File.tar.gz\\"), true);
+		RunIsPathLeafTest(TEXT("//"), true);
+		RunIsPathLeafTest(TEXT("\\\\"), true);
+		RunIsPathLeafTest(TEXT("/"), true);
+		RunIsPathLeafTest(TEXT("\\"), true);
+		RunIsPathLeafTest(TEXT("C:/"), true);
+		RunIsPathLeafTest(TEXT("C:\\"), true);
+		RunIsPathLeafTest(TEXT("//Folder"), false);
+		RunIsPathLeafTest(TEXT("\\\\Folder"), false);
+		RunIsPathLeafTest(TEXT("/Folder"), false);
+		RunIsPathLeafTest(TEXT("\\Folder"), false);
+		RunIsPathLeafTest(TEXT("C:/Folder"), false);
+		RunIsPathLeafTest(TEXT("C:\\Folder"), false);
+		RunIsPathLeafTest(TEXT("C:/Folder"), false);
+		RunIsPathLeafTest(TEXT("C:\\Folder"), false);
+
+		return true;
+	}
+
+	bool EqualsAndLessTest()
+	{
+		auto RunEqualsLessTest = [this](const TCHAR* A, const TCHAR* B, int Expected)
+		{
+			bool bEqual = FPathViews::Equals(A, B);
+			bool bALessThanB = FPathViews::Less(A, B);
+			bool bBLessThanA = FPathViews::Less(B, A);
+			if (bEqual != (Expected == 0))
+			{
+				FAIL_CHECK(FString::Printf(TEXT("Equals(%s,%s) == %s, expected %s"), A, B, BoolToString(bEqual), BoolToString(Expected == 0)));
+			}
+			if (bALessThanB != Expected < 0)
+			{
+				FAIL_CHECK(FString::Printf(TEXT("Less(%s,%s) == %s, expected %s"), A, B, BoolToString(bALessThanB), BoolToString(Expected < 0)));
+			}
+			if (bBLessThanA != Expected > 0)
+			{
+				FAIL_CHECK(FString::Printf(TEXT("Less(%s,%s) == %s, expected %s"), B, A, BoolToString(bBLessThanA), BoolToString(Expected > 0)));
+			}
+		};
+
+		RunEqualsLessTest(TEXT("A"), TEXT("B"), -1);
+		RunEqualsLessTest(TEXT("A"), TEXT("b"), -1);
+		RunEqualsLessTest(TEXT("a"), TEXT("B"), -1);
+		RunEqualsLessTest(TEXT("A"), TEXT("A"), 0);
+		RunEqualsLessTest(TEXT("A"), TEXT("a/"), 0);
+		RunEqualsLessTest(TEXT("A"), TEXT("a\\"), 0);
+		RunEqualsLessTest(TEXT("A"), TEXT("abby"), -1);
+		RunEqualsLessTest(TEXT("a"), TEXT("Abby"), -1);
+		RunEqualsLessTest(TEXT("\\A/B"), TEXT("/A\\B/"), 0);
+		//RunEqualsLessTest(TEXT("../../../Engine"), TEXT("C:/Engine"), 0); // Detecting relpath == abspath is not yet implemented
+		//RunEqualsLessTest(TEXT("C:/A/B"), TEXT("C:/A/../A/./B"), 0); // Collapsing .. and . is not yet implemented
+		RunEqualsLessTest(TEXT("/"), TEXT("/"), 0);
+		RunEqualsLessTest(TEXT("/"), TEXT("//"), -1);
+		RunEqualsLessTest(TEXT("/"), TEXT("C:/"), -1); // int('/') == 47 < int('C') == 67
+		RunEqualsLessTest(TEXT("/"), TEXT("A"), -1);
+		RunEqualsLessTest(TEXT("//"), TEXT("//"), 0);
+		RunEqualsLessTest(TEXT("//"), TEXT("C:/"), -1);
+		RunEqualsLessTest(TEXT("//"), TEXT("A"), -1);
+		RunEqualsLessTest(TEXT("C:/"), TEXT("C:/"), 0);
+		RunEqualsLessTest(TEXT("C:/"), TEXT("C"), 1);
+
+		return true;
+	}
+
+	bool TryMakeChildPathRelativeToTest()
+	{
+		auto RunRelChildTest = [this](const TCHAR* Child, const TCHAR* Parent, bool bExpectedIsChild, const TCHAR* ExpectedRelPath)
+		{
+			FStringView ActualRelPath;
+			bool bActualIsChild = FPathViews::TryMakeChildPathRelativeTo(Child, Parent, ActualRelPath);
+			bool bActualIsParent = FPathViews::IsParentPathOf(Parent, Child);
+			if (bExpectedIsChild != bActualIsChild || !ActualRelPath.Equals(ExpectedRelPath, ESearchCase::IgnoreCase))
+			{
+				FAIL_CHECK(FString::Printf(TEXT("TryMakeChildPathRelativeTo(\"%s\", \"%s\") returned (%s, \"%.*s\"), expected (%s, \"%s\")."), Child, Parent,
+					BoolToString(bActualIsChild), ActualRelPath.Len(), ActualRelPath.GetData(), BoolToString(bExpectedIsChild), ExpectedRelPath));
+			}
+			if (bActualIsParent != bExpectedIsChild)
+			{
+				FAIL_CHECK(FString::Printf(TEXT("IsParentPathOf(\"%s\", \"%s\") returned %s, expected %s."), Child, Parent,
+					BoolToString(bActualIsParent), BoolToString(bExpectedIsChild)));
+			}
+		};
+
+		RunRelChildTest(TEXT("A/B"), TEXT("A"), true, TEXT("B"));
+		RunRelChildTest(TEXT("A\\B"), TEXT("A"), true, TEXT("B"));
+		RunRelChildTest(TEXT("A/B"), TEXT("A/"), true, TEXT("B"));
+		RunRelChildTest(TEXT("A\\B"), TEXT("A/"), true, TEXT("B"));
+		RunRelChildTest(TEXT("A/B"), TEXT("A\\"), true, TEXT("B"));
+		RunRelChildTest(TEXT("A\\B"), TEXT("A\\"), true, TEXT("B"));
+		RunRelChildTest(TEXT("A"), TEXT("A"), true, TEXT(""));
+		RunRelChildTest(TEXT("A/"), TEXT("A"), true, TEXT(""));
+		RunRelChildTest(TEXT("A\\"), TEXT("A"), true, TEXT(""));
+		RunRelChildTest(TEXT("A"), TEXT("A/"), true, TEXT(""));
+		RunRelChildTest(TEXT("A"), TEXT("A\\"), true, TEXT(""));
+		RunRelChildTest(TEXT("../A"), TEXT(".."), true, TEXT("A"));
+		RunRelChildTest(TEXT("/A/B"), TEXT("/A"), true, TEXT("B"));
+		RunRelChildTest(TEXT("../A/B"), TEXT("../A"), true, TEXT("B"));
+		RunRelChildTest(TEXT("../"), TEXT(".."), true, TEXT(""));
+		RunRelChildTest(TEXT("C:/"), TEXT("C:/"), true, TEXT(""));
+		RunRelChildTest(TEXT("C:/A"), TEXT("C:/"), true, TEXT("A"));
+		RunRelChildTest(TEXT("//A"), TEXT("//A"), true, TEXT(""));
+		RunRelChildTest(TEXT("//A/"), TEXT("//A"), true, TEXT(""));
+		RunRelChildTest(TEXT("//A"), TEXT("//A/"), true, TEXT(""));
+		RunRelChildTest(TEXT("//A/"), TEXT("//A/"), true, TEXT(""));
+		RunRelChildTest(TEXT("//A/B"), TEXT("//A"), true, TEXT("B"));
+		RunRelChildTest(TEXT("//A/B/"), TEXT("//A"), true, TEXT("B/"));
+		RunRelChildTest(TEXT("//A/B"), TEXT("//A/"), true, TEXT("B"));
+		RunRelChildTest(TEXT("//A/B/"), TEXT("//A/"), true, TEXT("B/"));
+
+		RunRelChildTest(TEXT("//A/BFoo"), TEXT("//A/B"), false, TEXT(""));
+		RunRelChildTest(TEXT("//A/C"), TEXT("//A/B"), false, TEXT(""));
+		RunRelChildTest(TEXT("//A/C"), TEXT("C:/A"), false, TEXT(""));
+		RunRelChildTest(TEXT("//A/C"), TEXT("/A"), false, TEXT(""));
+		RunRelChildTest(TEXT("C:/A/BFoo"), TEXT("C:/A/B"), false, TEXT(""));
+		RunRelChildTest(TEXT("C:/A/C"), TEXT("C:/A/B"), false, TEXT(""));
+		RunRelChildTest(TEXT("C:/A/C"), TEXT("//A"), false, TEXT(""));
+		RunRelChildTest(TEXT("C:/A/C"), TEXT("/A"), false, TEXT(""));
+		RunRelChildTest(TEXT("/A/BFoo"), TEXT("/A/B"), false, TEXT(""));
+		RunRelChildTest(TEXT("/A/C"), TEXT("/A/B"), false, TEXT(""));
+		RunRelChildTest(TEXT("/A/C"), TEXT("C:/A"), false, TEXT(""));
+		RunRelChildTest(TEXT("/A/C"), TEXT("//A"), false, TEXT(""));
+
+		RunRelChildTest(TEXT("/"), TEXT("/"), true, TEXT(""));
+		RunRelChildTest(TEXT("/"), TEXT("//"), false, TEXT(""));
+		RunRelChildTest(TEXT("/"), TEXT("C"), false, TEXT(""));
+		RunRelChildTest(TEXT("/"), TEXT("C:/"), false, TEXT(""));
+		RunRelChildTest(TEXT("//"), TEXT("/"), false, TEXT(""));
+		RunRelChildTest(TEXT("//"), TEXT("//"), true, TEXT(""));
+		RunRelChildTest(TEXT("//"), TEXT("C"), false, TEXT(""));
+		RunRelChildTest(TEXT("//"), TEXT("C:/"), false, TEXT(""));
+		RunRelChildTest(TEXT("C"), TEXT("/"), false, TEXT(""));
+		RunRelChildTest(TEXT("C"), TEXT("//"), false, TEXT(""));
+		RunRelChildTest(TEXT("C"), TEXT("C"), true, TEXT(""));
+		RunRelChildTest(TEXT("C"), TEXT("C:/"), false, TEXT(""));
+		RunRelChildTest(TEXT("C:/"), TEXT("/"), false, TEXT(""));
+		RunRelChildTest(TEXT("C:/"), TEXT("//"), false, TEXT(""));
+		RunRelChildTest(TEXT("C:/"), TEXT("C"), false, TEXT(""));
+		RunRelChildTest(TEXT("C:/"), TEXT("C:/"), true, TEXT(""));
+
+		RunRelChildTest(TEXT("C:/A/B/C"), TEXT("C:/A"), true, TEXT("B/C"));
+		RunRelChildTest(TEXT("C:/A"), TEXT("C:/A/B/C"), false, TEXT(""));
+
+		// Correctly handle paths with invalid duplicate slashes
+		RunRelChildTest(TEXT("///A"), TEXT("///"), true, TEXT("A"));
+		RunRelChildTest(TEXT("////////////A"), TEXT("///"), true, TEXT("A"));
+		RunRelChildTest(TEXT("///A"), TEXT("//"), true, TEXT("A"));
+		RunRelChildTest(TEXT("////////////A"), TEXT("/"), false, TEXT(""));
+		RunRelChildTest(TEXT("C://A"), TEXT("C://"), true, TEXT("A"));
+		RunRelChildTest(TEXT("C://////////A"), TEXT("C://"), true, TEXT("A"));
+		RunRelChildTest(TEXT("C://A"), TEXT("C:/"), true, TEXT("A"));
+		RunRelChildTest(TEXT("C://////////A"), TEXT("C:/"), true, TEXT("A"));
+		RunRelChildTest(TEXT("A//B"), TEXT("A//"), true, TEXT("B"));
+		RunRelChildTest(TEXT("A///////////B"), TEXT("A//"), true, TEXT("B"));
+		RunRelChildTest(TEXT("A//B"), TEXT("A/"), true, TEXT("B"));
+		RunRelChildTest(TEXT("A///////////B"), TEXT("A/"), true, TEXT("B"));
+		RunRelChildTest(TEXT("A//B"), TEXT("A"), true, TEXT("B"));
+		RunRelChildTest(TEXT("A///////////B"), TEXT("A"), true, TEXT("B"));
+
+		return true;
+	}
+
+	bool IsRelativePathTest()
+	{
+		auto RunRelTest = [this](const TCHAR* A, bool bExpected)
+		{
+			bool bActual = FPathViews::IsRelativePath(A);
+			if (bActual != bExpected)
+			{
+				FAIL_CHECK(FString::Printf(TEXT("IsRelativePath(\"%s\") == %s, expected %s."), A, BoolToString(bActual), BoolToString(bExpected)));
+			}
+		};
+
+		RunRelTest(TEXT("A"), true);
+		RunRelTest(TEXT("A/"), true);
+		RunRelTest(TEXT("A/B"), true);
+		RunRelTest(TEXT("A\\B"), true);
+		RunRelTest(TEXT("/A"), false);
+		RunRelTest(TEXT("\\A"), false);
+		RunRelTest(TEXT("/A/B"), false);
+		RunRelTest(TEXT("//A"), false);
+		RunRelTest(TEXT("\\\\A"), false);
+		RunRelTest(TEXT("//A/B"), false);
+		RunRelTest(TEXT("C:/A"), false);
+		RunRelTest(TEXT("C:\\A"), false);
+		RunRelTest(TEXT("C:/A/B"), false);
+
+		return true;
+	}
+};
+
+TEST_CASE_NAMED(FPathViewsCollapseDirectoriesTest, "System::Core::Misc::PathViews::CollapseDirectories", "[ApplicationContextMask][SmokeFilter]")
+{
+	TestCollapseRelativeDirectories<FPathViews, TStringBuilder<64>>();
+}
+
+TEST_CASE_NAMED(FPathViewsRemoveDuplicateSlashesTest, "System::Core::Misc::PathViews::RemoveDuplicateSlashes", "[ApplicationContextMask][SmokeFilter]")
+{
+	TestRemoveDuplicateSlashes<FPathViews, TStringBuilder<64>>();
+}
+
+TEST_CASE_NAMED(FPathViewsGetCleanFilenameTest, "System::Core::Misc::PathViews::GetCleanFilename", "[ApplicationContextMask][SmokeFilter]")
+{
+	FPathViewsTest Instance = FPathViewsTest();
+	Instance.GetCleanFilenameTest();	
+}
+
+TEST_CASE_NAMED(FPathViewsGetBaseFilenameTest, "System::Core::Misc::PathViews::GetBaseFilename", "[ApplicationContextMask][SmokeFilter]")
+{
+	FPathViewsTest Instance = FPathViewsTest();
+	Instance.GetBaseFilenameTest();
+}
+
+TEST_CASE_NAMED(FPathViewsGetPathTest, "System::Core::Misc::PathViews::GetPath", "[ApplicationContextMask][SmokeFilter]")
+{
+	FPathViewsTest Instance = FPathViewsTest();
+	Instance.GetPathTest();
+}
+
+TEST_CASE_NAMED(FPathViewsGetExtensionTest, "System::Core::Misc::PathViews::GetExtension", "[ApplicationContextMask][SmokeFilter]")
+{
+	FPathViewsTest Instance = FPathViewsTest();
+	Instance.GetExtensionTest();
+}
+
+TEST_CASE_NAMED(FPathViewsGetPathLeafTest, "System::Core::Misc::PathViews::GetPathLeaf", "[ApplicationContextMask][SmokeFilter]")
+{
+	FPathViewsTest Instance = FPathViewsTest();
+	Instance.GetPathLeafTest();
+}
+
+TEST_CASE_NAMED(FPathViewsIsPathLeafTest, "System::Core::Misc::PathViews::IsPathLeaf", "[ApplicationContextMask][SmokeFilter]")
+{
+	FPathViewsTest Instance = FPathViewsTest();
+	Instance.IsPathLeafTest();
+}
+
+TEST_CASE_NAMED(FPathViewsSplitTest, "System::Core::Misc::PathViews::Split", "[ApplicationContextMask][SmokeFilter]")
+{
+	auto RunSplitTest = [](const TCHAR* InPath, const TCHAR* InExpectedPath, const TCHAR* InExpectedName, const TCHAR* InExpectedExt)
+	{
+		FStringView SplitPath, SplitName, SplitExt;
+		FPathViews::Split(InPath, SplitPath, SplitName, SplitExt);
+		if (SplitPath != InExpectedPath || SplitName != InExpectedName || SplitExt != InExpectedExt)
+		{
+			FAIL_CHECK(FString::Printf(TEXT("Failed to split path '%s' (got ('%.*s', '%.*s', '%.*s'), expected ('%s', '%s', '%s'))."), InPath,
+				SplitPath.Len(), SplitPath.GetData(), SplitName.Len(), SplitName.GetData(), SplitExt.Len(), SplitExt.GetData(),
+				InExpectedPath, InExpectedName, InExpectedExt));
+		}
+	};
+
+	RunSplitTest(TEXT(""), TEXT(""), TEXT(""), TEXT(""));
+	RunSplitTest(TEXT(".txt"), TEXT(""), TEXT(""), TEXT("txt"));
+	RunSplitTest(TEXT(".tar.gz"), TEXT(""), TEXT(".tar"), TEXT("gz"));
+	RunSplitTest(TEXT(".tar.gz/"), TEXT(".tar.gz"), TEXT(""), TEXT(""));
+	RunSplitTest(TEXT(".tar.gz\\"), TEXT(".tar.gz"), TEXT(""), TEXT(""));
+	RunSplitTest(TEXT("File"), TEXT(""), TEXT("File"), TEXT(""));
+	RunSplitTest(TEXT("File.txt"), TEXT(""), TEXT("File"), TEXT("txt"));
+	RunSplitTest(TEXT("File.tar.gz"), TEXT(""), TEXT("File.tar"), TEXT("gz"));
+	RunSplitTest(TEXT("File.tar.gz/"), TEXT("File.tar.gz"), TEXT(""), TEXT(""));
+	RunSplitTest(TEXT("File.tar.gz\\"), TEXT("File.tar.gz"), TEXT(""), TEXT(""));
+	RunSplitTest(TEXT("C:/Folder/"), TEXT("C:/Folder"), TEXT(""), TEXT(""));
+	RunSplitTest(TEXT("C:/Folder/File"), TEXT("C:/Folder"), TEXT("File"), TEXT(""));
+	RunSplitTest(TEXT("C:/Folder/File.txt"), TEXT("C:/Folder"), TEXT("File"), TEXT("txt"));
+	RunSplitTest(TEXT("C:/Folder/File.tar.gz"), TEXT("C:/Folder"), TEXT("File.tar"), TEXT("gz"));
+	RunSplitTest(TEXT("C:/Folder/First.Last/File"), TEXT("C:/Folder/First.Last"), TEXT("File"), TEXT(""));
+	RunSplitTest(TEXT("C:/Folder/First.Last/File.txt"), TEXT("C:/Folder/First.Last"), TEXT("File"), TEXT("txt"));
+	RunSplitTest(TEXT("C:/Folder/First.Last/File.tar.gz"), TEXT("C:/Folder/First.Last"), TEXT("File.tar"), TEXT("gz"));
+	RunSplitTest(TEXT("C:\\Folder\\"), TEXT("C:\\Folder"), TEXT(""), TEXT(""));
+	RunSplitTest(TEXT("C:\\Folder\\File"), TEXT("C:\\Folder"), TEXT("File"), TEXT(""));
+	RunSplitTest(TEXT("C:\\Folder\\First.Last\\"), TEXT("C:\\Folder\\First.Last"), TEXT(""), TEXT(""));
+	RunSplitTest(TEXT("C:\\Folder\\First.Last\\File"), TEXT("C:\\Folder\\First.Last"), TEXT("File"), TEXT(""));
+	RunSplitTest(TEXT("C:\\Folder\\First.Last\\File.txt"), TEXT("C:\\Folder\\First.Last"), TEXT("File"), TEXT("txt"));
+	RunSplitTest(TEXT("C:\\Folder\\First.Last\\File.tar.gz"), TEXT("C:\\Folder\\First.Last"), TEXT("File.tar"), TEXT("gz"));
+}
+
+TEST_CASE_NAMED(FPathViewsAppendTest, "System::Core::Misc::PathViews::Append", "[ApplicationContextMask][SmokeFilter]")
+{
+	TStringBuilder<256> Path;
+
+	FPathViews::Append(Path, TEXT("A"), TEXT(""));
+	CHECK_EQUALS(TEXT("FPathViews::Append('A', '')"), FStringView(Path), TEXTVIEW("A/"));
+	Path.Reset();
+
+	FPathViews::Append(Path, TEXT(""), TEXT("B"));
+	CHECK_EQUALS(TEXT("FPathViews::Append('', 'B')"), FStringView(Path), TEXTVIEW("B"));
+	Path.Reset();
+
+	FPathViews::Append(Path, TEXT("/"), TEXT("B"));
+	CHECK_EQUALS(TEXT("FPathViews::Append('/', 'B')"), FStringView(Path), TEXTVIEW("/B"));
+	Path.Reset();
+
+	FPathViews::Append(Path, TEXT("A"), TEXT("B"));
+	CHECK_EQUALS(TEXT("FPathViews::Append('A', 'B')"), FStringView(Path), TEXTVIEW("A/B"));
+	Path.Reset();
+
+	FPathViews::Append(Path, TEXT("A/"), TEXT("B"));
+	CHECK_EQUALS(TEXT("FPathViews::Append('A/', 'B')"), FStringView(Path), TEXTVIEW("A/B"));
+	Path.Reset();
+
+	FPathViews::Append(Path, TEXT("A\\"), TEXT("B"));
+	CHECK_EQUALS(TEXT("FPathViews::Append('A\\', 'B')"), FStringView(Path), TEXTVIEW("A\\B"));
+	Path.Reset();
+
+	FPathViews::Append(Path, TEXT("A/B"), TEXT("C/D"));
+	CHECK_EQUALS(TEXT("FPathViews::Append('A/B', 'C/D')"), FStringView(Path), TEXTVIEW("A/B/C/D"));
+	Path.Reset();
+
+	FPathViews::Append(Path, TEXT("A/"), TEXT("B"), TEXT("C/"), TEXT("D"));
+	CHECK_EQUALS(TEXT("FPathViews::Append('A/', 'B', 'C/', 'D')"), FStringView(Path), TEXTVIEW("A/B/C/D"));
+	Path.Reset();
+
+	FPathViews::Append(Path, TEXT("A/"), 16, TEXT("B"));
+	CHECK_EQUALS(TEXT("FPathViews::Append('A/', 16, 'B')"), FStringView(Path), TEXTVIEW("A/16/B"));
+	Path.Reset();
+}
+
+TEST_CASE_NAMED(FPathViewsChangeExtensionTest, "System::Core::Misc::PathViews::ChangeExtension", "[ApplicationContextMask][SmokeFilter]")
+{
+	auto RunChangeExtensionTest = [](const TCHAR* InPath, const TCHAR* InNewExt, const TCHAR* InExpectedPath)
+	{
+		// Run test
+		const FString NewPath = FPathViews::ChangeExtension(InPath, InNewExt);
+		if (NewPath != InExpectedPath)
+		{
+			FAIL_CHECK(FString::Printf(TEXT("Path '%s' failed to change the extension (got '%s', expected '%s')."), InPath, *NewPath, InExpectedPath));
+		}
+	};
+
+	RunChangeExtensionTest(nullptr, nullptr, TEXT(""));
+	RunChangeExtensionTest(TEXT(""), TEXT(""), TEXT(""));
+	RunChangeExtensionTest(TEXT(""), TEXT(".txt"), TEXT(""));
+	RunChangeExtensionTest(TEXT("file"), TEXT("log"), TEXT("file"));
+	RunChangeExtensionTest(TEXT("file.txt"), TEXT("log"), TEXT("file.log"));
+	RunChangeExtensionTest(TEXT("file.tar.gz"), TEXT("gz2"), TEXT("file.tar.gz2"));
+	RunChangeExtensionTest(TEXT("file.txt"), TEXT(""), TEXT("file"));
+	RunChangeExtensionTest(TEXT("C:/Folder/file"), TEXT("log"), TEXT("C:/Folder/file"));
+	RunChangeExtensionTest(TEXT("C:/Folder/file.txt"), TEXT("log"), TEXT("C:/Folder/file.log"));
+	RunChangeExtensionTest(TEXT("C:/Folder/file.tar.gz"), TEXT("gz2"), TEXT("C:/Folder/file.tar.gz2"));
+	RunChangeExtensionTest(TEXT("C:/Folder/First.Last/file"), TEXT("log"), TEXT("C:/Folder/First.Last/file"));
+	RunChangeExtensionTest(TEXT("C:/Folder/First.Last/file.txt"), TEXT("log"), TEXT("C:/Folder/First.Last/file.log"));
+	RunChangeExtensionTest(TEXT("C:/Folder/First.Last/file.tar.gz"), TEXT("gz2"), TEXT("C:/Folder/First.Last/file.tar.gz2"));
+}
+
+TEST_CASE_NAMED(FPathViewsEqualsAndLessTest, "System::Core::Misc::PathViews::EqualsAndLess", "[ApplicationContextMask][SmokeFilter]")
+{
+	FPathViewsTest Instance = FPathViewsTest();
+	Instance.EqualsAndLessTest();
+}
+
+TEST_CASE_NAMED(FPathViewsTryMakeChildPathRelativeToTest, "System::Core::Misc::PathViews::TryMakeChildPathRelativeTo", "[ApplicationContextMask][SmokeFilter]")
+{
+	FPathViewsTest Instance = FPathViewsTest();
+	Instance.TryMakeChildPathRelativeToTest();
+}
+
+TEST_CASE_NAMED(FPathViewsIsRelativePathTest, "System::Core::Misc::PathViews::IsRelativePath", "[ApplicationContextMask][SmokeFilter]")
+{
+	FPathViewsTest Instance = FPathViewsTest();
+	Instance.IsRelativePathTest();
+}
+
+TEST_CASE_NAMED(FPathViewsSplitFirstComponentTest, "System::Core::Misc::PathViews::SplitFirstComponent", "[ApplicationContextMask][SmokeFilter]")
+{
+	auto RunSplitFirstTest = [](const TCHAR* FullPath, const TCHAR* ExpectedFirst, const TCHAR* ExpectedRemaining)
+	{
+		FStringView ActualFirst;
+		FStringView ActualRemaining;
+		FPathViews::SplitFirstComponent(FullPath, ActualFirst, ActualRemaining);
+		if (!ActualFirst.Equals(ExpectedFirst, ESearchCase::IgnoreCase) || !ActualRemaining.Equals(ExpectedRemaining, ESearchCase::IgnoreCase))
+		{
+			FAIL_CHECK(FString::Printf(TEXT("SplitFirstComponent(\"%s\") == (\"%.*s\", \"%.*s\"), expected (\"%s\", \"%s\")."),
+				FullPath, ActualFirst.Len(), ActualFirst.GetData(), ActualRemaining.Len(), ActualRemaining.GetData(),
+				ExpectedFirst, ExpectedRemaining));
+		}
+	};
+
+	RunSplitFirstTest(TEXT(""), TEXT(""), TEXT(""));
+
+	RunSplitFirstTest(TEXT("A"), TEXT("A"), TEXT(""));
+	RunSplitFirstTest(TEXT("A/"), TEXT("A"), TEXT(""));
+	RunSplitFirstTest(TEXT("A\\"), TEXT("A"), TEXT(""));
+	RunSplitFirstTest(TEXT("A/B"), TEXT("A"), TEXT("B"));
+	RunSplitFirstTest(TEXT("A\\B"), TEXT("A"), TEXT("B"));
+	RunSplitFirstTest(TEXT("A/B/"), TEXT("A"), TEXT("B/"));
+	RunSplitFirstTest(TEXT("A\\B\\"), TEXT("A"), TEXT("B\\"));
+	RunSplitFirstTest(TEXT("A/B/C"), TEXT("A"), TEXT("B/C"));
+	RunSplitFirstTest(TEXT("A\\B\\C"), TEXT("A"), TEXT("B\\C"));
+
+	RunSplitFirstTest(TEXT("/A"), TEXT("/"), TEXT("A"));
+	RunSplitFirstTest(TEXT("\\A"), TEXT("\\"), TEXT("A"));
+	RunSplitFirstTest(TEXT("/A/"), TEXT("/"), TEXT("A/"));
+	RunSplitFirstTest(TEXT("\\A\\"), TEXT("\\"), TEXT("A\\"));
+	RunSplitFirstTest(TEXT("/A/B"), TEXT("/"), TEXT("A/B"));
+	RunSplitFirstTest(TEXT("\\A\\B"), TEXT("\\"), TEXT("A\\B"));
+
+	RunSplitFirstTest(TEXT("//A"), TEXT("//"), TEXT("A"));
+	RunSplitFirstTest(TEXT("\\\\A"), TEXT("\\\\"), TEXT("A"));
+	RunSplitFirstTest(TEXT("//A/"), TEXT("//"), TEXT("A/"));
+	RunSplitFirstTest(TEXT("\\\\A\\"), TEXT("\\\\"), TEXT("A\\"));
+	RunSplitFirstTest(TEXT("//A/B"), TEXT("//"), TEXT("A/B"));
+	RunSplitFirstTest(TEXT("\\\\A\\B"), TEXT("\\\\"), TEXT("A\\B"));
+
+	RunSplitFirstTest(TEXT("C:/A"), TEXT("C:/"), TEXT("A"));
+	RunSplitFirstTest(TEXT("C:\\A"), TEXT("C:\\"), TEXT("A"));
+	RunSplitFirstTest(TEXT("C:/A/"), TEXT("C:/"), TEXT("A/"));
+	RunSplitFirstTest(TEXT("C:\\A\\"), TEXT("C:\\"), TEXT("A\\"));
+	RunSplitFirstTest(TEXT("C:/A/B"), TEXT("C:/"), TEXT("A/B"));
+	RunSplitFirstTest(TEXT("C:\\A\\B"), TEXT("C:\\"), TEXT("A\\B"));
+
+	// Correctly handle paths with invalid duplicate slashes
+	RunSplitFirstTest(TEXT("///A"), TEXT("//"), TEXT("A"));
+	RunSplitFirstTest(TEXT("////////////A"), TEXT("//"), TEXT("A"));
+	RunSplitFirstTest(TEXT("///A"), TEXT("//"), TEXT("A"));
+	RunSplitFirstTest(TEXT("C://A"), TEXT("C:/"), TEXT("A"));
+	RunSplitFirstTest(TEXT("C://////////A"), TEXT("C:/"), TEXT("A"));
+	RunSplitFirstTest(TEXT("A//B"), TEXT("A"), TEXT("B"));
+	RunSplitFirstTest(TEXT("A///////////B"), TEXT("A"), TEXT("B"));
+
+}
+
+TEST_CASE_NAMED(FPathViewsAppendPathTest, "System::Core::Misc::PathViews::AppendPath", "[ApplicationContextMask][SmokeFilter]")
+{
+	auto RunAppendTest = [](const TCHAR* Base, const TCHAR* Append, const TCHAR* ExpectedNewBase)
+	{
+		TStringBuilder<128> BaseBuilder;
+		BaseBuilder << Base;
+		FPathViews::AppendPath(BaseBuilder, Append);
+		if (!FStringView(BaseBuilder).Equals(ExpectedNewBase, ESearchCase::IgnoreCase))
+		{
+			FAIL_CHECK(FString::Printf(TEXT("AppendPath(\"%s\", \"%s\") == \"%.*s\", expected \"%s\"."),
+				Base, Append, BaseBuilder.Len(), BaseBuilder.GetData(), ExpectedNewBase));
+		}
+	};
+
+	RunAppendTest(TEXT(""), TEXT("A"), TEXT("A"));
+	RunAppendTest(TEXT(""), TEXT("A/B"), TEXT("A/B"));
+	RunAppendTest(TEXT(""), TEXT("A\\B"), TEXT("A\\B"));
+
+	RunAppendTest(TEXT("Root"), TEXT("A"), TEXT("Root/A"));
+	RunAppendTest(TEXT("Root"), TEXT("A/B"), TEXT("Root/A/B"));
+	RunAppendTest(TEXT("Root"), TEXT("A\\B"), TEXT("Root/A\\B"));
+
+	RunAppendTest(TEXT("/"), TEXT("A"), TEXT("/A"));
+	RunAppendTest(TEXT("/"), TEXT("A/B"), TEXT("/A/B"));
+	RunAppendTest(TEXT("/"), TEXT("A\\B"), TEXT("/A\\B"));
+	RunAppendTest(TEXT("\\"), TEXT("A"), TEXT("\\A"));
+	RunAppendTest(TEXT("\\"), TEXT("A/B"), TEXT("\\A/B"));
+	RunAppendTest(TEXT("\\"), TEXT("A\\B"), TEXT("\\A\\B"));
+
+	RunAppendTest(TEXT("/Root"), TEXT("A"), TEXT("/Root/A"));
+	RunAppendTest(TEXT("/Root"), TEXT("A/B"), TEXT("/Root/A/B"));
+	RunAppendTest(TEXT("/Root"), TEXT("A\\B"), TEXT("/Root/A\\B"));
+	RunAppendTest(TEXT("\\Root"), TEXT("A"), TEXT("\\Root/A"));
+	RunAppendTest(TEXT("\\Root"), TEXT("A/B"), TEXT("\\Root/A/B"));
+	RunAppendTest(TEXT("\\Root"), TEXT("A\\B"), TEXT("\\Root/A\\B"));
+	RunAppendTest(TEXT("/Root/"), TEXT("A"), TEXT("/Root/A"));
+	RunAppendTest(TEXT("/Root/"), TEXT("A/B"), TEXT("/Root/A/B"));
+	RunAppendTest(TEXT("/Root/"), TEXT("A\\B"), TEXT("/Root/A\\B"));
+	RunAppendTest(TEXT("\\Root\\"), TEXT("A"), TEXT("\\Root\\A"));
+	RunAppendTest(TEXT("\\Root\\"), TEXT("A/B"), TEXT("\\Root\\A/B"));
+	RunAppendTest(TEXT("\\Root\\"), TEXT("A\\B"), TEXT("\\Root\\A\\B"));
+
+	RunAppendTest(TEXT("//"), TEXT("A"), TEXT("//A"));
+	RunAppendTest(TEXT("//"), TEXT("A/B"), TEXT("//A/B"));
+	RunAppendTest(TEXT("//"), TEXT("A\\B"), TEXT("//A\\B"));
+	RunAppendTest(TEXT("\\\\"), TEXT("A"), TEXT("\\\\A"));
+	RunAppendTest(TEXT("\\\\"), TEXT("A/B"), TEXT("\\\\A/B"));
+	RunAppendTest(TEXT("\\\\"), TEXT("A\\B"), TEXT("\\\\A\\B"));
+
+	RunAppendTest(TEXT("//Root"), TEXT("A"), TEXT("//Root/A"));
+	RunAppendTest(TEXT("//Root"), TEXT("A/B"), TEXT("//Root/A/B"));
+	RunAppendTest(TEXT("//Root"), TEXT("A\\B"), TEXT("//Root/A\\B"));
+	RunAppendTest(TEXT("\\\\Root"), TEXT("A"), TEXT("\\\\Root/A"));
+	RunAppendTest(TEXT("\\\\Root"), TEXT("A/B"), TEXT("\\\\Root/A/B"));
+	RunAppendTest(TEXT("\\\\Root"), TEXT("A\\B"), TEXT("\\\\Root/A\\B"));
+	RunAppendTest(TEXT("//Root/"), TEXT("A"), TEXT("//Root/A"));
+	RunAppendTest(TEXT("//Root/"), TEXT("A/B"), TEXT("//Root/A/B"));
+	RunAppendTest(TEXT("//Root/"), TEXT("A\\B"), TEXT("//Root/A\\B"));
+	RunAppendTest(TEXT("\\\\Root\\"), TEXT("A"), TEXT("\\\\Root\\A"));
+	RunAppendTest(TEXT("\\\\Root\\"), TEXT("A/B"), TEXT("\\\\Root\\A/B"));
+	RunAppendTest(TEXT("\\\\Root\\"), TEXT("A\\B"), TEXT("\\\\Root\\A\\B"));
+
+	RunAppendTest(TEXT("C:/"), TEXT("A"), TEXT("C:/A"));
+	RunAppendTest(TEXT("C:/"), TEXT("A/B"), TEXT("C:/A/B"));
+	RunAppendTest(TEXT("C:/"), TEXT("A\\B"), TEXT("C:/A\\B"));
+	RunAppendTest(TEXT("C:\\"), TEXT("A"), TEXT("C:\\A"));
+	RunAppendTest(TEXT("C:\\"), TEXT("A/B"), TEXT("C:\\A/B"));
+	RunAppendTest(TEXT("C:\\"), TEXT("A\\B"), TEXT("C:\\A\\B"));
+
+	RunAppendTest(TEXT("C:/Root"), TEXT("A"), TEXT("C:/Root/A"));
+	RunAppendTest(TEXT("C:/Root"), TEXT("A/B"), TEXT("C:/Root/A/B"));
+	RunAppendTest(TEXT("C:/Root"), TEXT("A\\B"), TEXT("C:/Root/A\\B"));
+	RunAppendTest(TEXT("C:\\Root"), TEXT("A"), TEXT("C:\\Root/A"));
+	RunAppendTest(TEXT("C:\\Root"), TEXT("A/B"), TEXT("C:\\Root/A/B"));
+	RunAppendTest(TEXT("C:\\Root"), TEXT("A\\B"), TEXT("C:\\Root/A\\B"));
+	RunAppendTest(TEXT("C:/Root/"), TEXT("A"), TEXT("C:/Root/A"));
+	RunAppendTest(TEXT("C:/Root/"), TEXT("A/B"), TEXT("C:/Root/A/B"));
+	RunAppendTest(TEXT("C:/Root/"), TEXT("A\\B"), TEXT("C:/Root/A\\B"));
+	RunAppendTest(TEXT("C:\\Root\\"), TEXT("A"), TEXT("C:\\Root\\A"));
+	RunAppendTest(TEXT("C:\\Root\\"), TEXT("A/B"), TEXT("C:\\Root\\A/B"));
+	RunAppendTest(TEXT("C:\\Root\\"), TEXT("A\\B"), TEXT("C:\\Root\\A\\B"));
+
+	// No matter the prefix, appending a rooted path should result in only the rooted path
+	for (const TCHAR* Prefix : { TEXT(""),
+		TEXT("/"), TEXT("\\"), TEXT("/Root"), TEXT("\\Root"), TEXT("/Root/"), TEXT("\\Root\\"),
+		TEXT("//"), TEXT("\\\\"), TEXT("//Root"), TEXT("\\\\Root"), TEXT("//Root/"), TEXT("\\\\Root\\"),
+		TEXT("C:/"), TEXT("C:\\"), TEXT("C:/Root"), TEXT("C:\\Root"), TEXT("C:/Root/"), TEXT("C:\\Root\\")})
+	{
+		RunAppendTest(Prefix, TEXT("/A"), TEXT("/A"));
+		RunAppendTest(Prefix, TEXT("\\A"), TEXT("\\A"));
+		RunAppendTest(Prefix, TEXT("/A/B"), TEXT("/A/B"));
+		RunAppendTest(Prefix, TEXT("\\A\\B"), TEXT("\\A\\B"));
+
+		RunAppendTest(Prefix, TEXT("//A"), TEXT("//A"));
+		RunAppendTest(Prefix, TEXT("\\\\A"), TEXT("\\\\A"));
+		RunAppendTest(Prefix, TEXT("//A/B"), TEXT("//A/B"));
+		RunAppendTest(Prefix, TEXT("\\\\A\\B"), TEXT("\\\\A\\B"));
+
+		RunAppendTest(Prefix, TEXT("C:/A"), TEXT("C:/A"));
+		RunAppendTest(Prefix, TEXT("C:\\A"), TEXT("C:\\A"));
+		RunAppendTest(Prefix, TEXT("C:/A/B"), TEXT("C:/A/B"));
+		RunAppendTest(Prefix, TEXT("C:\\A\\B"), TEXT("C:\\A\\B"));
+	}
+
+}
+
+
+TEST_CASE_NAMED(FPathViewsGetMountPointNameTest, "System::Core::Misc::PathViews::GetMountPointNameFromPath", "[ApplicationContextMask][SmokeFilter]")
+{
+	auto RunGetMountPointName = [](const TCHAR* Base, const TCHAR* ExpectedWithoutSlash, const TCHAR* ExpectedWithSlash)
+	{
+		TStringBuilder<128> BaseBuilder;
+		BaseBuilder << Base;
+		bool bHadClassPrefix = false;
+		FStringView Result = FPathViews::GetMountPointNameFromPath(Base, &bHadClassPrefix, true);
+		if (Result != FStringView(ExpectedWithoutSlash))
+		{
+			FAIL_CHECK(FString::Printf(TEXT("GetMountPointNameFromPath(\"%s\", &bHadClassPrefix, true) == \"%.*s\", expected \"%s\"."),
+				Base, Result.Len(), Result.GetData(), ExpectedWithoutSlash));
+		}
+		Result = FPathViews::GetMountPointNameFromPath(Base, &bHadClassPrefix, false);
+		if (Result != FStringView(ExpectedWithSlash))
+		{
+			FAIL_CHECK(FString::Printf(TEXT("AppendPath(\"%s\", \"%s\") == \"%.*s\", expected \"%s\"."),
+				Base, Result.Len(), Result.GetData(), ExpectedWithSlash));
+		}
+	};
+
+	RunGetMountPointName(TEXT(""), TEXT(""), TEXT(""));
+	RunGetMountPointName(TEXT("Root"), TEXT(""), TEXT(""));
+	RunGetMountPointName(TEXT("/"), TEXT(""), TEXT("/"));
+	RunGetMountPointName(TEXT("/Root"), TEXT("Root"), TEXT("/Root"));
+	RunGetMountPointName(TEXT("\\Root"), TEXT(""), TEXT(""));
+	RunGetMountPointName(TEXT("/Root/"), TEXT("Root"), TEXT("/Root"));
+	RunGetMountPointName(TEXT("/Root/A/B/"), TEXT("Root"), TEXT("/Root"));
+	RunGetMountPointName(TEXT("/Classes/A/"), TEXT("Classes"), TEXT("/Classes"));
+	RunGetMountPointName(TEXT("/Classes_A/B/"), TEXT("A"), TEXT("/Classes_A"));
+
+}
+
+
+TEST_CASE_NAMED(FPathViewsToAbsoluteTest, "System::Core::Misc::PathViews::ToAbsolute", "[ApplicationContextMask][SmokeFilter]")
+{
+	using namespace PathTest;
+
+	for (FTestPair Pair : ExpectedRelativeToAbsolutePaths)
+	{
+		TStringBuilder<64> ActualAppend;
+		FPathViews::ToAbsolutePath(BaseDir, Pair.Input, ActualAppend);
+		CHECK_EQUALS(TEXT("ToAbsolutePath"), ActualAppend.ToView(), Pair.Expected);
+
+		TStringBuilder<64> ActualInline;
+		ActualInline << Pair.Input;
+		FPathViews::ToAbsolutePathInline(BaseDir, ActualInline);
+		CHECK_EQUALS(TEXT("ToAbsolutePathInline"), ActualInline.ToView(), Pair.Expected);
+
+		const FStringView Original = TEXTVIEW("\\\\la/./.././la////");
+		TStringBuilder<64> ActualNondestructive;
+		ActualNondestructive << Original;
+		FPathViews::ToAbsolutePath(BaseDir, Pair.Input, ActualNondestructive);
+		CHECK_EQUALS(TEXT("ToAbsolutePath non-destructive append"), ActualNondestructive.ToView().Left(Original.Len()), Original);
+		CHECK_EQUALS(TEXT("ToAbsolutePath non-destructive append"), ActualNondestructive.ToView().RightChop(Original.Len()), Pair.Expected);
+	}
+
+}
+
+#endif //WITH_TESTS
diff --git a/Engine/Source/Runtime/Core/Tests/Misc/PathsTest.cpp b/Engine/Source/Runtime/Core/Tests/Misc/PathsTest.cpp
new file mode 100644
--- /dev/null
+++ b/Engine/Source/Runtime/Core/Tests/Misc/PathsTest.cpp
@@ -0,0 +1,137 @@
+// Copyright Epic Games, Inc. All Rights Reserved.
+
+#if WITH_TESTS
+
+#include "PathTests.h"
+#include "CoreTypes.h"
+#include "Misc/AssertionMacros.h"
+#include "Containers/UnrealString.h"
+#include "Misc/Paths.h"
+#include "Tests/TestHarnessAdapter.h"
+#if WITH_LOW_LEVEL_TESTS
+#include "TestCommon/Comparisons.h"
+#endif
+
+const FStringView PathTest::BaseDir = TEXTVIEW("/root");
+
+const PathTest::FTestPair PathTest::ExpectedRelativeToAbsolutePaths[10] =
+{
+	{ TEXTVIEW(""),					TEXTVIEW("/root/") },
+	{ TEXTVIEW("dir"),				TEXTVIEW("/root/dir") },
+	{ TEXTVIEW("/groot"),			TEXTVIEW("/groot") },
+	{ TEXTVIEW("/groot/"),			TEXTVIEW("/groot/") },
+	{ TEXTVIEW("/r/dir"),			TEXTVIEW("/r/dir") },
+	{ TEXTVIEW("/r/dir"),			TEXTVIEW("/r/dir") },
+	{ TEXTVIEW("C:\\"),				TEXTVIEW("C:/") },
+	{ TEXTVIEW("C:\\A\\B"),			TEXTVIEW("C:/A/B") },
+	{ TEXTVIEW("a/b/../c"),			TEXTVIEW("/root/a/c") },
+	{ TEXTVIEW("/a/b/../c"),		TEXTVIEW("/a/c") },
+};
+
+TEST_CASE_NAMED(FPathTests, "System::Core::Misc::Paths", "[ApplicationContextMask][SmokeFilter]")
+{
+	TestCollapseRelativeDirectories<FPaths, FString>();
+
+	// Extension texts
+	{
+		auto RunGetExtensionTest = [](const TCHAR* InPath, const TCHAR* InExpectedExt)
+		{
+			// Run test
+			const FString Ext = FPaths::GetExtension(FString(InPath));
+			if (Ext != InExpectedExt)
+			{
+				FAIL_CHECK(FString::Printf(TEXT("Path '%s' failed to get the extension (got '%s', expected '%s')."), InPath, *Ext, InExpectedExt));
+			}
+		};
+
+		RunGetExtensionTest(TEXT("file"),									TEXT(""));
+		RunGetExtensionTest(TEXT("file.txt"),								TEXT("txt"));
+		RunGetExtensionTest(TEXT("file.tar.gz"),							TEXT("gz"));
+		RunGetExtensionTest(TEXT("C:/Folder/file"),							TEXT(""));
+		RunGetExtensionTest(TEXT("C:/Folder/file.txt"),						TEXT("txt"));
+		RunGetExtensionTest(TEXT("C:/Folder/file.tar.gz"),					TEXT("gz"));
+		RunGetExtensionTest(TEXT("C:/Folder/First.Last/file"),				TEXT(""));
+		RunGetExtensionTest(TEXT("C:/Folder/First.Last/file.txt"),			TEXT("txt"));
+		RunGetExtensionTest(TEXT("C:/Folder/First.Last/file.tar.gz"),		TEXT("gz"));
+
+		auto RunSetExtensionTest = [](const TCHAR* InPath, const TCHAR* InNewExt, const FString& InExpectedPath)
+		{
+			// Run test
+			const FString NewPath = FPaths::SetExtension(FString(InPath), FString(InNewExt));
+			if (NewPath != InExpectedPath)
+			{
+				FAIL_CHECK(FString::Printf(TEXT("Path '%s' failed to set the extension (got '%s', expected '%s')."), InPath, *NewPath, *InExpectedPath));
+			}
+		};
+
+		RunSetExtensionTest(TEXT("file"),									TEXT("log"),	TEXT("file.log"));
+		RunSetExtensionTest(TEXT("file.txt"),								TEXT("log"),	TEXT("file.log"));
+		RunSetExtensionTest(TEXT("file.tar.gz"),							TEXT("gz2"),	TEXT("file.tar.gz2"));
+		RunSetExtensionTest(TEXT("C:/Folder/file"),							TEXT("log"),	TEXT("C:/Folder/file.log"));
+		RunSetExtensionTest(TEXT("C:/Folder/file.txt"),						TEXT("log"),	TEXT("C:/Folder/file.log"));
+		RunSetExtensionTest(TEXT("C:/Folder/file.tar.gz"),					TEXT("gz2"),	TEXT("C:/Folder/file.tar.gz2"));
+		RunSetExtensionTest(TEXT("C:/Folder/First.Last/file"),				TEXT("log"),	TEXT("C:/Folder/First.Last/file.log"));
+		RunSetExtensionTest(TEXT("C:/Folder/First.Last/file.txt"),			TEXT("log"),	TEXT("C:/Folder/First.Last/file.log"));
+		RunSetExtensionTest(TEXT("C:/Folder/First.Last/file.tar.gz"),		TEXT("gz2"),	TEXT("C:/Folder/First.Last/file.tar.gz2"));
+
+		auto RunChangeExtensionTest = [](const TCHAR* InPath, const TCHAR* InNewExt, const FString& InExpectedPath)
+		{
+			// Run test
+			const FString NewPath = FPaths::ChangeExtension(FString(InPath), FString(InNewExt));
+			if (NewPath != InExpectedPath)
+			{
+				FAIL_CHECK(FString::Printf(TEXT("Path '%s' failed to change the extension (got '%s', expected '%s')."), InPath, *NewPath, *InExpectedPath));
+			}
+		};
+
+		RunChangeExtensionTest(TEXT("file"),								TEXT("log"),	TEXT("file"));
+		RunChangeExtensionTest(TEXT("file.txt"),							TEXT("log"),	TEXT("file.log"));
+		RunChangeExtensionTest(TEXT("file.tar.gz"),							TEXT("gz2"),	TEXT("file.tar.gz2"));
+		RunChangeExtensionTest(TEXT("C:/Folder/file"),						TEXT("log"),	TEXT("C:/Folder/file"));
+		RunChangeExtensionTest(TEXT("C:/Folder/file.txt"),					TEXT("log"),	TEXT("C:/Folder/file.log"));
+		RunChangeExtensionTest(TEXT("C:/Folder/file.tar.gz"),				TEXT("gz2"),	TEXT("C:/Folder/file.tar.gz2"));
+		RunChangeExtensionTest(TEXT("C:/Folder/First.Last/file"),			TEXT("log"),	TEXT("C:/Folder/First.Last/file"));
+		RunChangeExtensionTest(TEXT("C:/Folder/First.Last/file.txt"),		TEXT("log"),	TEXT("C:/Folder/First.Last/file.log"));
+		RunChangeExtensionTest(TEXT("C:/Folder/First.Last/file.tar.gz"),	TEXT("gz2"),	TEXT("C:/Folder/First.Last/file.tar.gz2"));
+	}
+
+	// IsUnderDirectory
+	{
+		auto RunIsUnderDirectoryTest = [](const TCHAR* InPath1, const TCHAR* InPath2, bool ExpectedResult)
+		{
+			// Run test
+			bool Result = FPaths::IsUnderDirectory(FString(InPath1), FString(InPath2));
+			if (Result != ExpectedResult)
+			{
+				FAIL_CHECK(FString::Printf(TEXT("FPaths::IsUnderDirectory('%s', '%s') != %s."), InPath1, InPath2, ExpectedResult ? TEXT("true") : TEXT("false")));
+			}
+		};
+
+		RunIsUnderDirectoryTest(TEXT("C:/Folder"),			TEXT("C:/FolderN"), false);
+		RunIsUnderDirectoryTest(TEXT("C:/Folder1"),			TEXT("C:/Folder2"), false);
+		RunIsUnderDirectoryTest(TEXT("C:/Folder"),			TEXT("C:/Folder/SubDir"), false);
+
+		RunIsUnderDirectoryTest(TEXT("C:/Folder"),			TEXT("C:/Folder"), true);
+		RunIsUnderDirectoryTest(TEXT("C:/Folder/File"),		TEXT("C:/Folder"), true);
+		RunIsUnderDirectoryTest(TEXT("C:/Folder/File"),		TEXT("C:/Folder/"), true);
+		RunIsUnderDirectoryTest(TEXT("C:/Folder/"),			TEXT("C:/Folder"), true);
+		RunIsUnderDirectoryTest(TEXT("C:/Folder/"),			TEXT("C:/Folder/"), true);
+		RunIsUnderDirectoryTest(TEXT("C:/Folder/Subdir/"),	TEXT("C:/Folder"), true);
+		RunIsUnderDirectoryTest(TEXT("C:/Folder/Subdir/"),	TEXT("C:/Folder/"), true);
+	}
+
+	TestRemoveDuplicateSlashes<FPaths, FString>();
+
+	// ConvertRelativePathToFull
+	{
+		using namespace PathTest;
+
+		for (FTestPair Pair : ExpectedRelativeToAbsolutePaths)
+		{
+			FString Actual = FPaths::ConvertRelativePathToFull(FString(BaseDir), FString(Pair.Input));
+			CHECK_EQUALS(TEXT("ConvertRelativePathToFull"), FStringView(Actual), Pair.Expected);
+		}
+	}
+}
+
+#endif //WITH_TESTS
diff --git a/Engine/Source/Runtime/Core/Tests/Misc/PlatformStringConvertTests.cpp b/Engine/Source/Runtime/Core/Tests/Misc/PlatformStringConvertTests.cpp
new file mode 100644
--- /dev/null
+++ b/Engine/Source/Runtime/Core/Tests/Misc/PlatformStringConvertTests.cpp
@@ -0,0 +1,98 @@
+// Copyright Epic Games, Inc. All Rights Reserved.
+
+#include "CoreTypes.h"
+#include "Tests/TestHarnessAdapter.h"
+#include <type_traits>
+
+#if WITH_TESTS
+
+TEST_CASE_NAMED(FStringConvertTest, "System::Core::Misc::StringConvertTests", "[ApplicationContextMask][SmokeFilter]")
+{
+	auto CheckNullTerminatedStringConversion = [](const TCHAR* TestName, const auto* From, const auto* To)
+	{
+		using ToType = std::decay_t<decltype(*To)>;
+
+		// Basic version of Strcmp for arbitrary types
+		auto Strcmp = [](const auto* Lhs, const auto* Rhs)
+		{
+			for (;;)
+			{
+				if (!*Lhs)
+				{
+					return *Rhs ? -1 : 0;
+				}
+
+				if (!*Rhs)
+				{
+					return 1;
+				}
+
+				if (*Lhs != *Rhs)
+				{
+					return *Lhs < *Rhs ? -1 : 1;
+				}
+
+				++Lhs;
+				++Rhs;
+			}
+		};
+
+		int32 ConvertedLength = FPlatformString::ConvertedLength<ToType>(From);
+
+		// Allocate a buffer to hold the converted string and set it with 0xCD
+		TUniquePtr<ToType[]> UniquePtr = MakeUniqueForOverwrite<ToType[]>(ConvertedLength + 8);
+		ToType* Ptr = UniquePtr.Get();
+		FMemory::Memset(Ptr, 0xCD, (ConvertedLength + 8) * sizeof(ToType));
+
+		ToType* End = FPlatformString::Convert(Ptr + 4, ConvertedLength, From);
+
+		CHECK_MESSAGE(FString::Printf(TEXT("%s converted length matches"), TestName), End == Ptr + 4 + ConvertedLength);
+
+		// != changed to < because of false positive in MSVC's static analyzer: warning C6295: Ill-defined for-loop.  Loop executes infinitely.
+		for (int i = 0; i < 4 * sizeof(ToType); ++i)
+		{
+			CHECK_MESSAGE(FString::Printf(TEXT("%s guard block is preserved"), TestName), ((unsigned char*)Ptr)[i] == 0xCD && ((unsigned char*)End)[i] == 0xCD);
+		}
+
+		CHECK_MESSAGE(FString::Printf(TEXT("%s converted string matches"), TestName), Strcmp(To, Ptr + 4) == 0);
+	};
+
+	// We need to do UTF-8 literals like below because MSVC doesn't encode numeric escape sequences properly in u8"" literals:
+	// https://developercommunity.visualstudio.com/t/hex-escape-codes-in-a-utf8-literal-are-t/225847
+	//
+	// We cast to UTF16CHAR* because UTF16TEXT doesn't currently result in an array of UTF16CHAR.
+
+	// Freely interconvertible
+	{
+		const UTF16CHAR* UnrealEngineIsGreat16 = (const UTF16CHAR*)UTF16TEXT("\xD835\xDE50\x043F\xD835\xDDCB\xD835\xDE26\xD835\xDD52\x006C\x0020\xD835\xDC6C\xD835\xDD2B\xFF47\xD835\xDCBE\xD835\xDD93\xD835\xDD56\x0020\xD835\xDE92\xFF53\x0020\xD835\xDC6E\xD835\xDE9B\x212E\xD835\xDF36\xFF54\x0021");
+		const UTF8CHAR*  UnrealEngineIsGreat8  = (const UTF8CHAR*)"\xF0\x9D\x99\x90\xD0\xBF\xF0\x9D\x97\x8B\xF0\x9D\x98\xA6\xF0\x9D\x95\x92\x6C\x20\xF0\x9D\x91\xAC\xF0\x9D\x94\xAB\xEF\xBD\x87\xF0\x9D\x92\xBE\xF0\x9D\x96\x93\xF0\x9D\x95\x96\x20\xF0\x9D\x9A\x92\xEF\xBD\x93\x20\xF0\x9D\x91\xAE\xF0\x9D\x9A\x9B\xE2\x84\xAE\xF0\x9D\x9C\xB6\xEF\xBD\x94\x21";
+		CheckNullTerminatedStringConversion(TEXT("UnrealEngineIsGreat UTF-16 -> UTF-8"), UnrealEngineIsGreat16, UnrealEngineIsGreat8);
+		CheckNullTerminatedStringConversion(TEXT("UnrealEngineIsGreat UTF-8 -> UTF-16"), UnrealEngineIsGreat8, UnrealEngineIsGreat16);
+
+		const UTF16CHAR* Internationalization16 = (const UTF16CHAR*)UTF16TEXT("\x0049\x00F1\x0074\x00EB\x0072\x006E\x00E2\x0074\x0069\x00F4\x006E\x00E0\x006C\x0069\x007A\x00E6\x0074\x0069\x00F8\x006E");
+		const UTF8CHAR*  Internationalization8  = (const UTF8CHAR*)"\x49\xC3\xB1\x74\xC3\xAB\x72\x6E\xC3\xA2\x74\x69\xC3\xB4\x6E\xC3\xA0\x6C\x69\x7A\xC3\xA6\x74\x69\xC3\xB8\x6E";
+		CheckNullTerminatedStringConversion(TEXT("Internationalization UTF-16 -> UTF-8"), Internationalization16, Internationalization8);
+		CheckNullTerminatedStringConversion(TEXT("Internationalization UTF-8 -> UTF-16"), Internationalization8, Internationalization16);
+
+		const UTF16CHAR* WomanEmoji16 = (const UTF16CHAR*)UTF16TEXT("\xD83D\xDC69");
+		const UTF8CHAR*  WomanEmoji8  = (const UTF8CHAR*)"\xf0\x9f\x91\xa9";
+		CheckNullTerminatedStringConversion(TEXT("WomanEmoji UTF-16 -> UTF-8"), WomanEmoji16, WomanEmoji8);
+		CheckNullTerminatedStringConversion(TEXT("WomanEmoji UTF-8 -> UTF-16"), WomanEmoji8, WomanEmoji16);
+
+		const UTF16CHAR* LanguageProcessingChinese16 = (const UTF16CHAR*)UTF16TEXT("\x8bed\x8a00\x5904\x7406");
+		const UTF8CHAR*  LanguageProcessingChinese8  = (const UTF8CHAR*)"\xe8\xaf\xad\xe8\xa8\x80\xe5\xa4\x84\xe7\x90\x86";
+		CheckNullTerminatedStringConversion(TEXT("LanguageProcessingChinese UTF-16 -> UTF-8"), LanguageProcessingChinese16, LanguageProcessingChinese8);
+		CheckNullTerminatedStringConversion(TEXT("LanguageProcessingChinese UTF-8 -> UTF-16"), LanguageProcessingChinese8, LanguageProcessingChinese16);
+	}
+
+	// Broken UTF-16 encodings
+	{
+		CheckNullTerminatedStringConversion(TEXT("Double UTF-16 high surrogate"),   (const UTF16CHAR*)UTF16TEXT("q\xD83D\xD83Dq"), (const UTF8CHAR*)"q??q");
+		CheckNullTerminatedStringConversion(TEXT("Double UTF-16 low surrogate"),    (const UTF16CHAR*)UTF16TEXT("q\xDC69\xDC69q"), (const UTF8CHAR*)"q??q");
+		CheckNullTerminatedStringConversion(TEXT("Orphaned UTF-16 high surrogate"), (const UTF16CHAR*)UTF16TEXT("q\xDC69q"),       (const UTF8CHAR*)"q?q");
+		CheckNullTerminatedStringConversion(TEXT("Orphaned UTF-16 low surrogate"),  (const UTF16CHAR*)UTF16TEXT("q\xDC69q"),       (const UTF8CHAR*)"q?q");
+	}
+
+}
+
+#endif //WITH_TESTS
diff --git a/Engine/Source/Runtime/Core/Tests/Misc/QueueTest.cpp b/Engine/Source/Runtime/Core/Tests/Misc/QueueTest.cpp
new file mode 100644
--- /dev/null
+++ b/Engine/Source/Runtime/Core/Tests/Misc/QueueTest.cpp
@@ -0,0 +1,67 @@
+// Copyright Epic Games, Inc. All Rights Reserved.
+
+#if WITH_TESTS
+
+#include "CoreTypes.h"
+#include "Containers/Queue.h"
+#include "Tests/TestHarnessAdapter.h"
+#if WITH_LOW_LEVEL_TESTS
+#include "TestCommon/Comparisons.h"
+#endif
+
+TEST_CASE_NAMED(FQueueTest, "System::Core::Misc::Queue", "[ApplicationContextMask][SmokeFilter]")
+{
+	// empty queues
+	{
+		TQueue<int32> Queue;
+		int32 OutItem = 0;
+
+		CHECK_MESSAGE(TEXT("A new queue must be empty"), Queue.IsEmpty());
+
+		CHECK_FALSE_MESSAGE(TEXT("A new queue must not dequeue anything"), Queue.Dequeue(OutItem));
+		CHECK_FALSE_MESSAGE(TEXT("A new queue must not peek anything"), Queue.Peek(OutItem));
+	}
+
+	// insertion & removal
+	{
+		TQueue<int32> Queue;
+		int32 Item1 = 1;
+		int32 Item2 = 2;
+		int32 Item3 = 3;
+		int32 OutItem = 0;
+
+		CHECK_MESSAGE(TEXT("Inserting into a new queue must succeed"), Queue.Enqueue(Item1));
+		CHECK_MESSAGE(TEXT("Peek must succeed on a queue with one item"), Queue.Peek(OutItem));
+		CHECK_EQUALS(TEXT("Peek must return the first value"), OutItem, Item1);
+
+		CHECK_MESSAGE(TEXT("Inserting into a non-empty queue must succeed"), Queue.Enqueue(Item2));
+		CHECK_MESSAGE(TEXT("Peek must succeed on a queue with two items"), Queue.Peek(OutItem));
+		CHECK_EQUALS(TEXT("Peek must return the first item"), OutItem, Item1);
+
+		Queue.Enqueue(Item3);
+
+		CHECK_MESSAGE(TEXT("Dequeue must succeed on a queue with three items"), Queue.Dequeue(OutItem));
+		CHECK_EQUALS(TEXT("Dequeue must return the first item"), OutItem, Item1);
+		CHECK_MESSAGE(TEXT("Dequeue must succeed on a queue with two items"), Queue.Dequeue(OutItem));
+		CHECK_EQUALS(TEXT("Dequeue must return the second item"), OutItem, Item2);
+		CHECK_MESSAGE(TEXT("Dequeue must succeed on a queue with one item"), Queue.Dequeue(OutItem));
+		CHECK_EQUALS(TEXT("Dequeue must return the third item"), OutItem, Item3);
+
+		CHECK_MESSAGE(TEXT("After removing all items, the queue must be empty"), Queue.IsEmpty());
+	}
+
+	// emptying
+	{
+		TQueue<int32> Queue;
+
+		Queue.Enqueue(1);
+		Queue.Enqueue(2);
+		Queue.Enqueue(3);
+		Queue.Empty();
+
+		CHECK_MESSAGE(TEXT("An emptied queue must be empty"), Queue.IsEmpty());
+	}
+}
+
+
+#endif //WITH_TESTS
diff --git a/Engine/Source/Runtime/Core/Tests/Misc/TCharToUTF8.cpp b/Engine/Source/Runtime/Core/Tests/Misc/TCharToUTF8.cpp
new file mode 100644
--- /dev/null
+++ b/Engine/Source/Runtime/Core/Tests/Misc/TCharToUTF8.cpp
@@ -0,0 +1,191 @@
+// Copyright Epic Games, Inc. All Rights Reserved.
+
+#if WITH_TESTS
+
+#include "CoreTypes.h"
+#include "Tests/TestHarnessAdapter.h"
+#if WITH_LOW_LEVEL_TESTS
+#include "TestCommon/Comparisons.h"
+#endif
+
+TEST_CASE_NAMED(FTCharToUTF8Test, "System::Core::Misc::TCharToUtf8", "[ApplicationContextMask][SmokeFilter]")
+{
+	{
+		// Hebrew letter final nun
+		const FString TestString = UTF16_TO_TCHAR(u"\x05DF");
+		const ANSICHAR ExpectedResult[] = "\xD7\x9F";
+		{
+			const FTCHARToUTF8 ConvertedValue(*TestString);
+			const FUTF8ToTCHAR ReverseConvertedValue(ConvertedValue.Get());
+			CHECK_MESSAGE(FString(TEXT("Expected test string 1 to have different encoded value")), FCStringAnsi::Strcmp(ExpectedResult, (const ANSICHAR*)ConvertedValue.Get()) == 0);
+			CHECK_MESSAGE(FString(TEXT("Expected test string 1 to have different decoded value")), FCString::Strcmp(*TestString, ReverseConvertedValue.Get()) == 0);
+		}
+
+		{
+			const FUTF8ToTCHAR ConvertedValue(ExpectedResult);
+			const FTCHARToUTF8 ReverseConvertedValue(ConvertedValue.Get());
+			CHECK_MESSAGE(FString(TEXT("Expected test string 2 to have different encoded value")), FCString::Strcmp(*TestString, ConvertedValue.Get()) == 0);
+			CHECK_MESSAGE(FString(TEXT("Expected test string 2 to have different decoded value")), FCStringAnsi::Strcmp(ExpectedResult, (const ANSICHAR*)ReverseConvertedValue.Get()) == 0);
+		}
+	}
+
+	{
+		// Smiling face with open mouth and tightly-closed eyes and a Four of Circles
+		const FString TestString = UTF16_TO_TCHAR(u"\xD83D\xDE06\xD83C\xDC1C");
+		const ANSICHAR ExpectedResult[] = "\xF0\x9F\x98\x86\xF0\x9F\x80\x9C";
+		{
+			const FTCHARToUTF8 ConvertedValue(*TestString);
+			const FUTF8ToTCHAR ReverseConvertedValue(ConvertedValue.Get());
+			CHECK_MESSAGE(FString(TEXT("Expected test string 3 to have different encoded value")), FCStringAnsi::Strcmp(ExpectedResult, (const ANSICHAR*)ConvertedValue.Get()) == 0);
+			CHECK_MESSAGE(FString(TEXT("Expected test string 3 to have different decoded value")), FCString::Strcmp(*TestString, ReverseConvertedValue.Get()) == 0);
+		}
+
+		{
+			const FUTF8ToTCHAR ConvertedValue(ExpectedResult);
+			const FTCHARToUTF8 ReverseConvertedValue(ConvertedValue.Get());
+			CHECK_MESSAGE(FString(TEXT("Expected test string 4 to have different encoded value")), FCString::Strcmp(*TestString, ConvertedValue.Get()) == 0);
+			CHECK_MESSAGE(FString(TEXT("Expected test string 4 to have different decoded value")), FCStringAnsi::Strcmp(ExpectedResult, (const ANSICHAR*)ReverseConvertedValue.Get()) == 0);
+		}
+	}
+
+	{
+		// "Internationalization" Test string
+		const FString TestString = UTF16_TO_TCHAR(u"\x0049\x00F1\x0074\x00EB\x0072\x006E\x00E2\x0074\x0069\x00F4\x006E\x00E0\x006C\x0069\x007A\x00E6\x0074\x0069\x00F8\x006E");
+		const ANSICHAR ExpectedResult[] = "\x49\xC3\xB1\x74\xC3\xAB\x72\x6E\xC3\xA2\x74\x69\xC3\xB4\x6E\xC3\xA0\x6C\x69\x7A\xC3\xA6\x74\x69\xC3\xB8\x6E";
+		{
+			const FTCHARToUTF8 ConvertedValue(*TestString);
+			const FUTF8ToTCHAR ReverseConvertedValue(ConvertedValue.Get());
+			CHECK_MESSAGE(FString(TEXT("Expected test string 5 to have different encoded value")), FCStringAnsi::Strcmp(ExpectedResult, (const ANSICHAR*)ConvertedValue.Get()) == 0);
+			CHECK_MESSAGE(FString(TEXT("Expected test string 5 to have different decoded value")), FCString::Strcmp(*TestString, ReverseConvertedValue.Get()) == 0);
+		}
+
+		{
+			const FUTF8ToTCHAR ConvertedValue(ExpectedResult);
+			const FTCHARToUTF8 ReverseConvertedValue(ConvertedValue.Get());
+			CHECK_MESSAGE(FString(TEXT("Expected test string 6 to have different encoded value")), FCString::Strcmp(*TestString, ConvertedValue.Get()) == 0);
+			CHECK_MESSAGE(FString(TEXT("Expected test string 6 to have different decoded value")), FCStringAnsi::Strcmp(ExpectedResult, (const ANSICHAR*)ReverseConvertedValue.Get()) == 0);
+		}
+	}
+
+	{
+		// UE Test string
+		const FString TestString = UTF16_TO_TCHAR(u"\xD835\xDE50\x043F\xD835\xDDCB\xD835\xDE26\xD835\xDD52\x006C\x0020\xD835\xDC6C\xD835\xDD2B\xFF47\xD835\xDCBE\xD835\xDD93\xD835\xDD56\x0020\xD835\xDE92\xFF53\x0020\xD835\xDC6E\xD835\xDE9B\x212E\xD835\xDF36\xFF54\x0021");
+		const ANSICHAR ExpectedResult[] = "\xF0\x9D\x99\x90\xD0\xBF\xF0\x9D\x97\x8B\xF0\x9D\x98\xA6\xF0\x9D\x95\x92\x6C\x20\xF0\x9D\x91\xAC\xF0\x9D\x94\xAB\xEF\xBD\x87\xF0\x9D\x92\xBE\xF0\x9D\x96\x93\xF0\x9D\x95\x96\x20\xF0\x9D\x9A\x92\xEF\xBD\x93\x20\xF0\x9D\x91\xAE\xF0\x9D\x9A\x9B\xE2\x84\xAE\xF0\x9D\x9C\xB6\xEF\xBD\x94\x21";
+		{
+			const FTCHARToUTF8 ConvertedValue(*TestString);
+			const FUTF8ToTCHAR ReverseConvertedValue(ConvertedValue.Get());
+			CHECK_MESSAGE(FString(TEXT("Expected test string 7 to have different encoded value")), FCStringAnsi::Strcmp(ExpectedResult, (const ANSICHAR*)ConvertedValue.Get()) == 0);
+			CHECK_MESSAGE(FString(TEXT("Expected test string 7 to have different decoded value")), FCString::Strcmp(*TestString, ReverseConvertedValue.Get()) == 0);
+		}
+
+		{
+			const FUTF8ToTCHAR ConvertedValue(ExpectedResult);
+			const FTCHARToUTF8 ReverseConvertedValue(ConvertedValue.Get());
+			CHECK_MESSAGE(FString(TEXT("Expected test string 8 to have different encoded value")), FCString::Strcmp(*TestString, ConvertedValue.Get()) == 0);
+			CHECK_MESSAGE(FString(TEXT("Expected test string 8 to have different decoded value")), FCStringAnsi::Strcmp(ExpectedResult, (const ANSICHAR*)ReverseConvertedValue.Get()) == 0);
+		}
+	}
+
+	{
+		// Null terminator
+		const FString TestString = TEXT("Test");
+		{
+			const FTCHARToUTF8 ConvertedValue(*TestString);
+			const bool bIsNullTerminated = *(ConvertedValue.Get() + ConvertedValue.Length()) == 0;
+			CHECK_MESSAGE(TEXT("Expected converted value to be null-terminated"), bIsNullTerminated);
+		}
+		{
+			const FTCHARToUTF8 ConvertedValue(*TestString, TestString.Len() + 1); // include the null terminator
+			const bool bIsNullTerminated = *(ConvertedValue.Get() + ConvertedValue.Length()) == 0;
+			CHECK_MESSAGE(TEXT("Expected converted value to be null-terminated"), bIsNullTerminated);
+		}
+	}
+
+	{
+		// Bad UTF-16 surrogates
+		{
+			const FTCHARToUTF8 EndWithHighSurrogate(UTF16_TO_TCHAR(u"Hello\xD83D"));
+			CHECK_MESSAGE(TEXT("UTF-16 string ending with a single high surrogate converted to UTF-8"), FCStringUtf8::Strcmp((const UTF8CHAR*)EndWithHighSurrogate.Get(), (const UTF8CHAR*)"Hello?") == 0);
+		}
+		{
+			const FTCHARToUTF8 EndWithTwoHighSurrogates(UTF16_TO_TCHAR(u"Hello\xD83D\xD83D"));
+			CHECK_MESSAGE(TEXT("UTF-16 string ending with two high surrogates converted to UTF-8"), FCStringUtf8::Strcmp((const UTF8CHAR*)EndWithTwoHighSurrogates.Get(), (const UTF8CHAR*)"Hello??") == 0);
+		}
+		{
+			const FTCHARToUTF8 EndWithThreeHighSurrogates(UTF16_TO_TCHAR(u"Hello\xD83D\xD83D\xD83D"));
+			CHECK_MESSAGE(TEXT("UTF-16 string ending with three high surrogates converted to UTF-8"), FCStringUtf8::Strcmp((const UTF8CHAR*)EndWithThreeHighSurrogates.Get(), (const UTF8CHAR*)"Hello???") == 0);
+		}
+		{
+			const FTCHARToUTF8 EndWithHighSurrogateAndPair(UTF16_TO_TCHAR(u"Hello\xD83D\xD83D\xDC69"));
+			CHECK_MESSAGE(TEXT("UTF-16 string ending with a high surrogate and a surrogate pair converted to UTF-8"), FCStringUtf8::Strcmp((const UTF8CHAR*)EndWithHighSurrogateAndPair.Get(), (const UTF8CHAR*)"Hello?\xf0\x9f\x91\xa9") == 0);
+		}
+		{
+			const FTCHARToUTF8 EndWithLowSurrogate(UTF16_TO_TCHAR(u"Hello\xDC69"));
+			CHECK_MESSAGE(TEXT("UTF-16 string ending with a single low surrogate converted to UTF-8"), FCStringUtf8::Strcmp((const UTF8CHAR*)EndWithLowSurrogate.Get(), (const UTF8CHAR*)"Hello?") == 0);
+		}
+		{
+			const FTCHARToUTF8 EndWithTwoLowSurrogates(UTF16_TO_TCHAR(u"Hello\xDC69\xDC69"));
+			CHECK_MESSAGE(TEXT("UTF-16 string ending with two low surrogates converted to UTF-8"), FCStringUtf8::Strcmp((const UTF8CHAR*)EndWithTwoLowSurrogates.Get(), (const UTF8CHAR*)"Hello??") == 0);
+		}
+		{
+			const FTCHARToUTF8 EndWithThreeLowSurrogates(UTF16_TO_TCHAR(u"Hello\xDC69\xDC69\xDC69"));
+			CHECK_MESSAGE(TEXT("UTF-16 string ending with three low surrogates converted to UTF-8"), FCStringUtf8::Strcmp((const UTF8CHAR*)EndWithThreeLowSurrogates.Get(), (const UTF8CHAR*)"Hello???") == 0);
+		}
+		{
+			const FTCHARToUTF8 EndWithLowSurrogateAndPair(UTF16_TO_TCHAR(u"Hello\xDC69\xD83D\xDC69"));
+			CHECK_MESSAGE(TEXT("UTF-16 string ending with a low surrogate and a surrogate pair converted to UTF-8"), FCStringUtf8::Strcmp((const UTF8CHAR*)EndWithLowSurrogateAndPair.Get(), (const UTF8CHAR*)"Hello?\xf0\x9f\x91\xa9") == 0);
+		}
+		{
+			const FTCHARToUTF8 MidHighSurrogate(UTF16_TO_TCHAR(u"Hello\xD83DWorld"));
+			CHECK_MESSAGE(TEXT("UTF-16 string containing a single high surrogate converted to UTF-8"), FCStringUtf8::Strcmp((const UTF8CHAR*)MidHighSurrogate.Get(), (const UTF8CHAR*)"Hello?World") == 0);
+		}
+		{
+			const FTCHARToUTF8 MidTwoHighSurrogates(UTF16_TO_TCHAR(u"Hello\xD83D\xD83DWorld"));
+			CHECK_MESSAGE(TEXT("UTF-16 string containing two high surrogates converted to UTF-8"), FCStringUtf8::Strcmp((const UTF8CHAR*)MidTwoHighSurrogates.Get(), (const UTF8CHAR*)"Hello??World") == 0);
+		}
+		{
+			const FTCHARToUTF8 MidThreeHighSurrogates(UTF16_TO_TCHAR(u"Hello\xD83D\xD83D\xD83DWorld"));
+			CHECK_MESSAGE(TEXT("UTF-16 string containing three high surrogates converted to UTF-8"), FCStringUtf8::Strcmp((const UTF8CHAR*)MidThreeHighSurrogates.Get(), (const UTF8CHAR*)"Hello???World") == 0);
+		}
+		{
+			const FTCHARToUTF8 MidHighSurrogateAndPair(UTF16_TO_TCHAR(u"Hello\xD83D\xD83D\xDC69World"));
+			CHECK_MESSAGE(TEXT("UTF-16 string containing a high surrogate and a surrogate pair converted to UTF-8"), FCStringUtf8::Strcmp((const UTF8CHAR*)MidHighSurrogateAndPair.Get(), (const UTF8CHAR*)"Hello?\xf0\x9f\x91\xa9World") == 0);
+		}
+		{
+			const FTCHARToUTF8 MidLowSurrogate(UTF16_TO_TCHAR(u"Hello\xDC69World"));
+			CHECK_MESSAGE(TEXT("UTF-16 string containing a single low surrogate converted to UTF-8"), FCStringUtf8::Strcmp((const UTF8CHAR*)MidLowSurrogate.Get(), (const UTF8CHAR*)"Hello?World") == 0);
+		}
+		{
+			const FTCHARToUTF8 MidTwoLowSurrogates(UTF16_TO_TCHAR(u"Hello\xDC69\xDC69World"));
+			CHECK_MESSAGE(TEXT("UTF-16 string containing two low surrogates converted to UTF-8"), FCStringUtf8::Strcmp((const UTF8CHAR*)MidTwoLowSurrogates.Get(), (const UTF8CHAR*)"Hello??World") == 0);
+		}
+		{
+			const FTCHARToUTF8 MidThreeLowSurrogates(UTF16_TO_TCHAR(u"Hello\xDC69\xDC69\xDC69World"));
+			CHECK_MESSAGE(TEXT("UTF-16 string containing three low surrogates converted to UTF-8"), FCStringUtf8::Strcmp((const UTF8CHAR*)MidThreeLowSurrogates.Get(), (const UTF8CHAR*)"Hello???World") == 0);
+		}
+		{
+			const FTCHARToUTF8 MidLowSurrogateAndPair(UTF16_TO_TCHAR(u"Hello\xDC69\xD83D\xDC69World"));
+			CHECK_MESSAGE(TEXT("UTF-16 string containing a low surrogate and a surrogate pair converted to UTF-8"), FCStringUtf8::Strcmp((const UTF8CHAR*)MidLowSurrogateAndPair.Get(), (const UTF8CHAR*)"Hello?\xf0\x9f\x91\xa9World") == 0);
+		}
+	}
+
+	{
+		// Verify that we handle invalid UTF-16 strings that ends in half a surrogate pair w/o crashing
+#if PLATFORM_TCHAR_IS_UTF8CHAR
+#ifdef _MSC_VER
+		// We need to do this because MSVC doesn't encode numeric escape sequences properly in u8"" literals:
+		// https://developercommunity.visualstudio.com/t/hex-escape-codes-in-a-utf8-literal-are-t/225847
+		char EndWithIllegalSurrogatePairImpl[] = "ab\xED\xA0\x80";
+		const UTF8CHAR* EndWithIllegalSurrogatePair = (const UTF8CHAR*)EndWithIllegalSurrogatePairImpl;
+#else
+		const TCHAR* EndWithIllegalSurrogatePair = TEXT("ab\xED\xA0\x80"));
+#endif // #ifdef _MSC_VER
+#else
+		const TCHAR* EndWithIllegalSurrogatePair = TEXT("ab\xD800");
+#endif // #if PLATFORM_TCHAR_IS_UTF8CHAR
+
+		CHECK_EQUALS("IllegalSurrogatePair", TCHAR_TO_UTF8(EndWithIllegalSurrogatePair), "ab?");
+	}
+}
+
+#endif //WITH_TESTS
diff --git a/Engine/Source/Runtime/Core/Tests/Misc/TVariantTest.cpp b/Engine/Source/Runtime/Core/Tests/Misc/TVariantTest.cpp
new file mode 100644
--- /dev/null
+++ b/Engine/Source/Runtime/Core/Tests/Misc/TVariantTest.cpp
@@ -0,0 +1,497 @@
+// Copyright Epic Games, Inc. All Rights Reserved.
+
+#if WITH_TESTS
+
+#include "CoreTypes.h"
+#include "Templates/SharedPointer.h"
+#include "Templates/UniquePtr.h"
+#include "Tests/TestHarnessAdapter.h"
+#include "Misc/TVariant.h"
+#include "Serialization/MemoryReader.h"
+#include "Serialization/MemoryWriter.h"
+
+namespace UE::Core::Private
+{
+	// TIsAllVariant
+	static_assert(TIsAllVariant<TVariant<int, float>>::Value, "Expected TIsAllVariant<TVariant> to be true");
+	static_assert(TIsAllVariant<TVariant<int, float>, TVariant<float, int>>::Value, "Expected TIsAllVariant<TVariant, TVariant> to be true");
+	static_assert(TIsAllVariant<TVariant<int, float>, const TVariant<float, int>>::Value, "Expected TIsAllVariant<TVariant, const TVariant> to be true");
+	static_assert(TIsAllVariant<>::Value, "Expected TIsAllVariant<> for base case to be true");
+	static_assert(TIsAllVariant<bool>::Value == false, "Expected TIsAllVariant<non-TVariant> to be false");
+	static_assert(TIsAllVariant<TVariant<int, float>, bool>::Value == false, "Expected TIsAllVariant<TVariant, non-TVariant> to be false");
+	static_assert(TIsAllVariant<TVariant<int, float>, bool, TVariant<float, int>>::Value == false, "Expected TIsAllVariant<TVariant, non-TVariant, TVariant> to be false");
+
+	// TIsVariant
+	static_assert(TIsVariant<TVariant<int, float>>::Value, "Expected TIsVariant<TVariant> to be true");
+	static_assert(TIsVariant<const TVariant<int, float>>::Value, "Expected TIsVariant<const TVariant> to be true");
+	static_assert(TIsVariant<TVariant<int, float>&>::Value, "Expected TIsVariant<TVariant&> to be true");
+	static_assert(TIsVariant<const TVariant<int, float>&>::Value, "Expected TIsVariant<const TVariant&> to be true");
+	static_assert(TIsVariant<TVariant<int, float>&&>::Value, "Expected TIsVariant<TVariant&&> to be true");
+	static_assert(TIsVariant<bool>::Value == false, "Expected TIsVariant<non-TVariant> to be false");
+
+	// TVariantSize
+	static_assert(TVariantSize<TVariant<int, float>>::Value == 2, "Expected TVariantSize<TVariant> to be 2");
+	static_assert(TVariantSize<const TVariant<int, float>>::Value == 2, "Expected TVariantSize<const TVariant> to be 2");
+	static_assert(TVariantSize<TVariant<int, float>&>::Value == 2, "Expected TVariantSize<TVariant&> to be 2");
+	static_assert(TVariantSize<const TVariant<int, float>&>::Value == 2, "Expected TVariantSize<const TVariant&> to be 2");
+	static_assert(TVariantSize<TVariant<int, float>&&>::Value == 2, "Expected TVariantSize<TVariant&&> to be 2");
+} // namespace UE::Core::Private
+
+struct FNonDefaultConstructible
+{
+	int _A;
+	float _B;
+	FNonDefaultConstructible(int A, float B) : _A(A), _B(B) {}
+
+	bool operator==(const FNonDefaultConstructible& Other) const
+	{
+		return _A == Other._A && _B == Other._B;
+	}
+};
+
+TEST_CASE_NAMED(FTVariantInPlaceConstructionTest, "System::Core::Misc::TVariant::InPlaceConstruction", "[ApplicationContextMask][EngineFilter]")
+{
+	// TVariant supports in-place construction of a type in its parameter pack
+	using FVariant = TVariant<FNonDefaultConstructible>;
+	
+	FVariant Var{ TInPlaceType<FNonDefaultConstructible>(), 30, 40.0f };
+	check(Var.IsType<FNonDefaultConstructible>());
+	check(Var.Get<FNonDefaultConstructible>() == FNonDefaultConstructible(30, 40.0f));
+
+}
+
+TEST_CASE_NAMED(FTVariantDefaultConstructorTest, "System::Core::Misc::TVariant::DefaultConstruction", "[ApplicationContextMask][EngineFilter]")
+{
+	// TVariant default constructor constructs the first type in the parameter pack
+	TVariant<int, float, FString> Int;
+	check(Int.IsType<int>());
+	check(Int.Get<int>() == 0);
+	
+	TVariant<FString, int, float> String;
+	check(String.IsType<FString>());
+	check(String.Get<FString>() == FString{});
+	
+}
+
+TEST_CASE_NAMED(FTVariantEmptyVariantStateTest, "System::Core::Misc::TVariant::EmptyVariantState", "[ApplicationContextMask][EngineFilter]")
+{
+	// TVariant default constructor works for the provided FEmptyVariantState
+	TVariant<FEmptyVariantState, FNonDefaultConstructible> Var;
+	check(Var.IsType<FEmptyVariantState>());
+	
+}
+
+TEST_CASE_NAMED(FTVariantCopyTest, "System::Core::Misc::TVariant.Copy", "[ApplicationContextMask][EngineFilter]")
+{
+	// TVariants can be copied
+	using FVariant = TVariant<int, float, double>;
+	FVariant Int{ TInPlaceType<int>(), 40 };
+	FVariant Copy{ Int };
+
+	check(Int.IsType<int>());
+	check(Copy.IsType<int>());
+	check(Int.Get<int>() == Copy.Get<int>());
+
+}
+
+TEST_CASE_NAMED(FTVariantMoveTest, "System::Core::Misc::TVariant::Move", "[ApplicationContextMask][EngineFilter]")
+{
+	// TVariants can be moved
+	using FVariant = TVariant<int, float, double>;
+	FVariant Original{ TInPlaceType<int>(), 40 };
+	FVariant NewLoc{ MoveTemp(Original) };
+
+	check(NewLoc.IsType<int>());
+	check(NewLoc.Get<int>() == 40);
+
+}
+
+TEST_CASE_NAMED(FTVariantAssignTest, "System::Core::Misc::TVariant::Assign", "[ApplicationContextMask][EngineFilter]")
+{
+	// TVariants can be assigned from TVariants of the same template parameter pack
+
+	using FVariant = TVariant<int, float, double>;
+
+	// Copy Assign
+	{
+		FVariant One{ TInPlaceType<double>(), 20.0 };
+		FVariant Two{ TInPlaceType<int>(), 40 };
+
+		One = Two;
+		check(One.IsType<int>());
+		check(One.Get<int>() == 40);
+
+		check(Two.IsType<int>());
+		check(Two.Get<int>() == 40);
+	}
+
+	// Move Assign
+	{
+		FVariant One{ TInPlaceType<double>(), 20.0 };
+		FVariant Two{ TInPlaceType<int>(), 40 };
+
+		One = MoveTemp(Two);
+		check(One.IsType<int>());
+		check(One.Get<int>() == 40);
+	}
+
+}
+
+TEST_CASE_NAMED(FTVariantConstGetTest, "System::Core::Misc::TVariant::ConstGet", "[ApplicationContextMask][EngineFilter]")
+{
+	// const TVariant values can be gotten
+	const TVariant<int, float, double> ConstVar{ TInPlaceType<float>(), 20.0f };
+	check(ConstVar.IsType<float>());
+	check(ConstVar.Get<float>() == 20.0f);
+
+}
+
+TEST_CASE_NAMED(FTVariantTryGetTest, "System::Core::Misc::TVariant::TryGet", "[ApplicationContextMask][EngineFilter]")
+{
+	// TVariant supports conditional get
+
+	// mutable
+	{
+		TVariant<int, float, double> Int{ TInPlaceType<int>(), 20 };
+
+		check(Int.TryGet<double>() == nullptr);
+
+		int* Value = Int.TryGet<int>();
+		check(Value != nullptr);
+		check(*Value == 20);
+	}
+
+	// const
+	{
+		const TVariant<int, float, double> Int{ TInPlaceType<int>(), 20 };
+
+		check(Int.TryGet<double>() == nullptr);
+
+		const int* Value = Int.TryGet<int>();
+		check(Value != nullptr);
+		check(*Value == 20);
+	}
+
+}
+
+TEST_CASE_NAMED(FTVariantSharedPtrTest, "System::Core::Misc::TVariant::SharedPtr", "[ApplicationContextMask][EngineFilter]")
+{
+	// TVariant can successfully hold a TSharedPtr that goes out of scope
+	
+	// Assignment
+	{
+		TVariant<TSharedPtr<int>> Var;
+		{
+			TSharedPtr<int> Ptr = MakeShared<int>(20);
+			Var.Set<TSharedPtr<int>>(Ptr);
+		}
+
+		check(Var.IsType<TSharedPtr<int>>());
+
+		TSharedPtr<int>& PtrRef = Var.Get<TSharedPtr<int>>();
+		check(PtrRef.IsValid());
+		check(*PtrRef == 20);
+	}
+
+	// Construction
+	{
+		TUniquePtr<TVariant<TSharedPtr<int>>> Var = nullptr;
+		{
+			TSharedPtr<int> Ptr = MakeShared<int>(20);
+			Var = MakeUnique<TVariant<TSharedPtr<int>>>(TInPlaceType<TSharedPtr<int>>(), Ptr);
+		}
+
+		check(Var->IsType<TSharedPtr<int>>());
+
+		TSharedPtr<int>& PtrRef = Var->Get<TSharedPtr<int>>();
+		check(PtrRef.IsValid());
+		check(*PtrRef == 20);
+	}
+}
+
+TEST_CASE_NAMED(FTVariantUniquePtrTest, "System::Core::Misc::TVariant::UniquePtr", "[ApplicationContextMask][EngineFilter]")
+{
+	// TVariant can successfully hold a TUniquePtr
+
+	// Assignment
+	{
+		TVariant<TUniquePtr<int>> Var;
+		Var.Set<TUniquePtr<int>>(MakeUnique<int>(20));
+
+		check(Var.IsType<TUniquePtr<int>>());
+
+		TUniquePtr<int>& PtrRef = Var.Get<TUniquePtr<int>>();
+		check(*PtrRef == 20);
+	}
+
+	// Construction
+	{
+		TVariant<TUniquePtr<int>> Var{ TInPlaceType<TUniquePtr<int>>(), MakeUnique<int>(20) };
+		check(Var.IsType<TUniquePtr<int>>());
+
+		TUniquePtr<int>& PtrRef = Var.Get<TUniquePtr<int>>();
+		check(*PtrRef == 20);
+	}
+
+}
+
+TEST_CASE_NAMED(FTVariantMoveAssignNonCopyableTest, "System::Core::Misc::TVariant::MoveAssignNonCopyable", "[ApplicationContextMask][EngineFilter]")
+{
+	// TVariant holding a non-copyable can be move assigned
+	struct FNoCopy
+	{
+		int Data = 0;
+		FNoCopy() = default;
+		explicit FNoCopy(int _Data) : Data(_Data) {}
+		FNoCopy(const FNoCopy&) = delete;
+		FNoCopy(FNoCopy&& Other) = default;
+		FNoCopy& operator=(const FNoCopy&) = delete;
+		FNoCopy& operator=(FNoCopy&& Other) = default;
+	};
+
+	TVariant<FNoCopy> Var{ TInPlaceType<FNoCopy>(), 20 };
+
+	TVariant<FNoCopy> Moved;
+	Moved = MoveTemp(Var);
+	check(Moved.Get<FNoCopy>().Data == 20);
+
+}
+
+TEST_CASE_NAMED(FTVariantEmplaceTest, "System::Core::Misc::TVariant::Emplace", "[ApplicationContextMask][EngineFilter]")
+{
+	// TVariant can emplace values directly
+	TVariant<int, float, FString> Var;
+
+	Var.Emplace<int>(20);
+	check(Var.IsType<int>());
+	check(Var.Get<int>() == 20);
+
+	Var.Emplace<float>(30.0f);
+	check(Var.IsType<float>());
+	check(Var.Get<float>() == 30.0f);
+
+	Var.Emplace<FString>(TEXT("FOO"));
+	check(Var.IsType<FString>());
+	check(Var.Get<FString>() == FString{ TEXT("FOO") });
+
+}
+
+TEST_CASE_NAMED(FTVariantEmplaceNonMoveNonCopyTest, "System::Core::Misc::TVariant::EmplaceNonMoveNonCopy", "[ApplicationContextMask][EngineFilter]")
+{
+	// TVariant can emplace values that are neither moveable nor copyable
+	struct FNonMoveNonCopy
+	{
+		int Data;
+		explicit FNonMoveNonCopy(int _Data) : Data(_Data) {}
+		~FNonMoveNonCopy() = default;
+		FNonMoveNonCopy(FNonMoveNonCopy&&) = delete;
+		FNonMoveNonCopy(const FNonMoveNonCopy&) = delete;
+		FNonMoveNonCopy& operator=(FNonMoveNonCopy&&) = delete;
+		FNonMoveNonCopy& operator=(const FNonMoveNonCopy&) = delete;
+	};
+
+	TVariant<FEmptyVariantState, FNonMoveNonCopy> Var;
+	Var.Emplace<FNonMoveNonCopy>(20);
+	check(Var.Get<FNonMoveNonCopy>().Data == 20);
+
+}
+
+TEST_CASE_NAMED(FTVariantSelfAssignTest, "System::Core::Misc::TVariant::SelfAssign", "[ApplicationContextMask][EngineFilter]")
+{
+	// TVariant self-assignment is harmless
+	TVariant<int, float, double> Var{ TInPlaceType<int>(), 20 };
+
+	const TVariant<int, float, double>& Ref = Var; // trick clang into thinking we're not assigning the same variable to itself.
+
+	Var = Ref;
+	check(Var.Get<int>() == 20);
+
+	Var = MoveTemp(Var);
+	check(Var.Get<int>() == 20);
+
+}
+
+TEST_CASE_NAMED(FTVariantSwitchPatternMatchingTest, "System::Core::Misc::TVariant::SwitchPatternMatching", "[ApplicationContextMask][EngineFilter]")
+{
+	// TVariant supports switch-based 'pattern matching
+	using FVariant = TVariant<int, float, double>;
+	check(FVariant::IndexOfType<int>() == 0);
+	check(FVariant::IndexOfType<float>() == 1);
+	check(FVariant::IndexOfType<double>() == 2);
+
+	SIZE_T IndexTaken = (SIZE_T)-1;
+	auto Tester = [&IndexTaken](const FVariant& Var)
+	{
+		switch (IndexTaken = Var.GetIndex())
+		{
+			case FVariant::IndexOfType<int>():
+			{
+				int Value = Var.Get<int>();
+				check(Value == 20);
+				break;
+			}
+
+			case FVariant::IndexOfType<float>():
+			{
+				float Value = Var.Get<float>();
+				check(Value == 30.0f);
+				break;
+			}
+
+			case FVariant::IndexOfType<double>():
+			{
+				double Value = Var.Get<double>();
+				check(Value == 10.0);
+				break;
+			}
+
+			default:
+				check(!"Should not get here");
+				break;
+		}
+	};
+
+	FVariant Var;
+
+	Var.Set<int>(20);
+	Tester(Var);
+	check(IndexTaken == FVariant::IndexOfType<int>());
+
+	Var.Set<float>(30.0f);
+	Tester(Var);
+	check(IndexTaken == FVariant::IndexOfType<float>());
+
+	Var.Set<double>(10.0);
+	Tester(Var);
+	check(IndexTaken == FVariant::IndexOfType<double>());
+
+}
+
+
+TEST_CASE_NAMED(FTVariantDestructorTest, "System::Core::Misc::TVariant::Destructor", "[ApplicationContextMask][EngineFilter]")
+{
+	// Destructor test
+	struct FTestItem
+	{
+		int* DestructorCallCount;
+		FTestItem(int* _DestructorCallCount) : DestructorCallCount(_DestructorCallCount) {}
+		~FTestItem()
+		{
+			(*DestructorCallCount)++;
+		}
+	};
+
+	int DestructorCallCount = 0;
+	{
+		TVariant<TUniquePtr<FTestItem>> MovedFrom{ TInPlaceType<TUniquePtr<FTestItem>>(), MakeUnique<FTestItem>(&DestructorCallCount) };
+		TVariant<TUniquePtr<FTestItem>> MovedInto{ MoveTemp(MovedFrom) };
+	}
+	check(DestructorCallCount == 1);
+
+	DestructorCallCount = 0;
+	{
+		TVariant<TSharedPtr<FTestItem>> MovedFrom{ TInPlaceType<TSharedPtr<FTestItem>>(), MakeShared<FTestItem>(&DestructorCallCount) };
+		TVariant<TSharedPtr<FTestItem>> MovedInto{ MoveTemp(MovedFrom) };
+	}
+	check(DestructorCallCount == 1);
+
+}
+
+TEST_CASE_NAMED(FTVariantVisitTest, "System::Core::Misc::TVariant::Visit", "[ApplicationContextMask][EngineFilter]")
+{
+	// Visit calls the correct function overload
+	enum class EVisitType
+	{
+		IntFloat,
+		IntInt,
+		IntStr,
+		FloatInt,
+		FloatFloat,
+		FloatStr,
+		StrInt,
+		StrFloat,
+		StrStr,
+	};
+	struct FCallable
+	{
+		EVisitType operator()(int, float) { return EVisitType::IntFloat; }
+		EVisitType operator()(int, int) { return EVisitType::IntInt; }
+		EVisitType operator()(int, const FString&) { return EVisitType::IntStr; }
+		EVisitType operator()(float, int) { return EVisitType::FloatInt; }
+		EVisitType operator()(float, float) { return EVisitType::FloatFloat; }
+		EVisitType operator()(float, const FString&) { return EVisitType::FloatStr; }
+		EVisitType operator()(const FString&, int) { return EVisitType::StrInt; }
+		EVisitType operator()(const FString&, float) { return EVisitType::StrFloat; }
+		EVisitType operator()(const FString&, const FString&) { return EVisitType::StrStr; }
+	};
+
+	TVariant<int, float> IntV(TInPlaceType<int>(), 0);
+	const TVariant<int, float, FString> FloatV(TInPlaceType<float>(), 0.0f);
+	const TVariant<FString> StrV(TInPlaceType<FString>(), TEXT(""));
+
+#define CHECK_VISIT(Type1, Type2) check(Visit(FCallable{}, Type1##V, Type2##V) == EVisitType::Type1##Type2)
+	CHECK_VISIT(Int, Int);
+	CHECK_VISIT(Int, Float);
+	CHECK_VISIT(Int, Str);
+	CHECK_VISIT(Float, Int);
+	CHECK_VISIT(Float, Float);
+	CHECK_VISIT(Float, Str);
+	CHECK_VISIT(Str, Int);
+	CHECK_VISIT(Str, Float);
+	CHECK_VISIT(Str, Str);
+#undef CHECK_VISIT
+
+}
+
+TEST_CASE_NAMED(FTVariantVisitorRvalueReferenceParamTest, "System::Core::Misc::TVariant::VisitorRvalueReferenceParam", "[ApplicationContextMask][EngineFilter]")
+{
+	// Visitor implementation can have rvalue reference parameters
+	struct FAnonOne {};
+	struct FAnonTwo {};
+
+	using FVariant = TVariant<FAnonOne, FAnonTwo>;
+
+	struct FAnonVisitor
+	{
+		void operator()(FAnonOne&&)
+		{
+		}
+
+		void operator()(FAnonTwo&&)
+		{
+		}
+	};
+
+	FVariant MoveThisVariant;
+	Visit(FAnonVisitor{}, MoveTemp(MoveThisVariant));
+
+	// No checks, this is a compilation test.
+}
+
+TEST_CASE_NAMED(FTVariantSerializationTest, "System::Core::Misc::TVariant::Serialization", "[ApplicationContextMask][EngineFilter]")
+{
+	// TVariants can be serialized
+	TArray<uint8> Data;
+	FMemoryWriter Writer(Data);
+
+	TVariant<int, FString> VarData;
+	VarData.Emplace<int>(25);
+
+	Writer << VarData;
+
+	VarData.Emplace<FString>(TEXT("Foo"));
+	Writer << VarData;
+
+
+	FMemoryReader Reader(Data);
+	Reader << VarData;
+	check(VarData.IsType<int>());
+	check(VarData.Get<int>() == 25);
+
+	Reader << VarData;
+	check(VarData.IsType<FString>());
+	check(VarData.Get<FString>() == FString(TEXT("Foo")));
+
+}
+
+#endif //WITH_TESTS
diff --git a/Engine/Source/Runtime/Core/Tests/Misc/TimecodeTest.cpp b/Engine/Source/Runtime/Core/Tests/Misc/TimecodeTest.cpp
new file mode 100644
--- /dev/null
+++ b/Engine/Source/Runtime/Core/Tests/Misc/TimecodeTest.cpp
@@ -0,0 +1,249 @@
+// Copyright Epic Games, Inc. All Rights Reserved.
+
+#include "CoreTypes.h"
+#include "Async/Async.h"
+#include "Async/Future.h"
+#include "Containers/UnrealString.h"
+#include "Logging/LogMacros.h"
+#include "Math/NumericLimits.h"
+#include "Misc/FrameRate.h"
+#include "Misc/Timecode.h"
+#include "Misc/Timespan.h"
+#include "Tests/TestHarnessAdapter.h"
+
+
+#if WITH_TESTS
+
+/**
+ * Run a suite of timecode conversion operations to validate conversion from timecode to timespan/FrameNumber are working
+ *
+ * Drop Frame drop a frame every minute except every 10th minute
+ * 29.97fps
+ * 00:58:01:28 ; 00:58:01:29 ; 00:58:02:00 ; 00:58:02:01 (no skip)
+ * 01:00:59:28 ; 01:00:59:29 ; 01:01:00:02 ; 01:01:00:03 (every minute, we skip frame 0 and 1)
+ * 01:09:59:28 ; 01:09:59:29 ; 01:10:00:00 ; 01:10:00:01 (except every 10th minute, we include frame 0 and 1)
+ */
+TEST_CASE_NAMED(FTimecodeTest, "System::Core::Misc::Timecode", "[ApplicationContextMask][EngineFilter]")
+{
+	FFrameRate CommonFrameRates[]{
+		FFrameRate(12, 1),
+		FFrameRate(15, 1),
+		FFrameRate(24, 1),
+		FFrameRate(25, 1),
+		FFrameRate(30, 1),
+		FFrameRate(48, 1),
+		FFrameRate(48, 2), // Should give the same result as 24/1
+		FFrameRate(50, 1),
+		FFrameRate(60, 1),
+		FFrameRate(100, 1),
+		FFrameRate(120, 1),
+		FFrameRate(240, 1),
+		FFrameRate(24000, 1001),
+		FFrameRate(30000, 1001),
+		FFrameRate(48000, 1001),
+		FFrameRate(60000, 1001),
+	};
+
+	auto ConversionWithFrameRateTest = [](const FFrameRate FrameRate)
+	{
+		const bool bIsDropFrame = FTimecode::IsDropFormatTimecodeSupported(FrameRate);
+		int32 NumberOfErrors = 0;
+		FTimecode PreviousTimecodeValue;
+
+		const int64 StartIndex = 0;
+		for (int64 FrameIndex = StartIndex; FrameIndex <= MAX_int32; ++FrameIndex)
+		{
+			const FFrameNumber FrameNumber = static_cast<int32>(FrameIndex);
+			const FTimecode TimecodeValue = FTimecode::FromFrameNumber(FrameNumber, FrameRate, bIsDropFrame);
+			bool bDoTest = true;
+
+			// Conversion from FrameNumber to Timecode
+			if (bDoTest)
+			{
+				const FFrameNumber ExpectedFrameNumber = TimecodeValue.ToFrameNumber(FrameRate);
+				if (FrameNumber != ExpectedFrameNumber)
+				{
+					FAIL_CHECK(FString::Printf(TEXT("Timecode '%s' didn't convert properly from FrameNumber '%d' for FrameRate '%s'.")
+						, *TimecodeValue.ToString()
+						, FrameNumber.Value
+						, *FrameRate.ToPrettyText().ToString()
+					));
+					bDoTest = false;
+					++NumberOfErrors;
+				}
+			}
+			
+			// Conversion from Timespan to Timecode
+			if (bDoTest)
+			{
+				const FTimespan TimespanFromTimecode = TimecodeValue.ToTimespan(FrameRate);
+				const FTimecode TimecodeFromTimespanWithRollover = FTimecode::FromTimespan(TimespanFromTimecode, FrameRate, bIsDropFrame, true);
+				const FTimecode TimecodeFromTimespanWithoutRollover = FTimecode::FromTimespan(TimespanFromTimecode, FrameRate, bIsDropFrame, false);
+
+				if (TimecodeFromTimespanWithoutRollover != TimecodeValue)
+				{
+					FAIL_CHECK(FString::Printf(TEXT("Timecode '%s' didn't convert properly from Timespan '%f' with rollover for frame rate '%s'.")
+						, *TimecodeValue.ToString()
+						, TimespanFromTimecode.GetTotalSeconds()
+						, *FrameRate.ToPrettyText().ToString()
+					));
+					bDoTest = false;
+					++NumberOfErrors;
+				}
+				else if (TimecodeFromTimespanWithoutRollover.Minutes != TimecodeValue.Minutes || TimecodeFromTimespanWithoutRollover.Seconds != TimecodeValue.Seconds || TimecodeFromTimespanWithoutRollover.Frames != TimecodeValue.Frames)
+				{
+					FAIL_CHECK(FString::Printf(TEXT("Timecode '%s' didn't convert properly from Timespan '%f' without rollover for frame rate '%s'.")
+						, *TimecodeValue.ToString()
+						, TimespanFromTimecode.GetTotalSeconds()
+						, *FrameRate.ToPrettyText().ToString()
+					));
+					bDoTest = false;
+					++NumberOfErrors;
+				}
+				else if (!bIsDropFrame)
+				{
+					// Do they have the same hours, minutes, seconds
+					// To test this, we start from the number of events (FrameIndex) from which we got our timecode first
+					// Timecode is just a label and doesn't necessarily reflect real time especially with 23.976 FrameRate like
+					const int32 TotalSeconds = FMath::FloorToInt32((double)FrameIndex * FrameRate.AsInterval());
+					const int32 FrameHours = (TotalSeconds / (60 * 60));
+					const int32 FrameMinutes = ((TotalSeconds % (60 * 60)) / 60);
+					const int32 FrameSeconds = ((TotalSeconds % (60 * 60)) % 60);
+					const bool bHoursAreValid = (FrameHours % 24) == TimespanFromTimecode.GetHours() && (FrameHours / 24) == TimespanFromTimecode.GetDays();
+					const bool bMinutesAreValid = FrameMinutes == TimespanFromTimecode.GetMinutes();
+					const bool bSecondsAreValid = FrameSeconds == TimespanFromTimecode.GetSeconds();
+					if (!bHoursAreValid || !bMinutesAreValid || !bSecondsAreValid)
+					{
+						FAIL_CHECK(FString::Printf(TEXT("Timecode hours/minutes/seconds doesn't matches with Timespan '%s' from frame rate '%s'.")
+							, *TimespanFromTimecode.ToString()
+							, *FrameRate.ToPrettyText().ToString()
+						));
+						bDoTest = false;
+						++NumberOfErrors;
+					}
+				}
+			}
+
+			// Test if the frame number is incrementing
+			bool bIsPreviousTimecodeValid = FrameIndex != StartIndex;
+			if (bDoTest && bIsPreviousTimecodeValid)
+			{
+				bool bWrongFrame = PreviousTimecodeValue.Frames + 1 != TimecodeValue.Frames
+					&& TimecodeValue.Frames != 0;
+				const bool bWrongSeconds = PreviousTimecodeValue.Seconds != TimecodeValue.Seconds
+					&& PreviousTimecodeValue.Seconds + 1 != TimecodeValue.Seconds
+					&& TimecodeValue.Seconds != 0;
+				const bool bWrongMinutes = PreviousTimecodeValue.Minutes != TimecodeValue.Minutes
+					&& PreviousTimecodeValue.Minutes + 1 != TimecodeValue.Minutes
+					&& TimecodeValue.Minutes != 0;
+
+				if (bWrongFrame && bIsDropFrame)
+				{
+					// if new minute but not multiple of 10 mins, 2|4 is expected
+					const int32 NumberOfFramesInSecond = FMath::CeilToInt((float)FrameRate.AsDecimal());
+					const int32 NumberOfTimecodesToDrop = NumberOfFramesInSecond <= 30 ? 2 : 4;
+					bWrongFrame = !(TimecodeValue.Frames == NumberOfTimecodesToDrop && PreviousTimecodeValue.Minutes + 1 == TimecodeValue.Minutes && TimecodeValue.Minutes % 10 != 0);
+				}
+
+				if (bWrongFrame || bWrongSeconds || bWrongMinutes)
+				{
+					FAIL_CHECK(FString::Printf(TEXT("Timecode '%s' is not a continuity of the previous timecode '%s' from frame rate '%s'.")
+						, *TimecodeValue.ToString()
+						, *PreviousTimecodeValue.ToString()
+						, *FrameRate.ToPrettyText().ToString()
+					));
+					bDoTest = false;
+					++NumberOfErrors;
+				}
+			}
+
+			// Test frame rate that should be equivalent
+			if (bDoTest)
+			{
+				const FFrameRate EquivalentFrameRate = FFrameRate(FrameRate.Numerator * 3, FrameRate.Denominator * 3);
+				const FTimecode EquivalentTimecodeValue = FTimecode::FromFrameNumber(FrameNumber, EquivalentFrameRate, bIsDropFrame);
+				if (TimecodeValue != EquivalentTimecodeValue)
+				{
+					FAIL_CHECK(FString::Printf(TEXT("Timecode '%s' didn't convert properly from FrameNumber '%d' when the frame rate is tripled.")
+						, *TimecodeValue.ToString()
+						, FrameNumber.Value
+					));
+					bDoTest = false;
+					++NumberOfErrors;
+				}
+			}
+
+			// If we have a lot of errors with this frame rate, there is no need to log them all.
+			if (NumberOfErrors > 10)
+			{
+				WARN(FString::Printf(TEXT("Skip test for frame rate '%s'. Other errors may exists.")
+					, *FrameRate.ToPrettyText().ToString()
+				));
+				break;
+			}
+
+			PreviousTimecodeValue = TimecodeValue;
+
+			// LTC timecode support up to 40 hours
+			if (TimecodeValue.Hours >= 40)
+			{
+				break;
+			}
+		}
+
+		// Conversion from current time to Timecode
+		if (NumberOfErrors == 0)
+		{
+			const FTimespan CurrentTimespan = FTimespan(11694029893428);
+			const double CurrentSeconds = 1169402.9893428; //from FPlatformTime::Seconds()
+
+			const FTimecode FromTimespanTimecodeValueWithRollover = FTimecode::FromTimespan(CurrentTimespan, FrameRate, bIsDropFrame, true);
+			const FTimecode FromTimespanTimecodeValueWithoutRollover = FTimecode::FromTimespan(CurrentTimespan, FrameRate, bIsDropFrame, false);
+			const FTimecode FromSecondsTimecodeValueWithRollover = FTimecode(CurrentSeconds, FrameRate, bIsDropFrame, true);
+			const FTimecode FromSecondsTimecodeValueWithoutRollover = FTimecode(CurrentSeconds, FrameRate, bIsDropFrame, false);
+
+			if (FromTimespanTimecodeValueWithRollover != FromSecondsTimecodeValueWithRollover)
+			{
+				FAIL_CHECK(FString::Printf(TEXT("The timecode '%s' do not match timecode '%s' when converted from the computer clock's time and the frame rate is '%s'")
+					, *FromTimespanTimecodeValueWithRollover.ToString()
+					, *FromSecondsTimecodeValueWithRollover.ToString()
+					, *FrameRate.ToPrettyText().ToString()
+				));
+				++NumberOfErrors;
+			}
+			else if (FromTimespanTimecodeValueWithoutRollover != FromSecondsTimecodeValueWithoutRollover)
+			{
+				FAIL_CHECK(FString::Printf(TEXT("The timecode '%s' do not match timecode '%s' when converted from the computer clock's time and the frame rate is '%s'")
+					, *FromTimespanTimecodeValueWithoutRollover.ToString()
+					, *FromSecondsTimecodeValueWithoutRollover.ToString()
+					, *FrameRate.ToPrettyText().ToString()
+				));
+				++NumberOfErrors;
+			}
+			// Can't really test frame number matching between rollver timecode labels. We would need to exclude NDF fractional frame rates
+		}
+
+		INFO(FString::Printf(TEXT("Timecode test was completed with frame rate '%s'"), *FrameRate.ToPrettyText().ToString()));
+
+		return NumberOfErrors == 0;
+	};
+
+	// Test the conversion for all common frame rate
+	TArray<TFuture<bool>> Futures;
+	for (const FFrameRate& FrameRate : CommonFrameRates)
+	{
+		Futures.Add(Async(EAsyncExecution::Thread, [FrameRate, &ConversionWithFrameRateTest](){ return ConversionWithFrameRateTest(FrameRate); }));
+	}
+
+	bool bSuccessfully = true;
+	for (const TFuture<bool>& Future : Futures)
+	{
+		Future.Wait();
+		bSuccessfully = bSuccessfully && Future.Get();
+	}
+
+	REQUIRE(bSuccessfully);
+}
+
+
+#endif //WITH_TESTS
diff --git a/Engine/Source/Runtime/Core/Tests/Misc/TimespanTest.cpp b/Engine/Source/Runtime/Core/Tests/Misc/TimespanTest.cpp
new file mode 100644
--- /dev/null
+++ b/Engine/Source/Runtime/Core/Tests/Misc/TimespanTest.cpp
@@ -0,0 +1,133 @@
+// Copyright Epic Games, Inc. All Rights Reserved.
+
+#if WITH_TESTS
+
+#include "Misc/Timespan.h"
+
+#include "Containers/UnrealString.h"
+#include "Tests/TestHarnessAdapter.h"
+
+TEST_CASE_NAMED(FTimespanTest, "System::Core::Time::Timespan", "[Core][Time][SmokeFilter]")
+{
+	SECTION("Constructors")
+	{
+		FTimespan ts1_1 = FTimespan(3, 2, 1);
+		FTimespan ts1_2 = FTimespan(0, 3, 2, 1);
+		FTimespan ts1_3 = FTimespan(0, 3, 2, 1, 0);
+
+		CHECK(ts1_1 == ts1_2);
+		CHECK(ts1_1 == ts1_3);
+	}
+
+	SECTION("Components")
+	{
+		FTimespan ts2_1 = FTimespan(1, 2, 3, 4, 123456789);
+
+		CHECK(ts2_1.GetDays() == 1);
+		CHECK(ts2_1.GetHours() == 2);
+		CHECK(ts2_1.GetMinutes() == 3);
+		CHECK(ts2_1.GetSeconds() == 4);
+		CHECK(ts2_1.GetFractionMilli() == 123);
+		CHECK(ts2_1.GetFractionMicro() == 123456);
+		CHECK(ts2_1.GetFractionNano() == 123456700);
+	}
+
+	SECTION("Duration")
+	{
+		FTimespan ts3_1 = FTimespan(1, 2, 3, 4, 123456789);
+		FTimespan ts3_2 = FTimespan(-1, -2, -3, -4, -123456789);
+
+		CHECK(ts3_1.GetDuration() == ts3_2.GetDuration());
+	}
+
+	SECTION("FromXxx to GetTotalXxx")
+	{
+		CHECK(FTimespan::FromDays(123).GetTotalDays() == 123.0);
+		CHECK(FTimespan::FromHours(123).GetTotalHours() == 123.0);
+		CHECK(FTimespan::FromMinutes(123).GetTotalMinutes() == 123.0);
+		CHECK(FTimespan::FromSeconds(123).GetTotalSeconds() == 123.0);
+		CHECK(FTimespan::FromMilliseconds(123).GetTotalMilliseconds() == 123.0);
+		CHECK(FTimespan::FromMicroseconds(123).GetTotalMicroseconds() == 123.0);
+	}
+
+	SECTION("ToString")
+	{
+		FTimespan ts5_1 = FTimespan(1, 2, 3, 4, 123456789);
+
+		CHECK(ts5_1.ToString() == TEXT("+1.02:03:04.123"));
+		CHECK(ts5_1.ToString(TEXT("%d.%h:%m:%s.%f")) == TEXT("+1.02:03:04.123"));
+		CHECK(ts5_1.ToString(TEXT("%d.%h:%m:%s.%u")) == TEXT("+1.02:03:04.123456"));
+		CHECK(ts5_1.ToString(TEXT("%D.%h:%m:%s.%n")) == TEXT("+00000001.02:03:04.123456700"));
+	}
+
+	SECTION("Parse Valid")
+	{
+		FTimespan ts6_t;
+
+		FTimespan ts6_1 = FTimespan(1, 2, 3, 4, 123000000);
+		FTimespan ts6_2 = FTimespan(1, 2, 3, 4, 123456000);
+		FTimespan ts6_3 = FTimespan(1, 2, 3, 4, 123456700);
+
+		CHECK(FTimespan::Parse(TEXT("+1.02:03:04.123"), ts6_t));
+		CHECK(ts6_t == ts6_1);
+
+		CHECK(FTimespan::Parse(TEXT("+1.02:03:04.123456"), ts6_t));
+		CHECK(ts6_t == ts6_2);
+
+		CHECK(FTimespan::Parse(TEXT("+1.02:03:04.123456789"), ts6_t));
+		CHECK(ts6_t == ts6_3);
+
+		FTimespan ts6_4 = FTimespan(-1, -2, -3, -4, -123000000);
+		FTimespan ts6_5 = FTimespan(-1, -2, -3, -4, -123456000);
+		FTimespan ts6_6 = FTimespan(-1, -2, -3, -4, -123456700);
+
+		CHECK(FTimespan::Parse(TEXT("-1.02:03:04.123"), ts6_t));
+		CHECK(ts6_t == ts6_4);
+
+		CHECK(FTimespan::Parse(TEXT("-1.02:03:04.123456"), ts6_t));
+		CHECK(ts6_t == ts6_5);
+
+		CHECK(FTimespan::Parse(TEXT("-1.02:03:04.123456789"), ts6_t));
+		CHECK(ts6_t == ts6_6);
+	}
+
+	SECTION("Parse Invalid")
+	{
+		FTimespan ts7_1;
+
+		//CHECK_FALSE(FTimespan::Parse(TEXT("1,02:03:04.005"), ts7_1));
+		//CHECK_FALSE(FTimespan::Parse(TEXT("1.1.02:03:04:005"), ts7_1));
+		//CHECK_FALSE(FTimespan::Parse(TEXT("04:005"), ts7_1));
+	}
+
+	// `From*` converters must return correct value
+	// test normal and edge cases for polar conversions (FromMicroseconds() and FromDay()) and just normal case for all others
+	SECTION("FromXxx Exact Ticks")
+	{
+		CHECK(FTimespan::FromMicroseconds(0) == FTimespan(0));
+		CHECK(FTimespan::FromMicroseconds(1) == FTimespan(1 * ETimespan::TicksPerMicrosecond));
+		CHECK(FTimespan::FromMicroseconds(1.1) == FTimespan(1 * ETimespan::TicksPerMicrosecond + 1));
+		CHECK(FTimespan::FromMicroseconds(1.5) == FTimespan(1 * ETimespan::TicksPerMicrosecond + 5));
+		CHECK(FTimespan::FromMicroseconds(1.499999999999997) == FTimespan(1 * ETimespan::TicksPerMicrosecond + 5));
+		CHECK(FTimespan::FromMicroseconds(1.50000001) == FTimespan(1 * ETimespan::TicksPerMicrosecond + 5));
+		CHECK(FTimespan::FromMicroseconds(-1) == FTimespan(-1 * ETimespan::TicksPerMicrosecond));
+		CHECK(FTimespan::FromMicroseconds(-1.1) == FTimespan(-1 * ETimespan::TicksPerMicrosecond - 1));
+		CHECK(FTimespan::FromMicroseconds(-1.5) == FTimespan(-1 * ETimespan::TicksPerMicrosecond - 5));
+
+		CHECK(FTimespan::FromMilliseconds(1) == FTimespan(1 * ETimespan::TicksPerMillisecond));
+		CHECK(FTimespan::FromSeconds(1) == FTimespan(1 * ETimespan::TicksPerSecond));
+		CHECK(FTimespan::FromMinutes(1) == FTimespan(1 * ETimespan::TicksPerMinute));
+		CHECK(FTimespan::FromHours(1) == FTimespan(1 * ETimespan::TicksPerHour));
+
+		CHECK(FTimespan::FromDays(0) == FTimespan(0));
+		CHECK(FTimespan::FromDays(1) == FTimespan(1 * ETimespan::TicksPerDay));
+		CHECK(FTimespan::FromDays(1.25) == FTimespan(1 * ETimespan::TicksPerDay + 6 * ETimespan::TicksPerHour));
+		CHECK(FTimespan::FromDays(1.5) == FTimespan(1 * ETimespan::TicksPerDay + 12 * ETimespan::TicksPerHour));
+		CHECK(FTimespan::FromDays(1.499999999999997) == FTimespan(1 * ETimespan::TicksPerDay + 12 * ETimespan::TicksPerHour));
+		CHECK(FTimespan::FromDays(-1) == FTimespan(-1 * ETimespan::TicksPerDay));
+		CHECK(FTimespan::FromDays(-1.25) == FTimespan(-1 * ETimespan::TicksPerDay - 6 * ETimespan::TicksPerHour));
+		CHECK(FTimespan::FromDays(-1.5) == FTimespan(-1 * ETimespan::TicksPerDay - 12 * ETimespan::TicksPerHour));
+	}
+}
+
+#endif //WITH_TESTS
diff --git a/Engine/Source/Runtime/Core/Tests/Misc/TripleBufferTest.cpp b/Engine/Source/Runtime/Core/Tests/Misc/TripleBufferTest.cpp
new file mode 100644
--- /dev/null
+++ b/Engine/Source/Runtime/Core/Tests/Misc/TripleBufferTest.cpp
@@ -0,0 +1,107 @@
+// Copyright Epic Games, Inc. All Rights Reserved.
+
+#if WITH_TESTS
+
+#include "CoreTypes.h"
+#include "Containers/UnrealString.h"
+#include "Tests/TestHarnessAdapter.h"
+#include "Containers/TripleBuffer.h"
+#include "Math/RandomStream.h"
+#if WITH_LOW_LEVEL_TESTS
+#include "TestCommon/Comparisons.h"
+#endif
+
+TEST_CASE_NAMED(FTripleBufferTest, "System::Core::Misc::TripleBuffer", "[ApplicationContextMask][SmokeFilter]")
+{
+	// uninitialized buffer
+	{
+		TTripleBuffer<int32> Buffer(NoInit);
+
+		CHECK_FALSE_MESSAGE(TEXT("Uninitialized triple buffer must not be dirty"), Buffer.IsDirty());
+	}
+
+	// initialized buffer
+	{
+		TTripleBuffer<int32> Buffer(1);
+
+		CHECK_FALSE_MESSAGE(TEXT("Initialized triple buffer must not be dirty"), Buffer.IsDirty());
+		CHECK_EQUALS(TEXT("Initialized triple buffer must have correct read buffer value"), Buffer.Read(), 1);
+
+		Buffer.SwapReadBuffers();
+
+		CHECK_EQUALS(TEXT("Initialized triple buffer must have correct temp buffer value"), Buffer.Read(), 1);
+
+		Buffer.SwapWriteBuffers();
+
+		CHECK_MESSAGE(TEXT("Write buffer swap must set dirty flag"), Buffer.IsDirty());
+
+		Buffer.SwapReadBuffers();
+
+		CHECK_FALSE_MESSAGE(TEXT("Read buffer swap must clear dirty flag"), Buffer.IsDirty());
+		CHECK_EQUALS(TEXT("Initialized triple buffer must have correct temp buffer value"), Buffer.Read(), 1);
+	}
+
+	// pre-set buffer
+	{
+		int32 Array[3] = { 1, 2, 3 };
+		TTripleBuffer<int32> Buffer(Array);
+
+		int32 Read = Buffer.Read();
+		CHECK_EQUALS(TEXT("Pre-set triple buffer must have correct Read buffer value"), Read, 3);
+
+		Buffer.SwapReadBuffers();
+
+		int32 Temp = Buffer.Read();
+		CHECK_EQUALS(TEXT("Pre-set triple buffer must have correct Temp buffer value"), Temp, 1);
+
+		Buffer.SwapWriteBuffers();
+		Buffer.SwapReadBuffers();
+
+		int32 Write = Buffer.Read();
+		CHECK_EQUALS(TEXT("Pre-set triple buffer must have correct Write buffer value"), Write, 2);
+	}
+
+	// operations
+	{
+		TTripleBuffer<int32> Buffer;
+
+		for (int32 Index = 0; Index < 6; ++Index)
+		{
+			int32& Write = Buffer.GetWriteBuffer(); Write = Index; Buffer.SwapWriteBuffers();
+			Buffer.SwapReadBuffers();
+			CHECK_EQUALS(*FString::Printf(TEXT("Triple buffer must read correct value (%i)"), Index), Buffer.Read(), Index);
+		}
+
+		FRandomStream Rand;
+		int32 LastRead = -1;
+
+		for (int32 Index = 0; Index < 100; ++Index)
+		{
+			int32 Writes = Rand.GetUnsignedInt() % 4;
+
+			while (Writes > 0)
+			{
+				int32& Write = Buffer.GetWriteBuffer(); Write = Index; Buffer.SwapWriteBuffers();
+				--Writes;
+			}
+			
+			int32 Reads = Rand.GetUnsignedInt() % 4;
+
+			while (Reads > 0)
+			{
+				if (!Buffer.IsDirty())
+				{
+					break;
+				}
+
+				Buffer.SwapReadBuffers();
+				int32 Read = Buffer.Read();
+				CHECK_MESSAGE(TEXT("Triple buffer must read in increasing order"), Read > LastRead);
+				LastRead = Read;
+				--Reads;
+			}
+		}
+	}
+}
+
+#endif //WITH_TESTS
\ No newline at end of file
diff --git a/Engine/Source/Runtime/Core/Tests/Misc/TypeContainerTest.cpp b/Engine/Source/Runtime/Core/Tests/Misc/TypeContainerTest.cpp
new file mode 100644
--- /dev/null
+++ b/Engine/Source/Runtime/Core/Tests/Misc/TypeContainerTest.cpp
@@ -0,0 +1,250 @@
+// Copyright Epic Games, Inc. All Rights Reserved.
+
+#if WITH_TESTS
+
+#include "CoreTypes.h"
+#include "Templates/UnrealTypeTraits.h"
+#include "Containers/UnrealString.h"
+#include "Templates/Function.h"
+#include "Templates/SharedPointer.h"
+#include "Delegates/Delegate.h"
+#include "Tests/TestHarnessAdapter.h"
+#include "Async/Async.h"
+#include "Misc/TypeContainer.h"
+#if WITH_LOW_LEVEL_TESTS
+#include "TestCommon/Comparisons.h"
+#endif
+
+/* Helpers
+ *****************************************************************************/
+
+struct IFruit
+{
+	virtual ~IFruit() { }
+	virtual FString Name() = 0;
+};
+
+struct IBerry : public IFruit
+{
+	virtual ~IBerry() { }
+};
+
+struct FBanana : public IFruit
+{
+	virtual ~FBanana() { }
+	virtual FString Name() override { return TEXT("Banana"); }
+};
+
+struct FStrawberry : public IBerry
+{
+	virtual ~FStrawberry() { }
+	virtual FString Name() override { return TEXT("Strawberry"); }
+};
+
+template<ESPMode Mode = ESPMode::ThreadSafe>
+struct ISmoothie
+{
+	virtual ~ISmoothie() { }
+	virtual TSharedRef<IBerry, Mode> GetBerry() = 0;
+	virtual TSharedRef<IFruit, Mode> GetFruit() = 0;
+};
+
+template<ESPMode Mode = ESPMode::ThreadSafe>
+struct TSmoothie : public ISmoothie<Mode>
+{
+	TSmoothie(TSharedRef<IFruit, Mode> InFruit, TSharedRef<IBerry, Mode> InBerry) : Berry(InBerry), Fruit(InFruit) { }
+	virtual ~TSmoothie() { }
+	virtual TSharedRef<IBerry, Mode> GetBerry() override { return Berry; }
+	virtual TSharedRef<IFruit, Mode> GetFruit() override { return Fruit; }
+	TSharedRef<IBerry, Mode> Berry;
+	TSharedRef<IFruit, Mode> Fruit;
+};
+
+template<ESPMode Mode = ESPMode::ThreadSafe>
+struct TTwoSmoothies
+{
+	TSharedPtr<ISmoothie<Mode>, Mode> One;
+	TSharedPtr<ISmoothie<Mode>, Mode> Two;
+};
+
+
+DECLARE_DELEGATE_RetVal(TSharedRef<IBerry>, FBerryFactoryDelegate);
+DECLARE_DELEGATE_RetVal(TSharedRef<IFruit>, FFruitFactoryDelegate);
+DECLARE_DELEGATE_RetVal_TwoParams(TSharedRef<ISmoothie<>>, FSmoothieFactoryDelegate, TSharedRef<IFruit>, TSharedRef<IBerry>);
+
+Expose_TNameOf(IBerry)
+Expose_TNameOf(IFruit)
+Expose_TNameOf(ISmoothie<>)
+
+
+/* Tests
+ *****************************************************************************/
+
+TEST_CASE_NAMED(FTypeContainerTest, "System::Core::Misc::TypeContainer", "[ApplicationContextMask][EngineFilter]")
+{
+	// existing instance test
+	{
+		TSharedRef<IFruit> Fruit = MakeShareable(new FBanana());
+		TSharedRef<IBerry> Berry = MakeShareable(new FStrawberry());
+
+		TTypeContainer<> Container;
+		{
+			Container.RegisterInstance<IFruit>(Fruit);
+			Container.RegisterInstance<IBerry>(Berry);
+		}
+
+		auto Instance = Container.GetInstance<IFruit>();
+
+		CHECK_EQUALS(TEXT("Correct instance must be returned"), Instance, Fruit);
+		CHECK_NOT_EQUALS(TEXT("Incorrect instance must not be returned"), Instance, StaticCastSharedRef<IFruit>(Berry));
+	}
+
+	// per instance test
+	{
+		TTypeContainer<> Container;
+		{
+			Container.RegisterClass<IFruit, FBanana>(ETypeContainerScope::Instance);
+			Container.RegisterClass<IBerry, FStrawberry>(ETypeContainerScope::Instance);
+			Container.RegisterClass<ISmoothie<>, TSmoothie<>, IFruit, IBerry>(ETypeContainerScope::Instance); // !!!
+		}
+
+		auto Smoothie1 = Container.GetInstance<ISmoothie<>>();
+		auto Smoothie2 = Container.GetInstance<ISmoothie<>>();
+
+		CHECK_NOT_EQUALS(TEXT("For per-instances classes, a unique instance must be returned each time"), Smoothie1, Smoothie2);
+		CHECK_NOT_EQUALS(TEXT("For per-instances dependencies, a unique instance must be returned each time [1]"), Smoothie1->GetBerry(), Smoothie2->GetBerry());
+		CHECK_NOT_EQUALS(TEXT("For per-instances dependencies, a unique instance must be returned each time [2]"), Smoothie1->GetFruit(), Smoothie2->GetFruit());
+	}
+
+	// per thread test
+	{
+		TTypeContainer<ESPMode::ThreadSafe> Container;
+		{
+			Container.RegisterClass<IFruit, FBanana>(ETypeContainerScope::Thread);
+			Container.RegisterClass<IBerry, FStrawberry>(ETypeContainerScope::Instance);
+			Container.RegisterClass<ISmoothie<ESPMode::ThreadSafe>, TSmoothie<ESPMode::ThreadSafe>, IFruit, IBerry>(ETypeContainerScope::Thread); // !!!
+		}
+
+		TFunction<TTwoSmoothies<ESPMode::ThreadSafe>()> MakeSmoothies = [&]()
+		{
+			TTwoSmoothies<ESPMode::ThreadSafe> Smoothies;
+			Smoothies.One = Container.GetInstance<ISmoothie<ESPMode::ThreadSafe>>();
+			Smoothies.Two = Container.GetInstance<ISmoothie<ESPMode::ThreadSafe>>();
+			return Smoothies;
+		};
+
+		auto Smoothies1 = Async(EAsyncExecution::Thread, MakeSmoothies);
+		auto Smoothies2 = Async(EAsyncExecution::Thread, MakeSmoothies);
+
+		auto One1 = Smoothies1.Get().One;
+		auto Two1 = Smoothies1.Get().Two;
+		auto One2 = Smoothies2.Get().One;
+		auto Two2 = Smoothies2.Get().Two;
+
+		CHECK_EQUALS(TEXT("For per-thread classes, the same instance must be returned from the same thread [1]"), One1, Two1);
+		CHECK_EQUALS(TEXT("For per-thread classes, the same instance must be returned from the same thread [2]"), One2, Two2);
+		CHECK_NOT_EQUALS(TEXT("For per-thread classes, different instances must be returned from different threads [1]"), One1, One2);
+		CHECK_NOT_EQUALS(TEXT("For per-thread classes, different instances must be returned from different threads [2]"), Two1, Two2);
+	}
+
+	// per process test
+	{
+		TTypeContainer<ESPMode::ThreadSafe> Container;
+		{
+			Container.RegisterClass<IFruit, FBanana>(ETypeContainerScope::Thread);
+			Container.RegisterClass<IBerry, FStrawberry>(ETypeContainerScope::Instance);
+			Container.RegisterClass<ISmoothie<ESPMode::ThreadSafe>, TSmoothie<ESPMode::ThreadSafe>, IFruit, IBerry>(ETypeContainerScope::Process); //!!!
+		}
+
+		TFunction<TTwoSmoothies<ESPMode::ThreadSafe>()> MakeSmoothies = [&]()
+		{
+			TTwoSmoothies<ESPMode::ThreadSafe> Smoothies;
+			Smoothies.One = Container.GetInstance<ISmoothie<ESPMode::ThreadSafe>>();
+			Smoothies.Two = Container.GetInstance<ISmoothie<ESPMode::ThreadSafe>>();
+			return Smoothies;
+		};
+
+		auto Smoothies1 = Async(EAsyncExecution::Thread, MakeSmoothies);
+		auto Smoothies2 = Async(EAsyncExecution::Thread, MakeSmoothies);
+
+		auto One1 = Smoothies1.Get().One;
+		auto Two1 = Smoothies1.Get().Two;
+		auto One2 = Smoothies2.Get().One;
+		auto Two2 = Smoothies2.Get().Two;
+
+		CHECK_EQUALS(TEXT("For per-process classes, the same instance must be returned from the same thread [1]"), One1, Two1);
+		CHECK_EQUALS(TEXT("For per-process classes, the same instance must be returned from the same thread [2]"), One2, Two2);
+		CHECK_EQUALS(TEXT("For per-process classes, the same instance must be returned from different threads [1]"), One1, One2);
+		CHECK_EQUALS(TEXT("For per-process classes, the same instance must be returned from different threads [1]"), Two1, Two2);
+	}
+
+	// factory test
+	{
+		struct FLocal
+		{
+			static TSharedRef<IBerry> MakeStrawberry()
+			{
+				return MakeShareable(new FStrawberry());
+			}
+
+			static TSharedRef<ISmoothie<>> MakeSmoothie(TSharedRef<IFruit> Fruit, TSharedRef<IBerry> Berry)
+			{
+				return MakeShareable(new TSmoothie<>(Fruit, Berry));
+			}
+		};
+
+		TTypeContainer<> Container;
+		{
+			Container.RegisterFactory<IBerry>(&FLocal::MakeStrawberry);
+			Container.RegisterFactory<IFruit>(
+				[]() -> TSharedRef<IFruit> {
+					return MakeShareable(new FBanana());
+				}
+			);
+			Container.RegisterFactory<ISmoothie<>, IFruit, IBerry>(&FLocal::MakeSmoothie);
+		}
+
+		auto Berry = Container.GetInstance<IBerry>();
+		auto Fruit = Container.GetInstance<IFruit>();
+		auto Smoothie = Container.GetInstance<ISmoothie<>>();
+	}
+
+	// delegate test
+	{
+		struct FLocal
+		{
+			static TSharedRef<IBerry> MakeBerry()
+			{
+				return MakeShareable(new FStrawberry());
+			}
+
+			static TSharedRef<IFruit> MakeFruit(bool Banana)
+			{
+				if (Banana)
+				{
+					return MakeShareable(new FBanana());
+				}
+
+				return MakeShareable(new FStrawberry());
+			}
+
+			static TSharedRef<ISmoothie<>> MakeSmoothie(TSharedRef<IFruit> Fruit, TSharedRef<IBerry> Berry)
+			{
+				return MakeShareable(new TSmoothie<>(Fruit, Berry));
+			}
+		};
+
+		TTypeContainer<> Container;
+		{
+			Container.RegisterDelegate<IBerry>(FBerryFactoryDelegate::CreateStatic(&FLocal::MakeBerry));
+			Container.RegisterDelegate<IFruit>(FFruitFactoryDelegate::CreateStatic(&FLocal::MakeFruit, true));
+			Container.RegisterDelegate<ISmoothie<>, FSmoothieFactoryDelegate, IFruit, IBerry>(FSmoothieFactoryDelegate::CreateStatic(&FLocal::MakeSmoothie));
+		}
+
+		auto Fruit = Container.GetInstance<IFruit>();
+		auto Smoothie = Container.GetInstance<ISmoothie<>>();
+	}
+
+}
+
+#endif //WITH_TESTS
