Removed two unused PhysX stats

These stats were made obsolete by 10933328 and are not referenced anywhere. The trimesh one in particular has a very similar name to the "PhysXTriMesh" stat later in the file, differing only by the case of the 'm', causing any case-insensitive comparisons to not behave as oner might expect.

#rnx

diff --git a/Engine/Source/Runtime/Core/Private/HAL/LowLevelMemTracker.cpp b/Engine/Source/Runtime/Core/Private/HAL/LowLevelMemTracker.cpp
--- a/Engine/Source/Runtime/Core/Private/HAL/LowLevelMemTracker.cpp
+++ b/Engine/Source/Runtime/Core/Private/HAL/LowLevelMemTracker.cpp
@@ -236,8 +236,6 @@ DECLARE_LLM_MEMORY_STAT(TEXT("VirtualTextureSystem"), STAT_VirtualTextureSystemL
 DECLARE_LLM_MEMORY_STAT(TEXT("Render Targets"), STAT_RenderTargetsLLM, STATGROUP_LLMFULL);
 DECLARE_LLM_MEMORY_STAT(TEXT("SceneRender"), STAT_SceneRenderLLM, STATGROUP_LLMFULL);
 DECLARE_LLM_MEMORY_STAT(TEXT("RHIMisc"), STAT_RHIMiscLLM, STATGROUP_LLMFULL);
-DECLARE_LLM_MEMORY_STAT(TEXT("PhysX TriMesh"), STAT_PhysXTriMeshLLM, STATGROUP_LLMFULL);
-DECLARE_LLM_MEMORY_STAT(TEXT("PhysX ConvexMesh"), STAT_PhysXConvexMeshLLM, STATGROUP_LLMFULL);
 DECLARE_LLM_MEMORY_STAT(TEXT("AsyncLoading"), STAT_AsyncLoadingLLM, STATGROUP_LLMFULL);
 DECLARE_LLM_MEMORY_STAT(TEXT("UObject"), STAT_UObjectLLM, STATGROUP_LLMFULL);
 DECLARE_LLM_MEMORY_STAT(TEXT("Animation"), STAT_AnimationLLM, STATGROUP_LLMFULL);
