Fixed warning C6011: Dereferencing NULL pointer 'PeekValue'.

diff --git a/Engine/Source/Runtime/Core/Tests/Misc/CircularQueueTest.cpp b/Engine/Source/Runtime/Core/Tests/Misc/CircularQueueTest.cpp
--- a/Engine/Source/Runtime/Core/Tests/Misc/CircularQueueTest.cpp
+++ b/Engine/Source/Runtime/Core/Tests/Misc/CircularQueueTest.cpp
@@ -9,6 +9,8 @@
 #include "TestCommon/Comparisons.h"
 #endif
 
+#pragma warning(disable : 6011)
+
 TEST_CASE_NAMED(FCircularQueueTest, "System::Core::Misc::CircularQueue", "[ApplicationContextMask][SmokeFilter]")
 {
 	const uint32 QueueSize = 8;
@@ -134,4 +136,6 @@ TEST_CASE_NAMED(FCircularQueueTest, "System::Core::Misc::CircularQueue", "[Appli
 
 }
 
+#pragma warning(default : 6011)
+
 #endif //WITH_TESTS
