Used a scoped access handler instead of macros for delegate thread-safety detection.

diff --git a/Engine/Source/Runtime/Core/Public/Delegates/Delegate.h b/Engine/Source/Runtime/Core/Public/Delegates/Delegate.h
--- a/Engine/Source/Runtime/Core/Public/Delegates/Delegate.h
+++ b/Engine/Source/Runtime/Core/Public/Delegates/Delegate.h
@@ -3,7 +3,6 @@
 #pragma once
 
 #include "CoreTypes.h"
-#include "Delegates/DelegateMacros.h"
 #include "Misc/AssertionMacros.h"
 #include "UObject/NameTypes.h"
 #include "Templates/SharedPointer.h"
diff --git a/Engine/Source/Runtime/Core/Public/Delegates/DelegateAccessHandler.h b/Engine/Source/Runtime/Core/Public/Delegates/DelegateAccessHandler.h
new file mode 100644
--- /dev/null
+++ b/Engine/Source/Runtime/Core/Public/Delegates/DelegateAccessHandler.h
@@ -0,0 +1,91 @@
+// Copyright Epic Games, Inc. All Rights Reserved.
+
+#pragma once
+
+#include "CoreTypes.h"
+#include "HAL/CriticalSection.h"
+#include "Misc/MTAccessDetector.h"
+#include "Misc/ScopeLock.h"
+
+//#define UE_DETECT_DELEGATES_RACE_CONDITIONS 1
+
+#if !defined(UE_DETECT_DELEGATES_RACE_CONDITIONS)
+	#define UE_DETECT_DELEGATES_RACE_CONDITIONS 0
+#endif
+
+/**
+ * non thread-safe version that does not do any race detection. supposed to be used in a controlled environment that provides own
+ * detection or synchronization.
+ */
+class FDelegateAccessHandlerBaseUnchecked
+{
+protected:
+	struct FReadAccessScope {};
+	struct FWriteAccessScope {};
+
+	[[nodiscard]] FReadAccessScope GetReadAccessScope() const
+	{
+		return {};
+	}
+
+	[[nodiscard]] FWriteAccessScope GetWriteAccessScope()
+	{
+		return {};
+	}
+};
+
+#if UE_DETECT_DELEGATES_RACE_CONDITIONS
+
+/**
+ * non thread-safe version that detects not thread-safe delegates used concurrently (dev builds only)
+ */
+class FDelegateAccessHandlerBaseChecked
+{
+protected:
+	class FReadAccessScope
+	{
+	public:
+		UE_NONCOPYABLE(FReadAccessScope);
+
+		explicit FReadAccessScope(FRWFullyRecursiveAccessDetector& InAccessDetector)
+			: Accessor(MakeScopedReaderAccessDetector(InAccessDetector))
+		{
+		}
+
+	private:
+		TScopedReaderAccessDetector<FRWFullyRecursiveAccessDetector> Accessor;
+	};
+
+	[[nodiscard]] FReadAccessScope GetReadAccessScope() const
+	{
+		return FReadAccessScope(Accessor);
+	}
+
+	class FWriteAccessScope
+	{
+	public:
+		UE_NONCOPYABLE(FWriteAccessScope);
+
+		explicit FWriteAccessScope(FRWFullyRecursiveAccessDetector& InAccessDetector)
+			: Accessor(MakeScopedWriterAccessDetector(InAccessDetector))
+		{
+		}
+
+	private:
+		TScopedWriterDetector<FRWFullyRecursiveAccessDetector> Accessor;
+	};
+
+	[[nodiscard]] FWriteAccessScope GetWriteAccessScope()
+	{
+		return FWriteAccessScope(Accessor);
+	}
+
+private:
+	mutable FRWFullyRecursiveAccessDetector Accessor;
+};
+
+#else // UE_DETECT_DELEGATES_RACE_CONDITIONS
+
+using FDelegateAccessHandlerBaseChecked = FDelegateAccessHandlerBaseUnchecked;
+
+#endif // UE_DETECT_DELEGATES_RACE_CONDITIONS
diff --git a/Engine/Source/Runtime/Core/Public/Delegates/DelegateBase.h b/Engine/Source/Runtime/Core/Public/Delegates/DelegateBase.h
--- a/Engine/Source/Runtime/Core/Public/Delegates/DelegateBase.h
+++ b/Engine/Source/Runtime/Core/Public/Delegates/DelegateBase.h
@@ -6,9 +6,9 @@
 #include "Containers/ContainerAllocationPolicies.h"
 #include "Math/UnrealMathUtility.h"
 #include "UObject/NameTypes.h"
+#include "Delegates/DelegateAccessHandler.h"
 #include "Delegates/DelegateSettings.h"
 #include "Delegates/IDelegateInstance.h"
-#include "Delegates/DelegateMacros.h"
 
 #if !defined(_WIN32) || defined(_WIN64) || (defined(ALLOW_DELEGATE_INLINE_ALLOCATORS_ON_WIN32) && ALLOW_DELEGATE_INLINE_ALLOCATORS_ON_WIN32)
 	typedef TAlignedBytes<16, 16> FAlignedInlineDelegateType;
@@ -79,12 +79,20 @@ struct FDefaultTSDelegateUserPolicy
  * Base class for unicast delegates.
  */
 template <typename UserPolicy>
-class TDelegateBase
+class TDelegateBase : public FDelegateAccessHandlerBaseChecked
 {
+private:
+	using Super = FDelegateAccessHandlerBaseChecked;
+
 	template <typename>
 	friend class TMulticastDelegateBase;
 
 protected:
+	using typename Super::FReadAccessScope;
+	using Super::GetReadAccessScope;
+	using typename Super::FWriteAccessScope;
+	using Super::GetWriteAccessScope;
+
 	explicit TDelegateBase() = default;
 
 public:
@@ -109,7 +117,7 @@ public:
 	 */
 	FORCEINLINE void Unbind()
 	{
-		UE_DELEGATES_MT_SCOPED_WRITE_ACCESS(AccessDetector);
+		FWriteAccessScope WriteScope = GetWriteAccessScope();
 
 		if (IDelegateInstance* Ptr = GetDelegateInstanceProtected())
 		{
@@ -225,7 +233,7 @@ public:
 	template<typename DelegateInstanceType, typename... DelegateInstanceParams>
 	void CreateDelegateInstance(DelegateInstanceParams&&... Params)
 	{
-		UE_DELEGATES_MT_SCOPED_READ_ACCESS(AccessDetector);
+		FReadAccessScope ReadScope = GetReadAccessScope();
 
 		IDelegateInstance* DelegateInstance = GetDelegateInstanceProtected();
 		if (DelegateInstance)
@@ -245,14 +253,14 @@ protected:
 	 */
 	FORCEINLINE IDelegateInstance* GetDelegateInstanceProtected()
 	{
-		UE_DELEGATES_MT_SCOPED_READ_ACCESS(AccessDetector);
+		FReadAccessScope ReadScope = GetReadAccessScope();
 
 		return DelegateSize ? (IDelegateInstance*)DelegateAllocator.GetAllocation() : nullptr;
 	}
 
 	FORCEINLINE const IDelegateInstance* GetDelegateInstanceProtected() const
 	{
-		UE_DELEGATES_MT_SCOPED_READ_ACCESS(AccessDetector);
+		FReadAccessScope ReadScope = GetReadAccessScope();
 
 		return DelegateSize ? (const IDelegateInstance*)DelegateAllocator.GetAllocation() : nullptr;
 	}
@@ -260,7 +268,7 @@ protected:
 private:
 	void* Allocate(int32 Size)
 	{
-		UE_DELEGATES_MT_SCOPED_WRITE_ACCESS(AccessDetector);
+		FWriteAccessScope WriteScope = GetWriteAccessScope();
 
 		int32 NewDelegateSize = FMath::DivideAndRoundUp(Size, (int32)sizeof(FAlignedInlineDelegateType));
 		if (DelegateSize != NewDelegateSize)
@@ -281,7 +289,7 @@ private:
 
 	void MoveAssign(TDelegateBase&& Other)
 	{
-		UE_DELEGATES_MT_SCOPED_WRITE_ACCESS(AccessDetector);
+		FWriteAccessScope WriteScope = GetWriteAccessScope();
 
 		Unbind();
 		DelegateAllocator.MoveToEmpty(Other.DelegateAllocator);
@@ -292,6 +300,4 @@ private:
 private:
 	FDelegateAllocatorType::ForElementType<FAlignedInlineDelegateType> DelegateAllocator;
 	int32 DelegateSize = 0;
-
-	UE_DELEGATES_MT_ACCESS_DETECTOR(AccessDetector);
 };
diff --git a/Engine/Source/Runtime/Core/Public/Delegates/DelegateMacros.h b/Engine/Source/Runtime/Core/Public/Delegates/DelegateMacros.h
deleted file mode 100644
diff --git a/Engine/Source/Runtime/Core/Public/Delegates/MulticastDelegateBase.h b/Engine/Source/Runtime/Core/Public/Delegates/MulticastDelegateBase.h
--- a/Engine/Source/Runtime/Core/Public/Delegates/MulticastDelegateBase.h
+++ b/Engine/Source/Runtime/Core/Public/Delegates/MulticastDelegateBase.h
@@ -21,16 +21,22 @@
  * Abstract base class for multicast delegates.
  */
 template<typename UserPolicy>
-class TMulticastDelegateBase
+class TMulticastDelegateBase : public FDelegateAccessHandlerBaseChecked
 {
 protected:
+	using Super = FDelegateAccessHandlerBaseChecked;
+	using typename Super::FReadAccessScope;
+	using Super::GetReadAccessScope;
+	using typename Super::FWriteAccessScope;
+	using Super::GetWriteAccessScope;
+
 	using InvocationListType = TArray<TDelegateBase<UserPolicy>, FMulticastInvocationListAllocatorType>;
 
 public:
 	/** Removes all functions from this delegate's invocation list. */
 	void Clear( )
 	{
-		UE_DELEGATES_MT_SCOPED_WRITE_ACCESS(AccessDetector);
+		FWriteAccessScope WriteScope = GetWriteAccessScope();
 
 		for (TDelegateBase<UserPolicy>& DelegateBaseRef : InvocationList)
 		{
@@ -47,7 +53,7 @@ public:
 	 */
 	inline bool IsBound( ) const
 	{
-		UE_DELEGATES_MT_SCOPED_READ_ACCESS(AccessDetector);
+		FReadAccessScope ReadScope = GetReadAccessScope();
 	
 		for (const TDelegateBase<UserPolicy>& DelegateBaseRef : InvocationList)
 		{
@@ -66,7 +72,7 @@ public:
 	 */
 	inline bool IsBoundToObject( void const* InUserObject ) const
 	{
-		UE_DELEGATES_MT_SCOPED_READ_ACCESS(AccessDetector);
+		FReadAccessScope ReadScope = GetReadAccessScope();
 
 		for (const TDelegateBase<UserPolicy>& DelegateBaseRef : InvocationList)
 		{
@@ -89,7 +95,7 @@ public:
 	 */
 	int32 RemoveAll( const void* InUserObject )
 	{
-		UE_DELEGATES_MT_SCOPED_WRITE_ACCESS(AccessDetector);
+		FWriteAccessScope WriteScope = GetWriteAccessScope();
 
 		int32 Result = 0;
 		if (InvocationListLockCount > 0)
@@ -166,7 +172,7 @@ protected:
 	template<typename DelegateInstanceInterfaceType, typename DelegateType>
 	void CopyFrom(const TMulticastDelegateBase& Other)
 	{
-		UE_DELEGATES_MT_SCOPED_WRITE_ACCESS(AccessDetector);
+		FWriteAccessScope WriteScope = GetWriteAccessScope();
 
 		checkSlow(&Other != this);
 
@@ -186,7 +192,8 @@ protected:
 	template<typename DelegateInstanceInterfaceType, typename DelegateBaseType, typename... ParamTypes>
 	void Broadcast(ParamTypes... Params) const
 	{
-		UE_DELEGATES_MT_SCOPED_WRITE_ACCESS(AccessDetector);
+		// the `const` on the method is a lie
+		FWriteAccessScope WriteScope = const_cast<TMulticastDelegateBase*>(this)->GetWriteAccessScope();
 
 		bool NeedsCompaction = false;
 
@@ -223,7 +230,7 @@ protected:
 	template <typename NewDelegateType>
 	inline FDelegateHandle AddDelegateInstance(NewDelegateType&& NewDelegateBaseRef)
 	{
-		UE_DELEGATES_MT_SCOPED_WRITE_ACCESS(AccessDetector);
+		FWriteAccessScope WriteScope = GetWriteAccessScope();
 
 		FDelegateHandle Result;
 		if (NewDelegateBaseRef.IsBound())
@@ -244,7 +251,7 @@ protected:
 	 */
 	bool RemoveDelegateInstance(FDelegateHandle Handle)
 	{
-		UE_DELEGATES_MT_SCOPED_WRITE_ACCESS(AccessDetector);
+		FWriteAccessScope WriteScope = GetWriteAccessScope();
 
 		for (int32 InvocationListIndex = 0; InvocationListIndex < InvocationList.Num(); ++InvocationListIndex)
 		{
@@ -363,6 +370,4 @@ private:
 
 	/** Holds a lock counter for the invocation list. */
 	mutable int32 InvocationListLockCount;
-
-	UE_DELEGATES_MT_ACCESS_DETECTOR(AccessDetector);
 };
diff --git a/Engine/Source/Runtime/Core/Public/UObject/ScriptDelegates.h b/Engine/Source/Runtime/Core/Public/UObject/ScriptDelegates.h
--- a/Engine/Source/Runtime/Core/Public/UObject/ScriptDelegates.h
+++ b/Engine/Source/Runtime/Core/Public/UObject/ScriptDelegates.h
@@ -3,7 +3,6 @@
 #pragma once
 
 #include "CoreTypes.h"
-#include "Delegates/DelegateMacros.h"
 #include "Misc/AssertionMacros.h"
 #include "Templates/UnrealTypeTraits.h"
 #include "Containers/ContainerAllocationPolicies.h"
@@ -11,6 +10,7 @@
 #include "Containers/UnrealString.h"
 #include "UObject/NameTypes.h"
 #include "Templates/SharedPointer.h"
+#include "Delegates/DelegateAccessHandler.h"
 
 namespace UE::Core::Private
 {
@@ -58,13 +58,19 @@ namespace UE::Core::Private
  * Script delegate base class.
  */
 template <typename Dummy>
-class TScriptDelegate
+class TScriptDelegate : public FDelegateAccessHandlerBaseChecked
 {
 	using WeakPtrType = typename UE::Core::Private::TScriptDelegateTraits<Dummy>::WeakPtrType;
 
 	template <typename>
 	friend class TScriptDelegate;
 
+	using Super = FDelegateAccessHandlerBaseChecked;
+	using typename Super::FReadAccessScope;
+	using Super::GetReadAccessScope;
+	using typename Super::FWriteAccessScope;
+	using Super::GetWriteAccessScope;
+
 public:
 
 	/** Default constructor. */
@@ -92,7 +98,7 @@ public:
 
 	void operator=(const TScriptDelegate& Other)
 	{
-		UE_DELEGATES_MT_SCOPED_WRITE_ACCESS(AccessDetector);
+		FWriteAccessScope WriteScope = GetWriteAccessScope();
 
 		Object = Other.Object;
 		FunctionName = Other.FunctionName;
@@ -104,7 +110,7 @@ public:
 	/* UE_DEPRECATED(5.3, "Deprecated - remove after TScriptDelegateTraits<FWeakObjectPtr> is removed") */
 	FORCEINLINE void operator=(const TScriptDelegate<OtherDummy>& Other)
 	{
-		UE_DELEGATES_MT_SCOPED_WRITE_ACCESS(AccessDetector);
+		FWriteAccessScope WriteScope = GetWriteAccessScope();
 
 		Object = Other.Object;
 		FunctionName = Other.FunctionName;
@@ -136,7 +142,7 @@ public:
 	 */
 	void BindUFunction( UObject* InObject, const FName& InFunctionName )
 	{
-		UE_DELEGATES_MT_SCOPED_WRITE_ACCESS(AccessDetector);
+		FWriteAccessScope WriteScope = GetWriteAccessScope();
 
 		Object = InObject;
 		FunctionName = InFunctionName;
@@ -149,8 +155,7 @@ public:
 	 */
 	inline bool IsBound() const
 	{
-		UE_DELEGATES_MT_SCOPED_READ_ACCESS(AccessDetector);
-
+		FReadAccessScope ReadScope = GetReadAccessScope();
 		return IsBound_Internal<UObject>();
 	}
 
@@ -161,8 +166,7 @@ public:
 	 */
 	inline bool IsBoundToObject(void const* InUserObject) const
 	{
-		UE_DELEGATES_MT_SCOPED_READ_ACCESS(AccessDetector);
-
+		FReadAccessScope ReadScope = GetReadAccessScope();
 		return InUserObject && (InUserObject == GetUObject());
 	}
 
@@ -173,8 +177,7 @@ public:
 	 */
 	bool IsBoundToObjectEvenIfUnreachable(void const* InUserObject) const
 	{
-		UE_DELEGATES_MT_SCOPED_READ_ACCESS(AccessDetector);
-
+		FReadAccessScope ReadScope = GetReadAccessScope();
 		return InUserObject && InUserObject == GetUObjectEvenIfUnreachable();
 	}
 
@@ -185,8 +188,7 @@ public:
 	 */
 	inline bool IsCompactable() const
 	{
-		UE_DELEGATES_MT_SCOPED_READ_ACCESS(AccessDetector);
-
+		FReadAccessScope ReadScope = GetReadAccessScope();
 		return FunctionName == NAME_None || !Object.Get(true);
 	}
 
@@ -195,7 +197,7 @@ public:
 	 */
 	void Unbind()
 	{
-		UE_DELEGATES_MT_SCOPED_WRITE_ACCESS(AccessDetector);
+		FWriteAccessScope WriteScope = GetWriteAccessScope();
 
 		Object = nullptr;
 		FunctionName = NAME_None;
@@ -219,8 +221,7 @@ public:
 	{
 		if( IsBound() )
 		{
-			UE_DELEGATES_MT_SCOPED_READ_ACCESS(AccessDetector);
-
+			FReadAccessScope ReadScope = GetReadAccessScope();
 			return ((UObjectTemplate*)GetUObject())->GetPathName() + TEXT(".") + GetFunctionName().ToString();
 		}
 		return TEXT( "<Unbound>" );
@@ -229,7 +230,7 @@ public:
 	/** Delegate serialization */
 	friend FArchive& operator<<( FArchive& Ar, TScriptDelegate& D )
 	{
-		UE_DELEGATES_MT_SCOPED_READ_ACCESS(D.AccessDetector);
+		FReadAccessScope ReadScope = D.GetReadAccessScope();
 
 		Ar << D.Object << D.FunctionName;
 		return Ar;
@@ -238,7 +239,7 @@ public:
 	/** Delegate serialization */
 	friend void operator<<(FStructuredArchive::FSlot Slot, TScriptDelegate& D)
 	{
-		UE_DELEGATES_MT_SCOPED_READ_ACCESS(D.AccessDetector);
+		FReadAccessScope ReadScope = D.GetReadAccessScope();
 
 		FStructuredArchive::FRecord Record = Slot.EnterRecord();
 		Record << SA_VALUE(TEXT("Object"), D.Object) << SA_VALUE(TEXT("FunctionName"),D.FunctionName);
@@ -247,7 +248,7 @@ public:
 	/** Comparison operators */
 	FORCEINLINE bool operator==( const TScriptDelegate& Other ) const
 	{
-		UE_DELEGATES_MT_SCOPED_READ_ACCESS(AccessDetector);
+		FReadAccessScope ReadScope = GetReadAccessScope();
 
 		return Object == Other.Object && FunctionName == Other.FunctionName;
 	}
@@ -258,14 +259,14 @@ public:
 	/* UE_DEPRECATED(5.3, "Deprecated - remove after TScriptDelegateTraits<FWeakObjectPtr> is removed") */
 	FORCEINLINE bool operator==(const TScriptDelegate<OtherDummy>& Other) const
 	{
-		UE_DELEGATES_MT_SCOPED_READ_ACCESS(AccessDetector);
+		FReadAccessScope ReadScope = GetReadAccessScope();
 
 		return Object == Other.Object && FunctionName == Other.FunctionName;
 	}
 
 	FORCEINLINE bool operator!=( const TScriptDelegate& Other ) const
 	{
-		UE_DELEGATES_MT_SCOPED_READ_ACCESS(AccessDetector);
+		FReadAccessScope ReadScope = GetReadAccessScope();
 
 		return Object != Other.Object || FunctionName != Other.FunctionName;
 	}
@@ -276,7 +277,7 @@ public:
 	/* UE_DEPRECATED(5.3, "Deprecated - remove after TScriptDelegateTraits<FWeakObjectPtr> is removed") */
 	FORCEINLINE bool operator!=(const TScriptDelegate<OtherDummy>& Other) const
 	{
-		UE_DELEGATES_MT_SCOPED_READ_ACCESS(AccessDetector);
+		FReadAccessScope ReadScope = GetReadAccessScope();
 
 		return Object != Other.Object || FunctionName != Other.FunctionName;
 	}
@@ -288,7 +289,7 @@ public:
 	 */
 	UObject* GetUObject()
 	{
-		UE_DELEGATES_MT_SCOPED_WRITE_ACCESS(AccessDetector);
+		FWriteAccessScope WriteScope = GetWriteAccessScope();
 
 		// Downcast UObjectBase to UObject
 		return static_cast< UObject* >( Object.Get() );
@@ -301,7 +302,7 @@ public:
 	 */
 	const UObject* GetUObject() const
 	{
-		UE_DELEGATES_MT_SCOPED_READ_ACCESS(AccessDetector);
+		FReadAccessScope ReadScope = GetReadAccessScope();
 
 		// Downcast UObjectBase to UObject
 		return static_cast< const UObject* >( Object.Get() );
@@ -314,7 +315,7 @@ public:
 	 */
 	UObject* GetUObjectEvenIfUnreachable()
 	{
-		UE_DELEGATES_MT_SCOPED_WRITE_ACCESS(AccessDetector);
+		FWriteAccessScope WriteScope = GetWriteAccessScope();
 
 		// Downcast UObjectBase to UObject
 		return static_cast< UObject* >( Object.GetEvenIfUnreachable() );
@@ -327,7 +328,7 @@ public:
 	 */
 	const UObject* GetUObjectEvenIfUnreachable() const
 	{
-		UE_DELEGATES_MT_SCOPED_READ_ACCESS(AccessDetector);
+		FReadAccessScope ReadScope = GetReadAccessScope();
 
 		// Downcast UObjectBase to UObject
 		return static_cast< const UObject* >( Object.GetEvenIfUnreachable() );
@@ -340,8 +341,7 @@ public:
 	 */
 	FName GetFunctionName() const
 	{
-		UE_DELEGATES_MT_SCOPED_READ_ACCESS(AccessDetector);
-
+		FReadAccessScope ReadScope = GetReadAccessScope();
 		return FunctionName;
 	}
 
@@ -361,19 +361,20 @@ public:
 
 		{	// to avoid MT access check if the delegate is deleted from inside of its callback, we don't cover the callback execution
 			// by access protection scope
-			UE_DELEGATES_MT_SCOPED_WRITE_ACCESS(AccessDetector);
+			// the `const` on the method is a lie
+			FWriteAccessScope WriteScope = const_cast<TScriptDelegate*>(this)->GetWriteAccessScope();
 
-		checkf( Object.IsValid() != false, TEXT( "ProcessDelegate() called with no object bound to delegate!" ) );
-		checkf( FunctionName != NAME_None, TEXT( "ProcessDelegate() called with no function name set!" ) );
+			checkf( Object.IsValid() != false, TEXT( "ProcessDelegate() called with no object bound to delegate!" ) );
+			checkf( FunctionName != NAME_None, TEXT( "ProcessDelegate() called with no function name set!" ) );
 
-		// Object was pending kill, so we cannot execute the delegate.  Note that it's important to assert
-		// here and not simply continue execution, as memory may be left uninitialized if the delegate is
-		// not able to execute, resulting in much harder-to-detect code errors.  Users should always make
-		// sure IsBound() returns true before calling ProcessDelegate()!
+			// Object was pending kill, so we cannot execute the delegate.  Note that it's important to assert
+			// here and not simply continue execution, as memory may be left uninitialized if the delegate is
+			// not able to execute, resulting in much harder-to-detect code errors.  Users should always make
+			// sure IsBound() returns true before calling ProcessDelegate()!
 			ObjectPtr = static_cast<UObjectTemplate*>(Object.Get());	// Down-cast
-		checkSlow( IsValid(ObjectPtr) );
+			checkSlow( IsValid(ObjectPtr) );
 
-		// Object *must* implement the specified function
+			// Object *must* implement the specified function
 			Function = ObjectPtr->FindFunctionChecked(FunctionName);
 		}
 
@@ -383,7 +384,7 @@ public:
 
 	friend uint32 GetTypeHash(const TScriptDelegate& Delegate)
 	{
-		UE_DELEGATES_MT_SCOPED_READ_ACCESS(Delegate.AccessDetector);
+		FReadAccessScope ReadScope = Delegate.GetReadAccessScope();
 
 		return HashCombine(GetTypeHash(Delegate.Object), GetTypeHash(Delegate.GetFunctionName()));
 	}
@@ -396,8 +397,6 @@ protected:
 	/** Name of the function to call on the bound object */
 	FName FunctionName;
 
-	UE_DELEGATES_MT_ACCESS_DETECTOR(AccessDetector);
-
 	// 
 	friend class FCallDelegateHelper;
 
@@ -416,8 +415,14 @@ struct TIsZeroConstructType<TScriptDelegate<Dummy>>
  * Script multi-cast delegate base class
  */
 template <typename Dummy>
-class TMulticastScriptDelegate
+class TMulticastScriptDelegate : public FDelegateAccessHandlerBaseChecked
 {
+	using Super = FDelegateAccessHandlerBaseChecked;
+	using typename Super::FReadAccessScope;
+	using Super::GetReadAccessScope;
+	using typename Super::FWriteAccessScope;
+	using Super::GetWriteAccessScope;
+
 	using UnicastDelegateType = TScriptDelegate<Dummy>;
 
 public:
@@ -435,7 +440,7 @@ public:
 	 */
 	inline bool IsBound() const
 	{
-		UE_DELEGATES_MT_SCOPED_READ_ACCESS(AccessDetector);
+		FReadAccessScope ReadScope = GetReadAccessScope();
 
 		return InvocationList.Num() > 0;
 	}
@@ -448,7 +453,7 @@ public:
 	 */
 	bool Contains( const UnicastDelegateType& InDelegate ) const
 	{
-		UE_DELEGATES_MT_SCOPED_READ_ACCESS(AccessDetector);
+		FReadAccessScope ReadScope = GetReadAccessScope();
 
 		return InvocationList.Contains( InDelegate );
 	}
@@ -471,7 +476,7 @@ public:
 	 */
 	bool Contains( const UObject* InObject, FName InFunctionName ) const
 	{
-		UE_DELEGATES_MT_SCOPED_READ_ACCESS(AccessDetector);
+		FReadAccessScope ReadScope = GetReadAccessScope();
 
 		return InvocationList.ContainsByPredicate( [=]( const UnicastDelegateType& Delegate ){
 			return Delegate.GetFunctionName() == InFunctionName && Delegate.IsBoundToObjectEvenIfUnreachable(InObject);
@@ -485,7 +490,7 @@ public:
 	 */
 	void Add( const UnicastDelegateType& InDelegate )
 	{
-		UE_DELEGATES_MT_SCOPED_WRITE_ACCESS(AccessDetector);
+		FWriteAccessScope WriteScope = GetWriteAccessScope();
 
 		// First check for any objects that may have expired
 		CompactInvocationList();
@@ -511,7 +516,7 @@ public:
 	 */
 	void AddUnique( const UnicastDelegateType& InDelegate )
 	{
-		UE_DELEGATES_MT_SCOPED_WRITE_ACCESS(AccessDetector);
+		FWriteAccessScope WriteScope = GetWriteAccessScope();
 
 		// Add the delegate, if possible
 		AddUniqueInternal( InDelegate );
@@ -537,7 +542,7 @@ public:
 	 */
 	void Remove( const UnicastDelegateType& InDelegate )
 	{
-		UE_DELEGATES_MT_SCOPED_WRITE_ACCESS(AccessDetector);
+		FWriteAccessScope WriteScope = GetWriteAccessScope();
 
 		// Remove the delegate
 		RemoveInternal( InDelegate );
@@ -564,7 +569,7 @@ public:
 	 */
 	void Remove( const UObject* InObject, FName InFunctionName )
 	{
-		UE_DELEGATES_MT_SCOPED_WRITE_ACCESS(AccessDetector);
+		FWriteAccessScope WriteScope = GetWriteAccessScope();
 
 		// Remove the delegate
 		RemoveInternal( InObject, InFunctionName );
@@ -583,7 +588,7 @@ public:
 	 */
 	void RemoveAll(const UObject* Object)
 	{
-		UE_DELEGATES_MT_SCOPED_WRITE_ACCESS(AccessDetector);
+		FWriteAccessScope WriteScope = GetWriteAccessScope();
 
 		for (int32 BindingIndex = InvocationList.Num() - 1; BindingIndex >= 0; --BindingIndex)
 		{
@@ -601,8 +606,7 @@ public:
 	 */
 	void Clear()
 	{
-		UE_DELEGATES_MT_SCOPED_WRITE_ACCESS(AccessDetector);
-
+		FWriteAccessScope WriteScope = GetWriteAccessScope();
 		InvocationList.Empty();
 	}
 
@@ -614,7 +618,7 @@ public:
 	template <typename UObjectTemplate>
 	inline FString ToString() const
 	{
-		UE_DELEGATES_MT_SCOPED_READ_ACCESS(AccessDetector);
+		FReadAccessScope ReadScope = GetReadAccessScope();
 
 		if( IsBound() )
 		{
@@ -638,7 +642,7 @@ public:
 	/** Multi-cast delegate serialization */
 	friend FArchive& operator<<( FArchive& Ar, TMulticastScriptDelegate& D )
 	{
-		UE_DELEGATES_MT_SCOPED_WRITE_ACCESS(D.AccessDetector);
+		FWriteAccessScope WriteScope = D.GetWriteAccessScope();
 
 		if( Ar.IsSaving() )
 		{
@@ -659,7 +663,7 @@ public:
 
 	friend void operator<<(FStructuredArchive::FSlot Slot, TMulticastScriptDelegate& D)
 	{
-		UE_DELEGATES_MT_SCOPED_WRITE_ACCESS(D.AccessDetector);
+		FWriteAccessScope WriteScope = D.GetWriteAccessScope();
 
 		FArchive& UnderlyingArchive = Slot.GetUnderlyingArchive();
 
@@ -688,7 +692,8 @@ public:
 	template <class UObjectTemplate>
 	void ProcessMulticastDelegate(void* Parameters) const
 	{
-		UE_DELEGATES_MT_SCOPED_WRITE_ACCESS(AccessDetector);
+		// the `const` on the method is a lie
+		FWriteAccessScope WriteScope = const_cast<TMulticastScriptDelegate*>(this)->GetWriteAccessScope();
 
 		if( InvocationList.Num() > 0 )
 		{
@@ -720,7 +725,7 @@ public:
 	 */
 	TArray< UObject* > GetAllObjects() const
 	{
-		UE_DELEGATES_MT_SCOPED_READ_ACCESS(AccessDetector);
+		FReadAccessScope ReadScope = GetReadAccessScope();
 
 		TArray< UObject* > OutputList;
 		for( typename InvocationListType::TIterator CurDelegate( InvocationList ); CurDelegate; ++CurDelegate )
@@ -741,7 +746,7 @@ public:
 	 */
 	TArray< UObject* > GetAllObjectsEvenIfUnreachable() const
 	{
-		UE_DELEGATES_MT_SCOPED_READ_ACCESS(AccessDetector);
+		FReadAccessScope ReadScope = GetReadAccessScope();
 
 		TArray< UObject* > OutputList;
 		for( typename InvocationListType::TIterator CurDelegate( InvocationList ); CurDelegate; ++CurDelegate )
@@ -838,8 +843,6 @@ protected:
 	/** Ordered list functions to invoke when the Broadcast function is called */
 	mutable InvocationListType InvocationList;		// Mutable so that we can housekeep list even with 'const' broadcasts
 
-	UE_DELEGATES_MT_ACCESS_DETECTOR(AccessDetector);
-
 	// Declare ourselves as a friend of FMulticastDelegateProperty so that it can access our function list
 	friend class FMulticastDelegateProperty;
 	friend class FMulticastInlineDelegateProperty;
