Fixed mismatched copy constructor/assignment operator in FRemoteConfigAsyncIOInfo.

diff --git a/Engine/Source/Runtime/Core/Private/Misc/RemoteConfigIni.cpp b/Engine/Source/Runtime/Core/Private/Misc/RemoteConfigIni.cpp
--- a/Engine/Source/Runtime/Core/Private/Misc/RemoteConfigIni.cpp
+++ b/Engine/Source/Runtime/Core/Private/Misc/RemoteConfigIni.cpp
@@ -36,20 +36,6 @@ FRemoteConfigAsyncIOInfo::FRemoteConfigAsyncIOInfo(const TCHAR* InDefaultIniFile
 }
 
 
-FRemoteConfigAsyncIOInfo& FRemoteConfigAsyncIOInfo::operator=(const FRemoteConfigAsyncIOInfo& Other)
-{
-	Buffer = Other.Buffer;
-	TimeStamp = Other.TimeStamp;
-	StartReadTime = Other.StartReadTime;
-	StartWriteTime = Other.StartWriteTime;
-	bReadIOFailed = Other.bReadIOFailed;
-	bWasProcessed = Other.bWasProcessed;
-	FMemory::Memcpy(DefaultIniFile, Other.DefaultIniFile, 1024);
-
-	return *this;
-}
-
-
 /*-----------------------------------------------------------------------------
    FRemoteConfigAsyncWorker
 -----------------------------------------------------------------------------*/
diff --git a/Engine/Source/Runtime/Core/Public/Misc/RemoteConfigIni.h b/Engine/Source/Runtime/Core/Public/Misc/RemoteConfigIni.h
--- a/Engine/Source/Runtime/Core/Public/Misc/RemoteConfigIni.h
+++ b/Engine/Source/Runtime/Core/Public/Misc/RemoteConfigIni.h
@@ -23,9 +23,6 @@ struct FRemoteConfigAsyncIOInfo
 	FRemoteConfigAsyncIOInfo() { }
 	FRemoteConfigAsyncIOInfo(const TCHAR* InDefaultIniFile);
 
-	/** Assignment operator */
-	FRemoteConfigAsyncIOInfo& operator=(const FRemoteConfigAsyncIOInfo& Other);
-
 	/** Stores the contents of the remote config file */
 	FString Buffer;
 
