Reduced MemStack's memory footprint

diff --git a/Engine/Source/Runtime/Core/Private/Misc/MemStack.cpp b/Engine/Source/Runtime/Core/Private/Misc/MemStack.cpp
--- a/Engine/Source/Runtime/Core/Private/Misc/MemStack.cpp
+++ b/Engine/Source/Runtime/Core/Private/Misc/MemStack.cpp
@@ -345,7 +345,7 @@ void FMemStackBase::AllocateNewChunk(int32 MinSize)
 	// Create new chunk.
 	int32 TotalSize = MinSize + (int32)sizeof(FTaggedMemory);
 	uint32 AllocSize;
-	if (TopChunk || TotalSize > FPageAllocator::SmallPageSize)
+	if (TotalSize > FPageAllocator::SmallPageSize)
 	{
 		AllocSize = AlignArbitrary<int32>(TotalSize, FPageAllocator::PageSize);
 		if (AllocSize == FPageAllocator::PageSize)
diff --git a/Engine/Source/Runtime/Core/Public/Misc/MemStack.h b/Engine/Source/Runtime/Core/Public/Misc/MemStack.h
--- a/Engine/Source/Runtime/Core/Public/Misc/MemStack.h
+++ b/Engine/Source/Runtime/Core/Public/Misc/MemStack.h
@@ -38,7 +38,7 @@ public:
 	enum
 	{
 		PageSize = 64 * 1024,
-		SmallPageSize = 1024-16 // allow a little extra space for allocator headers, etc
+		SmallPageSize = 1024
 	};
 #if UE_BUILD_SHIPPING
 	typedef TLockFreeFixedSizeAllocator<PageSize, PLATFORM_CACHE_LINE_SIZE, FNoopCounter> TPageAllocator;
