Fixed mismatched copy constructor/assignment operator in TProcHandle.

diff --git a/Engine/Source/Runtime/Core/Public/GenericPlatform/GenericPlatformProcess.h b/Engine/Source/Runtime/Core/Public/GenericPlatform/GenericPlatformProcess.h
--- a/Engine/Source/Runtime/Core/Public/GenericPlatform/GenericPlatformProcess.h
+++ b/Engine/Source/Runtime/Core/Public/GenericPlatform/GenericPlatformProcess.h
@@ -91,16 +91,6 @@ public:
 		: Handle( Other )
 	{ }
 
-	/** Assignment operator. */
-	FORCEINLINE TProcHandle& operator=( const TProcHandle& Other )
-	{
-		if( this != &Other )
-		{
-			Handle = Other.Handle;
-		}
-		return *this;
-	}
-
 	/** Accessors. */
 	FORCEINLINE T Get() const
 	{
