Probable fix for IncludeTool complaining about multiple conflicting definitions of FDelegateAccessHandlerBaseChecked.

diff --git a/Engine/Source/Runtime/Core/Public/Delegates/DelegateAccessHandler.h b/Engine/Source/Runtime/Core/Public/Delegates/DelegateAccessHandler.h
--- a/Engine/Source/Runtime/Core/Public/Delegates/DelegateAccessHandler.h
+++ b/Engine/Source/Runtime/Core/Public/Delegates/DelegateAccessHandler.h
@@ -34,13 +34,15 @@ protected:
 	}
 };
 
-#if UE_DETECT_DELEGATES_RACE_CONDITIONS
-
 /**
  * non thread-safe version that detects not thread-safe delegates used concurrently (dev builds only)
  */
 class FDelegateAccessHandlerBaseChecked
+#if !UE_DETECT_DELEGATES_RACE_CONDITIONS
+	: public FDelegateAccessHandlerBaseUnchecked
+#endif // !UE_DETECT_DELEGATES_RACE_CONDITIONS
 {
+#if UE_DETECT_DELEGATES_RACE_CONDITIONS
 protected:
 	class FReadAccessScope
 	{
@@ -82,10 +84,5 @@ protected:
 
 private:
 	mutable FRWFullyRecursiveAccessDetector Accessor;
-};
-
-#else // UE_DETECT_DELEGATES_RACE_CONDITIONS
-
-using FDelegateAccessHandlerBaseChecked = FDelegateAccessHandlerBaseUnchecked;
-
 #endif // UE_DETECT_DELEGATES_RACE_CONDITIONS
+};
