Add deprecation path for mutable TArray/TArrayView<TObjectPtr<...>> implicit conversions

diff --git a/Engine/Source/Runtime/Core/Public/Containers/Array.h b/Engine/Source/Runtime/Core/Public/Containers/Array.h
--- a/Engine/Source/Runtime/Core/Public/Containers/Array.h
+++ b/Engine/Source/Runtime/Core/Public/Containers/Array.h
@@ -2300,7 +2300,8 @@ public:
 		return *Ptr;
 	}
 
-	/** Implicit conversion operator to container of compatible element type. */
+#if !UE_DEPRECATE_MUTABLE_TOBJECTPTR
+	/** Mutable implicit conversion operator to container of compatible element type. */
 	template <
 		typename AliasElementType = ElementType,
 		std::enable_if_t<TIsContainerElementTypeReinterpretable<AliasElementType>::Value>* = nullptr
@@ -2311,8 +2312,9 @@ public:
 		ElementCompat::ReinterpretRangeContiguous(begin(), end(), Num());
 		return *reinterpret_cast<TArray<typename ElementCompat::ReinterpretType>*>(this);
 	}
-
-	/** Implicit conversion operator to constant container of compatible element type. */
+#endif
+  
+	/** Immutable implicit conversion operator to constant container of compatible element type. */
 	template <
 		typename AliasElementType = ElementType,
 		typename std::enable_if_t<TIsContainerElementTypeReinterpretable<AliasElementType>::Value>* = nullptr
diff --git a/Engine/Source/Runtime/Core/Public/Containers/ArrayView.h b/Engine/Source/Runtime/Core/Public/Containers/ArrayView.h
--- a/Engine/Source/Runtime/Core/Public/Containers/ArrayView.h
+++ b/Engine/Source/Runtime/Core/Public/Containers/ArrayView.h
@@ -75,12 +75,17 @@ namespace ArrayViewPrivate
 	{
 	private:
 		using NaturalElementType = typename TRemovePointer<decltype(GetData(DeclVal<RangeType&>()))>::Type;
+		using TypeCompat = TContainerElementTypeCompatibility<NaturalElementType>;
 
 	public:
 		static constexpr bool Value = 
-			!std::is_same_v<typename TContainerElementTypeCompatibility<NaturalElementType>::ReinterpretType, NaturalElementType>
+			!std::is_same_v<typename TypeCompat::ReinterpretType, NaturalElementType>
 			&&
-			TIsCompatibleElementType<typename TContainerElementTypeCompatibility<NaturalElementType>::ReinterpretType, ElementType>::Value;
+			TIsCompatibleElementType<typename TypeCompat::ReinterpretType, ElementType>::Value
+			&&
+			(!UE_DEPRECATE_MUTABLE_TOBJECTPTR
+			 || std::is_same_v<ElementType, std::remove_pointer_t<typename TypeCompat::ReinterpretType>* const>
+			 || std::is_same_v<ElementType, const std::remove_pointer_t<typename TypeCompat::ReinterpretType>* const>);
 
 		template <typename T>
 		static decltype(auto) GetData(T&& Arg)
diff --git a/Engine/Source/Runtime/Core/Public/Misc/CoreDefines.h b/Engine/Source/Runtime/Core/Public/Misc/CoreDefines.h
--- a/Engine/Source/Runtime/Core/Public/Misc/CoreDefines.h
+++ b/Engine/Source/Runtime/Core/Public/Misc/CoreDefines.h
@@ -65,3 +65,7 @@ Can be overridden by setting to 1 or 0 in the project's .Target.cs files.
 #ifndef NO_CVARS
 	#define NO_CVARS 0
 #endif
+
+#ifndef UE_DEPRECATE_MUTABLE_TOBJECTPTR
+#define UE_DEPRECATE_MUTABLE_TOBJECTPTR 0
+#endif
diff --git a/Engine/Source/Runtime/CoreUObject/Public/UObject/ObjectPtr.h b/Engine/Source/Runtime/CoreUObject/Public/UObject/ObjectPtr.h
--- a/Engine/Source/Runtime/CoreUObject/Public/UObject/ObjectPtr.h
+++ b/Engine/Source/Runtime/CoreUObject/Public/UObject/ObjectPtr.h
@@ -987,7 +987,7 @@ using TPointedToType = typename TPointedToTypeImpl<T>::Type;
 
 //------------------------------------------------------------------------------
 
-namespace UE::Core::Private
+namespace UE::Core::Private // private facilities; not for direct use
 {
 	namespace Unsafe
 	{
@@ -1075,7 +1075,7 @@ namespace UE::Core::Private
 */
 
 template <typename T>
-UE::Core::Private::TMutableView<T> MutableView(T& A)
+[[nodiscard]] UE::Core::Private::TMutableView<T> MutableView(T& A)
 {
 	return UE::Core::Private::TMutableView<T>{A};
 }
