Added equality operator to FTask and a UE::Tasks::Wait variant that can take a single task.

diff --git a/Engine/Source/Runtime/Core/Public/Tasks/Task.h b/Engine/Source/Runtime/Core/Public/Tasks/Task.h
--- a/Engine/Source/Runtime/Core/Public/Tasks/Task.h
+++ b/Engine/Source/Runtime/Core/Public/Tasks/Task.h
@@ -182,6 +182,16 @@ namespace UE::Tasks
 				return IsValid() && Pimpl->IsAwaitable();
 			}
 
+			bool operator==(const FTaskHandle& Other) const
+			{
+				return Pimpl == Other.Pimpl;
+			}
+
+			bool operator!=(const FTaskHandle& Other) const
+			{
+				return Pimpl != Other.Pimpl;
+			}
+
 		protected:
 			TRefCountPtr<FTaskBase> Pimpl;
 		};
@@ -373,6 +383,11 @@ namespace UE::Tasks
 		}
 	}
 
+	inline void Wait(Private::FTaskHandle& Task)
+	{
+		Task.Wait();
+	}
+
 	// wait for multiple tasks, with timeout
 	// @param TaskCollectionType - an iterable collection of `TTask<T>`, e.g. `TArray<FTask>`
 	template<typename TaskCollectionType>
