Fixed mismatched copy constructor/assignment operator in FHashedName.

diff --git a/Engine/Source/Runtime/Core/Private/Serialization/MemoryImage.cpp b/Engine/Source/Runtime/Core/Private/Serialization/MemoryImage.cpp
--- a/Engine/Source/Runtime/Core/Private/Serialization/MemoryImage.cpp
+++ b/Engine/Source/Runtime/Core/Private/Serialization/MemoryImage.cpp
@@ -1231,6 +1231,13 @@ FHashedName::FHashedName(const FHashedName& InName) : Hash(InName.Hash)
 #endif
 }
 
+FHashedName& FHashedName::operator=(const FHashedName& InName)
+{
+	FHashedName Temp = InName;
+	Swap(Temp, *this);
+	return *this;
+}
+
 void FPointerTableBase::SaveToArchive(FArchive& Ar, const FPlatformTypeLayoutParameters& LayoutParams, const void* FrozenObject) const
 {
 	int32 NumDependencies = 0;
diff --git a/Engine/Source/Runtime/Core/Public/Serialization/MemoryImage.h b/Engine/Source/Runtime/Core/Public/Serialization/MemoryImage.h
--- a/Engine/Source/Runtime/Core/Public/Serialization/MemoryImage.h
+++ b/Engine/Source/Runtime/Core/Public/Serialization/MemoryImage.h
@@ -854,6 +854,7 @@ public:
 	inline FHashedName() : Hash(0u) {}
 	CORE_API explicit FHashedName(uint64 InHash);
 	CORE_API FHashedName(const FHashedName& InName);
+	CORE_API FHashedName& operator=(const FHashedName& InName);
 	CORE_API FHashedName(const TCHAR* InString);
 	CORE_API FHashedName(const FString& InString);
 	CORE_API FHashedName(const FName& InName);
