Tasks System: fixed "empty" prerequisites support. It used to deadlock if an "empty" task was used as a prerequisite

diff --git a/Engine/Source/Runtime/Core/Private/Tests/Tasks/TasksTest.cpp b/Engine/Source/Runtime/Core/Private/Tests/Tasks/TasksTest.cpp
--- a/Engine/Source/Runtime/Core/Private/Tests/Tasks/TasksTest.cpp
+++ b/Engine/Source/Runtime/Core/Private/Tests/Tasks/TasksTest.cpp
@@ -121,6 +121,13 @@ namespace UE { namespace TasksTests
 			verify(Task.GetResult() == 42);
 		}
 
+		{ // using an "empty" prerequisites
+			FTask EmptyPrereq;
+			FTask NonEmptyPrereq = Launch(UE_SOURCE_LOCATION, [] {});
+			FTask Task = Launch(UE_SOURCE_LOCATION, [] {}, Prerequisites(EmptyPrereq, NonEmptyPrereq));
+			check(Task.Wait(FTimespan::FromMilliseconds(1.0)));
+		}
+
 		{	// check that movable-only result types are supported, that only single instance of result is created and that it's destroyed
 			static std::atomic<uint32> ConstructionsNum{ 0 }; // `static` to make it visible to `FMoveConstructable`
 			static std::atomic<uint32> DestructionsNum{ 0 };
diff --git a/Engine/Source/Runtime/Core/Public/Tasks/TaskPrivate.h b/Engine/Source/Runtime/Core/Public/Tasks/TaskPrivate.h
--- a/Engine/Source/Runtime/Core/Public/Tasks/TaskPrivate.h
+++ b/Engine/Source/Runtime/Core/Public/Tasks/TaskPrivate.h
@@ -266,6 +266,7 @@ namespace UE::Tasks
 
 					if (Prerequisite == nullptr)
 					{
+						++NumCompletedPrerequisites;
 						continue;
 					}
 
