fix using TMap element after removal

diff --git a/Engine/Source/Runtime/Core/Public/Containers/DiscardableKeyValueCache.h b/Engine/Source/Runtime/Core/Public/Containers/DiscardableKeyValueCache.h
--- a/Engine/Source/Runtime/Core/Public/Containers/DiscardableKeyValueCache.h
+++ b/Engine/Source/Runtime/Core/Public/Containers/DiscardableKeyValueCache.h
@@ -281,9 +281,8 @@ protected:
 			if (!Found && BackfillFound)
 			{
 				// if shared refs, add/remove order is important
-				CurrentMap->Add(Key, *BackfillFound);
+				Found = &CurrentMap->Add(Key, *BackfillFound);
 				BackfillMap->Remove(Key);
-				Found = BackfillFound;
 			}
 		}
 
