Fix for SlowTask accessing the StartTime value from the multiple threads
Re-instated default loading of the analytics plugin

diff --git a/Engine/Source/Runtime/Core/Public/Misc/FeedbackContext.h b/Engine/Source/Runtime/Core/Public/Misc/FeedbackContext.h
--- a/Engine/Source/Runtime/Core/Public/Misc/FeedbackContext.h
+++ b/Engine/Source/Runtime/Core/Public/Misc/FeedbackContext.h
@@ -9,6 +9,7 @@
 #include "CoreTypes.h"
 #include "HAL/CriticalSection.h"
 #include "HAL/PlatformCrt.h"
+#include "HAL/PlatformTime.h"
 #include "Internationalization/Text.h"
 #include "Misc/OutputDevice.h"
 #include "Misc/ScopeLock.h"
@@ -50,7 +51,7 @@ public:
 	DECLARE_EVENT_OneParam(FFeedbackContext, FOnStartSlowTask, const FText& TaskName );
 	FOnStartSlowTask& OnStartSlowTask() { return StartSlowTaskEvent; }
 
-	DECLARE_EVENT_OneParam(FFeedbackContext, FOnFinalizeSlowTask, const FText& TaskName);
+	DECLARE_EVENT_TwoParams(FFeedbackContext, FOnFinalizeSlowTask, const FText& TaskName, double DurationInSeconds);
 	FOnFinalizeSlowTask& OnFinalizeSlowTask() { return FinalizeSlowTaskEvent; }
 
 
@@ -70,6 +71,7 @@ protected:
 	virtual void StartSlowTask( const FText& Task, bool bShowCancelButton=false )
 	{
 		TaskName = Task;
+		TaskStartTime = FPlatformTime::Seconds();
 		GIsSlowTask = true;
 		StartSlowTaskEvent.Broadcast(TaskName);
 	}
@@ -79,7 +81,7 @@ protected:
 	 */
 	virtual void FinalizeSlowTask( )
 	{
-		FinalizeSlowTaskEvent.Broadcast(TaskName);
+		FinalizeSlowTaskEvent.Broadcast(TaskName, FPlatformTime::Seconds() - TaskStartTime);
 		GIsSlowTask = false;
 	}
 
@@ -162,6 +164,7 @@ private:
 
 	/** The name of any task we are running */
 	FText TaskName;
+	double TaskStartTime;
 	FOnStartSlowTask StartSlowTaskEvent;
 	FOnFinalizeSlowTask FinalizeSlowTaskEvent;
 
