deprecated TDepletableMpmcQueue as it uses spin-waiting that can cause priority inversion and subsequently deadlock on some platforms. Switched all existing instances to TConsumeAllMpmcQueue

diff --git a/Engine/Source/Editor/UnrealEd/Private/PlayLevel.cpp b/Engine/Source/Editor/UnrealEd/Private/PlayLevel.cpp
--- a/Engine/Source/Editor/UnrealEd/Private/PlayLevel.cpp
+++ b/Engine/Source/Editor/UnrealEd/Private/PlayLevel.cpp
@@ -122,7 +122,7 @@
 #include "IAssetViewport.h"
 #include "IPIEAuthorizer.h"
 #include "Features/IModularFeatures.h"
-#include "Containers/DepletableMpmcQueue.h"
+#include "Containers/ConsumeAllMpmcQueue.h"
 #include "TickableEditorObject.h"
 
 DEFINE_LOG_CATEGORY(LogPlayLevel);
@@ -166,7 +166,7 @@ public:
 			}
 			else
 			{
-				QueuedLines.Enqueue(MoveTemp(Line));
+				QueuedLines.ProduceItem(MoveTemp(Line));
 			}
 		}
 	}
@@ -180,7 +180,7 @@ public:
 
 	virtual void Tick(float DeltaTime) final
 	{
-		QueuedLines.Deplete(LogLine);
+		QueuedLines.ConsumeAllFifo(LogLine);
 	}
 
 	virtual ETickableTickType GetTickableTickType() const final
@@ -216,7 +216,7 @@ private:
 		}
 	}
 
-	UE::TDepletableMpmcQueue<FLine> QueuedLines;
+	UE::TConsumeAllMpmcQueue<FLine> QueuedLines;
 };
 
 void UEditorEngine::EndPlayMap()
diff --git a/Engine/Source/Runtime/Core/Private/Misc/OutputDeviceRedirector.cpp b/Engine/Source/Runtime/Core/Private/Misc/OutputDeviceRedirector.cpp
--- a/Engine/Source/Runtime/Core/Private/Misc/OutputDeviceRedirector.cpp
+++ b/Engine/Source/Runtime/Core/Private/Misc/OutputDeviceRedirector.cpp
@@ -4,7 +4,7 @@
 
 #include "Async/EventCount.h"
 #include "Containers/BitArray.h"
-#include "Containers/DepletableMpmcQueue.h"
+#include "Containers/ConsumeAllMpmcQueue.h"
 #include "Experimental/ConcurrentLinearAllocator.h"
 #include "HAL/PlatformProcess.h"
 #include "HAL/PlatformTime.h"
@@ -120,7 +120,7 @@ struct FOutputDeviceRedirectorState
 	uint8 OutputDevicesLockPadding[CalculateRedirectorCacheLinePadding(sizeof(OutputDevicesLock) + sizeof(OutputDevicesLockState))]{};
 
 	/** A queue of items logged by non-primary threads. */
-	TDepletableMpmcQueue<FOutputDeviceItem, FOutputDeviceLinearAllocator> BufferedItems;
+	TConsumeAllMpmcQueue<FOutputDeviceItem, FOutputDeviceLinearAllocator> BufferedItems;
 
 	/** Array of output devices to redirect to from the primary thread. */
 	TArray<FOutputDevice*> BufferedOutputDevices;
@@ -469,14 +469,14 @@ void FOutputDeviceRedirectorState::FlushBufferedItems()
 
 	const uint32 ThreadId = FPlatformTLS::GetCurrentThreadId();
 	bool bIsPanicThread = this->IsPanicThread(ThreadId);
-	BufferedItems.Deplete([this, ThreadId, bIsPanicThread](FOutputDeviceItem&& Item)
+	BufferedItems.ConsumeAllFifo([this, ThreadId, bIsPanicThread](FOutputDeviceItem&& Item)
 	{
 		if (!bIsPanicThread && HasPanicThread())
 		{
 			// Enqueuing during Deplete is okay because the items passed to the Deplete consumer
 			// are a moved list that is no longer associated with the queue; Deplete will return
 			// even though new items are in the queue.
-			BufferedItems.Enqueue(MoveTemp(Item));
+			BufferedItems.ProduceItem(MoveTemp(Item));
 		}
 		else
 		{
@@ -653,7 +653,7 @@ void FOutputDeviceRedirector::SerializeRecord(const UE::FLogRecord& Record)
 	}
 
 	// Queue the record to serialize to buffered output devices from the primary thread.
-	if (State->BufferedItems.EnqueueAndReturnWasEmpty(Record))
+	if (State->BufferedItems.ProduceItem(Record) == UE::EConsumeAllMpmcQueueResult::WasEmpty)
 	{
 		State->ThreadWakeEvent.Notify();
 	}
@@ -714,7 +714,7 @@ void FOutputDeviceRedirector::Serialize(const TCHAR* const Data, const ELogVerbo
 	}
 
 	// Queue the line to serialize to buffered output devices from the primary thread.
-	if (State->BufferedItems.EnqueueAndReturnWasEmpty(Data, Category, Verbosity, RealTime))
+	if (State->BufferedItems.ProduceItem(Data, Category, Verbosity, RealTime) == UE::EConsumeAllMpmcQueueResult::WasEmpty)
 	{
 		State->ThreadWakeEvent.Notify();
 	}
diff --git a/Engine/Source/Runtime/Core/Public/Containers/DepletableMpmcQueue.h b/Engine/Source/Runtime/Core/Public/Containers/DepletableMpmcQueue.h
--- a/Engine/Source/Runtime/Core/Public/Containers/DepletableMpmcQueue.h
+++ b/Engine/Source/Runtime/Core/Public/Containers/DepletableMpmcQueue.h
@@ -5,6 +5,7 @@
 #include "CoreTypes.h"
 #include "HAL/UnrealMemory.h"
 #include "Templates/TypeCompatibleBytes.h"
+#include "Templates/UnrealTemplate.h"
 #include <atomic>
 #include <type_traits>
 
@@ -21,7 +22,7 @@ namespace UE
 	 * Consumes elements in FIFO order.
 	 */
 	template<typename T, typename AllocatorType = FMemory>
-	class TDepletableMpmcQueue final
+	class UE_DEPRECATED(5.3, "This concurrent queue was deprecated because it uses spin-waiting that can cause priority inversion and subsequently deadlocks on some platforms. Please use TConsumeAllMpmcQueue.") TDepletableMpmcQueue final
 	{
 	private:
 		struct FNode
@@ -146,6 +147,8 @@ namespace UE
 		}
 	};
 
+	PRAGMA_DISABLE_DEPRECATION_WARNINGS
 	template<typename T, typename AllocatorType = FMemory>
-	using TDepletableMpscQueue UE_DEPRECATED(5.2, "The queue was renamed to TDepletableMpmcQueue as now it supports multiple concurrent consumers. Please consider switching to TConsumeAllMpmcQueue that is not vulnerable to live locking.") = TDepletableMpmcQueue<T, AllocatorType>;
+	using TDepletableMpscQueue UE_DEPRECATED(5.2, "This concurrent queue was deprecated because it uses spin-waiting that can cause priority inversion and subsequently deadlocks on some platforms. Please use TConsumeAllMpmcQueue.") = TDepletableMpmcQueue<T, AllocatorType>;
+	PRAGMA_ENABLE_DEPRECATION_WARNINGS
 }
diff --git a/Engine/Source/Runtime/Core/Tests/Containers/ConcurrentQueuesTest.cpp b/Engine/Source/Runtime/Core/Tests/Containers/ConcurrentQueuesTest.cpp
--- a/Engine/Source/Runtime/Core/Tests/Containers/ConcurrentQueuesTest.cpp
+++ b/Engine/Source/Runtime/Core/Tests/Containers/ConcurrentQueuesTest.cpp
@@ -8,7 +8,7 @@
 #include "Containers/SpscQueue.h"
 #include "Containers/MpscQueue.h"
 #include "Containers/ClosableMpscQueue.h"
-#include "Containers/DepletableMpmcQueue.h"
+#include "Containers/ConsumeAllMpmcQueue.h"
 #include "Tests/Benchmark.h"
 
 #include <atomic>
@@ -459,26 +459,26 @@ namespace UE { namespace ConcurrentQueuesTests
 			}
 		}
 
-		UE_BENCHMARK(5, TestTCircularQueueSingleThread<5'000'000>);
-		UE_BENCHMARK(5, TestQueueSingleThread<5'000'000, TQueueAdapter<TQueue<uint32, EQueueMode::Spsc>>>);
-		UE_BENCHMARK(5, TestQueueSingleThread<5'000'000, TQueueAdapter<TQueue<uint32, EQueueMode::Mpsc>>>);
-		UE_BENCHMARK(5, TestQueueSingleThread<5'000'000, TMpscQueue<uint32>>);
-		UE_BENCHMARK(5, TestQueueSingleThread<5'000'000, TSpscQueue<uint32>>);
+		//UE_BENCHMARK(5, TestTCircularQueueSingleThread<5'000'000>);
+		//UE_BENCHMARK(5, TestQueueSingleThread<5'000'000, TQueueAdapter<TQueue<uint32, EQueueMode::Spsc>>>);
+		//UE_BENCHMARK(5, TestQueueSingleThread<5'000'000, TQueueAdapter<TQueue<uint32, EQueueMode::Mpsc>>>);
+		//UE_BENCHMARK(5, TestQueueSingleThread<5'000'000, TMpscQueue<uint32>>);
+		//UE_BENCHMARK(5, TestQueueSingleThread<5'000'000, TSpscQueue<uint32>>);
 
-		UE_BENCHMARK(5, TestSpscQueueCorrectness<1'000'000, TMpscQueue<uint32>>);
-		UE_BENCHMARK(5, TestSpscQueueCorrectness<1'000'000, TSpscQueue<uint32>>);
+		//UE_BENCHMARK(5, TestSpscQueueCorrectness<1'000'000, TMpscQueue<uint32>>);
+		//UE_BENCHMARK(5, TestSpscQueueCorrectness<1'000'000, TSpscQueue<uint32>>);
 
-		UE_BENCHMARK(5, TestTCircularQueue<5'000'000>);
-		UE_BENCHMARK(5, TestSpscQueue<5'000'000, TQueueAdapter<TQueue<uint32, EQueueMode::Spsc>>>);
-		UE_BENCHMARK(5, TestSpscQueue<5'000'000, TQueueAdapter<TQueue<uint32, EQueueMode::Mpsc>>>);
-		UE_BENCHMARK(5, TestSpscQueue<5'000'000, TMpscQueue<uint32>>);
-		UE_BENCHMARK(5, TestSpscQueue<5'000'000, TSpscQueue<uint32>>);
+		//UE_BENCHMARK(5, TestTCircularQueue<5'000'000>);
+		//UE_BENCHMARK(5, TestSpscQueue<5'000'000, TQueueAdapter<TQueue<uint32, EQueueMode::Spsc>>>);
+		//UE_BENCHMARK(5, TestSpscQueue<5'000'000, TQueueAdapter<TQueue<uint32, EQueueMode::Mpsc>>>);
+		//UE_BENCHMARK(5, TestSpscQueue<5'000'000, TMpscQueue<uint32>>);
+		//UE_BENCHMARK(5, TestSpscQueue<5'000'000, TSpscQueue<uint32>>);
 
 		UE_BENCHMARK(5, TestMpscQueueCorrectness<5'000'000, TQueueAdapter<TQueue<uint32, EQueueMode::Mpsc>>>);
 		UE_BENCHMARK(5, TestMpscQueueCorrectness<5'000'000, TMpscQueue<uint32>>);
 
-		UE_BENCHMARK(5, TestMpscQueue<100'000, TQueueAdapter<TQueue<uint32, EQueueMode::Mpsc>>>);
-		UE_BENCHMARK(5, TestMpscQueue<100'000, TMpscQueue<uint32>>);
+		//UE_BENCHMARK(5, TestMpscQueue<100'000, TQueueAdapter<TQueue<uint32, EQueueMode::Mpsc>>>);
+		//UE_BENCHMARK(5, TestMpscQueue<100'000, TMpscQueue<uint32>>);
 	}
 }
 
@@ -793,22 +793,22 @@ namespace ClosableMpscQueueTests
 	}
 }
 
-namespace DepletableMpscQueueTests
+namespace ConsumeAllMpmcQueueTests
 {
-	// straightforward implementation over MPSC TQueue for comparison with TResettableMpscQueue
+	// straightforward implementation over MPSC TQueue for comparison with TConsumeAllMpmcQueue
 	template<typename T>
-	class TDepletableMpmcQueueOnTQueue
+	class TConsumeAllMpmcQueueOnTQueue
 	{
 		using FQueue = TQueue<T, EQueueMode::Mpsc>;
 
 	public:
-		~TDepletableMpmcQueueOnTQueue()
+		~TConsumeAllMpmcQueueOnTQueue()
 		{
 			checkf(SubscribingThreadsNum.load(std::memory_order_relaxed) == 0, TEXT("No other threads can use the object during its destruction"));
 			delete Queue.load(std::memory_order_relaxed);
 		}
 
-		void Enqueue(T Value)
+		void ProduceItem(T Value)
 		{
 			// `acquire`/`release` to keep `Queue` usage inside the scope
 			verify(SubscribingThreadsNum.fetch_add(1, std::memory_order_acquire) >= 0);
@@ -818,7 +818,7 @@ namespace DepletableMpscQueueTests
 		}
 
 		template<typename F>
-		void Deplete(const F& Consumer)
+		void ConsumeAllFifo(const F& Consumer)
 		{
 			// `acquire` to make it "happen-before" reading `SubscribingThreadsNum`
 			FQueue* LocalQueue = Queue.exchange(new FQueue, std::memory_order_acquire);
@@ -846,16 +846,16 @@ namespace DepletableMpscQueueTests
 
 	// straightforward implementation over MPSC TQueue for comparison with TResettableMpscQueue
 	template<typename T>
-	class TDepletableMpmcQueueOnTMpscQueue
+	class TConsumeAllMpmcQueueOnTMpscQueue
 	{
 	public:
-		~TDepletableMpmcQueueOnTMpscQueue()
+		~TConsumeAllMpmcQueueOnTMpscQueue()
 		{
 			checkf(SubscribingThreadsNum.load(std::memory_order_relaxed) == 0, TEXT("No other threads can use the object during its destration"));
 			delete Queue.load(std::memory_order_relaxed);
 		}
 
-		void Enqueue(T Value)
+		void ProduceItem(T Value)
 		{
 			// `acquire`/`release` to keep `Queue` usage inside the scope
 			verify(SubscribingThreadsNum.fetch_add(1, std::memory_order_acquire) >= 0);
@@ -865,7 +865,7 @@ namespace DepletableMpscQueueTests
 		}
 
 		template<typename F>
-		void Deplete(const F& Consumer)
+		void ConsumeAllFifo(const F& Consumer)
 		{
 			// `acquire` to make it "happen-before" reading `SubscribingThreadsNum`
 			TMpscQueue<T>* LocalQueue = Queue.exchange(new TMpscQueue<T>, std::memory_order_acquire);
@@ -892,18 +892,18 @@ namespace DepletableMpscQueueTests
 
 	// wrapper over `TClosableLockFreePointerListUnorderedSingleConsumer`
 	template<typename T>
-	class TDepletableMpmcQueueOnTClosableLockFreePointerListUnorderedSingleConsumer
+	class TConsumeAllMpmcQueueOnTClosableLockFreePointerListUnorderedSingleConsumer
 	{
 		using FQueue = TClosableLockFreePointerListUnorderedSingleConsumer<T, 0 /* as was used in FGraphEvent */>;
 
 	public:
-		~TDepletableMpmcQueueOnTClosableLockFreePointerListUnorderedSingleConsumer()
+		~TConsumeAllMpmcQueueOnTClosableLockFreePointerListUnorderedSingleConsumer()
 		{
 			checkf(SubscribingThreadsNum.load(std::memory_order_relaxed) == 0, TEXT("No other threads can use the object during its destration"));
 			delete Queue.load(std::memory_order_relaxed);
 		}
 
-		void Enqueue(T* Value)
+		void ProduceItem(T* Value)
 		{
 			// `acquire`/`release` to keep `Queue` usage inside the scope
 			verify(SubscribingThreadsNum.fetch_add(1, std::memory_order_acquire) >= 0);
@@ -913,7 +913,7 @@ namespace DepletableMpscQueueTests
 		}
 
 		template<typename F>
-		void Deplete(const F& Consumer)
+		void ConsumeAllFifo(const F& Consumer)
 		{
 			// `acquire` to make it "happen-before" reading `SubscribingThreadsNum`
 			FQueue* LocalQueue = Queue.exchange(new FQueue, std::memory_order_acquire);
@@ -942,18 +942,18 @@ namespace DepletableMpscQueueTests
 
 	// wrapper over `TClosableMpscQueue`
 	template<typename T>
-	class TDepletableMpmcQueueOnTClosableMpscQueue
+	class TConsumeAllMpmcQueueOnTClosableMpscQueue
 	{
 		using FQueue = TClosableMpscQueue<T>;
 
 	public:
-		~TDepletableMpmcQueueOnTClosableMpscQueue()
+		~TConsumeAllMpmcQueueOnTClosableMpscQueue()
 		{
 			checkf(SubscribingThreadsNum.load(std::memory_order_relaxed) == 0, TEXT("No other threads can use the object during its destration"));
 			delete Queue.load(std::memory_order_relaxed);
 		}
 
-		void Enqueue(T Value)
+		void ProduceItem(T Value)
 		{
 			// `acquire`/`release` to keep `Queue` usage inside the scope
 			verify(SubscribingThreadsNum.fetch_add(1, std::memory_order_acquire) >= 0);
@@ -963,7 +963,7 @@ namespace DepletableMpscQueueTests
 		}
 
 		template<typename F>
-		void Deplete(const F& Consumer)
+		void ConsumeAllFifo(const F& Consumer)
 		{
 			// `acquire` to make it "happen-before" reading `SubscribingThreadsNum`
 			FQueue* LocalQueue = Queue.exchange(new FQueue, std::memory_order_acquire);
@@ -991,7 +991,7 @@ namespace DepletableMpscQueueTests
 		QueueT Q;
 		for (uint32 i = 0; i != Num; ++i)
 		{
-			Q.Deplete([](void*) { checkNoEntry(); });
+			Q.ConsumeAllFifo([](void*) { checkNoEntry(); });
 		}
 	}
 
@@ -1016,7 +1016,7 @@ namespace DepletableMpscQueueTests
 					uint64 Total = 0;
 					for (uint64 Product = 1; Product != Num + 1; ++Product)
 					{
-						Queue.Enqueue((void*)(intptr_t)Product);
+						Queue.ProduceItem((void*)(intptr_t)Product);
 						Total += Product;
 					}
 					check(Total == Num * (Num + 1) / 2);
@@ -1028,7 +1028,7 @@ namespace DepletableMpscQueueTests
 		uint64 Consumed = 0;
 		while (Consumed != Produced)
 		{
-			Queue.Deplete
+			Queue.ConsumeAllFifo
 			(
 				[&Consumed](void* Value) 
 				{ 
@@ -1043,39 +1043,39 @@ namespace DepletableMpscQueueTests
 
 		FTaskGraphInterface::Get().WaitUntilTasksComplete(Producers);
 
-		Queue.Deplete([](void*) { checkNoEntry(); }); // empty
+		Queue.ConsumeAllFifo([](void*) { checkNoEntry(); }); // empty
 	}
 
-	TEST_CASE_NAMED(FDepletableMpscQueueTest, "System::Core::Async::DepletableMpscQueueTest", "[.][ApplicationContextMask][EngineFilter][Perf]")
+	TEST_CASE_NAMED(FConsumeAllMpmcQueueTest, "System::Core::Async::ConsumeAllMpmcQueueTest", "[.][ApplicationContextMask][EngineFilter][Perf]")
 	{
 		{
-			TDepletableMpmcQueue<int> Q;
-			Q.Enqueue(1);
-			Q.Enqueue(2);
-			Q.Deplete([](int) {});
-			Q.Enqueue(3);
-			Q.Deplete([](int) {});
-			Q.Deplete([](int) {});
-			Q.Enqueue(4); // to destroy non-empty queue
+			TConsumeAllMpmcQueue<int> Q;
+			Q.ProduceItem(1);
+			Q.ProduceItem(2);
+			Q.ConsumeAllLifo([](int) {});
+			Q.ProduceItem(3);
+			Q.ConsumeAllLifo([](int) {});
+			Q.ConsumeAllLifo([](int) {});
+			Q.ProduceItem(4); // to destroy non-empty queue
 		}
 
 		{
-			TDepletableMpmcQueue<int> Q;
-			Q.Deplete([](int) { check(false); });
+			TConsumeAllMpmcQueue<int> Q;
+			Q.ConsumeAllLifo([](int) { check(false); });
 		}
 		{
-			TDepletableMpmcQueue<int> Q;
-			Q.Enqueue(1);
+			TConsumeAllMpmcQueue<int> Q;
+			Q.ProduceItem(1);
 			bool Done = false;
-			Q.Deplete([&Done](int Item) { check(!Done && Item == 1); Done = true; });
+			Q.ConsumeAllLifo([&Done](int Item) { check(!Done && Item == 1); Done = true; });
 		}
 
 		{
-			TDepletableMpmcQueue<int> Q;
-			Q.Enqueue(1);
-			Q.Enqueue(2);
+			TConsumeAllMpmcQueue<int> Q;
+			Q.ProduceItem(1);
+			Q.ProduceItem(2);
 			int Step = 0;
-			Q.Deplete(
+			Q.ConsumeAllFifo(
 				[&Step](int Item)
 				{
 					switch (Step)
@@ -1097,45 +1097,45 @@ namespace DepletableMpscQueueTests
 
 		//for (int i = 0; i != 1000000; ++i)
 		{
-			TDepletableMpmcQueue<int> Q;
-			Q.Deplete([](int) { check(false); });
-			Q.Deplete([](int) { check(false); });
+			TConsumeAllMpmcQueue<int> Q;
+			Q.ConsumeAllLifo([](int) { check(false); });
+			Q.ConsumeAllLifo([](int) { check(false); });
 		}
 
 		//for (int i = 0; i != 1000000; ++i)
 		{
 			int Res = 0;
-			TDepletableMpmcQueue<int> Q;
-			Q.Enqueue(1);
-			Q.Deplete([&Res](int i) { Res += i; });
-			Q.Enqueue(2);
-			Q.Deplete([&Res](int i) { Res += i; });
+			TConsumeAllMpmcQueue<int> Q;
+			Q.ProduceItem(1);
+			Q.ConsumeAllLifo([&Res](int i) { Res += i; });
+			Q.ProduceItem(2);
+			Q.ConsumeAllLifo([&Res](int i) { Res += i; });
 			checkf(Res == 3, TEXT("%d"), Res);
 		}
 
 		{	// EnqueueAndReturnWasEmpty
-			TDepletableMpmcQueue<int> Q;
-			verify(Q.EnqueueAndReturnWasEmpty(1));
-			verify(!Q.EnqueueAndReturnWasEmpty(2));
-			Q.Deplete([](int) {});
-			verify(Q.EnqueueAndReturnWasEmpty(1));
-			verify(!Q.EnqueueAndReturnWasEmpty(2));
-			Q.Deplete([](int) {});
-			verify(Q.EnqueueAndReturnWasEmpty(1));
-			verify(!Q.EnqueueAndReturnWasEmpty(2));
+			TConsumeAllMpmcQueue<int> Q;
+			verify(Q.ProduceItem(1) == EConsumeAllMpmcQueueResult::WasEmpty);
+			verify(Q.ProduceItem(2) == EConsumeAllMpmcQueueResult::HadItems);
+			Q.ConsumeAllLifo([](int) {});
+			verify(Q.ProduceItem(1) == EConsumeAllMpmcQueueResult::WasEmpty);
+			verify(Q.ProduceItem(2) == EConsumeAllMpmcQueueResult::HadItems);
+			Q.ConsumeAllLifo([](int) {});
+			verify(Q.ProduceItem(1) == EConsumeAllMpmcQueueResult::WasEmpty);
+			verify(Q.ProduceItem(2) == EConsumeAllMpmcQueueResult::HadItems);
 		}
 
-		UE_BENCHMARK(5, TestPerf<10'000, TDepletableMpmcQueue<void*>>);
-		UE_BENCHMARK(5, TestPerf<10'000, TDepletableMpmcQueueOnTQueue<void*>>);
-		UE_BENCHMARK(5, TestPerf<10'000, TDepletableMpmcQueueOnTMpscQueue<void*>>);
-		UE_BENCHMARK(5, TestPerf<10'000, TDepletableMpmcQueueOnTClosableLockFreePointerListUnorderedSingleConsumer<void>>);
-		UE_BENCHMARK(5, TestPerf<10'000, TDepletableMpmcQueueOnTClosableMpscQueue<void*>>);
-
-		UE_BENCHMARK(5, TestConsumingEmptyQueue<10'000'000, TDepletableMpmcQueue<void*>>);
-		UE_BENCHMARK(5, TestConsumingEmptyQueue<10'000'000, TDepletableMpmcQueueOnTQueue<void*>>);
-		UE_BENCHMARK(5, TestConsumingEmptyQueue<10'000'000, TDepletableMpmcQueueOnTMpscQueue<void*>>);
-		UE_BENCHMARK(5, TestConsumingEmptyQueue<10'000'000, TDepletableMpmcQueueOnTClosableLockFreePointerListUnorderedSingleConsumer<void>>);
-		UE_BENCHMARK(5, TestConsumingEmptyQueue<10'000'000, TDepletableMpmcQueueOnTClosableMpscQueue<void*>>);
+		UE_BENCHMARK(5, TestPerf<10'000, TConsumeAllMpmcQueue<void*>>);
+		UE_BENCHMARK(5, TestPerf<10'000, TConsumeAllMpmcQueueOnTQueue<void*>>);
+		UE_BENCHMARK(5, TestPerf<10'000, TConsumeAllMpmcQueueOnTMpscQueue<void*>>);
+		UE_BENCHMARK(5, TestPerf<10'000, TConsumeAllMpmcQueueOnTClosableLockFreePointerListUnorderedSingleConsumer<void>>);
+		UE_BENCHMARK(5, TestPerf<10'000, TConsumeAllMpmcQueueOnTClosableMpscQueue<void*>>);
+
+		UE_BENCHMARK(5, TestConsumingEmptyQueue<10'000'000, TConsumeAllMpmcQueue<void*>>);
+		UE_BENCHMARK(5, TestConsumingEmptyQueue<10'000'000, TConsumeAllMpmcQueueOnTQueue<void*>>);
+		UE_BENCHMARK(5, TestConsumingEmptyQueue<10'000'000, TConsumeAllMpmcQueueOnTMpscQueue<void*>>);
+		UE_BENCHMARK(5, TestConsumingEmptyQueue<10'000'000, TConsumeAllMpmcQueueOnTClosableLockFreePointerListUnorderedSingleConsumer<void>>);
+		UE_BENCHMARK(5, TestConsumingEmptyQueue<10'000'000, TConsumeAllMpmcQueueOnTClosableMpscQueue<void*>>);
 	}
 }}
 
diff --git a/Engine/Source/Runtime/Engine/Classes/Sound/QuartzSubscription.h b/Engine/Source/Runtime/Engine/Classes/Sound/QuartzSubscription.h
--- a/Engine/Source/Runtime/Engine/Classes/Sound/QuartzSubscription.h
+++ b/Engine/Source/Runtime/Engine/Classes/Sound/QuartzSubscription.h
@@ -3,7 +3,7 @@
 #pragma once
 
 #include "Sound/QuartzQuantizationUtilities.h"
-#include "Containers/DepletableMpmcQueue.h"
+#include "Containers/ConsumeAllMpmcQueue.h"
 
 // forwards
 class UQuartzSubsystem;
@@ -88,24 +88,24 @@ namespace Audio
 		template <typename T>
 		void PushEvent(const T& Data)
 		{
-			CommandQueue.Enqueue([InData = Data](ListenerType* Listener) { Listener->ExecCommand(InData); });
+			CommandQueue.ProduceItem([InData = Data](ListenerType* Listener) { Listener->ExecCommand(InData); });
 		}
 
 		void PushCommand(TFunction<void(ListenerType*)> InCommand)
 		{
-			CommandQueue.Enqueue([=](ListenerType* Listener) { InCommand(Listener); });
+			CommandQueue.ProduceItem([=](ListenerType* Listener) { InCommand(Listener); });
 		}
 
 		void PumpCommandQueue(ListenerType* InListener)
 		{
-			CommandQueue.Deplete([InListener](TFunction<void(ListenerType*)> Command)
+			CommandQueue.ConsumeAllFifo([InListener](TFunction<void(ListenerType*)> Command)
 			{
 				Command(InListener);
 			});
 		}
 
 	private:
-		UE::TDepletableMpmcQueue<TFunction<void(ListenerType*)>> CommandQueue;
+		UE::TConsumeAllMpmcQueue<TFunction<void(ListenerType*)>> CommandQueue;
 	};
 
 } // namespace Audio
