Replaced 0xFFFFFFFF for tls slot index with FPlatformTLS::InvalidTlsSlot named constant

diff --git a/Engine/Plugins/Importers/USDImporter/Source/UnrealUSDWrapper/Private/USDMemory.cpp b/Engine/Plugins/Importers/USDImporter/Source/UnrealUSDWrapper/Private/USDMemory.cpp
--- a/Engine/Plugins/Importers/USDImporter/Source/UnrealUSDWrapper/Private/USDMemory.cpp
+++ b/Engine/Plugins/Importers/USDImporter/Source/UnrealUSDWrapper/Private/USDMemory.cpp
@@ -13,7 +13,7 @@ LLM_DEFINE_TAG( Usd );
 class UNREALUSDWRAPPER_API FTlsSlot final
 {
 public:
-	static constexpr uint32 InvalidTlsSlotIndex = 0xFFFFFFFF;
+	static constexpr uint32 InvalidTlsSlotIndex = FPlatformTLS::InvalidTlsSlot;
 
 	FTlsSlot();
 	~FTlsSlot();
diff --git a/Engine/Source/Runtime/Core/Private/HAL/ThreadingBase.cpp b/Engine/Source/Runtime/Core/Private/HAL/ThreadingBase.cpp
--- a/Engine/Source/Runtime/Core/Private/HAL/ThreadingBase.cpp
+++ b/Engine/Source/Runtime/Core/Private/HAL/ThreadingBase.cpp
@@ -1345,12 +1345,12 @@ FTlsAutoCleanup* FThreadSingletonInitializer::Get( TFunctionRef<FTlsAutoCleanup*
 	UE_AUTORTFM_OPEN(
 	{
 		TlsSlot = (uint32)FPlatformAtomics::AtomicRead_Relaxed((int32*)&InOutTlsSlot);
-		if (TlsSlot == 0xFFFFFFFF)
+		if (TlsSlot == FPlatformTLS::InvalidTlsSlot)
 		{
 			const uint32 ThisTlsSlot = FPlatformTLS::AllocTlsSlot();
 			check(FPlatformTLS::IsValidTlsSlot(ThisTlsSlot));
-			const uint32 PrevTlsSlot = FPlatformAtomics::InterlockedCompareExchange( (int32*)&InOutTlsSlot, (int32)ThisTlsSlot, 0xFFFFFFFF );
-			if (PrevTlsSlot != 0xFFFFFFFF)
+			const uint32 PrevTlsSlot = FPlatformAtomics::InterlockedCompareExchange( (int32*)&InOutTlsSlot, (int32)ThisTlsSlot, FPlatformTLS::InvalidTlsSlot );
+			if (PrevTlsSlot != FPlatformTLS::InvalidTlsSlot)
 			{
 				FPlatformTLS::FreeTlsSlot( ThisTlsSlot );
 				TlsSlot = PrevTlsSlot;
@@ -1374,7 +1374,7 @@ FTlsAutoCleanup* FThreadSingletonInitializer::Get( TFunctionRef<FTlsAutoCleanup*
 
 FTlsAutoCleanup* FThreadSingletonInitializer::TryGet(uint32& TlsSlot)
 {
-	if (TlsSlot == 0xFFFFFFFF)
+	if (TlsSlot == FPlatformTLS::InvalidTlsSlot)
 	{
 		return nullptr;
 	}
@@ -1385,12 +1385,12 @@ FTlsAutoCleanup* FThreadSingletonInitializer::TryGet(uint32& TlsSlot)
 
 FTlsAutoCleanup* FThreadSingletonInitializer::Inject(FTlsAutoCleanup* Instance, uint32& TlsSlot)
 {
-	if (TlsSlot == 0xFFFFFFFF)
+	if (TlsSlot == FPlatformTLS::InvalidTlsSlot)
 	{
 		const uint32 ThisTlsSlot = FPlatformTLS::AllocTlsSlot();
 		check(FPlatformTLS::IsValidTlsSlot(ThisTlsSlot));
-		const uint32 PrevTlsSlot = FPlatformAtomics::InterlockedCompareExchange( (int32*)&TlsSlot, (int32)ThisTlsSlot, 0xFFFFFFFF );
-		if (PrevTlsSlot != 0xFFFFFFFF)
+		const uint32 PrevTlsSlot = FPlatformAtomics::InterlockedCompareExchange( (int32*)&TlsSlot, (int32)ThisTlsSlot, FPlatformTLS::InvalidTlsSlot);
+		if (PrevTlsSlot != FPlatformTLS::InvalidTlsSlot)
 		{
 			FPlatformTLS::FreeTlsSlot( ThisTlsSlot );
 		}
diff --git a/Engine/Source/Runtime/Core/Public/HAL/ThreadSingleton.h b/Engine/Source/Runtime/Core/Public/HAL/ThreadSingleton.h
--- a/Engine/Source/Runtime/Core/Public/HAL/ThreadSingleton.h
+++ b/Engine/Source/Runtime/Core/Public/HAL/ThreadSingleton.h
@@ -46,7 +46,7 @@ class TThreadSingleton : public FTlsAutoCleanup
 	 */
 	CORE_API static uint32& GetTlsSlot()
 	{
-		static uint32 TlsSlot = 0xFFFFFFFF;
+		static uint32 TlsSlot = FPlatformTLS::InvalidTlsSlot;
 		return TlsSlot;
 	}
 #else
@@ -58,7 +58,7 @@ class TThreadSingleton : public FTlsAutoCleanup
 #endif
 	static uint32& GetTlsSlot()
 	{
-		static uint32 TlsSlot = 0xFFFFFFFF;
+		static uint32 TlsSlot = FPlatformTLS::InvalidTlsSlot;
 		return TlsSlot;
 	}
 #endif
@@ -73,7 +73,7 @@ protected:
 	virtual ~TThreadSingleton()
 	{
 		// Clean the dangling pointer from the TLS.
-		check(GetTlsSlot() != 0xFFFFFFFF);
+		check(GetTlsSlot() != FPlatformTLS::InvalidTlsSlot);
 		if(((FTlsAutoCleanup*)FPlatformTLS::GetTlsValue(GetTlsSlot())) == static_cast<FTlsAutoCleanup*>(this))
 		{
 			FPlatformTLS::SetTlsValue(GetTlsSlot(), nullptr);
