ParkingLot: Use memory from FPlatformMemory to avoid GMalloc

ParkingLot cannot be used to synchronize allocator implementations if ParkingLot itself depends on those allocators.

#rnx

diff --git a/Engine/Source/Runtime/Core/Private/Async/ParkingLot.cpp b/Engine/Source/Runtime/Core/Private/Async/ParkingLot.cpp
--- a/Engine/Source/Runtime/Core/Private/Async/ParkingLot.cpp
+++ b/Engine/Source/Runtime/Core/Private/Async/ParkingLot.cpp
@@ -7,6 +7,8 @@
 #include "Async/WordMutex.h"
 #include "Containers/Array.h"
 #include "Containers/ArrayView.h"
+#include "HAL/Allocators/AnsiAllocator.h"
+#include "HAL/MallocAnsi.h"
 #include "HAL/PlatformManualResetEvent.h"
 #include "Templates/RefCounting.h"
 #include <atomic>
@@ -14,43 +16,89 @@
 namespace UE::ParkingLot::Private
 {
 
-using FBucketMutex = FWordMutex;
+///////////////////////////////////////////////////////////////////////////////////////////////////////////////////////
+
+/** A low-level allocator that maintains a free list for a given object type. Memory is never released. */
+template <typename ObjectType>
+class TObjectAllocator
+{
+	static_assert(sizeof(ObjectType) >= sizeof(void*));
+	static_assert(alignof(ObjectType) >= alignof(void*));
+
+public:
+	static void* Malloc(SIZE_T Size = sizeof(ObjectType), uint32 Alignment = alignof(ObjectType))
+	{
+		if (TUniqueLock Lock(FreeListMutex); void* Mem = FreeListHead)
+		{
+			FreeListHead = *(void**)Mem;
+			return Mem;
+		}
+		constexpr SIZE_T MinSize = sizeof(ObjectType);
+		constexpr uint32 MinAlignment = alignof(ObjectType);
+		return AnsiMalloc(FPlatformMath::Max(Size, MinSize), FPlatformMath::Max(Alignment, MinAlignment));
+	}
+
+	static void Free(void* Mem)
+	{
+		TUniqueLock Lock(FreeListMutex);
+		void* Next = FreeListHead;
+		*(void**)Mem = Next;
+		FreeListHead = Mem;
+	}
+
+private:
+	inline static FWordMutex FreeListMutex;
+	inline static void* FreeListHead = nullptr;
+};
 
 ///////////////////////////////////////////////////////////////////////////////////////////////////////////////////////
 
 /**
  * A thread as stored in the wait queue.
  */
-struct FThread final
+struct alignas(PLATFORM_CACHE_LINE_SIZE) FThread final
 {
 	FThread* Next = nullptr;
 	const void* WaitAddress = nullptr;
 	uint64 WakeToken = 0;
 	FPlatformManualResetEvent Event;
-	mutable std::atomic<uint32> ReferenceCount = 0;
+	std::atomic<uint32> ReferenceCount = 0;
 
-	inline void AddRef() const
+	inline void AddRef()
 	{
 		ReferenceCount.fetch_add(1, std::memory_order_relaxed);
 	}
 
-	inline void Release() const
+	inline void Release()
 	{
 		if (ReferenceCount.fetch_sub(1, std::memory_order_acq_rel) == 1)
 		{
-			delete this;
+			this->~FThread();
+			Allocator.Free(this);
 		}
 	}
 
-	FThread() = default;
+	inline static FThread* New()
+	{
+		// Performance is very sensitive to how this is aligned.
+		// This value likely needs to be tuned for different hardware/platforms.
+		constexpr SIZE_T ThreadAlignment = 2048;
+
+		return new(Allocator.Malloc(sizeof(FThread), ThreadAlignment)) FThread;
+	}
 
 private:
+	FThread() = default;
 	~FThread() = default;
 
 	FThread(const FThread&) = delete;
 	FThread& operator=(const FThread&) = delete;
+
+	static TObjectAllocator<FThread> Allocator;
 };
 
+TObjectAllocator<FThread> FThread::Allocator;
+
 class FThreadLocalData
 {
 public:
@@ -64,7 +112,7 @@ public:
 		FThreadLocalData& LocalThreadLocalData = ThreadLocalData;
 		if (!LocalThreadLocalData.Thread)
 		{
-			LocalThreadLocalData.Thread = new FThread;
+			LocalThreadLocalData.Thread = FThread::New();
 		}
 		return *LocalThreadLocalData.Thread;
 	}
@@ -102,11 +150,18 @@ enum class EQueueAction
 class alignas(PLATFORM_CACHE_LINE_SIZE) FBucket final
 {
 public:
-	FBucket() = default;
-	FBucket(const FBucket&) = delete;
-	FBucket& operator=(const FBucket&) = delete;
+	static FBucket* Create()
+	{
+		return new(Allocator.Malloc()) FBucket;
+	}
+
+	static void Destroy(FBucket* Bucket)
+	{
+		Bucket->~FBucket();
+		Allocator.Free(Bucket);
+	}
 
-	[[nodiscard]] inline TDynamicUniqueLock<FBucketMutex> LockDynamic() { return TDynamicUniqueLock(Mutex); }
+	[[nodiscard]] inline TDynamicUniqueLock<FWordMutex> LockDynamic() { return TDynamicUniqueLock(Mutex); }
 
 	inline void Lock() { Mutex.Lock(); }
 	inline void Unlock() { Mutex.Unlock(); }
@@ -128,11 +183,21 @@ public:
 	void DequeueIf(VisitorType&& Visitor);
 
 private:
-	FBucketMutex Mutex;
+	FBucket() = default;
+	~FBucket() = default;
+
+	FBucket(const FBucket&) = delete;
+	FBucket& operator=(const FBucket&) = delete;
+
+	FWordMutex Mutex;
 	FThread* Head = nullptr;
 	FThread* Tail = nullptr;
+
+	static TObjectAllocator<FBucket> Allocator;
 };
 
+TObjectAllocator<FBucket> FBucket::Allocator;
+
 void FBucket::Enqueue(FThread* Thread)
 {
 	checkSlow(Thread);
@@ -217,10 +282,10 @@ class FTable final
 
 public:
 	/** Find or create, and lock, the bucket for the memory address. */
-	static FBucket& FindOrCreateBucket(const void* Address, TDynamicUniqueLock<FBucketMutex>& OutLock);
+	static FBucket& FindOrCreateBucket(const void* Address, TDynamicUniqueLock<FWordMutex>& OutLock);
 
 	/** Find and lock the bucket for the memory address. Returns null if the bucket was not created. */
-	static FBucket* FindBucket(const void* Address, TDynamicUniqueLock<FBucketMutex>& OutLock);
+	static FBucket* FindBucket(const void* Address, TDynamicUniqueLock<FWordMutex>& OutLock);
 
 	/** Reserve memory for the table to handle at least ThreadCount waiting threads. */
 	static void Reserve(uint32 ThreadCount);
@@ -235,7 +300,7 @@ private:
 	static void Destroy(FTable& Table);
 
 	/** Try to lock the whole table by locking every one of its buckets. */
-	static bool TryLock(FTable& Table, TArray<FBucket*>& OutBuckets);
+	static bool TryLock(FTable& Table, TArray<FBucket*, FAnsiAllocator>& OutBuckets);
 	/** Unlock an array of buckets that was filled by TryLock. */
 	static void Unlock(TConstArrayView<FBucket*> LockedBuckets);
 
@@ -247,18 +312,19 @@ private:
 
 	FTable() = default;
 	~FTable() = default;
+
 	FTable(const FTable&) = delete;
 	FTable& operator=(const FTable&) = delete;
 
 	/** Find or create the bucket at the index. */
 	template <typename AllocatorType>
-	FBucket& FindOrCreateBucket(uint32 Index, AllocatorType&& Allocator);
+	FBucket& FindOrCreateBucket(uint32 Index, AllocatorType&& BucketAllocator);
 
 	uint32 BucketCount = 0;
 	std::atomic<FBucket*> Buckets[0];
 };
 
-FBucket& FTable::FindOrCreateBucket(const void* Address, TDynamicUniqueLock<FBucketMutex>& OutLock)
+FBucket& FTable::FindOrCreateBucket(const void* Address, TDynamicUniqueLock<FWordMutex>& OutLock)
 {
 	const uint32 Hash = HashAddress(Address);
 
@@ -266,7 +332,7 @@ FBucket& FTable::FindOrCreateBucket(const void* Address, TDynamicUniqueLock<FBuc
 	{
 		FTable& Table = CreateOrGet();
 		const uint32 Index = Hash % Table.BucketCount;
-		FBucket& Bucket = Table.FindOrCreateBucket(Index, [] { return new FBucket; });
+		FBucket& Bucket = Table.FindOrCreateBucket(Index, [] { return FBucket::Create(); });
 		OutLock = Bucket.LockDynamic();
 
 		if (LIKELY(&Table == GlobalTable.load(std::memory_order_acquire)))
@@ -279,7 +345,7 @@ FBucket& FTable::FindOrCreateBucket(const void* Address, TDynamicUniqueLock<FBuc
 	}
 }
 
-FBucket* FTable::FindBucket(const void* Address, TDynamicUniqueLock<FBucketMutex>& OutLock)
+FBucket* FTable::FindBucket(const void* Address, TDynamicUniqueLock<FWordMutex>& OutLock)
 {
 	const uint32 Hash = HashAddress(Address);
 
@@ -307,20 +373,20 @@ FBucket* FTable::FindBucket(const void* Address, TDynamicUniqueLock<FBucketMutex
 }
 
 template <typename AllocatorType>
-FBucket& FTable::FindOrCreateBucket(uint32 Index, AllocatorType&& Allocator)
+FBucket& FTable::FindOrCreateBucket(uint32 Index, AllocatorType&& BucketAllocator)
 {
 	std::atomic<FBucket*>& BucketPtr = Buckets[Index];
 	FBucket* Bucket = BucketPtr.load(std::memory_order_acquire);
 	if (UNLIKELY(!Bucket))
 	{
-		FBucket* NewBucket = Allocator();
+		FBucket* NewBucket = BucketAllocator();
 		if (BucketPtr.compare_exchange_strong(Bucket, NewBucket, std::memory_order_release, std::memory_order_acquire))
 		{
 			Bucket = NewBucket;
 		}
 		else
 		{
-			delete NewBucket;
+			FBucket::Destroy(NewBucket);
 		}
 		checkSlow(Bucket);
 	}
@@ -352,7 +418,7 @@ FTable& FTable::CreateOrGet()
 FTable& FTable::Create(const uint32 Size)
 {
 	const SIZE_T MemorySize = sizeof(FTable) + sizeof(FBucket*) * SIZE_T(Size);
-	void* const Memory = FMemory::Malloc(MemorySize, alignof(FTable));
+	void* const Memory = AnsiMalloc(MemorySize, alignof(FTable));
 	FMemory::Memzero(Memory, MemorySize);
 	FTable& Table = *new(Memory) FTable;
 	Table.BucketCount = Size;
@@ -362,13 +428,13 @@ FTable& FTable::Create(const uint32 Size)
 void FTable::Destroy(FTable& Table)
 {
 	Table.~FTable();
-	FMemory::Free(&Table);
+	AnsiFree(&Table);
 }
 
 void FTable::Reserve(const uint32 ThreadCount)
 {
 	const uint32 TargetBucketCount = FMath::RoundUpToPowerOfTwo(ThreadCount);
-	TArray<FBucket*> ExistingBuckets;
+	TArray<FBucket*, FAnsiAllocator> ExistingBuckets;
 
 	for (;;)
 	{
@@ -388,7 +454,7 @@ void FTable::Reserve(const uint32 ThreadCount)
 
 		// Gather waiting threads to be redistributed into the buckets of the new table.
 		// Threads with the same address remain in the same relative order as they were queued.
-		TArray<FThread*> Threads;
+		TArray<FThread*, FAnsiAllocator> Threads;
 		for (FBucket* Bucket : ExistingBuckets)
 		{
 			while (FThread* Thread = Bucket->Dequeue())
@@ -400,10 +466,10 @@ void FTable::Reserve(const uint32 ThreadCount)
 		FTable& NewTable = Create(TargetBucketCount);
 
 		// Reuse existing now-empty buckets when populating the new table.
-		TArray<FBucket*> AvailableBuckets = ExistingBuckets;
+		TArray<FBucket*, FAnsiAllocator> AvailableBuckets = ExistingBuckets;
 		const auto AllocateBucket = [&AvailableBuckets]() -> FBucket*
 		{
-			return !AvailableBuckets.IsEmpty() ? AvailableBuckets.Pop() : new FBucket;
+			return !AvailableBuckets.IsEmpty() ? AvailableBuckets.Pop(/*bAllowShrinking*/ false) : FBucket::Create();
 		};
 
 		// Add waiting threads to the new table.
@@ -432,14 +498,14 @@ void FTable::Reserve(const uint32 ThreadCount)
 	}
 }
 
-bool FTable::TryLock(FTable& Table, TArray<FBucket*>& OutBuckets)
+bool FTable::TryLock(FTable& Table, TArray<FBucket*, FAnsiAllocator>& OutBuckets)
 {
-	OutBuckets.Reset();
+	OutBuckets.Reset(Table.BucketCount);
 
 	// Gather buckets from the table, creating them as needed because the lock is on the bucket.
 	for (uint32 Index = 0; Index < Table.BucketCount; ++Index)
 	{
-		OutBuckets.Add(&Table.FindOrCreateBucket(Index, [] { return new FBucket; }));
+		OutBuckets.Add(&Table.FindOrCreateBucket(Index, [] { return FBucket::Create(); }));
 	}
 
 	// Lock the buckets in order by address to ensure consistent ordering regardless of the table being locked.
@@ -517,7 +583,7 @@ FWaitState WaitUntil(const void* Address, const TFunctionRef<bool()>& CanWait, c
 
 	// Enqueue the thread if CanWait returns true while the bucket is locked.
 	{
-		TDynamicUniqueLock<FBucketMutex> BucketLock;
+		TDynamicUniqueLock<FWordMutex> BucketLock;
 		FBucket& Bucket = FTable::FindOrCreateBucket(Address, BucketLock);
 		State.bDidWait = CanWait();
 		if (!State.bDidWait)
@@ -547,7 +613,7 @@ FWaitState WaitUntil(const void* Address, const TFunctionRef<bool()>& CanWait, c
 	// The timeout was reached and the thread needs to dequeue itself.
 	// This can race with a call to wake a thread, which means Self is unsafe to access outside of the lock.
 	bool bDequeued = false;
-	if (TDynamicUniqueLock<FBucketMutex> BucketLock; FBucket* Bucket = FTable::FindBucket(Address, BucketLock))
+	if (TDynamicUniqueLock<FWordMutex> BucketLock; FBucket* Bucket = FTable::FindBucket(Address, BucketLock))
 	{
 		Bucket->DequeueIf([Self = &Self, &bDequeued](FThread* Thread)
 		{
@@ -582,7 +648,7 @@ void WakeOne(const void* Address, const TFunctionRef<uint64(FWakeState)>& OnWake
 	uint64 WakeToken = 0;
 
 	{
-		TDynamicUniqueLock<FBucketMutex> BucketLock;
+		TDynamicUniqueLock<FWordMutex> BucketLock;
 		FBucket& Bucket = FTable::FindOrCreateBucket(Address, BucketLock);
 		Bucket.DequeueIf([Address, &WakeThread](FThread* Thread)
 		{
@@ -628,9 +694,9 @@ uint32 WakeMultiple(const void* const Address, const uint32 WakeCount)
 {
 	using namespace Private;
 
-	TArray<TRefCountPtr<FThread>> WakeThreads;
+	TArray<TRefCountPtr<FThread>, TInlineAllocator<128, FAnsiAllocator>> WakeThreads;
 
-	if (TDynamicUniqueLock<FBucketMutex> BucketLock; FBucket* Bucket = FTable::FindBucket(Address, BucketLock))
+	if (TDynamicUniqueLock<FWordMutex> BucketLock; FBucket* Bucket = FTable::FindBucket(Address, BucketLock))
 	{
 		Bucket->DequeueIf([Address, &WakeThreads, WakeCount](FThread* Thread)
 		{
