Logging: Added a public API for structured log formatting

diff --git a/Engine/Source/Runtime/Core/Private/Logging/StructuredLog.cpp b/Engine/Source/Runtime/Core/Private/Logging/StructuredLog.cpp
--- a/Engine/Source/Runtime/Core/Private/Logging/StructuredLog.cpp
+++ b/Engine/Source/Runtime/Core/Private/Logging/StructuredLog.cpp
@@ -12,6 +12,7 @@
 #include "HAL/PlatformTime.h"
 #include "Internationalization/Internationalization.h"
 #include "Logging/LogTrace.h"
+#include "Logging/StructuredLogFormat.h"
 #include "Misc/AsciiSet.h"
 #include "Misc/DateTime.h"
 #include "Misc/FeedbackContext.h"
@@ -238,11 +239,15 @@ public:
 	const uint8* GetOpData() const { return (const uint8*)(this + 1); }
 
 	template <typename CharType>
-	void FormatTo(TStringBuilderBase<CharType>& Out, const FCbObjectView& Fields) const;
+	void FormatTo(TStringBuilderBase<CharType>& Out, const FCbFieldViewIterator& Fields) const;
+
+	FText FormatToText(const FCbFieldViewIterator& Fields) const;
 
 private:
 	template <typename CharType>
-	void FormatLocalizedTo(TStringBuilderBase<CharType>& Out, const FCbObjectView& Fields) const;
+	void FormatLocalizedTo(TStringBuilderBase<CharType>& Out, const FCbFieldViewIterator& Fields) const;
+
+	FText FormatLocalizedToText(const FCbFieldViewIterator& Fields) const;
 
 	inline constexpr explicit FLogTemplate(const TCHAR* Format)
 		: StaticFormat(Format)
@@ -379,26 +384,30 @@ FLogTemplate* FLogTemplate::CreateLocalized(const TCHAR* TextNamespace, const TC
 
 	FTextFormat TextFormat(FInternationalization::ForUseOnlyByLocMacroAndGraphNodeTextLiterals_CreateText(Format, TextNamespace, TextKey));
 
-	int32 FormatFieldCount = 0;
-	TArray<FString> TextFormatFieldNames;
-	TextFormat.GetFormatArgumentNames(TextFormatFieldNames);
-	for (const FString& FormatName : TextFormatFieldNames)
+	const bool bFindFields = !!Fields;
+	if (bFindFields)
 	{
-		bool bFoundField = false;
-		const TCHAR* FormatNameStr = *FormatName;
-		const int32 FormatNameLen = FormatName.Len();
-		for (int32 BaseFieldIndex = 0; BaseFieldIndex < FieldCount; ++BaseFieldIndex)
+		int32 FormatFieldCount = 0;
+		TArray<FString> TextFormatFieldNames;
+		TextFormat.GetFormatArgumentNames(TextFormatFieldNames);
+		for (const FString& FormatName : TextFormatFieldNames)
 		{
-			const int32 FieldIndex = (FormatFieldCount + BaseFieldIndex) % FieldCount;
-			const ANSICHAR* FieldName = Fields[FieldIndex].Name;
-			if (FPlatformString::Strncmp(FieldName, FormatNameStr, FormatNameLen) == 0 && !FieldName[FormatNameLen])
+			bool bFoundField = false;
+			const TCHAR* FormatNameStr = *FormatName;
+			const int32 FormatNameLen = FormatName.Len();
+			for (int32 BaseFieldIndex = 0; BaseFieldIndex < FieldCount; ++BaseFieldIndex)
 			{
-				bFoundField = true;
-				break;
+				const int32 FieldIndex = (FormatFieldCount + BaseFieldIndex) % FieldCount;
+				const ANSICHAR* FieldName = Fields[FieldIndex].Name;
+				if (FPlatformString::Strncmp(FieldName, FormatNameStr, FormatNameLen) == 0 && !FieldName[FormatNameLen])
+				{
+					bFoundField = true;
+					break;
+				}
 			}
+			checkf(bFoundField, TEXT("Log format requires field '%s' which was not provided. [[%s]]"), FormatNameStr, Format);
+			++FormatFieldCount;
 		}
-		checkf(bFoundField, TEXT("Log format requires field '%s' which was not provided. [[%s]]"), FormatNameStr, Format);
-		++FormatFieldCount;
 	}
 
 	const FLogTemplateOp Ops[]{{FLogTemplateOp::OpLocText}, {FLogTemplateOp::OpEnd}};
@@ -431,11 +440,11 @@ void FLogTemplate::Destroy(FLogTemplate* Template)
 }
 
 template <typename CharType>
-void FLogTemplate::FormatTo(TStringBuilderBase<CharType>& Out, const FCbObjectView& Fields) const
+void FLogTemplate::FormatTo(TStringBuilderBase<CharType>& Out, const FCbFieldViewIterator& Fields) const
 {
 	using namespace Logging::Private;
 
-	auto FindField = [&Fields, It = Fields.CreateViewIterator(), Index = 0, Format = StaticFormat](FAnsiStringView Name, int32 IndexHint = -1) mutable -> FCbFieldView&
+	auto FindField = [&Fields, It = Fields, Index = 0, Format = StaticFormat](FAnsiStringView Name, int32 IndexHint = -1) mutable -> FCbFieldView&
 	{
 		if (IndexHint >= 0)
 		{
@@ -444,7 +453,7 @@ void FLogTemplate::FormatTo(TStringBuilderBase<CharType>& Out, const FCbObjectVi
 			}
 			if (IndexHint < Index)
 			{
-				It = Fields.CreateViewIterator();
+				It = Fields;
 				for (Index = 0; Index < IndexHint && It; ++Index, ++It);
 			}
 			if (IndexHint == Index && Name.Equals(It.GetName()))
@@ -460,7 +469,7 @@ void FLogTemplate::FormatTo(TStringBuilderBase<CharType>& Out, const FCbObjectVi
 				return It;
 			}
 		}
-		It = Fields.CreateViewIterator();
+		It = Fields;
 		for (Index = 0; Index < PrevIndex && It; ++Index, ++It)
 		{
 			if (Name.Equals(It.GetName()))
@@ -500,8 +509,30 @@ void FLogTemplate::FormatTo(TStringBuilderBase<CharType>& Out, const FCbObjectVi
 	}
 }
 
+FText FLogTemplate::FormatToText(const FCbFieldViewIterator& Fields) const
+{
+	using namespace Logging::Private;
+
+	const uint8* NextOp = GetOpData();
+	if (FLogTemplateOp::Load(NextOp).Code == FLogTemplateOp::OpLocText)
+	{
+		return FormatLocalizedToText(Fields);
+	}
+	else
+	{
+		TStringBuilder<512> Builder;
+		FormatTo(Builder, Fields);
+		return FText::FromStringView(Builder);
+	}
+}
+
 template <typename CharType>
-FORCENOINLINE void FLogTemplate::FormatLocalizedTo(TStringBuilderBase<CharType>& Out, const FCbObjectView& Fields) const
+FORCENOINLINE void FLogTemplate::FormatLocalizedTo(TStringBuilderBase<CharType>& Out, const FCbFieldViewIterator& Fields) const
+{
+	Out.Append(FormatLocalizedToText(Fields).ToString());
+}
+
+FText FLogTemplate::FormatLocalizedToText(const FCbFieldViewIterator& Fields) const
 {
 	using namespace Logging::Private;
 
@@ -512,8 +543,39 @@ FORCENOINLINE void FLogTemplate::FormatLocalizedTo(TStringBuilderBase<CharType>&
 	}
 
 	const FTextFormat* TextFormat = (const FTextFormat*)this - 1;
-	const FText Text = FText::Format(*TextFormat, TextFormatArguments);
-	Out.Append(Text.ToString());
+	return FText::Format(*TextFormat, TextFormatArguments);
+}
+
+///////////////////////////////////////////////////////////////////////////////////////////////////////////////////////
+
+FLogTemplate* CreateLogTemplate(const TCHAR* Format)
+{
+	return FLogTemplate::Create(Format);
+}
+
+FLogTemplate* CreateLogTemplate(const TCHAR* TextNamespace, const TCHAR* TextKey, const TCHAR* Format)
+{
+	return FLogTemplate::CreateLocalized(TextNamespace, TextKey, Format);
+}
+
+void DestroyLogTemplate(FLogTemplate* Template)
+{
+	FLogTemplate::Destroy(Template);
+}
+
+void FormatLogTo(FUtf8StringBuilderBase& Out, const FLogTemplate* Template, const FCbFieldViewIterator& Fields)
+{
+	Template->FormatTo(Out, Fields);
+}
+
+void FormatLogTo(FWideStringBuilderBase& Out, const FLogTemplate* Template, const FCbFieldViewIterator& Fields)
+{
+	Template->FormatTo(Out, Fields);
+}
+
+FText FormatLogToText(const FLogTemplate* Template, const FCbFieldViewIterator& Fields)
+{
+	return Template->FormatToText(Fields);
 }
 
 ///////////////////////////////////////////////////////////////////////////////////////////////////////////////////////
@@ -547,7 +609,7 @@ FORCENOINLINE static void FormatDynamicRecordMessageTo(TStringBuilderBase<CharTy
 		TEXT("Log record must have both or neither of the text namespace and text key. [[%s]]"), Format);
 
 	FLogTemplate* LocalTemplate = TextKey ? FLogTemplate::CreateLocalized(TextNamespace, TextKey, Format) : FLogTemplate::Create(Format);
-	LocalTemplate->FormatTo(Out, Record.GetFields());
+	LocalTemplate->FormatTo(Out, Record.GetFields().CreateViewIterator());
 	FLogTemplate::Destroy(LocalTemplate);
 }
 
@@ -557,7 +619,7 @@ static void FormatRecordMessageTo(TStringBuilderBase<CharType>& Out, const FLogR
 	const FLogTemplate* Template = Record.GetTemplate();
 	if (LIKELY(Template))
 	{
-		return Template->FormatTo(Out, Record.GetFields());
+		return Template->FormatTo(Out, Record.GetFields().CreateViewIterator());
 	}
 	FormatDynamicRecordMessageTo(Out, Record);
 }
diff --git a/Engine/Source/Runtime/Core/Public/Logging/StructuredLogFormat.h b/Engine/Source/Runtime/Core/Public/Logging/StructuredLogFormat.h
new file mode 100644
--- /dev/null
+++ b/Engine/Source/Runtime/Core/Public/Logging/StructuredLogFormat.h
@@ -0,0 +1,27 @@
+// Copyright Epic Games, Inc. All Rights Reserved.
+
+#pragma once
+
+#include "Containers/StringFwd.h"
+
+#define UE_API CORE_API
+
+class FCbFieldViewIterator;
+class FText;
+
+namespace UE { class FLogTemplate; }
+
+namespace UE
+{
+
+UE_API FLogTemplate* CreateLogTemplate(const TCHAR* Format);
+UE_API FLogTemplate* CreateLogTemplate(const TCHAR* TextNamespace, const TCHAR* TextKey, const TCHAR* Format);
+UE_API void DestroyLogTemplate(FLogTemplate* Template);
+
+UE_API void FormatLogTo(FUtf8StringBuilderBase& Out, const FLogTemplate* Template, const FCbFieldViewIterator& Fields);
+UE_API void FormatLogTo(FWideStringBuilderBase& Out, const FLogTemplate* Template, const FCbFieldViewIterator& Fields);
+UE_API FText FormatLogToText(const FLogTemplate* Template, const FCbFieldViewIterator& Fields);
+
+} // UE
+
+#undef UE_API
