Handle shared pointer acquire references correctly in AutoRTFM.

diff --git a/Engine/Source/Runtime/Core/Public/AutoRTFM/AutoRTFM.h b/Engine/Source/Runtime/Core/Public/AutoRTFM/AutoRTFM.h
--- a/Engine/Source/Runtime/Core/Public/AutoRTFM/AutoRTFM.h
+++ b/Engine/Source/Runtime/Core/Public/AutoRTFM/AutoRTFM.h
@@ -524,10 +524,12 @@ struct FRegisterOpenFunction
 // Macro-based variants so we completely compile away when not in use, even in debug builds
 #if UE_AUTORTFM
 #define UE_AUTORTFM_OPEN(x) AutoRTFM::Open([&]() { x })
+#define UE_AUTORTFM_OPENABORT(x) AutoRTFM::OpenAbort([=]() { x })
 #define UE_AUTORTFM_OPENCOMMIT(x) AutoRTFM::OpenCommit([=]() { x })
 #define UE_AUTORTFM_TRANSACT(x) AutoRTFM::Transact([&]() { x })
 #else
 #define UE_AUTORTFM_OPEN(x) do { x } while (false)
+#define UE_AUTORTFM_OPENABORT(x) do { /* do nothing */ } while (false)
 #define UE_AUTORTFM_OPENCOMMIT(x) do { x } while (false)
 #define UE_AUTORTFM_TRANSACT(x) do { x } while (false)
 #endif
diff --git a/Engine/Source/Runtime/Core/Public/Templates/SharedPointerInternals.h b/Engine/Source/Runtime/Core/Public/Templates/SharedPointerInternals.h
--- a/Engine/Source/Runtime/Core/Public/Templates/SharedPointerInternals.h
+++ b/Engine/Source/Runtime/Core/Public/Templates/SharedPointerInternals.h
@@ -115,12 +115,24 @@ namespace SharedPointerInternals
 				// in response to the increment, so there's nothing to order with.
 
 #if defined(_MSC_VER) && (defined(_M_X64) || defined(_M_IX86))
-				// We do a regular SC increment here because it maps to an _InterlockedIncrement (lock inc).
-				// The codegen for a relaxed fetch_add is actually much worse under MSVC (lock xadd).
-				++SharedReferenceCount;
+				UE_AUTORTFM_OPEN(
+				{
+					// We do a regular SC increment here because it maps to an _InterlockedIncrement (lock inc).
+					// The codegen for a relaxed fetch_add is actually much worse under MSVC (lock xadd).
+					++SharedReferenceCount;
+				});
 #else
-				SharedReferenceCount.fetch_add(1, std::memory_order_relaxed);
+				UE_AUTORTFM_OPEN(
+				{
+					SharedReferenceCount.fetch_add(1, std::memory_order_relaxed);
+				});
 #endif
+
+				// If the transaction would abort, we need to undo adding the shared reference.
+				UE_AUTORTFM_OPENABORT(
+				{
+					ReleaseSharedReference();
+				});
 			}
 			else
 			{
