CompressedBuffer : profile scope and note that decoder is serial and should be parallelized

#rnx

diff --git a/Engine/Source/Runtime/Core/Private/Compression/CompressedBuffer.cpp b/Engine/Source/Runtime/Core/Private/Compression/CompressedBuffer.cpp
--- a/Engine/Source/Runtime/Core/Private/Compression/CompressedBuffer.cpp
+++ b/Engine/Source/Runtime/Core/Private/Compression/CompressedBuffer.cpp
@@ -439,6 +439,8 @@ bool FBlockDecoder::TryDecompressTo(
 	const uint64 RawOffset,
 	FMutableMemoryView RawView) const
 {
+	TRACE_CPUPROFILER_EVENT_SCOPE(FBlockDecoder::TryDecompressTo);
+
 	if (Header.TotalRawSize < RawOffset + RawView.GetSize())
 	{
 		return false;
@@ -456,6 +458,8 @@ bool FBlockDecoder::TryDecompressTo(
 		Algo::TransformAccumulate(MakeArrayView(CompressedBlockSizes, FirstBlockIndex),
 			[](uint32 Size) -> uint64 { return NETWORK_ORDER32(Size); }, uint64(0));
 
+	// @todo : Parallel decoder needed!  This stream is chunked into blocks but this implementation is serial.  This is a bottleneck for large data loading.
+
 	for (uint32 BlockIndex = FirstBlockIndex; BlockIndex <= LastBlockIndex; BlockIndex++)
 	{
 		const uint64 RawBlockSize = BlockIndex == Header.BlockCount - 1 ? LastBlockSize : BlockSize;
