[autortfm]
- fix an issue with the WriteLog where we were applying undos in the order they were done. They need to be done in reverse order to correctly roll back the effects. It seems that sometimes we have more than one entry for the same address and size, and sometimes we also have overlapping ranges
- FWriteLog is now a doubly linked list, added a ReverseIterator

diff --git a/Engine/Source/Runtime/Core/Private/AutoRTFM/Transaction.cpp b/Engine/Source/Runtime/Core/Private/AutoRTFM/Transaction.cpp
--- a/Engine/Source/Runtime/Core/Private/AutoRTFM/Transaction.cpp
+++ b/Engine/Source/Runtime/Core/Private/AutoRTFM/Transaction.cpp
@@ -71,8 +71,18 @@ bool FTransaction::AttemptToCommit()
 
 void FTransaction::Undo()
 {
-    for (FWriteLogEntry& Entry : WriteLog)
+	const bool bVerboseUndo = false;
+
+	if (bVerboseUndo)
+	{
+		fprintf(GetLogFile(), "[FTransaction::Undo] START\n");
+	}
+
+	int VerboseCounter = 0;
+	for(auto Iter = WriteLog.rbegin(); Iter != WriteLog.rend(); ++Iter)
     {
+		FWriteLogEntry& Entry = *Iter;
+
         // Skip writes to our current transaction nest if we're scoped. We're about to
 		// leave so the changes don't matter. 
         if (IsScopedTransaction() && Context->IsInnerTransactionStack(Entry.OriginalAndSize.Get()))
@@ -83,8 +93,41 @@ void FTransaction::Undo()
         void* const Original = Entry.OriginalAndSize.Get();
         const size_t Size = Entry.OriginalAndSize.GetTopTag();
         void* const Copy = Entry.Copy;
+
+		if (bVerboseUndo)
+		{
+			FILE* LogFile = GetLogFile();
+			fprintf(LogFile, "%4d [UNDO] %p %4llu : ", VerboseCounter, Original, Size);
+
+			unsigned char* Current = (unsigned char*)Original;
+			unsigned char* Old = (unsigned char*)Copy;
+
+			fprintf(LogFile, "[");
+
+			for(size_t i = 0; i < Size; i++)
+			{
+				fprintf(LogFile, "%02X ", Current[i]);
+			}
+
+			fprintf(LogFile, "] -> [");
+
+			for(size_t i = 0; i < Size; i++)
+			{
+				fprintf(LogFile, "%02X ", Old[i]);
+			}
+
+			fprintf(LogFile, "]\n");
+		}
+
         memcpy(Original, Copy, Size);
+
+		VerboseCounter++;
     }
+
+	if (bVerboseUndo)
+	{
+		fprintf(GetLogFile(), "[FTransaction::Undo] END\n");
+	}
 }
 
 void FTransaction::AbortNested()
diff --git a/Engine/Source/Runtime/Core/Private/AutoRTFM/WriteLog.h b/Engine/Source/Runtime/Core/Private/AutoRTFM/WriteLog.h
--- a/Engine/Source/Runtime/Core/Private/AutoRTFM/WriteLog.h
+++ b/Engine/Source/Runtime/Core/Private/AutoRTFM/WriteLog.h
@@ -37,6 +37,7 @@ namespace AutoRTFM
             FWriteLogEntry Entries[BucketSize];
             size_t Size = 0;
             FWriteLogEntryBucket* Next = nullptr;
+            FWriteLogEntryBucket* Prev = nullptr;
         };
 
     public:
@@ -51,6 +52,7 @@ namespace AutoRTFM
             if (FWriteLogEntryBucket::BucketSize == Current->Size)
             {
                 Current->Next = new FWriteLogEntryBucket();
+				Current->Next->Prev = Current;
                 Current = Current->Next;
             }
 
@@ -95,16 +97,56 @@ namespace AutoRTFM
             size_t Offset = 0;
         };
 
+        struct ReverseIterator final
+        {
+            explicit ReverseIterator(FWriteLogEntryBucket* const Bucket) : Bucket(Bucket) {}
+
+            FWriteLogEntry& operator*()
+            {
+                return Bucket->Entries[Bucket->Size - Offset - 1];
+            }
+
+            void operator++()
+            {
+                Offset++;
+
+                if (Offset == Bucket->Size)
+                {
+                    Bucket = Bucket->Prev;
+                    Offset = 0;
+                }
+            }
+
+            bool operator!=(const ReverseIterator& Other) const
+            {
+                return (Other.Bucket != Bucket) || (Other.Offset != Offset);
+            }
+
+        private:
+            FWriteLogEntryBucket* Bucket;
+            size_t Offset = 0;
+        };
+
         Iterator begin()
         {
             return Iterator(Start);
         }
 
+        ReverseIterator rbegin()
+        {
+            return ReverseIterator(Current);
+        }
+
         Iterator end()
         {
             return Iterator(nullptr);
         }
 
+        ReverseIterator rend()
+        {
+            return ReverseIterator(nullptr);
+        }
+
         void Reset()
         {
             while (nullptr != Start)
