Fix stack walking other processes on Windows.

diff --git a/Engine/Source/Runtime/Core/Private/Windows/WindowsPlatformStackWalk.cpp b/Engine/Source/Runtime/Core/Private/Windows/WindowsPlatformStackWalk.cpp
--- a/Engine/Source/Runtime/Core/Private/Windows/WindowsPlatformStackWalk.cpp
+++ b/Engine/Source/Runtime/Core/Private/Windows/WindowsPlatformStackWalk.cpp
@@ -188,15 +188,111 @@ void FWindowsPlatformStackWalk::ThreadStackWalkAndDump(ANSICHAR* HumanReadableSt
 	}
 }
 
-uint32 FWindowsPlatformStackWalk::CaptureThreadStackBackTrace(uint64 ThreadId, uint64* BackTrace, uint32 MaxDepth, void* Context)
+static int32 CaptureStackTraceExternalProcess(uint64* BackTrace, uint32 MaxDepth, CONTEXT* InContext, HANDLE ThreadHandle, uint32* Depth)
+{
+	STACKFRAME64		StackFrame64;
+	HANDLE				ProcessHandle = GProcessHandle;
+	unsigned long		LastError;
+	bool				bStackWalkSucceeded	= true;
+	uint32				CurrentDepth		= 0;
+	uint32				MachineType			= IMAGE_FILE_MACHINE_I386;
+	CONTEXT				ContextCopy = *InContext;
+
+	*Depth = 0;
+
+#if !PLATFORM_SEH_EXCEPTIONS_DISABLED
+	__try
+#endif
+	{
+		// Zero out stack frame.
+		FMemory::Memzero( StackFrame64 );
+
+		// Initialize the STACKFRAME structure.
+		StackFrame64.AddrPC.Mode         = AddrModeFlat;
+		StackFrame64.AddrStack.Mode      = AddrModeFlat;
+		StackFrame64.AddrFrame.Mode      = AddrModeFlat;
+#if PLATFORM_64BITS
+		StackFrame64.AddrPC.Offset = ContextCopy.Rip;
+		StackFrame64.AddrStack.Offset = ContextCopy.Rsp;
+		StackFrame64.AddrFrame.Offset = ContextCopy.Rbp;
+		MachineType                      = IMAGE_FILE_MACHINE_AMD64;
+#else	//PLATFORM_64BITS
+		StackFrame64.AddrPC.Offset       = ContextCopy.Eip;
+		StackFrame64.AddrStack.Offset    = ContextCopy.Esp;
+		StackFrame64.AddrFrame.Offset    = ContextCopy.Ebp;
+#endif	//PLATFORM_64BITS
+
+		// Walk the stack one frame at a time.
+		while( bStackWalkSucceeded && (CurrentDepth < MaxDepth) )
+		{
+			bStackWalkSucceeded = !!StackWalk64(MachineType, 
+												ProcessHandle, 
+												ThreadHandle, 
+												&StackFrame64,
+												&ContextCopy,
+												NULL,
+												SymFunctionTableAccess64,
+												SymGetModuleBase64,
+												NULL );
+
+			if( !bStackWalkSucceeded  )
+			{
+				// StackWalk failed! give up.
+				LastError = GetLastError( );
+				break;
+			}
+
+			// Stop if the frame pointer is NULL.
+			// Note that the thread's PC 'StackFrame64.AddrPC.Offset' COULD be 0 in case something calls a nullptr.
+			if( StackFrame64.AddrFrame.Offset == 0 )
+			{
+				break;
+			}
+
+			BackTrace[CurrentDepth++] = StackFrame64.AddrPC.Offset;
+			*Depth = CurrentDepth;
+		}
+	} 
+#if !PLATFORM_SEH_EXCEPTIONS_DISABLED
+	__except( EXCEPTION_EXECUTE_HANDLER )
+	{
+		// We need to catch any exceptions within this function so they don't get sent to 
+		// the engine's error handler, hence causing an infinite loop.
+		return EXCEPTION_EXECUTE_HANDLER;
+	} 
+#endif
+
+	// NULL out remaining entries.
+	for ( ; CurrentDepth<MaxDepth; CurrentDepth++ )
+	{
+		BackTrace[CurrentDepth] = 0;
+	}
+
+	return EXCEPTION_EXECUTE_HANDLER;
+}
+
+void FWindowsPlatformStackWalk::CaptureStackTraceByProcess(uint64* OutBacktrace, uint32 MaxDepth, void* InContext, void* InThreadHandle, uint32* OutDepth, bool bExternalProcess)
 {
-	// Do not call InitStackWalking as we're just capturing addresses without using the dbghelp library etc.
+	if (bExternalProcess)
+	{
+		CaptureStackTraceExternalProcess(OutBacktrace, MaxDepth, reinterpret_cast<PCONTEXT>(InContext), reinterpret_cast<HANDLE>(InThreadHandle), OutDepth);	
+	}	
+	else
+	{
+		FMicrosoftPlatformStackWalk::CaptureStackTraceInternal(OutBacktrace, MaxDepth, InContext, InThreadHandle, OutDepth);
+	}
+}
 
+uint32 FWindowsPlatformStackWalk::CaptureThreadStackBackTrace(uint64 ThreadId, uint64* BackTrace, uint32 MaxDepth, void* Context)
+{
 	if (BackTrace == nullptr || MaxDepth == 0)
 	{
 		return 0;
 	}
+	
 
+	// Do not call InitStackWalking as we're just capturing addresses without using the dbghelp library etc if this is the same process.
+	// If we're walking another process, they will already have initialized stackwalking to set GProcessHandle.
 	// Don't suspend the calling thread, capture it's context directly and trace
 	if (ThreadId == FPlatformTLS::GetCurrentThreadId())
 	{
@@ -210,17 +306,19 @@ uint32 FWindowsPlatformStackWalk::CaptureThreadStackBackTrace(uint64 ThreadId, u
 	}
 
 	// Suspend the thread before grabbing its context
-	if (SuspendThread(ThreadHandle) < 0)
+	if (SuspendThread(ThreadHandle) == (DWORD)-1)
 	{
 		return 0;
 	}
 
+	bool bExternalProcess = GProcessHandle != INVALID_HANDLE_VALUE && GProcessHandle != GetCurrentProcess();
+
 	uint32 Depth = 0;
 	if (Context)
 	{
 		FWindowsThreadContextWrapper* ContextWrapper = reinterpret_cast<FWindowsThreadContextWrapper*>(Context);
 		ContextWrapper->ThreadHandle = ThreadHandle; // Use the thread handle open above to be sure it has the correct permissions.
-		FMicrosoftPlatformStackWalk::CaptureStackTraceInternal(BackTrace, MaxDepth, &ContextWrapper->Context, ThreadHandle, &Depth);
+		CaptureStackTraceByProcess(BackTrace, MaxDepth, &ContextWrapper->Context, ThreadHandle, &Depth, bExternalProcess);
 	}
 	else
 	{
@@ -230,7 +328,7 @@ uint32 FWindowsPlatformStackWalk::CaptureThreadStackBackTrace(uint64 ThreadId, u
 
 		if (::GetThreadContext(ThreadHandle, &ThreadContext))
 		{
-			FMicrosoftPlatformStackWalk::CaptureStackTraceInternal(BackTrace, MaxDepth, &ThreadContext, ThreadHandle, &Depth);
+			CaptureStackTraceByProcess(BackTrace, MaxDepth, &ThreadContext, ThreadHandle, &Depth, bExternalProcess);
 		}
 	}
 
diff --git a/Engine/Source/Runtime/Core/Public/Windows/WindowsPlatformStackWalk.h b/Engine/Source/Runtime/Core/Public/Windows/WindowsPlatformStackWalk.h
--- a/Engine/Source/Runtime/Core/Public/Windows/WindowsPlatformStackWalk.h
+++ b/Engine/Source/Runtime/Core/Public/Windows/WindowsPlatformStackWalk.h
@@ -53,7 +53,7 @@ struct CORE_API FWindowsPlatformStackWalk
 	 * The implementation extracts the information from the debug engine and the debug symbols and
 	 * takes care of loading the debug symbols if the debug engine was configured to load symbols
 	 * on demand. This function can be expensive if the debug symbols needs to be loaded.
-	 *
+	 FMovieSceneEventParameters*
 	 * @param FunctionSymbolName The function name to lookup.
 	 * @param FunctionModuleName The module name containing the function to lookup.
 	 * @param OutPathname The source file pathname.
@@ -63,6 +63,8 @@ struct CORE_API FWindowsPlatformStackWalk
 	 */
 	static bool GetFunctionDefinitionLocation(const FString& FunctionSymbolName, const FString& FunctionModuleName, FString& OutPathname, uint32& OutLineNumber, uint32& OutColumnNumber);
 
+protected:
+	static void CaptureStackTraceByProcess(uint64* OutBacktrace, uint32 MaxDepth, void* InContext, void* InThreadHandle, uint32* OutDepth, bool bExternalProcess);
 private:
 	static bool InitStackWalkingInternal(void* Process, bool bForceReinitOnProcessMismatch);
 };
