I/O - move FIoOffsetAndLength to separate header

diff --git a/Engine/Source/Runtime/Core/Internal/IO/IoOffsetLength.h b/Engine/Source/Runtime/Core/Internal/IO/IoOffsetLength.h
new file mode 100644
--- /dev/null
+++ b/Engine/Source/Runtime/Core/Internal/IO/IoOffsetLength.h
@@ -0,0 +1,53 @@
+// Copyright Epic Games, Inc. All Rights Reserved.
+
+#pragma once
+
+/**
+ * Combined offset and length.
+ */
+struct FIoOffsetAndLength
+{
+public:
+	inline uint64 GetOffset() const
+	{
+		return OffsetAndLength[4]
+			| (uint64(OffsetAndLength[3]) << 8)
+			| (uint64(OffsetAndLength[2]) << 16)
+			| (uint64(OffsetAndLength[1]) << 24)
+			| (uint64(OffsetAndLength[0]) << 32)
+			;
+	}
+
+	inline uint64 GetLength() const
+	{
+		return OffsetAndLength[9]
+			| (uint64(OffsetAndLength[8]) << 8)
+			| (uint64(OffsetAndLength[7]) << 16)
+			| (uint64(OffsetAndLength[6]) << 24)
+			| (uint64(OffsetAndLength[5]) << 32)
+			;
+	}
+
+	inline void SetOffset(uint64 Offset)
+	{
+		OffsetAndLength[0] = uint8(Offset >> 32);
+		OffsetAndLength[1] = uint8(Offset >> 24);
+		OffsetAndLength[2] = uint8(Offset >> 16);
+		OffsetAndLength[3] = uint8(Offset >>  8);
+		OffsetAndLength[4] = uint8(Offset >>  0);
+	}
+
+	inline void SetLength(uint64 Length)
+	{
+		OffsetAndLength[5] = uint8(Length >> 32);
+		OffsetAndLength[6] = uint8(Length >> 24);
+		OffsetAndLength[7] = uint8(Length >> 16);
+		OffsetAndLength[8] = uint8(Length >> 8);
+		OffsetAndLength[9] = uint8(Length >> 0);
+	}
+
+private:
+	// We use 5 bytes for offset and size, this is enough to represent
+	// an offset and size of 1PB
+	uint8 OffsetAndLength[5 + 5];
+};
diff --git a/Engine/Source/Runtime/Core/Internal/IO/IoStore.h b/Engine/Source/Runtime/Core/Internal/IO/IoStore.h
--- a/Engine/Source/Runtime/Core/Internal/IO/IoStore.h
+++ b/Engine/Source/Runtime/Core/Internal/IO/IoStore.h
@@ -7,6 +7,7 @@
 #include "HAL/UnrealMemory.h"
 #include "IO/IoContainerId.h"
 #include "IO/IoDispatcher.h"
+#include "IO/IoOffsetLength.h"
 #include "Logging/LogMacros.h"
 #include "Misc/EnumClassFlags.h"
 #include "Misc/Guid.h"
@@ -75,56 +76,6 @@ struct FIoStoreTocHeader
 	}
 };
 
-/**
- * Combined offset and length.
- */
-struct FIoOffsetAndLength
-{
-public:
-	inline uint64 GetOffset() const
-	{
-		return OffsetAndLength[4]
-			| (uint64(OffsetAndLength[3]) << 8)
-			| (uint64(OffsetAndLength[2]) << 16)
-			| (uint64(OffsetAndLength[1]) << 24)
-			| (uint64(OffsetAndLength[0]) << 32)
-			;
-	}
-
-	inline uint64 GetLength() const
-	{
-		return OffsetAndLength[9]
-			| (uint64(OffsetAndLength[8]) << 8)
-			| (uint64(OffsetAndLength[7]) << 16)
-			| (uint64(OffsetAndLength[6]) << 24)
-			| (uint64(OffsetAndLength[5]) << 32)
-			;
-	}
-
-	inline void SetOffset(uint64 Offset)
-	{
-		OffsetAndLength[0] = uint8(Offset >> 32);
-		OffsetAndLength[1] = uint8(Offset >> 24);
-		OffsetAndLength[2] = uint8(Offset >> 16);
-		OffsetAndLength[3] = uint8(Offset >>  8);
-		OffsetAndLength[4] = uint8(Offset >>  0);
-	}
-
-	inline void SetLength(uint64 Length)
-	{
-		OffsetAndLength[5] = uint8(Length >> 32);
-		OffsetAndLength[6] = uint8(Length >> 24);
-		OffsetAndLength[7] = uint8(Length >> 16);
-		OffsetAndLength[8] = uint8(Length >> 8);
-		OffsetAndLength[9] = uint8(Length >> 0);
-	}
-
-private:
-	// We use 5 bytes for offset and size, this is enough to represent
-	// an offset and size of 1PB
-	uint8 OffsetAndLength[5 + 5];
-};
-
 enum class FIoStoreTocEntryMetaFlags : uint8
 {
 	None,
