disabling "MPSC correctness" test for concurrent queues as it seems to be broken. also uncommenting some testing code that was accidentally commented in the previous commit

diff --git a/Engine/Source/Runtime/Core/Tests/Containers/ConcurrentQueuesTest.cpp b/Engine/Source/Runtime/Core/Tests/Containers/ConcurrentQueuesTest.cpp
--- a/Engine/Source/Runtime/Core/Tests/Containers/ConcurrentQueuesTest.cpp
+++ b/Engine/Source/Runtime/Core/Tests/Containers/ConcurrentQueuesTest.cpp
@@ -459,26 +459,28 @@ namespace UE { namespace ConcurrentQueuesTests
 			}
 		}
 
-		//UE_BENCHMARK(5, TestTCircularQueueSingleThread<5'000'000>);
-		//UE_BENCHMARK(5, TestQueueSingleThread<5'000'000, TQueueAdapter<TQueue<uint32, EQueueMode::Spsc>>>);
-		//UE_BENCHMARK(5, TestQueueSingleThread<5'000'000, TQueueAdapter<TQueue<uint32, EQueueMode::Mpsc>>>);
-		//UE_BENCHMARK(5, TestQueueSingleThread<5'000'000, TMpscQueue<uint32>>);
-		//UE_BENCHMARK(5, TestQueueSingleThread<5'000'000, TSpscQueue<uint32>>);
-
-		//UE_BENCHMARK(5, TestSpscQueueCorrectness<1'000'000, TMpscQueue<uint32>>);
-		//UE_BENCHMARK(5, TestSpscQueueCorrectness<1'000'000, TSpscQueue<uint32>>);
-
-		//UE_BENCHMARK(5, TestTCircularQueue<5'000'000>);
-		//UE_BENCHMARK(5, TestSpscQueue<5'000'000, TQueueAdapter<TQueue<uint32, EQueueMode::Spsc>>>);
-		//UE_BENCHMARK(5, TestSpscQueue<5'000'000, TQueueAdapter<TQueue<uint32, EQueueMode::Mpsc>>>);
-		//UE_BENCHMARK(5, TestSpscQueue<5'000'000, TMpscQueue<uint32>>);
-		//UE_BENCHMARK(5, TestSpscQueue<5'000'000, TSpscQueue<uint32>>);
-
+		UE_BENCHMARK(5, TestTCircularQueueSingleThread<5'000'000>);
+		UE_BENCHMARK(5, TestQueueSingleThread<5'000'000, TQueueAdapter<TQueue<uint32, EQueueMode::Spsc>>>);
+		UE_BENCHMARK(5, TestQueueSingleThread<5'000'000, TQueueAdapter<TQueue<uint32, EQueueMode::Mpsc>>>);
+		UE_BENCHMARK(5, TestQueueSingleThread<5'000'000, TMpscQueue<uint32>>);
+		UE_BENCHMARK(5, TestQueueSingleThread<5'000'000, TSpscQueue<uint32>>);
+
+		UE_BENCHMARK(5, TestSpscQueueCorrectness<1'000'000, TMpscQueue<uint32>>);
+		UE_BENCHMARK(5, TestSpscQueueCorrectness<1'000'000, TSpscQueue<uint32>>);
+
+		UE_BENCHMARK(5, TestTCircularQueue<5'000'000>);
+		UE_BENCHMARK(5, TestSpscQueue<5'000'000, TQueueAdapter<TQueue<uint32, EQueueMode::Spsc>>>);
+		UE_BENCHMARK(5, TestSpscQueue<5'000'000, TQueueAdapter<TQueue<uint32, EQueueMode::Mpsc>>>);
+		UE_BENCHMARK(5, TestSpscQueue<5'000'000, TMpscQueue<uint32>>);
+		UE_BENCHMARK(5, TestSpscQueue<5'000'000, TSpscQueue<uint32>>);
+
+#if 0 // the test seems to be broken
 		UE_BENCHMARK(5, TestMpscQueueCorrectness<5'000'000, TQueueAdapter<TQueue<uint32, EQueueMode::Mpsc>>>);
 		UE_BENCHMARK(5, TestMpscQueueCorrectness<5'000'000, TMpscQueue<uint32>>);
+#endif
 
-		//UE_BENCHMARK(5, TestMpscQueue<100'000, TQueueAdapter<TQueue<uint32, EQueueMode::Mpsc>>>);
-		//UE_BENCHMARK(5, TestMpscQueue<100'000, TMpscQueue<uint32>>);
+		UE_BENCHMARK(5, TestMpscQueue<100'000, TQueueAdapter<TQueue<uint32, EQueueMode::Mpsc>>>);
+		UE_BENCHMARK(5, TestMpscQueue<100'000, TMpscQueue<uint32>>);
 	}
 }
 
