Separated out a UnicastDelegateType typedef from the invocation list element type in TMulticastDelegateBase.

diff --git a/Engine/Source/Runtime/Core/Public/Delegates/MulticastDelegateBase.h b/Engine/Source/Runtime/Core/Public/Delegates/MulticastDelegateBase.h
--- a/Engine/Source/Runtime/Core/Public/Delegates/MulticastDelegateBase.h
+++ b/Engine/Source/Runtime/Core/Public/Delegates/MulticastDelegateBase.h
@@ -30,7 +30,9 @@ protected:
 	using typename Super::FWriteAccessScope;
 	using Super::GetWriteAccessScope;
 
-	using InvocationListType = TArray<TDelegateBase<UserPolicy>, FMulticastInvocationListAllocatorType>;
+	using UnicastDelegateType = TDelegateBase<UserPolicy>;
+
+	using InvocationListType = TArray<UnicastDelegateType, FMulticastInvocationListAllocatorType>;
 
 public:
 	/** Removes all functions from this delegate's invocation list. */
@@ -38,7 +40,7 @@ public:
 	{
 		FWriteAccessScope WriteScope = GetWriteAccessScope();
 
-		for (TDelegateBase<UserPolicy>& DelegateBaseRef : InvocationList)
+		for (UnicastDelegateType& DelegateBaseRef : InvocationList)
 		{
 			DelegateBaseRef.Unbind();
 		}
@@ -55,7 +57,7 @@ public:
 	{
 		FReadAccessScope ReadScope = GetReadAccessScope();
 	
-		for (const TDelegateBase<UserPolicy>& DelegateBaseRef : InvocationList)
+		for (const UnicastDelegateType& DelegateBaseRef : InvocationList)
 		{
 			if (DelegateBaseRef.GetDelegateInstanceProtected())
 			{
@@ -74,7 +76,7 @@ public:
 	{
 		FReadAccessScope ReadScope = GetReadAccessScope();
 
-		for (const TDelegateBase<UserPolicy>& DelegateBaseRef : InvocationList)
+		for (const UnicastDelegateType& DelegateBaseRef : InvocationList)
 		{
 			const IDelegateInstance* DelegateInstance = DelegateBaseRef.GetDelegateInstanceProtected();
 			if ((DelegateInstance != nullptr) && DelegateInstance->HasSameObject(InUserObject))
@@ -100,7 +102,7 @@ public:
 		int32 Result = 0;
 		if (InvocationListLockCount > 0)
 		{
-			for (TDelegateBase<UserPolicy>& DelegateBaseRef : InvocationList)
+			for (UnicastDelegateType& DelegateBaseRef : InvocationList)
 			{
 				IDelegateInstance* DelegateInstance = DelegateBaseRef.GetDelegateInstanceProtected();
 				if ((DelegateInstance != nullptr) && DelegateInstance->HasSameObject(InUserObject))
@@ -111,7 +113,7 @@ public:
 				}
 			}
 
-			// can't compact at the moment, but set out threshold to zero so the next add will do it
+			// can't compact at the moment as the invocation list is locked, but set out threshold to zero so the next add will do it
 			if (Result > 0)
 			{
 				CompactionThreshold = 0;
@@ -122,7 +124,7 @@ public:
 			// compact us while shuffling in later delegates to fill holes
 			for (int32 InvocationListIndex = 0; InvocationListIndex < InvocationList.Num();)
 			{
-				TDelegateBase<UserPolicy>& DelegateBaseRef = InvocationList[InvocationListIndex];
+				UnicastDelegateType& DelegateBaseRef = InvocationList[InvocationListIndex];
 
 				IDelegateInstance* DelegateInstance = DelegateBaseRef.GetDelegateInstanceProtected();
 				if (DelegateInstance == nullptr
@@ -153,7 +155,7 @@ public:
 	{
 		SIZE_T Size = 0;
 		Size += InvocationList.GetAllocatedSize();
-		for (const TDelegateBase<UserPolicy>& DelegateBaseRef : InvocationList)
+		for (const UnicastDelegateType& DelegateBaseRef : InvocationList)
 		{
 			Size += DelegateBaseRef.GetAllocatedSize();
 		}
@@ -178,7 +180,7 @@ protected:
 
 		Clear();
 
-		for (const TDelegateBase<UserPolicy>& OtherDelegateRef : Other.GetInvocationList())
+		for (const UnicastDelegateType& OtherDelegateRef : Other.GetInvocationList())
 		{
 			if (const IDelegateInstance* OtherInstance = GetDelegateInstanceProtectedHelper(OtherDelegateRef))
 			{
@@ -255,7 +257,7 @@ protected:
 
 		for (int32 InvocationListIndex = 0; InvocationListIndex < InvocationList.Num(); ++InvocationListIndex)
 		{
-			TDelegateBase<UserPolicy>& DelegateBase = InvocationList[InvocationListIndex];
+			UnicastDelegateType& DelegateBase = InvocationList[InvocationListIndex];
 
 			IDelegateInstance* DelegateInstance = DelegateBase.GetDelegateInstanceProtected();
 			if (DelegateInstance && DelegateInstance->GetHandle() == Handle)
