Added UE_STATS_MEMORY_PROFILER_ENABLED to toggle on/off code related to MemoryProfiler in STATS system. It still requires STATS to be enabled and "-MemoryProfiler" to be specified in cmd line.
This is off by default as MemoryProfiler in STATS system functionality is deprecated in UE 5.3. For memory profiling, use instead Trace/MemoryInsights and/or LLM.

diff --git a/Engine/Source/Developer/Profiler/Private/ProfilerModule.cpp b/Engine/Source/Developer/Profiler/Private/ProfilerModule.cpp
--- a/Engine/Source/Developer/Profiler/Private/ProfilerModule.cpp
+++ b/Engine/Source/Developer/Profiler/Private/ProfilerModule.cpp
@@ -87,13 +87,14 @@ void FProfilerModule::Shutdown( TSharedRef<SDockTab> TabBeingClosed )
 #endif // STATS
 }
 
+
 /*-----------------------------------------------------------------------------
 	StatsMemoryDumpCommand
 -----------------------------------------------------------------------------*/
 
 void FProfilerModule::StatsMemoryDumpCommand( const TCHAR* Filename )
 {
-#if STATS
+#if STATS && UE_STATS_MEMORY_PROFILER_ENABLED
 	TUniquePtr<FRawStatsMemoryProfiler> Instance( FStatsReader<FRawStatsMemoryProfiler>::Create( Filename ) );
 	if (Instance)
 	{
@@ -157,7 +158,7 @@ void FProfilerModule::StatsMemoryDumpCommand( const TCHAR* Filename )
 #endif // UE_BUILD_DEBUG
 		}
 	}
-#endif // STATS
+#endif // STATS && UE_STATS_MEMORY_PROFILER_ENABLED
 }
 
 
@@ -167,13 +168,14 @@ void FProfilerModule::StatsMemoryDumpCommand( const TCHAR* Filename )
 
 FRawStatsMemoryProfiler* FProfilerModule::OpenRawStatsForMemoryProfiling( const TCHAR* Filename )
 {
-#if STATS
+#if STATS && UE_STATS_MEMORY_PROFILER_ENABLED
 	FRawStatsMemoryProfiler* Instance = FStatsReader<FRawStatsMemoryProfiler>::Create( Filename );
 	if (Instance)
 	{
 		Instance->ReadAndProcessAsynchronously();
 	}
 	return Instance;
-#endif // STATS
+#else
 	return nullptr;
+#endif // STATS && UE_STATS_MEMORY_PROFILER_ENABLED
 }
diff --git a/Engine/Source/Developer/Profiler/Private/ProfilerRawStatsForMemory.cpp b/Engine/Source/Developer/Profiler/Private/ProfilerRawStatsForMemory.cpp
--- a/Engine/Source/Developer/Profiler/Private/ProfilerRawStatsForMemory.cpp
+++ b/Engine/Source/Developer/Profiler/Private/ProfilerRawStatsForMemory.cpp
@@ -2,7 +2,7 @@
 
 #include "ProfilerRawStatsForMemory.h"
 
-#if STATS
+#if STATS && UE_STATS_MEMORY_PROFILER_ENABLED
 
 #include "Stats/StatsMisc.h"
 #include "ProfilingDebugging/DiagnosticTable.h"
@@ -888,4 +888,4 @@ void FRawStatsMemoryProfiler::CompareSnapshotsHumanReadable( const FName BeginSn
 	}
 }
 
-#endif // STATS
+#endif // STATS && UE_STATS_MEMORY_PROFILER_ENABLED
diff --git a/Engine/Source/Developer/Profiler/Private/ProfilerRawStatsForMemory.h b/Engine/Source/Developer/Profiler/Private/ProfilerRawStatsForMemory.h
--- a/Engine/Source/Developer/Profiler/Private/ProfilerRawStatsForMemory.h
+++ b/Engine/Source/Developer/Profiler/Private/ProfilerRawStatsForMemory.h
@@ -4,7 +4,7 @@
 
 #include "CoreMinimal.h"
 
-#if STATS
+#if STATS && UE_STATS_MEMORY_PROFILER_ENABLED
 
 #include "Stats/StatsFile.h"
 
@@ -386,4 +386,4 @@ protected:
 	int32 LastSequenceTagForNamedMarker;
 };
 
-#endif // STATS
+#endif // STATS && UE_STATS_MEMORY_PROFILER_ENABLED
diff --git a/Engine/Source/Developer/Profiler/Public/IProfilerModule.h b/Engine/Source/Developer/Profiler/Public/IProfilerModule.h
--- a/Engine/Source/Developer/Profiler/Public/IProfilerModule.h
+++ b/Engine/Source/Developer/Profiler/Public/IProfilerModule.h
@@ -28,11 +28,13 @@ public:
 	virtual TSharedRef<SWidget> CreateProfilerWindow(const TSharedRef<ISessionManager>& InSessionManager, const TSharedRef<SDockTab>& ConstructUnderMajorTab) = 0;
 
 	/** Implements stats memory dump command. */
+	UE_DEPRECATED(5.3, "Use Trace/MemoryInsights and/or LLM for memory profiling.")
 	virtual void StatsMemoryDumpCommand(const TCHAR* Filename) = 0;
 
 	/** 
 	 * Creates a new instance of the memory profiler based the raw stats. 
 	 * When no longer needed must be stopped via RequestStop() and deleted to avoid memory leaks.
 	 */
+	UE_DEPRECATED(5.3, "Use Trace/MemoryInsights and/or LLM for memory profiling.")
 	virtual FRawStatsMemoryProfiler* OpenRawStatsForMemoryProfiling(const TCHAR* Filename) = 0;
 };
diff --git a/Engine/Source/Runtime/Core/Private/Stats/Stats2.cpp b/Engine/Source/Runtime/Core/Private/Stats/Stats2.cpp
--- a/Engine/Source/Runtime/Core/Private/Stats/Stats2.cpp
+++ b/Engine/Source/Runtime/Core/Private/Stats/Stats2.cpp
@@ -1280,6 +1280,7 @@ void FThreadStats::CheckForCollectingStartupStats()
 		DirectStatsCommand( TEXT( "stat dumpsum -start" ), true );
 	}
 
+#if UE_STATS_MEMORY_PROFILER_ENABLED
 	// Now we can safely enable malloc profiler.
 	if (FStatsMallocProfilerProxy::HasMemoryProfilerToken())
 	{
@@ -1288,6 +1289,7 @@ void FThreadStats::CheckForCollectingStartupStats()
 		FStatsMallocProfilerProxy::Get()->SetState( true );
 		DirectStatsCommand( TEXT( "stat startfileraw" ), true );
 	}
+#endif //UE_STATS_MEMORY_PROFILER_ENABLED
 
 	STAT_ADD_CUSTOMMESSAGE_NAME( STAT_NamedMarker, TEXT("CheckForCollectingStartupStats") );
 }
diff --git a/Engine/Source/Runtime/Core/Private/Stats/StatsCommand.cpp b/Engine/Source/Runtime/Core/Private/Stats/StatsCommand.cpp
--- a/Engine/Source/Runtime/Core/Private/Stats/StatsCommand.cpp
+++ b/Engine/Source/Runtime/Core/Private/Stats/StatsCommand.cpp
@@ -2045,6 +2045,7 @@ static void StatCmd(FString InCmd, bool bStatCommand, FOutputDevice* Ar /*= null
 			FCommandStatsFile::Get().Stop();
 			FThreadStats::DisableRawStats();
 
+#if UE_STATS_MEMORY_PROFILER_ENABLED
 			if (FStatsMallocProfilerProxy::HasMemoryProfilerToken())
 			{
 				if (FStatsMallocProfilerProxy::Get()->GetState())
@@ -2054,6 +2055,7 @@ static void StatCmd(FString InCmd, bool bStatCommand, FOutputDevice* Ar /*= null
 					IStatGroupEnableManager::Get().StatGroupEnableManagerCommand(TEXT("default"));
 				}
 			}
+#endif //UE_STATS_MEMORY_PROFILER_ENABLED
 
 			Stats.ResetStatsForRawStats();
 
diff --git a/Engine/Source/Runtime/Core/Private/Stats/StatsData.cpp b/Engine/Source/Runtime/Core/Private/Stats/StatsData.cpp
--- a/Engine/Source/Runtime/Core/Private/Stats/StatsData.cpp
+++ b/Engine/Source/Runtime/Core/Private/Stats/StatsData.cpp
@@ -845,26 +845,32 @@ void FStatsThreadState::ProcessNonFrameStats(FStatMessagesArray& Data, TSet<FNam
 			{
 				UE_LOG(LogStats, Fatal, TEXT( "Stat %s was not cleared every frame, but was used with a scope cycle counter." ), *Item.NameAndInfo.GetRawName().ToString() );
 			}
+#if UE_STATS_MEMORY_PROFILER_ENABLED
+			else if (Op == EStatOperation::Memory)
+			{
+				// Ignore any memory messages, they shouldn't be treated as regular stats messages.
+			}
+#endif //UE_STATS_MEMORY_PROFILER_ENABLED
+			else if (Op == EStatOperation::SpecialMessageMarker)
+			{
+				// Ignore special messages, they shouldn't be treated as regular stats messages.
+			}
 			else
 			{
-				// Ignore any memory or special messages, they shouldn't be treated as regular stats messages.
-				if( Op != EStatOperation::Memory && Op != EStatOperation::SpecialMessageMarker )
+				FStatMessage* Result = NotClearedEveryFrame.Find(Item.NameAndInfo.GetRawName());
+				if (!Result)
 				{
-					FStatMessage* Result = NotClearedEveryFrame.Find(Item.NameAndInfo.GetRawName());
-					if (!Result)
-					{
-						UE_LOG(LogStats, Error, TEXT( "Stat %s was cleared every frame, but we don't have metadata for it. Data loss." ), *Item.NameAndInfo.GetRawName().ToString() );
-					}
-					else
+					UE_LOG(LogStats, Error, TEXT( "Stat %s was cleared every frame, but we don't have metadata for it. Data loss." ), *Item.NameAndInfo.GetRawName().ToString() );
+				}
+				else
+				{
+					if (NonFrameStatsFound)
 					{
-						if (NonFrameStatsFound)
-						{
-							NonFrameStatsFound->Add(Item.NameAndInfo.GetRawName());
-						}
-						FStatsUtils::AccumulateStat(*Result, Item);
-						Item = *Result; // now just write the accumulated value back into the stream
-						check(Item.NameAndInfo.GetField<EStatOperation>() == EStatOperation::Set);
+						NonFrameStatsFound->Add(Item.NameAndInfo.GetRawName());
 					}
+					FStatsUtils::AccumulateStat(*Result, Item);
+					Item = *Result; // now just write the accumulated value back into the stream
+					check(Item.NameAndInfo.GetField<EStatOperation>() == EStatOperation::Set);
 				}
 			}
 		}
@@ -1514,10 +1520,12 @@ void FStatsThreadState::GetRawStackStats(int64 TargetFrame, FRawStatStackNode& R
 
 					}
 				}
-				else if( Op == EStatOperation::Memory )
+#if UE_STATS_MEMORY_PROFILER_ENABLED
+				else if (Op == EStatOperation::Memory)
 				{
 					// Should never happen.
 				}
+#endif //UE_STATS_MEMORY_PROFILER_ENABLED
 				else if (OutNonStackStats)
 				{
 					FStatsUtils::AddNonStackStats( LongName, Item, Op, ThisFrameNonStackStats );
@@ -2030,9 +2038,11 @@ void FStatsUtils::AccumulateStat(FStatMessage& Dest, FStatMessage const& Item, E
 					StatOpMaxVal_Int64( Dest.NameAndInfo, Dest.GetValue_int64(), Item.GetValue_int64() );
 					break;
 
+#if UE_STATS_MEMORY_PROFILER_ENABLED
 				// Nothing here at this moment.
 				case EStatOperation::Memory:
 					break;
+#endif //UE_STATS_MEMORY_PROFILER_ENABLED
 
 				default:
 					check(0);
@@ -2058,9 +2068,11 @@ void FStatsUtils::AccumulateStat(FStatMessage& Dest, FStatMessage const& Item, E
 					Dest.GetValue_double() = FMath::Max<double>(Dest.GetValue_double(), Item.GetValue_double());
 					break;
 
+#if UE_STATS_MEMORY_PROFILER_ENABLED
 				// Nothing here at this moment.
 				case EStatOperation::Memory:
 					break;
+#endif //UE_STATS_MEMORY_PROFILER_ENABLED
 
 				default:
 					check(0);
diff --git a/Engine/Source/Runtime/Core/Private/Stats/StatsFile.cpp b/Engine/Source/Runtime/Core/Private/Stats/StatsFile.cpp
--- a/Engine/Source/Runtime/Core/Private/Stats/StatsFile.cpp
+++ b/Engine/Source/Runtime/Core/Private/Stats/StatsFile.cpp
@@ -900,121 +900,120 @@ void FStatsReadFile::ProcessStats()
 						const EStatOperation::Type Op = Message.NameAndInfo.GetField<EStatOperation>();
 						const FName RawName = Message.NameAndInfo.GetRawName();
 
-						if (Op == EStatOperation::CycleScopeStart || Op == EStatOperation::CycleScopeEnd || Op == EStatOperation::Memory || Op == EStatOperation::SpecialMessageMarker)
+						if (Op == EStatOperation::CycleScopeStart)
 						{
-							if (Op == EStatOperation::CycleScopeStart)
+							StackState->Stack.Add( RawName );
+							StackState->Current = RawName;
+							ProcessCycleScopeStartOperation( Message, *StackState );
+						}
+#if UE_STATS_MEMORY_PROFILER_ENABLED
+						else if (Op == EStatOperation::Memory)
+						{
+							// First memory operation is Alloc or Free
+							const uint64 EncodedPtr = Message.GetValue_Ptr();
+							const EMemoryOperation MemOp = EMemoryOperation( EncodedPtr & (uint64)EMemoryOperation::Mask );
+							const uint64 Ptr = EncodedPtr & ~(uint64)EMemoryOperation::Mask;
+							if (MemOp == EMemoryOperation::Alloc)
 							{
-								StackState->Stack.Add( RawName );
-								StackState->Current = RawName;
-								ProcessCycleScopeStartOperation( Message, *StackState );
+								// @see FStatsMallocProfilerProxy::TrackAlloc
+								// After AllocPtr message there is always alloc size message and the sequence tag.
+								Index++; CurrentStatMessageIndex++;
+								const FStatMessage& AllocSizeMessage = Data[Index];
+								const int64 AllocSize = AllocSizeMessage.GetValue_int64();
+
+								// Read OperationSequenceTag.
+								Index++; CurrentStatMessageIndex++;
+								const FStatMessage& SequenceTagMessage = Data[Index];
+								const uint32 SequenceTag = (uint32)SequenceTagMessage.GetValue_int64();
+
+								//ThreadStats->AddMemoryMessage( GET_STATFNAME( STAT_Memory_AllocPtr ), (uint64)(UPTRINT)Ptr | (uint64)EMemoryOperation::Alloc );
+								//ThreadStats->AddMemoryMessage( GET_STATFNAME( STAT_Memory_AllocSize ), Size );
+								//ThreadStats->AddMemoryMessage( GET_STATFNAME( STAT_Memory_OperationSequenceTag ), (int64)SequenceTag );
+								ProcessMemoryOperation( MemOp, Ptr, 0, AllocSize, SequenceTag, *StackState );
 							}
-							else if (Op == EStatOperation::Memory)
+							else if (MemOp == EMemoryOperation::Realloc)
 							{
-								// First memory operation is Alloc or Free
-								const uint64 EncodedPtr = Message.GetValue_Ptr();
-								const EMemoryOperation MemOp = EMemoryOperation( EncodedPtr & (uint64)EMemoryOperation::Mask );
-								const uint64 Ptr = EncodedPtr & ~(uint64)EMemoryOperation::Mask;
-								if (MemOp == EMemoryOperation::Alloc)
-								{
-									// @see FStatsMallocProfilerProxy::TrackAlloc
-									// After AllocPtr message there is always alloc size message and the sequence tag.
-									Index++; CurrentStatMessageIndex++;
-									const FStatMessage& AllocSizeMessage = Data[Index];
-									const int64 AllocSize = AllocSizeMessage.GetValue_int64();
-
-									// Read OperationSequenceTag.
-									Index++; CurrentStatMessageIndex++;
-									const FStatMessage& SequenceTagMessage = Data[Index];
-									const uint32 SequenceTag = (uint32)SequenceTagMessage.GetValue_int64();
-
-									//ThreadStats->AddMemoryMessage( GET_STATFNAME( STAT_Memory_AllocPtr ), (uint64)(UPTRINT)Ptr | (uint64)EMemoryOperation::Alloc );
-									//ThreadStats->AddMemoryMessage( GET_STATFNAME( STAT_Memory_AllocSize ), Size );
-									//ThreadStats->AddMemoryMessage( GET_STATFNAME( STAT_Memory_OperationSequenceTag ), (int64)SequenceTag );
-									ProcessMemoryOperation( MemOp, Ptr, 0, AllocSize, SequenceTag, *StackState );
-								}
-								else if (MemOp == EMemoryOperation::Realloc)
-								{
-									const uint64 OldPtr = Ptr;
-
-									// Read NewPtr
-									Index++; CurrentStatMessageIndex++;
-									const FStatMessage& AllocPtrMessage = Data[Index];
-									const uint64 NewPtr = AllocPtrMessage.GetValue_Ptr() & ~(uint64)EMemoryOperation::Mask;
-
-									// After AllocPtr message there is always alloc size message and the sequence tag.
-									Index++; CurrentStatMessageIndex++;
-									const FStatMessage& ReallocSizeMessage = Data[Index];
-									const int64 ReallocSize = ReallocSizeMessage.GetValue_int64();
-
-									// Read OperationSequenceTag.
-									Index++; CurrentStatMessageIndex++;
-									const FStatMessage& SequenceTagMessage = Data[Index];
-									const uint32 SequenceTag = (uint32)SequenceTagMessage.GetValue_int64();
-
-									//ThreadStats->AddMemoryMessage( GET_STATFNAME( STAT_Memory_FreePtr ), (uint64)(UPTRINT)OldPtr | (uint64)EMemoryOperation::Realloc );
-									//ThreadStats->AddMemoryMessage( GET_STATFNAME( STAT_Memory_AllocPtr ), (uint64)(UPTRINT)NewPtr | (uint64)EMemoryOperation::Realloc );
-									//ThreadStats->AddMemoryMessage( GET_STATFNAME( STAT_Memory_AllocSize ), NewSize );
-									//ThreadStats->AddMemoryMessage( GET_STATFNAME( STAT_Memory_OperationSequenceTag ), (int64)SequenceTag );
-									ProcessMemoryOperation( MemOp, OldPtr, NewPtr, ReallocSize, SequenceTag, *StackState );
-								}
-								else if (MemOp == EMemoryOperation::Free)
-								{
-									// Read OperationSequenceTag.
-									Index++; CurrentStatMessageIndex++;
-									const FStatMessage& SequenceTagMessage = Data[Index];
-									const uint32 SequenceTag = (uint32)SequenceTagMessage.GetValue_int64();
-
-									//ThreadStats->AddMemoryMessage( GET_STATFNAME( STAT_Memory_FreePtr ), (uint64)(UPTRINT)Ptr | (uint64)EMemoryOperation::Free );	// 16 bytes total				
-									//ThreadStats->AddMemoryMessage( GET_STATFNAME( STAT_Memory_OperationSequenceTag ), (int64)SequenceTag );
-									ProcessMemoryOperation( MemOp, Ptr, 0, 0, SequenceTag, *StackState );
-								}
-								else
-								{
-									UE_LOG( LogStats, Warning, TEXT( "Pointer from a memory operation is invalid" ) );
-								}
+								const uint64 OldPtr = Ptr;
+
+								// Read NewPtr
+								Index++; CurrentStatMessageIndex++;
+								const FStatMessage& AllocPtrMessage = Data[Index];
+								const uint64 NewPtr = AllocPtrMessage.GetValue_Ptr() & ~(uint64)EMemoryOperation::Mask;
+
+								// After AllocPtr message there is always alloc size message and the sequence tag.
+								Index++; CurrentStatMessageIndex++;
+								const FStatMessage& ReallocSizeMessage = Data[Index];
+								const int64 ReallocSize = ReallocSizeMessage.GetValue_int64();
+
+								// Read OperationSequenceTag.
+								Index++; CurrentStatMessageIndex++;
+								const FStatMessage& SequenceTagMessage = Data[Index];
+								const uint32 SequenceTag = (uint32)SequenceTagMessage.GetValue_int64();
+
+								//ThreadStats->AddMemoryMessage( GET_STATFNAME( STAT_Memory_FreePtr ), (uint64)(UPTRINT)OldPtr | (uint64)EMemoryOperation::Realloc );
+								//ThreadStats->AddMemoryMessage( GET_STATFNAME( STAT_Memory_AllocPtr ), (uint64)(UPTRINT)NewPtr | (uint64)EMemoryOperation::Realloc );
+								//ThreadStats->AddMemoryMessage( GET_STATFNAME( STAT_Memory_AllocSize ), NewSize );
+								//ThreadStats->AddMemoryMessage( GET_STATFNAME( STAT_Memory_OperationSequenceTag ), (int64)SequenceTag );
+								ProcessMemoryOperation( MemOp, OldPtr, NewPtr, ReallocSize, SequenceTag, *StackState );
 							}
-							// Set, Clear, Add, Subtract
-							else if (Op == EStatOperation::CycleScopeEnd)
+							else if (MemOp == EMemoryOperation::Free)
 							{
-								if (StackState->Stack.Num() > 1)
-								{
-									const FName ScopeStart = StackState->Stack.Pop();
-									const FName ScopeEnd = Message.NameAndInfo.GetRawName();
-
-									check( ScopeStart == ScopeEnd );
-
-									StackState->Current = StackState->Stack.Last();
-
-									// The stack should be ok, but it may be partially broken.
-									// This will happen if memory profiling starts in the middle of executing a background thread.
-									StackState->bIsBrokenCallstack = false;
-
-									ProcessCycleScopeEndOperation( Message, *StackState );
-								}
-								else
-								{
-									const FName ShortName = Message.NameAndInfo.GetShortName();
-
-									UE_LOG( LogStats, Warning, TEXT( "Broken cycle scope end %s/%s, current %s" ),
-											*ThreadFName.ToString(),
-											*ShortName.ToString(),
-											*StackState->Current.ToString() );
-
-									// The stack is completely broken, only has the thread name and the last cycle scope.
-									// Rollback to the thread node.
-									StackState->bIsBrokenCallstack = true;
-									StackState->Stack.Empty();
-									StackState->Stack.Add( ThreadFName );
-									StackState->Current = ThreadFName;
-
-									//?ProcessCycleScopeEndOperation( Message, *StackState );
-								}
+								// Read OperationSequenceTag.
+								Index++; CurrentStatMessageIndex++;
+								const FStatMessage& SequenceTagMessage = Data[Index];
+								const uint32 SequenceTag = (uint32)SequenceTagMessage.GetValue_int64();
+
+								//ThreadStats->AddMemoryMessage( GET_STATFNAME( STAT_Memory_FreePtr ), (uint64)(UPTRINT)Ptr | (uint64)EMemoryOperation::Free );	// 16 bytes total				
+								//ThreadStats->AddMemoryMessage( GET_STATFNAME( STAT_Memory_OperationSequenceTag ), (int64)SequenceTag );
+								ProcessMemoryOperation( MemOp, Ptr, 0, 0, SequenceTag, *StackState );
 							}
-							else if (Op == EStatOperation::SpecialMessageMarker)
+							else
 							{
-								ProcessSpecialMessageMarkerOperation( Message, *StackState );
+								UE_LOG( LogStats, Warning, TEXT( "Pointer from a memory operation is invalid" ) );
 							}
 						}
+#endif //UE_STATS_MEMORY_PROFILER_ENABLED
+						// Set, Clear, Add, Subtract
+						else if (Op == EStatOperation::CycleScopeEnd)
+						{
+							if (StackState->Stack.Num() > 1)
+							{
+								const FName ScopeStart = StackState->Stack.Pop();
+								const FName ScopeEnd = Message.NameAndInfo.GetRawName();
+
+								check( ScopeStart == ScopeEnd );
+
+								StackState->Current = StackState->Stack.Last();
+
+								// The stack should be ok, but it may be partially broken.
+								// This will happen if memory profiling starts in the middle of executing a background thread.
+								StackState->bIsBrokenCallstack = false;
+
+								ProcessCycleScopeEndOperation( Message, *StackState );
+							}
+							else
+							{
+								const FName ShortName = Message.NameAndInfo.GetShortName();
+
+								UE_LOG( LogStats, Warning, TEXT( "Broken cycle scope end %s/%s, current %s" ),
+										*ThreadFName.ToString(),
+										*ShortName.ToString(),
+										*StackState->Current.ToString() );
+
+								// The stack is completely broken, only has the thread name and the last cycle scope.
+								// Rollback to the thread node.
+								StackState->bIsBrokenCallstack = true;
+								StackState->Stack.Empty();
+								StackState->Stack.Add( ThreadFName );
+								StackState->Current = ThreadFName;
+
+								//?ProcessCycleScopeEndOperation( Message, *StackState );
+							}
+						}
+						else if (Op == EStatOperation::SpecialMessageMarker)
+						{
+							ProcessSpecialMessageMarkerOperation( Message, *StackState );
+						}
 
 						if (CurrentStatMessageIndex > MessageIndexForStageProgressUpdate)
 						{
diff --git a/Engine/Source/Runtime/Core/Private/Stats/StatsMallocProfilerProxy.cpp b/Engine/Source/Runtime/Core/Private/Stats/StatsMallocProfilerProxy.cpp
--- a/Engine/Source/Runtime/Core/Private/Stats/StatsMallocProfilerProxy.cpp
+++ b/Engine/Source/Runtime/Core/Private/Stats/StatsMallocProfilerProxy.cpp
@@ -1,15 +1,15 @@
 // Copyright Epic Games, Inc. All Rights Reserved.
 
 #include "Stats/StatsMallocProfilerProxy.h"
+
+#if STATS && UE_STATS_MEMORY_PROFILER_ENABLED
+
 #include "UObject/NameTypes.h"
 #include "Logging/LogMacros.h"
 #include "Misc/Parse.h"
 #include "Misc/CommandLine.h"
 #include "Stats/Stats.h"
 
-#if STATS
-
-
 /** Fake stat group and memory stats. */
 DECLARE_STATS_GROUP( TEXT( "Memory Profiler" ), STATGROUP_MemoryProfiler, STATCAT_Advanced );
 
@@ -251,4 +251,4 @@ void FStatsMallocProfilerProxy::UpdateStats()
 	}
 }
 
-#endif //STATS
+#endif //STATS && UE_STATS_MEMORY_PROFILER_ENABLED
diff --git a/Engine/Source/Runtime/Core/Public/Misc/Build.h b/Engine/Source/Runtime/Core/Public/Misc/Build.h
--- a/Engine/Source/Runtime/Core/Public/Misc/Build.h
+++ b/Engine/Source/Runtime/Core/Public/Misc/Build.h
@@ -481,3 +481,10 @@
 #ifndef WITH_PROCESS_PRIORITY_CONTROL
 	#define WITH_PROCESS_PRIORITY_CONTROL 0
 #endif
+
+// Controls whether or not MemoryProfiler is enabled in STATS system.
+// This functionality is deprecated in UE 5.3.
+// For memory profiling, use instead Trace/MemoryInsights and/or LLM.
+#ifndef UE_STATS_MEMORY_PROFILER_ENABLED
+	#define UE_STATS_MEMORY_PROFILER_ENABLED 0
+#endif
diff --git a/Engine/Source/Runtime/Core/Public/Stats/Stats2.h b/Engine/Source/Runtime/Core/Public/Stats/Stats2.h
--- a/Engine/Source/Runtime/Core/Public/Stats/Stats2.h
+++ b/Engine/Source/Runtime/Core/Public/Stats/Stats2.h
@@ -372,7 +372,7 @@ struct EStatOperation
 		MaxVal,
 
 		/** This is a memory operation. @see EMemoryOperation. */
-		Memory,
+		Memory UE_DEPRECATED(5.3, "Use Trace/MemoryInsights and/or LLM for memory profiling."),
 
 		Num,
 		Mask = 0xf,
@@ -430,7 +430,7 @@ struct EMemoryRegion
 };
 
 /** Memory operation for STAT_Memory_AllocPtr. */
-enum class EMemoryOperation : uint8
+enum class UE_DEPRECATED(5.3, "Use Trace/MemoryInsights and/or LLM for memory profiling.") EMemoryOperation : uint8
 {
 	/** Invalid. */
 	Invalid,	
@@ -1484,9 +1484,12 @@ public:
 
 	/** Pseudo-Memory operation. */
 	template<typename TValue>
+	UE_DEPRECATED(5.3, "Use Trace/MemoryInsights and/or LLM for memory profiling.")
 	FORCEINLINE_STATS void AddMemoryMessage( FName InStatName, TValue Value )
 	{
+#if UE_STATS_MEMORY_PROFILER_ENABLED
 		AddStatMessage(FStatMessage(InStatName, EStatOperation::Memory, Value, false));
+#endif //UE_STATS_MEMORY_PROFILER_ENABLED
 	}
 
 	/** 
diff --git a/Engine/Source/Runtime/Core/Public/Stats/StatsFile.h b/Engine/Source/Runtime/Core/Public/Stats/StatsFile.h
--- a/Engine/Source/Runtime/Core/Public/Stats/StatsFile.h
+++ b/Engine/Source/Runtime/Core/Public/Stats/StatsFile.h
@@ -1074,9 +1074,11 @@ protected:
 // 	virtual void ProcessSubtractOperation( const FStatMessage& Message, const FStackState& StackState )
 // 	{}
 
+#if UE_STATS_MEMORY_PROFILER_ENABLED
 	/** Processes memory operation. @see EMemoryOperation. */
 	virtual void ProcessMemoryOperation( EMemoryOperation MemOp, uint64 Ptr, uint64 NewPtr, int64 Size, uint32 SequenceTag, const FStackState& StackState )
 	{}
+#endif
 
 	/** Sets a new processing stage for this file. */
 	void SetProcessingStage( EStatsProcessingStage NewStage )
diff --git a/Engine/Source/Runtime/Core/Public/Stats/StatsMallocProfilerProxy.h b/Engine/Source/Runtime/Core/Public/Stats/StatsMallocProfilerProxy.h
--- a/Engine/Source/Runtime/Core/Public/Stats/StatsMallocProfilerProxy.h
+++ b/Engine/Source/Runtime/Core/Public/Stats/StatsMallocProfilerProxy.h
@@ -7,16 +7,16 @@
 #include "HAL/ThreadSafeCounter.h"
 #include "Misc/Build.h"
 
+#if STATS && UE_STATS_MEMORY_PROFILER_ENABLED
+
 class FOutputDevice;
 class UWorld;
 struct FGenericMemoryStats;
 
-#if STATS
-
 /**
  * Malloc proxy for gathering memory messages.
  */
-class FStatsMallocProfilerProxy : public FMalloc
+class UE_DEPRECATED(5.3, "Use Trace/MemoryInsights and/or LLM for memory profiling.") FStatsMallocProfilerProxy : public FMalloc
 {
 private:
 	/** Malloc we're based on. */
@@ -151,5 +151,4 @@ public:
 	}
 };
 
-
-#endif //STATS
+#endif //STATS && UE_STATS_MEMORY_PROFILER_ENABLED
diff --git a/Engine/Source/Runtime/Launch/Private/LaunchEngineLoop.cpp b/Engine/Source/Runtime/Launch/Private/LaunchEngineLoop.cpp
--- a/Engine/Source/Runtime/Launch/Private/LaunchEngineLoop.cpp
+++ b/Engine/Source/Runtime/Launch/Private/LaunchEngineLoop.cpp
@@ -1984,7 +1984,7 @@ int32 FEngineLoop::PreInitPreStartupScreen(const TCHAR* CmdLine)
 	}
 #endif
 
-#if	STATS
+#if STATS && UE_STATS_MEMORY_PROFILER_ENABLED
 	// Create the stats malloc profiler proxy.
 	if (FStatsMallocProfilerProxy::HasMemoryProfilerToken())
 	{
@@ -1995,7 +1995,7 @@ int32 FEngineLoop::PreInitPreStartupScreen(const TCHAR* CmdLine)
 		// Assumes no concurrency here.
 		GMalloc = FStatsMallocProfilerProxy::Get();
 	}
-#endif // STATS
+#endif // STATS && UE_STATS_MEMORY_PROFILER_ENABLED
 
 #if WITH_APPLICATION_CORE
 	{
