Lots of scoping fixes for the delegate race detector.

diff --git a/Engine/Source/Runtime/Core/Public/Delegates/DelegateBase.h b/Engine/Source/Runtime/Core/Public/Delegates/DelegateBase.h
--- a/Engine/Source/Runtime/Core/Public/Delegates/DelegateBase.h
+++ b/Engine/Source/Runtime/Core/Public/Delegates/DelegateBase.h
@@ -119,12 +119,7 @@ public:
 	{
 		FWriteAccessScope WriteScope = GetWriteAccessScope();
 
-		if (IDelegateInstance* Ptr = GetDelegateInstanceProtected())
-		{
-			Ptr->~IDelegateInstance();
-			DelegateAllocator.ResizeAllocation(0, 0, sizeof(FAlignedInlineDelegateType));
-			DelegateSize = 0;
-		}
+		UnbindUnchecked();
 	}
 
 	/**
@@ -132,6 +127,8 @@ public:
 	 */
 	SIZE_T GetAllocatedSize() const
 	{
+		FReadAccessScope ReadScope = GetReadAccessScope();
+
 		return DelegateAllocator.GetAllocatedSize(DelegateSize, sizeof(FAlignedInlineDelegateType));
 	}
 
@@ -147,6 +144,8 @@ public:
 	 */
 	FName TryGetBoundFunctionName() const
 	{
+		FReadAccessScope ReadScope = GetReadAccessScope();
+
 		const IDelegateInstance* DelegateInstance = GetDelegateInstanceProtected();
 		return DelegateInstance ? DelegateInstance->TryGetBoundFunctionName() : NAME_None;
 	}
@@ -160,6 +159,8 @@ public:
 	 */
 	FORCEINLINE class UObject* GetUObject( ) const
 	{
+		FReadAccessScope ReadScope = GetReadAccessScope();
+
 		const IDelegateInstance* DelegateInstance = GetDelegateInstanceProtected();
 		return DelegateInstance ? DelegateInstance->GetUObject() : nullptr;
 	}
@@ -171,6 +172,8 @@ public:
 	 */
 	FORCEINLINE bool IsBound( ) const
 	{
+		FReadAccessScope ReadScope = GetReadAccessScope();
+
 		const IDelegateInstance* DelegateInstance = GetDelegateInstanceProtected();
 		return DelegateInstance && DelegateInstance->IsSafeToExecute();
 	}
@@ -182,6 +185,8 @@ public:
 	 */
 	FORCEINLINE const void* GetObjectForTimerManager() const
 	{
+		FReadAccessScope ReadScope = GetReadAccessScope();
+
 		const IDelegateInstance* DelegateInstance = GetDelegateInstanceProtected();
 		return DelegateInstance ? DelegateInstance->GetObjectForTimerManager() : nullptr;
 	}
@@ -196,6 +201,8 @@ public:
 	 */
 	uint64 GetBoundProgramCounterForTimerManager() const
 	{
+		FReadAccessScope ReadScope = GetReadAccessScope();
+
 		const IDelegateInstance* DelegateInstance = GetDelegateInstanceProtected();
 		return DelegateInstance ? DelegateInstance->GetBoundProgramCounterForTimerManager() : 0;
 	}
@@ -212,6 +219,8 @@ public:
 			return false;
 		}
 
+		FReadAccessScope ReadScope = GetReadAccessScope();
+
 		const IDelegateInstance* DelegateInstance = GetDelegateInstanceProtected();
 		return DelegateInstance && DelegateInstance->HasSameObject(InUserObject);
 	}
@@ -223,6 +232,8 @@ public:
 	 */
 	FORCEINLINE FDelegateHandle GetHandle() const
 	{
+		FReadAccessScope ReadScope = GetReadAccessScope();
+
 		const IDelegateInstance* DelegateInstance = GetDelegateInstanceProtected();
 		return DelegateInstance ? DelegateInstance->GetHandle() : FDelegateHandle{};
 	}
@@ -233,7 +244,7 @@ public:
 	template<typename DelegateInstanceType, typename... DelegateInstanceParams>
 	void CreateDelegateInstance(DelegateInstanceParams&&... Params)
 	{
-		FReadAccessScope ReadScope = GetReadAccessScope();
+		FWriteAccessScope WriteScope = GetWriteAccessScope();
 
 		IDelegateInstance* DelegateInstance = GetDelegateInstanceProtected();
 		if (DelegateInstance)
@@ -253,23 +264,17 @@ protected:
 	 */
 	FORCEINLINE IDelegateInstance* GetDelegateInstanceProtected()
 	{
-		FReadAccessScope ReadScope = GetReadAccessScope();
-
 		return DelegateSize ? (IDelegateInstance*)DelegateAllocator.GetAllocation() : nullptr;
 	}
 
 	FORCEINLINE const IDelegateInstance* GetDelegateInstanceProtected() const
 	{
-		FReadAccessScope ReadScope = GetReadAccessScope();
-
 		return DelegateSize ? (const IDelegateInstance*)DelegateAllocator.GetAllocation() : nullptr;
 	}
 
 private:
 	void* Allocate(int32 Size)
 	{
-		FWriteAccessScope WriteScope = GetWriteAccessScope();
-
 		int32 NewDelegateSize = FMath::DivideAndRoundUp(Size, (int32)sizeof(FAlignedInlineDelegateType));
 		if (DelegateSize != NewDelegateSize)
 		{
@@ -282,6 +287,8 @@ private:
 
 	void MoveConstruct(TDelegateBase&& Other)
 	{
+		FWriteAccessScope OtherWriteScope = Other.GetWriteAccessScope();
+
 		DelegateAllocator.MoveToEmpty(Other.DelegateAllocator);
 		DelegateSize = Other.DelegateSize;
 		Other.DelegateSize = 0;
@@ -289,15 +296,36 @@ private:
 
 	void MoveAssign(TDelegateBase&& Other)
 	{
-		FWriteAccessScope WriteScope = GetWriteAccessScope();
+		FDelegateAllocatorType::ForElementType<FAlignedInlineDelegateType> OtherDelegateAllocator;
+		int32 OtherDelegateSize;
+		{
+			FWriteAccessScope OtherWriteScope = Other.GetWriteAccessScope();
 
-		Unbind();
-		DelegateAllocator.MoveToEmpty(Other.DelegateAllocator);
-		DelegateSize = Other.DelegateSize;
-		Other.DelegateSize = 0;
+			OtherDelegateAllocator.MoveToEmpty(Other.DelegateAllocator);
+			OtherDelegateSize = Other.DelegateSize;
+			Other.DelegateSize = 0;
+		}
+
+		{
+			FWriteAccessScope WriteScope = GetWriteAccessScope();
+
+			UnbindUnchecked();
+			DelegateAllocator.MoveToEmpty(OtherDelegateAllocator);
+			DelegateSize = OtherDelegateSize;
+		}
 	}
 
 private:
+	FORCEINLINE void UnbindUnchecked()
+	{
+		if (IDelegateInstance* Ptr = GetDelegateInstanceProtected())
+		{
+			Ptr->~IDelegateInstance();
+			DelegateAllocator.ResizeAllocation(0, 0, sizeof(FAlignedInlineDelegateType));
+			DelegateSize = 0;
+		}
+	}
+
 	FDelegateAllocatorType::ForElementType<FAlignedInlineDelegateType> DelegateAllocator;
 	int32 DelegateSize = 0;
 };
diff --git a/Engine/Source/Runtime/Core/Public/Delegates/DelegateSignatureImpl.inl b/Engine/Source/Runtime/Core/Public/Delegates/DelegateSignatureImpl.inl
--- a/Engine/Source/Runtime/Core/Public/Delegates/DelegateSignatureImpl.inl
+++ b/Engine/Source/Runtime/Core/Public/Delegates/DelegateSignatureImpl.inl
@@ -68,6 +68,10 @@ public:
 
 private:
 	// Make sure FDelegateBase's protected functions are not accidentally exposed through the TDelegate API
+	using typename Super::FReadAccessScope;
+	using Super::GetReadAccessScope;
+	using typename Super::FWriteAccessScope;
+	using Super::GetWriteAccessScope;
 	using Super::GetDelegateInstanceProtected;
 
 public:
@@ -306,26 +310,12 @@ public:
 
 	inline TDelegate(const TDelegate& Other)
 	{
-		*this = Other;
+		CopyFrom(Other);
 	}
 
-	inline TDelegate& operator=(const TDelegate& Other)
+	TDelegate& operator=(const TDelegate& Other)
 	{
-		if (&Other != this)
-		{
-			// this down-cast is OK! allows for managing invocation list in the base class without requiring virtual functions
-			const DelegateInstanceInterfaceType* OtherInstance = Other.GetDelegateInstanceProtected();
-
-			if (OtherInstance != nullptr)
-			{
-				OtherInstance->CreateCopy(*this);
-			}
-			else
-			{
-				Unbind();
-			}
-		}
-
+		CopyFrom(Other);
 		return *this;
 	}
 
@@ -523,6 +513,8 @@ public:
 	 */
 	FORCEINLINE RetValType Execute(ParamTypes... Params) const
 	{
+		FReadAccessScope ReadScope = GetReadAccessScope();
+
 		const DelegateInstanceInterfaceType* LocalDelegateInstance = GetDelegateInstanceProtected();
 
 		// If this assert goes off, Execute() was called before a function was bound to the delegate.
@@ -546,6 +538,8 @@ public:
 	>
 	inline bool ExecuteIfBound(ParamTypes... Params) const
 	{
+		FReadAccessScope ReadScope = GetReadAccessScope();
+
 		if (const DelegateInstanceInterfaceType* Ptr = GetDelegateInstanceProtected())
 		{
 			return Ptr->ExecuteIfSafe(Params...);
@@ -566,6 +560,32 @@ protected:
 	{
 		return static_cast<const DelegateInstanceInterfaceType*>(Super::GetDelegateInstanceProtected());
 	}
+
+private:
+	void CopyFrom(const TDelegate& Other)
+	{
+		if ((void*)&Other == (void*)this)
+		{
+			return;
+		}
+
+		// to not hold both delegates locked, make a local copy of `Other` and then move it into this instance
+		TDelegate LocalCopy;
+
+		{
+			FReadAccessScope OtherReadScope = Other.GetReadAccessScope();
+
+			// this down-cast is OK! allows for managing invocation list in the base class without requiring virtual functions
+			const DelegateInstanceInterfaceType* OtherInstance = Other.GetDelegateInstanceProtected();
+
+			if (OtherInstance != nullptr)
+			{
+				OtherInstance->CreateCopy(LocalCopy);
+			}
+		}
+
+		*this = MoveTemp(LocalCopy);
+	}
 };
 
 /**
@@ -613,7 +633,6 @@ private:
 	// Make sure TMulticastDelegateBase's protected functions are not accidentally exposed through the TMulticastDelegate API
 	using Super::AddDelegateInstance;
 	using Super::RemoveDelegateInstance;
-	using Super::GetDelegateInstanceProtectedHelper;
 
 public:
 	TMulticastDelegate() = default;
diff --git a/Engine/Source/Runtime/Core/Public/Delegates/MulticastDelegateBase.h b/Engine/Source/Runtime/Core/Public/Delegates/MulticastDelegateBase.h
--- a/Engine/Source/Runtime/Core/Public/Delegates/MulticastDelegateBase.h
+++ b/Engine/Source/Runtime/Core/Public/Delegates/MulticastDelegateBase.h
@@ -17,6 +17,8 @@
 	typedef TInlineAllocator<NUM_MULTICAST_DELEGATE_INLINE_ENTRIES> FMulticastInvocationListAllocatorType;
 #endif
 
+#define UE_MULTICAST_DELEGATE_DEFAULT_COMPACTION_THRESHOLD 2
+
 /**
  * Abstract base class for multicast delegates.
  */
@@ -35,17 +37,46 @@ protected:
 	using InvocationListType = TArray<UnicastDelegateType, FMulticastInvocationListAllocatorType>;
 
 public:
-	/** Removes all functions from this delegate's invocation list. */
-	void Clear( )
+	TMulticastDelegateBase(TMulticastDelegateBase&& Other)
 	{
-		FWriteAccessScope WriteScope = GetWriteAccessScope();
+		*this = MoveTemp(Other);
+	}
 
-		for (UnicastDelegateType& DelegateBaseRef : InvocationList)
+	TMulticastDelegateBase& operator=(TMulticastDelegateBase&& Other)
+	{
+		if (&Other == this)
 		{
-			DelegateBaseRef.Unbind();
+			return *this;
 		}
 
-		CompactInvocationList(false);
+		InvocationListType LocalInvocationList;
+		int32 LocalCompactionThreshold;
+
+		{
+			FWriteAccessScope OtherWriteScope = Other.GetWriteAccessScope();
+
+			LocalInvocationList = MoveTemp(Other.InvocationList);
+			LocalCompactionThreshold = Other.CompactionThreshold;
+			Other.CompactionThreshold = UE_MULTICAST_DELEGATE_DEFAULT_COMPACTION_THRESHOLD;
+		}
+
+		{
+			FWriteAccessScope ThisWriteScope = GetWriteAccessScope();
+
+			ClearUnchecked();
+			InvocationList = MoveTemp(LocalInvocationList);
+			CompactionThreshold = LocalCompactionThreshold;
+		}
+
+		return *this;
+	}
+
+	/** Removes all functions from this delegate's invocation list. */
+	void Clear( )
+	{
+		FWriteAccessScope WriteScope = GetWriteAccessScope();
+
+		ClearUnchecked();
 	}
 
 	/**
@@ -140,7 +171,7 @@ public:
 				}
 			}
 
-			CompactionThreshold = FMath::Max(2, 2 * InvocationList.Num());
+			CompactionThreshold = FMath::Max(UE_MULTICAST_DELEGATE_DEFAULT_COMPACTION_THRESHOLD, 2 * InvocationList.Num());
 
 			InvocationList.Shrink();
 		}
@@ -153,6 +184,8 @@ public:
 	 */
 	SIZE_T GetAllocatedSize() const
 	{
+		FReadAccessScope ReadScope = GetReadAccessScope();
+	
 		SIZE_T Size = 0;
 		Size += InvocationList.GetAllocatedSize();
 		for (const UnicastDelegateType& DelegateBaseRef : InvocationList)
@@ -166,7 +199,7 @@ protected:
 
 	/** Hidden default constructor. */
 	inline TMulticastDelegateBase( )
-		: CompactionThreshold(2)
+		: CompactionThreshold(UE_MULTICAST_DELEGATE_DEFAULT_COMPACTION_THRESHOLD)
 		, InvocationListLockCount(0)
 	{ }
 
@@ -174,20 +207,27 @@ protected:
 	template<typename DelegateInstanceInterfaceType, typename DelegateType>
 	void CopyFrom(const TMulticastDelegateBase& Other)
 	{
-		FWriteAccessScope WriteScope = GetWriteAccessScope();
-
-		checkSlow(&Other != this);
+		InvocationListType TempInvocationList;
 
-		Clear();
+		{
+			FReadAccessScope OtherReadScope = Other.GetReadAccessScope();
+
+		    for (const UnicastDelegateType& OtherDelegateRef : Other.GetInvocationList())
+		    {
+				if (const IDelegateInstance* OtherInstance = OtherDelegateRef.GetDelegateInstanceProtected())
+			    {
+				    DelegateType TempDelegate;
+				    static_cast<const DelegateInstanceInterfaceType*>(OtherInstance)->CreateCopy(TempDelegate);
+				    TempInvocationList.Add(MoveTemp(TempDelegate));
+			    }
+		    }
+		}
 
-		for (const UnicastDelegateType& OtherDelegateRef : Other.GetInvocationList())
 		{
-			if (const IDelegateInstance* OtherInstance = GetDelegateInstanceProtectedHelper(OtherDelegateRef))
-			{
-				DelegateType TempDelegate;
-				static_cast<const DelegateInstanceInterfaceType*>(OtherInstance)->CreateCopy(TempDelegate);
-				AddDelegateInstance(MoveTemp(TempDelegate));
-			}
+			FWriteAccessScope ThisWriteScope = GetWriteAccessScope();
+
+			ClearUnchecked();
+			InvocationList = MoveTemp(TempInvocationList);
 		}
 	}
 
@@ -209,7 +249,7 @@ protected:
 				// this down-cast is OK! allows for managing invocation list in the base class without requiring virtual functions
 				const DelegateBaseType& DelegateBase = static_cast<const DelegateBaseType&>(LocalInvocationList[InvocationListIndex]);
 
-				const IDelegateInstance* DelegateInstanceInterface = GetDelegateInstanceProtectedHelper(DelegateBase);
+				const IDelegateInstance* DelegateInstanceInterface = DelegateBase.GetDelegateInstanceProtected();
 				if (DelegateInstanceInterface == nullptr || !static_cast<const DelegateInstanceInterfaceType*>(DelegateInstanceInterface)->ExecuteIfSafe(Params...))
 				{
 					NeedsCompaction = true;
@@ -271,16 +311,6 @@ protected:
 		return false;
 	}
 
-protected:
-	/**
-	 * Helper function for derived classes of TMulticastDelegateBase to get at the delegate instance.
-	 */
-	template <typename DelegateType>
-	static FORCEINLINE auto* GetDelegateInstanceProtectedHelper(const DelegateType& Base)
-	{
-		return Base.GetDelegateInstanceProtected();
-	}
-
 private:
 	/**
 	 * Removes any expired or deleted functions from the invocation list.
@@ -320,7 +350,7 @@ private:
 			}
 		}
 
-		CompactionThreshold = FMath::Max(2, 2 * InvocationList.Num());
+		CompactionThreshold = FMath::Max(UE_MULTICAST_DELEGATE_DEFAULT_COMPACTION_THRESHOLD, 2 * InvocationList.Num());
 
 		if (OldNumItems > CompactionThreshold)
 		{
@@ -363,6 +393,15 @@ private:
 	}
 
 private:
+	void ClearUnchecked()
+	{
+		for (UnicastDelegateType& DelegateBaseRef : InvocationList)
+		{
+			DelegateBaseRef.Unbind();
+		}
+
+		CompactInvocationList(false);
+	}
 
 	/** Holds the collection of delegate instances to invoke. */
 	InvocationListType InvocationList;
diff --git a/Engine/Source/Runtime/Core/Public/UObject/ScriptDelegates.h b/Engine/Source/Runtime/Core/Public/UObject/ScriptDelegates.h
--- a/Engine/Source/Runtime/Core/Public/UObject/ScriptDelegates.h
+++ b/Engine/Source/Runtime/Core/Public/UObject/ScriptDelegates.h
@@ -65,6 +65,9 @@ class TScriptDelegate : public FDelegateAccessHandlerBaseChecked
 	template <typename>
 	friend class TScriptDelegate;
 
+	template<typename>
+	friend class TMulticastScriptDelegate;
+
 	using Super = FDelegateAccessHandlerBaseChecked;
 	using typename Super::FReadAccessScope;
 	using Super::GetReadAccessScope;
@@ -80,9 +83,11 @@ public:
 	{ }
 
 	TScriptDelegate(const TScriptDelegate& Other)
-		: Object(Other.Object)
-		, FunctionName(Other.FunctionName)
 	{
+		FReadAccessScope OtherReadScope = Other.GetReadAccessScope();
+
+		Object = Other.Object;
+		FunctionName = Other.FunctionName;
 	}
 
 	template <
@@ -91,29 +96,55 @@ public:
 	>
 	/* UE_DEPRECATED(5.3, "Deprecated - remove after TScriptDelegateTraits<FWeakObjectPtr> is removed") */
 	TScriptDelegate(const TScriptDelegate<OtherDummy>& Other)
-		: Object(Other.Object)
-		, FunctionName(Other.FunctionName)
 	{
+		typename TScriptDelegate<OtherDummy>::FReadAccessScope OtherReadScope = Other.GetReadAccessScope();
+
+		Object = Other.Object;
+		FunctionName = Other.FunctionName;
 	}
 
-	void operator=(const TScriptDelegate& Other)
+	TScriptDelegate& operator=(const TScriptDelegate& Other)
 	{
-		FWriteAccessScope WriteScope = GetWriteAccessScope();
+		WeakPtrType OtherObject;
+		FName OtherFunctionName;
 
-		Object = Other.Object;
-		FunctionName = Other.FunctionName;
+		{
+			FReadAccessScope OtherReadScope = Other.GetReadAccessScope();
+			OtherObject = Other.Object;
+			OtherFunctionName = Other.FunctionName;
+		}
+
+		{
+			FWriteAccessScope WriteScope = GetWriteAccessScope();
+			Object = OtherObject;
+			FunctionName = OtherFunctionName;
+		}
+
+		return *this;
 	}
 	template <
 		typename OtherDummy
 		UE_REQUIRES(UE::Core::Private::BackwardCompatibilityCheck<Dummy, OtherDummy>())
 	>
 	/* UE_DEPRECATED(5.3, "Deprecated - remove after TScriptDelegateTraits<FWeakObjectPtr> is removed") */
-	FORCEINLINE void operator=(const TScriptDelegate<OtherDummy>& Other)
+	TScriptDelegate& operator=(const TScriptDelegate<OtherDummy>& Other)
 	{
-		FWriteAccessScope WriteScope = GetWriteAccessScope();
+		WeakPtrType OtherObject;
+		FName OtherFunctionName;
 
-		Object = Other.Object;
-		FunctionName = Other.FunctionName;
+		{
+			typename TScriptDelegate<OtherDummy>::FReadAccessScope OtherReadScope = Other.GetReadAccessScope();
+			OtherObject = Other.Object;
+			OtherFunctionName = Other.FunctionName;
+		}
+
+		{
+			FWriteAccessScope WriteScope = GetWriteAccessScope();
+			Object = OtherObject;
+			FunctionName = OtherFunctionName;
+		}
+
+		return *this;
 	}
 
 private:
@@ -248,9 +279,23 @@ public:
 	/** Comparison operators */
 	FORCEINLINE bool operator==( const TScriptDelegate& Other ) const
 	{
-		FReadAccessScope ReadScope = GetReadAccessScope();
+		WeakPtrType OtherObject;
+		FName OtherFunctionName;
+
+		{
+			FReadAccessScope OtherReadScope = Other.GetReadAccessScope();
+			OtherObject = Other.Object;
+			OtherFunctionName = Other.FunctionName;
+		}
+
+		bool bResult;
+
+		{
+			FReadAccessScope ThisReadScope = GetReadAccessScope();
+			bResult = Object == OtherObject && FunctionName == OtherFunctionName;
+		}
 
-		return Object == Other.Object && FunctionName == Other.FunctionName;
+		return bResult;
 	}
 	template <
 		typename OtherDummy
@@ -259,16 +304,28 @@ public:
 	/* UE_DEPRECATED(5.3, "Deprecated - remove after TScriptDelegateTraits<FWeakObjectPtr> is removed") */
 	FORCEINLINE bool operator==(const TScriptDelegate<OtherDummy>& Other) const
 	{
-		FReadAccessScope ReadScope = GetReadAccessScope();
+		WeakPtrType OtherObject;
+		FName OtherFunctionName;
+
+		{
+			typename TScriptDelegate<OtherDummy>::FReadAccessScope OtherReadScope = Other.GetReadAccessScope();
+			OtherObject = Other.Object;
+			OtherFunctionName = Other.FunctionName;
+		}
+
+		bool bResult;
+
+		{
+			FReadAccessScope ThisReadScope = GetReadAccessScope();
+			bResult = Object == OtherObject && FunctionName == OtherFunctionName;
+		}
 
-		return Object == Other.Object && FunctionName == Other.FunctionName;
+		return bResult;
 	}
 
 	FORCEINLINE bool operator!=( const TScriptDelegate& Other ) const
 	{
-		FReadAccessScope ReadScope = GetReadAccessScope();
-
-		return Object != Other.Object || FunctionName != Other.FunctionName;
+		return !operator==(Other);
 	}
 	template <
 		typename OtherDummy
@@ -277,9 +334,7 @@ public:
 	/* UE_DEPRECATED(5.3, "Deprecated - remove after TScriptDelegateTraits<FWeakObjectPtr> is removed") */
 	FORCEINLINE bool operator!=(const TScriptDelegate<OtherDummy>& Other) const
 	{
-		FReadAccessScope ReadScope = GetReadAccessScope();
-
-		return Object != Other.Object || FunctionName != Other.FunctionName;
+		return !operator==(Other);
 	}
 
 	/** 
@@ -389,6 +444,18 @@ public:
 		return HashCombine(GetTypeHash(Delegate.Object), GetTypeHash(Delegate.GetFunctionName()));
 	}
 
+	template<typename OtherDummy>
+	static TScriptDelegate CopyFrom(const TScriptDelegate<OtherDummy>& Other)
+	{
+		typename TScriptDelegate<OtherDummy>::FReadAccessScope OtherReadScope = Other.GetReadAccessScope();
+
+		TScriptDelegate Copy;
+		Copy.Object = Other.Object;
+		Copy.FunctionName = Other.FunctionName;
+
+		return Copy;
+	}
+
 protected:
 
 	/** The object bound to this delegate, or nullptr if no object is bound */
@@ -431,6 +498,63 @@ public:
 
 	TMulticastScriptDelegate() = default;
 
+	TMulticastScriptDelegate(const TMulticastScriptDelegate& Other)
+	{
+		InvocationListType LocalCopy;
+
+		{
+			FReadAccessScope OtherReadScope = Other.GetReadAccessScope();
+			LocalCopy = Other.InvocationList;
+		}
+
+		InvocationList = MoveTemp(LocalCopy);
+	}
+
+	TMulticastScriptDelegate& operator=(const TMulticastScriptDelegate& Other)
+	{
+		InvocationListType LocalCopy;
+		{
+			FReadAccessScope OtherReadScope = Other.GetReadAccessScope();
+			LocalCopy = Other.InvocationList;
+		}
+
+		{
+			FWriteAccessScope WriteScope = GetWriteAccessScope();
+			InvocationList = MoveTemp(LocalCopy);
+		}
+
+		return *this;
+	}
+
+	TMulticastScriptDelegate(TMulticastScriptDelegate&& Other)
+	{
+		InvocationListType LocalStorage;
+
+		{
+			FReadAccessScope OtherReadScope = Other.GetReadAccessScope();
+			LocalStorage = MoveTemp(Other.InvocationList);
+		}
+
+		InvocationList = MoveTemp(LocalStorage);
+	}
+
+	TMulticastScriptDelegate& operator=(TMulticastScriptDelegate&& Other)
+	{
+		InvocationListType LocalStorage;
+
+		{
+			FReadAccessScope OtherReadScope = Other.GetReadAccessScope();
+			LocalStorage = MoveTemp(Other.InvocationList);
+		}
+
+		{
+			FWriteAccessScope WriteScope = GetWriteAccessScope();
+			InvocationList = MoveTemp(LocalStorage);
+		}
+
+		return *this;
+	}
+
 public:
 
 	/**
@@ -453,18 +577,34 @@ public:
 	 */
 	bool Contains( const UnicastDelegateType& InDelegate ) const
 	{
-		FReadAccessScope ReadScope = GetReadAccessScope();
+		const UObject* Object;
+		FName FunctionName;
 
-		return InvocationList.Contains( InDelegate );
+		{
+			FReadAccessScope OtherReadScope = InDelegate.GetReadAccessScope();
+			Object = InDelegate.Object.Get();
+			FunctionName = InDelegate.FunctionName;
+		}
+
+		return Contains(Object, FunctionName);
 	}
 	template <
 		typename OtherDummy
 		UE_REQUIRES(UE::Core::Private::BackwardCompatibilityCheck<Dummy, OtherDummy>())
 	>
 	/* UE_DEPRECATED(5.3, "Deprecated - remove after TScriptDelegateTraits<FWeakObjectPtr> is removed") */
-	FORCEINLINE bool Contains(const TScriptDelegate<OtherDummy>& InDelegate) const
+	bool Contains(const TScriptDelegate<OtherDummy>& InDelegate) const
 	{
-		return Contains((UnicastDelegateType)InDelegate);
+		const UObject* Object;
+		FName FunctionName;
+
+		{
+			typename TScriptDelegate<OtherDummy>::FReadAccessScope OtherReadScope = InDelegate.GetReadAccessScope();
+			Object = InDelegate.Object.Get();
+			FunctionName = InDelegate.FunctionName;
+		}
+
+		return Contains(Object, FunctionName);
 	}
 
 	/**
@@ -490,22 +630,36 @@ public:
 	 */
 	void Add( const UnicastDelegateType& InDelegate )
 	{
-		FWriteAccessScope WriteScope = GetWriteAccessScope();
+		UnicastDelegateType LocalCopy = UnicastDelegateType::CopyFrom(InDelegate);
 
-		// First check for any objects that may have expired
-		CompactInvocationList();
+		{
+			FWriteAccessScope WriteScope = GetWriteAccessScope();
+
+			// First check for any objects that may have expired
+			CompactInvocationList();
 
-		// Add the delegate
-		AddInternal( InDelegate );
+			// Add the delegate
+			AddInternal(MoveTemp(LocalCopy));
+		}
 	}
 	template <
 		typename OtherDummy
 		UE_REQUIRES(UE::Core::Private::BackwardCompatibilityCheck<Dummy, OtherDummy>())
 	>
 	/* UE_DEPRECATED(5.3, "Deprecated - remove after TScriptDelegateTraits<FWeakObjectPtr> is removed") */
-	FORCEINLINE void Add(const TScriptDelegate<OtherDummy>& InDelegate)
+	void Add(const TScriptDelegate<OtherDummy>& InDelegate)
 	{
-		Add((UnicastDelegateType)InDelegate);
+		UnicastDelegateType LocalCopy = UnicastDelegateType::CopyFrom(InDelegate);
+
+		{
+			FWriteAccessScope WriteScope = GetWriteAccessScope();
+
+			// First check for any objects that may have expired
+			CompactInvocationList();
+
+			// Add the delegate
+			AddInternal(MoveTemp(LocalCopy));
+		}
 	}
 
 	/**
@@ -516,22 +670,36 @@ public:
 	 */
 	void AddUnique( const UnicastDelegateType& InDelegate )
 	{
-		FWriteAccessScope WriteScope = GetWriteAccessScope();
+		UnicastDelegateType LocalCopy = UnicastDelegateType::CopyFrom(InDelegate);
+
+		{
+			FWriteAccessScope WriteScope = GetWriteAccessScope();
 
-		// Add the delegate, if possible
-		AddUniqueInternal( InDelegate );
+			// Add the delegate, if possible
+			AddUniqueInternal(MoveTemp(LocalCopy));
 
-		// Then check for any objects that may have expired
-		CompactInvocationList();
+			// Then check for any objects that may have expired
+			CompactInvocationList();
+		}
 	}
 	template <
 		typename OtherDummy
 		UE_REQUIRES(UE::Core::Private::BackwardCompatibilityCheck<Dummy, OtherDummy>())
 	>
 	/* UE_DEPRECATED(5.3, "Deprecated - remove after TScriptDelegateTraits<FWeakObjectPtr> is removed") */
-	FORCEINLINE void AddUnique(const TScriptDelegate<OtherDummy>& InDelegate)
+	void AddUnique(const TScriptDelegate<OtherDummy>& InDelegate)
 	{
-		AddUnique((UnicastDelegateType)InDelegate);
+		UnicastDelegateType LocalCopy = UnicastDelegateType::CopyFrom(InDelegate);
+
+		{
+			FWriteAccessScope WriteScope = GetWriteAccessScope();
+
+			// Add the delegate
+			AddUniqueInternal(MoveTemp(LocalCopy));
+
+			// Then check for any objects that may have expired
+			CompactInvocationList();
+		}
 	}
 
 	/**
@@ -542,22 +710,36 @@ public:
 	 */
 	void Remove( const UnicastDelegateType& InDelegate )
 	{
-		FWriteAccessScope WriteScope = GetWriteAccessScope();
+		UnicastDelegateType LocalCopy = UnicastDelegateType::CopyFrom(InDelegate);
 
-		// Remove the delegate
-		RemoveInternal( InDelegate );
+		{
+			FWriteAccessScope WriteScope = GetWriteAccessScope();
 
-		// Check for any delegates that may have expired
-		CompactInvocationList();
+			// Remove the delegate
+			RemoveInternal(LocalCopy);
+
+			// Check for any delegates that may have expired
+			CompactInvocationList();
+		}
 	}
 	template <
 		typename OtherDummy
 		UE_REQUIRES(UE::Core::Private::BackwardCompatibilityCheck<Dummy, OtherDummy>())
 	>
 	/* UE_DEPRECATED(5.3, "Deprecated - remove after TScriptDelegateTraits<FWeakObjectPtr> is removed") */
-	FORCEINLINE void Remove(const TScriptDelegate<OtherDummy>& InDelegate)
+	void Remove(const TScriptDelegate<OtherDummy>& InDelegate)
 	{
-		Remove((UnicastDelegateType)InDelegate);
+		UnicastDelegateType LocalCopy = UnicastDelegateType::CopyFrom(InDelegate);
+
+		{
+			FWriteAccessScope WriteScope = GetWriteAccessScope();
+
+			// Remove the delegate
+			RemoveInternal(LocalCopy);
+
+			// Check for any delegates that may have expired
+			CompactInvocationList();
+		}
 	}
 
 	/**
@@ -765,6 +947,7 @@ public:
 	 */
 	SIZE_T GetAllocatedSize() const
 	{
+		FReadAccessScope ReadScope = GetReadAccessScope();
 		return InvocationList.GetAllocatedSize();
 	}
 
@@ -775,7 +958,7 @@ protected:
 	 *
 	 * @param	InDelegate	Delegate to add
 	*/
-	void AddInternal( const UnicastDelegateType& InDelegate )
+	void AddInternal(UnicastDelegateType&& InDelegate)
 	{
 #if DO_ENSURE
 		// Verify same function isn't already bound
@@ -785,7 +968,7 @@ protected:
 			(void)ensure( InvocationList[ CurFunctionIndex ] != InDelegate );
 		}
 #endif // DO_CHECK
-		InvocationList.Add( InDelegate );
+		InvocationList.Add(MoveTemp(InDelegate));
 	}
 
 	/**
@@ -794,10 +977,10 @@ protected:
 	 *
 	 * @param	InDelegate	Delegate to add
 	 */
-	void AddUniqueInternal( const UnicastDelegateType& InDelegate )
+	void AddUniqueInternal(UnicastDelegateType&& InDelegate)
 	{
 		// Add the item to the invocation list only if it is unique
-		InvocationList.AddUnique( InDelegate );
+		InvocationList.AddUnique(MoveTemp(InDelegate));
 	}
 
 	/**
