I/O - move FIoStatus to separate header

diff --git a/Engine/Source/Runtime/Core/Private/IO/IoStatus.cpp b/Engine/Source/Runtime/Core/Private/IO/IoStatus.cpp
--- a/Engine/Source/Runtime/Core/Private/IO/IoStatus.cpp
+++ b/Engine/Source/Runtime/Core/Private/IO/IoStatus.cpp
@@ -1,16 +1,14 @@
 // Copyright Epic Games, Inc. All Rights Reserved.
 
-#include "Containers/StringFwd.h"
-#include "Containers/StringView.h"
-#include "Containers/UnrealString.h"
+#include "IO/IoStatus.h"
 #include "HAL/Platform.h"
 #include "HAL/PlatformString.h"
 #include "HAL/UnrealMemory.h"
-#include "IO/IoDispatcher.h"
 #include "Logging/LogCategory.h"
 #include "Logging/LogMacros.h"
 #include "Math/UnrealMathUtility.h"
-#include "Trace/Detail/Channel.h"
+
+DEFINE_LOG_CATEGORY_STATIC(LogIoStatus, Log, All);
 
 //////////////////////////////////////////////////////////////////////////
 
@@ -71,7 +69,7 @@ FString FIoStatus::ToString() const
 
 void StatusOrCrash(const FIoStatus& Status)
 {
-	UE_LOG(LogIoDispatcher, Fatal, TEXT("I/O Error '%s'"), *Status.ToString());
+	UE_LOG(LogIoStatus, Fatal, TEXT("I/O Error '%s'"), *Status.ToString());
 }
 
 //////////////////////////////////////////////////////////////////////////
diff --git a/Engine/Source/Runtime/Core/Public/IO/IoDispatcher.h b/Engine/Source/Runtime/Core/Public/IO/IoDispatcher.h
--- a/Engine/Source/Runtime/Core/Public/IO/IoDispatcher.h
+++ b/Engine/Source/Runtime/Core/Public/IO/IoDispatcher.h
@@ -21,6 +21,7 @@
 #include "IO/IoContainerId.h"
 #include "IO/IoDispatcherPriority.h"
 #include "IO/IoHash.h"
+#include "IO/IoStatus.h"
 #include "Logging/LogMacros.h"
 #include "Math/NumericLimits.h"
 #include "Memory/MemoryFwd.h"
@@ -66,363 +67,6 @@ template <typename CharType> class TStringBuilderBase;
 
 CORE_API DECLARE_LOG_CATEGORY_EXTERN(LogIoDispatcher, Log, All);
 
-/*
- * I/O error code.
- */
-enum class EIoErrorCode
-{
-	Ok,
-	Unknown,
-	InvalidCode,
-	Cancelled,
-	FileOpenFailed,
-	FileNotOpen,
-	ReadError,
-	WriteError,
-	NotFound,
-	CorruptToc,
-	UnknownChunkID,
-	InvalidParameter,
-	SignatureError,
-	InvalidEncryptionKey,
-	CompressionError,
-};
-
-/*
- * Get I/O error code description.
- */
-inline const TCHAR* GetIoErrorText(EIoErrorCode ErrorCode)
-{
-	extern CORE_API const TCHAR* const* GetIoErrorText_ErrorCodeText;
-
-	return GetIoErrorText_ErrorCodeText[static_cast<uint32>(ErrorCode)];
-}
-
-/**
- * I/O status with error code and message.
- */
-class FIoStatus
-{
-public:
-	CORE_API			FIoStatus();
-	CORE_API			~FIoStatus();
-
-	CORE_API			FIoStatus(EIoErrorCode Code, const FStringView& InErrorMessage);
-	CORE_API			FIoStatus(EIoErrorCode Code);
-	CORE_API FIoStatus&	operator=(const FIoStatus& Other);
-	CORE_API FIoStatus&	operator=(const EIoErrorCode InErrorCode);
-
-	CORE_API bool		operator==(const FIoStatus& Other) const;
-			 bool		operator!=(const FIoStatus& Other) const { return !operator==(Other); }
-
-	inline bool			IsOk() const { return ErrorCode == EIoErrorCode::Ok; }
-	inline bool			IsCompleted() const { return ErrorCode != EIoErrorCode::Unknown; }
-	inline EIoErrorCode	GetErrorCode() const { return ErrorCode; }
-	CORE_API FString	ToString() const;
-
-	CORE_API static const FIoStatus Ok;
-	CORE_API static const FIoStatus Unknown;
-	CORE_API static const FIoStatus Invalid;
-
-private:
-	static constexpr int32 MaxErrorMessageLength = 128;
-	using FErrorMessage = TCHAR[MaxErrorMessageLength];
-
-	EIoErrorCode	ErrorCode = EIoErrorCode::Ok;
-	FErrorMessage	ErrorMessage;
-
-	friend class FIoStatusBuilder;
-};
-
-/**
- * Helper to make it easier to generate meaningful error messages.
- */
-class FIoStatusBuilder
-{
-	EIoErrorCode		StatusCode;
-	FString				Message;
-public:
-	CORE_API explicit	FIoStatusBuilder(EIoErrorCode StatusCode);
-	CORE_API			FIoStatusBuilder(const FIoStatus& InStatus, FStringView String);
-	CORE_API			~FIoStatusBuilder();
-
-	CORE_API			operator FIoStatus();
-
-	CORE_API FIoStatusBuilder& operator<<(FStringView String);
-};
-
-CORE_API FIoStatusBuilder operator<<(const FIoStatus& Status, FStringView String);
-
-/**
- * Optional I/O result or error status.
- */
-template<typename T>
-class TIoStatusOr
-{
-	template<typename U> friend class TIoStatusOr;
-
-public:
-	TIoStatusOr() : StatusValue(FIoStatus::Unknown) { }
-	TIoStatusOr(const TIoStatusOr& Other);
-	TIoStatusOr(TIoStatusOr&& Other);
-
-	TIoStatusOr(FIoStatus InStatus);
-	TIoStatusOr(const T& InValue);
-	TIoStatusOr(T&& InValue);
-
-	~TIoStatusOr();
-
-	template <typename... ArgTypes>
-	explicit TIoStatusOr(ArgTypes&&... Args);
-
-	template<typename U>
-	TIoStatusOr(const TIoStatusOr<U>& Other);
-
-	TIoStatusOr<T>& operator=(const TIoStatusOr<T>& Other);
-	TIoStatusOr<T>& operator=(TIoStatusOr<T>&& Other);
-	TIoStatusOr<T>& operator=(const FIoStatus& OtherStatus);
-	TIoStatusOr<T>& operator=(const T& OtherValue);
-	TIoStatusOr<T>& operator=(T&& OtherValue);
-
-	template<typename U>
-	TIoStatusOr<T>& operator=(const TIoStatusOr<U>& Other);
-
-	const FIoStatus&	Status() const;
-	bool				IsOk() const;
-
-	const T&			ValueOrDie();
-	T					ConsumeValueOrDie();
-
-	void				Reset();
-
-private:
-	FIoStatus				StatusValue;
-	TTypeCompatibleBytes<T>	Value;
-};
-
-CORE_API void StatusOrCrash(const FIoStatus& Status);
-
-template<typename T>
-void TIoStatusOr<T>::Reset()
-{
-	EIoErrorCode ErrorCode = StatusValue.GetErrorCode();
-	StatusValue = EIoErrorCode::Unknown;
-
-	if (ErrorCode == EIoErrorCode::Ok)
-	{
-		((T*)&Value)->~T();
-	}
-}
-
-template<typename T>
-const T& TIoStatusOr<T>::ValueOrDie()
-{
-	if (!StatusValue.IsOk())
-	{
-		StatusOrCrash(StatusValue);
-	}
-
-	return *Value.GetTypedPtr();
-}
-
-template<typename T>
-T TIoStatusOr<T>::ConsumeValueOrDie()
-{
-	if (!StatusValue.IsOk())
-	{
-		StatusOrCrash(StatusValue);
-	}
-
-	StatusValue = FIoStatus::Unknown;
-
-	return MoveTemp(*Value.GetTypedPtr());
-}
-
-template<typename T>
-TIoStatusOr<T>::TIoStatusOr(const TIoStatusOr& Other)
-{
-	StatusValue = Other.StatusValue;
-	if (StatusValue.IsOk())
-	{
-		new(&Value) T(*(const T*)&Other.Value);
-	}
-}
-
-template<typename T>
-TIoStatusOr<T>::TIoStatusOr(TIoStatusOr&& Other)
-{
-	StatusValue = Other.StatusValue;
-	if (StatusValue.IsOk())
-	{
-		new(&Value) T(MoveTempIfPossible(*(T*)&Other.Value));
-		Other.StatusValue = EIoErrorCode::Unknown;
-	}
-}
-
-template<typename T>
-TIoStatusOr<T>::TIoStatusOr(FIoStatus InStatus)
-{
-	check(!InStatus.IsOk());
-	StatusValue = InStatus;
-}
-
-template<typename T>
-TIoStatusOr<T>::TIoStatusOr(const T& InValue)
-{
-	StatusValue = FIoStatus::Ok;
-	new(&Value) T(InValue);
-}
-
-template<typename T>
-TIoStatusOr<T>::TIoStatusOr(T&& InValue)
-{
-	StatusValue = FIoStatus::Ok;
-	new(&Value) T(MoveTempIfPossible(InValue));
-}
-
-template <typename T>
-template <typename... ArgTypes>
-TIoStatusOr<T>::TIoStatusOr(ArgTypes&&... Args)
-{
-	StatusValue = FIoStatus::Ok;
-	new(&Value) T(Forward<ArgTypes>(Args)...);
-}
-
-template<typename T>
-TIoStatusOr<T>::~TIoStatusOr()
-{
-	Reset();
-}
-
-template<typename T>
-bool TIoStatusOr<T>::IsOk() const
-{
-	return StatusValue.IsOk();
-}
-
-template<typename T>
-const FIoStatus& TIoStatusOr<T>::Status() const
-{
-	return StatusValue;
-}
-
-template<typename T>
-TIoStatusOr<T>&
-TIoStatusOr<T>::operator=(const TIoStatusOr<T>& Other)
-{
-	if (&Other != this)
-	{
-		Reset();
-
-		if (Other.StatusValue.IsOk())
-		{
-			new(&Value) T(*(const T*)&Other.Value);
-			StatusValue = EIoErrorCode::Ok;
-		}
-		else
-		{
-			StatusValue = Other.StatusValue;
-		}
-	}
-
-	return *this;
-}
-
-template<typename T>
-TIoStatusOr<T>&
-TIoStatusOr<T>::operator=(TIoStatusOr<T>&& Other)
-{
-	if (&Other != this)
-	{
-		Reset();
- 
-		if (Other.StatusValue.IsOk())
-		{
-			new(&Value) T(MoveTempIfPossible(*(T*)&Other.Value));
-			Other.StatusValue = EIoErrorCode::Unknown;
-			StatusValue = EIoErrorCode::Ok;
-		}
-		else
-		{
-			StatusValue = Other.StatusValue;
-		}
-	}
-
-	return *this;
-}
-
-template<typename T>
-TIoStatusOr<T>&
-TIoStatusOr<T>::operator=(const FIoStatus& OtherStatus)
-{
-	check(!OtherStatus.IsOk());
-
-	Reset();
-	StatusValue = OtherStatus;
-
-	return *this;
-}
-
-template<typename T>
-TIoStatusOr<T>&
-TIoStatusOr<T>::operator=(const T& OtherValue)
-{
-	if (&OtherValue != (T*)&Value)
-	{
-		Reset();
-		
-		new(&Value) T(OtherValue);
-		StatusValue = EIoErrorCode::Ok;
-	}
-
-	return *this;
-}
-
-template<typename T>
-TIoStatusOr<T>&
-TIoStatusOr<T>::operator=(T&& OtherValue)
-{
-	if (&OtherValue != (T*)&Value)
-	{
-		Reset();
-		
-		new(&Value) T(MoveTempIfPossible(OtherValue));
-		StatusValue = EIoErrorCode::Ok;
-	}
-
-	return *this;
-}
-
-template<typename T>
-template<typename U>
-TIoStatusOr<T>::TIoStatusOr(const TIoStatusOr<U>& Other)
-:	StatusValue(Other.StatusValue)
-{
-	if (StatusValue.IsOk())
-	{
-		new(&Value) T(*(const U*)&Other.Value);
-	}
-}
-
-template<typename T>
-template<typename U>
-TIoStatusOr<T>& TIoStatusOr<T>::operator=(const TIoStatusOr<U>& Other)
-{
-	Reset();
-
-	if (Other.StatusValue.IsOk())
-	{
-		new(&Value) T(*(const U*)&Other.Value);
-		StatusValue = EIoErrorCode::Ok;
-	}
-	else
-	{
-		StatusValue = Other.StatusValue;
-	}
-
-	return *this;
-}
-
 //////////////////////////////////////////////////////////////////////////
 
 /** Helper used to manage creation of I/O store file handles etc
diff --git a/Engine/Source/Runtime/Core/Public/IO/IoStatus.h b/Engine/Source/Runtime/Core/Public/IO/IoStatus.h
new file mode 100644
--- /dev/null
+++ b/Engine/Source/Runtime/Core/Public/IO/IoStatus.h
@@ -0,0 +1,369 @@
+// Copyright Epic Games, Inc. All Rights Reserved.
+
+#pragma once
+
+#include "Containers/StringView.h"
+#include "Containers/UnrealString.h"
+
+template <typename CharType> class TStringBuilderBase;
+
+///////////////////////////////////////////////////////////////////////////////
+
+/**
+ * I/O error code.
+ */
+enum class EIoErrorCode
+{
+	Ok,
+	Unknown,
+	InvalidCode,
+	Cancelled,
+	FileOpenFailed,
+	FileNotOpen,
+	ReadError,
+	WriteError,
+	NotFound,
+	CorruptToc,
+	UnknownChunkID,
+	InvalidParameter,
+	SignatureError,
+	InvalidEncryptionKey,
+	CompressionError,
+};
+
+/** Get I/O error code description. */
+inline const TCHAR* GetIoErrorText(EIoErrorCode ErrorCode)
+{
+	extern CORE_API const TCHAR* const* GetIoErrorText_ErrorCodeText;
+
+	return GetIoErrorText_ErrorCodeText[static_cast<uint32>(ErrorCode)];
+}
+
+///////////////////////////////////////////////////////////////////////////////
+
+/**
+ * I/O status with error code and message.
+ */
+class FIoStatus
+{
+public:
+	CORE_API			FIoStatus();
+	CORE_API			~FIoStatus();
+
+	CORE_API			FIoStatus(EIoErrorCode Code, const FStringView& InErrorMessage);
+	CORE_API			FIoStatus(EIoErrorCode Code);
+	CORE_API FIoStatus&	operator=(const FIoStatus& Other);
+	CORE_API FIoStatus&	operator=(const EIoErrorCode InErrorCode);
+
+	CORE_API bool		operator==(const FIoStatus& Other) const;
+			 bool		operator!=(const FIoStatus& Other) const { return !operator==(Other); }
+
+	inline bool			IsOk() const { return ErrorCode == EIoErrorCode::Ok; }
+	inline bool			IsCompleted() const { return ErrorCode != EIoErrorCode::Unknown; }
+	inline EIoErrorCode	GetErrorCode() const { return ErrorCode; }
+	CORE_API FString	ToString() const;
+
+	CORE_API static const FIoStatus Ok;
+	CORE_API static const FIoStatus Unknown;
+	CORE_API static const FIoStatus Invalid;
+
+private:
+	static constexpr int32 MaxErrorMessageLength = 128;
+	using FErrorMessage = TCHAR[MaxErrorMessageLength];
+
+	EIoErrorCode	ErrorCode = EIoErrorCode::Ok;
+	FErrorMessage	ErrorMessage;
+
+	friend class FIoStatusBuilder;
+};
+
+///////////////////////////////////////////////////////////////////////////////
+
+/**
+ * Helper to make it easier to generate meaningful error messages.
+ */
+class FIoStatusBuilder
+{
+	EIoErrorCode		StatusCode;
+	FString				Message;
+public:
+	CORE_API explicit	FIoStatusBuilder(EIoErrorCode StatusCode);
+	CORE_API			FIoStatusBuilder(const FIoStatus& InStatus, FStringView String);
+	CORE_API			~FIoStatusBuilder();
+
+	CORE_API			operator FIoStatus();
+
+	CORE_API FIoStatusBuilder& operator<<(FStringView String);
+};
+
+CORE_API FIoStatusBuilder operator<<(const FIoStatus& Status, FStringView String);
+
+/**
+ * Optional I/O result or error status.
+ */
+template<typename T>
+class TIoStatusOr
+{
+	template<typename U> friend class TIoStatusOr;
+
+public:
+	TIoStatusOr() : StatusValue(FIoStatus::Unknown) { }
+	TIoStatusOr(const TIoStatusOr& Other);
+	TIoStatusOr(TIoStatusOr&& Other);
+
+	TIoStatusOr(FIoStatus InStatus);
+	TIoStatusOr(const T& InValue);
+	TIoStatusOr(T&& InValue);
+
+	~TIoStatusOr();
+
+	template <typename... ArgTypes>
+	explicit TIoStatusOr(ArgTypes&&... Args);
+
+	template<typename U>
+	TIoStatusOr(const TIoStatusOr<U>& Other);
+
+	TIoStatusOr<T>& operator=(const TIoStatusOr<T>& Other);
+	TIoStatusOr<T>& operator=(TIoStatusOr<T>&& Other);
+	TIoStatusOr<T>& operator=(const FIoStatus& OtherStatus);
+	TIoStatusOr<T>& operator=(const T& OtherValue);
+	TIoStatusOr<T>& operator=(T&& OtherValue);
+
+	template<typename U>
+	TIoStatusOr<T>& operator=(const TIoStatusOr<U>& Other);
+
+	const FIoStatus&	Status() const;
+	bool				IsOk() const;
+
+	const T&			ValueOrDie();
+	T					ConsumeValueOrDie();
+
+	void				Reset();
+
+private:
+	FIoStatus				StatusValue;
+	TTypeCompatibleBytes<T>	Value;
+};
+
+CORE_API void StatusOrCrash(const FIoStatus& Status);
+
+template<typename T>
+void TIoStatusOr<T>::Reset()
+{
+	EIoErrorCode ErrorCode = StatusValue.GetErrorCode();
+	StatusValue = EIoErrorCode::Unknown;
+
+	if (ErrorCode == EIoErrorCode::Ok)
+	{
+		((T*)&Value)->~T();
+	}
+}
+
+template<typename T>
+const T& TIoStatusOr<T>::ValueOrDie()
+{
+	if (!StatusValue.IsOk())
+	{
+		StatusOrCrash(StatusValue);
+	}
+
+	return *Value.GetTypedPtr();
+}
+
+template<typename T>
+T TIoStatusOr<T>::ConsumeValueOrDie()
+{
+	if (!StatusValue.IsOk())
+	{
+		StatusOrCrash(StatusValue);
+	}
+
+	StatusValue = FIoStatus::Unknown;
+
+	return MoveTemp(*Value.GetTypedPtr());
+}
+
+template<typename T>
+TIoStatusOr<T>::TIoStatusOr(const TIoStatusOr& Other)
+{
+	StatusValue = Other.StatusValue;
+	if (StatusValue.IsOk())
+	{
+		new(&Value) T(*(const T*)&Other.Value);
+	}
+}
+
+template<typename T>
+TIoStatusOr<T>::TIoStatusOr(TIoStatusOr&& Other)
+{
+	StatusValue = Other.StatusValue;
+	if (StatusValue.IsOk())
+	{
+		new(&Value) T(MoveTempIfPossible(*(T*)&Other.Value));
+		Other.StatusValue = EIoErrorCode::Unknown;
+	}
+}
+
+template<typename T>
+TIoStatusOr<T>::TIoStatusOr(FIoStatus InStatus)
+{
+	check(!InStatus.IsOk());
+	StatusValue = InStatus;
+}
+
+template<typename T>
+TIoStatusOr<T>::TIoStatusOr(const T& InValue)
+{
+	StatusValue = FIoStatus::Ok;
+	new(&Value) T(InValue);
+}
+
+template<typename T>
+TIoStatusOr<T>::TIoStatusOr(T&& InValue)
+{
+	StatusValue = FIoStatus::Ok;
+	new(&Value) T(MoveTempIfPossible(InValue));
+}
+
+template <typename T>
+template <typename... ArgTypes>
+TIoStatusOr<T>::TIoStatusOr(ArgTypes&&... Args)
+{
+	StatusValue = FIoStatus::Ok;
+	new(&Value) T(Forward<ArgTypes>(Args)...);
+}
+
+template<typename T>
+TIoStatusOr<T>::~TIoStatusOr()
+{
+	Reset();
+}
+
+template<typename T>
+bool TIoStatusOr<T>::IsOk() const
+{
+	return StatusValue.IsOk();
+}
+
+template<typename T>
+const FIoStatus& TIoStatusOr<T>::Status() const
+{
+	return StatusValue;
+}
+
+template<typename T>
+TIoStatusOr<T>&
+TIoStatusOr<T>::operator=(const TIoStatusOr<T>& Other)
+{
+	if (&Other != this)
+	{
+		Reset();
+
+		if (Other.StatusValue.IsOk())
+		{
+			new(&Value) T(*(const T*)&Other.Value);
+			StatusValue = EIoErrorCode::Ok;
+		}
+		else
+		{
+			StatusValue = Other.StatusValue;
+		}
+	}
+
+	return *this;
+}
+
+template<typename T>
+TIoStatusOr<T>&
+TIoStatusOr<T>::operator=(TIoStatusOr<T>&& Other)
+{
+	if (&Other != this)
+	{
+		Reset();
+ 
+		if (Other.StatusValue.IsOk())
+		{
+			new(&Value) T(MoveTempIfPossible(*(T*)&Other.Value));
+			Other.StatusValue = EIoErrorCode::Unknown;
+			StatusValue = EIoErrorCode::Ok;
+		}
+		else
+		{
+			StatusValue = Other.StatusValue;
+		}
+	}
+
+	return *this;
+}
+
+template<typename T>
+TIoStatusOr<T>&
+TIoStatusOr<T>::operator=(const FIoStatus& OtherStatus)
+{
+	check(!OtherStatus.IsOk());
+
+	Reset();
+	StatusValue = OtherStatus;
+
+	return *this;
+}
+
+template<typename T>
+TIoStatusOr<T>&
+TIoStatusOr<T>::operator=(const T& OtherValue)
+{
+	if (&OtherValue != (T*)&Value)
+	{
+		Reset();
+		
+		new(&Value) T(OtherValue);
+		StatusValue = EIoErrorCode::Ok;
+	}
+
+	return *this;
+}
+
+template<typename T>
+TIoStatusOr<T>&
+TIoStatusOr<T>::operator=(T&& OtherValue)
+{
+	if (&OtherValue != (T*)&Value)
+	{
+		Reset();
+		
+		new(&Value) T(MoveTempIfPossible(OtherValue));
+		StatusValue = EIoErrorCode::Ok;
+	}
+
+	return *this;
+}
+
+template<typename T>
+template<typename U>
+TIoStatusOr<T>::TIoStatusOr(const TIoStatusOr<U>& Other)
+:	StatusValue(Other.StatusValue)
+{
+	if (StatusValue.IsOk())
+	{
+		new(&Value) T(*(const U*)&Other.Value);
+	}
+}
+
+template<typename T>
+template<typename U>
+TIoStatusOr<T>& TIoStatusOr<T>::operator=(const TIoStatusOr<U>& Other)
+{
+	Reset();
+
+	if (Other.StatusValue.IsOk())
+	{
+		new(&Value) T(*(const U*)&Other.Value);
+		StatusValue = EIoErrorCode::Ok;
+	}
+	else
+	{
+		StatusValue = Other.StatusValue;
+	}
+
+	return *this;
+}
