task retraction refactor:
the main thing is removal of spin-waitings from task retraction and waiting for task completion. other changes include:
* add spinning for a short while in Wait() before getting blocked
* add timeout support for task retraction
* consolidates waiting w/ and w/o timeout into the same function with a default value for timeout.
* only the API function takes timeout value as `FTimespan`, all internal functions take `TTimeout` that allows to avoid recalculating the value before passing it around

diff --git a/Engine/Source/Runtime/Core/Private/Tasks/Pipe.cpp b/Engine/Source/Runtime/Core/Private/Tasks/Pipe.cpp
--- a/Engine/Source/Runtime/Core/Private/Tasks/Pipe.cpp
+++ b/Engine/Source/Runtime/Core/Private/Tasks/Pipe.cpp
@@ -43,7 +43,7 @@ namespace UE::Tasks
 		return false;
 	}
 
-	bool FPipe::WaitUntilEmpty(FTimespan Timeout/* = FTimespan::FromMilliseconds(MAX_uint32)*/)
+	bool FPipe::WaitUntilEmpty(FTimespan Timeout/* = FTimespan::MaxValue()*/)
 	{
 		struct FPlaceholderTask final : Private::FTaskBase
 		{
@@ -73,7 +73,7 @@ namespace UE::Tasks
 		{
 			LastTask_Local->Release(); // release pipe's ref to the last task
 			// when the placeholder replaced "the last task" in the pipe, it became its subsequent, so the last task will "complete" it
-			bRes = PlaceholderTask->Wait(Timeout);
+			bRes = PlaceholderTask->Wait(FTimeout{ Timeout });
 		}
 		else
 		{
diff --git a/Engine/Source/Runtime/Core/Private/Tasks/TaskPrivate.cpp b/Engine/Source/Runtime/Core/Private/Tasks/TaskPrivate.cpp
--- a/Engine/Source/Runtime/Core/Private/Tasks/TaskPrivate.cpp
+++ b/Engine/Source/Runtime/Core/Private/Tasks/TaskPrivate.cpp
@@ -68,13 +68,19 @@ namespace UE::Tasks
 			}
 		};
 
-		bool FTaskBase::TryRetractAndExecute(uint32 RecursionDepth/* = 0*/)
+		bool FTaskBase::TryRetractAndExecute(FTimeout Timeout, uint32 RecursionDepth/* = 0*/)
 		{
 			TRACE_CPUPROFILER_EVENT_SCOPE(TaskRetraction);
 
-			if (IsCompleted())
+			if (!IsAwaitable())
 			{
-				return true;
+				UE_LOG(LogTemp, Fatal, TEXT("Deadlock detected! A task can't be waited here, e.g. because it's being executed by the currect thread"));
+				return false;
+			}
+
+			if (IsCompleted() || Timeout)
+			{
+				return IsCompleted();
 			}
 
 #if TASKGRAPH_NEW_FRONTEND
@@ -108,7 +114,7 @@ namespace UE::Tasks
 				while (FTaskBase* Prerequisite = Prerequisites.Pop())
 				{
 					// ignore if retraction failed, as this thread still can try to help with other prerequisites instead of being blocked in waiting
-					Prerequisite->TryRetractAndExecute(RecursionDepth);
+					Prerequisite->TryRetractAndExecute(Timeout, RecursionDepth);
 					Prerequisite->Release();
 				}
 			}
@@ -131,6 +137,11 @@ namespace UE::Tasks
 					return true;
 				}
 
+				if (Timeout)
+				{
+					return IsCompleted();
+				}
+
 				if (!TryExecuteTask())
 				{
 					return false; // still locked by prerequisites, or another thread managed to set execution flag first, or we're inside this task execution
@@ -143,11 +154,6 @@ namespace UE::Tasks
 
 			// the task was launched so the scheduler will handle the internal reference held by low-level task
 
-			if (IsCompleted()) // still can be hold back by nested tasks, this is an optional early out for better perf
-			{
-				return true;
-			}
-
 			// retract nested tasks, if any
 			{
 				// keep trying retracting all nested tasks even if some of them fail, so the current worker can contribute instead of being blocked
@@ -156,7 +162,7 @@ namespace UE::Tasks
 				// this can be potentially improved by using a different container for prerequisites
 				while (FTaskBase* Prerequisite = Prerequisites.Pop())
 				{
-					if (!Prerequisite->TryRetractAndExecute(RecursionDepth))
+					if (!Prerequisite->TryRetractAndExecute(Timeout, RecursionDepth))
 					{
 						bSucceeded = false;
 					}
@@ -169,83 +175,40 @@ namespace UE::Tasks
 				}
 			}
 
-			// it happens that all nested tasks are completed and are in the process of completing the parent (this task) concurrently, 
-			// but the flag is not set yet. wait for it to maintain postconditions
-			while (!IsCompleted())
-			{
-				FPlatformProcess::Yield();
-			}
-
+			// at this point the task is executed and has no pending nested tasks, but still can be "not completed" (nested tasks can be 
+			// in the process of completing it (setting the flag) concurrently), so the caller still has to wait for completion
 			return true;
 		}
 
-		void FTaskBase::Wait()
+		bool FTaskBase::Wait(FTimeout Timeout)
 		{
-			if (IsCompleted())
+			if (IsCompleted() || Timeout)
 			{
-				return;
+				return IsCompleted();
 			}
 
 			TaskTrace::FWaitingScope WaitingScope(GetTraceId());
 			TRACE_CPUPROFILER_EVENT_SCOPE(Tasks::Wait);
 
-			if (!IsAwaitable())
-			{
-				UE_LOG(LogTemp, Fatal, TEXT("Deadlock detected! A task can't be waited here, e.g. because it's being executed by the currect thread"));
-				return;
-			}
-
-			if (TryRetractAndExecute())
-			{
-				return;
-			}
-
 			// if we are on a named thread, handle waiting in TaskGraph-specific style
 			if (TryWaitOnNamedThread(*this))
 			{
-				return;
-			}
-
-			FEventRef CompletionEvent;
-			auto WaitingTaskBody = [&CompletionEvent] { CompletionEvent->Trigger(); };
-			using FWaitingTask = TExecutableTask<decltype(WaitingTaskBody)>;
-
-			// the task is stored on the stack as we can guarantee that it's out of the system by the end of the call
-			FWaitingTask WaitingTask{ TEXT("Waiting Task"), MoveTemp(WaitingTaskBody), ETaskPriority::Default /* doesn't matter*/, EExtendedTaskPriority::Inline };
-			WaitingTask.AddPrerequisites(*this);
-
-			if (WaitingTask.TryLaunch(sizeof(WaitingTask)))
-			{	// was executed inline
-				check(WaitingTask.IsCompleted());
-			}
-			else
-			{
-				CompletionEvent->Wait();
-			}
-
-			// the waiting task will be destroyed leaving this scope, wait for the internal reference to it to be released
-			while (WaitingTask.GetRefCount(std::memory_order_acquire) != 1)
-			{
-				FPlatformProcess::Yield();
+				return true;
 			}
-		}
-
-		bool FTaskBase::Wait(FTimespan InTimeout)
-		{
-			TaskTrace::FWaitingScope WaitingScope(GetTraceId());
-			TRACE_CPUPROFILER_EVENT_SCOPE(Tasks::Wait);
 
-			FTimeout Timeout{ InTimeout };
+			// ignore the result as we still have to make sure the task is completed upon returning from this function call
+			TryRetractAndExecute(Timeout);
 
-			if (TryRetractAndExecute())
+			// spin for a while with hope the task is getting completed right now, to avoid getting blocked by a pricy syscall
+			const uint32 MaxSpinCount = 40;
+			for (uint32 SpinCount = 0; SpinCount != MaxSpinCount && !IsCompleted() && !Timeout; ++SpinCount)
 			{
-				return true;
+				FPlatformProcess::YieldThread();
 			}
 
-			if (GetCurrentTask() == this)
+			if (IsCompleted() || Timeout)
 			{
-				UE_LOG(LogTemp, Fatal, TEXT("A task waiting for itself detected"));
-				return true;
+				return IsCompleted();
 			}
 
 			// the event must be alive for the task and this function lifetime, we don't know which one will be finished first as waiting can 
@@ -263,7 +226,10 @@ namespace UE::Tasks
 				return true;
 			}
 
-			return CompletionEvent->Wait((uint32)FMath::Clamp<int64>(Timeout.GetRemainingTime().GetTicks() / ETimespan::TicksPerMillisecond, 0, MAX_uint32));
+			uint32 WaitForMsecs = Timeout == FTimeout::Never() ? 
+				MAX_uint32 :
+				(uint32)FMath::Clamp<int64>(Timeout.GetRemainingTime().GetTicks() / ETimespan::TicksPerMillisecond, 0, MAX_uint32);
+			return CompletionEvent->Wait(WaitForMsecs);
 		}
 
 		FTaskBase* FTaskBase::TryPushIntoPipe()
@@ -320,7 +286,7 @@ namespace UE::Tasks
 				ReturnTask.TryLaunch(sizeof(ReturnTask)); // the result doesn't matter
 
 				TaskGraph.ProcessThreadUntilRequestReturn(CurrentThread);
-				return true;
+				return IsCompleted();
 			}
 #endif
 
diff --git a/Engine/Source/Runtime/Core/Public/Tasks/Pipe.h b/Engine/Source/Runtime/Core/Public/Tasks/Pipe.h
--- a/Engine/Source/Runtime/Core/Public/Tasks/Pipe.h
+++ b/Engine/Source/Runtime/Core/Public/Tasks/Pipe.h
@@ -47,7 +47,7 @@ namespace UE::Tasks
 
 		// waits until the pipe is empty (its last task is executed)
 		// should be used only after no more tasks are launched in the pipe, e.g. preparing for the pipe destruction
-		CORE_API bool WaitUntilEmpty(FTimespan Timeout = FTimespan::FromMilliseconds(MAX_uint32));
+		CORE_API bool WaitUntilEmpty(FTimespan Timeout = FTimespan::MaxValue());
 
 		// launches a task in the pipe
 		// @param InDebugName helps to identify the task in debugger and profiler
diff --git a/Engine/Source/Runtime/Core/Public/Tasks/Task.h b/Engine/Source/Runtime/Core/Public/Tasks/Task.h
--- a/Engine/Source/Runtime/Core/Public/Tasks/Task.h
+++ b/Engine/Source/Runtime/Core/Public/Tasks/Task.h
@@ -15,10 +15,7 @@ namespace UE::Tasks
 	class TTask;
 
 	template<typename TaskCollectionType>
-	void Wait(const TaskCollectionType& Tasks);
-
-	template<typename TaskCollectionType>
-	bool Wait(const TaskCollectionType& Tasks, FTimespan InTimeout);
+	bool Wait(const TaskCollectionType& Tasks, FTimespan InTimeout = FTimespan::MaxValue());
 
 	template<typename TaskCollectionType>
 	bool BusyWait(const TaskCollectionType& Tasks, FTimespan InTimeout = FTimespan::MaxValue());
@@ -43,17 +40,11 @@ namespace UE::Tasks
 			friend void PrerequisitesUnpacker(ArrayType& Array, TaskType& FirstTask);
 
 			template<typename TaskCollectionType>
-			friend bool TryRetractAndExecute(const TaskCollectionType& Tasks);
-
-			template<typename TaskCollectionType>
-			friend bool TryRetractAndExecute(const TaskCollectionType& Tasks, FTimespan Timeout);
+			friend bool TryRetractAndExecute(const TaskCollectionType& Tasks, FTimeout Timeout);
 
 			template<typename TaskCollectionType>
 			friend TArray<TaskTrace::FId> GetTraceIds(const TaskCollectionType& Tasks);
 
-			template<typename TaskCollectionType>
-			friend void UE::Tasks::Wait(const TaskCollectionType& Tasks);
-
 			template<typename TaskCollectionType>
 			friend bool UE::Tasks::Wait(const TaskCollectionType& Tasks, FTimespan InTimeout);
 
@@ -84,40 +75,20 @@ namespace UE::Tasks
 				return !IsValid() || Pimpl->IsCompleted();
 			}
 
-			// waits for task's completion. Tries to retract the task and execute it in-place, if failed - blocks until the task 
-			// is completed by another thread. If timeout is zero, tries to retract the task and returns immedially after that.
-			// @return true if the task is completed
-			void Wait()
-			{
-				if (IsValid())
-				{
-					Pimpl->Wait();
-				}
-			}
-
 			// waits for task's completion with timeout. Tries to retract the task and execute it in-place, if failed - blocks until the task 
 			// is completed by another thread. If timeout is zero, tries to retract the task and returns immedially after that.
 			// @return true if the task is completed
-			bool Wait(FTimespan Timeout)
+			bool Wait(FTimespan Timeout = FTimespan::MaxValue())
 			{
-				return !IsValid() || Pimpl->Wait(Timeout);
-			}
-
-			// waits for task's completion while executing other tasks. Shouldn't be used inside a latency-sensitive task
-			void BusyWait()
-			{
-				if (IsValid())
-				{
-					Pimpl->BusyWait();
-				}
+				return !IsValid() || Pimpl->Wait(FTimeout{ Timeout });
 			}
 
 			// waits for task's completion for at least the specified amount of time, while executing other tasks.
 			// the call can return much later than the given timeout
 			// @return true if the task is completed
-			bool BusyWait(FTimespan Timeout)
+			bool BusyWait(FTimespan Timeout = FTimespan::MaxValue())
 			{
-				return !IsValid() || Pimpl->BusyWait(Timeout);
+				return !IsValid() || Pimpl->BusyWait(FTimeout{ Timeout });
 			}
 
 			// waits for task's completion or the given condition becomes true, while executing other tasks.
@@ -337,52 +308,6 @@ namespace UE::Tasks
 		}
 	}
 
-	// wait for multiple tasks
-	// @param TaskCollectionType - an iterable collection of `TTask<T>`, e.g. `TArray<FTask>`
-	template<typename TaskCollectionType>
-	void Wait(const TaskCollectionType& Tasks)
-	{
-		TaskTrace::FWaitingScope WaitingScope(Private::GetTraceIds(Tasks));
-		TRACE_CPUPROFILER_EVENT_SCOPE(Tasks::Wait);
-
-		if (Private::TryRetractAndExecute(Tasks))
-		{
-			return;
-		}
-
-		Private::FTaskBase* CurrentTask = Private::GetCurrentTask();
-		for (const auto& Task : Tasks)
-		{
-			if (Task.Pimpl == CurrentTask)
-			{
-				UE_LOG(LogTemp, Fatal, TEXT("A task waiting for itself detected"));
-				return;
-			}
-		}
-
-		FEventRef CompletionEvent;
-		auto WaitingTaskBody = [&CompletionEvent] { CompletionEvent->Trigger(); };
-		using FWaitingTask = Private::TExecutableTask<decltype(WaitingTaskBody)>;
-
-		FWaitingTask WaitingTask{ TEXT("Waiting Task"), MoveTemp(WaitingTaskBody), ETaskPriority::Default /* doesn't matter */,  EExtendedTaskPriority::Inline };
-		WaitingTask.AddPrerequisites(Tasks);
-
-		if (WaitingTask.TryLaunch(sizeof(WaitingTask)))
-		{	// was executed inline
-			check(WaitingTask.IsCompleted());
-		}
-		else
-		{ 
-			CompletionEvent->Wait();
-		}
-
-		// `WaitingTask` is allocated on the stack. we need to wait until its RefCount reaches zero before leaving the function to avoid "use after destruction"
-		while (WaitingTask.GetRefCount(std::memory_order_acquire) != 1)
-		{
-			FPlatformProcess::Yield();
-		}
-	}
-
 	inline void Wait(Private::FTaskHandle& Task)
 	{
 		Task.Wait();
@@ -391,17 +316,15 @@ namespace UE::Tasks
 	// wait for multiple tasks, with timeout
 	// @param TaskCollectionType - an iterable collection of `TTask<T>`, e.g. `TArray<FTask>`
 	template<typename TaskCollectionType>
-	bool Wait(const TaskCollectionType& Tasks, FTimespan InTimeout)
+	bool Wait(const TaskCollectionType& Tasks, FTimespan InTimeout/* = FTimespan::MaxValue()*/)
 	{
 		TaskTrace::FWaitingScope WaitingScope(Private::GetTraceIds(Tasks));
 		TRACE_CPUPROFILER_EVENT_SCOPE(Tasks::Wait);
 
 		FTimeout Timeout{ InTimeout };
 
-		if (Private::TryRetractAndExecute(Tasks, InTimeout))
-		{
-			return true;
-		}
+		// ignore the result as we still have to make sure the task is completed upon returning from this function call
+		Private::TryRetractAndExecute(Tasks, Timeout);
 
 		// the event must be alive for the task and this function lifetime, we don't know which one will be finished first as waiting can time out
 		// before the waiting task is completed
@@ -418,26 +341,27 @@ namespace UE::Tasks
 			return true;
 		}
 
-		return CompletionEvent->Wait((uint32)FMath::Clamp<int64>(Timeout.GetRemainingTime().GetTicks() / ETimespan::TicksPerMillisecond, 0, MAX_uint32));
+		uint32 WaitForMsecs = Timeout == FTimeout::Never() ?
+			MAX_uint32 :
+			(uint32)FMath::Clamp<int64>(Timeout.GetRemainingTime().GetTicks() / ETimespan::TicksPerMillisecond, 0, MAX_uint32);
+		return CompletionEvent->Wait(WaitForMsecs);
 	}
 
 	// wait for multiple tasks while executing other tasks
 	template<typename TaskCollectionType>
-	bool BusyWait(const TaskCollectionType& Tasks, FTimespan TimeoutValue/* = FTimespan::MaxValue()*/)
+	bool BusyWait(const TaskCollectionType& Tasks, FTimespan InTimeout/* = FTimespan::MaxValue()*/)
 	{
 		TaskTrace::FWaitingScope WaitingScope(Private::GetTraceIds(Tasks));
 		TRACE_CPUPROFILER_EVENT_SCOPE(Tasks::BusyWait);
 
-		FTimeout Timeout{ TimeoutValue };
+		FTimeout Timeout{ InTimeout };
 
-		if (Private::TryRetractAndExecute(Tasks, TimeoutValue))
-		{
-			return true;
-		}
+		// ignore the result as we still have to make sure the task is completed upon returning from this function call
+		Private::TryRetractAndExecute(Tasks, Timeout);
 
 		for (auto& Task : Tasks)
 		{
-			if (Timeout || !Task.BusyWait(Timeout.GetRemainingTime()))
+			if (Timeout || !Task.BusyWait(Timeout))
 			{
 				return false;
 			}
diff --git a/Engine/Source/Runtime/Core/Public/Tasks/TaskPrivate.h b/Engine/Source/Runtime/Core/Public/Tasks/TaskPrivate.h
--- a/Engine/Source/Runtime/Core/Public/Tasks/TaskPrivate.h
+++ b/Engine/Source/Runtime/Core/Public/Tasks/TaskPrivate.h
@@ -323,10 +323,12 @@ namespace UE::Tasks
 				return Subsequents.IsClosed();
 			}
 
-			// Tries to pull out the task from the system and execute it. If the task is locked by either prerequisites or nested tasks, tries to retract and execute them recursively. 
-			// @return true if task is completed, not necessarily by retraction. If the task is being executed (or its dependency) in parallel, it doesn't wait for task completion and 
-			// returns false immediately.
-			CORE_API bool TryRetractAndExecute(uint32 RecursionDepth = 0);
+			// Tries to pull out the task from the system and execute it. If the task is locked by either prerequisites or nested tasks, tries to 
+			// retract and execute them recursively. 
+			// WARNING: the function can return `true` even if the task is not completed yet. The `true` means only that the task is already
+			// executed and has no other pending dependencies, but can be in the process of completion (concurrently). The caller still needs 
+			// to wait for completion explicitly.
+			CORE_API bool TryRetractAndExecute(FTimeout Timeout, uint32 RecursionDepth = 0);
 
 			// releases internal reference and maintains low-level task state. must be called iff the task was never launched, otherwise 
 			// the scheduler will do this in due course
@@ -359,33 +361,16 @@ namespace UE::Tasks
 			// `Wait(FTimespan::Zero())` still tries to retract and execute the task, use `IsCompleted()` to check for completeness. 
 			// The version w/o timeout is slightly more efficient.
 			// @return true if the task is completed
-			CORE_API void Wait();
-			CORE_API bool Wait(FTimespan Timeout);
-
-			// waits until the task is completed while executing other tasks
-			void BusyWait()
-			{
-				TaskTrace::FWaitingScope WaitingScope(GetTraceId());
-				TRACE_CPUPROFILER_EVENT_SCOPE(Tasks::BusyWait);
-				
-				if (!TryRetractAndExecute())
-				{
-					LowLevelTasks::BusyWaitUntil([this] { return IsCompleted(); });
-				}
-			}
+			CORE_API bool Wait(FTimeout Timeout);
 
 			// waits until the task is completed or waiting timed out, while executing other tasks
-			bool BusyWait(FTimespan InTimeout)
+			bool BusyWait(FTimeout Timeout)
 			{
 				TaskTrace::FWaitingScope WaitingScope(GetTraceId());
 				TRACE_CPUPROFILER_EVENT_SCOPE(Tasks::BusyWait);
 
-				FTimeout Timeout{ InTimeout };
-				
-				if (TryRetractAndExecute())
-				{
-					return true;
-				}
+				// ignore the result as we still have to make sure the task is completed upon returning from this function call
+				TryRetractAndExecute(Timeout);
 
 				LowLevelTasks::BusyWaitUntil([this, Timeout] { return IsCompleted() || Timeout; });
 				return IsCompleted();
@@ -398,10 +383,8 @@ namespace UE::Tasks
 				TaskTrace::FWaitingScope WaitingScope(GetTraceId());
 				TRACE_CPUPROFILER_EVENT_SCOPE(Tasks::BusyWait);
 
-				if (TryRetractAndExecute())
-				{
-					return true;
-				}
+				// ignore the result as we still have to make sure the task is completed upon returning from this function call
+				TryRetractAndExecute(FTimeout::Never());
 
 				LowLevelTasks::BusyWaitUntil(
 					[this, Condition = Forward<ConditionType>(Condition)]{ return IsCompleted() || Condition(); }
@@ -826,36 +809,19 @@ namespace UE::Tasks
 			TaskEventBaseAllocator.Free(Ptr);
 		}
 
-		// task retraction of multiple tasks. 
-		// @return true if all tasks are completed
-		template<typename TaskCollectionType>
-		bool TryRetractAndExecute(const TaskCollectionType& Tasks)
-		{
-			bool bResult = true;
-
-			for (auto& Task : Tasks)
-			{
-				if (Task.IsValid() && !Task.Pimpl->TryRetractAndExecute())
-				{
-					bResult = false; // do not stop here to let this thread to help in executing tasks as much as possible, as it's waiting for their completion anyway
-				}
-			}
-
-			return bResult;
-		}
-
-		// task retraction of multiple tasks, with timeout. The timeout is rounded up to any successful task execution, which means that it can time out only in-between individual task
-		// retractions.
-		// @return true if all tasks are completed
+		// task retraction of multiple tasks, with timeout. The timeout is rounded up to any successful task execution, which means that it can 
+		// time out only in-between individual task retractions.
+		// WARNING: the function can return `true` even if some tasks are still not completed. The `true` means only that the tasks are executed
+		// and have no other pending dependencies, but can be still in the process of completion (concurrently). The caller still needs to wait 
+		// for completion.
 		template<typename TaskCollectionType>
-		bool TryRetractAndExecute(const TaskCollectionType& Tasks, FTimespan InTimeout)
+		bool TryRetractAndExecute(const TaskCollectionType& Tasks, FTimeout Timeout)
 		{
-			FTimeout Timeout{ InTimeout };
 			bool bResult = true;
 
 			for (auto& Task : Tasks)
 			{
-				if (Task.IsValid() && !Task.Pimpl->TryRetractAndExecute())
+				if (Task.IsValid() && !Task.Pimpl->TryRetractAndExecute(Timeout))
 				{
 					bResult = false;  // do not stop here to let this thread to help in executing tasks as much as possible, as it's waiting for their completion anyway
 				}
