I/O - move FIoBuffer to separate header

diff --git a/Engine/Source/Runtime/Core/Private/IO/IoBuffer.cpp b/Engine/Source/Runtime/Core/Private/IO/IoBuffer.cpp
--- a/Engine/Source/Runtime/Core/Private/IO/IoBuffer.cpp
+++ b/Engine/Source/Runtime/Core/Private/IO/IoBuffer.cpp
@@ -1,10 +1,9 @@
 // Copyright Epic Games, Inc. All Rights Reserved.
 
+#include "IO/IoBuffer.h"
+#include "IO/IoDispatcher.h"
 #include "HAL/Platform.h"
 #include "HAL/UnrealMemory.h"
-#include "IO/IoDispatcher.h"
-#include "Memory/MemoryFwd.h"
-#include "Memory/MemoryView.h"
 #include "Misc/AssertionMacros.h"
 #include "Templates/RefCounting.h"
 
diff --git a/Engine/Source/Runtime/Core/Public/IO/IoBuffer.h b/Engine/Source/Runtime/Core/Public/IO/IoBuffer.h
new file mode 100644
--- /dev/null
+++ b/Engine/Source/Runtime/Core/Public/IO/IoBuffer.h
@@ -0,0 +1,178 @@
+// Copyright Epic Games, Inc. All Rights Reserved.
+
+#pragma once
+
+#include "Memory/MemoryView.h"
+#include "Templates/RefCounting.h"
+
+template<typename T>
+class TIoStatusOr;
+
+/**
+ * Reference to buffer data used by I/O dispatcher APIs
+ */
+class FIoBuffer
+{
+public:
+	enum EAssumeOwnershipTag	{ AssumeOwnership };
+	enum ECloneTag				{ Clone };
+	enum EWrapTag				{ Wrap };
+
+	CORE_API			FIoBuffer();
+	CORE_API explicit	FIoBuffer(uint64 InSize);
+	CORE_API			FIoBuffer(const void* Data, uint64 InSize, const FIoBuffer& OuterBuffer);
+	CORE_API			FIoBuffer(FMemoryView Memory, const FIoBuffer& OuterBuffer);
+
+	CORE_API			FIoBuffer(EAssumeOwnershipTag,	const void* Data, uint64 InSize);
+	CORE_API			FIoBuffer(EAssumeOwnershipTag,	FMemoryView Memory);
+	CORE_API			FIoBuffer(ECloneTag,			const void* Data, uint64 InSize);
+	CORE_API			FIoBuffer(ECloneTag,			FMemoryView Memory);
+	CORE_API			FIoBuffer(EWrapTag,				const void* Data, uint64 InSize);
+	CORE_API			FIoBuffer(EWrapTag,				FMemoryView Memory);
+
+	// Note: we currently rely on implicit move constructor, thus we do not declare any
+	//		 destructor or copy/assignment operators or copy constructors
+
+	inline const uint8*	Data() const			{ return CorePtr->Data(); }
+	inline uint8*		Data()					{ return CorePtr->Data(); }
+	inline const uint8*	GetData() const			{ return CorePtr->Data(); }
+	inline uint8*		GetData()				{ return CorePtr->Data(); }
+	inline uint64		DataSize() const		{ return CorePtr->DataSize(); }
+	inline uint64		GetSize() const			{ return CorePtr->DataSize(); }
+
+	inline FMemoryView	GetView() const			{ return MakeMemoryView(CorePtr->Data(), CorePtr->DataSize()); }
+	inline FMutableMemoryView GetMutableView()	{ return MakeMemoryView(CorePtr->Data(), CorePtr->DataSize()); }
+
+	inline void			SetSize(uint64 InSize)	{ return CorePtr->SetSize(InSize); }
+
+	inline bool			IsMemoryOwned() const	{ return CorePtr->IsMemoryOwned(); }
+
+	inline void			EnsureOwned() const		{ if (!CorePtr->IsMemoryOwned()) { MakeOwned(); } }
+
+	inline bool			operator!=(const FIoBuffer& Rhs) const { return DataSize() != Rhs.DataSize() || FMemory::Memcmp(GetData(), Rhs.GetData(), DataSize()) != 0; }
+
+	CORE_API void		MakeOwned() const;
+	
+	/**
+	 * Relinquishes control of the internal buffer to the caller and removes it from the FIoBuffer.
+	 * This allows the caller to assume ownership of the internal data and prevent it from being deleted along with 
+	 * the FIoBuffer.
+	 *
+	 * NOTE: It is only valid to call this if the FIoBuffer currently owns the internal memory allocation, as the 
+	 * point of the call is to take ownership of it. If the FIoBuffer is only wrapping the allocation then it will
+	 * return a failed FIoStatus instead.
+	 *
+	 * @return A status wrapper around the memory pointer. Even if the status is valid the pointer might still be null.
+	 */
+	UE_NODISCARD CORE_API TIoStatusOr<uint8*> Release();
+
+private:
+	/** Core buffer object. For internal use only, used by FIoBuffer
+
+		Contains all state pertaining to a buffer.
+	  */
+	struct BufCore
+	{
+					BufCore();
+		CORE_API	~BufCore();
+
+		explicit	BufCore(uint64 InSize);
+					BufCore(const uint8* InData, uint64 InSize, bool InOwnsMemory);
+					BufCore(const uint8* InData, uint64 InSize, const BufCore* InOuter);
+					BufCore(ECloneTag, uint8* InData, uint64 InSize);
+
+					BufCore(const BufCore& Rhs) = delete;
+		
+		BufCore& operator=(const BufCore& Rhs) = delete;
+
+		inline uint8* Data()			{ return DataPtr; }
+		inline uint64 DataSize() const	{ return DataSizeLow | (uint64(DataSizeHigh) << 32); }
+
+		//
+
+		void	SetDataAndSize(const uint8* InData, uint64 InSize);
+		void	SetSize(uint64 InSize);
+
+		void	MakeOwned();
+
+		TIoStatusOr<uint8*> ReleaseMemory();
+
+		inline void SetIsOwned(bool InOwnsMemory)
+		{
+			if (InOwnsMemory)
+			{
+				Flags |= OwnsMemory;
+			}
+			else
+			{
+				Flags &= ~OwnsMemory;
+			}
+		}
+
+		inline uint32 AddRef() const
+		{
+			return uint32(FPlatformAtomics::InterlockedIncrement(&NumRefs));
+		}
+
+		inline uint32 Release() const
+		{
+#if DO_CHECK
+			CheckRefCount();
+#endif
+
+			const int32 Refs = FPlatformAtomics::InterlockedDecrement(&NumRefs);
+			if (Refs == 0)
+			{
+				delete this;
+			}
+
+			return uint32(Refs);
+		}
+
+		uint32 GetRefCount() const
+		{
+			return uint32(NumRefs);
+		}
+
+		bool IsMemoryOwned() const	{ return Flags & OwnsMemory; }
+
+	private:
+		CORE_API void				CheckRefCount() const;
+
+		uint8*						DataPtr = nullptr;
+
+		uint32						DataSizeLow = 0;
+		mutable int32				NumRefs = 0;
+
+		// Reference-counted outer "core", used for views into other buffer
+		//
+		// Ultimately this should probably just be an index into a pool
+		TRefCountPtr<const BufCore>	OuterCore;
+
+		// TODO: These two could be packed in the MSB of DataPtr on x64
+		uint8		DataSizeHigh = 0;	// High 8 bits of size (40 bits total)
+		uint8		Flags = 0;
+
+		enum
+		{
+			OwnsMemory		= 1 << 0,	// Buffer memory is owned by this instance
+			ReadOnlyBuffer	= 1 << 1,	// Buffer memory is immutable
+			
+			FlagsMask		= (1 << 2) - 1
+		};
+
+		void EnsureDataIsResident() {}
+
+		void ClearFlags()
+		{
+			Flags = 0;
+		}
+	};
+
+	// Reference-counted "core"
+	//
+	// Ultimately this should probably just be an index into a pool
+	TRefCountPtr<BufCore>	CorePtr;
+	
+	friend class FIoBufferManager;
+};
diff --git a/Engine/Source/Runtime/Core/Public/IO/IoDispatcher.h b/Engine/Source/Runtime/Core/Public/IO/IoDispatcher.h
--- a/Engine/Source/Runtime/Core/Public/IO/IoDispatcher.h
+++ b/Engine/Source/Runtime/Core/Public/IO/IoDispatcher.h
@@ -17,6 +17,7 @@
 #include "HAL/PlatformAtomics.h"
 #include "HAL/PlatformFile.h"
 #include "HAL/UnrealMemory.h"
+#include "IO/IoBuffer.h"
 #include "IO/IoChunkId.h"
 #include "IO/IoContainerId.h"
 #include "IO/IoDispatcherPriority.h"
@@ -87,174 +88,6 @@ private:
 	int32			Order = 0;
 };
 
-/** Reference to buffer data used by I/O dispatcher APIs
-  */
-class FIoBuffer
-{
-public:
-	enum EAssumeOwnershipTag	{ AssumeOwnership };
-	enum ECloneTag				{ Clone };
-	enum EWrapTag				{ Wrap };
-
-	CORE_API			FIoBuffer();
-	CORE_API explicit	FIoBuffer(uint64 InSize);
-	CORE_API			FIoBuffer(const void* Data, uint64 InSize, const FIoBuffer& OuterBuffer);
-	CORE_API			FIoBuffer(FMemoryView Memory, const FIoBuffer& OuterBuffer);
-
-	CORE_API			FIoBuffer(EAssumeOwnershipTag,	const void* Data, uint64 InSize);
-	CORE_API			FIoBuffer(EAssumeOwnershipTag,	FMemoryView Memory);
-	CORE_API			FIoBuffer(ECloneTag,			const void* Data, uint64 InSize);
-	CORE_API			FIoBuffer(ECloneTag,			FMemoryView Memory);
-	CORE_API			FIoBuffer(EWrapTag,				const void* Data, uint64 InSize);
-	CORE_API			FIoBuffer(EWrapTag,				FMemoryView Memory);
-
-	// Note: we currently rely on implicit move constructor, thus we do not declare any
-	//		 destructor or copy/assignment operators or copy constructors
-
-	inline const uint8*	Data() const			{ return CorePtr->Data(); }
-	inline uint8*		Data()					{ return CorePtr->Data(); }
-	inline const uint8*	GetData() const			{ return CorePtr->Data(); }
-	inline uint8*		GetData()				{ return CorePtr->Data(); }
-	inline uint64		DataSize() const		{ return CorePtr->DataSize(); }
-	inline uint64		GetSize() const			{ return CorePtr->DataSize(); }
-
-	inline FMemoryView	GetView() const			{ return MakeMemoryView(CorePtr->Data(), CorePtr->DataSize()); }
-	inline FMutableMemoryView GetMutableView()	{ return MakeMemoryView(CorePtr->Data(), CorePtr->DataSize()); }
-
-	inline void			SetSize(uint64 InSize)	{ return CorePtr->SetSize(InSize); }
-
-	inline bool			IsMemoryOwned() const	{ return CorePtr->IsMemoryOwned(); }
-
-	inline void			EnsureOwned() const		{ if (!CorePtr->IsMemoryOwned()) { MakeOwned(); } }
-
-	inline bool			operator!=(const FIoBuffer& Rhs) const { return DataSize() != Rhs.DataSize() || FMemory::Memcmp(GetData(), Rhs.GetData(), DataSize()) != 0; }
-
-	CORE_API void		MakeOwned() const;
-	
-	/**
-	 * Relinquishes control of the internal buffer to the caller and removes it from the FIoBuffer.
-	 * This allows the caller to assume ownership of the internal data and prevent it from being deleted along with 
-	 * the FIoBuffer.
-	 *
-	 * NOTE: It is only valid to call this if the FIoBuffer currently owns the internal memory allocation, as the 
-	 * point of the call is to take ownership of it. If the FIoBuffer is only wrapping the allocation then it will
-	 * return a failed FIoStatus instead.
-	 *
-	 * @return A status wrapper around the memory pointer. Even if the status is valid the pointer might still be null.
-	 */
-	UE_NODISCARD CORE_API TIoStatusOr<uint8*> Release();
-
-private:
-	/** Core buffer object. For internal use only, used by FIoBuffer
-
-		Contains all state pertaining to a buffer.
-	  */
-	struct BufCore
-	{
-					BufCore();
-		CORE_API	~BufCore();
-
-		explicit	BufCore(uint64 InSize);
-					BufCore(const uint8* InData, uint64 InSize, bool InOwnsMemory);
-					BufCore(const uint8* InData, uint64 InSize, const BufCore* InOuter);
-					BufCore(ECloneTag, uint8* InData, uint64 InSize);
-
-					BufCore(const BufCore& Rhs) = delete;
-		
-		BufCore& operator=(const BufCore& Rhs) = delete;
-
-		inline uint8* Data()			{ return DataPtr; }
-		inline uint64 DataSize() const	{ return DataSizeLow | (uint64(DataSizeHigh) << 32); }
-
-		//
-
-		void	SetDataAndSize(const uint8* InData, uint64 InSize);
-		void	SetSize(uint64 InSize);
-
-		void	MakeOwned();
-
-		TIoStatusOr<uint8*> ReleaseMemory();
-
-		inline void SetIsOwned(bool InOwnsMemory)
-		{
-			if (InOwnsMemory)
-			{
-				Flags |= OwnsMemory;
-			}
-			else
-			{
-				Flags &= ~OwnsMemory;
-			}
-		}
-
-		inline uint32 AddRef() const
-		{
-			return uint32(FPlatformAtomics::InterlockedIncrement(&NumRefs));
-		}
-
-		inline uint32 Release() const
-		{
-#if DO_CHECK
-			CheckRefCount();
-#endif
-
-			const int32 Refs = FPlatformAtomics::InterlockedDecrement(&NumRefs);
-			if (Refs == 0)
-			{
-				delete this;
-			}
-
-			return uint32(Refs);
-		}
-
-		uint32 GetRefCount() const
-		{
-			return uint32(NumRefs);
-		}
-
-		bool IsMemoryOwned() const	{ return Flags & OwnsMemory; }
-
-	private:
-		CORE_API void				CheckRefCount() const;
-
-		uint8*						DataPtr = nullptr;
-
-		uint32						DataSizeLow = 0;
-		mutable int32				NumRefs = 0;
-
-		// Reference-counted outer "core", used for views into other buffer
-		//
-		// Ultimately this should probably just be an index into a pool
-		TRefCountPtr<const BufCore>	OuterCore;
-
-		// TODO: These two could be packed in the MSB of DataPtr on x64
-		uint8		DataSizeHigh = 0;	// High 8 bits of size (40 bits total)
-		uint8		Flags = 0;
-
-		enum
-		{
-			OwnsMemory		= 1 << 0,	// Buffer memory is owned by this instance
-			ReadOnlyBuffer	= 1 << 1,	// Buffer memory is immutable
-			
-			FlagsMask		= (1 << 2) - 1
-		};
-
-		void EnsureDataIsResident() {}
-
-		void ClearFlags()
-		{
-			Flags = 0;
-		}
-	};
-
-	// Reference-counted "core"
-	//
-	// Ultimately this should probably just be an index into a pool
-	TRefCountPtr<BufCore>	CorePtr;
-	
-	friend class FIoBufferManager;
-};
-
 class FIoChunkHash
 {
 public:
