I/O - added caching to HTTP I/O backend

-  added simple in memory LRU cache for testing purposes, enabled with -HttpBackendMemoryCache=<size>

diff --git a/Engine/Source/Runtime/Core/Internal/IO/IoMemoryCache.h b/Engine/Source/Runtime/Core/Internal/IO/IoMemoryCache.h
new file mode 100644
--- /dev/null
+++ b/Engine/Source/Runtime/Core/Internal/IO/IoMemoryCache.h
@@ -0,0 +1,15 @@
+// Copyright Epic Games, Inc. All Rights Reserved.
+
+#pragma once
+
+#include "Misc/Build.h"
+
+#if !UE_BUILD_SHIPPING
+
+#include "Templates/UniquePtr.h"
+
+class IIoCache;
+
+TUniquePtr<IIoCache> MakeMemoryIoCache(uint64 CacheSize);
+
+#endif // !UE_BUILD_SHIPPING
diff --git a/Engine/Source/Runtime/Core/Private/IO/IoCache.cpp b/Engine/Source/Runtime/Core/Private/IO/IoCache.cpp
--- a/Engine/Source/Runtime/Core/Private/IO/IoCache.cpp
+++ b/Engine/Source/Runtime/Core/Private/IO/IoCache.cpp
@@ -55,3 +55,13 @@ void FIoCacheRequest::Cancel()
 		Pimpl->Cancel();
 	}
 }
+
+void FIoCacheRequest::Wait()
+{
+	check(IsValid());
+
+	if (!Status().IsCompleted())
+	{
+		Pimpl->Wait();
+	}
+}
diff --git a/Engine/Source/Runtime/Core/Private/IO/IoMemoryCache.cpp b/Engine/Source/Runtime/Core/Private/IO/IoMemoryCache.cpp
new file mode 100644
--- /dev/null
+++ b/Engine/Source/Runtime/Core/Private/IO/IoMemoryCache.cpp
@@ -0,0 +1,170 @@
+// Copyright Epic Games, Inc. All Rights Reserved.
+
+#include "IO/IoMemoryCache.h"
+
+#include "Containers/IntrusiveDoubleLinkedList.h"
+#include "IO/IoCache.h"
+#include "IO/IoDispatcher.h"
+#include "IO/IoHash.h"
+#include "Misc/ScopeLock.h"
+#include "Tasks/Task.h"
+
+#if !UE_BUILD_SHIPPING
+
+namespace UE::IO::Private
+{
+
+////////////////////////////////////////////////////////////////////////////////
+class FMemoryIoCacheRequest
+	: public FIoCacheRequestBase
+{
+public:
+	FMemoryIoCacheRequest(FIoReadCallback&& ReadCallback, FIoBuffer CachedBuffer);
+	virtual ~FMemoryIoCacheRequest() = default;
+	virtual void Wait() override;
+	virtual void Cancel() override;
+
+	void Issue();
+
+private:
+	FIoBuffer Buffer;
+	UE::Tasks::TTask<void> Task;
+	std::atomic_bool bCanceled{false};
+};
+
+FMemoryIoCacheRequest::FMemoryIoCacheRequest(FIoReadCallback&& ReadCallback, FIoBuffer CachedBuffer)
+	: FIoCacheRequestBase(MoveTemp(ReadCallback))
+	, Buffer(CachedBuffer)
+{
+}
+
+void FMemoryIoCacheRequest::Wait() 
+{
+	Task.Wait();
+}
+
+void FMemoryIoCacheRequest::Cancel() 
+{
+	bCanceled = true;
+}
+
+void FMemoryIoCacheRequest::Issue()
+{
+	Task = UE::Tasks::Launch(TEXT("Memory I/O Cache"), [this]()
+	{
+		if (bCanceled)
+		{
+			CompleteRequest(EIoErrorCode::Cancelled);
+		}
+		else if (Buffer.GetSize() > 0)
+		{
+			CompleteRequest(MoveTemp(Buffer));
+		}
+		else
+		{
+			CompleteRequest(EIoErrorCode::NotFound);
+		}
+	});
+}
+
+////////////////////////////////////////////////////////////////////////////////
+class FMemoryIoCache
+	: public IIoCache
+{
+	struct FCacheEntry
+		: public TIntrusiveDoubleLinkedListNode<FCacheEntry>
+	{
+		FIoHash Key;
+		FIoBuffer Buffer;
+	};
+
+	using FCacheEntryList = TIntrusiveDoubleLinkedList<FCacheEntry>;
+	using FCacheEntryLookup = TMap<FIoHash, TUniquePtr<FCacheEntry>>;
+
+public:
+	explicit FMemoryIoCache(uint64 CacheSize);
+	virtual ~FMemoryIoCache();
+	virtual bool ContainsChunk(const FIoHash& Key) const override;
+	virtual FIoCacheRequest GetChunk(const FIoHash& Key, const FIoReadOptions& Options, FIoReadCallback&& Callback) override;
+	virtual FIoStatus PutChunk(const FIoHash& Key, FMemoryView Data) override;
+
+private:
+	const uint64 MaxCacheSize;
+	uint64 TotalCacheSize;
+	FCacheEntryLookup Lookup;
+	FCacheEntryList Lru;
+	mutable FCriticalSection CacheCS;
+};
+
+FMemoryIoCache::FMemoryIoCache(uint64 CacheSize)
+	: MaxCacheSize(CacheSize)
+	, TotalCacheSize(0)
+{
+}
+
+FMemoryIoCache::~FMemoryIoCache()
+{
+}
+
+bool FMemoryIoCache::ContainsChunk(const FIoHash& Key) const
+{
+	FScopeLock _(&CacheCS);
+	return Lookup.Contains(Key);
+}
+
+FIoCacheRequest FMemoryIoCache::GetChunk(const FIoHash& Key, const FIoReadOptions& Options, FIoReadCallback&& Callback)
+{
+	FIoBuffer Buffer;
+	{
+		FScopeLock _(&CacheCS);
+
+		if (TUniquePtr<FCacheEntry>* Entry = Lookup.Find(Key))
+		{
+			FCacheEntry* CacheEntry = Entry->Get();
+			Buffer = CacheEntry->Buffer;
+			Lru.Remove(CacheEntry);
+			Lru.AddHead(CacheEntry);
+		}
+	}
+
+	TUniquePtr<FMemoryIoCacheRequest> Request = MakeUnique<FMemoryIoCacheRequest>(MoveTemp(Callback), Buffer);
+	Request->Issue();
+
+	return FIoCacheRequest(Request.Release());
+}
+
+FIoStatus FMemoryIoCache::PutChunk(const FIoHash& Key, FMemoryView Data)
+{
+	FScopeLock _(&CacheCS);
+
+	if (Lookup.Contains(Key))
+	{
+		return FIoStatus::Ok;
+	}
+
+	while ((TotalCacheSize + Data.GetSize()) > MaxCacheSize)
+	{
+		FCacheEntry* Entry = Lru.PopTail();
+		check(Entry);
+		TotalCacheSize -= Entry->Buffer.GetSize();
+		const FIoHash TmpKey = Entry->Key;
+		Lookup.Remove(TmpKey);
+	}
+
+	TUniquePtr<FCacheEntry>& NewEntry = Lookup.Add(Key, MakeUnique<FCacheEntry>());
+	NewEntry->Key = Key;
+	NewEntry->Buffer = FIoBuffer(FIoBuffer::Clone, Data);
+	Lru.AddHead(NewEntry.Get());
+
+	return FIoStatus::Ok;
+}
+
+} // namespace UE::IO::Private
+
+////////////////////////////////////////////////////////////////////////////////
+TUniquePtr<IIoCache> MakeMemoryIoCache(uint64 CacheSize)
+{
+	return MakeUnique<UE::IO::Private::FMemoryIoCache>(CacheSize);
+}
+
+#endif // !UE_BUILD_SHIPPING
diff --git a/Engine/Source/Runtime/Core/Public/IO/IoCache.h b/Engine/Source/Runtime/Core/Public/IO/IoCache.h
--- a/Engine/Source/Runtime/Core/Public/IO/IoCache.h
+++ b/Engine/Source/Runtime/Core/Public/IO/IoCache.h
@@ -76,6 +76,7 @@ public:
 		return IsValid() ? Pimpl->Status() : FIoStatus(EIoErrorCode::InvalidCode);
 	}
 
+	void Wait();
 	void Cancel();
 
 private:
