Explictly handle NaNs when evaluating curves.

diff --git a/Engine/Source/Runtime/Core/Public/Math/InterpCurve.h b/Engine/Source/Runtime/Core/Public/Math/InterpCurve.h
--- a/Engine/Source/Runtime/Core/Public/Math/InterpCurve.h
+++ b/Engine/Source/Runtime/Core/Public/Math/InterpCurve.h
@@ -313,6 +313,15 @@ T FInterpCurve<T>::Eval(const float InVal, const T& Default) const
 		return Default;
 	}
 
+	// If we let NaNs in through here, they fail the check() on the Alpha between the two points.
+	if (FPlatformMath::IsNaN(InVal))
+	{
+#if ENABLE_NAN_DIAGNOSTIC
+		logOrEnsureNanError(TEXT("FInterpCurve<T>::Eval has InVal == NaN"));
+#endif
+		return Default;
+	}
+
 	// Binary search to find index of lower bound of input value
 	const int32 Index = GetPointIndexForInputValue(InVal);
 
@@ -349,7 +358,7 @@ T FInterpCurve<T>::Eval(const float InVal, const T& Default) const
 	if (Diff > 0.0f && PrevPoint.InterpMode != CIM_Constant)
 	{
 		const float Alpha = (InVal - PrevPoint.InVal) / Diff;
-		check(Alpha >= 0.0f && Alpha <= 1.0f);
+		checkf(Alpha >= 0.0f && Alpha <= 1.0f, TEXT("Bad value in Eval(): in %f prev %f  diff %f alpha %f"), InVal, PrevPoint.InVal, Diff, Alpha);
 
 		if (PrevPoint.InterpMode == CIM_Linear)
 		{
@@ -379,6 +388,15 @@ T FInterpCurve<T>::EvalDerivative(const float InVal, const T& Default) const
 		return Default;
 	}
 
+	// If we let NaNs in through here, they fail the check() on the Alpha between the two points.
+	if (FPlatformMath::IsNaN(InVal))
+	{
+#if ENABLE_NAN_DIAGNOSTIC
+		logOrEnsureNanError(TEXT("FInterpCurve<T>::EvalDerivative has InVal == NaN"));
+#endif
+		return Default;
+	}
+	
 	// Binary search to find index of lower bound of input value
 	const int32 Index = GetPointIndexForInputValue(InVal);
 
@@ -421,7 +439,7 @@ T FInterpCurve<T>::EvalDerivative(const float InVal, const T& Default) const
 		else
 		{
 			const float Alpha = (InVal - PrevPoint.InVal) / Diff;
-			check(Alpha >= 0.0f && Alpha <= 1.0f);
+			checkf(Alpha >= 0.0f && Alpha <= 1.0f, TEXT("Bad value in EvalDerivative(): in %f prev %f  diff %f alpha %f"), InVal, PrevPoint.InVal, Diff, Alpha);
 
 			return FMath::CubicInterpDerivative(PrevPoint.OutVal, PrevPoint.LeaveTangent * Diff, NextPoint.OutVal, NextPoint.ArriveTangent * Diff, Alpha) / Diff;
 		}
@@ -445,6 +463,15 @@ T FInterpCurve<T>::EvalSecondDerivative(const float InVal, const T& Default) con
 		return Default;
 	}
 
+	// If we let NaNs in through here, they fail the check() on the Alpha between the two points.
+	if (FPlatformMath::IsNaN(InVal))
+	{
+#if ENABLE_NAN_DIAGNOSTIC
+		logOrEnsureNanError(TEXT("FInterpCurve<T>::EvalSecondDerivative has InVal == NaN"));
+#endif
+		return Default;
+	}
+	
 	// Binary search to find index of lower bound of input value
 	const int32 Index = GetPointIndexForInputValue(InVal);
 
@@ -483,7 +510,7 @@ T FInterpCurve<T>::EvalSecondDerivative(const float InVal, const T& Default) con
 		else
 		{
 			const float Alpha = (InVal - PrevPoint.InVal) / Diff;
-			check(Alpha >= 0.0f && Alpha <= 1.0f);
+			checkf(Alpha >= 0.0f && Alpha <= 1.0f, TEXT("Bad value in EvalSecondDerivative(): in %f prev %f  diff %f alpha %f"), InVal, PrevPoint.InVal, Diff, Alpha);
 
 			return FMath::CubicInterpSecondDerivative(PrevPoint.OutVal, PrevPoint.LeaveTangent * Diff, NextPoint.OutVal, NextPoint.ArriveTangent * Diff, Alpha) / (Diff * Diff);
 		}
