Loop as few times as possible when adding elements to TChunkedArray.

diff --git a/Engine/Source/Runtime/Core/Public/Containers/ChunkedArray.h b/Engine/Source/Runtime/Core/Public/Containers/ChunkedArray.h
--- a/Engine/Source/Runtime/Core/Public/Containers/ChunkedArray.h
+++ b/Engine/Source/Runtime/Core/Public/Containers/ChunkedArray.h
@@ -213,14 +213,14 @@ public:
 		checkSlow(NumElements>=0);
 
 		const int32 OldNum = NumElements;
-		for (int32 i = 0; i < Count; i++)
+		const int32 NewNumElements = OldNum + Count;
+		const int32 NewNumChunks = (NewNumElements + NumElementsPerChunk - 1)/NumElementsPerChunk;
+		NumElements = NewNumElements;
+		for (int32 NumChunks = Chunks.Num(); NumChunks < NewNumChunks; ++NumChunks)
 		{
-			if (NumElements % NumElementsPerChunk == 0)
-			{
-				Chunks.Add(new FChunk);
-			}
-			NumElements++;
+			Chunks.Add(new FChunk);
 		}
+
 		return OldNum;
 	}
 
diff --git a/Engine/Source/Runtime/Core/Tests/Containers/ChunkedArrayTest.cpp b/Engine/Source/Runtime/Core/Tests/Containers/ChunkedArrayTest.cpp
new file mode 100644
--- /dev/null
+++ b/Engine/Source/Runtime/Core/Tests/Containers/ChunkedArrayTest.cpp
@@ -0,0 +1,77 @@
+// Copyright Epic Games, Inc. All Rights Reserved.
+
+#if WITH_TESTS
+
+#include "Containers/ChunkedArray.h"
+#include "Tests/TestHarnessAdapter.h"
+
+namespace UE::Private
+{
+
+// Use TChunkArray of bytes in order to be able to easily specify the number of elements per chunk instead of bytes per chunk.
+template<uint32 ElementsPerChunk>
+class TChunkedArrayForAddTest : public TChunkedArray<uint8, ElementsPerChunk>
+{
+public:
+	int32 GetNumChunks() const
+	{
+		return this->Chunks.Num();
+	}
+};
+
+TEST_CASE_NAMED(TChunkedArrayAddAllocatesProperAmountOfChunksTest, "System::Core::ChunkedArray::Add", "[Core][ChunkedArray][SmokeFilter]")
+{
+	SECTION("One element per chunk")
+	{
+		constexpr uint32 ElementsPerChunk = 1;
+		const int32 NumElementsToAddTests[] = {0, 1, 2, 12345};
+		for (const int32 NumElementsToAdd : NumElementsToAddTests)
+		{
+			TChunkedArrayForAddTest<ElementsPerChunk> ChunkedArray;
+			ChunkedArray.Add(NumElementsToAdd);
+			REQUIRE(ChunkedArray.GetNumChunks() == NumElementsToAdd);
+			REQUIRE(ChunkedArray.Num() == NumElementsToAdd);
+		}
+	}
+
+	SECTION("Arbitrary num elements per chunk")
+	{
+		constexpr uint32 ElementsPerChunk = 47;
+		const int32 NumElementsToAddTests[] = {0, 1, ElementsPerChunk - 1, ElementsPerChunk, ElementsPerChunk + 1, 8999};
+		for (const int32 NumElementsToAdd : NumElementsToAddTests)
+		{
+			TChunkedArrayForAddTest<ElementsPerChunk> ChunkedArray;
+			ChunkedArray.Add(NumElementsToAdd);
+			const int32 ExpectedNumChunks = (NumElementsToAdd + ElementsPerChunk - 1)/ElementsPerChunk;
+			REQUIRE(ChunkedArray.GetNumChunks() == ExpectedNumChunks);
+			REQUIRE(ChunkedArray.Num() == NumElementsToAdd);
+		}
+	}
+
+	SECTION("Add to non-empty array")
+	{
+		constexpr uint32 ElementsPerChunk = 47;
+		const int32 NumElementsToAddTests[] = {0, 1, ElementsPerChunk - 1, ElementsPerChunk, ElementsPerChunk + 1, 8999};
+		for (const int32 NumElementsToStartWith : NumElementsToAddTests)
+		{
+			for (const int32 NumElementsToAdd : NumElementsToAddTests)
+			{
+				// Start off with an arbitrary number of elements
+				TChunkedArrayForAddTest<ElementsPerChunk> ChunkedArray;
+				ChunkedArray.Add(NumElementsToStartWith);
+
+				// Add additional elements
+				ChunkedArray.Add(NumElementsToAdd);
+
+				const int32 ExpectedNumElements = NumElementsToStartWith + NumElementsToAdd;
+				const int32 ExpectedNumChunks = (ExpectedNumElements + ElementsPerChunk - 1)/ElementsPerChunk;
+				REQUIRE(ChunkedArray.GetNumChunks() == ExpectedNumChunks);
+				REQUIRE(ChunkedArray.Num() == ExpectedNumElements);
+			}
+		}
+	}
+}
+
+}
+
+#endif // WITH_TESTS
