Added missing export

diff --git a/Engine/Source/Runtime/Core/Internal/IO/IoMemoryCache.h b/Engine/Source/Runtime/Core/Internal/IO/IoMemoryCache.h
--- a/Engine/Source/Runtime/Core/Internal/IO/IoMemoryCache.h
+++ b/Engine/Source/Runtime/Core/Internal/IO/IoMemoryCache.h
@@ -10,6 +10,6 @@
 
 class IIoCache;
 
-TUniquePtr<IIoCache> MakeMemoryIoCache(uint64 CacheSize);
+CORE_API TUniquePtr<IIoCache> MakeMemoryIoCache(uint64 CacheSize);
 
 #endif // !UE_BUILD_SHIPPING
