Adding a new macro to specify a command line parameter that will be used to set the name of the CustomConfig in the config hierarchy, even when ALLOW_INI_OVERRIDE_FROM_COMMANDLINE is false.

#rnx

diff --git a/Engine/Source/Runtime/Core/Private/Misc/ConfigCacheIni.cpp b/Engine/Source/Runtime/Core/Private/Misc/ConfigCacheIni.cpp
--- a/Engine/Source/Runtime/Core/Private/Misc/ConfigCacheIni.cpp
+++ b/Engine/Source/Runtime/Core/Private/Misc/ConfigCacheIni.cpp
@@ -3861,16 +3861,28 @@ const FString& FConfigCacheIni::GetCustomConfigString()
 		bInitialized = true;
 
 		// Set to compiled in value, then possibly override
+		bool bCustomConfigOverrideApplied = false;
 		CustomConfigString = TEXT(CUSTOM_CONFIG);
 
 #if ALLOW_INI_OVERRIDE_FROM_COMMANDLINE
 		if (FParse::Value(FCommandLine::Get(), CommandlineOverrideSpecifiers::CustomConfigIdentifier, CustomConfigString))
 		{
-			UE_LOG(LogConfig, Log, TEXT("Overriding CustomConfig from %s to %s"), TEXT(CUSTOM_CONFIG), *CustomConfigString);
+			bCustomConfigOverrideApplied = true;
+			UE_LOG(LogConfig, Log, TEXT("Overriding CustomConfig from %s to %s using -customconfig cmd line param"), TEXT(CUSTOM_CONFIG), *CustomConfigString);
+		}
+#endif
+
+#ifdef UE_USE_COMMAND_LINE_PARAM_FOR_CUSTOM_CONFIG
+		FString CustomName = PREPROCESSOR_TO_STRING(UE_USE_COMMAND_LINE_PARAM_FOR_CUSTOM_CONFIG);
+		if (!bCustomConfigOverrideApplied && FParse::Param(FCommandLine::Get(), *CustomName))
+		{
+			bCustomConfigOverrideApplied = true;
+			CustomConfigString = CustomName;
+			UE_LOG(LogConfig, Log, TEXT("Overriding CustomConfig from %s to %s using a custom cmd line param"), TEXT(CUSTOM_CONFIG), *CustomConfigString);
 		}
-		else
 #endif
-		if (!CustomConfigString.IsEmpty())
+
+		if (!bCustomConfigOverrideApplied && !CustomConfigString.IsEmpty())
 		{
 			UE_LOG(LogConfig, Log, TEXT("Using compiled CustomConfig %s"), *CustomConfigString);
 		}
