Allow overriding session guid for trace

diff --git a/Engine/Source/Runtime/Core/Private/ProfilingDebugging/TraceAuxiliary.cpp b/Engine/Source/Runtime/Core/Private/ProfilingDebugging/TraceAuxiliary.cpp
--- a/Engine/Source/Runtime/Core/Private/ProfilingDebugging/TraceAuxiliary.cpp
+++ b/Engine/Source/Runtime/Core/Private/ProfilingDebugging/TraceAuxiliary.cpp
@@ -1553,6 +1553,14 @@ void FTraceAuxiliary::Initialize(const TCHAR* CommandLine)
 	Desc.bUseWorkerThread = bShouldStartWorkerThread;
 	Desc.bUseImportantCache = (FParse::Param(CommandLine, TEXT("tracenocache")) == false);
 	Desc.OnConnectionFunc = &OnConnectionCallback;
+
+	FGuid SessionGuid;
+	if (!FParse::Value(CommandLine, TEXT("-tracesessionguid="), SessionGuid))
+	{
+		SessionGuid = FApp::GetInstanceId();
+	}
+	FMemory::Memcpy((FGuid&)Desc.SessionGuid, SessionGuid);
+	
 	if (FParse::Value(CommandLine, TEXT("-tracetailmb="), Desc.TailSizeBytes))
 	{
 		Desc.TailSizeBytes <<= 20;
