Made TDelegateBase::CreateDelegateInstance protected.

diff --git a/Engine/Source/Runtime/Core/Public/Delegates/DelegateBase.h b/Engine/Source/Runtime/Core/Public/Delegates/DelegateBase.h
--- a/Engine/Source/Runtime/Core/Public/Delegates/DelegateBase.h
+++ b/Engine/Source/Runtime/Core/Public/Delegates/DelegateBase.h
@@ -7,6 +7,7 @@
 #include "Math/UnrealMathUtility.h"
 #include "UObject/NameTypes.h"
 #include "Delegates/DelegateAccessHandler.h"
+#include "Delegates/DelegateInstancesImplFwd.h"
 #include "Delegates/DelegateSettings.h"
 #include "Delegates/IDelegateInstance.h"
 
@@ -87,6 +88,27 @@ private:
 	template <typename>
 	friend class TMulticastDelegateBase;
 
+	template <class, typename, typename, typename...>
+	friend class TBaseUFunctionDelegateInstance;
+
+	template <bool, class, ESPMode, typename, typename, typename...>
+	friend class TBaseSPMethodDelegateInstance;
+
+	template <bool, class, typename, typename, typename...>
+	friend class TBaseRawMethodDelegateInstance;
+
+	template <bool, class, typename, typename, typename...>
+	friend class TBaseUObjectMethodDelegateInstance;
+
+	template <typename, typename, typename...>
+	friend class TBaseStaticDelegateInstance;
+
+	template <typename, typename, typename, typename...>
+	friend class TBaseFunctorDelegateInstance;
+
+	template <typename, typename, typename, typename, typename...>
+	friend class TWeakBaseFunctorDelegateInstance;
+
 protected:
 	using typename Super::FReadAccessScope;
 	using Super::GetReadAccessScope;
@@ -252,6 +274,7 @@ public:
 		return DelegateInstance ? DelegateInstance->GetHandle() : FDelegateHandle{};
 	}
 
+protected:
 	/**
 	 * "emplacement" of delegate instance of the given type
 	 */
@@ -269,7 +292,6 @@ public:
 		new(Allocate(sizeof(DelegateInstanceType))) DelegateInstanceType(Forward<DelegateInstanceParams>(Params)...);
 	}
 
-protected:
 	/**
 	 * Gets the delegate instance.  Not intended for use by user code.
 	 *
diff --git a/Engine/Source/Runtime/Core/Public/Delegates/DelegateInstancesImpl.h b/Engine/Source/Runtime/Core/Public/Delegates/DelegateInstancesImpl.h
--- a/Engine/Source/Runtime/Core/Public/Delegates/DelegateInstancesImpl.h
+++ b/Engine/Source/Runtime/Core/Public/Delegates/DelegateInstancesImpl.h
@@ -1,7 +1,7 @@
 // Copyright Epic Games, Inc. All Rights Reserved.
 
 /*================================================================================
-	DelegateInstancesImpl.inl: Inline implementation of delegate bindings.
+	DelegateInstancesImpl.h: Inline implementation of delegate bindings.
 
 	The types declared in this file are for internal use only. 
 ================================================================================*/
@@ -12,6 +12,7 @@
 
 #include "CoreTypes.h"
 #include "Delegates/DelegateInstanceInterface.h"
+#include "Delegates/DelegateInstancesImplFwd.h"
 #include "Delegates/IDelegateInstance.h"
 #include "Misc/AssertionMacros.h"
 #include "Templates/RemoveReference.h"
@@ -22,7 +23,6 @@
 #include "UObject/WeakObjectPtrTemplates.h"
 
 class FDelegateHandle;
-enum class ESPMode : uint8;
 
 namespace UE::Delegates::Private
 {
@@ -54,13 +54,8 @@ protected:
 	FDelegateHandle Handle;
 };
 
-/**
- * Implements a delegate binding for UFunctions.
- *
- * @params UserClass Must be an UObject derived class.
- */
-template <class UserClass, typename FuncType, typename UserPolicy, typename... VarTypes>
-class TBaseUFunctionDelegateInstance;
+/* Delegate binding types
+ *****************************************************************************/
 
 template <class UserClass, typename RetValType, typename... ParamTypes, typename UserPolicy, typename... VarTypes>
 class TBaseUFunctionDelegateInstance<UserClass, RetValType(ParamTypes...), UserPolicy, VarTypes...> : public TCommonDelegateInstanceState<RetValType(ParamTypes...), UserPolicy, VarTypes...>
@@ -177,15 +172,6 @@ public:
 };
 
 
-/* Delegate binding types
- *****************************************************************************/
-
-/**
- * Implements a delegate binding for shared pointer member functions.
- */
-template <bool bConst, class UserClass, ESPMode SPMode, typename FuncType, typename UserPolicy, typename... VarTypes>
-class TBaseSPMethodDelegateInstance;
-
 template <bool bConst, class UserClass, ESPMode SPMode, typename RetValType, typename... ParamTypes, typename UserPolicy, typename... VarTypes>
 class TBaseSPMethodDelegateInstance<bConst, UserClass, SPMode, RetValType(ParamTypes...), UserPolicy, VarTypes...> : public TCommonDelegateInstanceState<RetValType(ParamTypes...), UserPolicy, VarTypes...>
 {
@@ -307,12 +293,6 @@ protected:
 };
 
 
-/**
- * Implements a delegate binding for C++ member functions.
- */
-template <bool bConst, class UserClass, typename FuncType, typename UserPolicy, typename... VarTypes>
-class TBaseRawMethodDelegateInstance;
-
 template <bool bConst, class UserClass, typename RetValType, typename... ParamTypes, typename UserPolicy, typename... VarTypes>
 class TBaseRawMethodDelegateInstance<bConst, UserClass, RetValType(ParamTypes...), UserPolicy, VarTypes...> : public TCommonDelegateInstanceState<RetValType(ParamTypes...), UserPolicy, VarTypes...>
 {
@@ -433,11 +413,6 @@ protected:
 	FMethodPtr MethodPtr;
 };
 
-/**
- * Implements a delegate binding for UObject methods.
- */
-template <bool bConst, class UserClass, typename FuncType, typename UserPolicy, typename... VarTypes>
-class TBaseUObjectMethodDelegateInstance;
 
 template <bool bConst, class UserClass, typename RetValType, typename... ParamTypes, typename UserPolicy, typename... VarTypes>
 class TBaseUObjectMethodDelegateInstance<bConst, UserClass, RetValType(ParamTypes...), UserPolicy, VarTypes...> : public TCommonDelegateInstanceState<RetValType(ParamTypes...), UserPolicy, VarTypes...>
@@ -564,12 +539,6 @@ protected:
 };
 
 
-/**
- * Implements a delegate binding for regular C++ functions.
- */
-template <typename FuncType, typename UserPolicy, typename... VarTypes>
-class TBaseStaticDelegateInstance;
-
 template <typename RetValType, typename... ParamTypes, typename UserPolicy, typename... VarTypes>
 class TBaseStaticDelegateInstance<RetValType(ParamTypes...), UserPolicy, VarTypes...> : public TCommonDelegateInstanceState<RetValType(ParamTypes...), UserPolicy, VarTypes...>
 {
@@ -664,11 +633,6 @@ private:
 	FFuncPtr StaticFuncPtr;
 };
 
-/**
- * Implements a delegate binding for C++ functors, e.g. lambdas.
- */
-template <typename FuncType, typename UserPolicy, typename FunctorType, typename... VarTypes>
-class TBaseFunctorDelegateInstance;
 
 template <typename RetValType, typename... ParamTypes, typename UserPolicy, typename FunctorType, typename... VarTypes>
 class TBaseFunctorDelegateInstance<RetValType(ParamTypes...), UserPolicy, FunctorType, VarTypes...> : public TCommonDelegateInstanceState<RetValType(ParamTypes...), UserPolicy, VarTypes...>
@@ -754,11 +718,6 @@ private:
 	mutable std::remove_const_t<FunctorType> Functor;
 };
 
-/**
- * Implements a weak object delegate binding for C++ functors, e.g. lambdas.
- */
-template <typename UserClass, typename FuncType, typename UserPolicy, typename FunctorType, typename... VarTypes>
-class TWeakBaseFunctorDelegateInstance;
 
 template <typename UserClass, typename RetValType, typename... ParamTypes, typename UserPolicy, typename FunctorType, typename... VarTypes>
 class TWeakBaseFunctorDelegateInstance<UserClass, RetValType(ParamTypes...), UserPolicy, FunctorType, VarTypes...> : public TCommonDelegateInstanceState<RetValType(ParamTypes...), UserPolicy, VarTypes...>
diff --git a/Engine/Source/Runtime/Core/Public/Delegates/DelegateInstancesImplFwd.h b/Engine/Source/Runtime/Core/Public/Delegates/DelegateInstancesImplFwd.h
new file mode 100644
--- /dev/null
+++ b/Engine/Source/Runtime/Core/Public/Delegates/DelegateInstancesImplFwd.h
@@ -0,0 +1,53 @@
+// Copyright Epic Games, Inc. All Rights Reserved.
+
+#pragma once
+
+// HEADER_UNIT_SKIP - Not included directly
+
+#include "CoreTypes.h"
+
+enum class ESPMode : uint8;
+
+/**
+ * Implements a delegate binding for UFunctions.
+ *
+ * @params UserClass Must be an UObject derived class.
+ */
+template <class UserClass, typename FuncType, typename UserPolicy, typename... VarTypes>
+class TBaseUFunctionDelegateInstance;
+
+/**
+ * Implements a delegate binding for shared pointer member functions.
+ */
+template <bool bConst, class UserClass, ESPMode SPMode, typename FuncType, typename UserPolicy, typename... VarTypes>
+class TBaseSPMethodDelegateInstance;
+
+/**
+ * Implements a delegate binding for C++ member functions.
+ */
+template <bool bConst, class UserClass, typename FuncType, typename UserPolicy, typename... VarTypes>
+class TBaseRawMethodDelegateInstance;
+
+/**
+ * Implements a delegate binding for UObject methods.
+ */
+template <bool bConst, class UserClass, typename FuncType, typename UserPolicy, typename... VarTypes>
+class TBaseUObjectMethodDelegateInstance;
+
+/**
+ * Implements a delegate binding for regular C++ functions.
+ */
+template <typename FuncType, typename UserPolicy, typename... VarTypes>
+class TBaseStaticDelegateInstance;
+
+/**
+ * Implements a delegate binding for C++ functors, e.g. lambdas.
+ */
+template <typename FuncType, typename UserPolicy, typename FunctorType, typename... VarTypes>
+class TBaseFunctorDelegateInstance;
+
+/**
+ * Implements a weak object delegate binding for C++ functors, e.g. lambdas.
+ */
+template <typename UserClass, typename FuncType, typename UserPolicy, typename FunctorType, typename... VarTypes>
+class TWeakBaseFunctorDelegateInstance;
