fixed tracing completion and destruction of old "fire and forgot" tasks

diff --git a/Engine/Source/Runtime/Core/Public/Async/TaskGraphInterfaces.h b/Engine/Source/Runtime/Core/Public/Async/TaskGraphInterfaces.h
--- a/Engine/Source/Runtime/Core/Public/Async/TaskGraphInterfaces.h
+++ b/Engine/Source/Runtime/Core/Public/Async/TaskGraphInterfaces.h
@@ -1274,6 +1274,12 @@ private:
 			FPlatformMisc::MemoryBarrier();
 			Subsequents->DispatchSubsequents(NewTasks, CurrentThread);
 		}
+		else
+		{
+			// "fire and forget" tasks don't have an accompanying FGraphEvent that traces completion and destruction
+			TaskTrace::Completed(GetTraceId());
+			TaskTrace::Destroyed(GetTraceId());
+		}
 
 		if (bDeleteOnCompletion)
 		{
