Move debug break in ensure macros into a macro body from inside the lambda

diff --git a/Engine/Source/Runtime/Core/Public/Misc/AssertionMacros.h b/Engine/Source/Runtime/Core/Public/Misc/AssertionMacros.h
--- a/Engine/Source/Runtime/Core/Public/Misc/AssertionMacros.h
+++ b/Engine/Source/Runtime/Core/Public/Misc/AssertionMacros.h
@@ -367,28 +367,20 @@ RetType FORCENOINLINE UE_DEBUG_SECTION DispatchCheckVerify(InnerType&& Inner, Ar
 		(LIKELY(!!(InExpression)) || (DispatchCheckVerify<bool>([Capture] () UE_DEBUG_SECTION \
 		{ \
 			static bool bExecuted = false; \
-			if (CheckVerifyImpl(bExecuted, Always, __FILE__, __LINE__, PLATFORM_RETURN_ADDRESS(), #InExpression, InFormat)) \
-			{ \
-				PLATFORM_BREAK(); \
-			} \
-			return false; \
-		})))
+			return CheckVerifyImpl(bExecuted, Always, __FILE__, __LINE__, PLATFORM_RETURN_ADDRESS(), #InExpression, InFormat); \
+		}) && [] () { PLATFORM_BREAK(); return false; } ()))
 
 	#define UE_ENSURE_IMPL2(Capture, Always, InExpression, ...) \
 		(LIKELY(!!(InExpression)) || (DispatchCheckVerify<bool>([Capture] () UE_DEBUG_SECTION \
 		{ \
 			static bool bExecuted = false; \
 			FValidateArgsInternal(__VA_ARGS__); \
-			if (CheckVerifyImpl(bExecuted, Always, __FILE__, __LINE__, PLATFORM_RETURN_ADDRESS(), #InExpression, ##__VA_ARGS__)) \
-			{ \
-				PLATFORM_BREAK(); \
-			} \
-			return false; \
-		})))
+			return CheckVerifyImpl(bExecuted, Always, __FILE__, __LINE__, PLATFORM_RETURN_ADDRESS(), #InExpression, ##__VA_ARGS__); \
+		}) && [] () { PLATFORM_BREAK(); return false; } ()))
 
-	#define ensure(           InExpression                ) UE_ENSURE_IMPL( , false, InExpression, TEXT(""))
+	#define ensure(           InExpression                ) UE_ENSURE_IMPL ( , false, InExpression, TEXT(""))
 	#define ensureMsgf(       InExpression, InFormat, ... ) UE_ENSURE_IMPL2(&, false, InExpression, InFormat, ##__VA_ARGS__)
-	#define ensureAlways(     InExpression                ) UE_ENSURE_IMPL( , true,  InExpression, TEXT(""))
+	#define ensureAlways(     InExpression                ) UE_ENSURE_IMPL ( , true,  InExpression, TEXT(""))
 	#define ensureAlwaysMsgf( InExpression, InFormat, ... ) UE_ENSURE_IMPL2(&, true,  InExpression, InFormat, ##__VA_ARGS__)
 
 #else	// DO_ENSURE
