Fixed mismatched copy constructor/assignment operator in TOptional<TNonNullPtr<OptionalType>>.

diff --git a/Engine/Source/Runtime/Core/Public/Templates/NonNullPointer.h b/Engine/Source/Runtime/Core/Public/Templates/NonNullPointer.h
--- a/Engine/Source/Runtime/Core/Public/Templates/NonNullPointer.h
+++ b/Engine/Source/Runtime/Core/Public/Templates/NonNullPointer.h
@@ -189,12 +189,6 @@ public:
 	{
 	}
 
-	TOptional& operator=(const TOptional& Other)
-	{
-		Pointer = Other.Pointer;
-		return *this;
-	}
-
 	TOptional& operator=(OptionalType* InPointer)
 	{
 		Pointer = InPointer;
