Move CPUPROFILERTRACE_FILE_AND_LINE_ENABLED to cpp file so it can be toggled without a full rebuild

diff --git a/Engine/Source/Runtime/Core/Private/ProfilingDebugging/CpuProfilerTrace.cpp b/Engine/Source/Runtime/Core/Private/ProfilingDebugging/CpuProfilerTrace.cpp
--- a/Engine/Source/Runtime/Core/Private/ProfilingDebugging/CpuProfilerTrace.cpp
+++ b/Engine/Source/Runtime/Core/Private/ProfilingDebugging/CpuProfilerTrace.cpp
@@ -14,6 +14,10 @@
 
 #if CPUPROFILERTRACE_ENABLED
 
+#if !defined(CPUPROFILERTRACE_FILE_AND_LINE_ENABLED)
+	#define CPUPROFILERTRACE_FILE_AND_LINE_ENABLED 1
+#endif
+
 UE_TRACE_CHANNEL_DEFINE(CpuChannel)
 
 UE_TRACE_EVENT_BEGIN(CpuProfiler, EventSpec, NoSync|Important)
diff --git a/Engine/Source/Runtime/Core/Public/ProfilingDebugging/CpuProfilerTrace.h b/Engine/Source/Runtime/Core/Public/ProfilingDebugging/CpuProfilerTrace.h
--- a/Engine/Source/Runtime/Core/Public/ProfilingDebugging/CpuProfilerTrace.h
+++ b/Engine/Source/Runtime/Core/Public/ProfilingDebugging/CpuProfilerTrace.h
@@ -19,9 +19,6 @@
 #endif
 #endif
 
-#if !defined(CPUPROFILERTRACE_FILE_AND_LINE_ENABLED)
-#define CPUPROFILERTRACE_FILE_AND_LINE_ENABLED 1
-#endif
 
 #if CPUPROFILERTRACE_ENABLED
 
