Mallocbinned2 for iOS/tvOS

diff --git a/Engine/Source/Runtime/Core/Private/Apple/ApplePlatformMemory.cpp b/Engine/Source/Runtime/Core/Private/Apple/ApplePlatformMemory.cpp
--- a/Engine/Source/Runtime/Core/Private/Apple/ApplePlatformMemory.cpp
+++ b/Engine/Source/Runtime/Core/Private/Apple/ApplePlatformMemory.cpp
@@ -24,7 +24,8 @@
 #include <sys/mman.h>
 #include <sys/mount.h>
 #include <objc/runtime.h>
-#if PLATFORM_IOS && defined(__IPHONE_13_0)
+#if PLATFORM_IOS
+#include "IOS/IOSPlatformMisc.h"
 #include <os/proc.h>
 #endif
 #include <CoreFoundation/CFBase.h>
@@ -302,53 +303,54 @@ FMalloc* FApplePlatformMemory::BaseAllocator()
 	FLowLevelMemTracker::Get().SetProgramSize(MemStats.UsedPhysical);
 #endif
 
-	if (FORCE_ANSI_ALLOCATOR)
-	{
-		AllocatorToUse = EMemoryAllocatorToUse::Ansi;
-	}
-	else if (USE_MALLOC_BINNED2)
-	{
-		AllocatorToUse = EMemoryAllocatorToUse::Binned2;
-	}
-	else
-	{
-		AllocatorToUse = EMemoryAllocatorToUse::Binned;
-	}
-	
-	// Force ansi malloc in some cases
-	if(getenv("UE4_FORCE_MALLOC_ANSI") != nullptr)
-	{
-		AllocatorToUse = EMemoryAllocatorToUse::Ansi;
-	}
-	
-	switch (AllocatorToUse)
-	{
-		case EMemoryAllocatorToUse::Ansi:
-			Instance = new FMallocAnsi();
-			break;
-
-		case EMemoryAllocatorToUse::Binned2:
-			Instance = new FMallocBinned2();
-			break;
-			
-		default:	// intentional fall-through
-		case EMemoryAllocatorToUse::Binned:
-		{
-			// get free memory
-			vm_statistics Stats;
-			mach_msg_type_number_t StatsSize = sizeof(Stats);
-			host_statistics(mach_host_self(), HOST_VM_INFO, (host_info_t)&Stats, &StatsSize);
-			// 1 << FMath::CeilLogTwo(MemoryConstants.TotalPhysical) should really be FMath::RoundUpToPowerOfTwo,
-			// but that overflows to 0 when MemoryConstants.TotalPhysical is close to 4GB, since CeilLogTwo returns 32
-			// this then causes the MemoryLimit to be 0 and crashing the app
-			uint64 MemoryLimit = FMath::Min<uint64>( uint64(1) << FMath::CeilLogTwo((Stats.free_count + Stats.inactive_count) * GetConstants().PageSize), 0x100000000);
-			
-			// [RCL] 2017-03-06 FIXME: perhaps BinnedPageSize should be used here, but leaving this change to the Mac platform owner.
-			Instance = new FMallocBinned((uint32)(GetConstants().PageSize&MAX_uint32), MemoryLimit);
-			break;
-		}
-	}
-	
+#if PLATFORM_IOS || PLATFORM_TVOS
+    else if (FIOSPlatformMisc::IsEntitlementEnabled("com.apple.developer.kernel.extended-virtual-addressing"))
+    {
+        AllocatorToUse = EMemoryAllocatorToUse::Binned2;
+#undef USE_MALLOC_BINNED2
+#define USE_MALLOC_BINNED2 1
+    }
+#elif USE_MALLOC_BINNED2
+    else if (USE_MALLOC_BINNED2)
+    {
+        AllocatorToUse = EMemoryAllocatorToUse::Binned2;
+    }
+#endif
+    else
+    {
+        AllocatorToUse = EMemoryAllocatorToUse::Ansi;
+    }
+    
+    switch (AllocatorToUse)
+    {
+        case EMemoryAllocatorToUse::Ansi:
+        {
+            Instance = new FMallocAnsi();
+            break;
+        }
+
+        case EMemoryAllocatorToUse::Binned2:
+        {
+            Instance = new FMallocBinned2();
+            break;
+        }
+
+        default:    // intentional fall-through
+        case EMemoryAllocatorToUse::Binned:
+        {
+            // get free memory
+            vm_statistics Stats;
+            mach_msg_type_number_t StatsSize = sizeof(Stats);
+            host_statistics(mach_host_self(), HOST_VM_INFO, (host_info_t)&Stats, &StatsSize);
+            // 1 << FMath::CeilLogTwo(MemoryConstants.TotalPhysical) should really be FMath::RoundUpToPowerOfTwo,
+            // but that overflows to 0 when MemoryConstants.TotalPhysical is close to 4GB, since CeilLogTwo returns 32
+            // this then causes the MemoryLimit to be 0 and crashing the app
+            uint64 MemoryLimit = FMath::Min<uint64>( uint64(1) << FMath::CeilLogTwo((Stats.free_count + Stats.inactive_count) * GetConstants().PageSize), 0x100000000);
+
+            // [RCL] 2017-03-06 FIXME: perhaps BinnedPageSize should be used here, but leaving this change to the Mac platform owner.
+            Instance = new FMallocBinned((uint32)(GetConstants().PageSize&MAX_uint32), MemoryLimit);
+        }
+    }
 	return Instance;
 }
 
diff --git a/Engine/Source/Runtime/Core/Private/IOS/IOSPlatformMisc.cpp b/Engine/Source/Runtime/Core/Private/IOS/IOSPlatformMisc.cpp
--- a/Engine/Source/Runtime/Core/Private/IOS/IOSPlatformMisc.cpp
+++ b/Engine/Source/Runtime/Core/Private/IOS/IOSPlatformMisc.cpp
@@ -47,6 +47,10 @@
 #include "Async/TaskGraphInterfaces.h"
 #include "Misc/CoreDelegates.h"
 
+
+#import <Foundation/Foundation.h>
+#import <mach-o/dyld.h>
+
 #if !defined ENABLE_ADVERTISING_IDENTIFIER
 	#define ENABLE_ADVERTISING_IDENTIFIER 0
 #endif
@@ -1256,6 +1260,154 @@ void FIOSPlatformMisc::UnregisterForRemoteNotifications()
 
 }
 
+
+// See for more information about the Blobs
+// https://opensource.apple.com/source/xnu/xnu-4570.61.1/osfmk/kern/cs_blobs.h.auto.html
+
+typedef struct __Blob {
+    uint32_t magic;
+    uint32_t length;
+    char data[];
+} CS_GenericBlob;
+
+typedef struct __BlobIndex {
+    uint32_t type;
+    uint32_t offset;
+} CS_BlobIndex;
+
+typedef struct __MultiBlob {
+    uint32_t magic;
+    uint32_t length;
+    uint32_t count;
+    CS_BlobIndex index[];
+} CS_MultiBlob;
+
+extern NSString *EntitlementsData(void)
+{
+    // iterate through the headers to find the executable, since only the executable has the entitlements
+    const struct mach_header_64 *executableHeader = nullptr;
+    char *ImageName = nullptr;
+    for (uint32_t i = 0; i < _dyld_image_count() && executableHeader == nullptr; i++)
+    {
+        const struct mach_header_64 *header = (struct mach_header_64 *)_dyld_get_image_header(i);
+        if (header->filetype == MH_EXECUTE)
+        {
+            ImageName = (char *)_dyld_get_image_name(i);
+            executableHeader = header;
+        }
+    }
+
+    if (executableHeader == nullptr)
+    {
+        return nil;
+    }
+
+    // verify that it's a 64bit app
+    if (executableHeader->magic != MH_MAGIC_64)
+    {
+        NSLog(@"Executable is NOT 64bit. Entitlement retrieval not supported.");
+        return nil;
+    }
+    uintptr_t cursor = (uintptr_t)executableHeader + sizeof(struct mach_header_64);
+    const struct linkedit_data_command *segmentCommand = NULL;
+
+    for (uint32_t i = 0; i < executableHeader->ncmds; i++, cursor += segmentCommand->cmdsize)
+    {
+        segmentCommand = (struct linkedit_data_command *)cursor;
+    
+        switch (segmentCommand->cmd)
+        {
+            case LC_CODE_SIGNATURE:
+                NSLog(@"LC_CODE_SIGNATURE found");
+                break;
+            default:
+                continue;
+        }
+
+        const struct linkedit_data_command *dataCommand = (const struct linkedit_data_command *)segmentCommand;
+
+        FILE *file = fopen(ImageName, "r");
+        if (file == NULL)
+        {
+            NSLog(@"Could not open binary file");
+            return nil;
+        }
+        CS_MultiBlob multiBlob;
+        int isSuccess = fseek(file, UInt64(dataCommand->dataoff), SEEK_SET);
+        size_t Count = fread(&multiBlob, sizeof(CS_MultiBlob), 1, file);
+        uint32 multiBlobCount = ntohl(multiBlob.count);
+        CS_BlobIndex multiBlobIndex[multiBlobCount];
+        fread(&multiBlobIndex[0], sizeof(CS_BlobIndex) * multiBlobCount, 1, file);
+        
+        if (__builtin_bswap32(multiBlob.magic) != 0xfade0cc0)
+        {
+            fclose(file);
+            return nil;
+        }
+        
+        uint32 multiBlobSize = sizeof(multiBlob);
+        uint32 blobSize = sizeof(CS_GenericBlob);
+
+        for (int j = 0; j < multiBlobCount; j++)
+        {
+            uint32_t currentOffset = dataCommand->dataoff;
+            uint32_t blobOffset = ntohl(multiBlobIndex[j].offset);
+            CS_GenericBlob blob;
+            isSuccess = fseek(file, currentOffset + blobOffset, SEEK_SET);
+            Count = fread(&blob, sizeof(CS_GenericBlob), 1, file);
+            
+
+            if (__builtin_bswap32(blob.magic) == 0xfade7171)
+            {
+                uint32 blobLength = ntohl(blob.length);
+                char data[blobLength];
+                fread(&data[0], sizeof(char) * blobLength, 1, file);
+                NSString *stringFromData = [NSString stringWithFormat: @"%s", data];
+                NSLog(@"%@", stringFromData);
+                fclose(file);
+                return stringFromData;
+            }
+            else
+            {
+                continue;
+            }
+        }
+        fclose(file);
+    }
+    return nil;
+}
+
+extern bool IsEntitlementPresentInEmbeddedProvision(NSString *EntitlementsToFind)
+{
+    NSString *mobileprovisionPath = [[[NSBundle mainBundle] bundlePath] stringByAppendingPathComponent:@"embedded.mobileprovision"];
+    NSError* error = nil;
+    NSString *myText = [NSString stringWithContentsOfFile:mobileprovisionPath encoding:NSUTF8StringEncoding error:&error];
+    
+    if ([myText rangeOfString: (@"%s", EntitlementsToFind)].location == NSNotFound)
+    {
+        return false;
+    }
+    else
+    {
+        return true;
+    }
+}
+
+bool FIOSPlatformMisc::IsEntitlementEnabled(const char * EntitlementToCheck)
+{
+    NSString *EntitlementsToFind = [NSString stringWithFormat: @"%s", EntitlementToCheck];
+    NSString *EntitlementData = EntitlementsData();
+    if ([EntitlementData rangeOfString: (@"%s", EntitlementsToFind)].location == NSNotFound)
+    {
+        NSLog(@"Entitlements not found in binary Mach-O header. Looking at the embedded mobile provision file.");
+        return IsEntitlementPresentInEmbeddedProvision(EntitlementsToFind);
+    }
+    else
+    {
+        return true;
+    }
+}
+
 void FIOSPlatformMisc::GetValidTargetPlatforms(TArray<FString>& TargetPlatformNames)
 {
 	// this is only used to cook with the proper TargetPlatform with COTF, it's not the runtime platform (which is just IOS for both)
diff --git a/Engine/Source/Runtime/Core/Public/IOS/IOSPlatformMisc.h b/Engine/Source/Runtime/Core/Public/IOS/IOSPlatformMisc.h
--- a/Engine/Source/Runtime/Core/Public/IOS/IOSPlatformMisc.h
+++ b/Engine/Source/Runtime/Core/Public/IOS/IOSPlatformMisc.h
@@ -73,6 +73,8 @@ struct CORE_API FIOSPlatformMisc : public FApplePlatformMisc
 	// Check if notifications are allowed if min iOS version is < 10
 	UE_DEPRECATED(4.21, "IsAllowedRemoteNotifications is deprecated. Use FIOSLocalNotificationService::CheckAllowedNotifications instead.")
 	static bool IsAllowedRemoteNotifications();
+    
+    static bool IsEntitlementEnabled(const char *EntitlementToCheck);
 	
 	static class IPlatformChunkInstall* GetPlatformChunkInstall();
 
