fix build

diff --git a/Engine/Source/Runtime/Core/Private/Apple/ApplePlatformMemory.cpp b/Engine/Source/Runtime/Core/Private/Apple/ApplePlatformMemory.cpp
--- a/Engine/Source/Runtime/Core/Private/Apple/ApplePlatformMemory.cpp
+++ b/Engine/Source/Runtime/Core/Private/Apple/ApplePlatformMemory.cpp
@@ -304,14 +304,14 @@ FMalloc* FApplePlatformMemory::BaseAllocator()
 #endif
 
 #if PLATFORM_IOS || PLATFORM_TVOS
-    else if (FIOSPlatformMisc::IsEntitlementEnabled("com.apple.developer.kernel.extended-virtual-addressing"))
+    if (FIOSPlatformMisc::IsEntitlementEnabled("com.apple.developer.kernel.extended-virtual-addressing"))
     {
         AllocatorToUse = EMemoryAllocatorToUse::Binned2;
 #undef USE_MALLOC_BINNED2
 #define USE_MALLOC_BINNED2 1
     }
 #elif USE_MALLOC_BINNED2
-    else if (USE_MALLOC_BINNED2)
+    if (USE_MALLOC_BINNED2)
     {
         AllocatorToUse = EMemoryAllocatorToUse::Binned2;
     }
