fixed waiting for task completion with timeout: if the smallest timeout value is used (1 msecs) it gets rounded down to zero because some time passed and so waiting just checks for completion w/o any waiting, which is not user wanted. changed to round up
also changed waiting for 1 msecs in unit-tests because it rarely can fail due to OS scheduling

diff --git a/Engine/Source/Runtime/Core/Private/Tasks/TaskPrivate.cpp b/Engine/Source/Runtime/Core/Private/Tasks/TaskPrivate.cpp
--- a/Engine/Source/Runtime/Core/Private/Tasks/TaskPrivate.cpp
+++ b/Engine/Source/Runtime/Core/Private/Tasks/TaskPrivate.cpp
@@ -226,9 +226,15 @@ namespace UE::Tasks
 				return true;
 			}
 
-			uint32 WaitForMsecs = Timeout == FTimeout::Never() ? 
-				MAX_uint32 :
-				(uint32)FMath::Clamp<int64>(Timeout.GetRemainingTime().GetTicks() / ETimespan::TicksPerMillisecond, 0, MAX_uint32);
+			uint32 WaitForMsecs = MAX_uint32;
+			if (Timeout != FTimeout::Never())
+			{
+				int64 RemainingTicks = Timeout.GetRemainingTime().GetTicks();
+				int64 RemainingMsecs = FMath::DivideAndRoundUp(RemainingTicks, ETimespan::TicksPerMillisecond);
+				int64 RemainingMsecsClamped = FMath::Clamp<int64>(RemainingMsecs, 0, MAX_uint32);
+				WaitForMsecs = (uint32)RemainingMsecsClamped;
+			}
+
 			return CompletionEvent->Wait(WaitForMsecs);
 		}
 
diff --git a/Engine/Source/Runtime/Core/Private/Tests/Tasks/TasksTest.cpp b/Engine/Source/Runtime/Core/Private/Tests/Tasks/TasksTest.cpp
--- a/Engine/Source/Runtime/Core/Private/Tests/Tasks/TasksTest.cpp
+++ b/Engine/Source/Runtime/Core/Private/Tests/Tasks/TasksTest.cpp
@@ -125,7 +125,7 @@ namespace UE { namespace TasksTests
 			FTask EmptyPrereq;
 			FTask NonEmptyPrereq = Launch(UE_SOURCE_LOCATION, [] {});
 			FTask Task = Launch(UE_SOURCE_LOCATION, [] {}, Prerequisites(EmptyPrereq, NonEmptyPrereq));
-			check(Task.Wait(FTimespan::FromMilliseconds(1.0)));
+			check(Task.Wait(FTimespan::FromMilliseconds(100)));
 		}
 
 		{	// check that movable-only result types are supported, that only single instance of result is created and that it's destroyed
@@ -1398,7 +1398,7 @@ namespace UE { namespace TasksTests
 		FPipe Pipe{ UE_SOURCE_LOCATION };
 		Pipe.Launch(UE_SOURCE_LOCATION, [] {});
 		FTask FinalTask = Pipe.Launch(UE_SOURCE_LOCATION, [] {});
-		check(FinalTask.Wait(FTimespan::FromMilliseconds(1.0)));
+		check(FinalTask.Wait(FTimespan::FromMilliseconds(100)));
 
 		return true;
 	}
diff --git a/Engine/Source/Runtime/Core/Public/Tasks/Task.h b/Engine/Source/Runtime/Core/Public/Tasks/Task.h
--- a/Engine/Source/Runtime/Core/Public/Tasks/Task.h
+++ b/Engine/Source/Runtime/Core/Public/Tasks/Task.h
@@ -341,9 +341,15 @@ namespace UE::Tasks
 			return true;
 		}
 
-		uint32 WaitForMsecs = Timeout == FTimeout::Never() ?
-			MAX_uint32 :
-			(uint32)FMath::Clamp<int64>(Timeout.GetRemainingTime().GetTicks() / ETimespan::TicksPerMillisecond, 0, MAX_uint32);
+		uint32 WaitForMsecs = MAX_uint32;
+		if (Timeout != FTimeout::Never())
+		{
+			int64 RemainingTicks = Timeout.GetRemainingTime().GetTicks();
+			int64 RemainingMsecs = FMath::DivideAndRoundUp(RemainingTicks, ETimespan::TicksPerMillisecond);
+			int64 RemainingMsecsClamped = FMath::Clamp<int64>(RemainingMsecs, 0, MAX_uint32);
+			WaitForMsecs = (uint32)RemainingMsecsClamped;
+		}
+
 		return CompletionEvent->Wait(WaitForMsecs);
 	}
 
