minor improvement to waiting with timeout for task completion

diff --git a/Engine/Source/Runtime/Core/Private/Tasks/TaskPrivate.cpp b/Engine/Source/Runtime/Core/Private/Tasks/TaskPrivate.cpp
--- a/Engine/Source/Runtime/Core/Private/Tasks/TaskPrivate.cpp
+++ b/Engine/Source/Runtime/Core/Private/Tasks/TaskPrivate.cpp
@@ -226,16 +226,7 @@ namespace UE::Tasks
 				return true;
 			}
 
-			uint32 WaitForMsecs = MAX_uint32;
-			if (Timeout != FTimeout::Never())
-			{
-				int64 RemainingTicks = Timeout.GetRemainingTime().GetTicks();
-				int64 RemainingMsecs = FMath::DivideAndRoundUp(RemainingTicks, ETimespan::TicksPerMillisecond);
-				int64 RemainingMsecsClamped = FMath::Clamp<int64>(RemainingMsecs, 0, MAX_uint32);
-				WaitForMsecs = (uint32)RemainingMsecsClamped;
-			}
-
-			return CompletionEvent->Wait(WaitForMsecs);
+			return CompletionEvent->Wait(Timeout.GetRemainingRoundedUpMilliseconds());
 		}
 
 		FTaskBase* FTaskBase::TryPushIntoPipe()
diff --git a/Engine/Source/Runtime/Core/Public/Misc/Timeout.h b/Engine/Source/Runtime/Core/Public/Misc/Timeout.h
--- a/Engine/Source/Runtime/Core/Public/Misc/Timeout.h
+++ b/Engine/Source/Runtime/Core/Public/Misc/Timeout.h
@@ -4,6 +4,7 @@
 
 #include "Misc/Timespan.h"
 #include "HAL/PlatformTime.h"
+#include "Math/UnrealMathUtility.h"
 #include "CoreTypes.h"
 
 namespace UE
@@ -47,6 +48,21 @@ namespace UE
 			return Timeout;
 		}
 
+		// intended for use in waiting functions, e.g. `FEvent::Wait()`
+		// returns the whole number (rounded up) of remaining milliseconds, clamped into [0, MAX_uint32] range
+		uint32 GetRemainingRoundedUpMilliseconds() const
+		{
+			if (*this == Never())
+			{
+				return MAX_uint32;
+			}
+
+			int64 RemainingTicks = GetRemainingTime().GetTicks();
+			int64 RemainingMsecs = FMath::DivideAndRoundUp(RemainingTicks, ETimespan::TicksPerMillisecond);
+			int64 RemainingMsecsClamped = FMath::Clamp<int64>(RemainingMsecs, 0, MAX_uint32);
+			return (uint32)RemainingMsecsClamped;
+		}
+
 		friend bool operator==(FTimeout Left, FTimeout Right)
 		{
 			return Left.Timeout == Right.Timeout && (Left.Timeout == FTimespan::MaxValue() || Left.Start == Right.Start);
diff --git a/Engine/Source/Runtime/Core/Public/Tasks/Task.h b/Engine/Source/Runtime/Core/Public/Tasks/Task.h
--- a/Engine/Source/Runtime/Core/Public/Tasks/Task.h
+++ b/Engine/Source/Runtime/Core/Public/Tasks/Task.h
@@ -341,16 +341,7 @@ namespace UE::Tasks
 			return true;
 		}
 
-		uint32 WaitForMsecs = MAX_uint32;
-		if (Timeout != FTimeout::Never())
-		{
-			int64 RemainingTicks = Timeout.GetRemainingTime().GetTicks();
-			int64 RemainingMsecs = FMath::DivideAndRoundUp(RemainingTicks, ETimespan::TicksPerMillisecond);
-			int64 RemainingMsecsClamped = FMath::Clamp<int64>(RemainingMsecs, 0, MAX_uint32);
-			WaitForMsecs = (uint32)RemainingMsecsClamped;
-		}
-
-		return CompletionEvent->Wait(WaitForMsecs);
+		return CompletionEvent->Wait(Timeout.GetRemainingRoundedUpMilliseconds());
 	}
 
 	// wait for multiple tasks while executing other tasks
