Translated tests from different directories from Engine\Source\Runtime\Core\Private\Tests

diff --git a/Engine/Source/Developer/LowLevelTestsRunner/Public/TestCommon/Comparisons.h b/Engine/Source/Developer/LowLevelTestsRunner/Public/TestCommon/Comparisons.h
deleted file mode 100644
diff --git a/Engine/Source/Developer/LowLevelTestsRunner/Public/TestCommon/Expectations.h b/Engine/Source/Developer/LowLevelTestsRunner/Public/TestCommon/Expectations.h
new file mode 100644
--- /dev/null
+++ b/Engine/Source/Developer/LowLevelTestsRunner/Public/TestCommon/Expectations.h
@@ -0,0 +1,262 @@
+// Copyright Epic Games, Inc. All Rights Reserved.
+
+#pragma once
+
+#include "CoreTypes.h"
+#include "Misc/CString.h"
+#include "Math/Color.h"
+#include "Math/MathFwd.h"
+#include "Math/Rotator.h"
+#include "Math/Transform.h"
+#include "Math/UnrealMathUtility.h"
+#include "Math/Vector.h"
+#include "Containers/UnrealString.h"
+#include "Containers/StringFwd.h"
+
+#include "../TestHarness.h"
+
+namespace {
+
+	bool TestTrue(const TCHAR* What, bool Value)
+	{
+		if (!Value)
+		{
+			FAIL_CHECK(FString::Printf(TEXT("Expected '%s' to be true."), What));
+			return false;
+		}
+		return true;
+	}
+
+	bool TestTrue(const FString& What, bool Value)
+	{
+		return TestTrue(*What, Value);
+	}
+
+	bool TestEqual(const TCHAR* What, const int32 Actual, const int32 Expected)
+	{
+		if (Actual != Expected)
+		{
+			FAIL_CHECK(FString::Printf(TEXT("Expected '%s' to be %d, but it was %d."), What, Expected, Actual));
+			return false;
+		}
+		return true;
+	}
+
+	bool TestEqual(const TCHAR* What, const int64 Actual, const int64 Expected)
+	{
+		if (Actual != Expected)
+		{
+			FAIL_CHECK(FString::Printf(TEXT("Expected '%s' to be %" "lld" ", but it was %" "lld" "."), What, Expected, Actual));
+			return false;
+		}
+		return true;
+	}
+
+#if PLATFORM_64BITS
+	bool TestEqual(const TCHAR* What, const SIZE_T Actual, const SIZE_T Expected)
+	{
+		if (Actual != Expected)
+		{
+			FAIL_CHECK(FString::Printf(TEXT("Expected '%s' to be %" "Iu" ", but it was %" "Iu" "."), What, Expected, Actual));
+			return false;
+		}
+		return true;
+	}
+#endif
+
+	bool TestEqual(const TCHAR* What, const float Actual, const float Expected, float Tolerance = UE_KINDA_SMALL_NUMBER)
+	{
+		if (!FMath::IsNearlyEqual(Actual, Expected, Tolerance))
+		{
+			FAIL_CHECK(FString::Printf(TEXT("Expected '%s' to be %f, but it was %f within tolerance %f."), What, Expected, Actual, Tolerance));
+			return false;
+		}
+		return true;
+	}
+
+	bool TestEqual(const TCHAR* What, const double Actual, const double Expected, double Tolerance = UE_KINDA_SMALL_NUMBER)
+	{
+		if (!FMath::IsNearlyEqual(Actual, Expected, Tolerance))
+		{
+			FAIL_CHECK(FString::Printf(TEXT("Expected '%s' to be %f, but it was %f within tolerance %f."), What, Expected, Actual, Tolerance));
+			return false;
+		}
+		return true;
+	}
+
+	bool TestEqual(const TCHAR* What, const FVector Actual, const FVector Expected, float Tolerance)
+	{
+		if (!Expected.Equals(Actual, Tolerance))
+		{
+			FAIL_CHECK(FString::Printf(TEXT("Expected '%s' to be %s, but it was %s within tolerance %f."), What, *Expected.ToString(), *Actual.ToString(), Tolerance));
+			return false;
+		}
+		return true;
+	}
+
+	bool TestEqual(const TCHAR* What, const FTransform Actual, const FTransform Expected, float Tolerance)
+	{
+		if (!Expected.Equals(Actual, Tolerance))
+		{
+			FAIL_CHECK(FString::Printf(TEXT("Expected '%s' to be %s, but it was %s within tolerance %f."), What, *Expected.ToString(), *Actual.ToString(), Tolerance));
+			return false;
+		}
+		return true;
+	}
+
+	bool TestEqual(const TCHAR* What, const FRotator Actual, const FRotator Expected, float Tolerance)
+	{
+		if (!Expected.Equals(Actual, Tolerance))
+		{
+			FAIL_CHECK(FString::Printf(TEXT("Expected '%s' to be %s, but it was %s within tolerance %f."), What, *Expected.ToString(), *Actual.ToString(), Tolerance));
+			return false;
+		}
+		return true;
+	}
+
+	bool TestEqual(const TCHAR* What, const FColor Actual, const FColor Expected)
+	{
+		if (Expected != Actual)
+		{
+			FAIL_CHECK(FString::Printf(TEXT("Expected '%s' to be %s, but it was %s."), What, *Expected.ToString(), *Actual.ToString()));
+			return false;
+		}
+		return true;
+	}
+
+	bool TestEqual(const TCHAR* What, const FLinearColor Actual, const FLinearColor Expected)
+	{
+		if (Expected != Actual)
+		{
+			FAIL_CHECK(FString::Printf(TEXT("Expected '%s' to be %s, but it was %s."), What, *Expected.ToString(), *Actual.ToString()));
+			return false;
+		}
+		return true;
+	}
+
+	bool TestEqual(const TCHAR* What, const TCHAR* Actual, const TCHAR* Expected)
+	{
+		if (FCString::Strcmp(Actual, Expected) != 0)
+		{
+			FAIL_CHECK(FString::Printf(TEXT("Expected '%s' to be \"%s\", but it was \"%s\"."), What, Expected, Actual));
+			return false;
+		}
+		return true;
+	}
+
+	bool TestEqualInsensitive(const TCHAR* What, const TCHAR* Actual, const TCHAR* Expected)
+	{
+		if (FCString::Stricmp(Actual, Expected) != 0)
+		{
+			FAIL_CHECK(FString::Printf(TEXT("Expected '%s' to be \"%s\", but it was \"%s\"."), What, Expected, Actual));
+			return false;
+		}
+		return true;
+	}
+
+	bool TestEqual(const FString& What, const int32 Actual, const int32 Expected)
+	{
+		return TestEqual(*What, Actual, Expected);
+	}
+
+	bool TestEqual(const FString& What, const float Actual, const float Expected, float Tolerance = UE_KINDA_SMALL_NUMBER)
+	{
+		return TestEqual(*What, Actual, Expected, Tolerance);
+	}
+
+	bool TestEqual(const FString& What, const double Actual, const double Expected, double Tolerance = UE_KINDA_SMALL_NUMBER)
+	{
+		return TestEqual(*What, Actual, Expected, Tolerance);
+	}
+
+	bool TestEqual(const FString& What, const FVector Actual, const FVector Expected, float Tolerance = UE_KINDA_SMALL_NUMBER)
+	{
+		return TestEqual(*What, Actual, Expected, Tolerance);
+	}
+
+	bool TestEqual(const FString& What, const FTransform Actual, const FTransform Expected, float Tolerance = UE_KINDA_SMALL_NUMBER)
+	{
+		return TestEqual(*What, Actual, Expected, Tolerance);
+	}
+
+	bool TestEqual(const FString& What, const FRotator Actual, const FRotator Expected, float Tolerance = UE_KINDA_SMALL_NUMBER)
+	{
+		return TestEqual(*What, Actual, Expected, Tolerance);
+	}
+
+	bool TestEqual(const FString& What, const FColor Actual, const FColor Expected)
+	{
+		return TestEqual(*What, Actual, Expected);
+	}
+
+	bool TestEqual(const FString& What, const TCHAR* Actual, const TCHAR* Expected)
+	{
+		return TestEqual(*What, Actual, Expected);
+	}
+
+	bool TestEqual(const TCHAR* What, const FString& Actual, const TCHAR* Expected)
+	{
+		return TestEqualInsensitive(What, *Actual, Expected);
+	}
+
+	bool TestEqual(const FString& What, const FString& Actual, const TCHAR* Expected)
+	{
+		return TestEqualInsensitive(*What, *Actual, Expected);
+	}
+
+	bool TestEqual(const TCHAR* What, const TCHAR* Actual, const FString& Expected)
+	{
+		return TestEqualInsensitive(What, Actual, *Expected);
+	}
+
+	bool TestEqual(const FString& What, const TCHAR* Actual, const FString& Expected)
+	{
+		return TestEqualInsensitive(*What, Actual, *Expected);
+	}
+
+	bool TestEqual(const TCHAR* What, const FString& Actual, const FString& Expected)
+	{
+		return TestEqualInsensitive(What, *Actual, *Expected);
+	}
+
+	bool TestEqual(const FString& What, const FString& Actual, const FString& Expected)
+	{
+		return TestEqualInsensitive(*What, *Actual, *Expected);
+	}
+
+	template<typename ValueType>
+	bool TestEqual(const TCHAR* What, const ValueType& Actual, const ValueType& Expected)
+	{
+		if (Actual != Expected)
+		{
+			FAIL_CHECK(FString::Printf(TEXT("%s: The two values are not equal."), What));
+			return false;
+		}
+		return true;
+	}
+
+	template<typename ValueType>
+	bool TestEqual(const FString& What, const ValueType& Actual, const ValueType& Expected)
+	{
+		return TestEqual(*What, Actual, Expected);
+	}
+
+	template<typename ValueType> bool TestNotEqual(const TCHAR* Description, const ValueType& Actual, const ValueType& Expected)
+	{
+		if (Actual == Expected)
+		{
+			FAIL_CHECK(FString::Printf(TEXT("%s: The two values are equal."), Description));
+			return false;
+		}
+		return true;
+	}
+
+	template<typename ValueType> bool TestNotEqual(const FString& Description, const ValueType& Actual, const ValueType& Expected)
+	{
+		return TestNotEqual(*Description, Actual, Expected);
+	}
+
+	#define CHECK_EQUALS(What, X, Y) TestEqual(What, X, Y);
+	#define CHECK_NOT_EQUALS(What, X, Y) TestNotEqual(What, X, Y);
+
+}
\ No newline at end of file
diff --git a/Engine/Source/Runtime/Core/Private/Tests/Async/AsyncTest.cpp b/Engine/Source/Runtime/Core/Private/Tests/Async/AsyncTest.cpp
deleted file mode 100644
diff --git a/Engine/Source/Runtime/Core/Private/Tests/Async/AsyncWorkTest.cpp b/Engine/Source/Runtime/Core/Private/Tests/Async/AsyncWorkTest.cpp
deleted file mode 100644
diff --git a/Engine/Source/Runtime/Core/Private/Tests/Async/SemaphoreTest.cpp b/Engine/Source/Runtime/Core/Private/Tests/Async/SemaphoreTest.cpp
deleted file mode 100644
diff --git a/Engine/Source/Runtime/Core/Private/Tests/Async/TaskGraphTest.cpp b/Engine/Source/Runtime/Core/Private/Tests/Async/TaskGraphTest.cpp
deleted file mode 100644
diff --git a/Engine/Source/Runtime/Core/Private/Tests/Delegates/TSDelegatesTest.cpp b/Engine/Source/Runtime/Core/Private/Tests/Delegates/TSDelegatesTest.cpp
deleted file mode 100644
diff --git a/Engine/Source/Runtime/Core/Private/Tests/Memory/SharedBufferTest.cpp b/Engine/Source/Runtime/Core/Private/Tests/Memory/SharedBufferTest.cpp
deleted file mode 100644
diff --git a/Engine/Source/Runtime/Core/Private/Tests/Serialization/CompactBinarySerializationTest.cpp b/Engine/Source/Runtime/Core/Private/Tests/Serialization/CompactBinarySerializationTest.cpp
deleted file mode 100644
diff --git a/Engine/Source/Runtime/Core/Private/Tests/Serialization/CompactBinaryValidationTest.cpp b/Engine/Source/Runtime/Core/Private/Tests/Serialization/CompactBinaryValidationTest.cpp
deleted file mode 100644
diff --git a/Engine/Source/Runtime/Core/Private/Tests/Serialization/CompactBinaryWriterTest.cpp b/Engine/Source/Runtime/Core/Private/Tests/Serialization/CompactBinaryWriterTest.cpp
deleted file mode 100644
diff --git a/Engine/Source/Runtime/Core/Private/Tests/Serialization/MemoryReaderTest.cpp b/Engine/Source/Runtime/Core/Private/Tests/Serialization/MemoryReaderTest.cpp
deleted file mode 100644
diff --git a/Engine/Source/Runtime/Core/Private/Tests/Serialization/MemoryWriterTest.cpp b/Engine/Source/Runtime/Core/Private/Tests/Serialization/MemoryWriterTest.cpp
deleted file mode 100644
diff --git a/Engine/Source/Runtime/Core/Private/Tests/Tasks/TasksTest.cpp b/Engine/Source/Runtime/Core/Private/Tests/Tasks/TasksTest.cpp
deleted file mode 100644
diff --git a/Engine/Source/Runtime/Core/Public/Misc/LowLevelTestAdapter.h b/Engine/Source/Runtime/Core/Public/Misc/LowLevelTestAdapter.h
--- a/Engine/Source/Runtime/Core/Public/Misc/LowLevelTestAdapter.h
+++ b/Engine/Source/Runtime/Core/Public/Misc/LowLevelTestAdapter.h
@@ -125,6 +125,8 @@ FString CaptureExpressionsAndValues(const FString& InExpressions, ArgTypes&&...
 #define REQUIRE_MESSAGE(Message, Expr) if (!(Expr)) { FAutomationTestFramework::Get().GetCurrentTest()->AddError(Message); return; }
 #define STATIC_REQUIRE(...) static_assert(__VA_ARGS__, #__VA_ARGS__);
 #define STATIC_CHECK(...) static_assert(__VA_ARGS__, #__VA_ARGS__);
+#define STATIC_CHECK_FALSE(...) static_assert(!(__VA_ARGS__), "!(" #__VA_ARGS__ ")");
+
 #define CHECK_EQUALS(What, X, Y) FAutomationTestFramework::Get().GetCurrentTest()->TestEqual(What, X, Y);
 #define CHECK_NOT_EQUALS(What, X, Y) FAutomationTestFramework::Get().GetCurrentTest()->TestNotEqual(What, X, Y);
 
diff --git a/Engine/Source/Runtime/Core/Tests/Async/AsyncTest.cpp b/Engine/Source/Runtime/Core/Tests/Async/AsyncTest.cpp
new file mode 100644
--- /dev/null
+++ b/Engine/Source/Runtime/Core/Tests/Async/AsyncTest.cpp
@@ -0,0 +1,94 @@
+// Copyright Epic Games, Inc. All Rights Reserved.
+
+#include "CoreTypes.h"
+#include "Templates/Function.h"
+#include "Tests/TestHarnessAdapter.h"
+#include "Async/Async.h"
+#if WITH_LOW_LEVEL_TESTS
+#include "TestCommon/Expectations.h"
+#endif
+
+#if WITH_TESTS
+
+/** Helper methods used in the test cases. */
+namespace AsyncTestUtils
+{
+	TFunction<int()> Task = [] {
+		return 123;
+	};
+
+	bool bHasVoidTaskFinished = false;
+
+	TFunction<void()> VoidTask = [] {
+		bHasVoidTaskFinished = true;
+	};
+}
+
+
+/** Test that task graph tasks return correctly. */
+TEST_CASE_NAMED(FAsyncGraphTest, "System::Core::Async::Async (Task Graph)", "[EditorContext][EngineFilter]")
+{
+	auto Future = Async(EAsyncExecution::TaskGraph, AsyncTestUtils::Task);
+	int Result = Future.Get();
+
+	CHECK_EQUALS(TEXT("Task graph task must return expected value"), Result, 123);
+}
+
+
+/** Test that threaded tasks return correctly. */
+TEST_CASE_NAMED(FAsyncThreadedTaskTest, "System::Core::Async::Async (Thread)", "[EditorContext][EngineFilter]")
+{
+	auto Future = Async(EAsyncExecution::Thread, AsyncTestUtils::Task);
+	int Result = Future.Get();
+
+	CHECK_EQUALS(TEXT("Threaded task must return expected value"), Result, 123);
+}
+
+
+/** Test that threaded pool tasks return correctly. */
+TEST_CASE_NAMED(FAsyncThreadedPoolTest, "System::Core::Async::Async (Thread Pool)", "[EditorContext][EngineFilter]")
+{
+	auto Future = Async(EAsyncExecution::ThreadPool, AsyncTestUtils::Task);
+	int Result = Future.Get();
+
+	CHECK_EQUALS(TEXT("Thread pool task must return expected value"), Result, 123);
+}
+
+
+/** Test that void tasks run without errors or warnings. */
+TEST_CASE_NAMED(FAsyncVoidTaskTest, "System::Core::Async::Async (Void)", "[EditorContext][EngineFilter]")
+{
+	// Reset test variable before running
+	AsyncTestUtils::bHasVoidTaskFinished = false;
+	auto Future = Async(EAsyncExecution::TaskGraph, AsyncTestUtils::VoidTask);
+	Future.Get();
+
+	// Check that the variable state was updated by task
+	CHECK_MESSAGE(TEXT("Void tasks should run"), AsyncTestUtils::bHasVoidTaskFinished);
+}
+
+
+/** Test that asynchronous tasks have their completion callback called. */
+TEST_CASE_NAMED(FAsyncCompletionCallbackTest, "System::Core::Async::Async (Completion Callback)", "[EditorContext][EngineFilter]")
+{
+	bool Completed = false;
+	FEvent* CompletedEvent = FPlatformProcess::GetSynchEventFromPool(true);
+
+	auto Future = Async(EAsyncExecution::TaskGraph, AsyncTestUtils::Task, [&] {
+		Completed = true;
+		CompletedEvent->Trigger();
+	});
+
+	int Result = Future.Get();
+
+	// We need an additional synchronization point here since the future get above will return after
+	// the task is done but before the completion callback has returned!
+
+	bool CompletedEventTriggered = CompletedEvent->Wait(FTimespan(0 /* hours */, 0 /* minutes */, 5 /* seconds */));
+	FPlatformProcess::ReturnSynchEventToPool(CompletedEvent);
+
+	CHECK_EQUALS(TEXT("Async Result"), Result, 123);
+	CHECK_MESSAGE(TEXT("Completion callback to be called"), CompletedEventTriggered && Completed);
+}
+
+#endif //WITH_TESTS
diff --git a/Engine/Source/Runtime/Core/Tests/Async/AsyncWorkTest.cpp b/Engine/Source/Runtime/Core/Tests/Async/AsyncWorkTest.cpp
new file mode 100644
--- /dev/null
+++ b/Engine/Source/Runtime/Core/Tests/Async/AsyncWorkTest.cpp
@@ -0,0 +1,135 @@
+// Copyright Epic Games, Inc. All Rights Reserved.
+
+#include "CoreTypes.h"
+#include "Misc/AutomationTest.h"
+#include "Async/AsyncWork.h"
+#include "Misc/QueuedThreadPool.h"
+#include "Tests/TestHarnessAdapter.h"
+
+#if WITH_TESTS
+
+namespace AsyncWorkTestImpl
+{
+	thread_local bool bIsOuterTaskRunning = false;
+
+	struct FInnerTask : public FNonAbandonableTask
+	{
+		void DoWork() 
+		{ 
+			TRACE_CPUPROFILER_EVENT_SCOPE(FInnerTask::DoWork);
+			FPlatformProcess::Sleep(0.01f); 
+		}
+
+		FORCEINLINE TStatId GetStatId() const
+		{
+			RETURN_QUICK_DECLARE_CYCLE_STAT(FInnerTask, STATGROUP_ThreadPoolAsyncTasks);
+		}
+	};
+
+	struct FOuterTask : public FNonAbandonableTask
+	{
+		FQueuedThreadPool* ThreadPool = nullptr;
+		int32 InnerTaskCount = 0;
+		void DoWork() 
+		{ 
+			TRACE_CPUPROFILER_EVENT_SCOPE(FOuterTask::DoWork);
+
+			bIsOuterTaskRunning = true;
+			TArray<FAsyncTask<FInnerTask>> InnerTasks;
+			InnerTasks.SetNum(InnerTaskCount);
+
+			for (int32 Index = 0; Index < InnerTasks.Num(); ++Index)
+			{
+				InnerTasks[Index].StartBackgroundTask(ThreadPool);
+			}
+
+			for (int32 Index = 0; Index < InnerTasks.Num(); ++Index)
+			{
+				TRACE_CPUPROFILER_EVENT_SCOPE(FOuterTask::EnsureCompletion);
+				InnerTasks[Index].EnsureCompletion();
+			}
+			bIsOuterTaskRunning = false;
+		}
+
+		FORCEINLINE TStatId GetStatId() const
+		{
+			RETURN_QUICK_DECLARE_CYCLE_STAT(FRootTask, STATGROUP_ThreadPoolAsyncTasks);
+		}
+	};
+
+	template <typename OtherTaskType>
+	struct FCompletionTask : public FNonAbandonableTask
+	{
+		OtherTaskType* OtherTask;
+		void DoWork() 
+		{ 
+			checkf(!bIsOuterTaskRunning, TEXT("We got picked by the busywait of the outer task we are going to wait on, we are now deadlocked"));
+
+			TRACE_CPUPROFILER_EVENT_SCOPE(FCompletionTask::DoWork);
+
+			// We need to wait a little bit before entering EnsureCompletion otherwise
+			// EnsureCompletion will busy wait into every completion task and there 
+			// will not be any completion tasks left for the outer task to pick up.
+			FPlatformProcess::Sleep(0.01f);
+
+			OtherTask->EnsureCompletion(); 
+		}
+
+		FORCEINLINE TStatId GetStatId() const
+		{
+			RETURN_QUICK_DECLARE_CYCLE_STAT(FCompletionTask, STATGROUP_ThreadPoolAsyncTasks);
+		}
+	};
+}
+
+TEST_CASE_NAMED(FAsyncWorkEnsureCompletionBusyWaitDeadLockTest, "System::Core::AsyncWork::EnsureCompletionBusyWaitDeadlockTest", "[.][EditorContext][EngineFilter][Disabled]")
+{
+	using namespace AsyncWorkTestImpl;
+	TRACE_CPUPROFILER_EVENT_SCOPE(FAsyncWorkEnsureCompletionBusyWaitDeadLockTest::RunTest);
+	
+	FQueuedThreadPool* ThreadPool = GThreadPool;
+#if WITH_EDITOR
+	ThreadPool = GLargeThreadPool;
+#endif
+
+	FAsyncTask<FOuterTask> OuterTask;
+	OuterTask.GetTask().ThreadPool = ThreadPool;
+	OuterTask.GetTask().InnerTaskCount = ThreadPool->GetNumThreads() * 2 * 100;
+
+	TArray<FAsyncTask<FCompletionTask<FAsyncTask<FOuterTask>>>> CompletionTasks;
+	CompletionTasks.SetNum(ThreadPool->GetNumThreads());
+	for (int32 Index = 0; Index < CompletionTasks.Num(); ++Index)
+	{
+		CompletionTasks[Index].GetTask().OtherTask = &OuterTask;
+	}
+
+	// Start the task that all the other tasks will try to complete as fast as possible
+	OuterTask.StartBackgroundTask(ThreadPool, EQueuedWorkPriority::Normal, EQueuedWorkFlags::None);
+
+	// Wait until the outer task had time to schedule it's own tasks
+	FPlatformProcess::Sleep(0.1f);
+
+	// Now fill the thread pool with tasks that simply want to complete the first one
+	for (int32 Index = 0; Index < CompletionTasks.Num(); ++Index)
+	{
+		EQueuedWorkFlags QueuedWorkFlags = EQueuedWorkFlags::None;
+
+		// Busy waiting on our task could cause a circular dependency between tasks on the same thread, which cannot be resolved and results into a deadlock.
+		// We need to tell the scheduler we want to opt-out of being picked up by busy wait.
+		// Comment this line to repro the deadlock scenario on the low-level task graph (will not reproduce on the old task graph as busywait doesn't exists).
+		QueuedWorkFlags |= EQueuedWorkFlags::DoNotRunInsideBusyWait;
+
+		// Use Highest priority so we have a chance to run before all the inner tasks are finished
+		CompletionTasks[Index].StartBackgroundTask(ThreadPool, EQueuedWorkPriority::Highest, QueuedWorkFlags);
+	}
+
+	// Do not contribute to any work here, we normally tick slate for the game-thread to stay responsive while we're waiting on tasks to complete
+	for (int32 Index = 0; Index < CompletionTasks.Num(); ++Index)
+	{
+		CompletionTasks[Index].EnsureCompletion(false);
+	}
+	OuterTask.EnsureCompletion(false);
+
+}
+
+#endif //WITH_TESTS
diff --git a/Engine/Source/Runtime/Core/Tests/Async/SemaphoreTest.cpp b/Engine/Source/Runtime/Core/Tests/Async/SemaphoreTest.cpp
new file mode 100644
--- /dev/null
+++ b/Engine/Source/Runtime/Core/Tests/Async/SemaphoreTest.cpp
@@ -0,0 +1,170 @@
+// Copyright Epic Games, Inc. All Rights Reserved.
+
+#include "CoreTypes.h"
+#include "Misc/AutomationTest.h"
+#include "Tests/Benchmark.h"
+#include "HAL/Thread.h"
+#include "HAL/UESemaphore.h"
+#include "HAL/Event.h"
+#include "HAL/PlatformProcess.h"
+#include "Containers/Array.h"
+#include "Templates/UniquePtr.h"
+#include "Tests/TestHarnessAdapter.h"
+
+#include <atomic>
+
+#if WITH_TESTS
+
+namespace UE::SemaphoreTests
+{
+	using namespace Tasks;
+
+	template<uint32 NumThreads, uint32 NumLoops, bool ReleaseAllAtOnce>
+	void TestSemaphoreWakeUpPerf()
+	{
+		TArray<FThread> Threads;
+		Threads.Reserve(NumThreads);
+		std::atomic<bool> bQuit{ false };
+
+		FSemaphore Semaphore{ 0, NumThreads };
+		std::atomic<uint32> Sync1{ 0 };
+		std::atomic<uint32> Sync2{ 0 };
+
+		for (uint32 i = 0; i < NumThreads; ++i)
+		{
+			Threads.Emplace(TEXT("SemaphoreWakeUpPerfTest"),
+				[&Semaphore, &Sync1, &Sync2, &bQuit]
+				{
+					while (!bQuit)
+					{
+						Semaphore.Acquire();
+						--Sync1;
+						while (Sync1 != 0)
+						{
+						}
+						--Sync2;
+					}
+				}
+				);
+		}
+
+		auto Iteration = [&Semaphore, &Sync1, &Sync2]
+		{
+			Sync2 = NumThreads;
+			Sync1 = NumThreads;
+
+			if constexpr (ReleaseAllAtOnce)
+			{
+				Semaphore.Release(NumThreads);
+			}
+			else
+			{
+				for (uint32 i = 0; i < NumThreads; ++i)
+				{
+					Semaphore.Release();
+				}
+			}
+
+			while (Sync2 != 0)
+			{
+			}
+			FPlatformProcess::Sleep(0.0);
+		};
+
+		for (uint32 i = 0; i < NumLoops; ++i)
+		{
+			Iteration();
+		}
+
+		Sync1 = NumThreads;
+		bQuit = true;
+
+		Semaphore.Release(NumThreads);
+		while (Sync1 != 0)
+		{
+		}
+
+		for (FThread& Thread : Threads)
+		{
+			Thread.Join();
+		}
+	}
+
+	template<uint32 NumThreads, uint32 NumLoops>
+	void TestMultiEventWakeUpPerf()
+	{
+		TArray<FThread> Threads;
+		Threads.Reserve(NumThreads);
+		std::atomic<bool> bQuit{ false };
+
+		FEventRef Events[NumThreads];
+
+		std::atomic<uint32> Sync1{ 0 };
+		std::atomic<uint32> Sync2{ 0 };
+
+		for (uint32 i = 0; i < NumThreads; ++i)
+		{
+			Threads.Emplace(TEXT("MultiEventWakeUpPerfTest"),
+				[Event = &Events[i], &Sync1, &Sync2, &bQuit]()
+				{
+					while (!bQuit)
+					{
+						Event->Get()->Wait();
+						--Sync1;
+						while (Sync1 != 0)
+						{
+						}
+						--Sync2;
+					}
+				}
+			);
+		}
+
+		auto Iteration = [&Events, &Sync1, &Sync2]
+		{
+			Sync2 = NumThreads;
+			Sync1 = NumThreads;
+
+			for (uint32 i = 0; i < NumThreads; ++i)
+			{
+				Events[i]->Trigger();
+			}
+
+			while (Sync2 != 0)
+			{
+			}
+			FPlatformProcess::Sleep(0.0);
+		};
+
+		for (uint32 i = 0; i != NumLoops; ++i)
+		{
+			Iteration();
+		}
+
+		Sync1 = NumThreads;
+		bQuit = true;
+
+		for (int i = 0; i != NumThreads; ++i)
+		{
+			Events[i]->Trigger();
+		}
+		while (Sync1 != 0)
+		{
+		}
+	
+		for (FThread& Thread : Threads)
+		{
+			Thread.Join();
+		}
+	}
+
+	// benchmark waking up multiple threads by releasing a semaphore and compare with triggering multiple auto-reset events
+	TEST_CASE_NAMED(FSemaphoreWakeUpPerfTest, "System::Core::Async::Semaphore::WakeUpPerf", "[.][ApplicationContextMask][EngineFilter]")
+	{
+		UE_BENCHMARK(5, TestSemaphoreWakeUpPerf<5, 100000, true>);
+		UE_BENCHMARK(5, TestSemaphoreWakeUpPerf<5, 100000, false>);
+		UE_BENCHMARK(5, TestMultiEventWakeUpPerf<5, 100000>);
+	}
+}
+
+#endif //WITH_TEST
diff --git a/Engine/Source/Runtime/Core/Tests/Async/TaskGraphTest.cpp b/Engine/Source/Runtime/Core/Tests/Async/TaskGraphTest.cpp
new file mode 100644
--- /dev/null
+++ b/Engine/Source/Runtime/Core/Tests/Async/TaskGraphTest.cpp
@@ -0,0 +1,1390 @@
+// Copyright Epic Games, Inc. All Rights Reserved.
+
+#include "CoreTypes.h"
+#include "Misc/Timespan.h"
+#include "Misc/DateTime.h"
+#include "HAL/PlatformProcess.h"
+#include "Stats/Stats.h"
+#include "Misc/AutomationTest.h"
+#include "Async/TaskGraphInterfaces.h"
+#include "Async/ParallelFor.h"
+#include "HAL/ThreadHeartBeat.h"
+#include "ProfilingDebugging/ScopedTimers.h"
+#include "Math/RandomStream.h"
+#include "Containers/CircularQueue.h"
+#include "Containers/Queue.h"
+#include "Tests/Benchmark.h"
+#include "HAL/Thread.h"
+#include "Async/Fundamental/Scheduler.h"
+#include "Tests/TestHarnessAdapter.h"
+
+#include <atomic>
+
+#if WITH_TESTS
+
+namespace OldTaskGraphTests
+{
+	static FORCEINLINE void DoWork(const void* Hash, FThreadSafeCounter& Counter, FThreadSafeCounter& Cycles, int32 Work)
+	{
+		if (Work > 0)
+		{
+			uint32 CyclesStart = FPlatformTime::Cycles();
+			Counter.Increment();
+			int32 Sum = 0;
+			for (int32 Index = 0; Index < Work; Index++)
+			{
+				Sum += PointerHash(((const uint64*)Hash) + Index);
+			}
+			Cycles.Add(FPlatformTime::Cycles() - CyclesStart + (Sum & 1));
+		}
+		else if (Work == 0)
+		{
+			Counter.Increment();
+		}
+	}
+
+	void PrintResult(double& StartTime, double& QueueTime, double& EndTime, FThreadSafeCounter& Counter, FThreadSafeCounter& Cycles, const TCHAR* Message)
+	{
+		UE_LOG(LogConsoleResponse, Display, TEXT("Total %6.3fms   %6.3fms queue   %6.3fms wait   %6.3fms work   : %s")
+			, float(1000.0 * (EndTime - StartTime)), float(1000.0 * (QueueTime - StartTime)), float(1000.0 * (EndTime - QueueTime)), float(FPlatformTime::GetSecondsPerCycle() * double(Cycles.GetValue()) * 1000.0)
+			, Message
+		);
+
+		Counter.Reset();
+		Cycles.Reset();
+		StartTime = 0.0;
+		QueueTime = 0.0;
+		EndTime = 0.0;
+	}
+
+	static void TaskGraphBenchmark(const TArray<FString>& Args)
+	{
+		FSlowHeartBeatScope SuspendHeartBeat;
+
+		double StartTime, QueueTime, EndTime;
+		FThreadSafeCounter Counter;
+		FThreadSafeCounter Cycles;
+
+		if (!FTaskGraphInterface::IsMultithread())
+		{
+			UE_LOG(LogConsoleResponse, Display, TEXT("WARNING: TaskGraphBenchmark disabled for non multi-threading platforms"));
+			return;
+		}
+
+		if (Args.Num() == 1 && Args[0] == TEXT("infinite"))
+		{
+			while (true)
+			{
+				{
+					StartTime = FPlatformTime::Seconds();
+
+					ParallelFor(1000,
+						[&Counter, &Cycles](int32 Index)
+						{
+							FFunctionGraphTask::CreateAndDispatchWhenReady(
+								[&Counter, &Cycles](ENamedThreads::Type CurrentThread, const FGraphEventRef& CompletionEvent)
+								{
+									DoWork(&CompletionEvent, Counter, Cycles, -1);
+								},
+								TStatId{}, nullptr, ENamedThreads::GameThread_Local
+							);
+						}
+					);
+					QueueTime = FPlatformTime::Seconds();
+					FTaskGraphInterface::Get().ProcessThreadUntilIdle(ENamedThreads::GameThread_Local);
+					EndTime = FPlatformTime::Seconds();
+				}
+			}
+		}
+		{
+			StartTime = FPlatformTime::Seconds();
+			FGraphEventArray Tasks;
+			Tasks.Reserve(1000);
+			for (int32 Index = 0; Index < 1000; Index++)
+			{
+				Tasks.Add(FFunctionGraphTask::CreateAndDispatchWhenReady([] {}, TStatId{}, nullptr, ENamedThreads::GameThread_Local));
+			}
+			QueueTime = FPlatformTime::Seconds();
+			FTaskGraphInterface::Get().WaitUntilTasksComplete(MoveTemp(Tasks), ENamedThreads::GameThread_Local);
+			EndTime = FPlatformTime::Seconds();
+		}
+		PrintResult(StartTime, QueueTime, EndTime, Counter, Cycles, TEXT("1000 tasks, ordinary local GT start"));
+		{
+			StartTime = FPlatformTime::Seconds();
+			FGraphEventArray Tasks;
+			Tasks.Reserve(1000);
+			for (int32 Index = 0; Index < 1000; Index++)
+			{
+				Tasks.Add(FFunctionGraphTask::CreateAndDispatchWhenReady(
+					[&Counter, &Cycles](ENamedThreads::Type CurrentThread, const FGraphEventRef& CompletionEvent)
+					{
+						DoWork(&CompletionEvent, Counter, Cycles, 100);
+					},
+					TStatId{}, nullptr, ENamedThreads::GameThread_Local
+				));
+			}
+			QueueTime = FPlatformTime::Seconds();
+			FTaskGraphInterface::Get().WaitUntilTasksComplete(MoveTemp(Tasks), ENamedThreads::GameThread_Local);
+			EndTime = FPlatformTime::Seconds();
+		}
+		PrintResult(StartTime, QueueTime, EndTime, Counter, Cycles, TEXT("1000 tasks, ordinary local GT start, with work"));
+		{
+			StartTime = FPlatformTime::Seconds();
+			FGraphEventArray Tasks;
+			Tasks.AddZeroed(1000);
+
+			ParallelFor(1000,
+				[&Tasks](int32 Index)
+				{
+					Tasks[Index] = FFunctionGraphTask::CreateAndDispatchWhenReady([] {});
+				}
+			);
+			QueueTime = FPlatformTime::Seconds();
+			FTaskGraphInterface::Get().WaitUntilTasksComplete(MoveTemp(Tasks), ENamedThreads::GameThread);
+			EndTime = FPlatformTime::Seconds();
+		}
+		PrintResult(StartTime, QueueTime, EndTime, Counter, Cycles, TEXT("1000 tasks, ParallelFor start"));
+		{
+			StartTime = FPlatformTime::Seconds();
+			FGraphEventArray Tasks;
+			Tasks.AddZeroed(10);
+
+			ParallelFor(10,
+				[&Tasks](int32 Index)
+				{
+					FGraphEventArray InnerTasks;
+					InnerTasks.AddZeroed(100);
+					for (int32 InnerIndex = 0; InnerIndex < 100; InnerIndex++)
+					{
+						InnerTasks[InnerIndex] = FFunctionGraphTask::CreateAndDispatchWhenReady([] {});
+					}
+					// join the above tasks
+					Tasks[Index] = FFunctionGraphTask::CreateAndDispatchWhenReady([] {}, TStatId{}, &InnerTasks, ENamedThreads::AnyThread);
+				}
+			);
+			QueueTime = FPlatformTime::Seconds();
+			FTaskGraphInterface::Get().WaitUntilTasksComplete(MoveTemp(Tasks), ENamedThreads::GameThread);
+			EndTime = FPlatformTime::Seconds();
+			PrintResult(StartTime, QueueTime, EndTime, Counter, Cycles, TEXT("1000 tasks, ParallelFor start, batched completion 10x100"));
+		}
+
+		{
+			StartTime = FPlatformTime::Seconds();
+			FGraphEventArray Tasks;
+			Tasks.AddZeroed(100);
+
+			ParallelFor(100,
+				[&Tasks](int32 Index)
+				{
+					FGraphEventArray InnerTasks;
+					InnerTasks.AddZeroed(10);
+					for (int32 InnerIndex = 0; InnerIndex < 10; InnerIndex++)
+					{
+						InnerTasks[InnerIndex] = FFunctionGraphTask::CreateAndDispatchWhenReady([] {});
+					}
+					// join the above tasks
+					Tasks[Index] = FFunctionGraphTask::CreateAndDispatchWhenReady([] {}, TStatId{}, &InnerTasks, ENamedThreads::AnyThread);
+				}
+			);
+			QueueTime = FPlatformTime::Seconds();
+			FTaskGraphInterface::Get().WaitUntilTasksComplete(MoveTemp(Tasks), ENamedThreads::GameThread);
+			EndTime = FPlatformTime::Seconds();
+		}
+		PrintResult(StartTime, QueueTime, EndTime, Counter, Cycles, TEXT("1000 tasks, ParallelFor start, batched completion 100x10"));
+
+		{
+			StartTime = FPlatformTime::Seconds();
+
+			ParallelFor(1000,
+				[&Counter, &Cycles](int32 Index)
+				{
+					FFunctionGraphTask::CreateAndDispatchWhenReady([&Counter, &Cycles](ENamedThreads::Type CurrentThread, const FGraphEventRef& CompletionEvent) { DoWork(&CompletionEvent, Counter, Cycles, 0); });
+				}
+			);
+			QueueTime = FPlatformTime::Seconds();
+			while (Counter.GetValue() < 1000)
+			{
+				FPlatformMisc::MemoryBarrier();
+			}
+			EndTime = FPlatformTime::Seconds();
+		}
+		PrintResult(StartTime, QueueTime, EndTime, Counter, Cycles, TEXT("1000 tasks, ParallelFor, counter tracking"));
+
+		{
+			StartTime = FPlatformTime::Seconds();
+
+			static bool Output[1000];
+			FPlatformMemory::Memzero(Output, 1000);
+
+			ParallelFor(1000,
+				[](int32 Index)
+				{
+					bool& Out = Output[Index];
+					FFunctionGraphTask::CreateAndDispatchWhenReady([&Out] { Out = true; });
+				}
+			);
+			QueueTime = FPlatformTime::Seconds();
+			for (int32 Index = 0; Index < 1000; Index++)
+			{
+				while (!Output[Index])
+				{
+					FPlatformProcess::Yield();
+				}
+			}
+			EndTime = FPlatformTime::Seconds();
+		}
+		PrintResult(StartTime, QueueTime, EndTime, Counter, Cycles, TEXT("1000 tasks, ParallelFor, bool* tracking"));
+
+		{
+			StartTime = FPlatformTime::Seconds();
+
+			ParallelFor(1000,
+				[&Counter, &Cycles](int32 Index)
+				{
+					FFunctionGraphTask::CreateAndDispatchWhenReady([&Counter, &Cycles](ENamedThreads::Type CurrentThread, const FGraphEventRef& CompletionEvent) { DoWork(&CompletionEvent, Counter, Cycles, 1000); });
+				}
+			);
+			QueueTime = FPlatformTime::Seconds();
+			while (Counter.GetValue() < 1000)
+			{
+				FPlatformProcess::Yield();
+			}
+			EndTime = FPlatformTime::Seconds();
+		}
+		PrintResult(StartTime, QueueTime, EndTime, Counter, Cycles, TEXT("1000 tasks, ParallelFor, counter tracking, with work"));
+		{
+			StartTime = FPlatformTime::Seconds();
+			for (int32 Index = 0; Index < 1000; Index++)
+			{
+				FFunctionGraphTask::CreateAndDispatchWhenReady([&Counter, &Cycles](ENamedThreads::Type CurrentThread, const FGraphEventRef& CompletionEvent) { DoWork(&CompletionEvent, Counter, Cycles, 1000); });
+			}
+			QueueTime = FPlatformTime::Seconds();
+			while (Counter.GetValue() < 1000)
+			{
+				FPlatformProcess::Yield();
+			}
+			EndTime = FPlatformTime::Seconds();
+		}
+		PrintResult(StartTime, QueueTime, EndTime, Counter, Cycles, TEXT("1000 tasks, GT submit, counter tracking, with work"));
+		{
+			StartTime = FPlatformTime::Seconds();
+
+			ParallelFor(1000,
+				[&Counter, &Cycles](int32 Index)
+				{
+					FFunctionGraphTask::CreateAndDispatchWhenReady(
+						[&Counter, &Cycles](ENamedThreads::Type CurrentThread, const FGraphEventRef& CompletionEvent)
+						{
+							DoWork(&CompletionEvent, Counter, Cycles, -1);
+						},
+						TStatId{}, nullptr, ENamedThreads::GameThread_Local
+					);
+				}
+			);
+			QueueTime = FPlatformTime::Seconds();
+			FTaskGraphInterface::Get().ProcessThreadUntilIdle(ENamedThreads::GameThread_Local);
+			EndTime = FPlatformTime::Seconds();
+		}
+		PrintResult(StartTime, QueueTime, EndTime, Counter, Cycles, TEXT("1000 local GT tasks, ParallelFor, no tracking (none needed)"));
+
+		{
+			StartTime = FPlatformTime::Seconds();
+			QueueTime = StartTime;
+			ParallelFor(1000,
+				[&Counter, &Cycles](int32 Index)
+				{
+					DoWork(&Counter, Counter, Cycles, -1);
+				}
+			);
+			EndTime = FPlatformTime::Seconds();
+		}
+		PrintResult(StartTime, QueueTime, EndTime, Counter, Cycles, TEXT("1000 element do-nothing ParallelFor"));
+		{
+			StartTime = FPlatformTime::Seconds();
+			QueueTime = StartTime;
+			ParallelFor(1000,
+				[&Counter, &Cycles](int32 Index)
+				{
+					DoWork(&Counter, Counter, Cycles, 1000);
+				}
+			);
+			EndTime = FPlatformTime::Seconds();
+		}
+		PrintResult(StartTime, QueueTime, EndTime, Counter, Cycles, TEXT("1000 element ParallelFor, with work"));
+
+		{
+			StartTime = FPlatformTime::Seconds();
+			QueueTime = StartTime;
+			ParallelFor(1000,
+				[&Counter, &Cycles](int32 Index)
+				{
+					DoWork(&Counter, Counter, Cycles, 1000);
+				},
+				true
+			);
+			EndTime = FPlatformTime::Seconds();
+		}
+		PrintResult(StartTime, QueueTime, EndTime, Counter, Cycles, TEXT("1000 element ParallelFor, single threaded, with work"));
+	}
+
+	static FAutoConsoleCommand TaskGraphBenchmarkCmd(
+		TEXT("TaskGraph.Benchmark"),
+		TEXT("Prints the time to run 1000 no-op tasks."),
+		FConsoleCommandWithArgsDelegate::CreateStatic(&TaskGraphBenchmark)
+	);
+
+	struct FTestStruct
+	{
+		int32 Index;
+		int32 Constant;
+		FTestStruct(int32 InIndex)
+			: Index(InIndex)
+			, Constant(0xfe05abcd)
+		{
+		}
+	};
+
+	struct FTestRigFIFO
+	{
+		FLockFreePointerFIFOBase<FTestStruct, PLATFORM_CACHE_LINE_SIZE> Test1;
+		FLockFreePointerFIFOBase<FTestStruct, 8> Test2;
+		FLockFreePointerFIFOBase<FTestStruct, 8, 1 << 4> Test3;
+	};
+
+	struct FTestRigLIFO
+	{
+		FLockFreePointerListLIFOBase<FTestStruct, PLATFORM_CACHE_LINE_SIZE> Test1;
+		FLockFreePointerListLIFOBase<FTestStruct, 8> Test2;
+		FLockFreePointerListLIFOBase<FTestStruct, 8, 1 << 4> Test3;
+	};
+
+	static void TestLockFree(int32 OuterIters = 3)
+	{
+		FSlowHeartBeatScope SuspendHeartBeat;
+
+
+		if (!FTaskGraphInterface::IsMultithread())
+		{
+			UE_LOG(LogConsoleResponse, Display, TEXT("WARNING: TestLockFree disabled for non multi-threading platforms"));
+			return;
+		}
+
+		const int32 NumWorkers = FTaskGraphInterface::Get().GetNumWorkerThreads();
+		// If we have too many threads active at once, they become too slow due to contention.  Set a reasonable maximum for how many are required to guarantee correctness of our LockFreePointers.
+		const int32 MaxWorkersForTest = 5;
+		const int32 MinWorkersForTest = 2; // With less than two threads we're not testing threading at all, so the test is pointless.
+		if (NumWorkers < MinWorkersForTest)
+		{
+			UE_LOG(LogConsoleResponse, Display, TEXT("WARNING: TestLockFree disabled for current machine because of not enough worker threads.  Need %d, have %d."), MinWorkersForTest, NumWorkers);
+			return;
+		}
+
+		FScopedDurationTimeLogger DurationLogger(TEXT("TestLockFree Runtime"));
+		const uint32 NumWorkersForTest = static_cast<uint32>(FMath::Clamp(NumWorkers, MinWorkersForTest, MaxWorkersForTest));
+		auto RunWorkersSynchronous = [NumWorkersForTest](const TFunction<void(uint32)>& WorkerTask)
+		{
+			FGraphEventArray Tasks;
+			for (uint32 Index = 0; Index < NumWorkersForTest; Index++)
+			{
+				TUniqueFunction<void()> WorkerTaskWithIndex{ [Index, &WorkerTask] { WorkerTask(Index); } };
+				Tasks.Add(FFunctionGraphTask::CreateAndDispatchWhenReady(MoveTemp(WorkerTaskWithIndex), TStatId{}, nullptr, ENamedThreads::AnyNormalThreadHiPriTask));
+			}
+			FTaskGraphInterface::Get().WaitUntilTasksComplete(MoveTemp(Tasks));
+		};
+
+		for (int32 Iter = 0; Iter < OuterIters; Iter++)
+		{
+			{
+				UE_LOG(LogTemp, Display, TEXT("******************************* Iter FIFO %d"), Iter);
+				FTestRigFIFO Rig;
+				for (int32 Index = 0; Index < 1000; Index++)
+				{
+					Rig.Test1.Push(new FTestStruct(Index));
+				}
+				TFunction<void(uint32)> Broadcast =
+					[&Rig](uint32 WorkerIndex)
+				{
+					FRandomStream Stream(((int32)WorkerIndex) * 7 + 13);
+					for (int32 Index = 0; Index < 1000000; Index++)
+					{
+						if (Index % 200000 == 1)
+						{
+							//UE_LOG(LogTemp, Log, TEXT("%8d iters thread=%d"), Index, int32(WorkerIndex));
+						}
+						if (Stream.FRand() < .03f)
+						{
+							TArray<FTestStruct*> Items;
+							{
+								float r = Stream.FRand();
+								if (r < .33f)
+								{
+									Rig.Test1.PopAll(Items);
+								}
+								else if (r < .66f)
+								{
+									Rig.Test2.PopAll(Items);
+								}
+								else
+								{
+									Rig.Test3.PopAll(Items);
+								}
+							}
+							for (FTestStruct* Item : Items)
+							{
+								float r = Stream.FRand();
+								if (r < .33f)
+								{
+									Rig.Test1.Push(Item);
+								}
+								else if (r < .66f)
+								{
+									Rig.Test2.Push(Item);
+								}
+								else
+								{
+									Rig.Test3.Push(Item);
+								}
+							}
+						}
+						else
+						{
+							FTestStruct* Item;
+							{
+								float r = Stream.FRand();
+								if (r < .33f)
+								{
+									Item = Rig.Test1.Pop();
+								}
+								else if (r < .66f)
+								{
+									Item = Rig.Test2.Pop();
+								}
+								else
+								{
+									Item = Rig.Test3.Pop();
+								}
+							}
+							if (Item)
+							{
+								float r = Stream.FRand();
+								if (r < .33f)
+								{
+									Rig.Test1.Push(Item);
+								}
+								else if (r < .66f)
+								{
+									Rig.Test2.Push(Item);
+								}
+								else
+								{
+									Rig.Test3.Push(Item);
+								}
+							}
+						}
+					}
+				};
+				RunWorkersSynchronous(Broadcast);
+
+				TArray<FTestStruct*> Items;
+				Rig.Test1.PopAll(Items);
+				Rig.Test2.PopAll(Items);
+				Rig.Test3.PopAll(Items);
+
+				checkf(Items.Num() == 1000, TEXT("Items %d"), Items.Num());
+
+				for (int32 LookFor = 0; LookFor < 1000; LookFor++)
+				{
+					bool bFound = false;
+					for (int32 Index = 0; Index < 1000; Index++)
+					{
+						if (Items[Index]->Index == LookFor && Items[Index]->Constant == 0xfe05abcd)
+						{
+							check(!bFound);
+							bFound = true;
+						}
+					}
+					check(bFound);
+				}
+				for (FTestStruct* Item : Items)
+				{
+					delete Item;
+				}
+
+				UE_LOG(LogTemp, Display, TEXT("******************************* Pass FTestRigFIFO"));
+
+			}
+			{
+				UE_LOG(LogTemp, Display, TEXT("******************************* Iter LIFO %d"), Iter);
+				FTestRigLIFO Rig;
+				for (int32 Index = 0; Index < 1000; Index++)
+				{
+					Rig.Test1.Push(new FTestStruct(Index));
+				}
+				TFunction<void(uint32)> Broadcast =
+					[&Rig](uint32 WorkerIndex)
+				{
+					FRandomStream Stream(((int32)WorkerIndex) * 7 + 13);
+					for (int32 Index = 0; Index < 1000000; Index++)
+					{
+						if (Index % 200000 == 1)
+						{
+							//UE_LOG(LogTemp, Log, TEXT("%8d iters thread=%d"), Index, int32(WorkerIndex));
+						}
+						if (Stream.FRand() < .03f)
+						{
+							TArray<FTestStruct*> Items;
+							{
+								float r = Stream.FRand();
+								if (r < .33f)
+								{
+									Rig.Test1.PopAll(Items);
+								}
+								else if (r < .66f)
+								{
+									Rig.Test2.PopAll(Items);
+								}
+								else
+								{
+									Rig.Test3.PopAll(Items);
+								}
+							}
+							for (FTestStruct* Item : Items)
+							{
+								float r = Stream.FRand();
+								if (r < .33f)
+								{
+									Rig.Test1.Push(Item);
+								}
+								else if (r < .66f)
+								{
+									Rig.Test2.Push(Item);
+								}
+								else
+								{
+									Rig.Test3.Push(Item);
+								}
+							}
+						}
+						else
+						{
+							FTestStruct* Item;
+							{
+								float r = Stream.FRand();
+								if (r < .33f)
+								{
+									Item = Rig.Test1.Pop();
+								}
+								else if (r < .66f)
+								{
+									Item = Rig.Test2.Pop();
+								}
+								else
+								{
+									Item = Rig.Test3.Pop();
+								}
+							}
+							if (Item)
+							{
+								float r = Stream.FRand();
+								if (r < .33f)
+								{
+									Rig.Test1.Push(Item);
+								}
+								else if (r < .66f)
+								{
+									Rig.Test2.Push(Item);
+								}
+								else
+								{
+									Rig.Test3.Push(Item);
+								}
+							}
+						}
+					}
+				};
+				RunWorkersSynchronous(Broadcast);
+
+				TArray<FTestStruct*> Items;
+				Rig.Test1.PopAll(Items);
+				Rig.Test2.PopAll(Items);
+				Rig.Test3.PopAll(Items);
+
+				checkf(Items.Num() == 1000, TEXT("Items %d"), Items.Num());
+
+				for (int32 LookFor = 0; LookFor < 1000; LookFor++)
+				{
+					bool bFound = false;
+					for (int32 Index = 0; Index < 1000; Index++)
+					{
+						if (Items[Index]->Index == LookFor && Items[Index]->Constant == 0xfe05abcd)
+						{
+							check(!bFound);
+							bFound = true;
+						}
+					}
+					check(bFound);
+				}
+				for (FTestStruct* Item : Items)
+				{
+					delete Item;
+				}
+
+				UE_LOG(LogTemp, Display, TEXT("******************************* Pass FTestRigLIFO"));
+
+			}
+		}
+	}
+
+	static void TestLockFree(const TArray<FString>& Args)
+	{
+		TestLockFree(10);
+	}
+
+	static FAutoConsoleCommand TestLockFreeCmd(
+		TEXT("TaskGraph.TestLockFree"),
+		TEXT("Test lock free lists"),
+		FConsoleCommandWithArgsDelegate::CreateStatic(&TestLockFree)
+	);
+
+	static void TestLowToHighPri(const TArray<FString>& Args)
+	{
+		UE_LOG(LogTemp, Display, TEXT("Starting latency test...."));
+
+		auto ForegroundTask = [](uint64 StartCycles)
+		{
+			float Latency = float(double(FPlatformTime::Cycles64() - StartCycles) * FPlatformTime::GetSecondsPerCycle64() * 1000.0 * 1000.0);
+			//UE_LOG(LogTemp, Display, TEXT("Latency %6.2fus"), Latency);
+			FPlatformMisc::LowLevelOutputDebugStringf(TEXT("Latency %6.2fus\r\n"), Latency);
+		};
+
+		auto BackgroundTask = [&ForegroundTask](ENamedThreads::Type CurrentThread, const FGraphEventRef& CompletionEvent)
+		{
+			while (true)
+			{
+				uint32 RunningCrc = 0;
+				for (int32 Index = 0; Index < 1000000; Index++)
+				{
+					FCrc::MemCrc32(CompletionEvent.GetReference(), sizeof(FGraphEvent), RunningCrc);
+				}
+				uint64 StartTime = FPlatformTime::Cycles64();
+				FFunctionGraphTask::CreateAndDispatchWhenReady([StartTime, &ForegroundTask] { ForegroundTask(StartTime); }, TStatId{}, nullptr, ENamedThreads::AnyHiPriThreadHiPriTask);
+			}
+		};
+
+#if 0
+		const int NumBackgroundTasks = 32;
+		const int NumNormalTasks = 32;
+		for (int32 Index = 0; Index < NumBackgroundTasks; Index++)
+		{
+			FFunctionGraphTask::CreateAndDispatchWhenReady(BackgroundTask, TStatId{}, nullptr, ENamedThreads::AnyNormalThreadNormalTask);
+		}
+		for (int32 Index = 0; Index < NumNormalTasks; Index++)
+		{
+			FFunctionGraphTask::CreateAndDispatchWhenReady(BackgroundTask, TStatId{}, nullptr, ENamedThreads::AnyBackgroundThreadNormalTask);
+		}
+		while (true)
+		{
+			FPlatformProcess::Sleep(25.0f);
+		}
+#else
+		FFunctionGraphTask::CreateAndDispatchWhenReady(MoveTemp(BackgroundTask), TStatId{}, nullptr, ENamedThreads::AnyBackgroundThreadNormalTask);
+#endif
+	}
+
+	static FAutoConsoleCommand TestLowToHighPriCmd(
+		TEXT("TaskGraph.TestLowToHighPri"),
+		TEXT("Test latency of high priority tasks when low priority tasks are saturating the CPU"),
+		FConsoleCommandWithArgsDelegate::CreateStatic(&TestLowToHighPri)
+	);
+
+
+	TEST_CASE_NAMED(FTaskGraphOldBenchmark, "System::Core::Async::TaskGraph::OldBenchmark", "[EditorContext][ClientContext][ServerContext][EngineFilter]")
+	{
+		TArray<FString> Args;
+		TaskGraphBenchmark(Args);
+	}
+
+	TEST_CASE_NAMED(FLockFreeTest, "System::Core::Async::TaskGraph::LockFree", "[.][ApplicationContextMask][EngineFilter]")
+	{
+		TestLockFree(3);
+	}
+}
+
+extern int32 GNumForegroundWorkers;
+
+namespace TaskGraphTests
+{
+	TEST_CASE_NAMED(FTaskGraphGraphEventTest, "System::Core::Async::TaskGraph::GraphEventTest", "[.][ApplicationContextMask][EngineFilter][Disabled]")
+	{
+		{	// task completes before it's waited for
+			FGraphEventRef Event = FFunctionGraphTask::CreateAndDispatchWhenReady(
+				[]
+				{
+					//UE_LOG(LogTemp, Log, TEXT("Main task"));
+				}
+			);
+			while (!Event->IsComplete() && FTaskGraphInterface::Get().GetNumWorkerThreads() != 0) // in single-threaded mode tasks are executed only when waited for
+			{}
+			Event->Wait(ENamedThreads::GameThread);
+		}
+
+		{	// task completes after it's waited for
+			FGraphEventRef Event = FFunctionGraphTask::CreateAndDispatchWhenReady([]()
+				{
+					//UE_LOG(LogTemp, Log, TEXT("Main task"));
+					FPlatformProcess::Sleep(0.1f); // pause for a bit to let waiting start
+				}
+			);
+			check(!Event->IsComplete());
+			Event->Wait(ENamedThreads::GameThread);
+		}
+
+		{	// event w/o a task, signaled by explicit call to DispatchSubsequents before it's waited for
+			FGraphEventRef Event = FGraphEvent::CreateGraphEvent();
+			FFunctionGraphTask::CreateAndDispatchWhenReady(
+				[&Event]
+				{
+					Event->DispatchSubsequents();
+				}
+			);
+			while (!Event->IsComplete() && FTaskGraphInterface::Get().GetNumWorkerThreads() != 0) // in single-threaded mode tasks are executed only when waited for
+			{}
+			Event->Wait(ENamedThreads::GameThread);
+		}
+
+		{	// event w/o a task, signaled by explicit call to DispatchSubsequents after it's waited for
+			FGraphEventRef Event = FGraphEvent::CreateGraphEvent();
+			auto Lambda = [&Event]
+				{
+					FPlatformProcess::Sleep(0.1f); // pause for a bit to let waiting start
+					Event->DispatchSubsequents();
+			};
+			FGraphEventRef Task = FFunctionGraphTask::CreateAndDispatchWhenReady(MoveTemp(Lambda));
+			check(!Event->IsComplete());
+			Event->Wait();
+			Task->Wait();
+		}
+
+		{	// wait for prereq by DontCompleteUntil
+			auto Lambda = [](ENamedThreads::Type CurrentThread, const FGraphEventRef& MyCompletionGraphEvent)
+				{
+					//UE_LOG(LogTemp, Log, TEXT("Main task"));
+
+					FGraphEventRef PrereqHolder = FGraphEvent::CreateGraphEvent();
+					PrereqHolder->SetDebugName(TEXT("PrereqHolder"));
+
+					FGraphEventRef Prereq = FFunctionGraphTask::CreateAndDispatchWhenReady(
+						[PrereqHolder]
+						{
+							//UE_LOG(LogTemp, Log, TEXT("Prereq"));
+
+							PrereqHolder->Wait(); // hold it until it's used for `DontCompleteUntil`
+						}
+					);
+					Prereq->SetDebugName(TEXT("Prereq"));
+
+					MyCompletionGraphEvent->DontCompleteUntil(Prereq);
+					check(!Prereq->IsComplete()); // check that prereq was incomplete during DontCompleteUntil ^^
+
+					// now that Prereq was registered in DontCompleteUntil, unlock it
+					PrereqHolder->DispatchSubsequents();
+			};
+
+			FGraphEventRef Event = FFunctionGraphTask::CreateAndDispatchWhenReady(MoveTemp(Lambda));
+			Event->SetDebugName(TEXT("MainEvent"));
+			check(!Event->IsComplete());
+			Event->Wait(ENamedThreads::GameThread);
+		}
+
+		{	// prereq is completed before DontCompleteUntil is called
+			FGraphEventRef Prereq = FFunctionGraphTask::CreateAndDispatchWhenReady(
+				[]
+				{
+					//UE_LOG(LogTemp, Log, TEXT("Prereq"));
+				}
+			);
+			Prereq->SetDebugName(TEXT("Prereq"));
+			Prereq->Wait(ENamedThreads::GameThread);
+
+			FGraphEventRef Event = FFunctionGraphTask::CreateAndDispatchWhenReady(
+				[&Prereq](ENamedThreads::Type CurrentThread, const FGraphEventRef& MyCompletionGraphEvent)
+				{
+					MyCompletionGraphEvent->DontCompleteUntil(Prereq);
+					//UE_LOG(LogTemp, Log, TEXT("Main task"));
+				}
+			);
+			Event->SetDebugName(TEXT("MainEvent"));
+			while (!Event->IsComplete() && FTaskGraphInterface::Get().GetNumWorkerThreads() != 0) // in single-threaded mode tasks are executed only when waited for
+			{}
+			Event->Wait(ENamedThreads::GameThread);
+		}
+	}
+
+	TEST_CASE_NAMED(FTaskGraphRecursionTest, "System::Core::Async::TaskGraph::RecursionTest", "[.][ApplicationContextMask][EngineFilter][Disabled]")
+	{
+		{	// recursive call on game thread
+			FGraphEventRef Event = FFunctionGraphTask::CreateAndDispatchWhenReady(
+				[]
+				{
+					FGraphEventRef Inner = FFunctionGraphTask::CreateAndDispatchWhenReady(
+						[]
+						{
+							check(IsInGameThread());
+						},
+						TStatId{}, nullptr, ENamedThreads::GameThread
+					);
+					Inner->Wait(ENamedThreads::GameThread);
+				},
+				TStatId{}, nullptr, ENamedThreads::GameThread
+			);
+			Event->Wait(ENamedThreads::GameThread);
+		}
+
+		//{	// didn't work in the old version
+		//	FGraphEventRef Event = FFunctionGraphTask::CreateAndDispatchWhenReady([] {}, TStatId{}, nullptr, ENamedThreads::GameThread_Local);
+		//	Event->Wait(ENamedThreads::GameThread);
+		//}
+	}
+
+	TEST_CASE_NAMED(FTaskGraphBasicTest, "System::Core::Async::TaskGraph::BasicTest", "[.][ApplicationContextMask][EngineFilter][Disabled]")
+	{
+		// thread and task priorities
+
+		{	// AnyNormalThreadNormalTask
+			bool bExecuted = false;
+			FFunctionGraphTask::CreateAndDispatchWhenReady([&bExecuted] { bExecuted = true; }, TStatId{}, nullptr, ENamedThreads::AnyNormalThreadNormalTask)->Wait();
+			check(bExecuted);
+		}
+
+		{	// AnyNormalThreadHiPriTask
+			bool bExecuted = false;
+			FFunctionGraphTask::CreateAndDispatchWhenReady([&bExecuted] { bExecuted = true; }, TStatId{}, nullptr, ENamedThreads::AnyNormalThreadHiPriTask)->Wait();
+			check(bExecuted);
+		}
+
+		{	// AnyBackgroundThreadNormalTask
+			bool bExecuted = false;
+			FFunctionGraphTask::CreateAndDispatchWhenReady([&bExecuted] { bExecuted = true; }, TStatId{}, nullptr, ENamedThreads::AnyBackgroundThreadNormalTask)->Wait();
+			check(bExecuted);
+		}
+
+		{	// AnyHiPriThreadNormalTask
+			bool bExecuted = false;
+			FFunctionGraphTask::CreateAndDispatchWhenReady([&bExecuted] { bExecuted = true; }, TStatId{}, nullptr, ENamedThreads::AnyHiPriThreadNormalTask)->Wait();
+			check(bExecuted);
+		}
+
+		{	// AnyHiPriThreadHiPriTask
+			bool bExecuted = false;
+			FFunctionGraphTask::CreateAndDispatchWhenReady([&bExecuted] { bExecuted = true; }, TStatId{}, nullptr, ENamedThreads::AnyHiPriThreadHiPriTask)->Wait();
+			check(bExecuted);
+		}
+
+		// named threads and local queues
+
+		if (IsRHIThreadRunning())
+		{	// RHIThread
+			bool bExecuted = false;
+			FFunctionGraphTask::CreateAndDispatchWhenReady([&bExecuted] { bExecuted = true; }, TStatId{}, nullptr, ENamedThreads::RHIThread)->Wait();
+			check(bExecuted);
+		}
+
+		{	// GameThread
+			bool bExecuted = false;
+			FFunctionGraphTask::CreateAndDispatchWhenReady([&bExecuted] { bExecuted = true; }, TStatId{}, nullptr, ENamedThreads::GameThread);
+			// GT must be executed explicitly
+			FTaskGraphInterface::Get().ProcessThreadUntilIdle(ENamedThreads::GameThread);
+			check(bExecuted);
+		}
+
+PRAGMA_DISABLE_DEPRECATION_WARNINGS
+		if (GRenderThreadId != 0)
+PRAGMA_ENABLE_DEPRECATION_WARNINGS
+		{	// ActualRenderingThread
+			bool bExecuted = false;
+			FFunctionGraphTask::CreateAndDispatchWhenReady([&bExecuted] { bExecuted = true; }, TStatId{}, nullptr, ENamedThreads::ActualRenderingThread)->Wait();
+			check(bExecuted);
+		}
+
+		{	// GameThread_Local
+			bool bExecuted = false;
+			FFunctionGraphTask::CreateAndDispatchWhenReady([&bExecuted] { bExecuted = true; }, TStatId{}, nullptr, ENamedThreads::GameThread_Local)->Wait(ENamedThreads::GameThread_Local);
+			check(bExecuted);
+		}
+
+PRAGMA_DISABLE_DEPRECATION_WARNINGS
+		if (GRenderThreadId != 0)
+PRAGMA_ENABLE_DEPRECATION_WARNINGS
+		{	// ActualRenderingThread_Local
+			bool bExecuted = false;
+			FFunctionGraphTask::CreateAndDispatchWhenReady(
+				[&bExecuted] 
+				{ 
+					FFunctionGraphTask::CreateAndDispatchWhenReady([&bExecuted] { bExecuted = true; }, TStatId{}, nullptr, ENamedThreads::ActualRenderingThread_Local)->Wait(ENamedThreads::ActualRenderingThread_Local);
+				}, 
+				TStatId{}, nullptr, ENamedThreads::ActualRenderingThread)->Wait();
+			check(bExecuted);
+		}
+
+		// dependencies
+
+		{	// a task is not executed until its prerequisite is completed
+			bool bExecuted = false;
+			FGraphEventRef Prereq = FGraphEvent::CreateGraphEvent();
+			FGraphEventRef MainTask = FFunctionGraphTask::CreateAndDispatchWhenReady([&bExecuted] { bExecuted = true; }, TStatId{}, Prereq);
+			// dummy task that is executed while the main task is waiting for its prereq
+			FFunctionGraphTask::CreateAndDispatchWhenReady([] {})->Wait();
+			check(!bExecuted);
+			Prereq->DispatchSubsequents();
+			MainTask->Wait();
+			check(bExecuted);
+		}
+
+		{	// a task is not executed until all its prerequisites are completed
+			bool bExecuted = false;
+			FGraphEventArray Prereqs{ FGraphEvent::CreateGraphEvent(), FGraphEvent::CreateGraphEvent() };
+			FGraphEventRef MainTask = FFunctionGraphTask::CreateAndDispatchWhenReady([&bExecuted] { bExecuted = true; }, TStatId{}, &Prereqs);
+			// dummy task that is executed while the main task is waiting for its prereqs
+			FFunctionGraphTask::CreateAndDispatchWhenReady([] {})->Wait();
+			check(!bExecuted);
+
+			Prereqs[0]->DispatchSubsequents();
+			FFunctionGraphTask::CreateAndDispatchWhenReady([] {})->Wait();
+			check(!bExecuted);
+			
+			Prereqs[1]->DispatchSubsequents();
+			MainTask->Wait();
+			check(bExecuted);
+		}
+
+		{	// ParallelFor
+			std::atomic<int32> Total{ 0 };
+			int32 Num = 1000;
+			ParallelFor(Num, [&Total](int32 i) { Total += i; });
+			check(Total == (Num - 1) * (Num / 2));
+		}
+
+		{	// holding a task
+			struct FTask
+			{
+				TStatId GetStatId() const
+				{
+					return TStatId{};
+				}
+
+				ENamedThreads::Type GetDesiredThread()
+				{
+					return ENamedThreads::AnyThread;
+				}
+
+				static ESubsequentsMode::Type GetSubsequentsMode()
+				{
+					return ESubsequentsMode::TrackSubsequents;
+				}
+
+				void DoTask(ENamedThreads::Type CurrentThread, const FGraphEventRef& MyCompletionGraphEvent)
+				{}
+			};
+
+			TGraphTask<FTask>* Task = TGraphTask<FTask>::CreateTask().ConstructAndHold();
+			FGraphEventRef Event = Task->GetCompletionEvent();
+			check(!Event->IsComplete());
+			Task->Unlock();
+			Event->Wait();
+			check(Event->IsComplete());
+		}
+
+		{	// check ref count for named thread tasks
+			FGraphEventRef LocalQueueTask = FFunctionGraphTask::CreateAndDispatchWhenReady([] {}, TStatId{}, nullptr, ENamedThreads::GameThread_Local);
+			LocalQueueTask->Wait(ENamedThreads::GameThread_Local);
+			check(LocalQueueTask.GetRefCount() == 1);
+		}
+
+		//for (int i = 0; i != 100000; ++i)
+		{	// a particular real-life case that doesn't work in the old TaskGraph if run in single-threaded mode.
+			// the culprit is that when a task is waited for, in single-threaded mode the queue it was pushed to is executed. 
+			// Here the (local queue) task  depends on a task that is not in the same queue and so it doesn't get executed
+
+			FGraphEventRef AnyTask = FFunctionGraphTask::CreateAndDispatchWhenReady([] {});
+			FGraphEventRef LocalQueueTask = FFunctionGraphTask::CreateAndDispatchWhenReady([] {}, TStatId{}, AnyTask, ENamedThreads::GameThread_Local);
+			LocalQueueTask->Wait(ENamedThreads::GameThread_Local);
+			check(LocalQueueTask.GetRefCount() == 1);
+		}
+
+		{	// launch a GT task, then an any-thread task that depends on it. wait for the any-thread task. this was a deadlock on the new frontend
+			FGraphEventRef GTTask = FFunctionGraphTask::CreateAndDispatchWhenReady([] { IsInGameThread(); }, TStatId{}, nullptr, ENamedThreads::GameThread);
+			FGraphEventRef AnyThreadTask = FFunctionGraphTask::CreateAndDispatchWhenReady([] {}, TStatId{}, GTTask);
+			AnyThreadTask->Wait();
+		}
+	}
+
+	// it's fast because tasks are too lightweight and so are executed almost as fast
+	template<int NumTasks>
+	void TestPerfBasic()
+	{
+		FGraphEventArray Tasks;
+		Tasks.Reserve(NumTasks);
+
+		for (int32 TaskIndex = 0; TaskIndex < NumTasks; ++TaskIndex)
+		{
+			Tasks.Emplace(FFunctionGraphTask::CreateAndDispatchWhenReady([] {}));
+		}
+
+		FTaskGraphInterface::Get().WaitUntilTasksComplete(Tasks);
+	}
+
+	template<int32 NumTasks, int32 BatchSize>
+	void TestPerfBatch()
+	{
+		static_assert(NumTasks % BatchSize == 0, "`NumTasks` must be divisible by `BatchSize`");
+		constexpr int32 NumBatches = NumTasks / BatchSize;
+
+		FGraphEventArray Batches;
+		Batches.Reserve(NumBatches);
+		FGraphEventArray Tasks;
+		Tasks.AddDefaulted(NumTasks);
+
+		for (int32 BatchIndex = 0; BatchIndex < NumBatches; ++BatchIndex)
+		{
+			Batches.Add(FFunctionGraphTask::CreateAndDispatchWhenReady(
+				[&Tasks, BatchIndex]
+				{
+					for (int32 TaskIndex = 0; TaskIndex < BatchSize; ++TaskIndex)
+					{
+						Tasks[BatchIndex * BatchSize + TaskIndex] = FFunctionGraphTask::CreateAndDispatchWhenReady([] {});
+					}
+				}
+			));
+		}
+
+		FTaskGraphInterface::Get().WaitUntilTasksComplete(Batches);
+		FTaskGraphInterface::Get().WaitUntilTasksComplete(Tasks);
+	}
+
+	template<int32 NumTasks, int32 BatchSize>
+	void TestPerfBatchOptimised()
+	{
+		static_assert(NumTasks % BatchSize == 0, "`NumTasks` must be divisible by `BatchSize`");
+		constexpr int32 NumBatches = NumTasks / BatchSize;
+
+		FGraphEventRef SpawnSignal = FGraphEvent::CreateGraphEvent();
+		FGraphEventArray AllDone;
+
+		for (int32 BatchIndex = 0; BatchIndex < NumBatches; ++BatchIndex)
+		{
+			AllDone.Add(FFunctionGraphTask::CreateAndDispatchWhenReady(
+				[](ENamedThreads::Type CurrentThread, const FGraphEventRef& CompletionEvent)
+				{
+					FGraphEventRef RunSignal = FGraphEvent::CreateGraphEvent();
+					for (int32 TaskIndex = 0; TaskIndex < BatchSize; ++TaskIndex)
+					{
+						CompletionEvent->DontCompleteUntil(FFunctionGraphTask::CreateAndDispatchWhenReady([] {}, TStatId{}, RunSignal, ENamedThreads::AnyThread));
+					}
+					RunSignal->DispatchSubsequents();
+				},
+				TStatId{}, SpawnSignal
+			));
+		}
+
+		SpawnSignal->DispatchSubsequents();
+		FTaskGraphInterface::Get().WaitUntilTasksComplete(MoveTemp(AllDone));
+	}
+
+	template<int NumTasks>
+	void TestLatency()
+	{
+		for (uint32 TaskIndex = 0; TaskIndex != NumTasks; ++TaskIndex)
+		{
+			FGraphEventRef GraphEvent = FFunctionGraphTask::CreateAndDispatchWhenReady([] {});
+			GraphEvent->Wait(ENamedThreads::GameThread);
+		}
+	}
+
+	int64 Fibonacci(int64 N)
+	{
+		check(N > 0);
+		if (N <= 2)
+		{
+			return 1;
+		}
+		else
+		{
+			std::atomic<int64> F1{ -1 };
+			std::atomic<int64> F2{ -1 };
+			FGraphEventArray GraphEvents;
+			GraphEvents.Add(FFunctionGraphTask::CreateAndDispatchWhenReady([&F1, N] { F1 = Fibonacci(N - 1); }));
+			GraphEvents.Add(FFunctionGraphTask::CreateAndDispatchWhenReady([&F2, N] { F2 = Fibonacci(N - 2); }));
+
+			FTaskGraphInterface::Get().WaitUntilTasksComplete(MoveTemp(GraphEvents));
+			check(F1 > 0 && F2 > 0);
+
+			return F1 + F2;
+		}
+	}
+
+	FGraphEventRef Fib(int64 N, int64* Res)
+	{
+		if (N <= 2)
+		{
+			*Res = 1;
+			FGraphEventRef ResEvent = FGraphEvent::CreateGraphEvent();
+			ResEvent->DispatchSubsequents();
+			return ResEvent;
+		}
+		else
+		{
+			TUniquePtr<int64> F1 = MakeUnique<int64>();
+			TUniquePtr<int64> F2 = MakeUnique<int64>();
+
+			FGraphEventArray SubTasks;
+
+			auto FibTask = [](int64 N, int64* Res)
+			{
+				return FFunctionGraphTask::CreateAndDispatchWhenReady
+				(
+					[N, Res]
+					(ENamedThreads::Type, const FGraphEventRef& CompletionEvent)
+					{
+						FGraphEventRef ResEvent = Fib(N, Res);
+						CompletionEvent->DontCompleteUntil(ResEvent);
+					}
+				);
+			};
+
+			SubTasks.Add(FibTask(N - 1, F1.Get()));
+			SubTasks.Add(FibTask(N - 2, F2.Get()));
+
+			FGraphEventRef ResEvent = FFunctionGraphTask::CreateAndDispatchWhenReady
+			(
+				[F1 = MoveTemp(F1), F2 = MoveTemp(F2), Res]
+				{
+					*Res = *F1 + *F2;
+				}, 
+				TStatId{}, &SubTasks
+			);
+
+			return ResEvent;
+		}
+	}
+
+	template<int64 N>
+	void Fib()
+	{
+		TUniquePtr<int64> Res = MakeUnique<int64>();
+		FGraphEventRef ResEvent = Fib(N, Res.Get());
+		ResEvent->Wait();
+		UE_LOG(LogTemp, Display, TEXT("Fibonacci(%d) = %d"), N, *Res);
+	}
+
+	template<uint32 NumTasks>
+	void TestFGraphEventPerf()
+	{
+		FGraphEventRef Prereq = FGraphEvent::CreateGraphEvent();
+		std::atomic<uint32> CompletedTasks{ 0 };
+
+		FGraphEventArray Tasks;
+		for (int i = 0; i != NumTasks; ++i)
+		{
+			Tasks.Add(FFunctionGraphTask::CreateAndDispatchWhenReady(
+				[&Prereq, &CompletedTasks](ENamedThreads::Type CurrentThread, const FGraphEventRef& MyCompletionGraphEvent)
+				{
+					MyCompletionGraphEvent->DontCompleteUntil(Prereq);
+					++CompletedTasks;
+				}
+			));
+		}
+
+		Prereq->DispatchSubsequents(ENamedThreads::GameThread);
+
+		FTaskGraphInterface::Get().WaitUntilTasksComplete(Tasks, ENamedThreads::GameThread);
+
+		check(CompletedTasks == NumTasks);
+	}
+
+	template<int NumTasks>
+	void TestSpawning()
+	{
+		{
+			FGraphEventArray Tasks;
+			Tasks.Reserve(NumTasks);
+			//double StartTime = FPlatformTime::Seconds();
+			for (uint32 TaskNo = 0; TaskNo != NumTasks; ++TaskNo)
+			{
+				Tasks.Add(FFunctionGraphTask::CreateAndDispatchWhenReady([] {}));
+			}
+
+			//double Duration = FPlatformTime::Seconds() - StartTime;
+			//UE_LOG(LogTemp, Display, TEXT("Spawning %d empty trackable tasks took %f secs"), NumTasks, Duration);
+
+			FTaskGraphInterface::Get().WaitUntilTasksComplete(MoveTemp(Tasks));
+		}
+		{
+			double StartTime = FPlatformTime::Seconds();
+			for (uint32 TaskNo = 0; TaskNo != NumTasks; ++TaskNo)
+			{
+				FFunctionGraphTask::CreateAndDispatchWhenReady([] {});
+			}
+
+			//double Duration = FPlatformTime::Seconds() - StartTime;
+			//UE_LOG(LogTemp, Display, TEXT("Spawning %d empty non-trackable tasks took %f secs"), NumTasks, Duration);
+		}
+	}
+
+	template<int NumTasks>
+	void TestBatchSpawning()
+	{
+		//double StartTime = FPlatformTime::Seconds();
+		FGraphEventRef Trigger = FGraphEvent::CreateGraphEvent();
+		for (uint32 TaskNo = 0; TaskNo != NumTasks; ++TaskNo)
+		{
+			FFunctionGraphTask::CreateAndDispatchWhenReady([] {}, TStatId{}, Trigger);
+		}
+
+		//double SpawnedTime = FPlatformTime::Seconds();
+		Trigger->DispatchSubsequents();
+
+		//double EndTime = FPlatformTime::Seconds();
+		//UE_LOG(LogTemp, Display, TEXT("Spawning %d empty non-trackable tasks took %f secs total, %f secs spawning and %f secs dispatching"), NumTasks, EndTime - StartTime, SpawnedTime - StartTime, EndTime - SpawnedTime);
+
+		//FTaskGraphInterface::Get().WaitUntilTasksComplete(MoveTemp(Tasks));
+	}
+
+	template<int64 NumBatches, int64 NumTasksPerBatch>
+	void TestWorkStealing()
+	{
+		FGraphEventArray Batches;
+		Batches.Reserve(NumBatches);
+
+		FGraphEventArray Tasks[NumBatches];
+		for (int32 BatchIndex = 0; BatchIndex != NumBatches; ++BatchIndex)
+		{
+			Tasks[BatchIndex].Reserve(NumBatches * NumTasksPerBatch);
+			Batches.Add(FFunctionGraphTask::CreateAndDispatchWhenReady(
+				[&Tasks, BatchIndex]()
+				{
+					for (int32 TaskIndex = 0; TaskIndex < NumTasksPerBatch; ++TaskIndex)
+					{
+						Tasks[BatchIndex].Add(FFunctionGraphTask::CreateAndDispatchWhenReady([] {}));
+					}
+				}
+			));
+		}
+
+		FTaskGraphInterface::Get().WaitUntilTasksComplete(Batches);
+		for (int32 BatchIndex = 0; BatchIndex != NumBatches; ++BatchIndex)
+		{
+			FTaskGraphInterface::Get().WaitUntilTasksComplete(Tasks[BatchIndex]);
+		}
+	}
+
+	TEST_CASE_NAMED(FTaskGraphPerfTest, "System::Core::Async::TaskGraph::PerfTest", "[.][ApplicationContextMask][EngineFilter][Disabled]")
+	{
+		TRACE_CPUPROFILER_EVENT_SCOPE(TaskGraphTests_PerfTest);
+
+		//UE_BENCHMARK(1, TestSpawning<100000>);
+		//return true;
+
+		//UE_BENCHMARK(5, Fib<18>);
+		//UE_BENCHMARK(5, [] { Fibonacci(15); });
+
+		UE_BENCHMARK(5, TestPerfBasic<100000>);
+		UE_BENCHMARK(5, TestPerfBatch<100000, 100>);
+		UE_BENCHMARK(5, TestPerfBatchOptimised<100000, 100>);
+		UE_BENCHMARK(5, TestLatency<10000>);
+		UE_BENCHMARK(5, TestFGraphEventPerf<100000>);
+		UE_BENCHMARK(5, TestWorkStealing<100, 1000>);
+		UE_BENCHMARK(5, TestSpawning<100000>);
+		UE_BENCHMARK(5, TestBatchSpawning<100000>);
+	}
+
+	template<uint32 Num>
+	void OversubscriptionStressTest()
+	{
+		// launch a number of tasks >= num of cores, each of them launches a task and waits for its completion. do this inside a task so 
+		// we can detect deadlock.
+		// outer tasks occupy all worker threads. when they launch inner tasks and wait for them, the inner tasks can't be 
+		// executed because all workers are blocked -> deadlock.
+		// can be solved by busy waiting
+
+		for (int i = 0; i != Num; ++i)
+		{
+			FSharedEventRef Event;
+			FFunctionGraphTask::CreateAndDispatchWhenReady(
+				[Event]
+				{
+					ParallelFor(200,
+						[](int32)
+						{
+							FPlatformProcess::Sleep(0.01f); // simulate some work and let all workers to pick up ParallelFor tasks
+							FFunctionGraphTask::CreateAndDispatchWhenReady([] {})->Wait();
+						}
+					);
+					Event->Trigger();
+				}
+			);
+			verify(Event->Wait(FTimespan::FromSeconds(5.f)));
+		}
+	}
+
+	TEST_CASE_NAMED(FTaskGraphOversubscriptionTest, "System::Core::Async::TaskGraph::Oversubscription", "[.][ApplicationContextMask][[EngineFilter][Disabled]")
+	{
+		UE_BENCHMARK(5, OversubscriptionStressTest<10>);
+	}
+
+	template<uint32 Nujm>
+	void SquaredOversubscriptionStressTest()
+	{
+		// same as before but using two ParallelFor nested one into another to simulate oversubscription in square
+
+		FSharedEventRef Event;
+		FFunctionGraphTask::CreateAndDispatchWhenReady(
+			[Event]
+			{
+				ParallelFor(200,
+					[](int32)
+					{
+						ParallelFor(200,
+							[](int32)
+							{
+								FPlatformProcess::Sleep(0.01f); // simulate some work and let all workers to pick up ParallelFor tasks
+								FFunctionGraphTask::CreateAndDispatchWhenReady([] {})->Wait();
+							}
+						);
+					}
+				);
+				Event->Trigger();
+			}
+		);
+		verify(Event->Wait(FTimespan::FromSeconds(30.f)));
+	}
+
+	TEST_CASE_NAMED(FTaskGraphSquaredOversubscriptionTest, "System::Core::Async::TaskGraph::SquaredOversubscription", "[.][ApplicationContextMask][EngineFilter][Disabled]")
+	{
+		UE_BENCHMARK(5, SquaredOversubscriptionStressTest<10>);
+	}
+
+	TEST_CASE_NAMED(FTaskGraphTaskDestructionTest, "System::Core::Async::TaskGraph::TaskDestruction", "[.][ApplicationContextMask][EngineFilter][Disabled]")
+	{
+		struct FDestructionTest
+		{
+			explicit FDestructionTest(bool* bDestroyedIn)
+				: bDestroyed(bDestroyedIn)
+			{}
+
+			~FDestructionTest()
+			{
+				*bDestroyed = true;
+			}
+
+			bool* bDestroyed;
+		};
+
+		bool bDestroyed = false;
+
+		FFunctionGraphTask::CreateAndDispatchWhenReady([DestructionTest = FDestructionTest{ &bDestroyed }]{})->Wait();
+
+		check(bDestroyed);
+	}
+}
+
+#endif //WITH_TESTS
diff --git a/Engine/Source/Runtime/Core/Tests/Delegates/TSDelegatesTest.cpp b/Engine/Source/Runtime/Core/Tests/Delegates/TSDelegatesTest.cpp
new file mode 100644
--- /dev/null
+++ b/Engine/Source/Runtime/Core/Tests/Delegates/TSDelegatesTest.cpp
@@ -0,0 +1,165 @@
+// Copyright Epic Games, Inc. All Rights Reserved.
+
+#include "CoreTypes.h"
+#include "Misc/AutomationTest.h"
+#include "Delegates/Delegate.h"
+#include "Tasks/Task.h"
+#include "Tests/TestHarnessAdapter.h"
+
+#if WITH_TESTS
+
+namespace UE { namespace TSDelegatesTest 
+{
+	using namespace Tasks;
+
+	void Func() {}
+
+	void OneParamFunc(bool* bExecuted)
+	{
+		*bExecuted = true;
+	}
+
+	void AnotherOneParamFunc(bool* bExecuted)
+	{
+		*bExecuted = true;
+	}
+
+	TEST_CASE_NAMED(FTSMulticastDelegateBasicTest, "System::Core::Delegates::MulticastBasic", "[ApplicationContextMask][EngineFilter]")
+	{
+		DECLARE_TS_MULTICAST_DELEGATE(FTestDelegate);
+
+		{	// single thread
+			FTestDelegate TestDelegate;
+			bool bExecuted = false;
+			FDelegateHandle DelegateHandle = TestDelegate.AddStatic(&OneParamFunc, &bExecuted);
+			TestDelegate.AddStatic(&AnotherOneParamFunc, &bExecuted);
+			TestDelegate.Broadcast();
+			check(bExecuted);
+			TestDelegate.Remove(DelegateHandle);
+		}
+
+		{	// concurrent Add/Remove and Execute
+			FTestDelegate TestDelegate;
+			bool bExecuted = false;
+			FTask AddStaticTask = Launch(UE_SOURCE_LOCATION, [&TestDelegate, &bExecuted] { TestDelegate.AddStatic(&AnotherOneParamFunc, &bExecuted); });
+			FTask BroadcastTask = Launch(UE_SOURCE_LOCATION, [&TestDelegate] { TestDelegate.Broadcast(); });
+			FDelegateHandle DelegateHandle = TestDelegate.AddStatic(&OneParamFunc, &bExecuted);
+			TestDelegate.Remove(DelegateHandle);
+			Wait(TArray{ AddStaticTask, BroadcastTask });
+		}
+
+		{ // same as above, but Add/Remove from different threads (not concurrently)
+			FTestDelegate TestDelegate;
+			bool bExecuted = false;
+			TTask<FDelegateHandle> AddStaticTask = Launch(UE_SOURCE_LOCATION, [&TestDelegate, &bExecuted] { return TestDelegate.AddStatic(&OneParamFunc, &bExecuted); });
+			FTask RemoveTask = Launch(UE_SOURCE_LOCATION, [&TestDelegate, &AddStaticTask] { TestDelegate.Remove(AddStaticTask.GetResult()); }, AddStaticTask);
+			TestDelegate.Broadcast();
+			RemoveTask.Wait();
+			TestDelegate.Remove(AddStaticTask.GetResult());
+		}
+
+		{	// recursive locking: removing a delegate handle from inside broadcasting
+			FTestDelegate TestDelegate;
+			FDelegateHandle DelegateHandle;
+			auto DelegateLambda = [&TestDelegate, &DelegateHandle] { TestDelegate.Remove(DelegateHandle); };
+			DelegateHandle = TestDelegate.AddLambda(MoveTemp(DelegateLambda));
+			check(TestDelegate.IsBound());
+			TestDelegate.Broadcast();
+			check(!TestDelegate.IsBound());
+		}
+
+		// test that all macros compile
+		{
+			DECLARE_TS_MULTICAST_DELEGATE_OneParam(FDelegate, int);
+			FDelegate Delegate;
+			Delegate.AddLambda([](int) {});
+			Delegate.Broadcast(0);
+		}
+		{
+			DECLARE_TS_MULTICAST_DELEGATE_TwoParams(FDelegate, int, int);
+			FDelegate Delegate;
+			Delegate.AddLambda([](int, int) {});
+			Delegate.Broadcast(0, 0);
+		}
+		{
+			DECLARE_TS_MULTICAST_DELEGATE_ThreeParams(FDelegate, int, int, int);
+			FDelegate Delegate;
+			Delegate.AddLambda([](int, int, int) {});
+			Delegate.Broadcast(0, 0, 0);
+		}
+		{
+			DECLARE_TS_MULTICAST_DELEGATE_FourParams(FDelegate, int, int, int, int);
+			FDelegate Delegate;
+			Delegate.AddLambda([](int, int, int, int) {});
+			Delegate.Broadcast(0, 0, 0, 0);
+		}
+		{
+			DECLARE_TS_MULTICAST_DELEGATE_FiveParams(FDelegate, int, int, int, int, int);
+			FDelegate Delegate;
+			Delegate.AddLambda([](int, int, int, int, int) {});
+			Delegate.Broadcast(0, 0, 0, 0, 0);
+		}
+		{
+			DECLARE_TS_MULTICAST_DELEGATE_SixParams(FDelegate, int, int, int, int, int, int);
+			FDelegate Delegate;
+			Delegate.AddLambda([](int, int, int, int, int, int) {});
+			Delegate.Broadcast(0, 0, 0, 0, 0, 0);
+		}
+		{
+			DECLARE_TS_MULTICAST_DELEGATE_SevenParams(FDelegate, int, int, int, int, int, int, int);
+			FDelegate Delegate;
+			Delegate.AddLambda([](int, int, int, int, int, int, int) {});
+			Delegate.Broadcast(0, 0, 0, 0, 0, 0, 0);
+		}
+		{
+			DECLARE_TS_MULTICAST_DELEGATE_EightParams(FDelegate, int, int, int, int, int, int, int, int);
+			FDelegate Delegate;
+			Delegate.AddLambda([](int, int, int, int, int, int, int, int) {});
+			Delegate.Broadcast(0, 0, 0, 0, 0, 0, 0, 0);
+		}
+		{
+			DECLARE_TS_MULTICAST_DELEGATE_NineParams(FDelegate, int, int, int, int, int, int, int, int, int);
+			FDelegate Delegate;
+			Delegate.AddLambda([](int, int, int, int, int, int, int, int, int) {});
+			Delegate.Broadcast(0, 0, 0, 0, 0, 0, 0, 0, 0);
+		}
+
+	}
+
+	TEST_CASE_NAMED(FTSMulticastDelegateStressTest, "System::Core::Delegates::MulticastStress", "[ApplicationContextMask][EngineFilter]")
+	{
+		DECLARE_TS_MULTICAST_DELEGATE(FTestDelegate);
+
+		{	// one thread repeatedly adds/removes delegates, another broadcasts
+			FTestDelegate Delegate;
+			std::atomic<bool> bQuit{ false };
+
+			FTask Binding = Launch(UE_SOURCE_LOCATION,
+				[&bQuit, &Delegate]
+				{
+					while (!bQuit)
+					{
+						FDelegateHandle Handle = Delegate.AddStatic(&Func);
+						Delegate.Remove(Handle);
+					}
+				}
+			);
+
+			FTask Executing = Launch(UE_SOURCE_LOCATION,
+				[&bQuit, &Delegate]
+				{
+					while (!bQuit)
+					{
+						Delegate.Broadcast();
+					}
+				}
+			);
+
+			FPlatformProcess::Sleep(1.0);
+			bQuit = true;
+			Wait(TArray<FTask>{ Binding, Executing });
+		}
+	}
+}}
+
+#endif //WITH_TESTS
\ No newline at end of file
diff --git a/Engine/Source/Runtime/Core/Tests/Memory/SharedBufferTest.cpp b/Engine/Source/Runtime/Core/Tests/Memory/SharedBufferTest.cpp
new file mode 100644
--- /dev/null
+++ b/Engine/Source/Runtime/Core/Tests/Memory/SharedBufferTest.cpp
@@ -0,0 +1,484 @@
+// Copyright Epic Games, Inc. All Rights Reserved.
+
+#if WITH_TESTS
+
+#include "Memory/SharedBuffer.h"
+#include "Tests/TestHarnessAdapter.h"
+#include <type_traits>
+#if WITH_LOW_LEVEL_TESTS
+#include "TestCommon/Expectations.h"
+#endif
+
+///////////////////////////////////////////////////////////////////////////////////////////////////////////////////////
+
+static_assert( std::is_constructible<FUniqueBuffer>::value, "Missing constructor");
+static_assert( std::is_constructible<FUniqueBuffer, FUniqueBuffer&&>::value, "Missing constructor");
+static_assert(!std::is_constructible<FUniqueBuffer, FSharedBuffer&&>::value, "Invalid constructor");
+static_assert(!std::is_constructible<FUniqueBuffer, FWeakSharedBuffer&&>::value, "Invalid constructor");
+static_assert(!std::is_constructible<FUniqueBuffer, const FUniqueBuffer&>::value, "Invalid constructor");
+static_assert(!std::is_constructible<FUniqueBuffer, const FSharedBuffer&>::value, "Invalid constructor");
+static_assert(!std::is_constructible<FUniqueBuffer, const FWeakSharedBuffer&>::value, "Invalid constructor");
+
+static_assert( std::is_constructible<FSharedBuffer>::value, "Missing constructor");
+static_assert(!std::is_constructible<FSharedBuffer, FUniqueBuffer&&>::value, "Invalid constructor");
+static_assert( std::is_constructible<FSharedBuffer, FSharedBuffer&&>::value, "Missing constructor");
+static_assert(!std::is_constructible<FSharedBuffer, FWeakSharedBuffer&&>::value, "Invalid constructor");
+static_assert(!std::is_constructible<FSharedBuffer, const FUniqueBuffer&>::value, "Invalid constructor");
+static_assert( std::is_constructible<FSharedBuffer, const FSharedBuffer&>::value, "Missing constructor");
+static_assert(!std::is_constructible<FSharedBuffer, const FWeakSharedBuffer&>::value, "Invalid constructor");
+
+static_assert( std::is_constructible<FWeakSharedBuffer>::value, "Missing constructor");
+static_assert(!std::is_constructible<FWeakSharedBuffer, FUniqueBuffer&&>::value, "Invalid constructor");
+static_assert( std::is_constructible<FWeakSharedBuffer, FSharedBuffer&&>::value, "Missing constructor");
+static_assert( std::is_constructible<FWeakSharedBuffer, FWeakSharedBuffer&&>::value, "Missing constructor");
+static_assert(!std::is_constructible<FWeakSharedBuffer, const FUniqueBuffer&>::value, "Invalid constructor");
+static_assert( std::is_constructible<FWeakSharedBuffer, const FSharedBuffer&>::value, "Missing constructor");
+static_assert( std::is_constructible<FWeakSharedBuffer, const FWeakSharedBuffer&>::value, "Missing constructor");
+
+///////////////////////////////////////////////////////////////////////////////////////////////////////////////////////
+
+static_assert( std::is_assignable<FUniqueBuffer, FUniqueBuffer&&>::value, "Missing assignment");
+static_assert(!std::is_assignable<FUniqueBuffer, FSharedBuffer&&>::value, "Invalid assignment");
+static_assert(!std::is_assignable<FUniqueBuffer, FWeakSharedBuffer&&>::value, "Invalid assignment");
+static_assert(!std::is_assignable<FUniqueBuffer, const FUniqueBuffer&>::value, "Invalid assignment");
+static_assert(!std::is_assignable<FUniqueBuffer, const FSharedBuffer&>::value, "Invalid assignment");
+static_assert(!std::is_assignable<FUniqueBuffer, const FWeakSharedBuffer&>::value, "Invalid assignment");
+
+static_assert(!std::is_assignable<FSharedBuffer, FUniqueBuffer&&>::value, "Invalid assignment");
+static_assert( std::is_assignable<FSharedBuffer, FSharedBuffer&&>::value, "Missing assignment");
+static_assert(!std::is_assignable<FSharedBuffer, FWeakSharedBuffer&&>::value, "Invalid assignment");
+static_assert(!std::is_assignable<FSharedBuffer, const FUniqueBuffer&>::value, "Invalid assignment");
+static_assert( std::is_assignable<FSharedBuffer, const FSharedBuffer&>::value, "Missing assignment");
+static_assert(!std::is_assignable<FSharedBuffer, const FWeakSharedBuffer&>::value, "Invalid assignment");
+
+static_assert(!std::is_assignable<FWeakSharedBuffer, FUniqueBuffer&&>::value, "Invalid assignment");
+static_assert( std::is_assignable<FWeakSharedBuffer, FSharedBuffer&&>::value, "Missing assignment");
+static_assert( std::is_assignable<FWeakSharedBuffer, FWeakSharedBuffer&&>::value, "Missing assignment");
+static_assert(!std::is_assignable<FWeakSharedBuffer, const FUniqueBuffer&>::value, "Invalid assignment");
+static_assert( std::is_assignable<FWeakSharedBuffer, const FSharedBuffer&>::value, "Missing assignment");
+static_assert( std::is_assignable<FWeakSharedBuffer, const FWeakSharedBuffer&>::value, "Missing assignment");
+
+///////////////////////////////////////////////////////////////////////////////////////////////////////////////////////
+
+static_assert(std::is_same<      void*, decltype(DeclVal<      FUniqueBuffer>().GetData())>::value, "Invalid accessor");
+static_assert(std::is_same<const void*, decltype(DeclVal<const FUniqueBuffer>().GetData())>::value, "Invalid accessor");
+static_assert(std::is_same<const void*, decltype(DeclVal<      FSharedBuffer>().GetData())>::value, "Invalid accessor");
+static_assert(std::is_same<const void*, decltype(DeclVal<const FSharedBuffer>().GetData())>::value, "Invalid accessor");
+
+static_assert(std::is_same<uint64, decltype(DeclVal<FUniqueBuffer>().GetSize())>::value, "Invalid accessor");
+static_assert(std::is_same<uint64, decltype(DeclVal<FSharedBuffer>().GetSize())>::value, "Invalid accessor");
+
+///////////////////////////////////////////////////////////////////////////////////////////////////////////////////////
+
+static_assert(std::is_same<FMutableMemoryView, decltype(DeclVal<      FUniqueBuffer>().GetView())>::value, "Invalid accessor");
+static_assert(std::is_same<       FMemoryView, decltype(DeclVal<const FUniqueBuffer>().GetView())>::value, "Invalid accessor");
+static_assert(std::is_same<       FMemoryView, decltype(DeclVal<      FSharedBuffer>().GetView())>::value, "Invalid accessor");
+static_assert(std::is_same<       FMemoryView, decltype(DeclVal<const FSharedBuffer>().GetView())>::value, "Invalid accessor");
+
+///////////////////////////////////////////////////////////////////////////////////////////////////////////////////////
+
+static_assert( std::is_convertible<      FUniqueBuffer, FMutableMemoryView>::value, "Missing conversion");
+static_assert(!std::is_convertible<const FUniqueBuffer, FMutableMemoryView>::value, "Invalid conversion");
+static_assert( std::is_convertible<      FUniqueBuffer, FMemoryView>::value, "Missing conversion");
+static_assert( std::is_convertible<const FUniqueBuffer, FMemoryView>::value, "Missing conversion");
+
+static_assert(!std::is_convertible<      FSharedBuffer, FMutableMemoryView>::value, "Invalid conversion");
+static_assert(!std::is_convertible<const FSharedBuffer, FMutableMemoryView>::value, "Invalid conversion");
+static_assert( std::is_convertible<      FSharedBuffer, FMemoryView>::value, "Missing conversion");
+static_assert( std::is_convertible<const FSharedBuffer, FMemoryView>::value, "Missing conversion");
+
+static_assert(!std::is_convertible<      FWeakSharedBuffer, FMutableMemoryView>::value, "Invalid conversion");
+static_assert(!std::is_convertible<const FWeakSharedBuffer, FMutableMemoryView>::value, "Invalid conversion");
+static_assert(!std::is_convertible<      FWeakSharedBuffer, FMemoryView>::value, "Invalid conversion");
+static_assert(!std::is_convertible<const FWeakSharedBuffer, FMemoryView>::value, "Invalid conversion");
+
+///////////////////////////////////////////////////////////////////////////////////////////////////////////////////////
+
+static_assert(std::is_same<bool, decltype(DeclVal<const FUniqueBuffer&>() == DeclVal<const FUniqueBuffer&>())>::value, "Invalid equality");
+static_assert(std::is_same<bool, decltype(DeclVal<const FUniqueBuffer&>() == DeclVal<const FSharedBuffer&>())>::value, "Invalid equality");
+static_assert(std::is_same<bool, decltype(DeclVal<const FUniqueBuffer&>() == DeclVal<const FWeakSharedBuffer&>())>::value, "Invalid equality");
+static_assert(std::is_same<bool, decltype(DeclVal<const FSharedBuffer&>() == DeclVal<const FUniqueBuffer&>())>::value, "Invalid equality");
+static_assert(std::is_same<bool, decltype(DeclVal<const FSharedBuffer&>() == DeclVal<const FSharedBuffer&>())>::value, "Invalid equality");
+static_assert(std::is_same<bool, decltype(DeclVal<const FSharedBuffer&>() == DeclVal<const FWeakSharedBuffer&>())>::value, "Invalid equality");
+static_assert(std::is_same<bool, decltype(DeclVal<const FWeakSharedBuffer&>() == DeclVal<const FUniqueBuffer&>())>::value, "Invalid equality");
+static_assert(std::is_same<bool, decltype(DeclVal<const FWeakSharedBuffer&>() == DeclVal<const FSharedBuffer&>())>::value, "Invalid equality");
+static_assert(std::is_same<bool, decltype(DeclVal<const FWeakSharedBuffer&>() == DeclVal<const FWeakSharedBuffer&>())>::value, "Invalid equality");
+
+///////////////////////////////////////////////////////////////////////////////////////////////////////////////////////
+
+static_assert(std::is_same<bool, decltype(DeclVal<const FUniqueBuffer&>() != DeclVal<const FUniqueBuffer&>())>::value, "Invalid inequality");
+static_assert(std::is_same<bool, decltype(DeclVal<const FUniqueBuffer&>() != DeclVal<const FSharedBuffer&>())>::value, "Invalid inequality");
+static_assert(std::is_same<bool, decltype(DeclVal<const FUniqueBuffer&>() != DeclVal<const FWeakSharedBuffer&>())>::value, "Invalid inequality");
+static_assert(std::is_same<bool, decltype(DeclVal<const FSharedBuffer&>() != DeclVal<const FUniqueBuffer&>())>::value, "Invalid inequality");
+static_assert(std::is_same<bool, decltype(DeclVal<const FSharedBuffer&>() != DeclVal<const FSharedBuffer&>())>::value, "Invalid inequality");
+static_assert(std::is_same<bool, decltype(DeclVal<const FSharedBuffer&>() != DeclVal<const FWeakSharedBuffer&>())>::value, "Invalid inequality");
+static_assert(std::is_same<bool, decltype(DeclVal<const FWeakSharedBuffer&>() != DeclVal<const FUniqueBuffer&>())>::value, "Invalid inequality");
+static_assert(std::is_same<bool, decltype(DeclVal<const FWeakSharedBuffer&>() != DeclVal<const FSharedBuffer&>())>::value, "Invalid inequality");
+static_assert(std::is_same<bool, decltype(DeclVal<const FWeakSharedBuffer&>() != DeclVal<const FWeakSharedBuffer&>())>::value, "Invalid inequality");
+
+///////////////////////////////////////////////////////////////////////////////////////////////////////////////////////
+
+static_assert(std::is_same<uint32, decltype(GetTypeHash(DeclVal<const FUniqueBuffer>()))>::value, "Missing or invalid hash function");
+static_assert(std::is_same<uint32, decltype(GetTypeHash(DeclVal<const FSharedBuffer>()))>::value, "Missing or invalid hash function");
+static_assert(std::is_same<uint32, decltype(GetTypeHash(DeclVal<const FWeakSharedBuffer>()))>::value, "Missing or invalid hash function");
+
+///////////////////////////////////////////////////////////////////////////////////////////////////////////////////////
+
+TEST_CASE_NAMED(FUniqueBufferTest, "System::Core::Memory::UniqueBuffer", "[ApplicationContextMask][SmokeFilter]")
+{
+	// Test Null
+	{
+		FUniqueBuffer Buffer;
+		CHECK_MESSAGE(TEXT("FUniqueBuffer().IsOwned()"), Buffer.IsOwned());
+		CHECK_EQUALS(TEXT("FUniqueBuffer().GetSize()"), Buffer.GetSize(), uint64(0));
+		CHECK_EQUALS(TEXT("FUniqueBuffer().GetView().GetSize()"), Buffer.GetView().GetSize(), uint64(0));
+		CHECK_EQUALS(TEXT("FUniqueBuffer().GetView().GetData()"), Buffer.GetView().GetData(), Buffer.GetData());
+	}
+
+	// Test Alloc
+	{
+		constexpr uint64 Size = 64;
+		FUniqueBuffer Buffer = FUniqueBuffer::Alloc(Size);
+		CHECK_MESSAGE(TEXT("FUniqueBuffer::Alloc().IsOwned()"), Buffer.IsOwned());
+		CHECK_EQUALS(TEXT("FUniqueBuffer::Alloc().GetSize()"), Buffer.GetSize(), Size);
+		CHECK_EQUALS(TEXT("FUniqueBuffer::Alloc().GetView().GetSize()"), Buffer.GetView().GetSize(), Size);
+		CHECK_EQUALS(TEXT("FUniqueBuffer::Alloc().GetView().GetData()"), Buffer.GetView().GetData(), Buffer.GetData());
+	}
+
+	// Test Clone
+	{
+		constexpr uint64 Size = 64;
+		const uint8 Data[Size]{};
+		FUniqueBuffer Buffer = FUniqueBuffer::Clone(Data, Size);
+		CHECK_MESSAGE(TEXT("FUniqueBuffer::Clone().IsOwned()"), Buffer.IsOwned());
+		CHECK_EQUALS(TEXT("FUniqueBuffer::Clone().GetSize()"), Buffer.GetSize(), Size);
+		CHECK_NOT_EQUALS(TEXT("FUniqueBuffer::Clone().GetData()"), static_cast<const void*>(Buffer.GetData()), static_cast<const void*>(Data));
+	}
+
+	// Test MakeView
+	{
+		constexpr uint64 Size = 64;
+		uint8 Data[Size]{};
+		FUniqueBuffer Buffer = FUniqueBuffer::MakeView(Data, Size);
+		CHECK_FALSE_MESSAGE(TEXT("FUniqueBuffer::MakeView().IsOwned()"), Buffer.IsOwned());
+		CHECK_EQUALS(TEXT("FUniqueBuffer::MakeView().GetSize()"), Buffer.GetSize(), Size);
+		CHECK_EQUALS(TEXT("FUniqueBuffer::MakeView().GetData()"), Buffer.GetData(), static_cast<void*>(Data));
+	}
+
+	// Test TakeOwnership
+	{
+		bool bDeleted = false;
+		constexpr uint64 Size = 64;
+		uint8 Data[Size]{};
+		FUniqueBuffer Buffer = FUniqueBuffer::TakeOwnership(Data, Size, [&bDeleted](void*, uint64) { bDeleted = true; });
+		CHECK_MESSAGE(TEXT("FUniqueBuffer::TakeOwnership().IsOwned()"), Buffer.IsOwned());
+		CHECK_EQUALS(TEXT("FUniqueBuffer::TakeOwnership().GetSize()"), Buffer.GetSize(), Size);
+		CHECK_EQUALS(TEXT("FUniqueBuffer::TakeOwnership().GetData()"), Buffer.GetData(), static_cast<void*>(Data));
+		Buffer.Reset();
+		CHECK_MESSAGE(TEXT("FUniqueBuffer::TakeOwnership() Deleted"), bDeleted);
+	}
+
+	// Test MakeOwned
+	{
+		constexpr uint64 Size = 64;
+		uint8 Data[Size]{};
+		FUniqueBuffer Buffer = FUniqueBuffer::MakeView(Data, Size).MakeOwned();
+		CHECK_MESSAGE(TEXT("FUniqueBuffer::MakeOwned(View).IsOwned()"), Buffer.IsOwned());
+		CHECK_EQUALS(TEXT("FUniqueBuffer::MakeOwned(View).GetSize()"), Buffer.GetSize(), Size);
+		CHECK_FALSE_MESSAGE(TEXT("FUniqueBuffer::MakeOwned(View).GetData()"), Buffer.GetData() == Data);
+		void* const OwnedData = Buffer.GetData();
+		Buffer = MoveTemp(Buffer).MakeOwned();
+		CHECK_MESSAGE(TEXT("FUniqueBuffer::MakeOwned(Owned).IsOwned()"), Buffer.IsOwned());
+		CHECK_EQUALS(TEXT("FUniqueBuffer::MakeOwned(Owned).GetSize()"), Buffer.GetSize(), Size);
+		CHECK_EQUALS(TEXT("FUniqueBuffer::MakeOwned(Owned).GetData()"), Buffer.GetData(), OwnedData);
+	}
+}
+
+TEST_CASE_NAMED(FSharedBufferTest, "System::Core::Memory::SharedBuffer", "[ApplicationContextMask][SmokeFilter]")
+{
+	// Test Null
+	{
+		FSharedBuffer Buffer;
+		CHECK_MESSAGE(TEXT("FSharedBuffer().IsOwned()"), Buffer.IsOwned());
+		CHECK_EQUALS(TEXT("FSharedBuffer().GetSize()"), Buffer.GetSize(), uint64(0));
+		CHECK_EQUALS(TEXT("FSharedBuffer().GetView().GetSize()"), Buffer.GetView().GetSize(), uint64(0));
+		CHECK_EQUALS(TEXT("FSharedBuffer().GetView().GetData()"), Buffer.GetView().GetData(), Buffer.GetData());
+	}
+
+	// Test Clone
+	{
+		constexpr uint64 Size = 64;
+		const uint8 Data[Size]{};
+		FSharedBuffer Buffer = FSharedBuffer::Clone(Data, Size);
+		CHECK_MESSAGE(TEXT("FSharedBuffer::Clone().IsOwned()"), Buffer.IsOwned());
+		CHECK_EQUALS(TEXT("FSharedBuffer::Clone().GetSize()"), Buffer.GetSize(), Size);
+		CHECK_NOT_EQUALS(TEXT("FSharedBuffer::Clone().GetData()"), Buffer.GetData(), static_cast<const void*>(Data));
+	}
+
+	// Test MakeView
+	{
+		constexpr uint64 Size = 64;
+		uint8 Data[Size]{};
+		FSharedBuffer Buffer = FSharedBuffer::MakeView(Data, Size);
+		CHECK_FALSE_MESSAGE(TEXT("FSharedBuffer::MakeView().IsOwned()"), Buffer.IsOwned());
+		CHECK_EQUALS(TEXT("FSharedBuffer::MakeView().GetSize()"), Buffer.GetSize(), Size);
+		CHECK_EQUALS(TEXT("FSharedBuffer::MakeView().GetData()"), Buffer.GetData(), static_cast<const void*>(Data));
+	}
+
+	// Test MakeView Outer View
+	{
+		constexpr uint64 Size = 64;
+		const uint8 Data[Size]{};
+		FWeakSharedBuffer WeakBuffer;
+		{
+			FSharedBuffer Buffer;
+			{
+				FSharedBuffer OuterBuffer = FSharedBuffer::MakeView(Data, Size);
+				Buffer = FSharedBuffer::MakeView(OuterBuffer.GetData(), OuterBuffer.GetSize() / 2, OuterBuffer);
+				WeakBuffer = OuterBuffer;
+			}
+			CHECK_FALSE_MESSAGE(TEXT("FSharedBuffer::MakeView(OuterView).IsOwned()"), Buffer.IsOwned());
+			CHECK_EQUALS(TEXT("FSharedBuffer::MakeView(OuterView).GetSize()"), Buffer.GetSize(), Size / 2);
+			CHECK_EQUALS(TEXT("FSharedBuffer::MakeView(OuterView).GetData()"), Buffer.GetData(), static_cast<const void*>(Data));
+			CHECK_FALSE_MESSAGE(TEXT("FSharedBuffer::MakeView(OuterView) Outer Not Null"), WeakBuffer.Pin().IsNull());
+		}
+		CHECK_MESSAGE(TEXT("FSharedBuffer::MakeView(OuterView) Outer Null"), WeakBuffer.Pin().IsNull());
+	}
+
+	// Test MakeView Outer Owned
+	{
+		constexpr uint64 Size = 64;
+		FWeakSharedBuffer WeakBuffer;
+		{
+			FSharedBuffer Buffer;
+			{
+				FSharedBuffer OuterBuffer = FUniqueBuffer::Alloc(Size).MoveToShared();
+				Buffer = FSharedBuffer::MakeView(OuterBuffer.GetData(), OuterBuffer.GetSize(), OuterBuffer);
+				CHECK_EQUALS(TEXT("FSharedBuffer::MakeView(OuterOwned, SameView)"), Buffer, OuterBuffer);
+				Buffer = FSharedBuffer::MakeView(OuterBuffer.GetData(), OuterBuffer.GetSize() / 2, OuterBuffer);
+				WeakBuffer = OuterBuffer;
+			}
+			CHECK_MESSAGE(TEXT("FSharedBuffer::MakeView(OuterOwned).IsOwned()"), Buffer.IsOwned());
+			CHECK_EQUALS(TEXT("FSharedBuffer::MakeView(OuterOwned).GetSize()"), Buffer.GetSize(), Size / 2);
+			CHECK_EQUALS(TEXT("FSharedBuffer::MakeView(OuterOwned).GetData()"), Buffer.GetData(), WeakBuffer.Pin().GetData());
+			CHECK_FALSE_MESSAGE(TEXT("FSharedBuffer::MakeView(OuterOwned) Outer Not Null"), WeakBuffer.Pin().IsNull());
+		}
+		CHECK_MESSAGE(TEXT("FSharedBuffer::MakeView(OuterOwned) Outer  Null"), WeakBuffer.Pin().IsNull());
+	}
+
+	// Test TakeOwnership
+	{
+		bool bDeleted = false;
+		constexpr uint64 Size = 64;
+		uint8 Data[Size]{};
+		FSharedBuffer Buffer = FSharedBuffer::TakeOwnership(Data, Size, [&bDeleted](void*, uint64) { bDeleted = true; });
+		CHECK_MESSAGE(TEXT("FSharedBuffer::TakeOwnership().IsOwned()"), Buffer.IsOwned());
+		CHECK_EQUALS(TEXT("FSharedBuffer::TakeOwnership().GetSize()"), Buffer.GetSize(), Size);
+		CHECK_EQUALS(TEXT("FSharedBuffer::TakeOwnership().GetData()"), Buffer.GetData(), static_cast<const void*>(Data));
+		Buffer.Reset();
+		CHECK_MESSAGE(TEXT("FSharedBuffer::TakeOwnership() Deleted"), bDeleted);
+	}
+
+	// Test MakeOwned
+	{
+		constexpr uint64 Size = 64;
+		uint8 Data[Size]{};
+		FSharedBuffer Buffer = FSharedBuffer::MakeView(Data, Size).MakeOwned();
+		CHECK_MESSAGE(TEXT("FSharedBuffer::MakeOwned(View).IsOwned()"), Buffer.IsOwned());
+		CHECK_EQUALS(TEXT("FSharedBuffer::MakeOwned(View).GetSize()"), Buffer.GetSize(), Size);
+		CHECK_FALSE_MESSAGE(TEXT("FSharedBuffer::MakeOwned(View).GetData()"), Buffer.GetData() == Data);
+		FSharedBuffer BufferCopy = Buffer.MakeOwned();
+		CHECK_MESSAGE(TEXT("FSharedBuffer::MakeOwned(Owned).IsOwned()"), BufferCopy.IsOwned());
+		CHECK_EQUALS(TEXT("FSharedBuffer::MakeOwned(Owned).GetSize()"), BufferCopy.GetSize(), Size);
+		CHECK_EQUALS(TEXT("FSharedBuffer::MakeOwned(Owned).GetData()"), BufferCopy.GetData(), Buffer.GetData());
+	}
+
+	// Test MoveToUnique
+	{
+		FUniqueBuffer Buffer = FSharedBuffer().MoveToUnique();
+		CHECK_MESSAGE(TEXT("FSharedBuffer::MoveToUnique(Null).IsOwned()"), Buffer.IsOwned());
+		CHECK_EQUALS(TEXT("FSharedBuffer::MoveToUnique(Null).GetSize()"), Buffer.GetSize(), uint64(0));
+	}
+	{
+		FSharedBuffer Shared = FSharedBuffer::Clone(MakeMemoryView<uint8>({ 1, 2, 3, 4 }));
+		FUniqueBuffer Buffer = FSharedBuffer(Shared).MoveToUnique();
+		CHECK_MESSAGE(TEXT("FSharedBuffer::MoveToUnique(Shared).IsOwned()"), Buffer.IsOwned());
+		CHECK_EQUALS(TEXT("FSharedBuffer::MoveToUnique(Shared).GetSize()"), Buffer.GetSize(), Shared.GetSize());
+		CHECK_NOT_EQUALS(TEXT("FSharedBuffer::MoveToUnique(Shared).GetData()"), const_cast<const void*>(Buffer.GetData()), Shared.GetData());
+		CHECK_MESSAGE(TEXT("FSharedBuffer::MoveToUnique(Shared)->Equal"), Buffer.GetView().EqualBytes(Shared.GetView()));
+	}
+	{
+		bool bDeleted = false;
+		constexpr uint64 Size = 64;
+		uint8 Data[Size]{};
+		FSharedBuffer Shared = FUniqueBuffer::TakeOwnership(Data, Size, [&bDeleted](void*) { bDeleted = true; }).MoveToShared();
+		FUniqueBuffer Buffer = Shared.MoveToUnique();
+		CHECK_MESSAGE(TEXT("FSharedBuffer::MoveToUnique(Owned).IsNull()"), Shared.IsNull());
+		CHECK_MESSAGE(TEXT("FSharedBuffer::MoveToUnique(Owned).IsOwned()"), Buffer.IsOwned());
+		CHECK_EQUALS(TEXT("FSharedBuffer::MoveToUnique(Owned).GetSize()"), Buffer.GetSize(), Size);
+		CHECK_EQUALS(TEXT("FSharedBuffer::MoveToUnique(Owned).GetData()"), Buffer.GetData(), static_cast<void*>(Data));
+		Buffer.Reset();
+		CHECK_MESSAGE(TEXT("FSharedBuffer::MoveToUnique(Owned) Deleted"), bDeleted);
+	}
+	{
+		constexpr uint64 Size = 64;
+		uint8 Data[Size]{};
+		FSharedBuffer Shared = FSharedBuffer::TakeOwnership(Data, Size, [](void*) {});
+		FUniqueBuffer Buffer = Shared.MoveToUnique();
+		CHECK_MESSAGE(TEXT("FSharedBuffer::MoveToUnique(ImmutableOwned).IsNull()"), Shared.IsNull());
+		CHECK_MESSAGE(TEXT("FSharedBuffer::MoveToUnique(ImmutableOwned).IsOwned()"), Buffer.IsOwned());
+		CHECK_EQUALS(TEXT("FSharedBuffer::MoveToUnique(ImmutableOwned).GetSize()"), Buffer.GetSize(), Size);
+		CHECK_NOT_EQUALS(TEXT("FSharedBuffer::MoveToUnique(ImmutableOwned).GetData()"), Buffer.GetData(), static_cast<void*>(Data));
+	}
+	{
+		constexpr uint64 Size = 64;
+		uint8 Data[Size]{};
+		FSharedBuffer Shared = FUniqueBuffer::MakeView(MakeMemoryView(Data)).MoveToShared();
+		FUniqueBuffer Buffer = Shared.MoveToUnique();
+		CHECK_MESSAGE(TEXT("FSharedBuffer::MoveToUnique(View).IsNull()"), Shared.IsNull());
+		CHECK_MESSAGE(TEXT("FSharedBuffer::MoveToUnique(View).IsOwned()"), Buffer.IsOwned());
+		CHECK_EQUALS(TEXT("FSharedBuffer::MoveToUnique(View).GetSize()"), Buffer.GetSize(), Size);
+		CHECK_NOT_EQUALS(TEXT("FSharedBuffer::MoveToUnique(View).GetData()"), Buffer.GetData(), static_cast<void*>(Data));
+	}
+	{
+		constexpr uint64 Size = 64;
+		uint8 Data[Size]{};
+		FSharedBuffer Shared = FSharedBuffer::MakeView(MakeMemoryView(Data));
+		FUniqueBuffer Buffer = Shared.MoveToUnique();
+		CHECK_MESSAGE(TEXT("FSharedBuffer::MoveToUnique(ImmutableView).IsNull()"), Shared.IsNull());
+		CHECK_MESSAGE(TEXT("FSharedBuffer::MoveToUnique(ImmutableView).IsOwned()"), Buffer.IsOwned());
+		CHECK_EQUALS(TEXT("FSharedBuffer::MoveToUnique(ImmutableView).GetSize()"), Buffer.GetSize(), Size);
+		CHECK_NOT_EQUALS(TEXT("FSharedBuffer::MoveToUnique(ImmutableView).GetData()"), Buffer.GetData(), static_cast<void*>(Data));
+	}
+	{
+		constexpr uint64 Size = 64;
+		uint8 Data[Size]{};
+		FSharedBuffer Shared = FUniqueBuffer::TakeOwnership(Data, Size, [](void*) {}).MoveToShared();
+		FUniqueBuffer Buffer = FSharedBuffer::MakeView(Shared.GetView().Left(Size / 2), Shared).MoveToUnique();
+		CHECK_FALSE_MESSAGE(TEXT("FSharedBuffer::MoveToUnique(OwnedView).IsNull()"), Shared.IsNull());
+		CHECK_MESSAGE(TEXT("FSharedBuffer::MoveToUnique(OwnedView).IsOwned()"), Buffer.IsOwned());
+		CHECK_EQUALS(TEXT("FSharedBuffer::MoveToUnique(OwnedView).GetSize()"), Buffer.GetSize(), Size / 2);
+		CHECK_NOT_EQUALS(TEXT("FSharedBuffer::MoveToUnique(OwnedView).GetData()"), Buffer.GetData(), static_cast<void*>(Data));
+	}
+
+}
+
+TEST_CASE_NAMED(FBufferOwnerTest, "System::Core::Memory::BufferOwner", "[ApplicationContextMask][SmokeFilter]")
+{
+	class FTestBufferOwner final : public FBufferOwner
+	{
+	public:
+		FTestBufferOwner(bool& bInMaterializedRef, bool& bInFreedRef, bool& bInDeletedRef)
+			: bMaterializedRef(bInMaterializedRef)
+			, bFreedRef(bInFreedRef)
+			, bDeletedRef(bInDeletedRef)
+		{
+			SetIsOwned();
+		}
+
+	private:
+		virtual void MaterializeBuffer()
+		{
+			SetIsMaterialized();
+			bMaterializedRef = true;
+		}
+
+		virtual void FreeBuffer() final
+		{
+			bFreedRef = true;
+		}
+
+		virtual ~FTestBufferOwner()
+		{
+			bDeletedRef = true;
+		}
+
+		bool& bMaterializedRef;
+		bool& bFreedRef;
+		bool& bDeletedRef;
+	};
+
+	bool bMaterialized = false;
+	bool bFreed = false;
+	bool bDeleted = false;
+
+	FSharedBuffer Buffer(new FTestBufferOwner(bMaterialized, bFreed, bDeleted));
+	CHECK_FALSE_MESSAGE(TEXT("FSharedBuffer(BufferOwner)->!Materialized"), bMaterialized);
+	CHECK_FALSE_MESSAGE(TEXT("FSharedBuffer(BufferOwner)->!Freed"), bFreed);
+	CHECK_FALSE_MESSAGE(TEXT("FSharedBuffer(BufferOwner)->!Deleted"), bDeleted);
+
+	(void)Buffer.GetData();
+	CHECK_MESSAGE(TEXT("FSharedBuffer(BufferOwner)->GetData()->Materialized"), bMaterialized);
+	bMaterialized = false;
+	(void)Buffer.GetData();
+	CHECK_FALSE_MESSAGE(TEXT("FSharedBuffer(BufferOwner)->GetData()->!Materialized"), bMaterialized);
+	bMaterialized = false;
+
+	Buffer = FSharedBuffer(new FTestBufferOwner(bMaterialized, bFreed, bDeleted));
+	CHECK_FALSE_MESSAGE(TEXT("FSharedBuffer(BufferOwner)->Assign(FSharedBuffer)->!Materialized"), bMaterialized);
+	CHECK_MESSAGE(TEXT("FSharedBuffer(BufferOwner)->Assign(FSharedBuffer)->Freed"), bFreed);
+	CHECK_MESSAGE(TEXT("FSharedBuffer(BufferOwner)->Assign(FSharedBuffer)->Deleted"), bDeleted);
+	bFreed = false;
+	bDeleted = false;
+
+	(void)Buffer.GetSize();
+	CHECK_MESSAGE(TEXT("FSharedBuffer(BufferOwner)->GetSize()->Materialized"), bMaterialized);
+	bMaterialized = false;
+	(void)Buffer.GetSize();
+	CHECK_FALSE_MESSAGE(TEXT("FSharedBuffer(BufferOwner)->GetSize()->!Materialized"), bMaterialized);
+	bMaterialized = false;
+
+	FWeakSharedBuffer WeakBuffer = Buffer;
+	Buffer.Reset();
+	CHECK_MESSAGE(TEXT("FSharedBuffer(BufferOwner)->Reset(Shared)->Freed"), bFreed);
+	CHECK_FALSE_MESSAGE(TEXT("FSharedBuffer(BufferOwner)->Reset(Shared)->!Deleted"), bDeleted);
+	WeakBuffer.Reset();
+	CHECK_MESSAGE(TEXT("FSharedBuffer(BufferOwner)->Reset(Weak)->Deleted"), bDeleted);
+
+	bMaterialized = false;
+	bFreed = false;
+	bDeleted = false;
+	Buffer = FSharedBuffer(new FTestBufferOwner(bMaterialized, bFreed, bDeleted));
+
+	CHECK_FALSE_MESSAGE(TEXT("FSharedBuffer(BufferOwner)->IsMaterialized()"), Buffer.IsMaterialized());
+	Buffer.Materialize();
+	CHECK_MESSAGE(TEXT("FSharedBuffer(BufferOwner)->Materialize()->IsMaterialized()"), Buffer.IsMaterialized());
+	CHECK_MESSAGE(TEXT("FSharedBuffer(BufferOwner)->Materialize()->Materialized"), bMaterialized);
+	bMaterialized = false;
+	Buffer.Materialize();
+	CHECK_FALSE_MESSAGE(TEXT("FSharedBuffer(BufferOwner)->Materialize()->!Materialized"), bMaterialized);
+	bMaterialized = false;
+	Buffer.Reset();
+
+	bFreed = false;
+	bDeleted = false;
+	Buffer = FSharedBuffer(new FTestBufferOwner(bMaterialized, bFreed, bDeleted));
+	Buffer.Reset();
+	CHECK_FALSE_MESSAGE(TEXT("FSharedBuffer(BufferOwner)->Reset(!Materialized)->!Materialized"), bMaterialized);
+	CHECK_MESSAGE(TEXT("FSharedBuffer(BufferOwner)->Reset(!Materialized)->Freed"), bFreed);
+	CHECK_MESSAGE(TEXT("FSharedBuffer(BufferOwner)->Reset(!Materialized)->Deleted"), bDeleted);
+}
+
+TEST_CASE_NAMED(FUniqueBufferFromArrayTest, "System::Core::Memory::UniqueBufferFromArray", "[ApplicationContextMask][SmokeFilter]")
+{
+	TArray<uint16> Array{ 1, 2, 3, 4, 5, 6, 7, 8 };
+	const void* const ExpectedData = Array.GetData();
+	const uint64 ExpectedSize = Array.Num() * sizeof(uint16);
+	const FUniqueBuffer Buffer = MakeUniqueBufferFromArray(MoveTemp(Array));
+	CHECK_MESSAGE(TEXT("MakeUniqueBufferFromArray -> Array.IsEmpty"), Array.IsEmpty());
+	CHECK_EQUALS(TEXT("MakeUniqueBufferFromArray -> Buffer.GetData()"), Buffer.GetData(), ExpectedData);
+	CHECK_EQUALS(TEXT("MakeUniqueBufferFromArray -> Buffer.GetSize()"), Buffer.GetSize(), ExpectedSize);
+}
+
+TEST_CASE_NAMED(FSharedBufferFromArrayTest, "System::Core::Memory::SharedBufferFromArray", "[ApplicationContextMask][SmokeFilter]")
+{
+	TArray<uint16> Array{1, 2, 3, 4, 5, 6, 7, 8};
+	const void* const ExpectedData = Array.GetData();
+	const uint64 ExpectedSize = Array.Num() * sizeof(uint16);
+	const FSharedBuffer Buffer = MakeSharedBufferFromArray(MoveTemp(Array));
+	CHECK_MESSAGE(TEXT("MakeSharedBufferFromArray -> Array.IsEmpty"), Array.IsEmpty());
+	CHECK_EQUALS(TEXT("MakeSharedBufferFromArray -> Buffer.GetData()"), Buffer.GetData(), ExpectedData);
+	CHECK_EQUALS(TEXT("MakeSharedBufferFromArray -> Buffer.GetSize()"), Buffer.GetSize(), ExpectedSize);
+}
+
+#endif // WITH_TESTS
diff --git a/Engine/Source/Runtime/Core/Tests/Misc/AsciiSetTest.cpp b/Engine/Source/Runtime/Core/Tests/Misc/AsciiSetTest.cpp
--- a/Engine/Source/Runtime/Core/Tests/Misc/AsciiSetTest.cpp
+++ b/Engine/Source/Runtime/Core/Tests/Misc/AsciiSetTest.cpp
@@ -5,7 +5,7 @@
 #include "Misc/AsciiSet.h"
 #include "Tests/TestHarnessAdapter.h"
 #if WITH_LOW_LEVEL_TESTS
-#include "TestCommon/Comparisons.h"
+#include "TestCommon/Expectations.h"
 #endif
 
 TEST_CASE_NAMED(TAsciiSetTest, "System::Core::Misc::AsciiSet", "[ApplicationContextMask][SmokeFilter]")
diff --git a/Engine/Source/Runtime/Core/Tests/Misc/CharTest.cpp b/Engine/Source/Runtime/Core/Tests/Misc/CharTest.cpp
--- a/Engine/Source/Runtime/Core/Tests/Misc/CharTest.cpp
+++ b/Engine/Source/Runtime/Core/Tests/Misc/CharTest.cpp
@@ -8,7 +8,7 @@
 #include <wctype.h>
 #include "Tests/TestHarnessAdapter.h"
 #if WITH_LOW_LEVEL_TESTS
-#include "TestCommon/Comparisons.h"
+#include "TestCommon/Expectations.h"
 #endif
 
 namespace crt
diff --git a/Engine/Source/Runtime/Core/Tests/Misc/CircularBufferTest.cpp b/Engine/Source/Runtime/Core/Tests/Misc/CircularBufferTest.cpp
--- a/Engine/Source/Runtime/Core/Tests/Misc/CircularBufferTest.cpp
+++ b/Engine/Source/Runtime/Core/Tests/Misc/CircularBufferTest.cpp
@@ -7,7 +7,7 @@
 #include "Containers/CircularBuffer.h"
 #include "Tests/TestHarnessAdapter.h"
 #if WITH_LOW_LEVEL_TESTS
-#include "TestCommon/Comparisons.h"
+#include "TestCommon/Expectations.h"
 #endif
 
 TEST_CASE_NAMED(FCircularBufferTest, "System::Core::Misc::CircularBuffer", "[ApplicationContextMask][SmokeFilter]")
diff --git a/Engine/Source/Runtime/Core/Tests/Misc/CircularQueueTest.cpp b/Engine/Source/Runtime/Core/Tests/Misc/CircularQueueTest.cpp
--- a/Engine/Source/Runtime/Core/Tests/Misc/CircularQueueTest.cpp
+++ b/Engine/Source/Runtime/Core/Tests/Misc/CircularQueueTest.cpp
@@ -6,7 +6,7 @@
 #include "Containers/CircularQueue.h"
 #include "Tests/TestHarnessAdapter.h"
 #if WITH_LOW_LEVEL_TESTS
-#include "TestCommon/Comparisons.h"
+#include "TestCommon/Expectations.h"
 #endif
 
 #pragma warning(disable : 6011)
diff --git a/Engine/Source/Runtime/Core/Tests/Misc/DateTimeTest.cpp b/Engine/Source/Runtime/Core/Tests/Misc/DateTimeTest.cpp
--- a/Engine/Source/Runtime/Core/Tests/Misc/DateTimeTest.cpp
+++ b/Engine/Source/Runtime/Core/Tests/Misc/DateTimeTest.cpp
@@ -7,7 +7,7 @@
 #include "Misc/DateTime.h"
 #include "Tests/TestHarnessAdapter.h"
 #if WITH_LOW_LEVEL_TESTS
-#include "TestCommon/Comparisons.h"
+#include "TestCommon/Expectations.h"
 #endif
 
 #define TestUnixEquivalent(Desc, A, B) if ((A).ToUnixTimestamp() != (B)) FAIL_CHECK(FString::Printf(TEXT("%s - A=%d B=%d"), Desc, (A).ToUnixTimestamp(), (B)));
diff --git a/Engine/Source/Runtime/Core/Tests/Misc/DisplayNameTest.cpp b/Engine/Source/Runtime/Core/Tests/Misc/DisplayNameTest.cpp
--- a/Engine/Source/Runtime/Core/Tests/Misc/DisplayNameTest.cpp
+++ b/Engine/Source/Runtime/Core/Tests/Misc/DisplayNameTest.cpp
@@ -5,7 +5,7 @@
 #include "UObject/UnrealNames.h"
 #include "Tests/TestHarnessAdapter.h"
 #if WITH_LOW_LEVEL_TESTS
-#include "TestCommon/Comparisons.h"
+#include "TestCommon/Expectations.h"
 #endif
 
 TEST_CASE_NAMED(TDisplayNameTest, "System::Core::Misc::DisplayName", "[ApplicationContextMask][SmokeFilter]")
diff --git a/Engine/Source/Runtime/Core/Tests/Misc/FileTest.cpp b/Engine/Source/Runtime/Core/Tests/Misc/FileTest.cpp
--- a/Engine/Source/Runtime/Core/Tests/Misc/FileTest.cpp
+++ b/Engine/Source/Runtime/Core/Tests/Misc/FileTest.cpp
@@ -15,7 +15,7 @@
 #include "Misc/ScopeExit.h"
 #include "Tests/TestHarnessAdapter.h"
 #if WITH_LOW_LEVEL_TESTS
-#include "TestCommon/Comparisons.h"
+#include "TestCommon/Expectations.h"
 #endif
 
 // These file tests are designed to ensure expected file writing behavior, as well as cross-platform consistency
diff --git a/Engine/Source/Runtime/Core/Tests/Misc/GuidTest.cpp b/Engine/Source/Runtime/Core/Tests/Misc/GuidTest.cpp
--- a/Engine/Source/Runtime/Core/Tests/Misc/GuidTest.cpp
+++ b/Engine/Source/Runtime/Core/Tests/Misc/GuidTest.cpp
@@ -8,7 +8,7 @@
 #include "Tests/TestHarnessAdapter.h"
 #include "Misc/SecureHash.h"
 #if WITH_LOW_LEVEL_TESTS
-#include "TestCommon/Comparisons.h"
+#include "TestCommon/Expectations.h"
 #endif
 
 TEST_CASE_NAMED(FGuidTest, "System::Core::Misc::Guid", "[ApplicationContextMask][SmokeFilter]")
diff --git a/Engine/Source/Runtime/Core/Tests/Misc/MonotonicTimeTest.cpp b/Engine/Source/Runtime/Core/Tests/Misc/MonotonicTimeTest.cpp
--- a/Engine/Source/Runtime/Core/Tests/Misc/MonotonicTimeTest.cpp
+++ b/Engine/Source/Runtime/Core/Tests/Misc/MonotonicTimeTest.cpp
@@ -1,15 +1,15 @@
 // Copyright Epic Games, Inc. All Rights Reserved.
 
-#if WITH_LOW_LEVEL_TESTS
+#if WITH_TESTS
 
 #include "Misc/MonotonicTime.h"
 
-#include "TestHarness.h"
+#include "Tests/TestHarnessAdapter.h"
 
 namespace UE
 {
 
-TEST_CASE("Core::Time::MonotonicTimeSpan", "[Core][Time][Smoke]")
+TEST_CASE_NAMED(FMonotonicTimeSpanTest, "System::Core::Time::MonotonicTimeSpan", "[Core][Time][SmokeFilter]")
 {
 	SECTION("Constructors")
 	{
@@ -74,7 +74,7 @@ TEST_CASE("Core::Time::MonotonicTimeSpan", "[Core][Time][Smoke]")
 	}
 }
 
-TEST_CASE("Core::Time::MonotonicTimePoint", "[Core][Time][Smoke]")
+TEST_CASE_NAMED(FMonotonicTimePointTest, "System::Core::Time::MonotonicTimePoint", "[Core][Time][SmokeFilter]")
 {
 	SECTION("Constructors")
 	{
@@ -139,4 +139,4 @@ TEST_CASE("Core::Time::MonotonicTimePoint", "[Core][Time][Smoke]")
 
 } // namespace UE
 
-#endif // WITH_LOW_LEVEL_TESTS
+#endif // WITH_TESTS
diff --git a/Engine/Source/Runtime/Core/Tests/Misc/PathTests.h b/Engine/Source/Runtime/Core/Tests/Misc/PathTests.h
--- a/Engine/Source/Runtime/Core/Tests/Misc/PathTests.h
+++ b/Engine/Source/Runtime/Core/Tests/Misc/PathTests.h
@@ -5,7 +5,7 @@
 #include "Containers/StringView.h"
 #include "Tests/TestHarnessAdapter.h"
 #if WITH_LOW_LEVEL_TESTS
-#include "TestCommon/Comparisons.h"
+#include "TestCommon/Expectations.h"
 #endif
 
 #if WITH_TESTS
diff --git a/Engine/Source/Runtime/Core/Tests/Misc/PathViewsTest.cpp b/Engine/Source/Runtime/Core/Tests/Misc/PathViewsTest.cpp
--- a/Engine/Source/Runtime/Core/Tests/Misc/PathViewsTest.cpp
+++ b/Engine/Source/Runtime/Core/Tests/Misc/PathViewsTest.cpp
@@ -6,7 +6,7 @@
 #include "Tests/TestHarnessAdapter.h"
 #include "Misc/StringBuilder.h"
 #if WITH_LOW_LEVEL_TESTS
-#include "TestCommon/Comparisons.h"
+#include "TestCommon/Expectations.h"
 #endif
 
 #if WITH_TESTS
diff --git a/Engine/Source/Runtime/Core/Tests/Misc/PathsTest.cpp b/Engine/Source/Runtime/Core/Tests/Misc/PathsTest.cpp
--- a/Engine/Source/Runtime/Core/Tests/Misc/PathsTest.cpp
+++ b/Engine/Source/Runtime/Core/Tests/Misc/PathsTest.cpp
@@ -9,7 +9,7 @@
 #include "Misc/Paths.h"
 #include "Tests/TestHarnessAdapter.h"
 #if WITH_LOW_LEVEL_TESTS
-#include "TestCommon/Comparisons.h"
+#include "TestCommon/Expectations.h"
 #endif
 
 const FStringView PathTest::BaseDir = TEXTVIEW("/root");
diff --git a/Engine/Source/Runtime/Core/Tests/Misc/QueueTest.cpp b/Engine/Source/Runtime/Core/Tests/Misc/QueueTest.cpp
--- a/Engine/Source/Runtime/Core/Tests/Misc/QueueTest.cpp
+++ b/Engine/Source/Runtime/Core/Tests/Misc/QueueTest.cpp
@@ -6,7 +6,7 @@
 #include "Containers/Queue.h"
 #include "Tests/TestHarnessAdapter.h"
 #if WITH_LOW_LEVEL_TESTS
-#include "TestCommon/Comparisons.h"
+#include "TestCommon/Expectations.h"
 #endif
 
 TEST_CASE_NAMED(FQueueTest, "System::Core::Misc::Queue", "[ApplicationContextMask][SmokeFilter]")
diff --git a/Engine/Source/Runtime/Core/Tests/Misc/TCharToUTF8.cpp b/Engine/Source/Runtime/Core/Tests/Misc/TCharToUTF8.cpp
--- a/Engine/Source/Runtime/Core/Tests/Misc/TCharToUTF8.cpp
+++ b/Engine/Source/Runtime/Core/Tests/Misc/TCharToUTF8.cpp
@@ -5,7 +5,7 @@
 #include "CoreTypes.h"
 #include "Tests/TestHarnessAdapter.h"
 #if WITH_LOW_LEVEL_TESTS
-#include "TestCommon/Comparisons.h"
+#include "TestCommon/Expectations.h"
 #endif
 
 TEST_CASE_NAMED(FTCharToUTF8Test, "System::Core::Misc::TCharToUtf8", "[ApplicationContextMask][SmokeFilter]")
diff --git a/Engine/Source/Runtime/Core/Tests/Misc/TripleBufferTest.cpp b/Engine/Source/Runtime/Core/Tests/Misc/TripleBufferTest.cpp
--- a/Engine/Source/Runtime/Core/Tests/Misc/TripleBufferTest.cpp
+++ b/Engine/Source/Runtime/Core/Tests/Misc/TripleBufferTest.cpp
@@ -8,7 +8,7 @@
 #include "Containers/TripleBuffer.h"
 #include "Math/RandomStream.h"
 #if WITH_LOW_LEVEL_TESTS
-#include "TestCommon/Comparisons.h"
+#include "TestCommon/Expectations.h"
 #endif
 
 TEST_CASE_NAMED(FTripleBufferTest, "System::Core::Misc::TripleBuffer", "[ApplicationContextMask][SmokeFilter]")
diff --git a/Engine/Source/Runtime/Core/Tests/Misc/TypeContainerTest.cpp b/Engine/Source/Runtime/Core/Tests/Misc/TypeContainerTest.cpp
--- a/Engine/Source/Runtime/Core/Tests/Misc/TypeContainerTest.cpp
+++ b/Engine/Source/Runtime/Core/Tests/Misc/TypeContainerTest.cpp
@@ -12,7 +12,7 @@
 #include "Async/Async.h"
 #include "Misc/TypeContainer.h"
 #if WITH_LOW_LEVEL_TESTS
-#include "TestCommon/Comparisons.h"
+#include "TestCommon/Expectations.h"
 #endif
 
 /* Helpers
diff --git a/Engine/Source/Runtime/Core/Tests/Serialization/CompactBinarySerializationTest.cpp b/Engine/Source/Runtime/Core/Tests/Serialization/CompactBinarySerializationTest.cpp
new file mode 100644
--- /dev/null
+++ b/Engine/Source/Runtime/Core/Tests/Serialization/CompactBinarySerializationTest.cpp
@@ -0,0 +1,243 @@
+// Copyright Epic Games, Inc. All Rights Reserved.
+
+#include "Serialization/CompactBinarySerialization.h"
+#include "Tests/TestHarnessAdapter.h"
+#include "Serialization/BufferArchive.h"
+#include "Serialization/MemoryReader.h"
+#if WITH_LOW_LEVEL_TESTS
+#include "TestCommon/Expectations.h"
+#endif
+
+#if WITH_TESTS
+
+TEST_CASE_NAMED(FCbMeasureTest, "System::Core::Serialization::MeasureCompactBinary", "[ApplicationContextMask][SmokeFilter]")
+{
+	const auto TestMeasure = [&](
+		const TCHAR* Test,
+		std::initializer_list<uint8> Data,
+		bool bExpectedValue,
+		uint64 ExpectedSize,
+		ECbFieldType ExpectedType,
+		ECbFieldType ExternalType = ECbFieldType::HasFieldType) -> void
+	{
+		ECbFieldType ActualType = ECbFieldType(255);
+		uint64 ActualSize = ~uint64(0);
+		const bool bActualValue = TryMeasureCompactBinary(MakeMemoryView(Data), ActualType, ActualSize, ExternalType);
+		CHECK_EQUALS(FString::Printf(TEXT("TryMeasureCompactBinary(%s)"), Test), bActualValue, bExpectedValue);
+		CHECK_EQUALS(FString::Printf(TEXT("TryMeasureCompactBinary(%s)->Type"), Test), ActualType, ExpectedType);
+		CHECK_EQUALS(FString::Printf(TEXT("TryMeasureCompactBinary(%s)->Size"), Test), ActualSize, ExpectedSize);
+		if (ActualSize == ExpectedSize)
+		{
+			const uint64 MeasureSize = MeasureCompactBinary(MakeMemoryView(Data), ExternalType);
+			CHECK_EQUALS(FString::Printf(TEXT("MeasureCompactBinary(%s)"), Test), MeasureSize, bExpectedValue ? ExpectedSize : 0);
+		}
+	};
+
+	using EType = ECbFieldType;
+
+	TestMeasure(TEXT("Empty"), {}, false, 1, EType::None);
+
+	TestMeasure(TEXT("None"), {uint8(EType::None)}, false, 0, EType::None);
+	TestMeasure(TEXT("None, NoType"), {}, false, 0, EType::None, EType::None);
+	TestMeasure(TEXT("None, NoType, Name"), {}, false, 0, EType::None, EType::None | EType::HasFieldName);
+
+	TestMeasure(TEXT("Null"), {uint8(EType::Null)}, true, 1, EType::Null);
+	TestMeasure(TEXT("Null, NameSize1B"), {uint8(EType::Null | EType::HasFieldName), 30}, true, 32, EType::Null | EType::HasFieldName);
+	TestMeasure(TEXT("Null, NameSize2B"), {uint8(EType::Null | EType::HasFieldName), 0x80, 0x80}, true, 131, EType::Null | EType::HasFieldName);
+	TestMeasure(TEXT("Null, NameSize2BShort"), {uint8(EType::Null | EType::HasFieldName), 0x80}, false, 3, EType::Null | EType::HasFieldName);
+	TestMeasure(TEXT("Null, NameSize3BShort"), {uint8(EType::Null | EType::HasFieldName), 0xc0}, false, 4, EType::Null | EType::HasFieldName);
+	TestMeasure(TEXT("Null, MissingName"), {uint8(EType::Null | EType::HasFieldName)}, false, 2, EType::Null | EType::HasFieldName);
+
+	TestMeasure(TEXT("Null, NoType"), {}, true, 0, EType::Null, EType::Null);
+	TestMeasure(TEXT("Null, NoType, NameSize1B"), {30}, true, 31, EType::Null | EType::HasFieldName, EType::Null | EType::HasFieldName);
+	TestMeasure(TEXT("Null, NoType, NameSize2B"), {0x80, 0x80}, true, 130, EType::Null | EType::HasFieldName, EType::Null | EType::HasFieldName);
+	TestMeasure(TEXT("Null, NoType, NameSize2BShort"), {0x80}, false, 2, EType::Null | EType::HasFieldName, EType::Null | EType::HasFieldName);
+	TestMeasure(TEXT("Null, NoType, NameSize3BShort"), {0xc0}, false, 3, EType::Null | EType::HasFieldName, EType::Null | EType::HasFieldName);
+	TestMeasure(TEXT("Null, NoType, MissingName"), {}, false, 1, EType::Null | EType::HasFieldName, EType::Null | EType::HasFieldName);
+
+	TestMeasure(TEXT("Object, NoSize"), {uint8(EType::Object)}, false, 2, EType::Object);
+	TestMeasure(TEXT("Object, Size1B"), {uint8(EType::Object), 30}, true, 32, EType::Object);
+	TestMeasure(TEXT("UniformObject, NoSize"), {uint8(EType::UniformObject)}, false, 2, EType::UniformObject);
+	TestMeasure(TEXT("UniformObject, Size1B"), {uint8(EType::UniformObject), 30}, true, 32, EType::UniformObject);
+
+	TestMeasure(TEXT("Array, NoSize"), {uint8(EType::Array)}, false, 2, EType::Array);
+	TestMeasure(TEXT("Array, Size1B"), {uint8(EType::Array), 30}, true, 32, EType::Array);
+	TestMeasure(TEXT("UniformArray, NoSize"), {uint8(EType::UniformArray)}, false, 2, EType::UniformArray);
+	TestMeasure(TEXT("UniformArray, Size1B"), {uint8(EType::UniformArray), 30}, true, 32, EType::UniformArray);
+
+	TestMeasure(TEXT("Binary, NoSize"), {uint8(EType::Binary)}, false, 2, EType::Binary);
+	TestMeasure(TEXT("Binary, Size1B"), {uint8(EType::Binary), 30}, true, 32, EType::Binary);
+
+	TestMeasure(TEXT("String, NoSize"), {uint8(EType::String)}, false, 2, EType::String);
+	TestMeasure(TEXT("String, Size1B"), {uint8(EType::String), 30}, true, 32, EType::String);
+	TestMeasure(TEXT("String, Size2B"), {uint8(EType::String), 0x80, 0x80}, true, 131, EType::String);
+	TestMeasure(TEXT("String, Size2BShort"), {uint8(EType::String), 0x80}, false, 3, EType::String);
+	TestMeasure(TEXT("String, Size3BShort"), {uint8(EType::String), 0xc0}, false, 4, EType::String);
+
+	TestMeasure(TEXT("String, NameNoSize"), {uint8(EType::String | EType::HasFieldName)}, false, 2, EType::String | EType::HasFieldName);
+	TestMeasure(TEXT("String, NameSize1BMissingText"), {uint8(EType::String | EType::HasFieldName), 1}, false, 3, EType::String | EType::HasFieldName);
+	TestMeasure(TEXT("String, NameSize1B, NoSize"), {uint8(EType::String | EType::HasFieldName), 1, 'A'}, false, 4, EType::String | EType::HasFieldName);
+	TestMeasure(TEXT("String, NameSize1B, Size1B"), {uint8(EType::String | EType::HasFieldName), 1, 'A', 30}, true, 34, EType::String | EType::HasFieldName);
+	TestMeasure(TEXT("String, NameSize1B, Size2B"), {uint8(EType::String | EType::HasFieldName), 1, 'A', 0x80, 0x80}, true, 133, EType::String | EType::HasFieldName);
+	TestMeasure(TEXT("String, NameSize1B, Size2BShort"), {uint8(EType::String | EType::HasFieldName), 1, 'A', 0x80}, false, 5, EType::String | EType::HasFieldName);
+	TestMeasure(TEXT("String, NameSize1B, Size3BShort"), {uint8(EType::String | EType::HasFieldName), 1, 'A', 0xc0}, false, 6, EType::String | EType::HasFieldName);
+	TestMeasure(TEXT("String, NameSize2BShort"), {uint8(EType::String | EType::HasFieldName), 0x80}, false, 3, EType::String | EType::HasFieldName);
+	TestMeasure(TEXT("String, NameSize3BShort"), {uint8(EType::String | EType::HasFieldName), 0xc0}, false, 4, EType::String | EType::HasFieldName);
+
+	TestMeasure(TEXT("IntegerPositive, NoValue"), {uint8(EType::IntegerPositive)}, false, 2, EType::IntegerPositive);
+	TestMeasure(TEXT("IntegerPositive, Value1B"), {uint8(EType::IntegerPositive), 0x7f}, true, 2, EType::IntegerPositive);
+	TestMeasure(TEXT("IntegerPositive, Value2B"), {uint8(EType::IntegerPositive), 0x80}, true, 3, EType::IntegerPositive);
+
+	TestMeasure(TEXT("IntegerNegative, NoValue"), {uint8(EType::IntegerNegative)}, false, 2, EType::IntegerNegative);
+	TestMeasure(TEXT("IntegerNegative, Value1B"), {uint8(EType::IntegerNegative), 0x7f}, true, 2, EType::IntegerNegative);
+	TestMeasure(TEXT("IntegerNegative, Value2B"), {uint8(EType::IntegerNegative), 0x80}, true, 3, EType::IntegerNegative);
+
+	TestMeasure(TEXT("Float32"), {uint8(EType::Float32)}, true, 5, EType::Float32);
+	TestMeasure(TEXT("Float32, NameSize1B"), {uint8(EType::Float32 | EType::HasFieldName), 30}, true, 36, EType::Float32 | EType::HasFieldName);
+	TestMeasure(TEXT("Float32, NameSize2B"), {uint8(EType::Float32 | EType::HasFieldName), 0x80, 0x80}, true, 135, EType::Float32 | EType::HasFieldName);
+	TestMeasure(TEXT("Float32, NameSize2BShort"), {uint8(EType::Float32 | EType::HasFieldName), 0x80}, false, 3, EType::Float32 | EType::HasFieldName);
+	TestMeasure(TEXT("Float32, MissingName"), {uint8(EType::Float32 | EType::HasFieldName)}, false, 2, EType::Float32 | EType::HasFieldName);
+
+	TestMeasure(TEXT("Float32, NoType"), {}, true, 4, EType::Float32, EType::Float32);
+	TestMeasure(TEXT("Float32, NoType, NameSize1B"), {30}, true, 35, EType::Float32 | EType::HasFieldName, EType::Float32 | EType::HasFieldName);
+	TestMeasure(TEXT("Float32, NoType, NameSize2B"), {0x80, 0x80}, true, 134, EType::Float32 | EType::HasFieldName, EType::Float32 | EType::HasFieldName);
+	TestMeasure(TEXT("Float32, NoType, NameSize2BShort"), {0x80}, false, 2, EType::Float32 | EType::HasFieldName, EType::Float32 | EType::HasFieldName);
+	TestMeasure(TEXT("Float32, NoType, MissingName"), {}, false, 1, EType::Float32 | EType::HasFieldName, EType::Float32 | EType::HasFieldName);
+
+	TestMeasure(TEXT("Float64"), {uint8(EType::Float64)}, true, 9, EType::Float64);
+
+	TestMeasure(TEXT("BoolFalse"), {uint8(EType::BoolFalse)}, true, 1, EType::BoolFalse);
+	TestMeasure(TEXT("BoolTrue"), {uint8(EType::BoolTrue)}, true, 1, EType::BoolTrue);
+
+	TestMeasure(TEXT("ObjectAttachment"), {uint8(EType::ObjectAttachment)}, true, 21, EType::ObjectAttachment);
+	TestMeasure(TEXT("BinaryAttachment"), {uint8(EType::BinaryAttachment)}, true, 21, EType::BinaryAttachment);
+
+	TestMeasure(TEXT("Hash"), {uint8(EType::Hash)}, true, 21, EType::Hash);
+	TestMeasure(TEXT("Uuid"), {uint8(EType::Uuid)}, true, 17, EType::Uuid);
+
+	TestMeasure(TEXT("DateTime"), {uint8(EType::DateTime)}, true, 9, EType::DateTime);
+	TestMeasure(TEXT("TimeSpan"), {uint8(EType::TimeSpan)}, true, 9, EType::TimeSpan);
+
+	TestMeasure(TEXT("ObjectId"), {uint8(EType::ObjectId)}, true, 13, EType::ObjectId);
+
+	TestMeasure(TEXT("CustomById, Size1B"), {uint8(EType::CustomById), 0x7f}, true, 129, EType::CustomById);
+	TestMeasure(TEXT("CustomById, Size2B"), {uint8(EType::CustomById), 0x80, 0x80}, true, 131, EType::CustomById);
+	TestMeasure(TEXT("CustomById, Size1BShort"), {uint8(EType::CustomById)}, false, 2, EType::CustomById);
+	TestMeasure(TEXT("CustomById, Size2BShort"), {uint8(EType::CustomById), 0x80}, false, 3, EType::CustomById);
+	TestMeasure(TEXT("CustomById, Size9BShort"), {uint8(EType::CustomById), 0xff}, false, 10, EType::CustomById);
+
+	TestMeasure(TEXT("CustomByName, Size1B"), {uint8(EType::CustomByName), 0x7f}, true, 129, EType::CustomByName);
+	TestMeasure(TEXT("CustomByName, Size2B"), {uint8(EType::CustomByName), 0x80, 0x80}, true, 131, EType::CustomByName);
+	TestMeasure(TEXT("CustomByName, Size1BShort"), {uint8(EType::CustomByName)}, false, 2, EType::CustomByName);
+	TestMeasure(TEXT("CustomByName, Size2BShort"), {uint8(EType::CustomByName), 0x80}, false, 3, EType::CustomByName);
+	TestMeasure(TEXT("CustomByName, Size9BShort"), {uint8(EType::CustomByName), 0xff}, false, 10, EType::CustomByName);
+}
+
+///////////////////////////////////////////////////////////////////////////////////////////////////////////////////////
+
+TEST_CASE_NAMED(FCbSaveTest, "System::Core::Serialization::SaveCompactBinary", "[ApplicationContextMask][SmokeFilter]")
+{
+
+	const auto TestSave = [&](const TCHAR* Test, auto Value, FMemoryView ExpectedData)
+	{
+		{
+			FBufferArchive WriteAr;
+			SaveCompactBinary(WriteAr, Value);
+			if (TestEqual(FString::Printf(TEXT("SaveCompactBinary(%s)->Size"), Test), uint64(WriteAr.Num()), ExpectedData.GetSize()) &&
+				TestTrue(FString::Printf(TEXT("SaveCompactBinary(%s)->EqualBytes"), Test), ExpectedData.EqualBytes(MakeMemoryView(WriteAr))))
+			{
+				FMemoryReader ReadAr(WriteAr);
+				FCbField Field = LoadCompactBinary(ReadAr);
+				CHECK_MESSAGE(FString::Printf(TEXT("LoadCompactBinary(%s)->EqualBytes"), Test), ExpectedData.EqualBytes(Field.GetOuterBuffer().GetView()));
+			}
+		}
+		{
+			FBufferArchive WriteAr;
+			WriteAr << Value;
+			if (TestEqual(FString::Printf(TEXT("Ar << CompactBinary Save(%s)->Size"), Test), uint64(WriteAr.Num()), ExpectedData.GetSize()) &&
+				TestTrue(FString::Printf(TEXT("Ar << CompactBinary Save(%s)->EqualBytes"), Test), ExpectedData.EqualBytes(MakeMemoryView(WriteAr))))
+			{
+				FMemoryReader ReadAr(WriteAr);
+				ReadAr << Value;
+				CHECK_MESSAGE(FString::Printf(TEXT("Ar << CompactBinary Load(%s)->EqualBytes"), Test), ExpectedData.EqualBytes(Value.GetOuterBuffer().GetView()));
+			}
+		}
+	};
+
+	// Field
+	{
+		const uint8 Value[] = { uint8(ECbFieldType::IntegerPositive), 42 };
+		TestSave(TEXT("Field"), FCbField(FCbFieldView(Value), FSharedBuffer()), MakeMemoryView(Value));
+	}
+	{
+		const uint8 Value[] = { uint8(ECbFieldType::IntegerPositive | ECbFieldType::HasFieldName), 1, 'A', 42 };
+		TestSave(TEXT("Field, Name"), FCbField(FCbFieldView(Value), FSharedBuffer()), MakeMemoryView(Value));
+	}
+	{
+		const uint8 Value[] = { 42 };
+		TestSave(TEXT("Field, NoType"), FCbField(FCbFieldView(Value, ECbFieldType::IntegerPositive), FSharedBuffer()),
+			MakeMemoryView<uint8>({ uint8(ECbFieldType::IntegerPositive), 42 }));
+	}
+	{
+		const uint8 Value[] = { 1, 'I', 42 };
+		TestSave(TEXT("Field, NoType, Name"), FCbField(FCbFieldView(Value, ECbFieldType::IntegerPositive | ECbFieldType::HasFieldName), FSharedBuffer()),
+			MakeMemoryView<uint8>({ uint8(ECbFieldType::IntegerPositive | ECbFieldType::HasFieldName), 1, 'I', 42 }));
+	}
+
+	// Array
+	{
+		constexpr uint8 IntType = uint8(ECbFieldType::IntegerPositive);
+		const uint8 Value[] = { uint8(ECbFieldType::UniformArray), 5, 3, IntType, 1, 2, 3 };
+		TestSave(TEXT("Array"), FCbArray(FCbArrayView(Value), FSharedBuffer()),
+			MakeMemoryView(Value));
+	}
+	{
+		constexpr uint8 IntType = uint8(ECbFieldType::IntegerPositive);
+		const uint8 Value[] = { uint8(ECbFieldType::UniformArray | ECbFieldType::HasFieldName), 1, 'A', 5, 3, IntType, 1, 2, 3 };
+		TestSave(TEXT("Array, Name"), FCbArray(FCbArrayView(Value), FSharedBuffer()),
+			MakeMemoryView<uint8>({ uint8(ECbFieldType::UniformArray), 5, 3, IntType, 1, 2, 3 }));
+	}
+	{
+		constexpr uint8 IntType = uint8(ECbFieldType::IntegerPositive);
+		const uint8 Value[] = { 5, 3, IntType, 1, 2, 3 };
+		TestSave(TEXT("Array, NoType"), FCbArray(FCbArrayView(Value, ECbFieldType::UniformArray), FSharedBuffer()),
+			MakeMemoryView<uint8>({ uint8(ECbFieldType::UniformArray), 5, 3, IntType, 1, 2, 3 }));
+	}
+	{
+		constexpr uint8 IntType = uint8(ECbFieldType::IntegerPositive);
+		const uint8 Value[] = { 1, 'A', 5, 3, IntType, 1, 2, 3 };
+		TestSave(TEXT("Array, NoType, Name"), FCbArray(FCbArrayView(Value, ECbFieldType::UniformArray | ECbFieldType::HasFieldName), FSharedBuffer()),
+			MakeMemoryView<uint8>({ uint8(ECbFieldType::UniformArray), 5, 3, IntType, 1, 2, 3 }));
+	}
+
+	// Object
+	{
+		constexpr uint8 IntType = uint8(ECbFieldType::HasFieldName | ECbFieldType::IntegerPositive);
+		const uint8 Value[] = { uint8(ECbFieldType::UniformObject), 10, IntType, 1, 'A', 1, 1, 'B', 2, 1, 'C', 3 };
+		TestSave(TEXT("Object"), FCbObject(FCbObjectView(Value), FSharedBuffer()),
+			MakeMemoryView(Value));
+	}
+	{
+		constexpr uint8 IntType = uint8(ECbFieldType::HasFieldName | ECbFieldType::IntegerPositive);
+		const uint8 Value[] = { uint8(ECbFieldType::UniformObject | ECbFieldType::HasFieldName), 1, 'O', 10, IntType, 1, 'A', 1, 1, 'B', 2, 1, 'C', 3 };
+		TestSave(TEXT("Object, Name"), FCbObject(FCbObjectView(Value), FSharedBuffer()),
+			MakeMemoryView<uint8>({ uint8(ECbFieldType::UniformObject), 10, IntType, 1, 'A', 1, 1, 'B', 2, 1, 'C', 3 }));
+	}
+	{
+		constexpr uint8 IntType = uint8(ECbFieldType::HasFieldName | ECbFieldType::IntegerPositive);
+		const uint8 Value[] = { 10, IntType, 1, 'A', 1, 1, 'B', 2, 1, 'C', 3 };
+		TestSave(TEXT("Object, NoType"), FCbObject(FCbObjectView(Value, ECbFieldType::UniformObject), FSharedBuffer()),
+			MakeMemoryView<uint8>({ uint8(ECbFieldType::UniformObject), 10, IntType, 1, 'A', 1, 1, 'B', 2, 1, 'C', 3 }));
+	}
+	{
+		constexpr uint8 IntType = uint8(ECbFieldType::HasFieldName | ECbFieldType::IntegerPositive);
+		const uint8 Value[] = { 1, 'O', 10, IntType, 1, 'A', 1, 1, 'B', 2, 1, 'C', 3 };
+		TestSave(TEXT("Object, NoType, Name"), FCbObject(FCbObjectView(Value, ECbFieldType::UniformObject | ECbFieldType::HasFieldName), FSharedBuffer()),
+			MakeMemoryView<uint8>({ uint8(ECbFieldType::UniformObject), 10, IntType, 1, 'A', 1, 1, 'B', 2, 1, 'C', 3 }));
+	}
+
+}
+
+///////////////////////////////////////////////////////////////////////////////////////////////////////////////////////
+
+#endif // WITH_TESTS
diff --git a/Engine/Source/Runtime/Core/Tests/Serialization/CompactBinaryValidationTest.cpp b/Engine/Source/Runtime/Core/Tests/Serialization/CompactBinaryValidationTest.cpp
new file mode 100644
--- /dev/null
+++ b/Engine/Source/Runtime/Core/Tests/Serialization/CompactBinaryValidationTest.cpp
@@ -0,0 +1,624 @@
+// Copyright Epic Games, Inc. All Rights Reserved.
+
+#if WITH_TESTS
+
+#include "Serialization/CompactBinaryValidation.h"
+#include "Tests/TestHarnessAdapter.h"
+#include "Serialization/BufferArchive.h"
+#include "Serialization/CompactBinary.h"
+#include "Serialization/CompactBinaryPackage.h"
+#include "Serialization/CompactBinaryWriter.h"
+#if WITH_LOW_LEVEL_TESTS
+#include "TestCommon/Expectations.h"
+#endif
+
+TEST_CASE_NAMED(FCbValidateTest, "System::Core::Serialization::ValidateCompactBinary", "[ApplicationContextMask][SmokeFilter]")
+{
+	auto Validate = [](std::initializer_list<uint8> Data, ECbFieldType Type = ECbFieldType::HasFieldType) -> ECbValidateError
+	{
+		return ValidateCompactBinary(MakeMemoryView(Data), ECbValidateMode::All, Type);
+	};
+	auto ValidateMode = [](std::initializer_list<uint8> Data, ECbValidateMode Mode, ECbFieldType Type = ECbFieldType::HasFieldType) -> ECbValidateError
+	{
+		return ValidateCompactBinary(MakeMemoryView(Data), Mode, Type);
+	};
+
+	auto AddName = [](ECbFieldType Type) -> uint8 { return uint8(Type | ECbFieldType::HasFieldName); };
+
+	constexpr uint8 NullNoName = uint8(ECbFieldType::Null);
+	constexpr uint8 NullWithName = uint8(ECbFieldType::Null | ECbFieldType::HasFieldName);
+	constexpr uint8 IntNoName = uint8(ECbFieldType::IntegerPositive);
+	constexpr uint8 IntWithName = uint8(ECbFieldType::IntegerPositive | ECbFieldType::HasFieldName);
+
+	// Test OutOfBounds
+	CHECK_EQUALS(TEXT("ValidateCompactBinary(OutOfBounds, Empty)"), Validate({}), ECbValidateError::OutOfBounds);
+	CHECK_EQUALS(TEXT("ValidateCompactBinary(OutOfBounds, Empty, Mode)"), ValidateMode({}, ECbValidateMode::None), ECbValidateError::None);
+
+	CHECK_EQUALS(TEXT("ValidateCompactBinary(Valid, Null)"), Validate({ NullNoName }), ECbValidateError::None);
+	CHECK_EQUALS(TEXT("ValidateCompactBinary(Valid, Null, Name)"), Validate({ NullWithName, 1, 'N' }), ECbValidateError::None);
+	CHECK_EQUALS(TEXT("ValidateCompactBinary(OutOfBounds, Null, Name)"), Validate({ NullWithName }), ECbValidateError::OutOfBounds);
+	CHECK_EQUALS(TEXT("ValidateCompactBinary(OutOfBounds, Null, Name)"), Validate({ NullWithName, 1 }), ECbValidateError::OutOfBounds);
+	CHECK_EQUALS(TEXT("ValidateCompactBinary(OutOfBounds, Null, Name)"), Validate({ NullWithName, 0x80 }), ECbValidateError::OutOfBounds);
+	CHECK_EQUALS(TEXT("ValidateCompactBinary(OutOfBounds, Null, Name)"), Validate({ NullWithName, 0x80, 128 }), ECbValidateError::OutOfBounds);
+
+	CHECK_EQUALS(TEXT("ValidateCompactBinary(Valid, Object, Empty)"), Validate({ uint8(ECbFieldType::Object), 0 }), ECbValidateError::None);
+	CHECK_EQUALS(TEXT("ValidateCompactBinary(Valid, Object, Empty, NoType)"), Validate({ 0 }, ECbFieldType::Object), ECbValidateError::None);
+	CHECK_EQUALS(TEXT("ValidateCompactBinary(Valid, Object, Field)"), Validate({ uint8(ECbFieldType::Object), 7, NullWithName, 1, 'N', IntWithName, 1, 'I', 0 }), ECbValidateError::None);
+	CHECK_EQUALS(TEXT("ValidateCompactBinary(Valid, Object, Field, NoType)"), Validate({ 7, NullWithName, 1, 'N', IntWithName, 1, 'I', 0 }, ECbFieldType::Object), ECbValidateError::None);
+	CHECK_EQUALS(TEXT("ValidateCompactBinary(OutOfBounds, Object)"), Validate({ uint8(ECbFieldType::Object) }), ECbValidateError::OutOfBounds);
+	CHECK_EQUALS(TEXT("ValidateCompactBinary(OutOfBounds, Object, NoType)"), Validate({}, ECbFieldType::Object), ECbValidateError::OutOfBounds);
+	CHECK_EQUALS(TEXT("ValidateCompactBinary(OutOfBounds, Object)"), Validate({ uint8(ECbFieldType::Object), 1 }), ECbValidateError::OutOfBounds);
+	CHECK_EQUALS(TEXT("ValidateCompactBinary(OutOfBounds, Object, NoType)"), Validate({ 1 }, ECbFieldType::Object), ECbValidateError::OutOfBounds);
+	CHECK_EQUALS(TEXT("ValidateCompactBinary(OutOfBounds, Object, OOB Field)"), Validate({ uint8(ECbFieldType::Object), 3, AddName(ECbFieldType::Float32), 1, 'N' }), ECbValidateError::OutOfBounds);
+	CHECK_EQUALS(TEXT("ValidateCompactBinary(OutOfBounds, Object, OOB Field, NoType)"), Validate({ 3, AddName(ECbFieldType::Float32), 1, 'N' }, ECbFieldType::Object), ECbValidateError::OutOfBounds);
+
+	CHECK_EQUALS(TEXT("ValidateCompactBinary(Valid, UniformObject, Field)"), Validate({ uint8(ECbFieldType::UniformObject), 3, NullWithName, 1, 'N' }), ECbValidateError::None);
+	CHECK_EQUALS(TEXT("ValidateCompactBinary(Valid, UniformObject, Field, NoType)"), Validate({ 3, NullWithName, 1, 'N' }, ECbFieldType::UniformObject), ECbValidateError::None);
+	CHECK_EQUALS(TEXT("ValidateCompactBinary(OutOfBounds, UniformObject)"), Validate({ uint8(ECbFieldType::UniformObject) }), ECbValidateError::OutOfBounds);
+	CHECK_EQUALS(TEXT("ValidateCompactBinary(OutOfBounds, UniformObject, NoType)"), Validate({}, ECbFieldType::UniformObject), ECbValidateError::OutOfBounds);
+	CHECK_EQUALS(TEXT("ValidateCompactBinary(OutOfBounds, UniformObject)"), Validate({ uint8(ECbFieldType::UniformObject), 1 }), ECbValidateError::OutOfBounds);
+	CHECK_EQUALS(TEXT("ValidateCompactBinary(OutOfBounds, UniformObject, NoType)"), Validate({ 1 }, ECbFieldType::UniformObject), ECbValidateError::OutOfBounds);
+	CHECK_EQUALS(TEXT("ValidateCompactBinary(OutOfBounds, UniformObject, OOB Field)"), Validate({ uint8(ECbFieldType::UniformObject), 3, AddName(ECbFieldType::Float32), 1, 'N' }), ECbValidateError::OutOfBounds);
+	CHECK_EQUALS(TEXT("ValidateCompactBinary(OutOfBounds, UniformObject, OOB Field, NoType)"), Validate({ 3, AddName(ECbFieldType::Float32), 1, 'N' }, ECbFieldType::UniformObject), ECbValidateError::OutOfBounds);
+
+	CHECK_EQUALS(TEXT("ValidateCompactBinary(Valid, Array, Empty)"), Validate({ uint8(ECbFieldType::Array), 1, 0 }), ECbValidateError::None);
+	CHECK_EQUALS(TEXT("ValidateCompactBinary(Valid, Array, Empty, NoType)"), Validate({ 1, 0 }, ECbFieldType::Array), ECbValidateError::None);
+	CHECK_EQUALS(TEXT("ValidateCompactBinary(Valid, Array, Field)"), Validate({ uint8(ECbFieldType::Array), 4, 2, NullNoName, IntNoName, 0 }), ECbValidateError::None);
+	CHECK_EQUALS(TEXT("ValidateCompactBinary(Valid, Array, Field, NoType)"), Validate({ 4, 2, NullNoName, IntNoName, 0 }, ECbFieldType::Array), ECbValidateError::None);
+	CHECK_EQUALS(TEXT("ValidateCompactBinary(OutOfBounds, Array)"), Validate({ uint8(ECbFieldType::Array) }), ECbValidateError::OutOfBounds);
+	CHECK_EQUALS(TEXT("ValidateCompactBinary(OutOfBounds, Array, NoType)"), Validate({}, ECbFieldType::Array), ECbValidateError::OutOfBounds);
+	CHECK_EQUALS(TEXT("ValidateCompactBinary(OutOfBounds, Array)"), Validate({ uint8(ECbFieldType::Array), 1 }), ECbValidateError::OutOfBounds);
+	CHECK_EQUALS(TEXT("ValidateCompactBinary(OutOfBounds, Array, NoType)"), Validate({ 1 }, ECbFieldType::Array), ECbValidateError::OutOfBounds);
+	CHECK_EQUALS(TEXT("ValidateCompactBinary(OutOfBounds, Array, OOB Field)"), Validate({ uint8(ECbFieldType::Array), 2, 1, uint8(ECbFieldType::Float32) }), ECbValidateError::OutOfBounds);
+	CHECK_EQUALS(TEXT("ValidateCompactBinary(OutOfBounds, Array, OOB Field, NoType)"), Validate({ 2, 1, uint8(ECbFieldType::Float32) }, ECbFieldType::Array), ECbValidateError::OutOfBounds);
+
+	CHECK_EQUALS(TEXT("ValidateCompactBinary(Valid, UniformArray, Field)"), Validate({ uint8(ECbFieldType::UniformArray), 3, 1, IntNoName, 0 }), ECbValidateError::None);
+	CHECK_EQUALS(TEXT("ValidateCompactBinary(Valid, UniformArray, Field, NoType)"), Validate({ 3, 1, IntNoName, 0 }, ECbFieldType::UniformArray), ECbValidateError::None);
+	CHECK_EQUALS(TEXT("ValidateCompactBinary(OutOfBounds, UniformArray)"), Validate({ uint8(ECbFieldType::UniformArray) }), ECbValidateError::OutOfBounds);
+	CHECK_EQUALS(TEXT("ValidateCompactBinary(OutOfBounds, UniformArray, NoType)"), Validate({}, ECbFieldType::UniformArray), ECbValidateError::OutOfBounds);
+	CHECK_EQUALS(TEXT("ValidateCompactBinary(OutOfBounds, UniformArray)"), Validate({ uint8(ECbFieldType::UniformArray), 1 }), ECbValidateError::OutOfBounds);
+	CHECK_EQUALS(TEXT("ValidateCompactBinary(OutOfBounds, UniformArray, NoType)"), Validate({ 1 }, ECbFieldType::UniformArray), ECbValidateError::OutOfBounds);
+	CHECK_EQUALS(TEXT("ValidateCompactBinary(OutOfBounds, UniformArray, OOB Field)"), Validate({ uint8(ECbFieldType::UniformArray), 2, 1, uint8(ECbFieldType::Float32) }), ECbValidateError::OutOfBounds);
+	CHECK_EQUALS(TEXT("ValidateCompactBinary(OutOfBounds, UniformArray, OOB Field, NoType)"), Validate({ 2, 1, uint8(ECbFieldType::Float32) }, ECbFieldType::UniformArray), ECbValidateError::OutOfBounds);
+
+	CHECK_EQUALS(TEXT("ValidateCompactBinary(Valid, Binary, Empty)"), Validate({ uint8(ECbFieldType::Binary), 0 }), ECbValidateError::None);
+	CHECK_EQUALS(TEXT("ValidateCompactBinary(Valid, Binary, Empty, NoType)"), Validate({ 0 }, ECbFieldType::Binary), ECbValidateError::None);
+	CHECK_EQUALS(TEXT("ValidateCompactBinary(Valid, Binary, Field)"), Validate({ uint8(ECbFieldType::Binary), 1, 0 }), ECbValidateError::None);
+	CHECK_EQUALS(TEXT("ValidateCompactBinary(Valid, Binary, Field, NoType)"), Validate({ 1, 0 }, ECbFieldType::Binary), ECbValidateError::None);
+	CHECK_EQUALS(TEXT("ValidateCompactBinary(OutOfBounds, Binary)"), Validate({ uint8(ECbFieldType::Binary) }), ECbValidateError::OutOfBounds);
+	CHECK_EQUALS(TEXT("ValidateCompactBinary(OutOfBounds, Binary, NoType)"), Validate({}, ECbFieldType::Binary), ECbValidateError::OutOfBounds);
+	CHECK_EQUALS(TEXT("ValidateCompactBinary(OutOfBounds, Binary)"), Validate({ uint8(ECbFieldType::Binary), 1 }), ECbValidateError::OutOfBounds);
+	CHECK_EQUALS(TEXT("ValidateCompactBinary(OutOfBounds, Binary, NoType)"), Validate({ 1 }, ECbFieldType::Binary), ECbValidateError::OutOfBounds);
+
+	CHECK_EQUALS(TEXT("ValidateCompactBinary(Valid, String, Empty)"), Validate({ uint8(ECbFieldType::String), 0 }), ECbValidateError::None);
+	CHECK_EQUALS(TEXT("ValidateCompactBinary(Valid, String, Empty, NoType)"), Validate({ 0 }, ECbFieldType::String), ECbValidateError::None);
+	CHECK_EQUALS(TEXT("ValidateCompactBinary(Valid, String, Field)"), Validate({ uint8(ECbFieldType::String), 1, 'S' }), ECbValidateError::None);
+	CHECK_EQUALS(TEXT("ValidateCompactBinary(Valid, String, Field, NoType)"), Validate({ 1, 'S' }, ECbFieldType::String), ECbValidateError::None);
+	CHECK_EQUALS(TEXT("ValidateCompactBinary(OutOfBounds, String)"), Validate({ uint8(ECbFieldType::String) }), ECbValidateError::OutOfBounds);
+	CHECK_EQUALS(TEXT("ValidateCompactBinary(OutOfBounds, String, NoType)"), Validate({}, ECbFieldType::String), ECbValidateError::OutOfBounds);
+	CHECK_EQUALS(TEXT("ValidateCompactBinary(OutOfBounds, String)"), Validate({ uint8(ECbFieldType::String), 1 }), ECbValidateError::OutOfBounds);
+	CHECK_EQUALS(TEXT("ValidateCompactBinary(OutOfBounds, String, NoType)"), Validate({ 1 }, ECbFieldType::String), ECbValidateError::OutOfBounds);
+
+	CHECK_EQUALS(TEXT("ValidateCompactBinary(Valid, IntegerPositive, 1-byte)"), Validate({ uint8(ECbFieldType::IntegerPositive), 0 }), ECbValidateError::None);
+	CHECK_EQUALS(TEXT("ValidateCompactBinary(Valid, IntegerPositive, 1-byte, NoType)"), Validate({ 0 }, ECbFieldType::IntegerPositive), ECbValidateError::None);
+	CHECK_EQUALS(TEXT("ValidateCompactBinary(Valid, IntegerPositive, 2-byte)"), Validate({ uint8(ECbFieldType::IntegerPositive), 0x80, 0x80 }), ECbValidateError::None);
+	CHECK_EQUALS(TEXT("ValidateCompactBinary(Valid, IntegerPositive, 2-byte, NoType)"), Validate({ 0x80, 0x80 }, ECbFieldType::IntegerPositive), ECbValidateError::None);
+	CHECK_EQUALS(TEXT("ValidateCompactBinary(OutOfBounds, IntegerPositive, 1-byte)"), Validate({ uint8(ECbFieldType::IntegerPositive) }), ECbValidateError::OutOfBounds);
+	CHECK_EQUALS(TEXT("ValidateCompactBinary(OutOfBounds, IntegerPositive, 1-byte, NoType)"), Validate({}, ECbFieldType::IntegerPositive), ECbValidateError::OutOfBounds);
+	CHECK_EQUALS(TEXT("ValidateCompactBinary(OutOfBounds, IntegerPositive, 2-byte)"), Validate({ uint8(ECbFieldType::IntegerPositive), 0x80 }), ECbValidateError::OutOfBounds);
+	CHECK_EQUALS(TEXT("ValidateCompactBinary(OutOfBounds, IntegerPositive, 2-byte, NoType)"), Validate({ 0x80 }, ECbFieldType::IntegerPositive), ECbValidateError::OutOfBounds);
+	CHECK_EQUALS(TEXT("ValidateCompactBinary(OutOfBounds, IntegerPositive, 9-byte)"), Validate({ uint8(ECbFieldType::IntegerPositive), 0xff, 0, 0, 0, 0, 0, 0, 0 }), ECbValidateError::OutOfBounds);
+	CHECK_EQUALS(TEXT("ValidateCompactBinary(OutOfBounds, IntegerPositive, 9-byte, NoType)"), Validate({ 0xff, 0, 0, 0, 0, 0, 0, 0 }, ECbFieldType::IntegerPositive), ECbValidateError::OutOfBounds);
+
+	CHECK_EQUALS(TEXT("ValidateCompactBinary(Valid, IntegerNegative, 1-byte)"), Validate({ uint8(ECbFieldType::IntegerNegative), 0 }), ECbValidateError::None);
+	CHECK_EQUALS(TEXT("ValidateCompactBinary(Valid, IntegerNegative, 1-byte, NoType)"), Validate({ 0 }, ECbFieldType::IntegerNegative), ECbValidateError::None);
+	CHECK_EQUALS(TEXT("ValidateCompactBinary(Valid, IntegerNegative, 2-byte)"), Validate({ uint8(ECbFieldType::IntegerNegative), 0x80, 0x80 }), ECbValidateError::None);
+	CHECK_EQUALS(TEXT("ValidateCompactBinary(Valid, IntegerNegative, 2-byte, NoType)"), Validate({ 0x80, 0x80 }, ECbFieldType::IntegerNegative), ECbValidateError::None);
+	CHECK_EQUALS(TEXT("ValidateCompactBinary(OutOfBounds, IntegerNegative, 1-byte)"), Validate({ uint8(ECbFieldType::IntegerNegative) }), ECbValidateError::OutOfBounds);
+	CHECK_EQUALS(TEXT("ValidateCompactBinary(OutOfBounds, IntegerNegative, 1-byte, NoType)"), Validate({}, ECbFieldType::IntegerNegative), ECbValidateError::OutOfBounds);
+	CHECK_EQUALS(TEXT("ValidateCompactBinary(OutOfBounds, IntegerNegative, 2-byte)"), Validate({ uint8(ECbFieldType::IntegerNegative), 0x80 }), ECbValidateError::OutOfBounds);
+	CHECK_EQUALS(TEXT("ValidateCompactBinary(OutOfBounds, IntegerNegative, 2-byte, NoType)"), Validate({ 0x80 }, ECbFieldType::IntegerNegative), ECbValidateError::OutOfBounds);
+	CHECK_EQUALS(TEXT("ValidateCompactBinary(OutOfBounds, IntegerNegative, 9-byte)"), Validate({ uint8(ECbFieldType::IntegerNegative), 0xff, 0, 0, 0, 0, 0, 0, 0 }), ECbValidateError::OutOfBounds);
+	CHECK_EQUALS(TEXT("ValidateCompactBinary(OutOfBounds, IntegerNegative, 9-byte, NoType)"), Validate({ 0xff, 0, 0, 0, 0, 0, 0, 0 }, ECbFieldType::IntegerNegative), ECbValidateError::OutOfBounds);
+
+	CHECK_EQUALS(TEXT("ValidateCompactBinary(Valid, Float32)"), Validate({ uint8(ECbFieldType::Float32), 0, 0, 0, 0 }), ECbValidateError::None);
+	CHECK_EQUALS(TEXT("ValidateCompactBinary(Valid, Float32, NoType)"), Validate({ 0, 0, 0, 0 }, ECbFieldType::Float32), ECbValidateError::None);
+	CHECK_EQUALS(TEXT("ValidateCompactBinary(OutOfBounds, Float32)"), Validate({ uint8(ECbFieldType::Float32), 0, 0, 0 }), ECbValidateError::OutOfBounds);
+	CHECK_EQUALS(TEXT("ValidateCompactBinary(OutOfBounds, Float32, NoType)"), Validate({ 0, 0, 0 }, ECbFieldType::Float32), ECbValidateError::OutOfBounds);
+
+	CHECK_EQUALS(TEXT("ValidateCompactBinary(Valid, Float64)"), Validate({ uint8(ECbFieldType::Float64), 0x3f, 0xff, 0xff, 0xff, 0xf0, 0x00, 0x00, 0x00 }), ECbValidateError::None);
+	CHECK_EQUALS(TEXT("ValidateCompactBinary(Valid, Float64, NoType)"), Validate({ 0x3f, 0xff, 0xff, 0xff, 0xf0, 0x00, 0x00, 0x00 }, ECbFieldType::Float64), ECbValidateError::None);
+	CHECK_EQUALS(TEXT("ValidateCompactBinary(OutOfBounds, Float64)"), Validate({ uint8(ECbFieldType::Float64), 0, 0, 0, 0, 0, 0, 0 }), ECbValidateError::OutOfBounds);
+	CHECK_EQUALS(TEXT("ValidateCompactBinary(OutOfBounds, Float64, NoType)"), Validate({ 0, 0, 0, 0, 0, 0, 0 }, ECbFieldType::Float64), ECbValidateError::OutOfBounds);
+
+	CHECK_EQUALS(TEXT("ValidateCompactBinary(Valid, BoolFalse)"), Validate({ uint8(ECbFieldType::BoolFalse) }), ECbValidateError::None);
+	CHECK_EQUALS(TEXT("ValidateCompactBinary(Valid, BoolTrue)"), Validate({ uint8(ECbFieldType::BoolTrue) }), ECbValidateError::None);
+
+	CHECK_EQUALS(TEXT("ValidateCompactBinary(Valid, ObjectAttachment)"), Validate({ uint8(ECbFieldType::ObjectAttachment), 0, 0, 0, 0, 0, 0, 0, 0, 0, 0, 0, 0, 0, 0, 0, 0, 0, 0, 0, 0 }), ECbValidateError::None);
+	CHECK_EQUALS(TEXT("ValidateCompactBinary(Valid, ObjectAttachment, NoType)"), Validate({ 0, 0, 0, 0, 0, 0, 0, 0, 0, 0, 0, 0, 0, 0, 0, 0, 0, 0, 0, 0 }, ECbFieldType::ObjectAttachment), ECbValidateError::None);
+	CHECK_EQUALS(TEXT("ValidateCompactBinary(OutOfBounds, ObjectAttachment)"), Validate({ uint8(ECbFieldType::ObjectAttachment), 0, 0, 0, 0, 0, 0, 0, 0, 0, 0, 0, 0, 0, 0, 0, 0, 0, 0, 0 }), ECbValidateError::OutOfBounds);
+	CHECK_EQUALS(TEXT("ValidateCompactBinary(OutOfBounds, ObjectAttachment, NoType)"), Validate({ 0, 0, 0, 0, 0, 0, 0, 0, 0, 0, 0, 0, 0, 0, 0, 0, 0, 0, 0 }, ECbFieldType::ObjectAttachment), ECbValidateError::OutOfBounds);
+
+	CHECK_EQUALS(TEXT("ValidateCompactBinary(Valid, BinaryAttachment)"), Validate({ uint8(ECbFieldType::BinaryAttachment), 0, 0, 0, 0, 0, 0, 0, 0, 0, 0, 0, 0, 0, 0, 0, 0, 0, 0, 0, 0 }), ECbValidateError::None);
+	CHECK_EQUALS(TEXT("ValidateCompactBinary(Valid, BinaryAttachment, NoType)"), Validate({ 0, 0, 0, 0, 0, 0, 0, 0, 0, 0, 0, 0, 0, 0, 0, 0, 0, 0, 0, 0 }, ECbFieldType::BinaryAttachment), ECbValidateError::None);
+	CHECK_EQUALS(TEXT("ValidateCompactBinary(OutOfBounds, BinaryAttachment)"), Validate({ uint8(ECbFieldType::BinaryAttachment), 0, 0, 0, 0, 0, 0, 0, 0, 0, 0, 0, 0, 0, 0, 0, 0, 0, 0, 0 }), ECbValidateError::OutOfBounds);
+	CHECK_EQUALS(TEXT("ValidateCompactBinary(OutOfBounds, BinaryAttachment, NoType)"), Validate({ 0, 0, 0, 0, 0, 0, 0, 0, 0, 0, 0, 0, 0, 0, 0, 0, 0, 0, 0 }, ECbFieldType::BinaryAttachment), ECbValidateError::OutOfBounds);
+
+	CHECK_EQUALS(TEXT("ValidateCompactBinary(Valid, Hash)"), Validate({ uint8(ECbFieldType::Hash), 0, 0, 0, 0, 0, 0, 0, 0, 0, 0, 0, 0, 0, 0, 0, 0, 0, 0, 0, 0 }), ECbValidateError::None);
+	CHECK_EQUALS(TEXT("ValidateCompactBinary(Valid, Hash, NoType)"), Validate({ 0, 0, 0, 0, 0, 0, 0, 0, 0, 0, 0, 0, 0, 0, 0, 0, 0, 0, 0, 0 }, ECbFieldType::Hash), ECbValidateError::None);
+	CHECK_EQUALS(TEXT("ValidateCompactBinary(OutOfBounds, Hash)"), Validate({ uint8(ECbFieldType::Hash), 0, 0, 0, 0, 0, 0, 0, 0, 0, 0, 0, 0, 0, 0, 0, 0, 0, 0, 0 }), ECbValidateError::OutOfBounds);
+	CHECK_EQUALS(TEXT("ValidateCompactBinary(OutOfBounds, Hash, NoType)"), Validate({ 0, 0, 0, 0, 0, 0, 0, 0, 0, 0, 0, 0, 0, 0, 0, 0, 0, 0, 0 }, ECbFieldType::Hash), ECbValidateError::OutOfBounds);
+
+	CHECK_EQUALS(TEXT("ValidateCompactBinary(Valid, Uuid)"), Validate({ uint8(ECbFieldType::Uuid), 0, 0, 0, 0, 0, 0, 0, 0, 0, 0, 0, 0, 0, 0, 0, 0 }), ECbValidateError::None);
+	CHECK_EQUALS(TEXT("ValidateCompactBinary(Valid, Uuid, NoType)"), Validate({ 0, 0, 0, 0, 0, 0, 0, 0, 0, 0, 0, 0, 0, 0, 0, 0 }, ECbFieldType::Uuid), ECbValidateError::None);
+	CHECK_EQUALS(TEXT("ValidateCompactBinary(OutOfBounds, Uuid)"), Validate({ uint8(ECbFieldType::Uuid), 0, 0, 0, 0, 0, 0, 0, 0, 0, 0, 0, 0, 0, 0, 0 }), ECbValidateError::OutOfBounds);
+	CHECK_EQUALS(TEXT("ValidateCompactBinary(OutOfBounds, Uuid, NoType)"), Validate({ 0, 0, 0, 0, 0, 0, 0, 0, 0, 0, 0, 0, 0, 0, 0 }, ECbFieldType::Uuid), ECbValidateError::OutOfBounds);
+
+	CHECK_EQUALS(TEXT("ValidateCompactBinary(Valid, DateTime)"), Validate({ uint8(ECbFieldType::DateTime), 0, 0, 0, 0, 0, 0, 0, 0 }), ECbValidateError::None);
+	CHECK_EQUALS(TEXT("ValidateCompactBinary(Valid, DateTime, NoType)"), Validate({ 0, 0, 0, 0, 0, 0, 0, 0 }, ECbFieldType::DateTime), ECbValidateError::None);
+	CHECK_EQUALS(TEXT("ValidateCompactBinary(OutOfBounds, DateTime)"), Validate({ uint8(ECbFieldType::DateTime), 0, 0, 0, 0, 0, 0, 0 }), ECbValidateError::OutOfBounds);
+	CHECK_EQUALS(TEXT("ValidateCompactBinary(OutOfBounds, DateTime, NoType)"), Validate({ 0, 0, 0, 0, 0, 0, 0 }, ECbFieldType::DateTime), ECbValidateError::OutOfBounds);
+
+	CHECK_EQUALS(TEXT("ValidateCompactBinary(Valid, TimeSpan)"), Validate({ uint8(ECbFieldType::TimeSpan), 0, 0, 0, 0, 0, 0, 0, 0 }), ECbValidateError::None);
+	CHECK_EQUALS(TEXT("ValidateCompactBinary(Valid, TimeSpan, NoType)"), Validate({ 0, 0, 0, 0, 0, 0, 0, 0 }, ECbFieldType::TimeSpan), ECbValidateError::None);
+	CHECK_EQUALS(TEXT("ValidateCompactBinary(OutOfBounds, TimeSpan)"), Validate({ uint8(ECbFieldType::TimeSpan), 0, 0, 0, 0, 0, 0, 0 }), ECbValidateError::OutOfBounds);
+	CHECK_EQUALS(TEXT("ValidateCompactBinary(OutOfBounds, TimeSpan, NoType)"), Validate({ 0, 0, 0, 0, 0, 0, 0 }, ECbFieldType::TimeSpan), ECbValidateError::OutOfBounds);
+
+	CHECK_EQUALS(TEXT("ValidateCompactBinary(Valid, ObjectId)"), Validate({ uint8(ECbFieldType::ObjectId), 0, 0, 0, 0, 0, 0, 0, 0, 0, 0, 0, 0 }), ECbValidateError::None);
+	CHECK_EQUALS(TEXT("ValidateCompactBinary(Valid, ObjectId, NoType)"), Validate({ 0, 0, 0, 0, 0, 0, 0, 0, 0, 0, 0, 0 }, ECbFieldType::ObjectId), ECbValidateError::None);
+	CHECK_EQUALS(TEXT("ValidateCompactBinary(OutOfBounds, ObjectId)"), Validate({ uint8(ECbFieldType::ObjectId), 0, 0, 0, 0, 0, 0, 0, 0, 0, 0, 0 }), ECbValidateError::OutOfBounds);
+	CHECK_EQUALS(TEXT("ValidateCompactBinary(OutOfBounds, ObjectId, NoType)"), Validate({ 0, 0, 0, 0, 0, 0, 0, 0, 0, 0, 0 }, ECbFieldType::ObjectId), ECbValidateError::OutOfBounds);
+
+	CHECK_EQUALS(TEXT("ValidateCompactBinary(Valid, CustomById, Empty)"), Validate({ uint8(ECbFieldType::CustomById), 1, 0 }), ECbValidateError::None);
+	CHECK_EQUALS(TEXT("ValidateCompactBinary(Valid, CustomById, Empty, NoType)"), Validate({ 1, 0 }, ECbFieldType::CustomById), ECbValidateError::None);
+	CHECK_EQUALS(TEXT("ValidateCompactBinary(Valid, CustomById, Value)"), Validate({ uint8(ECbFieldType::CustomById), 2, 0, 0 }), ECbValidateError::None);
+	CHECK_EQUALS(TEXT("ValidateCompactBinary(Valid, CustomById, Value, NoType)"), Validate({ 2, 0, 0 }, ECbFieldType::CustomById), ECbValidateError::None);
+	CHECK_EQUALS(TEXT("ValidateCompactBinary(Valid, CustomById, MaxTypeId)"), Validate({ uint8(ECbFieldType::CustomById), 9, 0xff, 0xff, 0xff, 0xff, 0xff, 0xff, 0xff, 0xff, 0xff }), ECbValidateError::None);
+	CHECK_EQUALS(TEXT("ValidateCompactBinary(Valid, CustomById, MaxTypeId, NoType)"), Validate({ 9, 0xff, 0xff, 0xff, 0xff, 0xff, 0xff, 0xff, 0xff, 0xff }, ECbFieldType::CustomById), ECbValidateError::None);
+	CHECK_EQUALS(TEXT("ValidateCompactBinary(Valid, CustomById, MaxTypeId, Value)"), Validate({ uint8(ECbFieldType::CustomById), 10, 0xff, 0xff, 0xff, 0xff, 0xff, 0xff, 0xff, 0xff, 0xff, 0 }), ECbValidateError::None);
+	CHECK_EQUALS(TEXT("ValidateCompactBinary(Valid, CustomById, MaxTypeId, Value, NoType)"), Validate({ 10, 0xff, 0xff, 0xff, 0xff, 0xff, 0xff, 0xff, 0xff, 0xff, 0 }, ECbFieldType::CustomById), ECbValidateError::None);
+	CHECK_EQUALS(TEXT("ValidateCompactBinary(OutOfBounds, CustomById, NoSize)"), Validate({ uint8(ECbFieldType::CustomById) }), ECbValidateError::OutOfBounds);
+	CHECK_EQUALS(TEXT("ValidateCompactBinary(OutOfBounds, CustomById, NoSize, NoType)"), Validate({}, ECbFieldType::CustomById), ECbValidateError::OutOfBounds);
+	CHECK_EQUALS(TEXT("ValidateCompactBinary(OutOfBounds, CustomById, NoTypeId)"), Validate({ uint8(ECbFieldType::CustomById), 0 }), ECbValidateError::OutOfBounds);
+	CHECK_EQUALS(TEXT("ValidateCompactBinary(OutOfBounds, CustomById, NoTypeId, NoType)"), Validate({ 0 }, ECbFieldType::CustomById), ECbValidateError::OutOfBounds);
+	CHECK_EQUALS(TEXT("ValidateCompactBinary(OutOfBounds, CustomById, ShortTypeId)"), Validate({ uint8(ECbFieldType::CustomById), 1, 0x80 }), ECbValidateError::OutOfBounds);
+	CHECK_EQUALS(TEXT("ValidateCompactBinary(OutOfBounds, CustomById, ShortTypeId, NoType)"), Validate({ 1, 0x80 }, ECbFieldType::CustomById), ECbValidateError::OutOfBounds);
+	CHECK_EQUALS(TEXT("ValidateCompactBinary(OutOfBounds, CustomById, NoValue)"), Validate({ uint8(ECbFieldType::CustomById), 2, 0 }), ECbValidateError::OutOfBounds);
+	CHECK_EQUALS(TEXT("ValidateCompactBinary(OutOfBounds, CustomById, NoValue, NoType)"), Validate({ 2, 0 }, ECbFieldType::CustomById), ECbValidateError::OutOfBounds);
+	CHECK_EQUALS(TEXT("ValidateCompactBinary(OutOfBounds, CustomById, ShortValue)"), Validate({ uint8(ECbFieldType::CustomById), 3, 0, 0 }), ECbValidateError::OutOfBounds);
+	CHECK_EQUALS(TEXT("ValidateCompactBinary(OutOfBounds, CustomById, ShortValue, NoType)"), Validate({ 3, 0, 0 }, ECbFieldType::CustomById), ECbValidateError::OutOfBounds);
+
+	CHECK_EQUALS(TEXT("ValidateCompactBinary(Valid, CustomByName, Empty)"), Validate({ uint8(ECbFieldType::CustomByName), 2, 1, 'A' }), ECbValidateError::None);
+	CHECK_EQUALS(TEXT("ValidateCompactBinary(Valid, CustomByName, Empty, NoType)"), Validate({ 2, 1, 'A' }, ECbFieldType::CustomByName), ECbValidateError::None);
+	CHECK_EQUALS(TEXT("ValidateCompactBinary(Valid, CustomByName, Value)"), Validate({ uint8(ECbFieldType::CustomByName), 3, 1, 'A', 0 }), ECbValidateError::None);
+	CHECK_EQUALS(TEXT("ValidateCompactBinary(Valid, CustomByName, Value, NoType)"), Validate({ 3, 1, 'A', 0 }, ECbFieldType::CustomByName), ECbValidateError::None);
+	CHECK_EQUALS(TEXT("ValidateCompactBinary(OutOfBounds, CustomByName, NoSize)"), Validate({ uint8(ECbFieldType::CustomByName) }), ECbValidateError::OutOfBounds);
+	CHECK_EQUALS(TEXT("ValidateCompactBinary(OutOfBounds, CustomByName, NoSize, NoType)"), Validate({}, ECbFieldType::CustomByName), ECbValidateError::OutOfBounds);
+	CHECK_EQUALS(TEXT("ValidateCompactBinary(OutOfBounds, CustomByName, NoTypeSize)"), Validate({ uint8(ECbFieldType::CustomByName), 0 }), ECbValidateError::OutOfBounds);
+	CHECK_EQUALS(TEXT("ValidateCompactBinary(OutOfBounds, CustomByName, NoTypeSize, NoType)"), Validate({ 0 }, ECbFieldType::CustomByName), ECbValidateError::OutOfBounds);
+	CHECK_EQUALS(TEXT("ValidateCompactBinary(OutOfBounds, CustomByName, ShortTypeSize)"), Validate({ uint8(ECbFieldType::CustomByName), 1, 0x80 }), ECbValidateError::OutOfBounds);
+	CHECK_EQUALS(TEXT("ValidateCompactBinary(OutOfBounds, CustomByName, ShortTypeSize, NoType)"), Validate({ 1, 0x80 }, ECbFieldType::CustomByName), ECbValidateError::OutOfBounds);
+	CHECK_EQUALS(TEXT("ValidateCompactBinary(OutOfBounds, CustomByName, ShortTypeName)"), Validate({ uint8(ECbFieldType::CustomByName), 1, 1 }), ECbValidateError::OutOfBounds);
+	CHECK_EQUALS(TEXT("ValidateCompactBinary(OutOfBounds, CustomByName, ShortTypeName, NoType)"), Validate({ 1, 1 }, ECbFieldType::CustomByName), ECbValidateError::OutOfBounds);
+	CHECK_EQUALS(TEXT("ValidateCompactBinary(OutOfBounds, CustomByName, NoValue)"), Validate({ uint8(ECbFieldType::CustomByName), 3, 1, 'A' }), ECbValidateError::OutOfBounds);
+	CHECK_EQUALS(TEXT("ValidateCompactBinary(OutOfBounds, CustomByName, NoValue, NoType)"), Validate({ 3, 1, 'A' }, ECbFieldType::CustomByName), ECbValidateError::OutOfBounds);
+	CHECK_EQUALS(TEXT("ValidateCompactBinary(OutOfBounds, CustomByName, ShortValue)"), Validate({ uint8(ECbFieldType::CustomByName), 4, 1, 'A', 0 }), ECbValidateError::OutOfBounds);
+	CHECK_EQUALS(TEXT("ValidateCompactBinary(OutOfBounds, CustomByName, ShortValue, NoType)"), Validate({ 4, 1, 'A', 0 }, ECbFieldType::CustomByName), ECbValidateError::OutOfBounds);
+
+	// Test InvalidType
+	CHECK_EQUALS(TEXT("ValidateCompactBinary(InvalidType, Unknown)"), Validate({ uint8(ECbFieldType::ObjectId) + 1 }), ECbValidateError::InvalidType);
+	CHECK_EQUALS(TEXT("ValidateCompactBinary(InvalidType, Unknown)"), Validate({}, ECbFieldType(uint8(ECbFieldType::ObjectId) + 1)), ECbValidateError::InvalidType);
+	CHECK_EQUALS(TEXT("ValidateCompactBinary(InvalidType, HasFieldType)"), Validate({ uint8(ECbFieldType::Null | ECbFieldType::HasFieldType) }), ECbValidateError::InvalidType);
+
+	CHECK_EQUALS(TEXT("ValidateCompactBinary(InvalidType, ZeroSizeField)"), Validate({}, ECbFieldType::Null), ECbValidateError::InvalidType);
+	CHECK_EQUALS(TEXT("ValidateCompactBinary(InvalidType, ZeroSizeField, BoolFalse)"), Validate({}, ECbFieldType::BoolFalse), ECbValidateError::InvalidType);
+	CHECK_EQUALS(TEXT("ValidateCompactBinary(InvalidType, ZeroSizeField, BoolTrue)"), Validate({}, ECbFieldType::BoolTrue), ECbValidateError::InvalidType);
+
+	CHECK_EQUALS(TEXT("ValidateCompactBinary(InvalidType, ZeroSizeField, Array)"), Validate({ uint8(ECbFieldType::UniformArray), 2, 2, NullNoName }), ECbValidateError::InvalidType);
+	CHECK_EQUALS(TEXT("ValidateCompactBinary(InvalidType, ZeroSizeField, Object)"), Validate({ uint8(ECbFieldType::UniformObject), 2, NullNoName, 0 }), ECbValidateError::InvalidType);
+
+	CHECK_EQUALS(TEXT("ValidateCompactBinary(InvalidType, CustomByName, EmptyTypeName)"), Validate({ uint8(ECbFieldType::CustomByName), 1, 0 }), ECbValidateError::InvalidType);
+	CHECK_EQUALS(TEXT("ValidateCompactBinary(InvalidType, CustomByName, EmptyTypeName, NoType)"), Validate({ 1, 0 }, ECbFieldType::CustomByName), ECbValidateError::InvalidType);
+
+	// Test DuplicateName
+	CHECK_EQUALS(TEXT("ValidateCompactBinary(DuplicateName)"), Validate({ uint8(ECbFieldType::UniformObject), 7, NullWithName, 1, 'A', 1, 'B', 1, 'A' }), ECbValidateError::DuplicateName);
+	CHECK_EQUALS(TEXT("ValidateCompactBinary(DuplicateName, CaseSensitive)"), Validate({ uint8(ECbFieldType::UniformObject), 7, NullWithName, 1, 'A', 1, 'B', 1, 'a' }), ECbValidateError::None);
+	CHECK_EQUALS(TEXT("ValidateCompactBinary(DuplicateName, Mode)"), ValidateMode({ uint8(ECbFieldType::UniformObject), 7, NullWithName, 1, 'A', 1, 'B', 1, 'A' }, ECbValidateMode::All & ~ECbValidateMode::Names), ECbValidateError::None);
+
+	// Test MissingName
+	CHECK_EQUALS(TEXT("ValidateCompactBinary(MissingName)"), Validate({ uint8(ECbFieldType::Object), 3, NullNoName, IntNoName, 0 }), ECbValidateError::MissingName);
+	CHECK_EQUALS(TEXT("ValidateCompactBinary(MissingName, Uniform)"), Validate({ uint8(ECbFieldType::UniformObject), 3, IntNoName, 0, 0 }), ECbValidateError::MissingName);
+	CHECK_EQUALS(TEXT("ValidateCompactBinary(MissingName, Mode)"), ValidateMode({ uint8(ECbFieldType::Object), 3, NullNoName, IntNoName, 0 }, ECbValidateMode::All & ~ECbValidateMode::Names), ECbValidateError::None);
+	CHECK_EQUALS(TEXT("ValidateCompactBinary(MissingName, Uniform, Mode)"), ValidateMode({ uint8(ECbFieldType::UniformObject), 3, IntNoName, 0, 0 }, ECbValidateMode::All & ~ECbValidateMode::Names), ECbValidateError::None);
+
+	// Test ArrayName
+	CHECK_EQUALS(TEXT("ValidateCompactBinary(ArrayName)"), Validate({ uint8(ECbFieldType::Array), 5, 2, NullNoName, NullWithName, 1, 'F' }), ECbValidateError::ArrayName);
+	CHECK_EQUALS(TEXT("ValidateCompactBinary(ArrayName, Uniform)"), Validate({ uint8(ECbFieldType::UniformArray), 4, 1, NullWithName, 1, 'F' }), ECbValidateError::ArrayName);
+	CHECK_EQUALS(TEXT("ValidateCompactBinary(ArrayName, Mode)"), ValidateMode({ uint8(ECbFieldType::Array), 5, 2, NullNoName, NullWithName, 1, 'F' }, ECbValidateMode::All & ~ECbValidateMode::Names), ECbValidateError::None);
+	CHECK_EQUALS(TEXT("ValidateCompactBinary(ArrayName, Uniform, Mode)"), ValidateMode({ uint8(ECbFieldType::UniformArray), 4, 1, NullWithName, 1, 'F' }, ECbValidateMode::All & ~ECbValidateMode::Names), ECbValidateError::None);
+
+	// Test InvalidString
+	// Not tested or implemented yet because the engine does not provide enough UTF-8 functionality.
+
+	// Test InvalidInteger
+	CHECK_EQUALS(TEXT("ValidateCompactBinary(InvalidInteger, NameSize)"), Validate({ NullWithName, 0x80, 1, 'N' }), ECbValidateError::InvalidInteger);
+	CHECK_EQUALS(TEXT("ValidateCompactBinary(InvalidInteger, ObjectSize)"), Validate({ uint8(ECbFieldType::Object), 0xc0, 0, 0 }), ECbValidateError::InvalidInteger);
+	CHECK_EQUALS(TEXT("ValidateCompactBinary(InvalidInteger, ArraySize)"), Validate({ uint8(ECbFieldType::Array), 0xe0, 0, 0, 1, 0 }), ECbValidateError::InvalidInteger);
+	CHECK_EQUALS(TEXT("ValidateCompactBinary(InvalidInteger, ArrayCount)"), Validate({ uint8(ECbFieldType::Array), 5, 0xf0, 0, 0, 0, 0 }), ECbValidateError::InvalidInteger);
+	CHECK_EQUALS(TEXT("ValidateCompactBinary(InvalidInteger, BinarySize)"), Validate({ uint8(ECbFieldType::Binary), 0xf8, 0, 0, 0, 0, 0 }), ECbValidateError::InvalidInteger);
+	CHECK_EQUALS(TEXT("ValidateCompactBinary(InvalidInteger, StringSize)"), Validate({ uint8(ECbFieldType::String), 0xfc, 0, 0, 0, 0, 0, 0 }), ECbValidateError::InvalidInteger);
+	CHECK_EQUALS(TEXT("ValidateCompactBinary(InvalidInteger, IntegerPositive)"), Validate({ uint8(ECbFieldType::IntegerPositive), 0xfe, 0, 0, 0, 0, 0, 0, 0 }), ECbValidateError::InvalidInteger);
+	CHECK_EQUALS(TEXT("ValidateCompactBinary(InvalidInteger, IntegerNegative)"), Validate({ uint8(ECbFieldType::IntegerNegative), 0xff, 0, 0, 0, 0, 0, 0, 0, 0 }), ECbValidateError::InvalidInteger);
+	CHECK_EQUALS(TEXT("ValidateCompactBinary(InvalidInteger, ArraySize)"), Validate({ uint8(ECbFieldType::Array), 0x80, 1, 0 }), ECbValidateError::InvalidInteger);
+	CHECK_EQUALS(TEXT("ValidateCompactBinary(InvalidInteger, ArrayCount)"), Validate({ uint8(ECbFieldType::Array), 3, 0xc0, 0, 0 }), ECbValidateError::InvalidInteger);
+	CHECK_EQUALS(TEXT("ValidateCompactBinary(InvalidInteger, ObjectSize)"), Validate({ uint8(ECbFieldType::Object), 0xe0, 0, 0, 0 }), ECbValidateError::InvalidInteger);
+	CHECK_EQUALS(TEXT("ValidateCompactBinary(InvalidInteger, CustomById)"), Validate({ uint8(ECbFieldType::CustomById), 2, 0x80, 0 }), ECbValidateError::InvalidInteger);
+	CHECK_EQUALS(TEXT("ValidateCompactBinary(InvalidInteger, CustomByName)"), Validate({ uint8(ECbFieldType::CustomByName), 3, 0x80, 1, 'A' }), ECbValidateError::InvalidInteger);
+
+	CHECK_EQUALS(TEXT("ValidateCompactBinary(InvalidInteger, NameSize, Mode)"), ValidateMode({ NullWithName, 0x80, 1, 'N' }, ECbValidateMode::All & ~ECbValidateMode::Format), ECbValidateError::None);
+	CHECK_EQUALS(TEXT("ValidateCompactBinary(InvalidInteger, ArraySize, Mode)"), ValidateMode({ uint8(ECbFieldType::Array), 0xc0, 0, 1, 0 }, ECbValidateMode::All & ~ECbValidateMode::Format), ECbValidateError::None);
+	CHECK_EQUALS(TEXT("ValidateCompactBinary(InvalidInteger, ObjectSize, Mode)"), ValidateMode({ uint8(ECbFieldType::Object), 0xe0, 0, 0, 0 }, ECbValidateMode::All & ~ECbValidateMode::Format), ECbValidateError::None);
+	CHECK_EQUALS(TEXT("ValidateCompactBinary(InvalidInteger, CustomById, Mode)"), ValidateMode({ uint8(ECbFieldType::CustomById), 2, 0x80, 0 }, ECbValidateMode::All & ~ECbValidateMode::Format), ECbValidateError::None);
+	CHECK_EQUALS(TEXT("ValidateCompactBinary(InvalidInteger, CustomByName, Mode)"), ValidateMode({ uint8(ECbFieldType::CustomByName), 3, 0x80, 1, 'A' }, ECbValidateMode::All & ~ECbValidateMode::Format), ECbValidateError::None);
+
+	// Test InvalidFloat
+	CHECK_EQUALS(TEXT("ValidateCompactBinary(InvalidFloat, MaxSignificant+1)"), Validate({ uint8(ECbFieldType::Float64), 0x3f, 0xff, 0xff, 0xff, 0xf0, 0x00, 0x00, 0x00 }), ECbValidateError::None); // 1.9999999403953552
+	CHECK_EQUALS(TEXT("ValidateCompactBinary(InvalidFloat, MaxExponent+1)"), Validate({ uint8(ECbFieldType::Float64), 0x47, 0xff, 0xff, 0xff, 0xe0, 0x00, 0x00, 0x00 }), ECbValidateError::None); // 6.8056469327705771e38
+	CHECK_EQUALS(TEXT("ValidateCompactBinary(InvalidFloat, MaxSignificand)"), Validate({ uint8(ECbFieldType::Float64), 0x3f, 0xff, 0xff, 0xff, 0xe0, 0x00, 0x00, 0x00 }), ECbValidateError::InvalidFloat); // 1.9999998807907104
+	CHECK_EQUALS(TEXT("ValidateCompactBinary(InvalidFloat, MaxExponent)"), Validate({ uint8(ECbFieldType::Float64), 0x47, 0xef, 0xff, 0xff, 0xe0, 0x00, 0x00, 0x00 }), ECbValidateError::InvalidFloat); // 3.4028234663852886e38
+	CHECK_EQUALS(TEXT("ValidateCompactBinary(InvalidFloat, MaxSignificand, Mode)"), ValidateMode({ uint8(ECbFieldType::Float64), 0x3f, 0xff, 0xff, 0xff, 0xe0, 0x00, 0x00, 0x00 }, ECbValidateMode::All & ~ECbValidateMode::Format), ECbValidateError::None); // 1.9999998807907104
+	CHECK_EQUALS(TEXT("ValidateCompactBinary(InvalidFloat, MaxExponent, Mode)"), ValidateMode({ uint8(ECbFieldType::Float64), 0x47, 0xef, 0xff, 0xff, 0xe0, 0x00, 0x00, 0x00 }, ECbValidateMode::All & ~ECbValidateMode::Format), ECbValidateError::None); // 3.4028234663852886e38
+
+	// Test NonUniformObject
+	CHECK_EQUALS(TEXT("ValidateCompactBinary(NonUniformObject)"), Validate({ uint8(ECbFieldType::Object), 3, NullWithName, 1, 'A' }), ECbValidateError::NonUniformObject);
+	CHECK_EQUALS(TEXT("ValidateCompactBinary(NonUniformObject)"), Validate({ uint8(ECbFieldType::Object), 6, NullWithName, 1, 'A', NullWithName, 1, 'B' }), ECbValidateError::NonUniformObject);
+	CHECK_EQUALS(TEXT("ValidateCompactBinary(NonUniformObject, Mode)"), ValidateMode({ uint8(ECbFieldType::Object), 3, NullWithName, 1, 'A' }, ECbValidateMode::All & ~ECbValidateMode::Format), ECbValidateError::None);
+	CHECK_EQUALS(TEXT("ValidateCompactBinary(NonUniformObject, Mode)"), ValidateMode({ uint8(ECbFieldType::Object), 6, NullWithName, 1, 'A', NullWithName, 1, 'B' }, ECbValidateMode::All & ~ECbValidateMode::Format), ECbValidateError::None);
+
+	// Test NonUniformArray
+	CHECK_EQUALS(TEXT("ValidateCompactBinary(NonUniformArray)"), Validate({ uint8(ECbFieldType::Array), 3, 1, IntNoName, 0 }), ECbValidateError::NonUniformArray);
+	CHECK_EQUALS(TEXT("ValidateCompactBinary(NonUniformArray)"), Validate({ uint8(ECbFieldType::Array), 5, 2, IntNoName, 1, IntNoName, 2 }), ECbValidateError::NonUniformArray);
+	CHECK_EQUALS(TEXT("ValidateCompactBinary(NonUniformArray, Null)"), Validate({ uint8(ECbFieldType::Array), 3, 2, NullNoName, NullNoName }), ECbValidateError::None);
+	CHECK_EQUALS(TEXT("ValidateCompactBinary(NonUniformArray, Bool)"), Validate({ uint8(ECbFieldType::Array), 3, 2, uint8(ECbFieldType::BoolFalse), uint8(ECbFieldType::BoolFalse) }), ECbValidateError::None);
+	CHECK_EQUALS(TEXT("ValidateCompactBinary(NonUniformArray, Bool)"), Validate({ uint8(ECbFieldType::Array), 3, 2, uint8(ECbFieldType::BoolTrue), uint8(ECbFieldType::BoolTrue) }), ECbValidateError::None);
+	CHECK_EQUALS(TEXT("ValidateCompactBinary(NonUniformArray, Mode)"), ValidateMode({ uint8(ECbFieldType::Array), 3, 1, IntNoName, 0 }, ECbValidateMode::All & ~ECbValidateMode::Format), ECbValidateError::None);
+	CHECK_EQUALS(TEXT("ValidateCompactBinary(NonUniformArray, Mode)"), ValidateMode({ uint8(ECbFieldType::Array), 5, 2, IntNoName, 1, IntNoName, 2 }, ECbValidateMode::All & ~ECbValidateMode::Format), ECbValidateError::None);
+
+	// Test Padding
+	CHECK_EQUALS(TEXT("ValidateCompactBinary(Padding)"), Validate({ NullNoName, 0 }), ECbValidateError::Padding);
+	CHECK_EQUALS(TEXT("ValidateCompactBinary(Padding)"), Validate({ uint8(ECbFieldType::Array), 1, 0, 0 }), ECbValidateError::Padding);
+	CHECK_EQUALS(TEXT("ValidateCompactBinary(Padding)"), Validate({ uint8(ECbFieldType::Object), 0, 0 }), ECbValidateError::Padding);
+	CHECK_EQUALS(TEXT("ValidateCompactBinary(Padding, Mode)"), ValidateMode({ NullNoName, 0 }, ECbValidateMode::All & ~ECbValidateMode::Padding), ECbValidateError::None);
+	CHECK_EQUALS(TEXT("ValidateCompactBinary(Padding, Mode)"), ValidateMode({ uint8(ECbFieldType::Array), 1, 0, 0 }, ECbValidateMode::All & ~ECbValidateMode::Padding), ECbValidateError::None);
+	CHECK_EQUALS(TEXT("ValidateCompactBinary(Padding, Mode)"), ValidateMode({ uint8(ECbFieldType::Object), 0, 0 }, ECbValidateMode::All & ~ECbValidateMode::Padding), ECbValidateError::None);
+}
+
+///////////////////////////////////////////////////////////////////////////////////////////////////////////////////////
+
+TEST_CASE_NAMED(FCbValidateRangeTest, "System::Core::Serialization::ValidateCompactBinaryRange", "[ApplicationContextMask][SmokeFilter]")
+{
+	auto Validate = [](std::initializer_list<uint8> Data, ECbValidateMode Mode = ECbValidateMode::All) -> ECbValidateError
+	{
+		return ValidateCompactBinaryRange(MakeMemoryView(Data), Mode);
+	};
+	// Test Empty
+	CHECK_EQUALS(TEXT("ValidateCompactBinaryRange(Empty)"), Validate({}), ECbValidateError::None);
+
+	// Test Valid
+	CHECK_EQUALS(TEXT("ValidateCompactBinaryRange(Null x2)"), Validate({ uint8(ECbFieldType::Null), uint8(ECbFieldType::Null) }), ECbValidateError::None);
+
+	// Test Padding
+	CHECK_EQUALS(TEXT("ValidateCompactBinaryRange(Padding InvalidType)"), Validate({ uint8(ECbFieldType::Null), 0 }), ECbValidateError::InvalidType);
+	CHECK_EQUALS(TEXT("ValidateCompactBinaryRange(Padding OutOfBounds)"), Validate({ uint8(ECbFieldType::Null), uint8(ECbFieldType::Binary) }), ECbValidateError::OutOfBounds);
+	CHECK_EQUALS(TEXT("ValidateCompactBinaryRange(Padding OutOfBounds, Mode)"), Validate({ uint8(ECbFieldType::Null), uint8(ECbFieldType::Binary) }, ECbValidateMode::None), ECbValidateError::None);
+}
+
+///////////////////////////////////////////////////////////////////////////////////////////////////////////////////////
+
+TEST_CASE_NAMED(FCbValidateAttachmentTest, "System::Core::Serialization::ValidateCompactBinaryAttachment", "[ApplicationContextMask][SmokeFilter]")
+{
+	auto Validate = [](std::initializer_list<uint8> Data, ECbValidateMode Mode = ECbValidateMode::All) -> ECbValidateError
+	{
+		return ValidateCompactBinaryAttachment(MakeMemoryView(Data), Mode);
+	};
+
+	const uint8 BinaryValue[] = { 0, 1, 2, 3 };
+	const FMemoryView BinaryView = MakeMemoryView(BinaryValue);
+
+	// Test Null
+	{
+		CHECK_EQUALS(TEXT("ValidateCompactBinaryAttachment(Null)"), Validate({ uint8(ECbFieldType::Binary), 0 }, ECbValidateMode::All & ~ECbValidateMode::Package), ECbValidateError::None);
+		CHECK_EQUALS(TEXT("ValidateCompactBinaryAttachment(Null, Padding)"), Validate({ uint8(ECbFieldType::Binary), 0, 0 }, ECbValidateMode::All & ~ECbValidateMode::Package), ECbValidateError::Padding);
+		CHECK_EQUALS(TEXT("ValidateCompactBinaryAttachment(Null, Padding, Mode)"), Validate({ uint8(ECbFieldType::Binary), 0, 0 }, ECbValidateMode::All & ~(ECbValidateMode::Padding | ECbValidateMode::Package)), ECbValidateError::None);
+	}
+
+	// Test Binary
+	{
+		FBufferArchive Buffer;
+		FCbAttachment(FSharedBuffer::MakeView(MakeMemoryView<uint8>({ 0 }))).Save(Buffer);
+		CHECK_EQUALS(TEXT("ValidateCompactBinaryAttachment(Binary)"), ValidateCompactBinaryAttachment(MakeMemoryView(Buffer), ECbValidateMode::All), ECbValidateError::None);
+		++Buffer[2];
+		CHECK_EQUALS(TEXT("ValidateCompactBinaryAttachment(Binary, InvalidHash)"), ValidateCompactBinaryAttachment(MakeMemoryView(Buffer), ECbValidateMode::All), ECbValidateError::InvalidPackageHash);
+		CHECK_EQUALS(TEXT("ValidateCompactBinaryAttachment(Binary, InvalidHash, Mode)"), ValidateCompactBinaryAttachment(MakeMemoryView(Buffer), ECbValidateMode::All & ~ECbValidateMode::PackageHash), ECbValidateError::None);
+		--Buffer[2];
+		Buffer.AddZeroed(1);
+		CHECK_EQUALS(TEXT("ValidateCompactBinaryAttachment(Binary, Padding)"), ValidateCompactBinaryAttachment(MakeMemoryView(Buffer), ECbValidateMode::All), ECbValidateError::Padding);
+		CHECK_EQUALS(TEXT("ValidateCompactBinaryAttachment(Binary, Padding, Mode)"), ValidateCompactBinaryAttachment(MakeMemoryView(Buffer), ECbValidateMode::All & ~ECbValidateMode::Padding), ECbValidateError::None);
+		Buffer.Pop(/*bAllowShrinking*/ false);
+		CHECK_EQUALS(TEXT("ValidateCompactBinaryAttachment(Binary, OutOfBoundsValue)"), ValidateCompactBinaryAttachment(MakeMemoryView(Buffer).Left(1), ECbValidateMode::All), ECbValidateError::OutOfBounds);
+		CHECK_EQUALS(TEXT("ValidateCompactBinaryAttachment(Binary, OutOfBoundsHash)"), ValidateCompactBinaryAttachment(MakeMemoryView(Buffer).Left(4), ECbValidateMode::All), ECbValidateError::OutOfBounds);
+		CHECK_EQUALS(TEXT("ValidateCompactBinaryAttachment(Binary, OutOfBoundsHash, Mode)"), ValidateCompactBinaryAttachment(MakeMemoryView(Buffer).Left(4), ECbValidateMode::None), ECbValidateError::None);
+	}
+
+	// Test Object
+	{
+		TCbWriter<256> Writer;
+		Writer.BeginObject();
+		Writer.AddNull(ANSITEXTVIEW("N"));
+		Writer.EndObject();
+		FBufferArchive Buffer;
+		FCbAttachment(Writer.Save().AsObject()).Save(Buffer);
+		CHECK_EQUALS(TEXT("ValidateCompactBinaryAttachment(Object)"), ValidateCompactBinaryAttachment(MakeMemoryView(Buffer), ECbValidateMode::All), ECbValidateError::None);
+		++Buffer[2];
+		CHECK_EQUALS(TEXT("ValidateCompactBinaryAttachment(Object, InvalidHash)"), ValidateCompactBinaryAttachment(MakeMemoryView(Buffer), ECbValidateMode::All), ECbValidateError::InvalidPackageHash);
+		CHECK_EQUALS(TEXT("ValidateCompactBinaryAttachment(Object, InvalidHash, Mode)"), ValidateCompactBinaryAttachment(MakeMemoryView(Buffer), ECbValidateMode::All & ~ECbValidateMode::PackageHash), ECbValidateError::None);
+		--Buffer[2];
+		Buffer.AddZeroed(1);
+		CHECK_EQUALS(TEXT("ValidateCompactBinaryAttachment(Object, Padding)"), ValidateCompactBinaryAttachment(MakeMemoryView(Buffer), ECbValidateMode::All), ECbValidateError::Padding);
+		CHECK_EQUALS(TEXT("ValidateCompactBinaryAttachment(Object, Padding, Mode)"), ValidateCompactBinaryAttachment(MakeMemoryView(Buffer), ECbValidateMode::All & ~ECbValidateMode::Padding), ECbValidateError::None);
+		Buffer.Pop(/*bAllowShrinking*/ false);
+		CHECK_EQUALS(TEXT("ValidateCompactBinaryAttachment(Object, OutOfBoundsValue)"), ValidateCompactBinaryAttachment(MakeMemoryView(Buffer).Left(1), ECbValidateMode::All), ECbValidateError::OutOfBounds);
+		CHECK_EQUALS(TEXT("ValidateCompactBinaryAttachment(Object, OutOfBoundsHash)"), ValidateCompactBinaryAttachment(MakeMemoryView(Buffer).Left(8), ECbValidateMode::All), ECbValidateError::OutOfBounds);
+		CHECK_EQUALS(TEXT("ValidateCompactBinaryAttachment(Object, OutOfBoundsHash, Mode)"), ValidateCompactBinaryAttachment(MakeMemoryView(Buffer).Left(8), ECbValidateMode::None), ECbValidateError::None);
+	}
+
+	// Test Empty Object
+	{
+		FBufferArchive Buffer;
+		FCbFieldIterator::MakeSingle(FCbObject().AsField()).CopyRangeTo(Buffer);
+		CHECK_EQUALS(TEXT("ValidateCompactBinaryAttachment(Object)"), ValidateCompactBinaryAttachment(MakeMemoryView(Buffer), ECbValidateMode::All), ECbValidateError::None);
+	}
+
+	// Test InvalidPackageFormat
+	{
+		CHECK_EQUALS(TEXT("ValidateCompactBinaryAttachment(Empty)"), Validate({}), ECbValidateError::InvalidPackageFormat);
+		CHECK_EQUALS(TEXT("ValidateCompactBinaryAttachment(Empty, Mode)"), Validate({}, ECbValidateMode::All & ~ECbValidateMode::Package), ECbValidateError::None);
+	}
+	{
+		FBufferArchive Buffer;
+		FCbWriter Writer;
+		Writer.AddBinaryAttachment(FIoHash::HashBuffer(BinaryView));
+		Writer.SetName(ANSITEXTVIEW("Name")).AddBinary(BinaryView);
+		Writer.Save(Buffer);
+		CHECK_EQUALS(TEXT("ValidateCompactBinaryAttachment(NameOnValue)"), ValidateCompactBinaryAttachment(MakeMemoryView(Buffer), ECbValidateMode::All), ECbValidateError::InvalidPackageFormat);
+		CHECK_EQUALS(TEXT("ValidateCompactBinaryAttachment(NameOnValue, Mode)"), ValidateCompactBinaryAttachment(MakeMemoryView(Buffer), ECbValidateMode::All & ~ECbValidateMode::Package), ECbValidateError::None);
+	}
+	{
+		FBufferArchive Buffer;
+		FCbWriter Writer;
+		Writer.SetName(ANSITEXTVIEW("Name")).AddBinaryAttachment(FIoHash::HashBuffer(BinaryView));
+		Writer.AddBinary(BinaryView);
+		Writer.Save(Buffer);
+		CHECK_EQUALS(TEXT("ValidateCompactBinaryAttachment(NameOnHash)"), ValidateCompactBinaryAttachment(MakeMemoryView(Buffer), ECbValidateMode::All), ECbValidateError::InvalidPackageFormat);
+		CHECK_EQUALS(TEXT("ValidateCompactBinaryAttachment(NameOnHash, Mode)"), ValidateCompactBinaryAttachment(MakeMemoryView(Buffer), ECbValidateMode::All & ~ECbValidateMode::Package), ECbValidateError::None);
+	}
+}
+
+///////////////////////////////////////////////////////////////////////////////////////////////////////////////////////
+
+TEST_CASE_NAMED(FCbValidatePackageTest, "System::Core::Serialization::ValidateCompactBinaryPackage", "[ApplicationContextMask][SmokeFilter]")
+{
+	const uint8 BinaryValue[] = { 0, 1, 2, 3 };
+	const FMemoryView BinaryView = MakeMemoryView(BinaryValue);
+	const FSharedBuffer BinaryBuffer = FSharedBuffer::MakeView(BinaryView);
+	FIoHash BinaryBufferHash = FIoHash::HashBuffer(BinaryBuffer);
+	const FCompressedBuffer CompressedBinaryBuffer = FCompressedBuffer::Compress(FSharedBuffer::MakeView(BinaryView));
+	
+	TCbWriter<32> CommonRootObjectWriter;
+	CommonRootObjectWriter.BeginObject();
+	CommonRootObjectWriter << "Field" << 42;
+	CommonRootObjectWriter.EndObject();
+	FCbObject CommonRootObject(CommonRootObjectWriter.Save().AsObject());
+	FIoHash CommonRootObjectHash = CommonRootObject.GetHash();
+
+	// Test Null
+	{
+		CHECK_EQUALS(TEXT("ValidateCompactBinaryPackage(Null)"), ValidateCompactBinaryPackage(MakeMemoryView<uint8>({ uint8(ECbFieldType::Null) }), ECbValidateMode::All), ECbValidateError::None);
+		CHECK_EQUALS(TEXT("ValidateCompactBinaryPackage(Null, Padding)"), ValidateCompactBinaryPackage(MakeMemoryView<uint8>({ uint8(ECbFieldType::Null), 0 }), ECbValidateMode::All), ECbValidateError::Padding);
+		CHECK_EQUALS(TEXT("ValidateCompactBinaryPackage(Null, Padding, Mode)"), ValidateCompactBinaryPackage(MakeMemoryView<uint8>({ uint8(ECbFieldType::Null), 0 }), ECbValidateMode::All & ~ECbValidateMode::Padding), ECbValidateError::None);
+	}
+
+	// Test Object
+	{
+		FBufferArchive Buffer;
+		TCbWriter<256> Writer;
+		Writer.AddHash(CommonRootObjectHash);
+		Writer.AddObject(CommonRootObject);
+		Writer.AddNull();
+		Writer.Save(Buffer);
+		CHECK_EQUALS(TEXT("ValidateCompactBinaryPackage(Object)"), ValidateCompactBinaryPackage(MakeMemoryView(Buffer), ECbValidateMode::All), ECbValidateError::None);
+		CHECK_EQUALS(TEXT("ValidateCompactBinaryPackage(Object, MissingHash)"), ValidateCompactBinaryPackage(MakeMemoryView(Buffer).LeftChop(34), ECbValidateMode::All), ECbValidateError::InvalidPackageFormat);
+		CHECK_EQUALS(TEXT("ValidateCompactBinaryPackage(Object, MissingHash, Mode)"), ValidateCompactBinaryPackage(MakeMemoryView(Buffer).LeftChop(34), ECbValidateMode::All & ~ECbValidateMode::Package), ECbValidateError::None);
+		CHECK_EQUALS(TEXT("ValidateCompactBinaryPackage(Object, MissingNull)"), ValidateCompactBinaryPackage(MakeMemoryView(Buffer).LeftChop(1), ECbValidateMode::All), ECbValidateError::InvalidPackageFormat);
+		CHECK_EQUALS(TEXT("ValidateCompactBinaryPackage(Object, MissingNull, Mode)"), ValidateCompactBinaryPackage(MakeMemoryView(Buffer).LeftChop(1), ECbValidateMode::All & ~ECbValidateMode::Package), ECbValidateError::None);
+		++Buffer.Last(1);
+		CHECK_EQUALS(TEXT("ValidateCompactBinaryPackage(Object, InvalidHash)"), ValidateCompactBinaryPackage(MakeMemoryView(Buffer), ECbValidateMode::All), ECbValidateError::InvalidPackageHash);
+		CHECK_EQUALS(TEXT("ValidateCompactBinaryPackage(Object, InvalidHash, Mode)"), ValidateCompactBinaryPackage(MakeMemoryView(Buffer), ECbValidateMode::All & ~ECbValidateMode::PackageHash), ECbValidateError::None);
+		--Buffer.Last(1);
+		Buffer.AddZeroed(1);
+		CHECK_EQUALS(TEXT("ValidateCompactBinaryPackage(Object, Padding)"), ValidateCompactBinaryPackage(MakeMemoryView(Buffer), ECbValidateMode::All), ECbValidateError::Padding);
+		CHECK_EQUALS(TEXT("ValidateCompactBinaryPackage(Object, Padding, Mode)"), ValidateCompactBinaryPackage(MakeMemoryView(Buffer), ECbValidateMode::All & ~ECbValidateMode::Padding), ECbValidateError::None);
+		Buffer.Pop(/*bAllowShrinking*/ false);
+		CHECK_EQUALS(TEXT("ValidateCompactBinaryPackage(Object, OutOfBoundsValue)"), ValidateCompactBinaryPackage(MakeMemoryView(Buffer).Left(1), ECbValidateMode::All), ECbValidateError::OutOfBounds);
+		CHECK_EQUALS(TEXT("ValidateCompactBinaryPackage(Object, OutOfBoundsHash)"), ValidateCompactBinaryPackage(MakeMemoryView(Buffer).Left(11), ECbValidateMode::All), ECbValidateError::OutOfBounds);
+		CHECK_EQUALS(TEXT("ValidateCompactBinaryPackage(Object, OutOfBoundsHash, Mode)"), ValidateCompactBinaryPackage(MakeMemoryView(Buffer).Left(11), ECbValidateMode::None), ECbValidateError::None);
+	}
+
+	// Test Object + Attachment
+	{
+		FBufferArchive Buffer;
+		TCbWriter<256> Writer;
+		Writer.AddHash(CommonRootObjectHash);
+		Writer.AddObject(CommonRootObject);
+		const uint64 AttachmentOffset = Writer.GetSaveSize();
+		Writer.AddBinaryAttachment(BinaryBufferHash);
+		Writer.AddBinary(BinaryBuffer);
+		Writer.AddNull();
+		Writer.Save(Buffer);
+		CHECK_EQUALS(TEXT("ValidateCompactBinaryPackage(Object+Attachment)"), ValidateCompactBinaryPackage(MakeMemoryView(Buffer), ECbValidateMode::All), ECbValidateError::None);
+		CHECK_EQUALS(TEXT("ValidateCompactBinaryPackage(Object+Attachment, MissingNull)"), ValidateCompactBinaryPackage(MakeMemoryView(Buffer).LeftChop(1), ECbValidateMode::All), ECbValidateError::InvalidPackageFormat);
+		CHECK_EQUALS(TEXT("ValidateCompactBinaryPackage(Object+Attachment, MissingNull, Mode)"), ValidateCompactBinaryPackage(MakeMemoryView(Buffer).LeftChop(1), ECbValidateMode::All & ~ECbValidateMode::Package), ECbValidateError::None);
+		++Buffer[2];
+		CHECK_EQUALS(TEXT("ValidateCompactBinaryPackage(Object+Attachment, InvalidHash)"), ValidateCompactBinaryPackage(MakeMemoryView(Buffer), ECbValidateMode::All), ECbValidateError::InvalidPackageHash);
+		CHECK_EQUALS(TEXT("ValidateCompactBinaryPackage(Object+Attachment, InvalidHash, Mode)"), ValidateCompactBinaryPackage(MakeMemoryView(Buffer), ECbValidateMode::All & ~ECbValidateMode::PackageHash), ECbValidateError::None);
+		--Buffer[2];
+		Buffer.AddZeroed(1);
+		CHECK_EQUALS(TEXT("ValidateCompactBinaryPackage(Object+Attachment, Padding)"), ValidateCompactBinaryPackage(MakeMemoryView(Buffer), ECbValidateMode::All), ECbValidateError::Padding);
+		CHECK_EQUALS(TEXT("ValidateCompactBinaryPackage(Object+Attachment, Padding, Mode)"), ValidateCompactBinaryPackage(MakeMemoryView(Buffer), ECbValidateMode::All & ~ECbValidateMode::Padding), ECbValidateError::None);
+		Buffer.Pop(/*bAllowShrinking*/ false);
+		CHECK_EQUALS(TEXT("ValidateCompactBinaryPackage(Object+Attachment, OutOfBoundsValue)"), ValidateCompactBinaryPackage(MakeMemoryView(Buffer).Left(AttachmentOffset + 1), ECbValidateMode::All), ECbValidateError::OutOfBounds);
+		CHECK_EQUALS(TEXT("ValidateCompactBinaryPackage(Object+Attachment, OutOfBoundsHash)"), ValidateCompactBinaryPackage(MakeMemoryView(Buffer).Left(AttachmentOffset + 7), ECbValidateMode::All), ECbValidateError::OutOfBounds);
+		CHECK_EQUALS(TEXT("ValidateCompactBinaryPackage(Object+Attachment, OutOfBoundsHash, Mode)"), ValidateCompactBinaryPackage(MakeMemoryView(Buffer).Left(AttachmentOffset + 7), ECbValidateMode::None), ECbValidateError::None);
+	}
+
+	// Test Object + Compressed Binary Attachment
+	{
+		FBufferArchive Buffer;
+		TCbWriter<256> Writer;
+		Writer.AddHash(CommonRootObjectHash);
+		Writer.AddObject(CommonRootObject);
+		const uint64 AttachmentOffset = Writer.GetSaveSize();
+		Writer.AddBinary(CompressedBinaryBuffer.GetCompressed());
+		Writer.AddNull();
+		Writer.Save(Buffer);
+		CHECK_EQUALS(TEXT("ValidateCompactBinaryPackage(Object+CompressedBinaryAttachment)"), ValidateCompactBinaryPackage(MakeMemoryView(Buffer), ECbValidateMode::All), ECbValidateError::None);
+		CHECK_EQUALS(TEXT("ValidateCompactBinaryPackage(Object+CompressedBinaryAttachment, MissingNull)"), ValidateCompactBinaryPackage(MakeMemoryView(Buffer).LeftChop(1), ECbValidateMode::All), ECbValidateError::InvalidPackageFormat);
+		CHECK_EQUALS(TEXT("ValidateCompactBinaryPackage(Object+CompressedBinaryAttachment, MissingNull, Mode)"), ValidateCompactBinaryPackage(MakeMemoryView(Buffer).LeftChop(1), ECbValidateMode::All & ~ECbValidateMode::Package), ECbValidateError::None);
+		++Buffer[2];
+		CHECK_EQUALS(TEXT("ValidateCompactBinaryPackage(Object+CompressedBinaryAttachment, InvalidHash)"), ValidateCompactBinaryPackage(MakeMemoryView(Buffer), ECbValidateMode::All), ECbValidateError::InvalidPackageHash);
+		CHECK_EQUALS(TEXT("ValidateCompactBinaryPackage(Object+CompressedBinaryAttachment, InvalidHash, Mode)"), ValidateCompactBinaryPackage(MakeMemoryView(Buffer), ECbValidateMode::All & ~ECbValidateMode::PackageHash), ECbValidateError::None);
+		--Buffer[2];
+		Buffer.AddZeroed(1);
+		CHECK_EQUALS(TEXT("ValidateCompactBinaryPackage(Object+CompressedBinaryAttachment, Padding)"), ValidateCompactBinaryPackage(MakeMemoryView(Buffer), ECbValidateMode::All), ECbValidateError::Padding);
+		CHECK_EQUALS(TEXT("ValidateCompactBinaryPackage(Object+CompressedBinaryAttachment, Padding, Mode)"), ValidateCompactBinaryPackage(MakeMemoryView(Buffer), ECbValidateMode::All & ~ECbValidateMode::Padding), ECbValidateError::None);
+		Buffer.Pop(/*bAllowShrinking*/ false);
+		CHECK_EQUALS(TEXT("ValidateCompactBinaryPackage(Object+CompressedBinaryAttachment, OutOfBoundsValue)"), ValidateCompactBinaryPackage(MakeMemoryView(Buffer).Left(AttachmentOffset + 1), ECbValidateMode::All), ECbValidateError::OutOfBounds);
+		CHECK_EQUALS(TEXT("ValidateCompactBinaryPackage(Object+CompressedBinaryAttachment, OutOfBoundsHash)"), ValidateCompactBinaryPackage(MakeMemoryView(Buffer).Left(AttachmentOffset + 7), ECbValidateMode::All), ECbValidateError::OutOfBounds);
+		CHECK_EQUALS(TEXT("ValidateCompactBinaryPackage(Object+CompressedBinaryAttachment, OutOfBoundsHash, Mode)"), ValidateCompactBinaryPackage(MakeMemoryView(Buffer).Left(AttachmentOffset + 7), ECbValidateMode::None), ECbValidateError::None);
+	}
+
+	// Test InvalidPackageFormat
+	{
+		CHECK_EQUALS(TEXT("ValidateCompactBinaryPackage(Empty)"), ValidateCompactBinaryPackage(FMemoryView(), ECbValidateMode::All), ECbValidateError::InvalidPackageFormat);
+		CHECK_EQUALS(TEXT("ValidateCompactBinaryPackage(Empty, Mode)"), ValidateCompactBinaryPackage(FMemoryView(), ECbValidateMode::All & ~ECbValidateMode::Package), ECbValidateError::None);
+	}
+	{
+		FBufferArchive Buffer;
+		TCbWriter<256> Writer;
+		Writer.AddHash(CommonRootObjectHash);
+		Writer.SetName(ANSITEXTVIEW("Name"));
+		Writer.AddObject(CommonRootObject);
+		Writer.AddNull();
+		Writer.Save(Buffer);
+		CHECK_EQUALS(TEXT("ValidateCompactBinaryPackage(NameOnObject)"), ValidateCompactBinaryPackage(MakeMemoryView(Buffer), ECbValidateMode::All), ECbValidateError::InvalidPackageFormat);
+		CHECK_EQUALS(TEXT("ValidateCompactBinaryPackage(NameOnObject, Mode)"), ValidateCompactBinaryPackage(MakeMemoryView(Buffer), ECbValidateMode::All & ~ECbValidateMode::Package), ECbValidateError::None);
+	}
+	{
+		FBufferArchive Buffer;
+		TCbWriter<256> Writer;
+		Writer.SetName(ANSITEXTVIEW("Name"));
+		Writer.AddHash(CommonRootObjectHash);
+		Writer.AddObject(CommonRootObject);
+		Writer.AddNull();
+		Writer.Save(Buffer);
+		CHECK_EQUALS(TEXT("ValidateCompactBinaryPackage(NameOnHash)"), ValidateCompactBinaryPackage(MakeMemoryView(Buffer), ECbValidateMode::All), ECbValidateError::InvalidPackageFormat);
+		CHECK_EQUALS(TEXT("ValidateCompactBinaryPackage(NameOnHash, Mode)"), ValidateCompactBinaryPackage(MakeMemoryView(Buffer), ECbValidateMode::All & ~ECbValidateMode::Package), ECbValidateError::None);
+	}
+	{
+		FBufferArchive Buffer;
+		TCbWriter<256> Writer;
+		Writer.AddHash(CommonRootObjectHash);
+		Writer.AddObject(CommonRootObject);
+		Writer.AddBinaryAttachment(BinaryBufferHash);
+		Writer.SetName(ANSITEXTVIEW("Name")).AddBinary(BinaryBuffer);
+		Writer.AddNull();
+		Writer.Save(Buffer);
+		CHECK_EQUALS(TEXT("ValidateCompactBinaryPackage(NameOnAttachment)"), ValidateCompactBinaryPackage(MakeMemoryView(Buffer), ECbValidateMode::All), ECbValidateError::InvalidPackageFormat);
+		CHECK_EQUALS(TEXT("ValidateCompactBinaryPackage(NameOnAttachment, Mode)"), ValidateCompactBinaryPackage(MakeMemoryView(Buffer), ECbValidateMode::All & ~ECbValidateMode::Package), ECbValidateError::None);
+	}
+	{
+		FBufferArchive Buffer;
+		TCbWriter<256> Writer;
+		Writer.AddHash(CommonRootObjectHash);
+		Writer.AddObject(CommonRootObject);
+		Writer.AddHash(CommonRootObjectHash);
+		Writer.AddObject(CommonRootObject);
+		Writer.AddNull();
+		Writer.Save(Buffer);
+		CHECK_EQUALS(TEXT("ValidateCompactBinaryPackage(MultipleObjects)"), ValidateCompactBinaryPackage(MakeMemoryView(Buffer), ECbValidateMode::All), ECbValidateError::MultiplePackageObjects);
+		CHECK_EQUALS(TEXT("ValidateCompactBinaryPackage(MultipleObjects, Mode)"), ValidateCompactBinaryPackage(MakeMemoryView(Buffer), ECbValidateMode::All & ~ECbValidateMode::Package), ECbValidateError::None);
+	}
+	{
+		FBufferArchive Buffer;
+		TCbWriter<256> Writer;
+		Writer.AddBinaryAttachment(BinaryBufferHash);
+		Writer.AddBinary(BinaryBuffer);
+		Writer.AddBinaryAttachment(BinaryBufferHash);
+		Writer.AddBinary(BinaryBuffer);
+		Writer.AddNull();
+		Writer.Save(Buffer);
+		CHECK_EQUALS(TEXT("ValidateCompactBinaryPackage(DuplicateAttachments)"), ValidateCompactBinaryPackage(MakeMemoryView(Buffer), ECbValidateMode::All), ECbValidateError::DuplicateAttachments);
+		CHECK_EQUALS(TEXT("ValidateCompactBinaryPackage(DuplicateAttachments, Mode)"), ValidateCompactBinaryPackage(MakeMemoryView(Buffer), ECbValidateMode::All & ~ECbValidateMode::Package), ECbValidateError::None);
+	}
+	{
+		FBufferArchive Buffer;
+		TCbWriter<256> Writer;
+		Writer.AddBinary(CompressedBinaryBuffer.GetCompressed());
+		Writer.AddBinary(CompressedBinaryBuffer.GetCompressed());
+		Writer.AddNull();
+		Writer.Save(Buffer);
+		CHECK_EQUALS(TEXT("ValidateCompactBinaryPackage(DuplicateCompressedAttachments)"), ValidateCompactBinaryPackage(MakeMemoryView(Buffer), ECbValidateMode::All), ECbValidateError::DuplicateAttachments);
+		CHECK_EQUALS(TEXT("ValidateCompactBinaryPackage(DuplicateCompressedAttachments, Mode)"), ValidateCompactBinaryPackage(MakeMemoryView(Buffer), ECbValidateMode::All & ~ECbValidateMode::Package), ECbValidateError::None);
+	}
+	{
+		FBufferArchive Buffer;
+		TCbWriter<256> Writer;
+		Writer.AddBinaryAttachment(BinaryBufferHash);
+		Writer.AddBinary(BinaryBuffer);
+		Writer.AddBinary(CompressedBinaryBuffer.GetCompressed());
+		Writer.AddNull();
+		Writer.Save(Buffer);
+		CHECK_EQUALS(TEXT("ValidateCompactBinaryPackage(DuplicateMixedAttachments)"), ValidateCompactBinaryPackage(MakeMemoryView(Buffer), ECbValidateMode::All), ECbValidateError::DuplicateAttachments);
+		CHECK_EQUALS(TEXT("ValidateCompactBinaryPackage(DuplicateMixedAttachments, Mode)"), ValidateCompactBinaryPackage(MakeMemoryView(Buffer), ECbValidateMode::All & ~ECbValidateMode::Package), ECbValidateError::None);
+	}
+	{
+		TCbWriter<32> RootObjectWriter;
+		RootObjectWriter.BeginObject();
+		RootObjectWriter.EndObject();
+		FCbObject NullRootObject(RootObjectWriter.Save().AsObject());
+		FIoHash NullRootObjectHash = NullRootObject.GetHash();
+
+		FBufferArchive Buffer;
+		TCbWriter<256> Writer;
+		Writer.AddHash(NullRootObjectHash);
+		Writer.AddObject(NullRootObject);
+		Writer.AddNull();
+		Writer.Save(Buffer);
+		CHECK_EQUALS(TEXT("ValidateCompactBinaryPackage(NullObject)"), ValidateCompactBinaryPackage(MakeMemoryView(Buffer), ECbValidateMode::All), ECbValidateError::NullPackageObject);
+		CHECK_EQUALS(TEXT("ValidateCompactBinaryPackage(NullObject, Mode)"), ValidateCompactBinaryPackage(MakeMemoryView(Buffer), ECbValidateMode::All & ~ECbValidateMode::Package), ECbValidateError::None);
+	}
+	{
+		FBufferArchive Buffer;
+		TCbWriter<256> Writer;
+		Writer.AddBinaryAttachment(FIoHash::HashBuffer(FMemoryView()));
+		Writer.AddBinary(FMemoryView());
+		Writer.AddNull();
+		Writer.Save(Buffer);
+		CHECK_EQUALS(TEXT("ValidateCompactBinaryPackage(NullAttachment)"), ValidateCompactBinaryPackage(MakeMemoryView(Buffer), ECbValidateMode::All), ECbValidateError::NullPackageAttachment);
+		CHECK_EQUALS(TEXT("ValidateCompactBinaryPackage(NullAttachment, Mode)"), ValidateCompactBinaryPackage(MakeMemoryView(Buffer), ECbValidateMode::All & ~ECbValidateMode::Package), ECbValidateError::None);
+	}
+}
+
+///////////////////////////////////////////////////////////////////////////////////////////////////////////////////////
+
+#endif // WITH_TESTS
diff --git a/Engine/Source/Runtime/Core/Tests/Serialization/CompactBinaryWriterTest.cpp b/Engine/Source/Runtime/Core/Tests/Serialization/CompactBinaryWriterTest.cpp
new file mode 100644
--- /dev/null
+++ b/Engine/Source/Runtime/Core/Tests/Serialization/CompactBinaryWriterTest.cpp
@@ -0,0 +1,1110 @@
+// Copyright Epic Games, Inc. All Rights Reserved.
+
+#if WITH_TESTS
+
+#include "Serialization/CompactBinaryWriter.h"
+#include "IO/IoHash.h"
+#include "Tests/TestHarnessAdapter.h"
+#include "Misc/DateTime.h"
+#include "Misc/Guid.h"
+#include "Misc/Timespan.h"
+#include "Misc/StringBuilder.h"
+#include "Serialization/BufferArchive.h"
+#include "Serialization/CompactBinaryPackage.h"
+#include "Serialization/CompactBinaryValidation.h"
+#if WITH_LOW_LEVEL_TESTS
+#include "TestCommon/Expectations.h"
+#endif
+
+TEST_CASE_NAMED(FCbWriterObjectTest, "System::Core::Serialization::CbWriter::Object", "[ApplicationContextMask][SmokeFilter]")
+{
+	TCbWriter<256> Writer;
+
+	// Test Empty Object
+	{
+		Writer.Reset();
+		Writer.BeginObject();
+		Writer.EndObject();
+		FCbField Field = Writer.Save();
+		if (TestEqual(TEXT("FCbWriter(Object, Empty) Validate"), ValidateCompactBinary(Field.GetOuterBuffer(), ECbValidateMode::All), ECbValidateError::None))
+		{
+			CHECK_MESSAGE(TEXT("FCbWriter(Object, Empty).IsObject()"), Field.IsObject());
+			CHECK_FALSE_MESSAGE(TEXT("FCbWriter(Object, Empty).AsObjectView()"), Field.AsObjectView().CreateViewIterator().HasValue());
+		}
+	}
+
+	// Test Named Empty Object
+	{
+		Writer.Reset();
+		Writer.SetName(ANSITEXTVIEW("Object"));
+		Writer.BeginObject();
+		Writer.EndObject();
+		FCbField Field = Writer.Save();
+		if (TestEqual(TEXT("FCbWriter(Object, Empty, Name) Validate"), ValidateCompactBinary(Field.GetOuterBuffer(), ECbValidateMode::All), ECbValidateError::None))
+		{
+			CHECK_MESSAGE(TEXT("FCbWriter(Object, Empty, Name).IsObject()"), Field.IsObject());
+			CHECK_FALSE_MESSAGE(TEXT("FCbWriter(Object, Empty, Name).AsObjectView()"), Field.AsObjectView().CreateViewIterator().HasValue());
+		}
+	}
+
+	// Test Basic Object
+	{
+		Writer.Reset();
+		Writer.BeginObject();
+		Writer.SetName(ANSITEXTVIEW("Integer")).AddInteger(0);
+		Writer.SetName(ANSITEXTVIEW("Float")).AddFloat(0.0f);
+		Writer.EndObject();
+		FCbField Field = Writer.Save();
+		if (TestEqual(TEXT("FCbWriter(Object, Basic) Validate"), ValidateCompactBinary(Field.GetOuterBuffer(), ECbValidateMode::All), ECbValidateError::None))
+		{
+			CHECK_MESSAGE(TEXT("FCbWriter(Object, Basic).IsObject()"), Field.IsObject());
+			FCbObjectView Object = Field.AsObjectView();
+			CHECK_MESSAGE(TEXT("FCbWriter(Object, Basic).AsObjectView()[Integer]"), Object[ANSITEXTVIEW("Integer")].IsInteger());
+			CHECK_MESSAGE(TEXT("FCbWriter(Object, Basic).AsObjectView()[Float]"), Object[ANSITEXTVIEW("Float")].IsFloat());
+		}
+	}
+
+	// Test Uniform Object
+	{
+		Writer.Reset();
+		Writer.BeginObject();
+		Writer.SetName(ANSITEXTVIEW("Field1")).AddInteger(0);
+		Writer.SetName(ANSITEXTVIEW("Field2")).AddInteger(1);
+		Writer.EndObject();
+		FCbField Field = Writer.Save();
+		if (TestEqual(TEXT("FCbWriter(Object, Uniform) Validate"), ValidateCompactBinary(Field.GetOuterBuffer(), ECbValidateMode::All), ECbValidateError::None))
+		{
+			CHECK_MESSAGE(TEXT("FCbWriter(Object, Uniform).IsObject()"), Field.IsObject());
+			FCbObjectView Object = Field.AsObjectView();
+			CHECK_MESSAGE(TEXT("FCbWriter(Object, Uniform).AsObjectView()[Field1]"), Object[ANSITEXTVIEW("Field1")].IsInteger());
+			CHECK_MESSAGE(TEXT("FCbWriter(Object, Uniform).AsObjectView()[Field2]"), Object[ANSITEXTVIEW("Field2")].IsInteger());
+		}
+	}
+}
+
+///////////////////////////////////////////////////////////////////////////////////////////////////////////////////////
+
+TEST_CASE_NAMED(FCbWriterArrayTest, "System::Core::Serialization::CbWriter::Array", "[ApplicationContextMask][SmokeFilter]")
+{
+	TCbWriter<256> Writer;
+
+	// Test Empty Array
+	{
+		Writer.Reset();
+		Writer.BeginArray();
+		Writer.EndArray();
+		FCbField Field = Writer.Save();
+		if (TestEqual(TEXT("FCbWriter(Array, Empty) Validate"), ValidateCompactBinary(Field.GetOuterBuffer(), ECbValidateMode::All), ECbValidateError::None))
+		{
+			CHECK_MESSAGE(TEXT("FCbWriter(Array, Empty).IsArray()"), Field.IsArray());
+			CHECK_EQUALS(TEXT("FCbWriter(Array, Empty).AsArrayView()"), Field.AsArrayView().Num(), uint64(0));
+		}
+	}
+
+	// Test Named Empty Array
+	{
+		Writer.Reset();
+		Writer.SetName(ANSITEXTVIEW("Array"));
+		Writer.BeginArray();
+		Writer.EndArray();
+		FCbField Field = Writer.Save();
+		if (TestEqual(TEXT("FCbWriter(Array, Empty, Name) Validate"), ValidateCompactBinary(Field.GetOuterBuffer(), ECbValidateMode::All), ECbValidateError::None))
+		{
+			CHECK_MESSAGE(TEXT("FCbWriter(Array, Empty, Name).IsArray()"), Field.IsArray());
+			CHECK_EQUALS(TEXT("FCbWriter(Array, Empty, Name).AsArrayView()"), Field.AsArrayView().Num(), uint64(0));
+		}
+	}
+
+	// Test Basic Array
+	{
+		Writer.Reset();
+		Writer.BeginArray();
+		Writer.AddInteger(0);
+		Writer.AddFloat(0.0f);
+		Writer.EndArray();
+		FCbField Field = Writer.Save();
+		if (TestEqual(TEXT("FCbWriter(Array, Basic) Validate"), ValidateCompactBinary(Field.GetOuterBuffer(), ECbValidateMode::All), ECbValidateError::None))
+		{
+			CHECK_MESSAGE(TEXT("FCbWriter(Array, Basic).IsArray()"), Field.IsArray());
+			FCbFieldViewIterator Iterator = Field.AsArrayView().CreateViewIterator();
+			CHECK_MESSAGE(TEXT("FCbWriter(Array, Basic).AsArrayView()[Integer]"), Iterator.IsInteger());
+			++Iterator;
+			CHECK_MESSAGE(TEXT("FCbWriter(Array, Basic).AsArrayView()[Float]"), Iterator.IsFloat());
+			++Iterator;
+			CHECK_FALSE_MESSAGE(TEXT("FCbWriter(Array, Basic).AsArrayView()[End]"), Iterator.HasValue());
+		}
+	}
+
+	// Test Uniform Array
+	{
+		Writer.Reset();
+		Writer.BeginArray();
+		Writer.AddInteger(0);
+		Writer.AddInteger(1);
+		Writer.EndArray();
+		FCbField Field = Writer.Save();
+		if (TestEqual(TEXT("FCbWriter(Array, Uniform) Validate"), ValidateCompactBinary(Field.GetOuterBuffer(), ECbValidateMode::All), ECbValidateError::None))
+		{
+			CHECK_MESSAGE(TEXT("FCbWriter(Array, Uniform).IsArray()"), Field.IsArray());
+			FCbFieldViewIterator Iterator = Field.AsArrayView().CreateViewIterator();
+			CHECK_MESSAGE(TEXT("FCbWriter(Array, Basic).AsArrayView()[Field1]"), Iterator.IsInteger());
+			++Iterator;
+			CHECK_MESSAGE(TEXT("FCbWriter(Array, Basic).AsArrayView()[Field2]"), Iterator.IsInteger());
+			++Iterator;
+			CHECK_FALSE_MESSAGE(TEXT("FCbWriter(Array, Basic).AsArrayView()[End]"), Iterator.HasValue());
+		}
+	}
+}
+
+///////////////////////////////////////////////////////////////////////////////////////////////////////////////////////
+
+TEST_CASE_NAMED(FCbWriterNullTest, "System::Core::Serialization::CbWriter::Null", "[ApplicationContextMask][SmokeFilter]")
+{
+	TCbWriter<256> Writer;
+
+	// Test Null
+	{
+		Writer.Reset();
+		Writer.AddNull();
+		FCbField Field = Writer.Save();
+		if (TestEqual(TEXT("FCbWriter(Null) Validate"), ValidateCompactBinary(Field.GetOuterBuffer(), ECbValidateMode::All), ECbValidateError::None))
+		{
+			CHECK_FALSE_MESSAGE(TEXT("FCbWriter(Null).HasName()"), Field.HasName());
+			CHECK_MESSAGE(TEXT("FCbWriter(Null).IsNull()"), Field.IsNull());
+		}
+	}
+
+	// Test Null with Name
+	{
+		Writer.Reset();
+		Writer.SetName(ANSITEXTVIEW("Null"));
+		Writer.AddNull();
+		FCbField Field = Writer.Save();
+		if (TestEqual(TEXT("FCbWriter(Null, Name) Validate"), ValidateCompactBinary(Field.GetOuterBuffer(), ECbValidateMode::All), ECbValidateError::None))
+		{
+			CHECK_EQUALS(TEXT("FCbWriter(Null, Name).GetName()"), Field.GetName(), UTF8TEXTVIEW("Null"));
+			CHECK_MESSAGE(TEXT("FCbWriter(Null, Name).HasName()"), Field.HasName());
+			CHECK_MESSAGE(TEXT("FCbWriter(Null, Name).IsNull()"), Field.IsNull());
+		}
+	}
+
+	// Test Null Array/Object Uniformity
+	{
+		Writer.Reset();
+
+		Writer.BeginArray();
+		Writer.AddNull();
+		Writer.AddNull();
+		Writer.AddNull();
+		Writer.EndArray();
+
+		Writer.BeginObject();
+		Writer.SetName(ANSITEXTVIEW("N1")).AddNull();
+		Writer.SetName(ANSITEXTVIEW("N2")).AddNull();
+		Writer.SetName(ANSITEXTVIEW("N3")).AddNull();
+		Writer.EndObject();
+
+		FCbFieldIterator Fields = Writer.Save();
+		CHECK_EQUALS(TEXT("FCbWriter(Null, Uniform) Validate"), ValidateCompactBinaryRange(Fields.GetOuterBuffer(), ECbValidateMode::All), ECbValidateError::None);
+	}
+
+	// Test Null with Save(Buffer)
+	{
+		Writer.Reset();
+		constexpr int NullCount = 3;
+		for (int Index = 0; Index < NullCount; ++Index)
+		{
+			Writer.AddNull();
+		}
+		uint8 Buffer[NullCount]{};
+		FCbFieldViewIterator Fields = Writer.Save(MakeMemoryView(Buffer));
+		if (TestEqual(TEXT("FCbWriter(Null, Memory) Validate"), ValidateCompactBinaryRange(MakeMemoryView(Buffer), ECbValidateMode::All), ECbValidateError::None))
+		{
+			for (int Index = 0; Index < NullCount; ++Index)
+			{
+				CHECK_MESSAGE(TEXT("FCbWriter(Null, Memory) IsNull"), Fields.IsNull());
+				++Fields;
+			}
+			CHECK_FALSE_MESSAGE(TEXT("FCbWriter(Null, Memory) HasValue"), Fields.HasValue());
+		}
+	}
+}
+
+///////////////////////////////////////////////////////////////////////////////////////////////////////////////////////
+
+
+template <typename T, typename Size>
+void TestEqual(const TCHAR* What, TArrayView<T, Size> Actual, TArrayView<T, Size> Expected)
+{
+	CHECK_MESSAGE(What, Actual.Num() == Expected.Num() && CompareItems(Actual.GetData(), Expected.GetData(), Actual.Num()));
+}
+
+TEST_CASE_NAMED(FCbWriterBinaryTest, "System::Core::Serialization::CbWriter::Binary", "[ApplicationContextMask][SmokeFilter]")
+{
+	TCbWriter<256> Writer;
+
+	// Test Empty Binary
+	{
+		Writer.Reset();
+		Writer.AddBinary(nullptr, 0);
+		FCbField Field = Writer.Save();
+		if (TestEqual(TEXT("FCbWriter(Binary, Empty) Validate"), ValidateCompactBinary(Field.GetOuterBuffer(), ECbValidateMode::All), ECbValidateError::None))
+		{
+			CHECK_FALSE_MESSAGE(TEXT("FCbWriter(Binary, Empty).HasName()"), Field.HasName());
+			CHECK_MESSAGE(TEXT("FCbWriter(Binary, Empty).IsBinary()"), Field.IsBinary());
+			CHECK_MESSAGE(TEXT("FCbWriter(Binary, Empty).AsBinaryView()"), Field.AsBinaryView().IsEmpty());
+		}
+	}
+
+	// Test Basic Binary
+	{
+		Writer.Reset();
+		const uint8 BinaryValue[] = { 1, 2, 3, 4, 5, 6 };
+		Writer.SetName(ANSITEXTVIEW("Binary"));
+		Writer.AddBinary(BinaryValue, sizeof(BinaryValue));
+		FCbField Field = Writer.Save();
+		if (TestEqual(TEXT("FCbWriter(Binary, Array) Validate"), ValidateCompactBinary(Field.GetOuterBuffer(), ECbValidateMode::All), ECbValidateError::None))
+		{
+			CHECK_EQUALS(TEXT("FCbWriter(Binary, Array).GetName()"), Field.GetName(), UTF8TEXTVIEW("Binary"));
+			CHECK_MESSAGE(TEXT("FCbWriter(Binary, Array).HasName()"), Field.HasName());
+			CHECK_MESSAGE(TEXT("FCbWriter(Binary, Array).IsBinary()"), Field.IsBinary());
+			CHECK_MESSAGE(TEXT("FCbWriter(Binary, Array).AsBinaryView()"), Field.AsBinaryView().EqualBytes(MakeMemoryView(BinaryValue)));
+		}
+	}
+}
+
+///////////////////////////////////////////////////////////////////////////////////////////////////////////////////////
+
+TEST_CASE_NAMED(FCbWriterStringTest, "System::Core::Serialization::CbWriter::String", "[ApplicationContextMask][SmokeFilter]")
+{
+	TCbWriter<256> Writer;
+
+	// Test Empty Strings
+	{
+		Writer.Reset();
+		Writer.AddString(FAnsiStringView());
+		Writer.AddString(FWideStringView());
+		FCbFieldIterator Fields = Writer.Save();
+		if (TestEqual(TEXT("FCbWriter(String, Empty) Validate"), ValidateCompactBinaryRange(Fields.GetOuterBuffer(), ECbValidateMode::All), ECbValidateError::None))
+		{
+			for (FCbFieldView Field : Fields)
+			{
+				CHECK_FALSE_MESSAGE(TEXT("FCbWriter(String, Empty).HasName()"), Field.HasName());
+				CHECK_MESSAGE(TEXT("FCbWriter(String, Empty).IsString()"), Field.IsString());
+				CHECK_MESSAGE(TEXT("FCbWriter(String, Empty).AsString()"), Field.AsString().IsEmpty());
+			}
+		}
+	}
+
+	// Test Basic Strings
+	{
+		Writer.Reset();
+		Writer.SetName(ANSITEXTVIEW("String")).AddString(ANSITEXTVIEW("Value"));
+		Writer.SetName(ANSITEXTVIEW("String")).AddString(TEXTVIEW("Value"));
+		FCbFieldIterator Fields = Writer.Save();
+		if (TestEqual(TEXT("FCbWriter(String, Basic) Validate"), ValidateCompactBinaryRange(Fields.GetOuterBuffer(), ECbValidateMode::All), ECbValidateError::None))
+		{
+			for (FCbFieldView Field : Fields)
+			{
+				CHECK_EQUALS(TEXT("FCbWriter(String, Basic).GetName()"), Field.GetName(), UTF8TEXTVIEW("String"));
+				CHECK_MESSAGE(TEXT("FCbWriter(String, Basic).HasName()"), Field.HasName());
+				CHECK_MESSAGE(TEXT("FCbWriter(String, Basic).IsString()"), Field.IsString());
+				CHECK_MESSAGE(TEXT("FCbWriter(String, Basic).AsString()"), Field.AsString() == UTF8TEXTVIEW("Value"));
+			}
+		}
+	}
+
+	// Test Long Strings
+	{
+		Writer.Reset();
+		constexpr int DotCount = 256;
+		TUtf8StringBuilder<DotCount + 1> Dots;
+		for (int Index = 0; Index < DotCount; ++Index)
+		{
+			Dots << '.';
+		}
+		Writer.AddString(Dots);
+		Writer.AddString(FString::ChrN(DotCount, TEXT('.')));
+		FCbFieldIterator Fields = Writer.Save();
+		if (TestEqual(TEXT("FCbWriter(String, Long) Validate"), ValidateCompactBinaryRange(Fields.GetOuterBuffer(), ECbValidateMode::All), ECbValidateError::None))
+		{
+			for (FCbFieldView Field : Fields)
+			{
+				CHECK_EQUALS(TEXT("FCbWriter(String, Long).AsString()"), Field.AsString(), Dots.ToView());
+			}
+		}
+	}
+
+	// Test Non-ASCII String
+	{
+		Writer.Reset();
+		WIDECHAR Value[2] = { 0xd83d, 0xde00 };
+		Writer.AddString(ANSITEXTVIEW("\xf0\x9f\x98\x80"));
+		Writer.AddString(FWideStringView(Value, UE_ARRAY_COUNT(Value)));
+		FCbFieldIterator Fields = Writer.Save();
+		if (TestEqual(TEXT("FCbWriter(String, Unicode) Validate"), ValidateCompactBinaryRange(Fields.GetOuterBuffer(), ECbValidateMode::All), ECbValidateError::None))
+		{
+			for (FCbFieldView Field : Fields)
+			{
+				CHECK_EQUALS(TEXT("FCbWriter(String, Unicode).AsString()"), Field.AsString(), UTF8TEXTVIEW("\xf0\x9f\x98\x80"));
+			}
+		}
+	}
+}
+
+///////////////////////////////////////////////////////////////////////////////////////////////////////////////////////
+
+TEST_CASE_NAMED(FCbWriterIntegerTest, "System::Core::Serialization::CbWriter::Integer", "[ApplicationContextMask][SmokeFilter]")
+{
+	TCbWriter<256> Writer;
+
+	auto TestInt32 = [&](int32 Value)
+	{
+		FCbWriter Writer;
+		Writer.Reset();
+		Writer.AddInteger(Value);
+		FCbField Field = Writer.Save();
+		if (TestEqual(TEXT("FCbWriter(Integer, Int32) Validate"), ValidateCompactBinary(Field.GetOuterBuffer(), ECbValidateMode::All), ECbValidateError::None))
+		{
+			CHECK_EQUALS(TEXT("FCbWriter(Integer, Int32) Value"), Field.AsInt32(), Value);
+			CHECK_FALSE_MESSAGE(TEXT("FCbWriter(Integer, Int32) Error"), Field.HasError());
+		}
+	};
+
+	auto TestUInt32 = [&](uint32 Value)
+	{
+		FCbWriter Writer;
+		Writer.Reset();
+		Writer.AddInteger(Value);
+		FCbField Field = Writer.Save();
+		if (TestEqual(TEXT("FCbWriter(Integer, UInt32) Validate"), ValidateCompactBinary(Field.GetOuterBuffer(), ECbValidateMode::All), ECbValidateError::None))
+		{
+			CHECK_EQUALS(TEXT("FCbWriter(Integer, UInt32) Value"), Field.AsUInt32(), Value);
+			CHECK_FALSE_MESSAGE(TEXT("FCbWriter(Integer, UInt32) Error"), Field.HasError());
+		}
+	};
+
+	auto TestInt64 = [&](int64 Value)
+	{
+		FCbWriter Writer;
+		Writer.Reset();
+		Writer.AddInteger(Value);
+		FCbField Field = Writer.Save();
+		if (TestEqual(TEXT("FCbWriter(Integer, Int64) Validate"), ValidateCompactBinary(Field.GetOuterBuffer(), ECbValidateMode::All), ECbValidateError::None))
+		{
+			CHECK_EQUALS(TEXT("FCbWriter(Integer, Int64) Value"), Field.AsInt64(), Value);
+			CHECK_FALSE_MESSAGE(TEXT("FCbWriter(Integer, Int64) Error"), Field.HasError());
+		}
+	};
+
+	auto TestUInt64 = [&](uint64 Value)
+	{
+		FCbWriter Writer;
+		Writer.Reset();
+		Writer.AddInteger(Value);
+		FCbField Field = Writer.Save();
+		if (TestEqual(TEXT("FCbWriter(Integer, UInt64) Validate"), ValidateCompactBinary(Field.GetOuterBuffer(), ECbValidateMode::All), ECbValidateError::None))
+		{
+			CHECK_EQUALS(TEXT("FCbWriter(Integer, UInt64) Value"), Field.AsUInt64(), Value);
+			CHECK_FALSE_MESSAGE(TEXT("FCbWriter(Integer, UInt64) Error"), Field.HasError());
+		}
+	};
+
+	TestUInt32(uint32(0x00));
+	TestUInt32(uint32(0x7f));
+	TestUInt32(uint32(0x80));
+	TestUInt32(uint32(0xff));
+	TestUInt32(uint32(0x0100));
+	TestUInt32(uint32(0x7fff));
+	TestUInt32(uint32(0x8000));
+	TestUInt32(uint32(0xffff));
+	TestUInt32(uint32(0x0001'0000));
+	TestUInt32(uint32(0x7fff'ffff));
+	TestUInt32(uint32(0x8000'0000));
+	TestUInt32(uint32(0xffff'ffff));
+
+	TestUInt64(uint64(0x0000'0001'0000'0000));
+	TestUInt64(uint64(0x7fff'ffff'ffff'ffff));
+	TestUInt64(uint64(0x8000'0000'0000'0000));
+	TestUInt64(uint64(0xffff'ffff'ffff'ffff));
+
+	TestInt32(int32(0x01));
+	TestInt32(int32(0x80));
+	TestInt32(int32(0x81));
+	TestInt32(int32(0x8000));
+	TestInt32(int32(0x8001));
+	TestInt32(int32(0x7fff'ffff));
+	TestInt32(int32(0x8000'0000));
+	TestInt32(int32(0x8000'0001));
+
+	TestInt64(int64(0x0000'0001'0000'0000));
+	TestInt64(int64(0x8000'0000'0000'0000));
+	TestInt64(int64(0x7fff'ffff'ffff'ffff));
+	TestInt64(int64(0x8000'0000'0000'0001));
+	TestInt64(int64(0xffff'ffff'ffff'ffff));
+}
+
+///////////////////////////////////////////////////////////////////////////////////////////////////////////////////////
+
+TEST_CASE_NAMED(FCbWriterFloatTest, "System::Core::Serialization::CbWriter::Float", "[ApplicationContextMask][SmokeFilter]")
+{
+	TCbWriter<256> Writer;
+
+	// Test Float32
+	{
+		Writer.Reset();
+		constexpr float Values[] = { 0.0f, 1.0f, -1.0f, UE_PI };
+		for (float Value : Values)
+		{
+			Writer.AddFloat(Value);
+		}
+		FCbFieldIterator Fields = Writer.Save();
+		if (TestEqual(TEXT("FCbWriter(Float, Single) Validate"), ValidateCompactBinaryRange(Fields.GetOuterBuffer(), ECbValidateMode::All), ECbValidateError::None))
+		{
+			const float* CheckValue = Values;
+			for (FCbFieldView Field : Fields)
+			{
+				CHECK_EQUALS(TEXT("FCbWriter(Float, Single).AsFloat()"), Field.AsFloat(), *CheckValue++);
+				CHECK_FALSE_MESSAGE(TEXT("FCbWriter(Float, Single) Error"), Field.HasError());
+			}
+		}
+	}
+
+	// Test Float64
+	{
+		Writer.Reset();
+		constexpr double Values[] = { 0.0f, 1.0f, -1.0f, UE_PI, 1.9999998807907104, 1.9999999403953552, 3.4028234663852886e38, 6.8056469327705771e38 };
+		for (double Value : Values)
+		{
+			Writer.AddFloat(Value);
+		}
+		FCbFieldIterator Fields = Writer.Save();
+		if (TestEqual(TEXT("FCbWriter(Float, Double) Validate"), ValidateCompactBinaryRange(Fields.GetOuterBuffer(), ECbValidateMode::All), ECbValidateError::None))
+		{
+			const double* CheckValue = Values;
+			for (FCbFieldView Field : Fields)
+			{
+				CHECK_EQUALS(TEXT("FCbWriter(Float, Double).AsDouble()"), Field.AsDouble(), *CheckValue++);
+				CHECK_FALSE_MESSAGE(TEXT("FCbWriter(Float, Double) Error"), Field.HasError());
+			}
+		}
+	}
+}
+
+///////////////////////////////////////////////////////////////////////////////////////////////////////////////////////
+
+TEST_CASE_NAMED(FCbWriterBoolTest, "System::Core::Serialization::CbWriter::Bool", "[ApplicationContextMask][SmokeFilter]")
+{
+	TCbWriter<256> Writer;
+
+	// Test Bool Values
+	{
+		Writer.AddBool(true);
+		Writer.AddBool(false);
+
+		FCbFieldIterator Fields = Writer.Save();
+		if (TestEqual(TEXT("FCbWriter(Bool) Validate"), ValidateCompactBinaryRange(Fields.GetOuterBuffer(), ECbValidateMode::All), ECbValidateError::None))
+		{
+			CHECK_MESSAGE(TEXT("FCbWriter(Bool).AsBool()"), Fields.AsBool());
+			CHECK_FALSE_MESSAGE(TEXT("FCbWriter(Bool) Error"), Fields.HasError());
+			++Fields;
+			CHECK_FALSE_MESSAGE(TEXT("FCbWriter(Bool).AsBool()"), Fields.AsBool());
+			CHECK_FALSE_MESSAGE(TEXT("FCbWriter(Bool) Error"), Fields.HasError());
+		}
+	}
+
+	// Test Bool Array/Object Uniformity
+	{
+		Writer.Reset();
+
+		Writer.BeginArray();
+		Writer.AddBool(false);
+		Writer.AddBool(false);
+		Writer.AddBool(false);
+		Writer.EndArray();
+
+		Writer.BeginObject();
+		Writer.SetName(ANSITEXTVIEW("B1")).AddBool(false);
+		Writer.SetName(ANSITEXTVIEW("B2")).AddBool(false);
+		Writer.SetName(ANSITEXTVIEW("B3")).AddBool(false);
+		Writer.EndObject();
+
+		FCbFieldIterator Fields = Writer.Save();
+		CHECK_EQUALS(TEXT("FCbWriter(Bool, Uniform) Validate"), ValidateCompactBinaryRange(Fields.GetOuterBuffer(), ECbValidateMode::All), ECbValidateError::None);
+	}
+}
+
+///////////////////////////////////////////////////////////////////////////////////////////////////////////////////////
+
+TEST_CASE_NAMED(FCbWriterObjectAttachmentTest, "System::Core::Serialization::CbWriter::ObjectAttachment", "[ApplicationContextMask][SmokeFilter]")
+{
+	TCbWriter<256> Writer;
+
+	const FIoHash::ByteArray ZeroBytes{};
+	const FIoHash::ByteArray SequentialBytes{1, 2, 3, 4, 5, 6, 7, 8, 9, 10, 11, 12, 13, 14, 15, 16, 17, 18, 19, 20};
+
+	const FIoHash Values[] = { FIoHash(ZeroBytes), FIoHash(SequentialBytes) };
+	for (const FIoHash& Value : Values)
+	{
+		Writer.AddObjectAttachment(Value);
+	}
+
+	FCbFieldIterator Fields = Writer.Save();
+	if (TestEqual(TEXT("FCbWriter(ObjectAttachment) Validate"), ValidateCompactBinaryRange(Fields.GetOuterBuffer(), ECbValidateMode::All), ECbValidateError::None))
+	{
+		const FIoHash* CheckValue = Values;
+		for (FCbFieldView Field : Fields)
+		{
+			CHECK_EQUALS(TEXT("FCbWriter(ObjectAttachment).AsObjectAttachment()"), Field.AsObjectAttachment(), *CheckValue++);
+			CHECK_FALSE_MESSAGE(TEXT("FCbWriter(ObjectAttachment) Error"), Field.HasError());
+		}
+	}
+}
+
+///////////////////////////////////////////////////////////////////////////////////////////////////////////////////////
+
+TEST_CASE_NAMED(FCbWriterBinaryAttachmentTest, "System::Core::Serialization::CbWriter::BinaryAttachment", "[ApplicationContextMask][SmokeFilter]")
+{
+	TCbWriter<256> Writer;
+
+	const FIoHash::ByteArray ZeroBytes{};
+	const FIoHash::ByteArray SequentialBytes{1, 2, 3, 4, 5, 6, 7, 8, 9, 10, 11, 12, 13, 14, 15, 16, 17, 18, 19, 20};
+
+	const FIoHash Values[] = { FIoHash(ZeroBytes), FIoHash(SequentialBytes) };
+	for (const FIoHash& Value : Values)
+	{
+		Writer.AddBinaryAttachment(Value);
+	}
+
+	FCbFieldIterator Fields = Writer.Save();
+	if (TestEqual(TEXT("FCbWriter(BinaryAttachment) Validate"), ValidateCompactBinaryRange(Fields.GetOuterBuffer(), ECbValidateMode::All), ECbValidateError::None))
+	{
+		const FIoHash* CheckValue = Values;
+		for (FCbFieldView Field : Fields)
+		{
+			CHECK_EQUALS(TEXT("FCbWriter(BinaryAttachment).AsBinaryAttachment()"), Field.AsBinaryAttachment(), *CheckValue++);
+			CHECK_FALSE_MESSAGE(TEXT("FCbWriter(BinaryAttachment) Error"), Field.HasError());
+		}
+	}
+}
+
+///////////////////////////////////////////////////////////////////////////////////////////////////////////////////////
+
+TEST_CASE_NAMED(FCbWriterHashTest, "System::Core::Serialization::CbWriter::Hash", "[ApplicationContextMask][SmokeFilter]")
+{
+	TCbWriter<256> Writer;
+
+	const FIoHash::ByteArray ZeroBytes{};
+	const FIoHash::ByteArray SequentialBytes{1, 2, 3, 4, 5, 6, 7, 8, 9, 10, 11, 12, 13, 14, 15, 16, 17, 18, 19, 20};
+
+	const FIoHash Values[] = { FIoHash(ZeroBytes), FIoHash(SequentialBytes) };
+	for (const FIoHash& Value : Values)
+	{
+		Writer.AddHash(Value);
+	}
+
+	FCbFieldIterator Fields = Writer.Save();
+	if (TestEqual(TEXT("FCbWriter(Hash) Validate"), ValidateCompactBinaryRange(Fields.GetOuterBuffer(), ECbValidateMode::All), ECbValidateError::None))
+	{
+		const FIoHash* CheckValue = Values;
+		for (FCbFieldView Field : Fields)
+		{
+			CHECK_EQUALS(TEXT("FCbWriter(Hash).AsHash()"), Field.AsHash(), *CheckValue++);
+			CHECK_FALSE_MESSAGE(TEXT("FCbWriter(Hash) Error"), Field.HasError());
+		}
+	}
+}
+
+///////////////////////////////////////////////////////////////////////////////////////////////////////////////////////
+
+TEST_CASE_NAMED(FCbWriterUuidTest, "System::Core::Serialization::CbWriter::Uuid", "[ApplicationContextMask][SmokeFilter]")
+{
+	TCbWriter<256> Writer;
+
+	const FGuid Values[] = { FGuid(), FGuid::NewGuid() };
+	for (const FGuid& Value : Values)
+	{
+		Writer.AddUuid(Value);
+	}
+
+	FCbFieldIterator Fields = Writer.Save();
+	if (TestEqual(TEXT("FCbWriter(Uuid) Validate"), ValidateCompactBinaryRange(Fields.GetOuterBuffer(), ECbValidateMode::All), ECbValidateError::None))
+	{
+		const FGuid* CheckValue = Values;
+		for (FCbFieldView Field : Fields)
+		{
+			CHECK_EQUALS(TEXT("FCbWriter(Uuid).AsUuid()"), Field.AsUuid(), *CheckValue++);
+			CHECK_FALSE_MESSAGE(TEXT("FCbWriter(Uuid) Error"), Field.HasError());
+		}
+	}
+}
+
+///////////////////////////////////////////////////////////////////////////////////////////////////////////////////////
+
+TEST_CASE_NAMED(FCbWriterDateTimeTest, "System::Core::Serialization::CbWriter::DateTime", "[ApplicationContextMask][SmokeFilter]")
+{
+	TCbWriter<256> Writer;
+
+	const FDateTime Values[] = { FDateTime(0), FDateTime(2020, 5, 13, 15, 10) };
+	for (FDateTime Value : Values)
+	{
+		Writer.AddDateTime(Value);
+	}
+
+	FCbFieldIterator Fields = Writer.Save();
+	if (TestEqual(TEXT("FCbWriter(DateTime) Validate"), ValidateCompactBinaryRange(Fields.GetOuterBuffer(), ECbValidateMode::All), ECbValidateError::None))
+	{
+		const FDateTime* CheckValue = Values;
+		for (FCbFieldView Field : Fields)
+		{
+			CHECK_EQUALS(TEXT("FCbWriter(DateTime).AsDateTime()"), Field.AsDateTime(), *CheckValue++);
+			CHECK_FALSE_MESSAGE(TEXT("FCbWriter(DateTime) Error"), Field.HasError());
+		}
+	}
+}
+
+///////////////////////////////////////////////////////////////////////////////////////////////////////////////////////
+
+TEST_CASE_NAMED(FCbWriterTimeSpanTest, "System::Core::Serialization::CbWriter::TimeSpan", "[ApplicationContextMask][SmokeFilter]")
+{
+	TCbWriter<256> Writer;
+
+	const FTimespan Values[] = { FTimespan(0), FTimespan(1, 2, 4, 8) };
+	for (FTimespan Value : Values)
+	{
+		Writer.AddTimeSpan(Value);
+	}
+
+	FCbFieldIterator Fields = Writer.Save();
+	if (TestEqual(TEXT("FCbWriter(TimeSpan) Validate"), ValidateCompactBinaryRange(Fields.GetOuterBuffer(), ECbValidateMode::All), ECbValidateError::None))
+	{
+		const FTimespan* CheckValue = Values;
+		for (FCbFieldView Field : Fields)
+		{
+			CHECK_EQUALS(TEXT("FCbWriter(TimeSpan).AsTimeSpan()"), Field.AsTimeSpan(), *CheckValue++);
+			CHECK_FALSE_MESSAGE(TEXT("FCbWriter(TimeSpan) Error"), Field.HasError());
+		}
+	}
+}
+
+///////////////////////////////////////////////////////////////////////////////////////////////////////////////////////
+
+TEST_CASE_NAMED(FCbWriterObjectIdTest, "System::Core::Serialization::CbWriter::ObjectId", "[ApplicationContextMask][SmokeFilter]")
+{
+	TCbWriter<256> Writer;
+
+	const FCbObjectId Values[] = { FCbObjectId(), FCbObjectId(MakeMemoryView<uint8>({1, 2, 3, 4, 5, 6, 7, 8, 9, 10, 11, 12})) };
+	for (const FCbObjectId& Value : Values)
+	{
+		Writer.AddObjectId(Value);
+	}
+
+	FCbFieldIterator Fields = Writer.Save();
+	if (TestEqual(TEXT("FCbWriter(ObjectId) Validate"), ValidateCompactBinaryRange(Fields.GetOuterBuffer(), ECbValidateMode::All), ECbValidateError::None))
+	{
+		const FCbObjectId* CheckValue = Values;
+		for (FCbFieldView Field : Fields)
+		{
+			CHECK_EQUALS(TEXT("FCbWriter(ObjectId).AsObjectId()"), Field.AsObjectId(), *CheckValue++);
+			CHECK_FALSE_MESSAGE(TEXT("FCbWriter(ObjectId) Error"), Field.HasError());
+		}
+	}
+}
+
+///////////////////////////////////////////////////////////////////////////////////////////////////////////////////////
+
+TEST_CASE_NAMED(FCbWriterCustomByIdTest, "System::Core::Serialization::CbWriter::CustomById", "[ApplicationContextMask][SmokeFilter]")
+{
+	TCbWriter<256> Writer;
+
+	struct FCustomValue
+	{
+		uint64 Type;
+		TArray<uint8, TInlineAllocator<16>> Bytes;
+	};
+	const FCustomValue Values[] =
+	{
+		{ 1, {1, 2, 3} },
+		{ MAX_uint64, {4, 5, 6} },
+	};
+
+	for (const FCustomValue& Value : Values)
+	{
+		Writer.AddCustom(Value.Type, MakeMemoryView(Value.Bytes));
+	}
+
+	FCbFieldIterator Fields = Writer.Save();
+	if (TestEqual(TEXT("FCbWriter(CustomById) Validate"), ValidateCompactBinaryRange(Fields.GetOuterBuffer(), ECbValidateMode::All), ECbValidateError::None))
+	{
+		const FCustomValue* CheckValue = Values;
+		for (FCbFieldView Field : Fields)
+		{
+			CHECK_MESSAGE(TEXT("FCbWriter(CustomById).AsCustom()"), Field.AsCustom(CheckValue->Type).EqualBytes(MakeMemoryView(CheckValue->Bytes)));
+			CHECK_FALSE_MESSAGE(TEXT("FCbWriter(CustomById) Error"), Field.HasError());
+			++CheckValue;
+		}
+	}
+}
+
+///////////////////////////////////////////////////////////////////////////////////////////////////////////////////////
+
+TEST_CASE_NAMED(FCbWriterCustomByNameTest, "System::Core::Serialization::CbWriter::CustomByName", "[ApplicationContextMask][SmokeFilter]")
+{
+	TCbWriter<256> Writer;
+
+	struct FCustomValue
+	{
+		FUtf8StringView Type;
+		TArray<uint8, TInlineAllocator<16>> Bytes;
+	};
+	const FCustomValue Values[] =
+	{
+		{ UTF8TEXTVIEW("Type1"), {1, 2, 3} },
+		{ UTF8TEXTVIEW("Type2"), {4, 5, 6} },
+	};
+
+	for (const FCustomValue& Value : Values)
+	{
+		Writer.AddCustom(Value.Type, MakeMemoryView(Value.Bytes));
+	}
+
+	FCbFieldIterator Fields = Writer.Save();
+	if (TestEqual(TEXT("FCbWriter(CustomByName) Validate"), ValidateCompactBinaryRange(Fields.GetOuterBuffer(), ECbValidateMode::All), ECbValidateError::None))
+	{
+		const FCustomValue* CheckValue = Values;
+		for (FCbFieldView Field : Fields)
+		{
+			CHECK_MESSAGE(TEXT("FCbWriter(CustomByName).AsCustom()"), Field.AsCustom(CheckValue->Type).EqualBytes(MakeMemoryView(CheckValue->Bytes)));
+			CHECK_FALSE_MESSAGE(TEXT("FCbWriter(CustomByName) Error"), Field.HasError());
+			++CheckValue;
+		}
+	}
+}
+
+///////////////////////////////////////////////////////////////////////////////////////////////////////////////////////
+
+TEST_CASE_NAMED(FCbWriterComplexTest, "System::Core::Serialization::CbWriter::Complex", "[ApplicationContextMask][SmokeFilter]")
+{
+	FCbObject Object;
+	FBufferArchive Archive;
+	{
+		FCbWriter Writer;
+		Writer.BeginObject();
+
+		const uint8 LocalField[] = { uint8(ECbFieldType::IntegerPositive | ECbFieldType::HasFieldName), 1, 'I', 42 };
+		Writer.AddField(ANSITEXTVIEW("FieldViewCopy"), FCbFieldView(LocalField));
+		Writer.AddField(ANSITEXTVIEW("FieldCopy"), FCbField(FSharedBuffer::Clone(MakeMemoryView(LocalField))));
+
+		const uint8 LocalObject[] = { uint8(ECbFieldType::Object | ECbFieldType::HasFieldName), 1, 'O', 7,
+			uint8(ECbFieldType::IntegerPositive | ECbFieldType::HasFieldName), 1, 'I', 42,
+			uint8(ECbFieldType::Null | ECbFieldType::HasFieldName), 1, 'N' };
+		Writer.AddObject(ANSITEXTVIEW("ObjectViewCopy"), FCbObjectView(LocalObject));
+		Writer.AddObject(ANSITEXTVIEW("ObjectCopy"), FCbObject(FSharedBuffer::Clone(MakeMemoryView(LocalObject))));
+
+		const uint8 LocalArray[] = { uint8(ECbFieldType::UniformArray | ECbFieldType::HasFieldName), 1, 'A', 4, 2,
+			uint8(ECbFieldType::IntegerPositive), 42, 21 };
+		Writer.AddArray(ANSITEXTVIEW("ArrayViewCopy"), FCbArrayView(LocalArray));
+		Writer.AddArray(ANSITEXTVIEW("ArrayCopy"), FCbArray(FSharedBuffer::Clone(MakeMemoryView(LocalArray))));
+
+		Writer.AddNull(ANSITEXTVIEW("Null"));
+
+		Writer.BeginObject(ANSITEXTVIEW("Binary"));
+		{
+			Writer.AddBinary(ANSITEXTVIEW("Empty"), FMemoryView());
+			Writer.AddBinary(ANSITEXTVIEW("Value"), MakeMemoryView("BinaryValue"));
+			Writer.AddBinary(ANSITEXTVIEW("LargeViewValue"), MakeMemoryView(FString::ChrN(256, TEXT('.'))));
+			Writer.AddBinary(ANSITEXTVIEW("LargeValue"), FSharedBuffer::Clone(MakeMemoryView(FString::ChrN(256, TEXT('!')))));
+		}
+		Writer.EndObject();
+
+		Writer.BeginObject(ANSITEXTVIEW("Strings"));
+		{
+			Writer.AddString(ANSITEXTVIEW("AnsiString"), ANSITEXTVIEW("AnsiValue"));
+			Writer.AddString(ANSITEXTVIEW("WideString"), FString::ChrN(256, TEXT('.')));
+			Writer.AddString(ANSITEXTVIEW("EmptyAnsiString"), FAnsiStringView());
+			Writer.AddString(ANSITEXTVIEW("EmptyWideString"), FWideStringView());
+			Writer.AddString(UTF8TEXTVIEW("EmptyUtf8String"), FUtf8StringView());
+			Writer.AddString("AnsiStringLiteral", "AnsiValue");
+			Writer.AddString("WideStringLiteral", TEXT("AnsiValue"));
+		}
+		Writer.EndObject();
+
+		Writer.BeginArray(ANSITEXTVIEW("Integers"));
+		{
+			Writer.AddInteger(int32(-1));
+			Writer.AddInteger(int64(-1));
+			Writer.AddInteger(uint32(1));
+			Writer.AddInteger(uint64(1));
+			Writer.AddInteger(MIN_int32);
+			Writer.AddInteger(MAX_int32);
+			Writer.AddInteger(MAX_uint32);
+			Writer.AddInteger(MIN_int64);
+			Writer.AddInteger(MAX_int64);
+			Writer.AddInteger(MAX_uint64);
+		}
+		Writer.EndArray();
+
+		Writer.BeginArray(ANSITEXTVIEW("UniformIntegers"));
+		{
+			Writer.AddInteger(0);
+			Writer.AddInteger(MAX_int32);
+			Writer.AddInteger(MAX_uint32);
+			Writer.AddInteger(MAX_int64);
+			Writer.AddInteger(MAX_uint64);
+		}
+		Writer.EndArray();
+
+		Writer.AddFloat(ANSITEXTVIEW("Float32"), 1.0f);
+		Writer.AddFloat(ANSITEXTVIEW("Float64as32"), 2.0);
+		Writer.AddFloat(ANSITEXTVIEW("Float64"), 3.0e100);
+
+		Writer.AddBool(ANSITEXTVIEW("False"), false);
+		Writer.AddBool(ANSITEXTVIEW("True"), true);
+
+		Writer.AddObjectAttachment(ANSITEXTVIEW("ObjectAttachment"), FIoHash());
+		Writer.AddBinaryAttachment(ANSITEXTVIEW("BinaryAttachment"), FIoHash());
+		Writer.AddAttachment(ANSITEXTVIEW("Attachment"), FCbAttachment());
+
+		Writer.AddHash(ANSITEXTVIEW("Hash"), FIoHash());
+		Writer.AddUuid(ANSITEXTVIEW("Uuid"), FGuid());
+
+		Writer.AddDateTimeTicks(ANSITEXTVIEW("DateTimeZero"), 0);
+		Writer.AddDateTime(ANSITEXTVIEW("DateTime2020"), FDateTime(2020, 5, 13, 15, 10));
+
+		Writer.AddTimeSpanTicks(ANSITEXTVIEW("TimeSpanZero"), 0);
+		Writer.AddTimeSpan(ANSITEXTVIEW("TimeSpan"), FTimespan(1, 2, 4, 8));
+
+		Writer.AddObjectId(ANSITEXTVIEW("ObjectIdZero"), FCbObjectId());
+		Writer.AddObjectId(ANSITEXTVIEW("ObjectId"), FCbObjectId(MakeMemoryView<uint8>({1, 2, 3, 4, 5, 6, 7, 8, 9, 10, 11, 12})));
+
+		Writer.BeginObject(ANSITEXTVIEW("NestedObjects"));
+		{
+			Writer.BeginObject(ANSITEXTVIEW("Empty"));
+			Writer.EndObject();
+
+			Writer.BeginObject(ANSITEXTVIEW("Null"));
+			Writer.AddNull(ANSITEXTVIEW("Null"));
+			Writer.EndObject();
+		}
+		Writer.EndObject();
+
+		Writer.BeginArray(ANSITEXTVIEW("NestedArrays"));
+		{
+			Writer.BeginArray();
+			Writer.EndArray();
+
+			Writer.BeginArray();
+			Writer.AddNull();
+			Writer.AddNull();
+			Writer.AddNull();
+			Writer.EndArray();
+
+			Writer.BeginArray();
+			Writer.AddBool(false);
+			Writer.AddBool(false);
+			Writer.AddBool(false);
+			Writer.EndArray();
+
+			Writer.BeginArray();
+			Writer.AddBool(true);
+			Writer.AddBool(true);
+			Writer.AddBool(true);
+			Writer.EndArray();
+		}
+		Writer.EndArray();
+
+		Writer.BeginArray(ANSITEXTVIEW("ArrayOfObjects"));
+		{
+			Writer.BeginObject();
+			Writer.EndObject();
+
+			Writer.BeginObject();
+			Writer.AddNull(ANSITEXTVIEW("Null"));
+			Writer.EndObject();
+		}
+		Writer.EndArray();
+
+		Writer.BeginArray(ANSITEXTVIEW("LargeArray"));
+		for (int Index = 0; Index < 256; ++Index)
+		{
+			Writer.AddInteger(Index - 128);
+		}
+		Writer.EndArray();
+
+		Writer.BeginArray(ANSITEXTVIEW("LargeUniformArray"));
+		for (int Index = 0; Index < 256; ++Index)
+		{
+			Writer.AddInteger(Index);
+		}
+		Writer.EndArray();
+
+		Writer.BeginArray(ANSITEXTVIEW("NestedUniformArray"));
+		for (int Index = 0; Index < 16; ++Index)
+		{
+			Writer.BeginArray();
+			for (int Value = 0; Value < 4; ++Value)
+			{
+				Writer.AddInteger(Value);
+			}
+			Writer.EndArray();
+		}
+		Writer.EndArray();
+
+		Writer.EndObject();
+		Object = Writer.Save().AsObject();
+
+		Writer.Save(Archive);
+		CHECK_EQUALS(TEXT("FCbWriter(Complex).Save(Ar)->Num()"), uint64(Archive.Num()), Writer.GetSaveSize());
+	}
+
+	CHECK_EQUALS(TEXT("FCbWriter(Complex).Save()->Validate"),
+		ValidateCompactBinary(Object.GetOuterBuffer(), ECbValidateMode::All), ECbValidateError::None);
+
+	CHECK_EQUALS(TEXT("FCbWriter(Complex).Save(Ar)->Validate"),
+		ValidateCompactBinary(MakeMemoryView(Archive), ECbValidateMode::All), ECbValidateError::None);
+}
+
+///////////////////////////////////////////////////////////////////////////////////////////////////////////////////////
+
+TEST_CASE_NAMED(FCbWriterOwnedReadOnlyTest, "System::Core::Serialization::CbWriter::OwnedReadOnly", "[ApplicationContextMask][SmokeFilter]")
+{
+	FCbWriter Writer;
+	Writer.BeginObject();
+	Writer.EndObject();
+	FCbObject Object = Writer.Save().AsObject();
+	CHECK_MESSAGE(TEXT("FCbWriter().Save().IsOwned()"), Object.IsOwned());
+}
+
+///////////////////////////////////////////////////////////////////////////////////////////////////////////////////////
+
+TEST_CASE_NAMED(FCbWriterStreamTest, "System::Core::Serialization::CbWriter::Stream", "[ApplicationContextMask][SmokeFilter]")
+{
+	FCbObject Object;
+	{
+		FCbWriter Writer;
+		Writer.BeginObject();
+
+		const uint8 LocalField[] = { uint8(ECbFieldType::IntegerPositive | ECbFieldType::HasFieldName), 1, 'I', 42 };
+		Writer << ANSITEXTVIEW("FieldCopy") << FCbFieldView(LocalField);
+
+		const uint8 LocalObject[] = { uint8(ECbFieldType::Object | ECbFieldType::HasFieldName), 1, 'O', 7,
+			uint8(ECbFieldType::IntegerPositive | ECbFieldType::HasFieldName), 1, 'I', 42,
+			uint8(ECbFieldType::Null | ECbFieldType::HasFieldName), 1, 'N' };
+		Writer << ANSITEXTVIEW("ObjectCopy") << FCbObjectView(LocalObject);
+
+		const uint8 LocalArray[] = { uint8(ECbFieldType::UniformArray | ECbFieldType::HasFieldName), 1, 'A', 4, 2,
+			uint8(ECbFieldType::IntegerPositive), 42, 21 };
+		Writer << ANSITEXTVIEW("ArrayCopy") << FCbArrayView(LocalArray);
+
+		Writer << ANSITEXTVIEW("Null") << nullptr;
+
+		Writer << ANSITEXTVIEW("Strings");
+		Writer.BeginObject();
+		Writer
+			<< ANSITEXTVIEW("AnsiString") << ANSITEXTVIEW("AnsiValue")
+			<< ANSITEXTVIEW("AnsiStringLiteral") << "AnsiValue"
+			<< ANSITEXTVIEW("WideString") << TEXTVIEW("WideValue")
+			<< ANSITEXTVIEW("WideStringLiteral") << TEXT("WideValue");
+		Writer.EndObject();
+
+		Writer << ANSITEXTVIEW("Integers");
+		Writer.BeginArray();
+		Writer << int32(-1) << int64(-1) << uint32(1) << uint64(1);
+		Writer.EndArray();
+
+		Writer << ANSITEXTVIEW("Float32") << 1.0f;
+		Writer << ANSITEXTVIEW("Float64") << 2.0;
+
+		Writer << ANSITEXTVIEW("False") << false << ANSITEXTVIEW("True") << true;
+
+		Writer << ANSITEXTVIEW("Attachment") << FCbAttachment();
+
+		Writer << ANSITEXTVIEW("Hash") << FIoHash();
+		Writer << ANSITEXTVIEW("Uuid") << FGuid();
+
+		Writer << ANSITEXTVIEW("DateTime") << FDateTime(2020, 5, 13, 15, 10);
+		Writer << ANSITEXTVIEW("TimeSpan") << FTimespan(1, 2, 4, 8);
+
+		Writer << ANSITEXTVIEW("ObjectIdZero") << FCbObjectId();
+		Writer << ANSITEXTVIEW("ObjectId") << FCbObjectId(MakeMemoryView<uint8>({1, 2, 3, 4, 5, 6, 7, 8, 9, 10, 11, 12}));
+
+		Writer << "LiteralName" << nullptr;
+
+		Writer.EndObject();
+		Object = Writer.Save().AsObject();
+	}
+
+	CHECK_EQUALS(TEXT("FCbWriter(Stream) Validate"), ValidateCompactBinary(Object.GetOuterBuffer(), ECbValidateMode::All), ECbValidateError::None);
+}
+
+///////////////////////////////////////////////////////////////////////////////////////////////////////////////////////
+
+TEST_CASE_NAMED(FCbWriterStateTest, "System::Core::Serialization::CbWriter::State", "[ApplicationContextMask][SmokeFilter]")
+{
+	FCbWriter Writer;
+
+	// Assert on saving an empty writer.
+	//uint8 EmptyField[1];
+	//Writer.Reset();
+	//Writer.Save(MakeMemoryView(EmptyField));
+
+	// Assert on under-sized save buffer.
+	//uint8 ZeroFieldSmall[1];
+	//Writer.Reset();
+	//Writer.AddInteger(0);
+	//Writer.Save(MakeMemoryView(ZeroFieldSmall));
+
+	// Assert on over-sized save buffer.
+	//uint8 ZeroFieldLarge[3];
+	//Writer.Reset();
+	//Writer.AddInteger(0);
+	//Writer.Save(MakeMemoryView(ZeroFieldLarge));
+
+	// Assert on empty name.
+	//Writer.SetName(ANSITEXTVIEW(""));
+
+	// Assert on name after name.
+	//Writer.SetName(ANSITEXTVIEW("Field")).SetName(ANSITEXTVIEW("Field"));
+
+	// Assert on missing name.
+	//Writer.BeginObject();
+	//Writer.AddNull();
+	//Writer.EndObject();
+
+	// Assert on name in array.
+	//Writer.BeginArray();
+	//Writer.SetName(ANSITEXTVIEW("Field"));
+	//Writer.EndArray();
+
+	// Assert on save in object.
+	//uint8 InvalidObject[1];
+	//Writer.Reset();
+	//Writer.BeginObject();
+	//Writer.Save(MakeMemoryView(InvalidObject));
+	//Writer.EndObject();
+
+	// Assert on save in array.
+	//uint8 InvalidArray[1];
+	//Writer.Reset();
+	//Writer.BeginArray();
+	//Writer.Save(MakeMemoryView(InvalidArray));
+	//Writer.EndArray();
+
+	// Assert on object end with no begin.
+	//Writer.EndObject();
+
+	// Assert on array end with no begin.
+	//Writer.EndArray();
+
+	// Assert on object end after name with no value.
+	//Writer.BeginObject();
+	//Writer.SetName(ANSITEXTVIEW("Field"));
+	//Writer.EndObject();
+
+	// Assert on writing a field with no value.
+	//Writer.Field(FCbFieldView());
+}
+
+///////////////////////////////////////////////////////////////////////////////////////////////////////////////////////
+
+#endif // WITH_TESTS
diff --git a/Engine/Source/Runtime/Core/Tests/Serialization/MemoryReaderTest.cpp b/Engine/Source/Runtime/Core/Tests/Serialization/MemoryReaderTest.cpp
new file mode 100644
--- /dev/null
+++ b/Engine/Source/Runtime/Core/Tests/Serialization/MemoryReaderTest.cpp
@@ -0,0 +1,180 @@
+// Copyright Epic Games, Inc. All Rights Reserved.
+
+#include "CoreMinimal.h"
+#include "Tests/TestHarnessAdapter.h"
+#include "Serialization/MemoryReader.h"
+#include "Serialization/MemoryWriter.h"
+
+#if WITH_TESTS
+
+TEST_CASE_NAMED(FMemoryReaderTest, "System::Core::Serialization::MemoryReader", "[ApplicationContextMask][SmokeFilter]")
+{
+	// Test reading uint64 byte swapped, ensuring that bytes really get swapped.
+	{
+		const uint64 WrittenValue  = 0x8877665544332211ull;
+		const uint64 ExpectedValue = 0x1122334455667788ull;
+
+		TArray<uint8> Bytes;
+		Bytes.AddUninitialized(sizeof(uint64));
+		*reinterpret_cast<uint64*>(Bytes.GetData()) = WrittenValue;
+
+		FMemoryReader Reader(Bytes);
+		Reader.SetByteSwapping(true);
+		uint64 ReadValue;
+		Reader << ReadValue;
+		CHECK_MESSAGE("Test reading uint64 in swapped byte order.", ReadValue == ExpectedValue);
+	}
+
+	// Write all supported types and read them back. (Assuming that MemoryWriterTest tests don't fail).
+	{
+		// Keeps the 'official' test value.
+		const uint8  TestValueU8   = 0x12;
+		const int8   TestValueS8   = 0x34;
+		const uint16 TestValueU16  = 0x1122;
+		const int16  TestValueS16  = 0x3344;
+		const uint32 TestValueU32  = 0x11223344;
+		const int32  TestValueS32  = 0x55667788;
+		const uint64 TestValueU64  = 0x1122334455667788;
+		const int64  TestValueS64  = 0x99AABBCCDDEEFF00;
+		const float  TestValueF    = 128.5f;
+		const double TestValueD    = 256.5;
+		const bool   TestValueB    = true;
+		const WIDECHAR TestValueCh = 0xF2;
+		const FString TestAnsiStr  = TEXT("Joe");
+		const FString TestUtf16Str = TEXT("\uC11C\uC6B8\uC0AC\uB78C"); // Must contains at least one non-ANSI codepoint to ensure UTF16 is used.
+
+		// The values to write.
+		TArray<uint8> Bytes;
+		FMemoryWriter Writer(Bytes);
+		uint8  ValueU8   = TestValueU8;
+		int8   ValueS8   = TestValueS8;
+		uint16 ValueU16  = TestValueU16;
+		int16  ValueS16  = TestValueS16;
+		uint32 ValueU32  = TestValueU32;
+		int32  ValueS32  = TestValueS32;
+		uint64 ValueU64  = TestValueU64;
+		int64  ValueS64  = TestValueS64;
+		float  ValueF    = TestValueF;
+		double ValueD    = TestValueD;
+		bool   ValueB    = TestValueB;
+		WIDECHAR ValueCh = TestValueCh;
+		FString AnsiStr  = TestAnsiStr;
+		FString Utf16Str = TestUtf16Str;
+
+		// Try to not align the values in the buffer, so that the test would fail (on platforms not supporting misaligned read) if the reader tries to cast and dereference a pointer type with misaligned address. Doing like: *(uint64*)Array.GetData();
+		Writer << ValueU8;
+		Writer << ValueU16;
+		Writer << ValueU32;
+		Writer << ValueU64;
+		Writer << ValueS8;
+		Writer << ValueS16;
+		Writer << ValueS32;
+		Writer << ValueS64;
+		Writer << ValueF;
+		Writer << ValueD;
+		Writer << ValueB;
+		Writer << ValueCh;
+		Writer << AnsiStr;
+		Writer << Utf16Str;
+
+		// Write again, but write the bytes swapped.
+		Writer.SetByteSwapping(true);
+
+		Writer << ValueU8;
+		Writer << ValueU16;
+		Writer << ValueU32;
+		Writer << ValueU64;
+		Writer << ValueS8;
+		Writer << ValueS16;
+		Writer << ValueS32;
+		Writer << ValueS64;
+		Writer << ValueF;
+		Writer << ValueD;
+		Writer << ValueB;
+		Writer << ValueCh;
+		Writer << AnsiStr;
+		Writer << Utf16Str;
+
+		// The value to reads without swapping the the value swapped back to this platform endianness.
+		uint8  ReadValueU8,  ReadValueSwapU8;
+		int8   ReadValueS8,  ReadValueSwapS8;
+		uint16 ReadValueU16, ReadValueSwapU16;
+		int16  ReadValueS16, ReadValueSwapS16;
+		uint32 ReadValueU32, ReadValueSwapU32;
+		int32  ReadValueS32, ReadValueSwapS32;
+		uint64 ReadValueU64, ReadValueSwapU64;
+		int64  ReadValueS64, ReadValueSwapS64;
+		float  ReadValueF,   ReadValueSwapF;
+		double ReadValueD,   ReadValueSwapD;
+		bool   ReadValueB,   ReadValueSwapB;
+		WIDECHAR ReadValueCh, ReadValueSwapCh;
+		FString ReadAnsiStr,  ReadSwapAnsiStr;
+		FString ReadUtf16Str, ReadSwapUtf16Str;
+
+		// Read the first set of values written in this platform endianness.
+		FMemoryReader Reader(Bytes);
+		Reader << ReadValueU8;
+		Reader << ReadValueU16;
+		Reader << ReadValueU32;
+		Reader << ReadValueU64;
+		Reader << ReadValueS8;
+		Reader << ReadValueS16;
+		Reader << ReadValueS32;
+		Reader << ReadValueS64;
+		Reader << ReadValueF;
+		Reader << ReadValueD;
+		Reader << ReadValueB;
+		Reader << ReadValueCh;
+		Reader << ReadAnsiStr;
+		Reader << ReadUtf16Str;
+
+		// Read the swapped values, swapping them back to their original value (Assuming that MemoryWriterTest tests don't fail)
+		Reader.SetByteSwapping(true);
+		Reader << ReadValueSwapU8;
+		Reader << ReadValueSwapU16;
+		Reader << ReadValueSwapU32;
+		Reader << ReadValueSwapU64;
+		Reader << ReadValueSwapS8;
+		Reader << ReadValueSwapS16;
+		Reader << ReadValueSwapS32;
+		Reader << ReadValueSwapS64;
+		Reader << ReadValueSwapF;
+		Reader << ReadValueSwapD;
+		Reader << ReadValueSwapB;
+		Reader << ReadValueSwapCh;
+		Reader << ReadSwapAnsiStr;
+		Reader << ReadSwapUtf16Str;
+
+		// Validate that the value read are the expected ones.
+		CHECK_MESSAGE("Test reading 'u8' from byte stream",             ReadValueU8      == TestValueU8);
+		CHECK_MESSAGE("Test reading 'u16' from byte stream",            ReadValueU16     == TestValueU16);
+		CHECK_MESSAGE("Test reading 'u32' from byte stream",            ReadValueU32     == TestValueU32);
+		CHECK_MESSAGE("Test reading 'u64' from byte stream",            ReadValueU64     == TestValueU64);
+		CHECK_MESSAGE("Test reading 's8' from byte stream",             ReadValueS8      == TestValueS8);
+		CHECK_MESSAGE("Test reading 's16' from byte stream",            ReadValueS16     == TestValueS16);
+		CHECK_MESSAGE("Test reading 's32' from byte stream",            ReadValueS32     == TestValueS32);
+		CHECK_MESSAGE("Test reading 's64' from byte stream",            ReadValueS64     == TestValueS64);
+		CHECK_MESSAGE("Test reading 'foat' from byte stream",           ReadValueF       == TestValueF);
+		CHECK_MESSAGE("Test reading 'double' from byte stream",         ReadValueD       == TestValueD);
+		CHECK_MESSAGE("Test reading 'bool' from byte stream",           ReadValueB       == TestValueB);
+		CHECK_MESSAGE("Test reading 'wchar' from byte stream",          ReadValueCh      == TestValueCh);
+		CHECK_MESSAGE("Test reading 'ansi str' from byte stream",       ReadAnsiStr      == TestAnsiStr);
+		CHECK_MESSAGE("Test reading 'utf16 str' from byte stream",      ReadUtf16Str     == TestUtf16Str);
+		CHECK_MESSAGE("Test reading 'u8-swapped' from byte stream",     ReadValueSwapU8  == TestValueU8);
+		CHECK_MESSAGE("Test reading 'u16-swapped' from byte stream",    ReadValueSwapU16 == TestValueU16);
+		CHECK_MESSAGE("Test reading 'u32-swapped' from byte stream",    ReadValueSwapU32 == TestValueU32);
+		CHECK_MESSAGE("Test reading 'u64-swapped' from byte stream",    ReadValueSwapU64 == TestValueU64);
+		CHECK_MESSAGE("Test reading 's8-swapped' from byte stream",     ReadValueSwapS8  == TestValueS8);
+		CHECK_MESSAGE("Test reading 's16'-swapped' from byte stream",   ReadValueSwapS16 == TestValueS16);
+		CHECK_MESSAGE("Test reading 's32-swapped' from byte stream",    ReadValueSwapS32 == TestValueS32);
+		CHECK_MESSAGE("Test reading 's64-swapped' from byte stream",    ReadValueSwapS64 == TestValueS64);
+		CHECK_MESSAGE("Test reading 'float-swapped' from byte stream",  ReadValueSwapF   == TestValueF);
+		CHECK_MESSAGE("Test reading 'double-swapped' from byte stream", ReadValueSwapD   == TestValueD);
+		CHECK_MESSAGE("Test reading 'bool-swapped' from byte stream",   ReadValueSwapB   == TestValueB);
+		CHECK_MESSAGE("Test reading 'wchar-swapped' from byte stream",  ReadValueSwapCh  == TestValueCh);
+		CHECK_MESSAGE("Test reading 'ansi-swapped' from byte stream",   ReadSwapAnsiStr  == TestAnsiStr);
+		CHECK_MESSAGE("Test reading 'utf16-swapped' from byte stream",  ReadSwapUtf16Str == TestUtf16Str);
+	}
+}
+
+#endif // WITH_TESTS
diff --git a/Engine/Source/Runtime/Core/Tests/Serialization/MemoryWriterTest.cpp b/Engine/Source/Runtime/Core/Tests/Serialization/MemoryWriterTest.cpp
new file mode 100644
--- /dev/null
+++ b/Engine/Source/Runtime/Core/Tests/Serialization/MemoryWriterTest.cpp
@@ -0,0 +1,273 @@
+// Copyright Epic Games, Inc. All Rights Reserved.
+
+#include "CoreMinimal.h"
+#include "Tests/TestHarnessAdapter.h"
+#include "Misc/ByteSwap.h"
+#include "Serialization/MemoryWriter.h"
+#include "Serialization/MemoryReader.h"
+
+#if WITH_TESTS
+
+namespace MemoryWriterTestUtil
+{
+template <typename T> struct TIsBoolean { enum { Value = false }; };
+template <> struct TIsBoolean<bool> { enum { Value = true }; };
+
+class MemoryWriterTester
+{
+public:
+	MemoryWriterTester()
+		: Writer(Bytes)
+	{
+	}
+
+	template<typename T>
+	void TestWritePlatformByteOrder(T Value)
+	{
+		const T Original = Value;
+		Writer << Value;
+		CheckWrittenByteCount<T>(Bytes.Num());
+		checkf(Value == Original, TEXT("The writer unexpectedly modified the input value")); // Ensure the written value wasn't changed during the write. (Regression check when writing byte swapped)
+		checkf(FMemory::Memcmp(Bytes.GetData(), &Original, sizeof(Original)) == 0, TEXT("The written value doesn't match the expected one."));
+	}
+
+	template<typename T>
+	void TestWriteSwappedByteOrder(T Value, T Swapped)
+	{
+		const T Original = Value;
+		Writer.SetByteSwapping(true);
+		Writer << Value;
+		CheckWrittenByteCount<T>(Bytes.Num());
+		checkf(Value == Original, TEXT("The writer unexpectedly modified the input value")); // Ensure the written value wasn't changed during the write. (Regression check when writing byte swapped)
+		checkf(FMemory::Memcmp(Bytes.GetData(), &Swapped, sizeof(Swapped)) == 0, TEXT("The written value doesn't match the swapped value."));
+	}
+
+	template<typename T>
+	void CheckWrittenByteCount(int32 WrittenCount)
+	{
+		if (TIsBoolean<T>::Value)
+		{
+			// Boolean are written as 4 bytes integer.
+			checkf(WrittenCount == sizeof(int32), TEXT("Unexpected number of bytes written by the writer"));
+		}
+		else
+		{
+			checkf(WrittenCount == sizeof(T), TEXT("Unexpected number of bytes written by the writer"));
+		}
+	}
+
+public:
+	TArray<uint8> Bytes;
+	FMemoryWriter Writer;
+};
+
+}
+
+TEST_CASE_NAMED(FMemoryWriterTest, "System::Core::Serialization::MemoryWriter", "[ApplicationContextMask][SmokeFilter]")
+{
+	// Keeps the 'official' test value as const, to prevent overwriting them.
+	const uint8  TestValueU8   = 0x12;
+	const int8   TestValueS8   = 0x34;
+	const uint16 TestValueU16  = 0x1122;
+	const int16  TestValueS16  = 0x3344;
+	const uint32 TestValueU32  = 0x11223344;
+	const int32  TestValueS32  = 0x55667788;
+	const uint64 TestValueU64  = 0x1122334455667788;
+	const int64  TestValueS64  = 0x99AABBCCDDEEFF00;
+	const float  TestValueF    = 128.5f;
+	const double TestValueD    = 256.5;
+	const bool   TestValueB    = true;
+	const WIDECHAR TestValueCh = 0xF2;
+
+	// Platform endianness tests.
+	{
+		MemoryWriterTestUtil::MemoryWriterTester Tester;
+		Tester.TestWritePlatformByteOrder(TestValueU8);
+	}
+
+	{
+		MemoryWriterTestUtil::MemoryWriterTester Tester;
+		Tester.TestWritePlatformByteOrder(TestValueS8);
+	}
+
+	{
+		MemoryWriterTestUtil::MemoryWriterTester Tester;
+		Tester.TestWritePlatformByteOrder(TestValueU16);
+	}
+
+	{
+		MemoryWriterTestUtil::MemoryWriterTester Tester;
+		Tester.TestWritePlatformByteOrder(TestValueS16);
+	}
+
+	{
+		MemoryWriterTestUtil::MemoryWriterTester Tester;
+		Tester.TestWritePlatformByteOrder(TestValueU32);
+	}
+
+	{
+		MemoryWriterTestUtil::MemoryWriterTester Tester;
+		Tester.TestWritePlatformByteOrder(TestValueS32);
+	}
+
+	{
+		MemoryWriterTestUtil::MemoryWriterTester Tester;
+		Tester.TestWritePlatformByteOrder(TestValueU64);
+	}
+
+	{
+		MemoryWriterTestUtil::MemoryWriterTester Tester;
+		Tester.TestWritePlatformByteOrder(TestValueS64);
+	}
+
+	{
+		MemoryWriterTestUtil::MemoryWriterTester Tester;
+		Tester.TestWritePlatformByteOrder(TestValueF);
+	}
+
+	{
+		MemoryWriterTestUtil::MemoryWriterTester Tester;
+		Tester.TestWritePlatformByteOrder(TestValueD);
+	}
+
+	{
+		MemoryWriterTestUtil::MemoryWriterTester Tester;
+		Tester.TestWritePlatformByteOrder(TestValueB);
+	}
+
+	// Non Platform endianness tests. (Byte swapping)
+	{
+		MemoryWriterTestUtil::MemoryWriterTester Tester;
+		Tester.TestWriteSwappedByteOrder(TestValueU8, TestValueU8);
+	}
+
+	{
+		MemoryWriterTestUtil::MemoryWriterTester Tester;
+		Tester.TestWriteSwappedByteOrder(TestValueS8, TestValueS8);
+	}
+
+	{
+		MemoryWriterTestUtil::MemoryWriterTester Tester;
+		Tester.TestWriteSwappedByteOrder(TestValueU16, ByteSwap(TestValueU16));
+	}
+
+	{
+		MemoryWriterTestUtil::MemoryWriterTester Tester;
+		Tester.TestWriteSwappedByteOrder(TestValueS16, ByteSwap(TestValueS16));
+	}
+
+	{
+		MemoryWriterTestUtil::MemoryWriterTester Tester;
+		Tester.TestWriteSwappedByteOrder(TestValueU32, ByteSwap(TestValueU32));
+	}
+
+	{
+		MemoryWriterTestUtil::MemoryWriterTester Tester;
+		Tester.TestWriteSwappedByteOrder(TestValueS32, ByteSwap(TestValueS32));
+	}
+
+	{
+		MemoryWriterTestUtil::MemoryWriterTester Tester;
+		Tester.TestWriteSwappedByteOrder(TestValueU64, ByteSwap(TestValueU64));
+	}
+
+	{
+		MemoryWriterTestUtil::MemoryWriterTester Tester;
+		Tester.TestWriteSwappedByteOrder(TestValueS64, ByteSwap(TestValueS64));
+	}
+
+	{
+		MemoryWriterTestUtil::MemoryWriterTester Tester;
+		Tester.TestWriteSwappedByteOrder(TestValueF, ByteSwap(TestValueF));
+	}
+
+	{
+		MemoryWriterTestUtil::MemoryWriterTester Tester;
+		Tester.TestWriteSwappedByteOrder(TestValueD, ByteSwap(TestValueD));
+	}
+
+	{
+		MemoryWriterTestUtil::MemoryWriterTester Tester;
+		Tester.TestWriteSwappedByteOrder(TestValueB, TestValueB);
+	}
+
+	// ANSI String
+	{
+		const TCHAR* RawStr = TEXT("Joe"); // Only contains ANSI chars -> should be serialized as an ANSI string (1 byte per char)
+		const char* RawStrBytes = "Joe";
+
+		const FString TestString = RawStr;
+		int32 CharCount = TestString.Len() + 1; // Include \0.
+		int32 StrByteCount = CharCount;
+		int32 CharCountInBuffer = CharCount; // A positive count let the deserializer know this an ANSI string rather than a UTF16 string.
+		int32 CharCountInBufferSwapped = ByteSwap(CharCountInBuffer);
+
+		// Platform endianness
+		{
+			TArray<uint8> Bytes;
+			FMemoryWriter Writer(Bytes);
+
+			FString WriteValue = TestString;
+			Writer << WriteValue;
+			check(WriteValue == TestString); // Ensure the original value did not change (as side effect).
+			check(FMemory::Memcmp(Bytes.GetData(), &CharCountInBuffer, sizeof(int32)) == 0); // Check the count encoded as the first 4 bytes.
+			check(FMemory::Memcmp(Bytes.GetData() + sizeof(int32), RawStrBytes, StrByteCount) == 0); // Check the string content.
+		}
+	
+		// Swapped endianness
+		{
+			TArray<uint8> Bytes;
+			FMemoryWriter Writer(Bytes);
+			Writer.SetByteSwapping(true); // Swap bytes when writing.
+
+			FString WriteValue = TestString;
+			Writer << WriteValue;
+			check(WriteValue == TestString); // Ensure the original value did not change (as side effect).
+			check(FMemory::Memcmp(Bytes.GetData(), &CharCountInBufferSwapped, sizeof(int32)) == 0); // Check the count encoded as the first 4 bytes.
+			check(FMemory::Memcmp(Bytes.GetData() + sizeof(int32), RawStrBytes, StrByteCount) == 0); // Check the string content.
+		}
+	}
+
+	// UTF16 Strings
+	{
+		// Must not contain at least one non-ANSI char otherwise, the serialization finds out and serialize and ANSI string (1-byte per char).
+		const TCHAR RawStr[] = {TEXT('\u0404'), TEXT('\u0400'), TEXT('\uC0AC'), TEXT('\u0000')};
+		const uint8* RawStrBytes = reinterpret_cast<const uint8*>(RawStr);
+
+		const TCHAR RawStrSwapped[] = { ByteSwap(static_cast<uint16>(RawStr[0])), ByteSwap(static_cast<uint16>(RawStr[1])), ByteSwap(static_cast<uint16>(RawStr[2])), ByteSwap(static_cast<uint16>(RawStr[3]))};
+		const uint8* RawStrBytesSwapped = reinterpret_cast<const uint8*>(RawStrSwapped);
+
+		const FString TestString = RawStr;
+		int32 CharCount = TestString.Len() + 1; // Include \0.
+		int32 StrByteCount = CharCount * sizeof(TCHAR);
+		int32 CharCountInBuffer = -CharCount; // A negative count let the deserializer know this an UTF16 string rather than a ANSI string.
+		int32 CharCountInBufferSwapped = ByteSwap(CharCountInBuffer);
+
+		// Platform endianness
+		{
+			TArray<uint8> Bytes;
+			FMemoryWriter Writer(Bytes);
+
+			FString WriteValue = TestString;
+			Writer << WriteValue;
+			check(WriteValue == TestString); // Ensure the original value did not change (as side effect).
+			check(FMemory::Memcmp(Bytes.GetData(), &CharCountInBuffer, sizeof(int32)) == 0); // Check the count encoded as the first 4 bytes.
+			check(FMemory::Memcmp(Bytes.GetData() + sizeof(int32), RawStrBytes, StrByteCount) == 0); // Check the string content.
+		}
+	
+		// Swapped endianness
+		{
+			TArray<uint8> Bytes;
+			FMemoryWriter Writer(Bytes);
+			Writer.SetByteSwapping(true); // Swap bytes when writing.
+
+			FString WriteValue = TestString;
+			Writer << WriteValue;
+			check(WriteValue == TestString); // Ensure the original value did not change (as side effect).
+			check(FMemory::Memcmp(Bytes.GetData(), &CharCountInBufferSwapped, sizeof(int32)) == 0); // Check the count encoded as the first 4 bytes.
+			check(FMemory::Memcmp(Bytes.GetData() + sizeof(int32), RawStrBytesSwapped, StrByteCount) == 0); // Check the string content.
+		}
+	}
+}
+
+#endif // WITH_TESTS
diff --git a/Engine/Source/Runtime/Core/Tests/Tasks/TasksTest.cpp b/Engine/Source/Runtime/Core/Tests/Tasks/TasksTest.cpp
new file mode 100644
--- /dev/null
+++ b/Engine/Source/Runtime/Core/Tests/Tasks/TasksTest.cpp
@@ -0,0 +1,1364 @@
+// Copyright Epic Games, Inc. All Rights Reserved.
+
+#include "CoreTypes.h"
+#include "Misc/AutomationTest.h"
+#include "Tests/Benchmark.h"
+#include "Tasks/Pipe.h"
+#include "HAL/Thread.h"
+#include "Async/ParallelFor.h"
+#include "Tests/TestHarnessAdapter.h"
+
+#include <atomic>
+
+#if WITH_TESTS
+
+namespace UE { namespace TasksTests
+{
+	using namespace Tasks;
+
+	void DummyFunc()
+	{
+	}
+
+	TEST_CASE_NAMED(FTasksBasicTest, "System::Core::Tasks::Basic", "[.][ApplicationContextMask][EngineFilter]")
+	{
+		if (!FPlatformProcess::SupportsMultithreading())
+		{
+			// the new API doesn't support single-threaded execution (`-nothreading`) until it's feature-compatible with the old API and completely replaces it
+			return;
+		}
+
+		{	// basic example, fire and forget a high-pri task
+			Launch(
+				UE_SOURCE_LOCATION, // debug name
+				[] {}, // task body
+				Tasks::ETaskPriority::High /* task priority, `Normal` by default */
+			);
+		}
+
+		{	// launch a task and wait till it's executed
+			Launch(UE_SOURCE_LOCATION, [] {}).Wait();
+			Launch(UE_SOURCE_LOCATION, [] {}).BusyWait();
+		}
+
+		{	// FTaskEvent asserts on destruction if it wasn't triggered. uncomment the code to verify
+			//FTaskEvent Event{ UE_SOURCE_LOCATION };
+		}
+
+		{	// FTaskEvent blocks execution until it's signalled
+			FTaskEvent Event{ UE_SOURCE_LOCATION };
+			check(!Event.IsCompleted());
+
+			// check that waiting blocks
+			FTask Task = Launch(UE_SOURCE_LOCATION, [&Event] { Event.Wait(); });
+			FPlatformProcess::Sleep(0.1f);
+			check(!Task.IsCompleted());
+
+			Event.Trigger();
+			check(Event.IsCompleted());
+			verify(Event.Wait(FTimespan::Zero()));
+			verify(Event.BusyWait(FTimespan::Zero()));
+		}
+
+		{	// FTaskEvent can be triggered multiple times
+			FTaskEvent Event{ UE_SOURCE_LOCATION };
+			Event.Trigger();
+			check(Event.IsCompleted());
+			Event.Trigger();
+			Event.Trigger();
+			check(Event.IsCompleted());
+		}
+
+		{	// same but using busy-waiting
+			FTaskEvent Event{ UE_SOURCE_LOCATION };
+			check(!Event.IsCompleted());
+
+			// check that waiting blocks
+			FTask Task = Launch(UE_SOURCE_LOCATION, [Event]() mutable { Event.BusyWait(); });
+			FPlatformProcess::Sleep(0.1f);
+			check(!Task.IsCompleted());
+
+			Event.Trigger();
+			check(Event.IsCompleted());
+			verify(Event.Wait(FTimespan::Zero()));
+			verify(Event.BusyWait(FTimespan::Zero()));
+		}
+
+		{	// basic use-case, postpone waiting so the task is executed first
+			std::atomic<bool> Done{ false };
+			FTask Task = Launch(UE_SOURCE_LOCATION, [&Done] { Done = true; });
+			while (!Task.IsCompleted())
+			{
+				FPlatformProcess::Yield();
+			}
+			Task.Wait();
+			check(Done);
+		}
+
+		{	// basic use-case, postpone busy-waiting so the task is executed first
+			std::atomic<bool> Done{ false };
+			FTask Task = Launch(UE_SOURCE_LOCATION, [&Done] { Done = true; });
+			while (!Task.IsCompleted())
+			{
+				FPlatformProcess::Yield();
+			}
+			Task.BusyWait();
+			check(Done);
+		}
+
+		{	// basic use-case with result, postpone execution so waiting kicks in first
+			TTask<int> Task = Launch(UE_SOURCE_LOCATION, [] { FPlatformProcess::Sleep(0.1f);  return 42; });
+			verify(Task.GetResult() == 42);
+		}
+
+		{	// basic use-case with result, postpone waiting so the task is executed first
+			TTask<int> Task = Launch(UE_SOURCE_LOCATION, [] { return 42; });
+			while (!Task.IsCompleted())
+			{
+				FPlatformProcess::Yield();
+			}
+			verify(Task.GetResult() == 42);
+		}
+
+		{ // using an "empty" prerequisites
+			FTask EmptyPrereq;
+			FTask NonEmptyPrereq = Launch(UE_SOURCE_LOCATION, [] {});
+			FTask Task = Launch(UE_SOURCE_LOCATION, [] {}, Prerequisites(EmptyPrereq, NonEmptyPrereq));
+			check(Task.Wait(FTimespan::FromMilliseconds(100)));
+		}
+
+		{	// check that movable-only result types are supported, that only single instance of result is created and that it's destroyed
+			static std::atomic<uint32> ConstructionsNum{ 0 }; // `static` to make it visible to `FMoveConstructable`
+			static std::atomic<uint32> DestructionsNum{ 0 };
+
+			struct FMoveConstructable
+			{
+				FORCENOINLINE FMoveConstructable()
+				{
+					ConstructionsNum.fetch_add(1, std::memory_order_relaxed);
+				}
+
+				FMoveConstructable(FMoveConstructable&&)
+					: FMoveConstructable()
+				{
+				}
+
+				FMoveConstructable(const FMoveConstructable&) = delete;
+				FMoveConstructable& operator=(FMoveConstructable&&) = delete;
+				FMoveConstructable& operator=(const FMoveConstructable&) = delete;
+
+				FORCENOINLINE ~FMoveConstructable()
+				{
+					DestructionsNum.fetch_add(1, std::memory_order_relaxed);
+				}
+			};
+
+			{
+				Launch(UE_SOURCE_LOCATION, [] { return FMoveConstructable{}; }).GetResult();
+			}
+
+#if 0	// unreliable test, destruction can happen on a worker thread, after the task is flagged as completed and so the check can be hit before the destruction
+			uint32 LocalConstructionsNum = ConstructionsNum.load(std::memory_order_relaxed);
+			uint32 LocalDestructionsNum = DestructionsNum.load(std::memory_order_relaxed);
+			checkf(LocalConstructionsNum == 1, TEXT("%d result instances were created but one was expected: the value stored in the task"), LocalConstructionsNum);
+			checkf(LocalConstructionsNum == LocalDestructionsNum, TEXT("Mismatched number of constructions (%d) and destructions (%d)"), LocalConstructionsNum, LocalDestructionsNum);
+
+			ConstructionsNum = 0;
+			DestructionsNum = 0;
+#endif
+
+			{
+				FMoveConstructable Res{ MoveTemp(Launch(UE_SOURCE_LOCATION, [] { return FMoveConstructable{}; }).GetResult()) }; // consume the result
+			}
+
+#if 0	// unreliable test, destruction can happen on a worker thread, after the task is flagged as completed and so the check can be hit before the destruction
+			LocalConstructionsNum = ConstructionsNum.load(std::memory_order_relaxed);
+			LocalDestructionsNum = DestructionsNum.load(std::memory_order_relaxed);
+			checkf(LocalConstructionsNum == 2, TEXT("%d result instances were created but 2 was expected: the value stored in the task"), LocalConstructionsNum);
+			checkf(LocalConstructionsNum == LocalDestructionsNum, TEXT("Mismatched number of constructions (%d) and destructions (%d)"), LocalConstructionsNum, LocalDestructionsNum);
+
+			ConstructionsNum = 0;
+			DestructionsNum = 0;
+#endif
+		}
+
+		// fire and forget: launch a task w/o keeping its reference
+		if (LowLevelTasks::FScheduler::Get().GetNumWorkers() != 0)
+		{
+			std::atomic<bool> bDone{ false };
+			Launch(UE_SOURCE_LOCATION, [&bDone] { bDone = true; });
+			while (!bDone)
+			{
+				FPlatformProcess::Yield();
+			}
+		}
+
+		{	// mutable lambda, compilation check
+			Launch(UE_SOURCE_LOCATION, []() mutable {}).Wait();
+			Launch(UE_SOURCE_LOCATION, []() mutable { return false; }).GetResult();
+		}
+
+		{	// free memory occupied by task, can be required if task instance is held as a member var
+			FTask Task = Launch(UE_SOURCE_LOCATION, [] {});
+			Task.Wait();
+			Task = {};
+		}
+
+		{	// accessing the task from inside its execution
+			FTask Task;
+			Task.Launch(UE_SOURCE_LOCATION, [&Task] { check(!Task.IsCompleted()); });
+			Task.Wait();
+		}
+
+		///////////////////////////////////////////////////////////////////////////////
+		// nested tasks
+		///////////////////////////////////////////////////////////////////////////////
+
+		{	// one nested task
+			FTaskEvent FinishSignal{ UE_SOURCE_LOCATION };
+			FTaskEvent ExecutedSignal{ UE_SOURCE_LOCATION };
+			FTask Task{ Launch(UE_SOURCE_LOCATION,
+				[&FinishSignal, &ExecutedSignal]
+				{
+					AddNested(FinishSignal);
+					ExecutedSignal.Trigger();
+				}
+			) };
+
+			verify(!Task.Wait(FTimespan::FromMilliseconds(100)));
+			verify(ExecutedSignal.IsCompleted());
+			FinishSignal.Trigger();
+			Task.Wait();
+		}
+
+		{	// nested task completed before the parent finishes its execution
+			Launch(UE_SOURCE_LOCATION,
+				[]
+				{
+					FTask NestedTask = Launch(UE_SOURCE_LOCATION, [] {});
+					AddNested(NestedTask);
+					NestedTask.Wait();
+				}
+			).Wait();
+		}
+
+		{	// multiple nested tasks
+			FTaskEvent Signal1{ UE_SOURCE_LOCATION };
+			FTaskEvent Signal2{ UE_SOURCE_LOCATION };
+			FTaskEvent Signal3{ UE_SOURCE_LOCATION };
+
+			FTask Task{ Launch(UE_SOURCE_LOCATION,
+				[&Signal1, &Signal2, &Signal3]
+				{
+					AddNested(Signal1);
+					AddNested(Signal2);
+					AddNested(Signal3);
+				}
+			) };
+
+			verify(!Task.Wait(FTimespan::FromMilliseconds(100)));
+			Signal1.Trigger();
+			verify(!Task.Wait(FTimespan::FromMilliseconds(100)));
+			Signal2.Trigger();
+			verify(!Task.Wait(FTimespan::FromMilliseconds(100)));
+			Signal3.Trigger();
+			Task.Wait();
+		}
+
+		{	// nested in square
+			FTaskEvent Signal{ UE_SOURCE_LOCATION };
+			FTask Task{ Launch(UE_SOURCE_LOCATION,
+				[&Signal]
+				{
+					AddNested(Signal);
+
+					FTaskEvent NestedSignal{ UE_SOURCE_LOCATION };
+					FTask NestedTask{ Launch(UE_SOURCE_LOCATION,
+						[&NestedSignal]
+						{
+							AddNested(NestedSignal);
+						}
+					) };
+
+					verify(!NestedTask.Wait(FTimespan::FromMilliseconds(100)));
+					NestedSignal.Trigger();
+					NestedTask.Wait();
+				}
+			)};
+
+			verify(!Task.Wait(FTimespan::FromMilliseconds(200)));
+			Signal.Trigger();
+			Task.Wait();
+		}
+
+		{	// nested task should block pipe execution, also an example of pipe suspension
+
+			struct FPipeSuspensionScope
+			{
+				explicit FPipeSuspensionScope(FPipe& Pipe)
+				{
+					Pipe.Launch(UE_SOURCE_LOCATION,
+						[this]
+						{
+							SuspendSignal.Trigger();
+							AddNested(ResumeSignal);
+						}
+					);
+
+					SuspendSignal.Wait();
+				}
+
+				~FPipeSuspensionScope()
+				{
+					ResumeSignal.Trigger();
+				}
+
+				FTaskEvent SuspendSignal{ UE_SOURCE_LOCATION };
+				FTaskEvent ResumeSignal{ UE_SOURCE_LOCATION };
+			};
+
+			FPipe Pipe{ UE_SOURCE_LOCATION };
+			FTask Task;
+			{
+				FPipeSuspensionScope Suspension(Pipe);
+				FPlatformProcess::Sleep(0.1f); // let pipe suspension finish its business, which was a bug as pipe was cleared (and unblocked) before AddNested kicked in
+				Task = Pipe.Launch(UE_SOURCE_LOCATION, [] {});
+				verify(!Task.Wait(FTimespan::FromMilliseconds(100)));
+			}
+			Task.Wait();
+		}
+
+		{	// test basic functionality of TTask::IsAwaitable()
+			FTask Task;
+			Task.Launch(UE_SOURCE_LOCATION,
+				[&Task] 
+				{
+					check(!Task.IsAwaitable()); // Task.Wait() would deadlock if called here inside its execution
+				}
+			);
+			check(Task.IsAwaitable());
+			Task.Wait();
+		}
+
+
+		{	// TTask::IsAwaitable() returns false when called from its inner task execution
+			FTask Outer;
+			Outer.Launch(UE_SOURCE_LOCATION,
+				[&Outer]
+				{
+					// the inner task is executed "inline" (synchronously)
+					FTask Inner = Launch
+					(
+						UE_SOURCE_LOCATION, 
+						[&Outer] { check(!Outer.IsAwaitable()); /* still inside `Outer` execution */ }, 
+						ETaskPriority::Default, 
+						EExtendedTaskPriority::Inline
+					);
+					check(Inner.IsCompleted());
+				}
+			);
+			Outer.Wait();
+		}
+
+		{ // check that `Prerequisites()` support containers
+			Prerequisites(TArray<FTask>{});
+			Prerequisites(TArrayView<FTask>{});
+		}
+
+#if TASKGRAPH_NEW_FRONTEND
+		{	// a basic test for a named thread task
+			FTask GTTask = Launch
+			(
+				UE_SOURCE_LOCATION, 
+				[] { check(IsInGameThread()); }, 
+				ETaskPriority::Default, 
+				EExtendedTaskPriority::GameThreadNormalPri
+			);
+			GTTask.Wait();
+		}
+#endif
+
+		//{	// Cancelled task is completed only after nested
+		//	FTaskEvent ParentBlocker{ UE_SOURCE_LOCATION };
+		//	FTaskEvent NestedBlocker{ UE_SOURCE_LOCATION };
+		//	FTask Parent = Launch(UE_SOURCE_LOCATION,
+		//		[&NestedBlocker]
+		//		{
+		//			FTask Nested{ Launch(UE_SOURCE_LOCATION, [] {}, NestedBlocker) };
+		//			AddNested(Nested);
+		//			Nested.Wait();
+		//			// the only reference left is hold by Parent because of `AddNested` dependency
+		//		},
+		//		ParentBlocker
+		//	);
+		//	verify(Parent.TryCancel());
+		//	ParentBlocker.Trigger(); // unblock `Parent` so it can notice that it's been cancelled
+		//	check(Parent.Wait(FTimespan::FromMilliseconds(100))); // but it's still not complete because of the nested task is blocked
+		//	NestedBlocker.Trigger();
+		//	Parent.Wait();
+		//}
+
+	}
+
+	template<uint32 SpawnerGroupsNum, uint32 SpawnersPerGroupNum, uint32 RepeatNum>
+	void BasicStressTest()
+	{
+		TArray<FTask> SpawnerGroups;
+		SpawnerGroups.Reserve(SpawnerGroupsNum);
+
+		constexpr uint32 TasksNum = SpawnerGroupsNum * SpawnersPerGroupNum;
+		TArray<FTask> Spawners;
+		Spawners.AddDefaulted(TasksNum);
+		TArray<FTask> Tasks;
+		Tasks.AddDefaulted(TasksNum);
+
+		for (uint32 RepeatIt = 0; RepeatIt != RepeatNum; ++RepeatIt)
+		{
+			std::atomic<uint32> TasksExecutedNum{ 0 };
+
+			for (uint32 GroupIndex = 0; GroupIndex != SpawnerGroupsNum; ++GroupIndex)
+			{
+				SpawnerGroups.Add(Launch(UE_SOURCE_LOCATION,
+					[
+						Spawners = &Spawners[GroupIndex * SpawnersPerGroupNum],
+						Tasks = &Tasks[GroupIndex * SpawnersPerGroupNum],
+						&TasksExecutedNum
+					]
+					{
+						for (uint32 SpawnerIndex = 0; SpawnerIndex != SpawnersPerGroupNum; ++SpawnerIndex)
+						{
+							Spawners[SpawnerIndex] = Launch(UE_SOURCE_LOCATION,
+								[
+									Task = &Tasks[SpawnerIndex],
+									&TasksExecutedNum
+								]
+								{
+									*Task = Launch(UE_SOURCE_LOCATION,
+										[&TasksExecutedNum]
+										{
+											++TasksExecutedNum;
+										}
+									);
+								}
+							);
+						}
+					}
+				));
+			}
+
+			Wait(SpawnerGroups);
+			Wait(Spawners);
+			Wait(Tasks);
+
+			check(TasksExecutedNum == TasksNum);
+		}
+	}
+
+	TEST_CASE_NAMED(FTasksBasicStressTest, "System::Core::Tasks::BasicStressTest", "[ApplicationContextMask][EngineFilter]")
+	{
+		UE_BENCHMARK(5, BasicStressTest<100, 100, 100>);
+	}
+
+	template<uint32 SpawnerGroupsNum, uint32 SpawnersPerGroupNum>
+	void PipeStressTest();
+
+	TEST_CASE_NAMED(FTasksPipeTest, "System::Core::Tasks::Pipe", "[.][ApplicationContextMask][EngineFilter]")
+	{
+		if (!FPlatformProcess::SupportsMultithreading())
+		{
+			// the new API doesn't support single-threaded execution (`-nothreading`) until it's feature-compatible with the new API and completely replaces it
+			return;
+		}
+
+		{	// a basic usage example
+			Tasks::FPipe Pipe{ UE_SOURCE_LOCATION }; // a debug name, user-provided or 
+				// `UE_SOURCE_LOCATION` - source file name and line number
+			// launch two tasks in the pipe, they will be executed sequentially, but in parallel
+			// with other tasks (including TaskGraph's old API tasks)
+			Tasks::FTask Task1 = Pipe.Launch(UE_SOURCE_LOCATION, [] {});
+			Tasks::FTask Task2 = Pipe.Launch(UE_SOURCE_LOCATION, [] {});
+			Task2.Wait(); // wait for `Task2` completion
+		}
+
+		{	// an example of thread-safe async interface, kind of a primitive "actor"
+			class FAsyncClass
+			{
+			public:
+				TTask<bool> DoSomething()
+				{
+					return Pipe.Launch(TEXT("DoSomething()"), [this] { return DoSomethingImpl(); });
+				}
+
+				FTask DoSomethingElse()
+				{
+					return Pipe.Launch(TEXT("DoSomethingElse()"), [this] { DoSomethingElseImpl(); });
+				}
+
+			private:
+				bool DoSomethingImpl() { return false; }
+				void DoSomethingElseImpl() {}
+
+			private:
+				FPipe Pipe{ UE_SOURCE_LOCATION };
+			};
+
+			// access the same instance from multiple threads
+			FAsyncClass AsyncInstance;
+			bool bRes = AsyncInstance.DoSomething().GetResult();
+			AsyncInstance.DoSomethingElse().Wait();
+		}
+
+		{	// basic
+			FPipe Pipe{ UE_SOURCE_LOCATION };
+			Pipe.Launch(UE_SOURCE_LOCATION, [] {});
+			Pipe.Launch(UE_SOURCE_LOCATION, [] {}).Wait();
+		}
+
+		{	// launching a piped task with pointer to a function
+			FPipe Pipe{ UE_SOURCE_LOCATION };
+			Pipe.Launch(UE_SOURCE_LOCATION, &DummyFunc).Wait();
+		}
+
+		{	// launching a piped task with a functor object
+			struct FFunctor
+			{
+				void operator()()
+				{
+				}
+			};
+
+			FPipe Pipe{ UE_SOURCE_LOCATION };
+			Pipe.Launch(UE_SOURCE_LOCATION, FFunctor{}).Wait();
+		}
+
+		{	// hold the first piped task execution until the next one is piped to test for non-concurrent execution
+			FPipe Pipe{ UE_SOURCE_LOCATION };
+			bool bTask1Done = false;
+			FTask Task1 = Pipe.Launch(UE_SOURCE_LOCATION, 
+				[&bTask1Done] 
+				{ 
+					FPlatformProcess::Sleep(0.1f); 
+					bTask1Done = true; 
+				}
+			);
+			// we can't just check if `Task1` is completed because pipe gets unblocked and so the next piped task can start execution before the
+			// previous piped task's comlpetion flag is set
+			Pipe.Launch(UE_SOURCE_LOCATION, [&bTask1Done] { check(bTask1Done); }).Wait();
+		}
+
+		{	// piping another task after the previous one is completed and destroyed
+			FPipe Pipe{ UE_SOURCE_LOCATION };
+
+			Pipe.Launch(UE_SOURCE_LOCATION, [] {}).Wait();
+			Pipe.Launch(UE_SOURCE_LOCATION, [] {}).Wait();
+		}
+
+		{	// an example of blocking a pipe
+			FPipe Pipe{ UE_SOURCE_LOCATION };
+			std::atomic<bool> bBlocked;
+			FTaskEvent Event{ UE_SOURCE_LOCATION };
+			FTask Task = Pipe.Launch(UE_SOURCE_LOCATION,
+				[&bBlocked, &Event]
+				{
+					bBlocked = true;
+					Event.Wait(); 
+				}
+			);
+			while (!bBlocked)
+			{
+			}
+			// now it's blocked
+			ensure(!Task.Wait(FTimespan::FromMilliseconds(100)));
+
+			Event.Trigger(); // unblock
+			Task.Wait();
+		}
+
+		{	// multiple "inline" tasks piped in one go, check that pipe maintains FIFO order in this case
+			FPipe Pipe{ UE_SOURCE_LOCATION };
+			FTaskEvent Block{ UE_SOURCE_LOCATION };
+			bool bFirstDone = false;
+			bool bSecondDone = false;
+			FTask Task1 = Pipe.Launch(UE_SOURCE_LOCATION, [&] { check(!bSecondDone); bFirstDone = true; }, Prerequisites(Block), ETaskPriority::Normal, EExtendedTaskPriority::Inline);
+			FTask Task2 = Pipe.Launch(UE_SOURCE_LOCATION, [&] { check(bFirstDone); bSecondDone = true; }, Prerequisites(Block), ETaskPriority::Normal, EExtendedTaskPriority::Inline);
+			Block.Trigger();
+			Wait(TArray{ Task1, Task2 });
+			check(bFirstDone && bSecondDone);
+		}
+
+		UE_BENCHMARK(5, PipeStressTest<200, 100>);
+
+	}
+
+	// stress test for named thread tasks, checks spawning a large number of tasks from multiple threads and executing them
+	template<uint32 SpawnerGroupsNum, uint32 SpawnersPerGroupNum>
+	void PipeStressTest()
+	{
+		TArray<FTask> SpawnerGroups;
+		SpawnerGroups.Reserve(SpawnerGroupsNum);
+
+		constexpr uint32 TasksNum = SpawnerGroupsNum * SpawnersPerGroupNum;
+		TArray<FTask> Spawners;
+		Spawners.AddDefaulted(TasksNum);
+		TArray<FTask> Tasks;
+		Tasks.AddDefaulted(TasksNum);
+
+		std::atomic<bool> bExecuting{ false };
+		std::atomic<uint32> TasksExecutedNum{ 0 };
+
+		FPipe Pipe{ UE_SOURCE_LOCATION };
+
+		for (uint32 GroupIndex = 0; GroupIndex != SpawnerGroupsNum; ++GroupIndex)
+		{
+			SpawnerGroups.Add(Launch(UE_SOURCE_LOCATION,
+				[
+					Spawners = &Spawners[GroupIndex * SpawnersPerGroupNum],
+					Tasks = &Tasks[GroupIndex * SpawnersPerGroupNum],
+					&bExecuting,
+					&TasksExecutedNum,
+					&Pipe
+				]
+				{
+					for (uint32 SpawnerIndex = 0; SpawnerIndex != SpawnersPerGroupNum; ++SpawnerIndex)
+					{
+						Spawners[SpawnerIndex] = Launch(UE_SOURCE_LOCATION,
+							[
+								Task = &Tasks[SpawnerIndex],
+								&bExecuting,
+								&TasksExecutedNum,
+								&Pipe
+							]
+							{
+								*Task = Pipe.Launch(UE_SOURCE_LOCATION,
+									[&bExecuting, &TasksExecutedNum]
+									{
+										check(!bExecuting);
+										bExecuting = true;
+										++TasksExecutedNum;
+										bExecuting = false;
+									}
+								);
+							}
+						);
+					}
+				}
+				));
+		}
+
+		Wait(SpawnerGroups);
+		Wait(Spawners);
+		Wait(Tasks);
+
+		check(TasksExecutedNum == TasksNum);
+	}
+
+	TEST_CASE_NAMED(FTasksPipeWaitUntilEmptyTest, "System::Core::Tasks::Pipe::WaitUntilEmpty", "[ApplicationContextMask][EngineFilte]")
+	{
+		//for (int i = 0; i != 100000; ++i)
+		{	// waiting until an empty pipe is empty
+			FPipe Pipe{ UE_SOURCE_LOCATION };
+			Pipe.WaitUntilEmpty();
+		}
+
+		//for (int i = 0; i != 100000; ++i)
+		{	// waiting until a not empty pipe is empty
+			FPipe Pipe{ UE_SOURCE_LOCATION };
+			Pipe.Launch(UE_SOURCE_LOCATION, [] {});
+			Pipe.WaitUntilEmpty();
+		}
+
+		if (false) // disabled as it's an invalid case, but keep it compiling
+		//for (int i = 0; i != 100000; ++i)
+		{	// waiting until a not empty pipe is empty, while new tasks are piped. just an example, this would assert in `Pipe.WaitUntilEmtpy()`
+			// as no more tasks can be piped after that
+
+			FPipe Pipe{ UE_SOURCE_LOCATION };
+
+			std::atomic<bool> bQuit{ false };
+
+			auto TaskBody = [&Pipe, &bQuit](auto& TaskBodyRef) -> void
+			{
+				if (!bQuit)
+				{
+					Pipe.Launch(UE_SOURCE_LOCATION, [&TaskBodyRef] { TaskBodyRef(TaskBodyRef); });
+				}
+			};
+
+			// start the chain
+			TaskBody(TaskBody);
+
+			check(!Pipe.WaitUntilEmpty(FTimespan::FromMilliseconds(100))); // it's not empty until we break the chain
+
+			// break the chain
+			bQuit = true;
+
+			check(Pipe.WaitUntilEmpty());
+		}
+	}
+
+	struct FAutoTlsSlot
+	{
+		uint32 Slot;
+
+		FAutoTlsSlot()
+			: Slot(FPlatformTLS::AllocTlsSlot())
+		{
+		}
+
+		~FAutoTlsSlot()
+		{
+			FPlatformTLS::FreeTlsSlot(Slot);
+		}
+	};
+
+	template<uint64 Num>
+	void UeTlsStressTest()
+	{
+		static FAutoTlsSlot Slot;
+		double Dummy = 0;
+		for (uint64 i = 0; i != Num; ++i)
+		{
+			Dummy += (double)(uintptr_t)(FPlatformTLS::GetTlsValue(Slot.Slot));
+			double Now = FPlatformTime::Seconds();
+			FPlatformTLS::SetTlsValue(Slot.Slot, (void*)(uintptr_t)(Now));
+		}
+		FPlatformTLS::SetTlsValue(Slot.Slot, (void*)(uintptr_t)(Dummy));
+	}
+
+	template<uint64 Num>
+	void ThreadLocalStressTest()
+	{
+		static thread_local double TlsValue;
+		double Dummy = 0;
+		for (uint64 i = 0; i != Num; ++i)
+		{
+			Dummy += TlsValue;
+			double Now = FPlatformTime::Seconds();
+			TlsValue = Now;
+		}
+		TlsValue = Dummy;
+	}
+
+	TEST_CASE_NAMED(FTlsTest, "System::Core::Tls", "[ApplicationContextMask][EngineFilter]")
+	{
+		UE_BENCHMARK(5, UeTlsStressTest<10000000>);
+		UE_BENCHMARK(5, ThreadLocalStressTest<10000000>);
+	}
+
+	template<uint64 NumBranches, uint64 NumTasks>
+	void DependenciesPerfTest()
+	{
+
+		TArray<TTask<FTaskEvent>> Branches;
+		Branches.Reserve(NumBranches);
+		for (uint64 BranchIndex = 0; BranchIndex != NumBranches; ++BranchIndex)
+		{
+			auto Branch = []
+			{
+				TArray<FTask> Tasks;
+				Tasks.Reserve(NumTasks);
+				for (uint64 TaskIndex = 0; TaskIndex != NumTasks; ++TaskIndex)
+				{
+					Tasks.Add(Launch(UE_SOURCE_LOCATION, [] { /*FPlatformProcess::YieldCycles(100000);*/ }));
+				}
+				FTaskEvent Joiner{ UE_SOURCE_LOCATION };
+				Joiner.AddPrerequisites(Tasks);
+				Joiner.Trigger();
+				return Joiner;
+			};
+
+			Branches.Add(Launch(UE_SOURCE_LOCATION, MoveTemp(Branch)));
+		}
+		TArray<FTaskEvent> BranchTasks;
+		BranchTasks.Reserve(NumBranches);
+		for (TTask<FTaskEvent>& Task : Branches)
+		{
+			BranchTasks.Add(Task.GetResult());
+		}
+		Wait(BranchTasks);
+	}
+
+	TEST_CASE_NAMED(FTasksDependenciesTest, "System::Core::Tasks::Dependencies", "[ApplicationContextMask][EngineFilter]")
+	{
+		{	// a task is not executed until its prerequisite (FTaskEvent) is completed
+			FTaskEvent Prereq{ UE_SOURCE_LOCATION };
+
+			FTask Task{ Launch(UE_SOURCE_LOCATION, [] {}, Prereq) };
+			FPlatformProcess::Sleep(0.1f);
+			check(!Task.IsCompleted());
+
+			Prereq.Trigger();
+			Task.Wait();
+		}
+
+		{	// a task is not executed until its prerequisite (FTaskEvent) is completed. with explicit task priority
+			FTaskEvent Prereq{ UE_SOURCE_LOCATION };
+
+			FTask Task{ Launch(UE_SOURCE_LOCATION, [] {}, Prereq, ETaskPriority::Normal) };
+			FPlatformProcess::Sleep(0.1f);
+			check(!Task.IsCompleted());
+
+			Prereq.Trigger();
+			Task.Wait();
+		}
+
+		{	// a task is not executed until its prerequisite (FTask) is completed
+			FTaskEvent Event{ UE_SOURCE_LOCATION };
+			FTask Prereq{ Launch(UE_SOURCE_LOCATION, [&Event] { Event.Wait(); }) };
+			FTask Task{ Launch(UE_SOURCE_LOCATION, [] {}, Prereq) };
+			FPlatformProcess::Sleep(0.1f);
+			check(!Task.IsCompleted());
+
+			Event.Trigger();
+			Task.Wait();
+		}
+
+		{	// compilation test of an iterable collection as prerequisites
+			TArray<FTask> Prereqs{ Launch(UE_SOURCE_LOCATION, [] {}), FTaskEvent{ UE_SOURCE_LOCATION } };
+
+			FTask Task{ Launch(UE_SOURCE_LOCATION, [] {}, Prereqs) };
+			((FTaskEvent&)Prereqs[1]).Trigger();
+			Task.Wait();
+		}
+
+		{	// compilation test of an initializer list as prerequisites
+			auto Prereqs = { Launch(UE_SOURCE_LOCATION, [] {}), Launch(UE_SOURCE_LOCATION, [] {}) };
+			Launch(UE_SOURCE_LOCATION, [] {}, Prereqs).Wait();
+			// the following line doesn't compile because the compiler can't deduce the initializer list type
+			//Launch(UE_SOURCE_LOCATION, [] {}, { Launch(UE_SOURCE_LOCATION, [] {}), Launch(UE_SOURCE_LOCATION, [] {}) }).Wait();
+		}
+
+		{	// a task is not executed until all its prerequisites (FTask and FTaskEvent instances) are completed
+			FTaskEvent Prereq1{ UE_SOURCE_LOCATION };
+			FTaskEvent Event{ UE_SOURCE_LOCATION };
+			FTask Prereq2{ Launch(UE_SOURCE_LOCATION, [&Event] { Event.Wait(); }) };
+
+			TTask<void> Task{ Launch(UE_SOURCE_LOCATION, [] {}, Prerequisites(Prereq1, Prereq2)) };
+			FPlatformProcess::Sleep(0.1f);
+			check(!Task.IsCompleted());
+
+			Prereq1.Trigger();
+			FPlatformProcess::Sleep(0.1f);
+			check(!Task.IsCompleted());
+
+			Event.Trigger();
+			Task.Wait();
+		}
+
+		{	// a task is not executed until all its prerequisites (FTask and FTaskEvent instances) are completed. with explicit task priority
+			FTaskEvent Prereq1{ UE_SOURCE_LOCATION };
+			FTaskEvent Event{ UE_SOURCE_LOCATION };
+			FTask Prereq2{ Launch(UE_SOURCE_LOCATION, [&Event] { Event.Wait(); }) };
+			TArray<FTask> Prereqs{ Prereq1, Prereq2 }; // to check if a random iterable container works as a prerequisite collection
+
+			TTask<void> Task{ Launch(UE_SOURCE_LOCATION, [] {}, Prereqs, ETaskPriority::Normal) };
+			FPlatformProcess::Sleep(0.1f);
+			check(!Task.IsCompleted());
+
+			Prereq1.Trigger();
+			FPlatformProcess::Sleep(0.1f);
+			check(!Task.IsCompleted());
+
+			Event.Trigger();
+			Task.Wait();
+		}
+
+		{	// a piped task blocked by a prerequisites doesn't block the pipe
+			FPipe Pipe{ UE_SOURCE_LOCATION };
+			FTaskEvent Prereq{ UE_SOURCE_LOCATION };
+
+			FTask Task1{ Pipe.Launch(UE_SOURCE_LOCATION, [] {}, Prereq) };
+			FPlatformProcess::Sleep(0.1f);
+			check(!Task1.IsCompleted());
+
+			FTask Task2{ Pipe.Launch(UE_SOURCE_LOCATION, [] {}) };
+			Task2.Wait();
+
+			Prereq.Trigger();
+			Task1.Wait();
+		}
+
+		{	// a piped task with multiple prerequisites
+			FPipe Pipe{ UE_SOURCE_LOCATION };
+			FTaskEvent Prereq1{ UE_SOURCE_LOCATION };
+			FTaskEvent Event{ UE_SOURCE_LOCATION };
+			FTask Prereq2{ Launch(UE_SOURCE_LOCATION, [&Event] { Event.Wait(); }) };
+
+			FTask Task{ Pipe.Launch(UE_SOURCE_LOCATION, [] {}, Prerequisites(Prereq1, Prereq2)) };
+			FPlatformProcess::Sleep(0.1f);
+			check(!Task.IsCompleted());
+
+			Prereq1.Trigger();
+			FPlatformProcess::Sleep(0.1f);
+			check(!Task.IsCompleted());
+
+			Event.Trigger();
+			Task.Wait();
+		}
+
+		{	// a piped task with multiple prerequisites. with explicit task priority
+			FPipe Pipe{ UE_SOURCE_LOCATION };
+			FTaskEvent Prereq1{ UE_SOURCE_LOCATION };
+			FTaskEvent Event{ UE_SOURCE_LOCATION };
+			FTask Prereq2{ Launch(UE_SOURCE_LOCATION, [&Event] { Event.Wait(); }) };
+
+			FTask Task{ Pipe.Launch(UE_SOURCE_LOCATION, [] {}, Prerequisites(Prereq1, Prereq2), ETaskPriority::Normal) };
+			FPlatformProcess::Sleep(0.1f);
+			check(!Task.IsCompleted());
+
+			Prereq1.Trigger();
+			FPlatformProcess::Sleep(0.1f);
+			check(!Task.IsCompleted());
+
+			Event.Trigger();
+			Task.Wait();
+		}
+
+		UE_BENCHMARK(5, DependenciesPerfTest<150, 150>);
+
+	}
+
+	// blocks all workers (except reserve workers) until given event is triggered. Returns blocking tasks.
+	TArray<LowLevelTasks::FTask> BlockWorkers(FTaskEvent& ResumeEvent)
+	{
+		FPlatformProcess::Sleep(0.1f); // give workers time to fall asleep, to avoid any reserve worker messing around
+
+		uint32 NumWorkers = LowLevelTasks::FScheduler::Get().GetNumWorkers();
+
+		TArray<LowLevelTasks::FTask> WorkerBlockers; // tasks that block worker threads
+		WorkerBlockers.Reserve(NumWorkers);
+
+		std::atomic<uint32> NumWorkersNotBlocked{ NumWorkers };
+
+		for (int i = 0; i != NumWorkers; ++i)
+		{
+			WorkerBlockers.Emplace();
+			LowLevelTasks::FTask& Task = WorkerBlockers.Last();
+			Task.Init(TEXT("WorkerBlocker"),
+				[&NumWorkersNotBlocked, &ResumeEvent]
+				{
+					checkf(LowLevelTasks::FScheduler::Get().IsWorkerThread(), TEXT("No reserve workers are expected to get blocked"));
+					--NumWorkersNotBlocked;
+					ResumeEvent.Wait();
+				},
+				LowLevelTasks::ETaskFlags::AllowNothing
+			);
+			LowLevelTasks::FScheduler::Get().TryLaunchAffinity(Task, i);
+		}
+
+		UE::FTimeout Timeout{ FTimespan::FromSeconds(1) };
+		while (NumWorkersNotBlocked != 0)
+		{
+			check(!Timeout);
+			FPlatformProcess::Sleep(0.001f);
+		}
+
+		return WorkerBlockers;
+	}
+
+	// two levels of prerequisites and two levels of nested tasks
+	void TwoLevelsDeepRetractionTest()
+	{
+		FTask P11{ Launch(TEXT("P11"), [] {}) };
+		FTask P12{ Launch(TEXT("P12"), [] {}) };
+		FTask P21{ Launch(TEXT("P21"), [] {}, Prerequisites(P11, P12)) };
+		FTask P22{ Launch(TEXT("P22"), [] {}) };
+		FTask N11, N12, N21, N22;
+		FTask Task = Launch(UE_SOURCE_LOCATION,
+			[&N11, &N12, &N21, &N22]
+			{
+				AddNested(N11 = Launch(TEXT("N11"),
+					[&N21, &N22]
+					{
+						AddNested(N21 = Launch(TEXT("N21"), [] {}));
+						AddNested(N22 = Launch(TEXT("N22"), [] {}));
+					}
+				));
+				AddNested(N12 = Launch(TEXT("N12"), [] {}));
+			},
+			Prerequisites(P21, P22)
+		);
+		Task.Wait();
+		check(P11.IsCompleted() && P12.IsCompleted() && P21.IsCompleted() && P22.IsCompleted() &&
+			N11.IsCompleted() && N12.IsCompleted() && N21.IsCompleted() && N22.IsCompleted());
+	}
+
+	TEST_CASE_NAMED(FTasksDeepRetractionTest, "System::Core::Tasks::DeepRetraction", "[.][ApplicationContextMask][EngineFilter]")
+	{
+		FTaskEvent ResumeEvent{ UE_SOURCE_LOCATION };
+		TArray<LowLevelTasks::FTask> WorkerBlockers = BlockWorkers(ResumeEvent);
+
+		{	// basic retraction, no dependencies
+			bool bDone = false;
+			Launch(UE_SOURCE_LOCATION, [&bDone] { bDone = true; }).Wait();
+			check(bDone);
+		}
+
+		{	// basic deep retraction: single prerequisite
+			bool bPrerequisiteDone = false;
+			bool bTaskDone = false;
+			FTask Prerequisite{ Launch(UE_SOURCE_LOCATION, 
+				[&bPrerequisiteDone, &bTaskDone] 
+				{
+					check(!bPrerequisiteDone);
+					check(!bTaskDone);
+					bPrerequisiteDone = true;
+				}
+			)};
+			Launch(UE_SOURCE_LOCATION, 
+				[&bPrerequisiteDone, &bTaskDone] 
+				{ 
+					check(bPrerequisiteDone);
+					check(!bTaskDone);
+					bTaskDone = true; 
+				},
+				Prerequisite
+			).Wait();
+			check(bPrerequisiteDone);
+			check(bTaskDone);
+		}
+
+		{	// basic deep retraction: single nested task
+			bool bParentDone = false;
+			bool bNestedDone = false;
+			FTask NestedTask;
+			Launch(UE_SOURCE_LOCATION,
+				[&bParentDone, &bNestedDone, &NestedTask]
+				{
+					check(!bParentDone);
+					check(!bNestedDone);
+					NestedTask = Launch(UE_SOURCE_LOCATION,
+						[&bParentDone, &bNestedDone]
+						{
+							check(bParentDone); // due to single-threaded execution (cos all workers are blocked), nested task can't be executed
+							// until the parent is
+							check(!bNestedDone);
+							bNestedDone = true;
+						}
+					);
+					AddNested(NestedTask);
+					check(!bNestedDone);
+					bParentDone = true;
+				}
+			).Wait();
+			check(bParentDone);
+			check(bNestedDone);
+			check(NestedTask.IsCompleted());
+		}
+
+		TwoLevelsDeepRetractionTest();
+
+		{	// retraction of a piped task
+			FPipe Pipe{ UE_SOURCE_LOCATION };
+			FTask PipedTask = Pipe.Launch(UE_SOURCE_LOCATION, [] {});
+			PipedTask.Wait(); // the pipe is not blocked, so the task is retracted successfully
+		}
+
+		{	// deep retraction of a piped task: waiting for a piped task when the pipe has other incomplete tasks before it
+			FPipe Pipe{ UE_SOURCE_LOCATION };
+			FTask PipedTask1 = Pipe.Launch(UE_SOURCE_LOCATION, [] {});
+			FTask PipedTask2 = Pipe.Launch(UE_SOURCE_LOCATION, [] {});
+			PipedTask2.Wait(); // the pipe is not blocked, deep retraction first executes `PipedTask1`, and then `PipedTask2`
+		}
+
+		// an example of a deadlock caused by waiting for a piped task from inside execution of another task of the same pipe
+		//{	// retraction of a piped task from inside that pipe
+		//	FPipe Pipe{ UE_SOURCE_LOCATION };
+		//	FTask PipedOuterTask = Pipe.Launch(UE_SOURCE_LOCATION, 
+		//		[&Pipe] () mutable
+		//		{
+		//			FTask PipedInnerTask = Pipe.Launch(UE_SOURCE_LOCATION, [] {});
+		//			PipedInnerTask.Wait(); // deadlock as the pipe is blocked (we are inside it)
+		//		}
+		//	);
+		//	PipedOuterTask.Wait();
+		//}
+
+		ResumeEvent.Trigger();
+		LowLevelTasks::BusyWaitForTasks<LowLevelTasks::FTask>(WorkerBlockers);
+	}
+
+	template<uint32 Num>
+	void DeepRetractionStressTest()
+	{
+		for (uint32 i = 0; i != Num; ++i)
+		{
+			TwoLevelsDeepRetractionTest();
+		}
+	}
+
+	TEST_CASE_NAMED(FTasksDeepRetractionStressTest, "System::Core::Tasks::DeepRetraction::Stress", "[ApplicationContextMask][EngineFilter]")
+	{
+		UE_BENCHMARK(5, DeepRetractionStressTest<1000>);
+	}
+
+	template<uint64 Num>
+	void NestedTasksStressTest()
+	{
+		for (uint64 i = 0; i != Num; ++i)
+		{
+			FTask Nested;
+			FTask Parent = Launch(TEXT("Parent"),
+				[&Nested]
+				{
+					AddNested(Nested = Launch(TEXT("Nested"), [] {}));
+				}
+			);
+			Parent.Wait();
+			check(Nested.IsCompleted() && Parent.IsCompleted());
+		}
+	}
+
+	TEST_CASE_NAMED(FTasksNestedTasksStressTest, "System::Core::Tasks::NestedTasks::Stress", "[ApplicationContextMask][EngineFilter]")
+	{
+		UE_BENCHMARK(5, NestedTasksStressTest<10000>);
+	}
+
+	template<int NumTasks>
+	void TestPerfBasic()
+	{
+		TArray<FTask> Tasks;
+		Tasks.Reserve(NumTasks);
+
+		for (int32 TaskIndex = 0; TaskIndex < NumTasks; ++TaskIndex)
+		{
+			Tasks.Emplace(Launch(UE_SOURCE_LOCATION, [] {}));
+		}
+
+		Wait(Tasks);
+	}
+
+	template<int32 NumTasks, int32 BatchSize>
+	void TestPerfBatch()
+	{
+		static_assert(NumTasks % BatchSize == 0, "`NumTasks` must be divisible by `BatchSize`");
+		constexpr int32 NumBatches = NumTasks / BatchSize;
+
+		TArray<FTask> Batches;
+		Batches.Reserve(NumBatches);
+		TArray<FTask> Tasks;
+		Tasks.AddDefaulted(NumTasks);
+
+		for (int32 BatchIndex = 0; BatchIndex < NumBatches; ++BatchIndex)
+		{
+			Batches.Add(Launch(UE_SOURCE_LOCATION,
+				[&Tasks, BatchIndex]
+				{
+					for (int32 TaskIndex = 0; TaskIndex < BatchSize; ++TaskIndex)
+					{
+						Tasks[BatchIndex * BatchSize + TaskIndex] = Launch(UE_SOURCE_LOCATION, [] {});
+					}
+				}
+			));
+		}
+
+		Wait(Batches);
+		Wait(Tasks);
+	}
+
+	template<int32 NumTasks, int32 BatchSize>
+	void TestPerfBatchOptimised()
+	{
+		static_assert(NumTasks % BatchSize == 0, "`NumTasks` must be divisible by `BatchSize`");
+		constexpr int32 NumBatches = NumTasks / BatchSize;
+
+		FTaskEvent SpawnSignal{ UE_SOURCE_LOCATION };
+		TArray<FTask> AllDone;
+
+		for (int32 BatchIndex = 0; BatchIndex < NumBatches; ++BatchIndex)
+		{
+			AllDone.Add(Launch(UE_SOURCE_LOCATION,
+				[]
+				{
+					FTaskEvent RunSignal{ UE_SOURCE_LOCATION };
+					for (int32 TaskIndex = 0; TaskIndex < BatchSize; ++TaskIndex)
+					{
+						//AddNested(Launch(UE_SOURCE_LOCATION, [] {}, RunSignal));
+					}
+					RunSignal.Trigger();
+				},
+				SpawnSignal
+			));
+		}
+
+		SpawnSignal.Trigger();
+		Wait(AllDone);
+	}
+
+	template<int NumTasks>
+	void TestLatency()
+	{
+		for (uint32 TaskIndex = 0; TaskIndex != NumTasks; ++TaskIndex)
+		{
+			Launch(UE_SOURCE_LOCATION, [] {}).Wait();
+		}
+	}
+
+	template<uint32 NumTasks>
+	void TestFGraphEventPerf()
+	{
+		FTaskEvent Prereq{ UE_SOURCE_LOCATION };
+		std::atomic<uint32> CompletedTasks{ 0 };
+
+		TArray<FTask> Tasks;
+		for (int i = 0; i != NumTasks; ++i)
+		{
+			Tasks.Add(Launch(UE_SOURCE_LOCATION,
+				[&Prereq, &CompletedTasks]()
+				{
+					Prereq.Wait();
+					++CompletedTasks;
+				}
+			));
+		}
+
+		Prereq.Trigger();
+		Wait(Tasks);
+
+		check(CompletedTasks == NumTasks);
+	}
+
+	template<int NumTasks>
+	void TestSpawning()
+	{
+		{
+			TArray<FTask> Tasks;
+			Tasks.Reserve(NumTasks);
+			//double StartTime = FPlatformTime::Seconds();
+			for (uint32 TaskNo = 0; TaskNo != NumTasks; ++TaskNo)
+			{
+				Tasks.Add(Launch(UE_SOURCE_LOCATION, [] {}));
+			}
+
+			//double Duration = FPlatformTime::Seconds() - StartTime;
+			//UE_LOG(LogTemp, Display, TEXT("Spawning %d empty trackable tasks took %f secs"), NumTasks, Duration);
+
+			Wait(Tasks);
+		}
+		{
+			double StartTime = FPlatformTime::Seconds();
+			for (uint32 TaskNo = 0; TaskNo != NumTasks; ++TaskNo)
+			{
+				Launch(UE_SOURCE_LOCATION, [] {});
+			}
+
+			//double Duration = FPlatformTime::Seconds() - StartTime;
+			//UE_LOG(LogTemp, Display, TEXT("Spawning %d empty non-trackable tasks took %f secs"), NumTasks, Duration);
+		}
+	}
+
+	template<int NumTasks>
+	void TestBatchSpawning()
+	{
+		//double StartTime = FPlatformTime::Seconds();
+		FTaskEvent Prereq{ UE_SOURCE_LOCATION };
+		TArray<FTask> Tasks;
+		Tasks.Reserve(NumTasks);
+		for (uint32 TaskNo = 0; TaskNo != NumTasks; ++TaskNo)
+		{
+			Tasks.Add(Launch(UE_SOURCE_LOCATION, [] {}, Prereq));
+		}
+
+		//double SpawnedTime = FPlatformTime::Seconds();
+		Prereq.Trigger();
+
+		//double EndTime = FPlatformTime::Seconds();
+		//UE_LOG(LogTemp, Display, TEXT("Spawning %d empty non-trackable tasks took %f secs total, %f secs spawning and %f secs dispatching"), NumTasks, EndTime - StartTime, SpawnedTime - StartTime, EndTime - SpawnedTime);
+
+		Wait(Tasks);
+	}
+
+	template<int64 NumBatches, int64 NumTasksPerBatch>
+	void TestWorkStealing()
+	{
+		TArray<FTask> Batches;
+		Batches.Reserve(NumBatches);
+
+		TArray<FTask> Tasks[NumBatches];
+		for (int32 BatchIndex = 0; BatchIndex != NumBatches; ++BatchIndex)
+		{
+			Tasks[BatchIndex].Reserve(NumTasksPerBatch);
+			Batches.Add(Launch(UE_SOURCE_LOCATION,
+				[&Tasks, BatchIndex]()
+				{
+					for (int32 TaskIndex = 0; TaskIndex < NumTasksPerBatch; ++TaskIndex)
+					{
+						Tasks[BatchIndex].Add(Launch(UE_SOURCE_LOCATION, [] {}));
+					}
+				}
+			));
+		}
+
+		Wait(Batches);
+		for (int32 BatchIndex = 0; BatchIndex != NumBatches; ++BatchIndex)
+		{
+			Wait(Tasks[BatchIndex]);
+		}
+	}
+
+	TEST_CASE_NAMED(FTasksPerfTest, "System::Core::Tasks::PerfTest", "[ApplicationContextMask][EngineFilter]")
+	{
+		TRACE_CPUPROFILER_EVENT_SCOPE(TaskGraphTests_PerfTest);
+
+		UE_BENCHMARK(5, TestPerfBasic<10000>);
+		UE_BENCHMARK(5, TestPerfBatch<10000, 100>);
+		UE_BENCHMARK(5, TestPerfBatchOptimised<10000, 100>);
+		UE_BENCHMARK(5, TestLatency<1000>);
+		UE_BENCHMARK(5, TestWorkStealing<100, 100>);
+		UE_BENCHMARK(5, TestSpawning<10000>);
+		UE_BENCHMARK(5, TestBatchSpawning<10000>);
+	}
+
+	TEST_CASE_NAMED(FTasksPriorityCVarTest, "System::Core::Tasks::PriorityCVar", "[ApplicationContextMask][EngineFilter]")
+	{
+		// set every combination of task priority and extended task priority
+		const TCHAR* CVarName = TEXT("TasksPriorityTestCVar");
+		UE::Tasks::FTaskPriorityCVar CVar{ CVarName, TEXT("CVarHelp"), ETaskPriority::Normal, EExtendedTaskPriority::None };
+		IConsoleVariable* CVarPtr = IConsoleManager::Get().FindConsoleVariable(CVarName);
+		for (int Priority = 0; Priority != (int)ETaskPriority::Count; ++Priority)
+		{
+			for (int ExtendedPriority = 0; ExtendedPriority != (int)EExtendedTaskPriority::Count; ++ExtendedPriority)
+			{
+				TStringBuilder<1024> TaskPriorities;
+				TaskPriorities.Append(ToString((ETaskPriority)Priority));
+				TaskPriorities.Append(TEXT(" "));
+				TaskPriorities.Append(ToString((EExtendedTaskPriority)ExtendedPriority));
+
+				CVarPtr->Set(*TaskPriorities);
+
+				check(CVar.GetTaskPriority() == (ETaskPriority)Priority);
+				check(CVar.GetExtendedTaskPriority() == (EExtendedTaskPriority)ExtendedPriority);
+				check(CVarPtr->GetString() == *TaskPriorities);
+			}
+		}
+
+		// test setting only task priority
+		CVarPtr->Set(TEXT("High"));
+		check(CVar.GetTaskPriority() == ETaskPriority::High);
+		check(CVar.GetExtendedTaskPriority() == EExtendedTaskPriority::None);
+		CVarPtr->Set(TEXT("Normal"));
+		check(CVar.GetTaskPriority() == ETaskPriority::Normal);
+
+		// usage example
+		CVarPtr->Set(TEXT("Normal"));
+		Launch(UE_SOURCE_LOCATION, [] {}, CVar.GetTaskPriority(), CVar.GetExtendedTaskPriority()).Wait();
+
+	}
+
+	// test Pipe support for `-nothreading` config (`FPlatformProcess::SupportsMultithreading()` is false)
+	TEST_CASE_NAMED(FTasksPipeNoThreadingTest, "System::Core::Tasks::PipeNoThreading", "[ApplicationContextMask][EngineFilter]")
+	{
+		if (FPlatformProcess::SupportsMultithreading())
+		{
+			return; // run with `-nothreading` for this test to make any sense
+		}
+
+		FPipe Pipe{ UE_SOURCE_LOCATION };
+		Pipe.Launch(UE_SOURCE_LOCATION, [] {});
+		FTask FinalTask = Pipe.Launch(UE_SOURCE_LOCATION, [] {});
+		check(FinalTask.Wait(FTimespan::FromMilliseconds(100)));
+
+	}
+}}
+
+#endif // WITH_TESTS
