Allowing multiple DPCVars

DPCVar(s|)= are now processed in a loop to allow for multiple entries form the command line.
Useful for launching from EGS

diff --git a/Engine/Source/Runtime/Core/Private/Misc/Parse.cpp b/Engine/Source/Runtime/Core/Private/Misc/Parse.cpp
--- a/Engine/Source/Runtime/Core/Private/Misc/Parse.cpp
+++ b/Engine/Source/Runtime/Core/Private/Misc/Parse.cpp
@@ -234,12 +234,18 @@ bool FParse::Value(
 	const TCHAR*	Match,
 	TCHAR*			Value,
 	int32			MaxLen,
-	bool			bShouldStopOnSeparator
+	bool			bShouldStopOnSeparator,
+	const TCHAR**	OptStreamGotTo
 )
 {
 	bool bSuccess = false;
 	int32 MatchLen = FCString::Strlen(Match);
 
+	if (OptStreamGotTo)
+	{
+		*OptStreamGotTo = nullptr;
+	}
+
 	for (const TCHAR* Found = FCString::Strifind(Stream, Match, true); Found != nullptr; Found = FCString::Strifind(Found + MatchLen, Match, true))
 	{
 		const TCHAR* Start = Found + MatchLen;
@@ -282,6 +288,11 @@ bool FParse::Value(
 			}
 		}
 
+		if (OptStreamGotTo) 
+		{
+			*OptStreamGotTo = Start + FCString::Strlen(Value);
+		}
+
 		bSuccess = true;
 		break;
 	}
@@ -318,7 +329,7 @@ bool FParse::Param( const TCHAR* Stream, const TCHAR* Param )
 // 
 // Parse a string.
 //
-bool FParse::Value( const TCHAR* Stream, const TCHAR* Match, FString& Value, bool bShouldStopOnSeparator )
+bool FParse::Value( const TCHAR* Stream, const TCHAR* Match, FString& Value, bool bShouldStopOnSeparator, const TCHAR** OptStreamGotTo)
 {
 	if (!Stream)
 	{
@@ -332,7 +343,7 @@ bool FParse::Value( const TCHAR* Stream, const TCHAR* Match, FString& Value, boo
 		ValueCharArray.AddUninitialized(StreamLen + 1);
 		ValueCharArray[0] = TCHAR('\0');
 
-		if( FParse::Value(Stream, Match, ValueCharArray.GetData(), ValueCharArray.Num(), bShouldStopOnSeparator) )
+		if( FParse::Value(Stream, Match, ValueCharArray.GetData(), ValueCharArray.Num(), bShouldStopOnSeparator, OptStreamGotTo) )
 		{
 			Value = FString(ValueCharArray.GetData());
 			return true;
diff --git a/Engine/Source/Runtime/Core/Public/Misc/Parse.h b/Engine/Source/Runtime/Core/Public/Misc/Parse.h
--- a/Engine/Source/Runtime/Core/Public/Misc/Parse.h
+++ b/Engine/Source/Runtime/Core/Public/Misc/Parse.h
@@ -30,8 +30,18 @@ struct CORE_API FParse
 	static bool Value( const TCHAR* Stream, const TCHAR* Match, uint32& Value );
 	/** Parses a globally unique identifier. */
 	static bool Value( const TCHAR* Stream, const TCHAR* Match, struct FGuid& Guid );
-	/** Parses a string from a text string. */
-	static bool Value( const TCHAR* Stream, const TCHAR* Match, TCHAR* Value, int32 MaxLen, bool bShouldStopOnSeparator=true );
+	/** Parses a string from a text string. 
+	 * @param Stream, the string you want to extract the value from.
+	 * @param Match, the identifier for the value in the stream.
+	 * @param Value, the destination to the value to be extracted to.
+	 * @param MaxLen, the maximum size eof the string that can be extracted.
+	 * @param bShouldStopOnSeparator, (default = false) If this is true, and the value doesn't start with a '"'
+				then it may be truncated to ',' or ')' in addition to whitespace.
+	 * @param OptStreamGotTo, (default = nullptr) If this is not null, then its dereference is set to the address
+				of the end of the value within Stream. This permits consuming of stream in a loop where Match may
+				occur multiple times.
+	*/
+	static bool Value( const TCHAR* Stream, const TCHAR* Match, TCHAR* Value, int32 MaxLen, bool bShouldStopOnSeparator=true, const TCHAR** OptStreamGotTo = nullptr);
 	/** Parses a byte. */
 	static bool Value( const TCHAR* Stream, const TCHAR* Match, uint8& Value );
 	/** Parses a signed byte. */
@@ -46,8 +56,17 @@ struct CORE_API FParse
 	static bool Value(const TCHAR* Stream, const TCHAR* Match, double& Value);
 	/** Parses a signed double word. */
 	static bool Value( const TCHAR* Stream, const TCHAR* Match, int32& Value );
-	/** Parses a string. */
-	static bool Value( const TCHAR* Stream, const TCHAR* Match, FString& Value, bool bShouldStopOnSeparator =true );
+	/** Parses a string. 
+	 * @param Stream, the string you want to extract the value from.
+	 * @param Match, the identifier for the value in the stream.
+	 * @param Value, the destination to the value to be extracted to.
+	 * @param bShouldStopOnSeparator, (default = false) If this is true, and the value doesn't start with a '"'
+				then it may be truncated to ',' or ')' in addition to whitespace.
+	 * @param OptStreamGotTo, (default = nullptr) If this is not null, then its dereference is set to the address
+				of the end of the value within Stream. This permits consuming of stream in a loop where Match may
+				occur multiple times.
+	*/
+	static bool Value( const TCHAR* Stream, const TCHAR* Match, FString& Value, bool bShouldStopOnSeparator =true, const TCHAR** OptStreamGotTo = nullptr);
 	/** Parses an FText. */
 	static bool Value( const TCHAR* Stream, const TCHAR* Match, FText& Value, const TCHAR* Namespace = NULL );
 	/** Parses a quadword. */
diff --git a/Engine/Source/Runtime/Engine/Private/DeviceProfiles/DeviceProfileManager.cpp b/Engine/Source/Runtime/Engine/Private/DeviceProfiles/DeviceProfileManager.cpp
--- a/Engine/Source/Runtime/Engine/Private/DeviceProfiles/DeviceProfileManager.cpp
+++ b/Engine/Source/Runtime/Engine/Private/DeviceProfiles/DeviceProfileManager.cpp
@@ -472,22 +472,27 @@ void UDeviceProfileManager::SetDeviceProfileCVars(const FString& DeviceProfileNa
 	}
 #endif
 	// pre-apply any -dpcvars= items, so that they override anything in the DPs
-	FString DPCVarString;
-	if (FParse::Value(FCommandLine::Get(), TEXT("DPCVars="), DPCVarString, false) || FParse::Value(FCommandLine::Get(), TEXT("DPCVar="), DPCVarString, false))
+	// Search for all occurrences of dpcvars and dpcvar on the command line.
+	static const TCHAR* DPCVarTags[]{ TEXT("DPCVars="), TEXT("DPCVar=") };
+	for (const TCHAR* Tag : DPCVarTags)
 	{
-		// look over a list of cvars
-		TArray<FString> DPCVars;
-		DPCVarString.ParseIntoArray(DPCVars, TEXT(","), true);
-		for (const FString& DPCVar : DPCVars)
+		FString DPCVarString;
+		for (const TCHAR* Cursor = FCommandLine::Get(); (Cursor != nullptr) && FParse::Value(Cursor, Tag, DPCVarString, false, &Cursor);)
 		{
-			// split up each Key=Value pair
-			FString CVarKey, CVarValue;
-			if (DPCVar.Split(TEXT("="), &CVarKey, &CVarValue))
+			// look over a list of cvars
+			TArray<FString> DPCVars;
+			DPCVarString.ParseIntoArray(DPCVars, TEXT(","), true);
+			for (const FString& DPCVar : DPCVars)
 			{
-				UE_LOG(LogDeviceProfileManager, Log, TEXT("Setting CommandLine Device Profile CVar: [[%s:%s]]"), *CVarKey, *CVarValue);
+				// split up each Key=Value pair
+				FString CVarKey, CVarValue;
+				if (DPCVar.Split(TEXT("="), &CVarKey, &CVarValue))
+				{
+					UE_LOG(LogDeviceProfileManager, Log, TEXT("Setting CommandLine Device Profile CVar: [[%s:%s]]"), *CVarKey, *CVarValue);
 
-				// set it and remember it (no thanks, Ron Popeil)
-				UE::ConfigUtilities::OnSetCVarFromIniEntry(*GDeviceProfilesIni, *CVarKey, *CVarValue, ECVF_SetByDeviceProfile);
+					// set it and remember it (no thanks, Ron Popeil)
+					UE::ConfigUtilities::OnSetCVarFromIniEntry(*GDeviceProfilesIni, *CVarKey, *CVarValue, ECVF_SetByDeviceProfile);
+				}
 			}
 		}
 	}
