I/O - parallel HTTP requests

- removed background thread in favor using Tasks
- using fixed set of task pipes to throttle HTTP requests
- updated chunk encoding API to use FIoStatus and FOffsetAndLength

diff --git a/Engine/Source/Runtime/Core/Internal/IO/IoOffsetLength.h b/Engine/Source/Runtime/Core/Internal/IO/IoOffsetLength.h
--- a/Engine/Source/Runtime/Core/Internal/IO/IoOffsetLength.h
+++ b/Engine/Source/Runtime/Core/Internal/IO/IoOffsetLength.h
@@ -2,12 +2,21 @@
 
 #pragma once
 
+#include "CoreTypes.h"
+
 /**
  * Combined offset and length.
  */
 struct FIoOffsetAndLength
 {
-public:
+	FIoOffsetAndLength() = default;
+
+	explicit FIoOffsetAndLength(const uint64 Offset, const uint64 Length)
+	{
+		SetOffset(Offset);
+		SetLength(Length);
+	}
+
 	inline uint64 GetOffset() const
 	{
 		return OffsetAndLength[4]
diff --git a/Engine/Source/Runtime/Core/Private/IO/IoChunkEncoding.cpp b/Engine/Source/Runtime/Core/Private/IO/IoChunkEncoding.cpp
--- a/Engine/Source/Runtime/Core/Private/IO/IoChunkEncoding.cpp
+++ b/Engine/Source/Runtime/Core/Private/IO/IoChunkEncoding.cpp
@@ -10,8 +10,7 @@
 #include "Misc/Compression.h"
 #include <atomic>
 
-///////////////////////////////////////////////////////////////////////////////////////////////////////////////////////
-
+////////////////////////////////////////////////////////////////////////////////
 bool FIoChunkEncoding::FHeader::IsValid() const
 {
 	return Magic == FIoChunkEncoding::ExpectedMagic
@@ -54,8 +53,7 @@ const FIoChunkEncoding::FHeader* FIoChunkEncoding::FHeader::Decode(FMemoryView H
 	return !Header->IsValid() || HeaderData.GetSize() < Header->GetTotalHeaderSize() ? nullptr : Header;
 }
 
-///////////////////////////////////////////////////////////////////////////////////////////////////////////////////////
-
+////////////////////////////////////////////////////////////////////////////////
 bool FIoChunkEncoding::Encode(const FIoChunkEncodingParams& Params, FMemoryView RawData, FIoBuffer& OutEncodedData)
 {
 	FIoBuffer Encoding, Blocks;
@@ -202,8 +200,7 @@ bool FIoChunkEncoding::Encode(const FIoChunkEncodingParams& Params, FMemoryView
 	return true;
 }
 
-///////////////////////////////////////////////////////////////////////////////////////////////////////////////////////
-
+////////////////////////////////////////////////////////////////////////////////
 bool FIoChunkEncoding::Decode(
 	FMemoryView EncodedData,
 	FName CompressionFormat,
@@ -313,23 +310,15 @@ bool FIoChunkEncoding::Decode(const FIoChunkDecodingParams& Params, FMemoryView
 	return true;
 }
 
-///////////////////////////////////////////////////////////////////////////////////////////////////////////////////////
-
-bool FIoChunkEncoding::GetEncodedRange(
-	uint64 TotalRawSize,
-	uint32 RawBlockSize,
-	TConstArrayView<uint32> EncodedBlockSize,
-	uint64 RawOffset,
-	uint64 RawSize,
-	uint64& OutEncodedStart,
-	uint64& OutEncodedEnd)
+////////////////////////////////////////////////////////////////////////////////
+TIoStatusOr<FIoOffsetAndLength> FIoChunkEncoding::GetChunkRange(uint64 TotalRawSize, uint32 RawBlockSize, TConstArrayView<uint32> EncodedBlockSize, uint64 RawOffset, uint64 RawSize)
 {
 	check(RawBlockSize > 0);
 	check(EncodedBlockSize.Num() > 0);
 
 	if (TotalRawSize < RawOffset + RawSize)
 	{
-		return false;
+		return FIoStatus(EIoErrorCode::InvalidParameter);
 	}
 
 	const uint64 TotalBlockCount = FMath::DivideAndRoundUp(TotalRawSize, uint64(RawBlockSize));
@@ -338,28 +327,28 @@ bool FIoChunkEncoding::GetEncodedRange(
 
 	if (TotalBlockCount < FirstBlockIndex || TotalBlockCount < LastBlockIndex)
 	{
-		return false;
+		return FIoStatus(EIoErrorCode::InvalidParameter);
 	}
 
-	OutEncodedStart = 0;
+	uint64 Offset = 0;
 	uint32 BlockIndex = 0;
 	while (BlockIndex < FirstBlockIndex)
 	{
-		OutEncodedStart += Align(EncodedBlockSize[BlockIndex++], FAES::AESBlockSize);
+		Offset += Align(EncodedBlockSize[BlockIndex++], FAES::AESBlockSize);
 	}
 
-	OutEncodedEnd = OutEncodedStart;
+	uint64 Length = 0;
 	while (BlockIndex <= LastBlockIndex)
 	{
-		OutEncodedEnd += Align(EncodedBlockSize[BlockIndex++], FAES::AESBlockSize);
+		Length += Align(EncodedBlockSize[BlockIndex++], FAES::AESBlockSize);
 	}
 
-	return true;
+	return FIoOffsetAndLength(Offset, Length);
 }
 
-bool FIoChunkEncoding::GetEncodedRange(const FIoChunkDecodingParams& Params, uint64 RawSize, uint64& OutEncodedStart, uint64& OutEncodedEnd)
+TIoStatusOr<FIoOffsetAndLength> FIoChunkEncoding::GetChunkRange(const FIoChunkDecodingParams& Params, uint64 RawSize)
 {
-	return GetEncodedRange(Params.TotalRawSize, Params.BlockSize, Params.EncodedBlockSize, Params.RawOffset, RawSize, OutEncodedStart, OutEncodedEnd);
+	return GetChunkRange(Params.TotalRawSize, Params.BlockSize, Params.EncodedBlockSize, Params.RawOffset, RawSize);
 }
 
 uint64 FIoChunkEncoding::GetTotalEncodedSize(TConstArrayView<uint32> EncodedBlockSize)
diff --git a/Engine/Source/Runtime/Core/Public/IO/IoChunkEncoding.h b/Engine/Source/Runtime/Core/Public/IO/IoChunkEncoding.h
--- a/Engine/Source/Runtime/Core/Public/IO/IoChunkEncoding.h
+++ b/Engine/Source/Runtime/Core/Public/IO/IoChunkEncoding.h
@@ -2,6 +2,8 @@
 
 #pragma once
 
+#include "IO/IoOffsetLength.h"
+#include "IO/IoStatus.h"
 #include "Memory/MemoryView.h"
 #include "UObject/NameTypes.h"
 
@@ -71,7 +73,7 @@ public:
 	static bool Decode(const FIoChunkDecodingParams& Params, FMemoryView EncodedBlocks, FMutableMemoryView OutRawData);
 	static bool Decode(FMemoryView EncodedData, FName CompressionFormat, FMemoryView EncryptionKey, FMutableMemoryView OutRawData, uint64 Offset = 0);
 
-	static bool GetEncodedRange(uint64 TotalRawSize, uint32 RawBlockSize, TConstArrayView<uint32> EncodedBlockSize, uint64 RawOffset, uint64 RawSize, uint64& OutEncodedStart, uint64& OutEncodedEnd);
-	static bool GetEncodedRange(const FIoChunkDecodingParams& Params, uint64 RawSize, uint64& OutEncodedStart, uint64& OutEncodedEnd);
+	static TIoStatusOr<FIoOffsetAndLength> GetChunkRange(uint64 TotalRawSize, uint32 RawBlockSize, TConstArrayView<uint32> EncodedBlockSize, uint64 RawOffset, uint64 RawSize);
+	static TIoStatusOr<FIoOffsetAndLength> GetChunkRange(const FIoChunkDecodingParams& Params, uint64 RawSize);
 	static uint64 GetTotalEncodedSize(TConstArrayView<uint32> EncodedBlockSize);
 };
