Change deprecated FCoreDelegates::OnConfigSectionsChanged to FCoreDelegates::TSOnConfigSectionsChanged() after merge

diff --git a/Engine/Plugins/Online/OnlineSubsystem/Source/Private/OnlineSubsystemImpl.cpp b/Engine/Plugins/Online/OnlineSubsystem/Source/Private/OnlineSubsystemImpl.cpp
--- a/Engine/Plugins/Online/OnlineSubsystem/Source/Private/OnlineSubsystemImpl.cpp
+++ b/Engine/Plugins/Online/OnlineSubsystem/Source/Private/OnlineSubsystemImpl.cpp
@@ -45,12 +45,12 @@ FOnlineSubsystemImpl::FOnlineSubsystemImpl(FName InSubsystemName, FName InInstan
 	NamedInterfaces(nullptr),
 	bTickerStarted(true)
 {
-	FCoreDelegates::OnConfigSectionsChanged.AddRaw(this, &FOnlineSubsystemImpl::OnConfigSectionsChanged);
+	FCoreDelegates::TSOnConfigSectionsChanged().AddRaw(this, &FOnlineSubsystemImpl::OnConfigSectionsChanged);
 }
 
 FOnlineSubsystemImpl::~FOnlineSubsystemImpl()
 {
-	FCoreDelegates::OnConfigSectionsChanged.RemoveAll(this);
+	FCoreDelegates::TSOnConfigSectionsChanged().RemoveAll(this);
 }
 
 void FOnlineSubsystemImpl::PreUnload()
diff --git a/Engine/Source/Runtime/Core/Private/Logging/LogSuppressionInterface.cpp b/Engine/Source/Runtime/Core/Private/Logging/LogSuppressionInterface.cpp
--- a/Engine/Source/Runtime/Core/Private/Logging/LogSuppressionInterface.cpp
+++ b/Engine/Source/Runtime/Core/Private/Logging/LogSuppressionInterface.cpp
@@ -565,7 +565,7 @@ class FLogSuppressionImplementation: public FLogSuppressionInterface, private FS
 
 		if (!bInitialized)
 		{
-			FCoreDelegates::OnConfigSectionsChanged.AddLambda([this](const FString& IniFilename, const TSet<FString>& SectionNames)
+			FCoreDelegates::TSOnConfigSectionsChanged().AddLambda([this](const FString& IniFilename, const TSet<FString>& SectionNames)
 			{
 				if (auto* LogSuppressionImplementation = static_cast<FLogSuppressionImplementation*>(FLogSuppressionInterface::TryGet()))
 				{
@@ -573,7 +573,7 @@ class FLogSuppressionImplementation: public FLogSuppressionInterface, private FS
 				}
 				else
 				{
-					FCoreDelegates::OnConfigSectionsChanged.RemoveAll(this);
+					FCoreDelegates::TSOnConfigSectionsChanged().RemoveAll(this);
 				}
 			});
 		}
diff --git a/Engine/Source/Runtime/Core/Private/Misc/ConfigCacheIni.cpp b/Engine/Source/Runtime/Core/Private/Misc/ConfigCacheIni.cpp
--- a/Engine/Source/Runtime/Core/Private/Misc/ConfigCacheIni.cpp
+++ b/Engine/Source/Runtime/Core/Private/Misc/ConfigCacheIni.cpp
@@ -3857,7 +3857,7 @@ void FConfigCacheIni::InitializeConfigSystem()
 	// load editor, etc config files
 	LoadRemainingConfigFiles(Context);
 
-	FCoreDelegates::OnConfigSectionsChanged.AddStatic(OnConfigSectionsChanged);
+	FCoreDelegates::TSOnConfigSectionsChanged().AddStatic(OnConfigSectionsChanged);
 
 	// now we can make use of GConfig
 	GConfig->bIsReadyForUse = true;
diff --git a/Engine/Source/Runtime/Online/HTTP/Private/HttpModule.cpp b/Engine/Source/Runtime/Online/HTTP/Private/HttpModule.cpp
--- a/Engine/Source/Runtime/Online/HTTP/Private/HttpModule.cpp
+++ b/Engine/Source/Runtime/Online/HTTP/Private/HttpModule.cpp
@@ -78,7 +78,7 @@ void FHttpModule::StartupModule()
 	HttpThreadIdleMinimumSleepTimeInSeconds = 0.0f;	
 
 	// override the above defaults from configs
-	FCoreDelegates::OnConfigSectionsChanged.AddRaw(this, &FHttpModule::OnConfigSectionsChanged);
+	FCoreDelegates::TSOnConfigSectionsChanged().AddRaw(this, &FHttpModule::OnConfigSectionsChanged);
 	UpdateConfigs();
 
 	if (!FParse::Value(FCommandLine::Get(), TEXT("httpproxy="), ProxyAddress))
@@ -141,7 +141,7 @@ void FHttpModule::ShutdownModule()
 	// at least on Linux, the code in HTTP manager (e.g. request destructors) expects platform to be initialized yet
 	delete HttpManager;	// can be passed NULLs
 
-	FCoreDelegates::OnConfigSectionsChanged.RemoveAll(this);
+	FCoreDelegates::TSOnConfigSectionsChanged().RemoveAll(this);
 
 	FPlatformHttp::Shutdown();
 
