LowLevelMemTracker support to capture a snapshot of memory and report deltas to CSV/FPS.
Exec command to REMEMBER (which will take the snapshot), and FORGET to clear the snapshot information.

diff --git a/Engine/Source/Runtime/Core/Private/HAL/LowLevelMemTracker.cpp b/Engine/Source/Runtime/Core/Private/HAL/LowLevelMemTracker.cpp
--- a/Engine/Source/Runtime/Core/Private/HAL/LowLevelMemTracker.cpp
+++ b/Engine/Source/Runtime/Core/Private/HAL/LowLevelMemTracker.cpp
@@ -307,7 +307,7 @@ public:
 
 	void Publish(FLowLevelMemTracker& LLMRef, const FTrackerTagSizeMap& TagSizes,
 		const FTagData* OverrideTrackedTotalTagData, const FTagData* OverrideUntaggedTagData, int64 TrackedTotal,
-		bool bTrackPeaks);
+		ESizeParams SizeParams);
 
 	void OnPreFork();
 
@@ -320,7 +320,7 @@ private:
 	void WriteHeader(const FTagData* OverrideTrackedTotalTagData, const FTagData* OverrideUntaggedTagData);
 	void AddRow(FLowLevelMemTracker& LLMRef, const FTrackerTagSizeMap& TagSizes,
 		const FTagData* OverrideTrackedTotalName, const FTagData* OverrideUntaggedName, int64 TrackedTotal,
-		bool bTrackPeaks);
+		ESizeParams SizeParams);
 
 	FConstTagDataArray Columns;
 	TFastPointerLLMSet<const FTagData*> ExistingColumns;
@@ -339,7 +339,7 @@ public:
 	void Clear();
 	void Publish(FLowLevelMemTracker& LLMRef, const FTrackerTagSizeMap& TagSizes,
 		const FTagData* OverrideTrackedTotalName, const FTagData* OverrideUntaggedName, int64 TrackedTotal,
-		bool bTrackPeaks);
+		ESizeParams SizeParams);
 	static void TraceGenericTags(FLowLevelMemTracker& LLMRef);
 
 private:
@@ -359,7 +359,7 @@ public:
 	void SetTracker(ELLMTracker InTracker);
 	void Publish(FLowLevelMemTracker& LLMRef, const FTrackerTagSizeMap& TagSizes,
 		const FTagData* OverrideTrackedTotalName, const FTagData* OverrideUntaggedName, int64 TrackedTotal,
-		bool bTrackPeaks);
+		ESizeParams SizeParams);
 protected:
 	void RecordTagToCsv(int32 CsvCategoryIndex, const FTagData* TagData, int64 TagSize);
 #if LLM_CSV_PROFILER_WRITER_ENABLED 
@@ -441,10 +441,10 @@ public:
 
 	void Clear();
 
-	void PublishStats(bool bTrackPeaks);
-	void PublishCsv(bool bTrackPeaks);
-	void PublishTrace(bool bTrackPeaks);
-	void PublishCsvProfiler(bool bTrackPeaks);
+	void PublishStats(ESizeParams SizeParam);
+	void PublishCsv(ESizeParams SizeParam);
+	void PublishTrace(ESizeParams SizeParam);
+	void PublishCsvProfiler(ESizeParams SizeParam);
 
 	void OnPreFork();
 
@@ -543,6 +543,9 @@ public:
 	void OnTagsResorted(FTagDataArray& OldTagDatas);
 	void LockAllThreadTags(bool bLock);
 
+	void CaptureTagSnapshot();
+	void ClearTagSnapshot();
+
 	int64 GetTagAmount(const FTagData* TagData, bool bPeakAmount = false) const;
 	void SetTagAmountExternal(const FTagData* TagData, int64 Amount, bool bAddToTotal);
 	void SetTagAmountInUpdate(const FTagData* TagData, int64 Amount, bool bAddToTotal);
@@ -556,8 +559,12 @@ public:
 
 	int64 GetAllocTypeAmount(ELLMAllocType AllocType);
 
-	int64 GetTrackedTotal() const
+	int64 GetTrackedTotal(ESizeParams SizeParams = ESizeParams::Default) const
 	{
+		if (EnumHasAnyFlags(SizeParams, ESizeParams::RelativeToSnapshot))
+		{
+			return FMath::Clamp<int64>(TrackedTotal - TrackedTotalInSnapshot, 0, INT64_MAX);
+		}
 		return TrackedTotal;
 	}
 
@@ -578,7 +585,12 @@ protected:
 	FCriticalSection PendingThreadStatesGuard;
 	TArray<FLLMThreadState*, FDefaultLLMAllocator> PendingThreadStates;
 
+	/** Sum of memory from all tracked tags. Duplicated in separate storage to make it instantly available at any time of frame without waiting for accumulation from threads during update.
+	    GCC_ALIGN is required because it is modified in FPlatformAtomics::InterlockedAdd, which requires aligned values. */
 	int64 TrackedTotal GCC_ALIGN(8);
+	
+	// The total tracked memory when the snapshot was taken
+	int64 TrackedTotalInSnapshot;
 
 	FLLMAllocMap AllocationMap;
 
@@ -759,6 +771,7 @@ FLowLevelMemTracker::FLowLevelMemTracker()
 	, bTagAdded(false)
 	, bAutoPublish(true)
 	, bPublishSingleFrame(false)
+	, bCapturedSizeSnapshot(false)
 {
 	using namespace UE::LLMPrivate;
 
@@ -1082,6 +1095,18 @@ void FLowLevelMemTracker::PublishDataPerFrame(const TCHAR* LogName)
 		FLLMTracker& PlatformTracker = *GetTracker(ELLMTracker::Platform);
 
 		bool bTrackPeaks = CVarLLMTrackPeaks.GetValueOnAnyThread() != 0;
+
+		ESizeParams SizeParams(ESizeParams::Default);
+		if (bTrackPeaks)
+		{
+			EnumAddFlags(SizeParams, ESizeParams::ReportPeak);
+		}
+
+		if (bCapturedSizeSnapshot)
+		{
+			EnumAddFlags(SizeParams, ESizeParams::RelativeToSnapshot);
+		}
+
 #if !LLM_ENABLED_TRACK_PEAK_MEMORY
 		if (bTrackPeaks)
 		{
@@ -1095,26 +1120,26 @@ void FLowLevelMemTracker::PublishDataPerFrame(const TCHAR* LogName)
 		}
 #endif
 
-		DefaultTracker.PublishStats(bTrackPeaks);
-		PlatformTracker.PublishStats(bTrackPeaks);
+		DefaultTracker.PublishStats(SizeParams);
+		PlatformTracker.PublishStats(SizeParams);
 
 		if (bCsvWriterEnabled)
 		{
-			DefaultTracker.PublishCsv(bTrackPeaks);
-			PlatformTracker.PublishCsv(bTrackPeaks);
+			DefaultTracker.PublishCsv(SizeParams);
+			PlatformTracker.PublishCsv(SizeParams);
 		}
 
 		if (bTraceWriterEnabled)
 		{
-			DefaultTracker.PublishTrace(bTrackPeaks);
-			PlatformTracker.PublishTrace(bTrackPeaks);
+			DefaultTracker.PublishTrace(SizeParams);
+			PlatformTracker.PublishTrace(SizeParams);
 		}
 
 #if LLM_CSV_PROFILER_WRITER_ENABLED
 		if (FCsvProfiler::Get()->IsCapturing())
 		{
-			DefaultTracker.PublishCsvProfiler(bTrackPeaks);
-			PlatformTracker.PublishCsvProfiler(bTrackPeaks);
+			DefaultTracker.PublishCsvProfiler(SizeParams);
+			PlatformTracker.PublishCsvProfiler(SizeParams);
 		}
 #endif
 	}
@@ -1396,6 +1421,26 @@ bool FLowLevelMemTracker::Exec(const TCHAR* Cmd, FOutputDevice& Ar)
 
 			UpdateStatsPerFrame(TEXT("After cleanup"));
 		}
+		else if (FParse::Command(&Cmd, TEXT("REMEMBER")))
+		{
+			FScopeLock UpdateScopeLock(&UpdateLock);
+			// Go through all the tags and REMEMBER the current recorded size values
+			for (int32 TrackerIndex = 0; TrackerIndex < static_cast<int32>(ELLMTracker::Max); TrackerIndex++)
+			{
+				GetTracker((ELLMTracker)TrackerIndex)->CaptureTagSnapshot();
+			}
+			bCapturedSizeSnapshot = true;
+		}
+		else if (FParse::Command(&Cmd, TEXT("FORGET")))
+		{
+			FScopeLock UpdateScopeLock(&UpdateLock);
+			// Go through all tags and make them forget their current size values
+			for (int32 TrackerIndex = 0; TrackerIndex < static_cast<int32>(ELLMTracker::Max); TrackerIndex++)
+			{
+				GetTracker((ELLMTracker)TrackerIndex)->ClearTagSnapshot();
+			}
+			bCapturedSizeSnapshot = false;
+		}
 		return true;
 	}
 
@@ -3652,6 +3697,35 @@ void FLLMTracker::SetTotalTags(const FTagData* InOverrideUntaggedTagData,
 	OverrideTrackedTotalTagData = InOverrideTrackedTotalTagData;
 }
 
+void FLLMTracker::CaptureTagSnapshot()
+{
+	UpdateThreads();
+	for (FLLMThreadState* ThreadState : ThreadStates)
+	{
+		ThreadState->PropagateChildSizesToParents();
+		ThreadState->FetchAndClearTagSizes(TagSizes, AllocTypeAmounts, false);
+	}
+
+	for (TPair<const FTagData*, FTrackerTagSizeData>& It : TagSizes)
+	{
+		FTrackerTagSizeData& TrackerInfo = It.Value;
+		TrackerInfo.CaptureSnapshot();
+	}
+
+	TrackedTotalInSnapshot = TrackedTotal;
+}
+
+void FLLMTracker::ClearTagSnapshot()
+{
+	for (TPair<const FTagData*, FTrackerTagSizeData>& It : TagSizes)
+	{
+		FTrackerTagSizeData& TrackerInfo = It.Value;
+		TrackerInfo.ClearSnapshot();
+	}
+
+	TrackedTotalInSnapshot = 0;
+}
+
 void FLLMTracker::Update()
 {
 	UpdateThreads();
@@ -3722,7 +3796,7 @@ void FLLMTracker::UpdateThreads()
 	}
 }
 
-void FLLMTracker::PublishStats(bool bTrackPeaks)
+void FLLMTracker::PublishStats(ESizeParams SizeParams)
 {
 	if (OverrideTrackedTotalTagData)
 	{
@@ -3735,9 +3809,9 @@ void FLLMTracker::PublishStats(bool bTrackPeaks)
 		const FTagData* TagData = LLMRef.FindTagData(TagName_Untagged, ELLMTagSet::None);
 		const FTrackerTagSizeData* AllocationData = TagData ? TagSizes.Find(TagData) : nullptr;
 		SetMemoryStatByFName(OverrideUntaggedTagData->GetStatName(),
-			AllocationData ? AllocationData->GetSize(bTrackPeaks) : 0);
+			AllocationData ? AllocationData->GetSize(SizeParams) : 0);
 		SetMemoryStatByFName(OverrideUntaggedTagData->GetSummaryStatName(),
-			AllocationData ? AllocationData->GetSize(bTrackPeaks) : 0);
+			AllocationData ? AllocationData->GetSize(SizeParams) : 0);
 	}
 
 	for (const TPair<const FTagData*, FTrackerTagSizeData>& It : TagSizes)
@@ -3753,29 +3827,29 @@ void FLLMTracker::PublishStats(bool bTrackPeaks)
 			continue;
 		}
 		const FTrackerTagSizeData& AllocationData = It.Value;
-		int64 Amount = AllocationData.GetSize(bTrackPeaks);
+		int64 Amount = AllocationData.GetSize(SizeParams);
 
 		SetMemoryStatByFName(TagData->GetStatName(), Amount);
 		SetMemoryStatByFName(TagData->GetSummaryStatName(), Amount);
 	}
 }
 
-void FLLMTracker::PublishCsv(bool bTrackPeaks)
+void FLLMTracker::PublishCsv(ESizeParams SizeParams)
 {
-	CsvWriter.Publish(LLMRef, TagSizes, OverrideTrackedTotalTagData, OverrideUntaggedTagData, TrackedTotal,
-		bTrackPeaks);
+	CsvWriter.Publish(LLMRef, TagSizes, OverrideTrackedTotalTagData, OverrideUntaggedTagData, GetTrackedTotal(SizeParams), SizeParams);
 }
 
-void FLLMTracker::PublishTrace(bool bTrackPeaks)
+void FLLMTracker::PublishTrace(ESizeParams SizeParams)
 {
-	TraceWriter.Publish(LLMRef, TagSizes, OverrideTrackedTotalTagData, OverrideUntaggedTagData, TrackedTotal,
-		bTrackPeaks);
+	// Trace tracker does not support the handling of Snapshot currently
+	EnumRemoveFlags(SizeParams, ESizeParams::RelativeToSnapshot);
+
+	TraceWriter.Publish(LLMRef, TagSizes, OverrideTrackedTotalTagData, OverrideUntaggedTagData, GetTrackedTotal(SizeParams), SizeParams);
 }
 
-void FLLMTracker::PublishCsvProfiler(bool bTrackPeaks)
+void FLLMTracker::PublishCsvProfiler(ESizeParams SizeParams)
 {
-	CsvProfilerWriter.Publish(LLMRef, TagSizes, OverrideTrackedTotalTagData, OverrideUntaggedTagData, TrackedTotal,
-		bTrackPeaks);
+	CsvProfilerWriter.Publish(LLMRef, TagSizes, OverrideTrackedTotalTagData, OverrideUntaggedTagData, GetTrackedTotal(SizeParams), SizeParams);
 }
 
 void FLLMTracker::OnTagsResorted(FTagDataArray& OldTagDatas)
@@ -3851,7 +3925,7 @@ void FLLMTracker::GetTagsNamesWithAmount(TMap<FName, uint64>& OutTagsNamesWithAm
 		const FTagData* TagData = It.Key;
 		if (TagData->GetTagSet() == TagSet)
 		{
-			OutTagsNamesWithAmount.Add(TagData->GetName(), It.Value.GetSize(false));
+			OutTagsNamesWithAmount.Add(TagData->GetName(), It.Value.GetSize(ESizeParams::Default));
 		}
 	}
 }
@@ -3923,7 +3997,7 @@ int64 FLLMTracker::GetTagAmount(const FTagData* TagData, bool bPeakAmount/* = fa
 	const FTrackerTagSizeData* AllocationData = TagSizes.Find(TagData);
 	if (AllocationData)
 	{
-		return AllocationData->GetSize(bPeakAmount);
+		return AllocationData->GetSize((bPeakAmount ? ESizeParams::ReportPeak : ESizeParams::Default));
 	}
 	else
 	{
@@ -4243,7 +4317,7 @@ void FLLMCsvWriter::SetTracker(ELLMTracker InTracker)
 
 void FLLMCsvWriter::Publish(FLowLevelMemTracker& LLMRef, const FTrackerTagSizeMap& TagSizes,
 	const FTagData* OverrideTrackedTotalTagData, const FTagData* OverrideUntaggedTagData, int64 TrackedTotal,
-	bool bTrackPeaks)
+	ESizeParams SizeParams)
 {
 	const double Now = FPlatformTime::Seconds();
 
@@ -4265,7 +4339,7 @@ void FLLMCsvWriter::Publish(FLowLevelMemTracker& LLMRef, const FTrackerTagSizeMa
 		WriteHeader(OverrideTrackedTotalTagData, OverrideUntaggedTagData);
 	}
 
-	AddRow(LLMRef, TagSizes, OverrideTrackedTotalTagData, OverrideUntaggedTagData, TrackedTotal, bTrackPeaks);
+	AddRow(LLMRef, TagSizes, OverrideTrackedTotalTagData, OverrideUntaggedTagData, TrackedTotal, SizeParams);
 }
 
 const TCHAR* FLLMCsvWriter::GetTrackerCsvName(ELLMTracker InTracker)
@@ -4388,7 +4462,7 @@ void FLLMCsvWriter::WriteHeader(const FTagData* OverrideTrackedTotalTagData, con
 
 void FLLMCsvWriter::AddRow(FLowLevelMemTracker& LLMRef, const FTrackerTagSizeMap& TagSizes,
 	const FTagData* OverrideTrackedTotalTagData, const FTagData* OverrideUntaggedTagData, int64 TrackedTotal,
-	bool bTrackPeaks)
+	ESizeParams SizeParams)
 {
 	TStringBuilder<256> TextBuffer;
 	auto WriteValue = [this, &TextBuffer](int64 Value)
@@ -4397,7 +4471,7 @@ void FLLMCsvWriter::AddRow(FLowLevelMemTracker& LLMRef, const FTrackerTagSizeMap
 		TextBuffer.Appendf(TEXT("%0.2f,"), (float)Value / 1024.0f / 1024.0f);
 		Write(TextBuffer);
 	};
-	auto WriteTag = [&WriteValue, &TagSizes, bTrackPeaks](const FTagData* TagData)
+	auto WriteTag = [&WriteValue, &TagSizes, SizeParams](const FTagData* TagData)
 	{
 		if (!TagData)
 		{
@@ -4412,7 +4486,7 @@ void FLLMCsvWriter::AddRow(FLowLevelMemTracker& LLMRef, const FTrackerTagSizeMap
 			}
 			else
 			{
-				WriteValue(AllocationData->GetSize(bTrackPeaks));
+				WriteValue(AllocationData->GetSize(SizeParams));
 			}
 		}
 	};
@@ -4468,7 +4542,7 @@ const void* FLLMTraceWriter::GetTagId(const FTagData* TagData)
 
 void FLLMTraceWriter::Publish(FLowLevelMemTracker& LLMRef, const FTrackerTagSizeMap& TagSizes,
 	const FTagData* OverrideTrackedTotalTagData, const FTagData* OverrideUntaggedTagData, int64 TrackedTotal,
-	bool bTrackPeaks)
+	ESizeParams SizeParams)
 {
 	if (!UE_TRACE_CHANNELEXPR_IS_ENABLED(MemTagChannel))
 	{
@@ -4550,7 +4624,7 @@ void FLLMTraceWriter::Publish(FLowLevelMemTracker& LLMRef, const FTrackerTagSize
 		else
 		{
 			const FTrackerTagSizeData* AllocationData = TagSizes.Find(TagData);
-			AddValue(OverrideUntaggedTagData, AllocationData ? AllocationData->GetSize(bTrackPeaks) : 0);
+			AddValue(OverrideUntaggedTagData, AllocationData ? AllocationData->GetSize(SizeParams) : 0);
 		}
 	}
 
@@ -4565,7 +4639,7 @@ void FLLMTraceWriter::Publish(FLowLevelMemTracker& LLMRef, const FTrackerTagSize
 		{
 			continue; // Handled by OverrideUntaggedTagData.
 		}
-		AddValue(TagData, It.Value.GetSize(bTrackPeaks));
+		AddValue(TagData, It.Value.GetSize(SizeParams));
 	}
 
 	int TagCount = TagIds.Num();
@@ -4613,7 +4687,7 @@ void FLLMCsvProfilerWriter::SetTracker(ELLMTracker InTracker)
 
 void FLLMCsvProfilerWriter::Publish(FLowLevelMemTracker& LLMRef, const FTrackerTagSizeMap& TagSizes,
 	const FTagData* OverrideTrackedTotalTagData, const FTagData* OverrideUntaggedTagData, int64 TrackedTotal,
-	bool bTrackPeaks)
+	ESizeParams SizeParams)
 {
 #if LLM_CSV_PROFILER_WRITER_ENABLED
 	int32 CsvCategoryIndex = (Tracker == ELLMTracker::Platform) ?
@@ -4629,7 +4703,7 @@ void FLLMCsvProfilerWriter::Publish(FLowLevelMemTracker& LLMRef, const FTrackerT
 		const FTagData* TagData = LLMRef.FindTagData(TagName_Untagged, ELLMTagSet::None);
 		const FTrackerTagSizeData* AllocationData = TagData ? TagSizes.Find(TagData) : nullptr;
 		RecordTagToCsv(CsvCategoryIndex, OverrideUntaggedTagData,
-			AllocationData ? AllocationData->GetSize(bTrackPeaks) : 0);
+			AllocationData ? AllocationData->GetSize(SizeParams) : 0);
 	}
 
 	for (const TPair<const FTagData*, FTrackerTagSizeData>& It : TagSizes)
@@ -4645,7 +4719,7 @@ void FLLMCsvProfilerWriter::Publish(FLowLevelMemTracker& LLMRef, const FTrackerT
 			continue;
 		}
 		const FTrackerTagSizeData* AllocationData = TagSizes.Find(TagData);
-		RecordTagToCsv(CsvCategoryIndex, TagData, AllocationData ? AllocationData->GetSize(bTrackPeaks) : 0);
+		RecordTagToCsv(CsvCategoryIndex, TagData, AllocationData ? AllocationData->GetSize(SizeParams) : 0);
 	}
 #endif
 }
diff --git a/Engine/Source/Runtime/Core/Private/HAL/LowLevelMemTrackerPrivate.h b/Engine/Source/Runtime/Core/Private/HAL/LowLevelMemTrackerPrivate.h
--- a/Engine/Source/Runtime/Core/Private/HAL/LowLevelMemTrackerPrivate.h
+++ b/Engine/Source/Runtime/Core/Private/HAL/LowLevelMemTrackerPrivate.h
@@ -44,8 +44,6 @@
 // Disable if you need a little more memory or speed
 #define LLM_ENABLED_TRACK_PEAK_MEMORY 1
 
-
-
 namespace UE
 {
 namespace LLMPrivate
@@ -175,6 +173,15 @@ namespace UE
 {
 namespace LLMPrivate
 {
+	enum class ESizeParams : uint8
+	{
+		Default = 0,
+		ReportCurrent = 0,
+		ReportPeak = 1,
+		RelativeToSnapshot = 2
+	};
+	ENUM_CLASS_FLAGS(ESizeParams);
+
 	/** Size information stored on the tracker for a tag; includes amounts aggregated from threadstates and from external api users */
 	struct FTrackerTagSizeData
 	{
@@ -182,23 +189,39 @@ namespace LLMPrivate
 #if LLM_ENABLED_TRACK_PEAK_MEMORY
 		int64 PeakSize = 0;
 #endif
+		int64 SizeInSnapshot = 0;
 		int64 ExternalAmount = 0;
 		bool bExternalValid = false;
 		bool bExternalAddToTotal = false;
 
-		int64 GetSize(bool bTrackPeaks) const
+		int64 GetSize(ESizeParams SizeParams) const
 		{
+			int64 CurrentSize = Size;
+
 #if LLM_ENABLED_TRACK_PEAK_MEMORY
-			if (bTrackPeaks)
+			if (EnumHasAnyFlags(SizeParams, ESizeParams::ReportPeak))
 			{
-				return PeakSize;
+				CurrentSize = PeakSize;
 			}
-			else
 #endif
+			// Note, this will also subtract the snapshotted size from PeakSize if that flag is enabled
+			if (EnumHasAnyFlags(SizeParams, ESizeParams::RelativeToSnapshot))
 			{
-				return Size;
+				CurrentSize = FMath::Clamp<int64>(CurrentSize - SizeInSnapshot, 0, INT64_MAX);
 			}
+
+			return CurrentSize;
 		};
+
+		void CaptureSnapshot()
+		{
+			SizeInSnapshot = Size;
+		}
+
+		void ClearSnapshot()
+		{
+			SizeInSnapshot = 0;
+		}
 	};
 	typedef TFastPointerLLMMap<const FTagData*, FTrackerTagSizeData> FTrackerTagSizeMap;
 
diff --git a/Engine/Source/Runtime/Core/Public/HAL/LowLevelMemTracker.h b/Engine/Source/Runtime/Core/Public/HAL/LowLevelMemTracker.h
--- a/Engine/Source/Runtime/Core/Public/HAL/LowLevelMemTracker.h
+++ b/Engine/Source/Runtime/Core/Public/HAL/LowLevelMemTracker.h
@@ -748,6 +748,7 @@ private:
 	bool bTagAdded;
 	bool bAutoPublish;
 	bool bPublishSingleFrame;
+	bool bCapturedSizeSnapshot;
 
 	static FLowLevelMemTracker* TrackerInstance;
 public: // really internal but needs to be visible for LLM_IF_ENABLED macro
