Adding CSV Event for when LLM Snapshot is captured and cleared. Makes it easier when viewing the charts.

diff --git a/Engine/Source/Runtime/Core/Private/HAL/LowLevelMemTracker.cpp b/Engine/Source/Runtime/Core/Private/HAL/LowLevelMemTracker.cpp
--- a/Engine/Source/Runtime/Core/Private/HAL/LowLevelMemTracker.cpp
+++ b/Engine/Source/Runtime/Core/Private/HAL/LowLevelMemTracker.cpp
@@ -1423,6 +1423,8 @@ bool FLowLevelMemTracker::Exec(const TCHAR* Cmd, FOutputDevice& Ar)
 		}
 		else if (FParse::Command(&Cmd, TEXT("REMEMBER")))
 		{
+			CSV_EVENT_GLOBAL(TEXT("LLM_REMEMBER"));
+
 			FScopeLock UpdateScopeLock(&UpdateLock);
 			// Go through all the tags and REMEMBER the current recorded size values
 			for (int32 TrackerIndex = 0; TrackerIndex < static_cast<int32>(ELLMTracker::Max); TrackerIndex++)
@@ -1433,6 +1435,8 @@ bool FLowLevelMemTracker::Exec(const TCHAR* Cmd, FOutputDevice& Ar)
 		}
 		else if (FParse::Command(&Cmd, TEXT("FORGET")))
 		{
+			CSV_EVENT_GLOBAL(TEXT("LLM_FORGET"));
+
 			FScopeLock UpdateScopeLock(&UpdateLock);
 			// Go through all tags and make them forget their current size values
 			for (int32 TrackerIndex = 0; TrackerIndex < static_cast<int32>(ELLMTracker::Max); TrackerIndex++)
