Implement MinimumGroupingDigits for FastDecimalFormat

Enable "min two" grouping for Spanish by default (this can be disabled via the "Localization.SpanishUsesMinTwoGrouping" CVar). When enabled "1234" formats as "1234" rather than "1 234", but "12345" still formats as "12 345"

diff --git a/Engine/Source/Runtime/Core/Private/Internationalization/FastDecimalFormat.cpp b/Engine/Source/Runtime/Core/Private/Internationalization/FastDecimalFormat.cpp
--- a/Engine/Source/Runtime/Core/Private/Internationalization/FastDecimalFormat.cpp
+++ b/Engine/Source/Runtime/Core/Private/Internationalization/FastDecimalFormat.cpp
@@ -205,7 +205,7 @@ void SanitizeNumberFormattingOptions(FNumberFormattingOptions& InOutFormattingOp
 
 int32 IntegralToString_UInt64ToString(
 	const uint64 InVal, 
-	const bool InUseGrouping, const uint8 InPrimaryGroupingSize, const uint8 InSecondaryGroupingSize, const TCHAR InGroupingSeparatorCharacter, const TCHAR* InDigitCharacters, 
+	const bool InUseGrouping, const uint8 InPrimaryGroupingSize, const uint8 InSecondaryGroupingSize, const uint8 InMinimumGroupingDigits, const TCHAR InGroupingSeparatorCharacter, const TCHAR* InDigitCharacters,
 	const int32 InMinDigitsToPrint, const int32 InMaxDigitsToPrint, 
 	TCHAR* InBufferToFill, const int32 InBufferToFillSize
 	)
@@ -249,9 +249,18 @@ int32 IntegralToString_UInt64ToString(
 			}
 
 			TmpBuffer[StringLen++] = InDigitCharacters[0];
+
+			++DigitsPrinted;
 		}
 	}
 
+	// Did we actually write enough digits to warrant a group separator? If not then we need to remove it again
+	if (InUseGrouping && DigitsPrinted > static_cast<int32>(InPrimaryGroupingSize) && DigitsPrinted < static_cast<int32>(InPrimaryGroupingSize + InMinimumGroupingDigits))
+	{
+		FMemory::Memmove(TmpBuffer + InPrimaryGroupingSize, TmpBuffer + InPrimaryGroupingSize + 1, StringLen - InPrimaryGroupingSize - 1);
+		--StringLen;
+	}
+
 	// TmpBuffer is backwards, flip it into the final output buffer
 	for (int32 FinalBufferIndex = 0; FinalBufferIndex < StringLen; ++FinalBufferIndex)
 	{
@@ -270,6 +279,7 @@ FORCEINLINE int32 IntegralToString_Common(const uint64 InVal, const FDecimalNumb
 		InFormattingOptions.UseGrouping && InFormattingRules.PrimaryGroupingSize > 0,
 		InFormattingRules.PrimaryGroupingSize, 
 		InFormattingRules.SecondaryGroupingSize, 
+		InFormattingRules.MinimumGroupingDigits, 
 		InFormattingRules.GroupingSeparatorCharacter, 
 		InFormattingRules.DigitCharacters, 
 		InFormattingOptions.MinimumIntegralDigits, 
@@ -557,7 +567,7 @@ void FractionalToString(const double InVal, const FDecimalNumberFormattingRules&
 	int32 FractionalPartLen = 0;
 	if (FractionalPart != 0.0)
 	{
-		FractionalPartLen = IntegralToString_UInt64ToString(static_cast<uint64>(FractionalPart), false, 0, 0, TEXT(' '), InFormattingRules.DigitCharacters, 0, InFormattingOptions.MaximumFractionalDigits, FractionalPartBuffer, UE_ARRAY_COUNT(FractionalPartBuffer));
+		FractionalPartLen = IntegralToString_UInt64ToString(static_cast<uint64>(FractionalPart), false, 0, 0, 1, TEXT(' '), InFormattingRules.DigitCharacters, 0, InFormattingOptions.MaximumFractionalDigits, FractionalPartBuffer, UE_ARRAY_COUNT(FractionalPartBuffer));
 	
 		{
 			// Pad the fractional part with any leading zeros that may have been lost when the number was split
diff --git a/Engine/Source/Runtime/Core/Private/Internationalization/ICUCulture.cpp b/Engine/Source/Runtime/Core/Private/Internationalization/ICUCulture.cpp
--- a/Engine/Source/Runtime/Core/Private/Internationalization/ICUCulture.cpp
+++ b/Engine/Source/Runtime/Core/Private/Internationalization/ICUCulture.cpp
@@ -10,10 +10,17 @@
 #include "Internationalization/ICUUtilities.h"
 #include "HAL/IConsoleManager.h"
 
-static TAutoConsoleVariable<int32> CVarSpanishUsesRAENumberFormat(
+static TAutoConsoleVariable<bool> CVarSpanishUsesRAENumberFormat(
 	TEXT("Localization.SpanishUsesRAENumberFormat"),
-	1,
-	TEXT("0: Disabled (CLDR format), 1: Enabled (RAE format, default)."),
+	true,
+	TEXT("False: Disabled (CLDR format), True: Enabled (RAE format, default)."),
+	ECVF_Default
+	);
+
+static TAutoConsoleVariable<bool> CVarSpanishUsesMinTwoGrouping(
+	TEXT("Localization.SpanishUsesMinTwoGrouping"),
+	true,
+	TEXT("False: 1234 will use a group separator, True: 1234 will not use a group separator (default)."),
 	ECVF_Default
 	);
 
@@ -675,6 +682,9 @@ FDecimalNumberFormattingRules ExtractNumberFormattingRulesFromICUDecimalFormatte
 	NewUEDecimalNumberFormattingRules.SecondaryGroupingSize			= (InICUDecimalFormat.getSecondaryGroupingSize() < 1) 
 																		? NewUEDecimalNumberFormattingRules.PrimaryGroupingSize 
 																		: static_cast<uint8>(InICUDecimalFormat.getSecondaryGroupingSize());
+#if WITH_ICU_V64
+	NewUEDecimalNumberFormattingRules.MinimumGroupingDigits			= static_cast<uint8>(FMath::Max(InICUDecimalFormat.getMinimumGroupingDigits(), 1));
+#endif
 
 	NewUEDecimalNumberFormattingRules.DigitCharacters[0]			= ExtractFormattingSymbolAsCharacter(icu::DecimalFormatSymbols::kZeroDigitSymbol,	TEXT('0'));
 	NewUEDecimalNumberFormattingRules.DigitCharacters[1]			= ExtractFormattingSymbolAsCharacter(icu::DecimalFormatSymbols::kOneDigitSymbol,	TEXT('1'));
@@ -687,11 +697,20 @@ FDecimalNumberFormattingRules ExtractNumberFormattingRulesFromICUDecimalFormatte
 	NewUEDecimalNumberFormattingRules.DigitCharacters[8]			= ExtractFormattingSymbolAsCharacter(icu::DecimalFormatSymbols::kEightDigitSymbol,	TEXT('8'));
 	NewUEDecimalNumberFormattingRules.DigitCharacters[9]			= ExtractFormattingSymbolAsCharacter(icu::DecimalFormatSymbols::kNineDigitSymbol,	TEXT('9'));
 
-	// The CLDR uses a dot as the group separator for Spanish, however the RAE favor using a space: https://www.rae.es/dpd/n%C3%BAmeros
-	if (FCStringAnsi::Strcmp(InICULocale.getLanguage(), "es") == 0 && CVarSpanishUsesRAENumberFormat.AsVariable()->GetInt())
+	if (FCStringAnsi::Strcmp(InICULocale.getLanguage(), "es") == 0)
 	{
-		NewUEDecimalNumberFormattingRules.GroupingSeparatorCharacter = TEXT('\u00A0'); // No-Break Space
-		NewUEDecimalNumberFormattingRules.DecimalSeparatorCharacter = TEXT(',');
+		// The CLDR uses a dot as the group separator for Spanish, however the RAE favor using a space: https://www.rae.es/dpd/n%C3%BAmeros
+		if (CVarSpanishUsesRAENumberFormat.AsVariable()->GetBool())
+		{
+			NewUEDecimalNumberFormattingRules.GroupingSeparatorCharacter = TEXT('\u00A0'); // No-Break Space
+			NewUEDecimalNumberFormattingRules.DecimalSeparatorCharacter = TEXT(',');
+		}
+
+		// Should we use "min two" grouping for Spanish (eg, "1234" formats as "1234" rather than "1 234", but "12345" still formats as "12 345")
+		if (CVarSpanishUsesMinTwoGrouping.AsVariable()->GetBool())
+		{
+			NewUEDecimalNumberFormattingRules.MinimumGroupingDigits = 2;
+		}
 	}
 
 	return NewUEDecimalNumberFormattingRules;
diff --git a/Engine/Source/Runtime/Core/Public/Internationalization/FastDecimalFormat.h b/Engine/Source/Runtime/Core/Public/Internationalization/FastDecimalFormat.h
--- a/Engine/Source/Runtime/Core/Public/Internationalization/FastDecimalFormat.h
+++ b/Engine/Source/Runtime/Core/Public/Internationalization/FastDecimalFormat.h
@@ -18,8 +18,9 @@ struct FDecimalNumberFormattingRules
 	FDecimalNumberFormattingRules()
 		: GroupingSeparatorCharacter(TEXT('\0'))
 		, DecimalSeparatorCharacter(TEXT('\0'))
-		, PrimaryGroupingSize(TEXT('\0'))
-		, SecondaryGroupingSize(TEXT('\0'))
+		, PrimaryGroupingSize(0)
+		, SecondaryGroupingSize(0)
+		, MinimumGroupingDigits(1)
 	{
 		DigitCharacters[0] = TEXT('0');
 		DigitCharacters[1] = TEXT('1');
@@ -45,6 +46,7 @@ struct FDecimalNumberFormattingRules
 	TCHAR DecimalSeparatorCharacter;
 	uint8 PrimaryGroupingSize;
 	uint8 SecondaryGroupingSize;
+	uint8 MinimumGroupingDigits;
 	TCHAR DigitCharacters[10];
 
 	/** Default number formatting options for a given culture */
