Fix unused variable warnings in AutoRTFM.h when compiling UE_AUTORTFM == 0

diff --git a/Engine/Source/Runtime/Core/Public/AutoRTFM/AutoRTFM.h b/Engine/Source/Runtime/Core/Public/AutoRTFM/AutoRTFM.h
--- a/Engine/Source/Runtime/Core/Public/AutoRTFM/AutoRTFM.h
+++ b/Engine/Source/Runtime/Core/Public/AutoRTFM/AutoRTFM.h
@@ -31,6 +31,8 @@
 template <typename FuncType> class TFunction;
 #endif
 
+#define UE_AUTORTFM_UNUSED(UNUSEDVAR)   (void)UNUSEDVAR
+
 #ifdef __cplusplus
 extern "C"
 {
@@ -120,6 +122,8 @@ void autortfm_commit(void (*work)(void* arg), void* arg);
 #else
 UE_AUTORTFM_FORCEINLINE void autortfm_commit(void (*work)(void* arg), void* arg)
 {
+	UE_AUTORTFM_UNUSED(work);
+	UE_AUTORTFM_UNUSED(arg);
     abort();
 }
 #endif
@@ -201,6 +205,8 @@ UE_AUTORTFM_FORCEINLINE void autortfm_open(void (*work)(void* arg), void* arg)
 #else
 [[nodiscard]] UE_AUTORTFM_FORCEINLINE autortfm_status autortfm_close(void (*work)(void* arg), void* arg)
 {
+	UE_AUTORTFM_UNUSED(work);
+	UE_AUTORTFM_UNUSED(arg);
     abort();
 	return autortfm_status_aborted_by_language;
 }
@@ -215,6 +221,8 @@ void autortfm_record_open_write(void* Ptr, size_t Size);
 #else
 UE_AUTORTFM_FORCEINLINE void autortfm_record_open_write(void* Ptr, size_t Size)
 {
+	UE_AUTORTFM_UNUSED(Ptr);
+	UE_AUTORTFM_UNUSED(Size);
 }
 #endif
 
@@ -245,7 +253,11 @@ UE_AUTORTFM_FORCEINLINE void autortfm_record_open_write(void* Ptr, size_t Size)
 #if UE_AUTORTFM
 void autortfm_register_open_function(void* original_function, void* new_function);
 #else
-UE_AUTORTFM_FORCEINLINE void autortfm_register_open_function(void* original_function, void* new_function) { }
+UE_AUTORTFM_FORCEINLINE void autortfm_register_open_function(void* original_function, void* new_function) 
+{ 
+	UE_AUTORTFM_UNUSED(original_function);
+	UE_AUTORTFM_UNUSED(new_function);
+}
 #endif
 
 // Have some work happen when this transaction commits. For nested transactions,
@@ -268,7 +280,11 @@ UE_AUTORTFM_FORCEINLINE void autortfm_defer_until_commit(void (*work)(void* arg)
 #if UE_AUTORTFM
 void autortfm_defer_until_abort(void (*work)(void* arg), void* arg);
 #else
-UE_AUTORTFM_FORCEINLINE void autortfm_defer_until_abort(void (*work)(void* arg), void* arg) { }
+UE_AUTORTFM_FORCEINLINE void autortfm_defer_until_abort(void (*work)(void* arg), void* arg)
+{
+	UE_AUTORTFM_UNUSED(work);
+	UE_AUTORTFM_UNUSED(arg);
+}
 #endif
 
 // Have some work happen when this transaction commits. For nested transactions,
@@ -289,7 +305,11 @@ UE_AUTORTFM_FORCEINLINE void autortfm_open_commit(void (*work)(void* arg), void*
 #if UE_AUTORTFM
 void autortfm_open_abort(void (*work)(void* arg), void* arg);
 #else
-UE_AUTORTFM_FORCEINLINE void autortfm_open_abort(void (*work)(void* arg), void* arg) { }
+UE_AUTORTFM_FORCEINLINE void autortfm_open_abort(void (*work)(void* arg), void* arg)
+{
+	UE_AUTORTFM_UNUSED(work);
+	UE_AUTORTFM_UNUSED(arg);
+}
 #endif
 
 // Inform the runtime that we have performed a new object allocation. It's only
@@ -303,6 +323,7 @@ void* autortfm_did_allocate(void* ptr, size_t size);
 #else
 UE_AUTORTFM_FORCEINLINE void* autortfm_did_allocate(void* ptr, size_t size)
 {
+	UE_AUTORTFM_UNUSED(size);
     return ptr;
 }
 #endif
@@ -327,7 +348,11 @@ UE_AUTORTFM_FORCEINLINE void autortfm_check_consistency_assuming_no_races(void)
 #if UE_AUTORTFM
 void autortfm_check_abi(void* ptr, size_t size);
 #else
-UE_AUTORTFM_FORCEINLINE void autortfm_check_abi(void* ptr, size_t size) { }
+UE_AUTORTFM_FORCEINLINE void autortfm_check_abi(void* ptr, size_t size)
+{
+	UE_AUTORTFM_UNUSED(ptr);
+	UE_AUTORTFM_UNUSED(size);
+}
 #endif
 
 #ifdef __cplusplus
@@ -408,6 +433,8 @@ UE_AUTORTFM_FORCEINLINE void RecordOpenWrite(TTYPE* Ptr)
 // RecordOpenRead does nothing right now, but it is intended as a support stub for the day when we move to full AutoSTM
 UE_AUTORTFM_FORCEINLINE void RecordOpenRead(void const* Ptr, size_t Size)
 {
+	UE_AUTORTFM_UNUSED(Ptr);
+	UE_AUTORTFM_UNUSED(Size);
 }
 
 template<typename TTYPE>
