Fix for CompileId not being assigned during the constructor of the CompileTask leading to incorrectly evaluating whether the experimental VM should be used

diff --git a/Engine/Plugins/FX/Niagara/Source/Niagara/Private/NiagaraAsyncCompile.cpp b/Engine/Plugins/FX/Niagara/Source/Niagara/Private/NiagaraAsyncCompile.cpp
--- a/Engine/Plugins/FX/Niagara/Source/Niagara/Private/NiagaraAsyncCompile.cpp
+++ b/Engine/Plugins/FX/Niagara/Source/Niagara/Private/NiagaraAsyncCompile.cpp
@@ -117,8 +117,6 @@ FNiagaraAsyncCompileTask::FNiagaraAsyncCompileTask(UNiagaraSystem* InOwningSyste
 	ScriptPair = InScriptPair;
 
 	CurrentState = ENiagaraCompilationState::CheckDDC;
-
-	bExperimentalVMDisabled = InScriptPair.CompileId.AdditionalDefines.Contains(FNiagaraCompileOptions::ExperimentalVMDisabled);
 }
 
 void FNiagaraAsyncCompileTask::ProcessCurrentState()
@@ -330,7 +328,7 @@ void FNiagaraAsyncCompileTask::OptimizeByteCode()
 	check(ExeData);
 
 #if VECTORVM_SUPPORTS_LEGACY
-	if (bExperimentalVMDisabled)
+	if (ScriptPair.CompileId.AdditionalDefines.Contains(FNiagaraCompileOptions::ExperimentalVMDisabled))
 	{
 		return;
 	}
diff --git a/Engine/Plugins/FX/Niagara/Source/Niagara/Private/NiagaraAsyncCompile.h b/Engine/Plugins/FX/Niagara/Source/Niagara/Private/NiagaraAsyncCompile.h
--- a/Engine/Plugins/FX/Niagara/Source/Niagara/Private/NiagaraAsyncCompile.h
+++ b/Engine/Plugins/FX/Niagara/Source/Niagara/Private/NiagaraAsyncCompile.h
@@ -69,7 +69,6 @@ public:
 	bool bWaitForCompileJob = false;
 	bool bUsedShaderCompilerWorker = false;
 	bool bFetchedGCObjects = false;
-	bool bExperimentalVMDisabled = true;
 
 	// in order to coordinate between tasks associated with a system CheckDDC can be handled external to the task.
 	// In that case we disable this flag
diff --git a/Engine/Source/Runtime/Core/Private/UObject/UE5ReleaseStreamObjectVersion.cpp b/Engine/Source/Runtime/Core/Private/UObject/UE5ReleaseStreamObjectVersion.cpp
--- a/Engine/Source/Runtime/Core/Private/UObject/UE5ReleaseStreamObjectVersion.cpp
+++ b/Engine/Source/Runtime/Core/Private/UObject/UE5ReleaseStreamObjectVersion.cpp
@@ -12,7 +12,7 @@ TMap<FGuid, FGuid> FUE5ReleaseStreamObjectVersion::GetSystemGuids()
 	SystemGuids.Add(DevGuids.MATERIALSHADERMAP_DERIVEDDATA_VER, FGuid("79090A66F9D94B5BB285D49E5D39468E"));
 	SystemGuids.Add(DevGuids.NANITE_DERIVEDDATA_VER, FGuid("AC88CFBDDA614A9C8488F64D2DAF699D"));
 	SystemGuids.Add(DevGuids.NIAGARASHADERMAP_DERIVEDDATA_VER, FGuid("6BE35B6FACB34568970120F9BC8DAB80"));
-	SystemGuids.Add(DevGuids.Niagara_LatestScriptCompileVersion, FGuid("CCA1C9E89704F54E9DF3639F7E75E78E"));
+	SystemGuids.Add(DevGuids.Niagara_LatestScriptCompileVersion, FGuid("43A777C54EB24E9894525B8D04529F23"));
 	SystemGuids.Add(DevGuids.SkeletalMeshDerivedDataVersion, FGuid("25C49E579B3142DDA2A8C14037267679"));
 
 	return SystemGuids;
