Make sure the entitlement for virtual address space on iOS is not only found but also set to true.

diff --git a/Engine/Source/Runtime/Core/Private/Apple/ApplePlatformMemory.cpp b/Engine/Source/Runtime/Core/Private/Apple/ApplePlatformMemory.cpp
--- a/Engine/Source/Runtime/Core/Private/Apple/ApplePlatformMemory.cpp
+++ b/Engine/Source/Runtime/Core/Private/Apple/ApplePlatformMemory.cpp
@@ -306,20 +306,27 @@ FMalloc* FApplePlatformMemory::BaseAllocator()
 #if PLATFORM_IOS || PLATFORM_TVOS
     if (FIOSPlatformMisc::IsEntitlementEnabled("com.apple.developer.kernel.extended-virtual-addressing"))
     {
+        UE_LOG(LogTemp, Display, TEXT("Virtual Address Space entitlement found. Using MallocBinned2 to allocate memory"));
         AllocatorToUse = EMemoryAllocatorToUse::Binned2;
 #undef USE_MALLOC_BINNED2
 #define USE_MALLOC_BINNED2 1
     }
+    else
+    {
+        UE_LOG(LogTemp, Warning, TEXT("Virtual Address Space entitlement NOT found. Defaulting to Ansi allocator"));
+    }
 #elif USE_MALLOC_BINNED2
     if (USE_MALLOC_BINNED2)
     {
+        UE_LOG(LogTemp, Display, TEXT("Using MallocBinned2 to allocate memory"));
         AllocatorToUse = EMemoryAllocatorToUse::Binned2;
     }
-#endif
     else
     {
+        UE_LOG(LogTemp, Warning, TEXT("Defaulting to Ansi allocator"));
         AllocatorToUse = EMemoryAllocatorToUse::Ansi;
     }
+#endif
     
     switch (AllocatorToUse)
     {
diff --git a/Engine/Source/Runtime/Core/Private/IOS/IOSPlatformMisc.cpp b/Engine/Source/Runtime/Core/Private/IOS/IOSPlatformMisc.cpp
--- a/Engine/Source/Runtime/Core/Private/IOS/IOSPlatformMisc.cpp
+++ b/Engine/Source/Runtime/Core/Private/IOS/IOSPlatformMisc.cpp
@@ -1267,8 +1267,8 @@ typedef struct __MultiBlob {
 extern NSString *EntitlementsData(void)
 {
     // iterate through the headers to find the executable, since only the executable has the entitlements
-    const struct mach_header_64 *executableHeader = nullptr;
-    char *ImageName = nullptr;
+    const struct mach_header_64* executableHeader = nullptr;
+    char* ImageName = nullptr;
     for (uint32_t i = 0; i < _dyld_image_count() && executableHeader == nullptr; i++)
     {
         const struct mach_header_64 *header = (struct mach_header_64 *)_dyld_get_image_header(i);
@@ -1287,7 +1287,7 @@ extern NSString *EntitlementsData(void)
     // verify that it's a 64bit app
     if (executableHeader->magic != MH_MAGIC_64)
     {
-        NSLog(@"Executable is NOT 64bit. Entitlement retrieval not supported.");
+        UE_LOG(LogIOS, Error, TEXT("Executable is NOT 64bit. Entitlement retrieval not supported."));
         return nil;
     }
     uintptr_t cursor = (uintptr_t)executableHeader + sizeof(struct mach_header_64);
@@ -1300,7 +1300,7 @@ extern NSString *EntitlementsData(void)
         switch (segmentCommand->cmd)
         {
             case LC_CODE_SIGNATURE:
-                NSLog(@"LC_CODE_SIGNATURE found");
+                UE_LOG(LogIOS, Log, TEXT("LC_CODE_SIGNATURE found"));
                 break;
             default:
                 continue;
@@ -1308,10 +1308,10 @@ extern NSString *EntitlementsData(void)
 
         const struct linkedit_data_command *dataCommand = (const struct linkedit_data_command *)segmentCommand;
 
-        FILE *file = fopen(ImageName, "r");
+        FILE* file = fopen(ImageName, "rb");
         if (file == NULL)
         {
-            NSLog(@"Could not open binary file");
+            UE_LOG(LogIOS, Error, TEXT("Could not open binary file"));
             return nil;
         }
         CS_MultiBlob multiBlob;
@@ -1359,30 +1359,95 @@ extern NSString *EntitlementsData(void)
     return nil;
 }
 
-extern bool IsEntitlementPresentInEmbeddedProvision(NSString *EntitlementsToFind)
+extern bool IsEntitlementPresentInEmbeddedProvision(const char *EntitlementsToFind)
 {
-    NSString *mobileprovisionPath = [[[NSBundle mainBundle] bundlePath] stringByAppendingPathComponent:@"embedded.mobileprovision"];
-    NSError* error = nil;
-    NSString *myText = [NSString stringWithContentsOfFile:mobileprovisionPath encoding:NSUTF8StringEncoding error:&error];
+    NSString* mobileprovisionPath = [[[NSBundle mainBundle] bundlePath] stringByAppendingPathComponent:@"embedded.mobileprovision"];
+  
+    FILE* file = fopen([mobileprovisionPath cStringUsingEncoding:1],"rb");
+
+    size_t entitlements_len = strlen(EntitlementsToFind);
     
-    if ([myText rangeOfString: (@"%s", EntitlementsToFind)].location == NSNotFound)
+    if (file == NULL)
     {
+        NSLog(@"Mobile Provision not found");
         return false;
     }
-    else
+
+    size_t readcount;
+    char buffer[1025]; //room for null
+    while(true)
     {
-        return true;
+        readcount = fread(buffer, 1, 1024, file);
+        if (readcount <= 0)
+        {
+            break;
+        }
+        for(size_t i = 0; i < readcount; i++)
+        {
+            if(buffer[i] == 0)
+            {
+                buffer[i] = ' ';  // replace any null terminators that might be in the binary data
+            }
+        }
+        
+        buffer[readcount] = 0; // null terminate buffer!
+        char* entitlementskeyptr = strstr(buffer, EntitlementsToFind);
+        
+        if(entitlementskeyptr)
+        {
+            fseek(file, -readcount + entitlementskeyptr - buffer, SEEK_CUR); // seek to immediately after the entitlements key
+            readcount = fread(buffer, 1, 1024, file); // read 1024 bytes immediately after the entitlements key so we definitely get the value untruncated
+            // we don't expect any \0 characters between the key and the value, so we don't do the replace thing.
+            buffer[readcount] = 0; // null terminate buffer!
+            
+            // there could be keys following our one with their own true or false values that we could accidentally match if we just look for 'true' and return immediately.
+            char* trueptr = strstr(buffer, "true");
+            char* falseptr = strstr(buffer, "false");
+            
+            if(trueptr == NULL && falseptr == NULL)
+            {
+                UE_LOG(LogIOS, Error, TEXT("Unexpected Behaviour. The entitlement key is found but its value is not set."));
+                return false;
+            }
+            if(trueptr && falseptr == NULL)  // only true
+            {
+                UE_LOG(LogIOS, Log, TEXT("Entitlements found in embedded mobile provision file."));
+                return true;
+            }
+            if(trueptr == NULL && falseptr) // only false
+            {
+                UE_LOG(LogIOS, Log, TEXT("Entitlements found but set to false."));
+                return false;
+            }
+            return (trueptr < falseptr); // return true if true comes before false
+        }
+        
+        if(readcount < 1024)
+        {
+            break; // end of file
+        }
+        
+        // seek back in case the entitlements was truncated.
+        fseek(file, -entitlements_len, SEEK_CUR);
     }
+    fclose(file);
+    return false;
 }
 
 bool FIOSPlatformMisc::IsEntitlementEnabled(const char * EntitlementToCheck)
 {
-    NSString *EntitlementsToFind = [NSString stringWithFormat: @"%s", EntitlementToCheck];
-    NSString *EntitlementData = EntitlementsData();
-    if ([EntitlementData rangeOfString: (@"%s", EntitlementsToFind)].location == NSNotFound)
+    NSString* TrueFlag = @"</key><true/>";
+    NSString* EntitlementsToFind = [NSString stringWithFormat: @"%s%@" , EntitlementToCheck, TrueFlag];
+    
+    NSString* CleanedEntitlementData = [EntitlementsData() stringByReplacingOccurrencesOfString: @"\\s+" withString: @"" options: NSRegularExpressionSearch range: NSMakeRange(0, EntitlementsToFind.length)];
+
+    CleanedEntitlementData = [[CleanedEntitlementData componentsSeparatedByCharactersInSet:[NSCharacterSet characterSetWithCharactersInString:@"\t"]] componentsJoinedByString:@""];
+    CleanedEntitlementData = [[CleanedEntitlementData componentsSeparatedByCharactersInSet:[NSCharacterSet newlineCharacterSet]] componentsJoinedByString:@""];
+
+    if ([CleanedEntitlementData rangeOfString: (@"%s", EntitlementsToFind)].location == NSNotFound)
     {
-        NSLog(@"Entitlements not found in binary Mach-O header. Looking at the embedded mobile provision file.");
-        return IsEntitlementPresentInEmbeddedProvision(EntitlementsToFind);
+        UE_LOG(LogIOS, Log, TEXT("Entitlements not found in binary Mach-O header. Looking at the embedded mobile provision file."));
+        return IsEntitlementPresentInEmbeddedProvision(EntitlementToCheck);
     }
     else
     {
