TaskGraph: fixed tracing of standalone graph events that don't have nested tasks at the moment of dispatching subsequents. because of this issue a lot of graph events weren't shown in TasksInsights
This is a rather "hacky" fix. The old task implementation (TaskGraph API) is on a life support so there's no point of spending too much time for this

diff --git a/Engine/Source/Runtime/Core/Private/Async/TaskGraph.cpp b/Engine/Source/Runtime/Core/Private/Async/TaskGraph.cpp
--- a/Engine/Source/Runtime/Core/Private/Async/TaskGraph.cpp
+++ b/Engine/Source/Runtime/Core/Private/Async/TaskGraph.cpp
@@ -2443,7 +2443,7 @@ void FGraphEvent::DispatchSubsequents(ENamedThreads::Type CurrentThreadIfKnown)
 	DispatchSubsequents(NewTasks, CurrentThreadIfKnown);
 }
 
-void FGraphEvent::DispatchSubsequents(TArray<FBaseGraphTask*>& NewTasks, ENamedThreads::Type CurrentThreadIfKnown)
+void FGraphEvent::DispatchSubsequents(TArray<FBaseGraphTask*>& NewTasks, ENamedThreads::Type CurrentThreadIfKnown, bool bInternal/* = false */)
 {
 	if (EventsToWaitFor.Num())
 	{
@@ -2500,6 +2500,10 @@ void FGraphEvent::DispatchSubsequents(TArray<FBaseGraphTask*>& NewTasks, ENamedT
 		NewTask->ConditionalQueueTask(CurrentThreadIfKnown, bWakeUpWorker);
 	}
 
+	if (!bInternal)
+	{
+		TaskTrace::Launched(GetTraceId(), TEXT("Standalone graph event"), true, ENamedThreads::AnyThread, sizeof(FGraphEvent));
+	}
 	TaskTrace::Completed(GetTraceId());
 }
 
diff --git a/Engine/Source/Runtime/Core/Public/Async/TaskGraphInterfaces.h b/Engine/Source/Runtime/Core/Public/Async/TaskGraphInterfaces.h
--- a/Engine/Source/Runtime/Core/Public/Async/TaskGraphInterfaces.h
+++ b/Engine/Source/Runtime/Core/Public/Async/TaskGraphInterfaces.h
@@ -984,10 +984,11 @@ public:
 
 	/**
 	 *	"Complete" the event. This grabs the list of subsequents and atomically closes it. Then for each subsequent and for each item in "NewTasks" it reduces the number of prerequisites outstanding and if that drops to zero, the task is queued. 
-	 * @param NewTasks subsequents to add
+	 *	@param NewTasks subsequents to add
 	 *	@param	 CurrentThreadIfKnown if the current thread is known, provide it here. Otherwise it will be determined via TLS if any task ends up being queued.
+	 *	@param bInternal used to distinguish between "internal" graph events that have an accompanying task, and "standalone" graph events created explicitly by `FGraphEvent::CreateGraphEvent()`
 	**/
-	CORE_API void DispatchSubsequents(TArray<FBaseGraphTask*>& NewTasks, ENamedThreads::Type CurrentThreadIfKnown = ENamedThreads::AnyThread);
+	CORE_API void DispatchSubsequents(TArray<FBaseGraphTask*>& NewTasks, ENamedThreads::Type CurrentThreadIfKnown = ENamedThreads::AnyThread, bool bInternal = false);
 
 	/**
 	 *	Determine if the event has been completed. This can be used to poll for completion. 
@@ -1043,10 +1044,11 @@ private:
 	**/
 	static CORE_API void Recycle(FGraphEvent* ToRecycle);
 
+	friend struct FGraphEventAndSmallTaskStorage;
+
 	/**
 	 *	Hidden Constructor
 	**/
-	friend struct FGraphEventAndSmallTaskStorage;
 	FGraphEvent()
 		: ThreadToDoGatherOn(ENamedThreads::AnyHiPriThreadHiPriTask)
 	{
@@ -1272,7 +1274,7 @@ private:
 		if (TTask::GetSubsequentsMode() == ESubsequentsMode::TrackSubsequents)
 		{
 			FPlatformMisc::MemoryBarrier();
-			Subsequents->DispatchSubsequents(NewTasks, CurrentThread);
+			Subsequents->DispatchSubsequents(NewTasks, CurrentThread, true);
 		}
 		else
 		{
