Chaos : add strength to chaos physics material and wire it to the high level engine physics material

diff --git a/Engine/Source/Runtime/Core/Private/Math/UnitConversion.cpp b/Engine/Source/Runtime/Core/Private/Math/UnitConversion.cpp
--- a/Engine/Source/Runtime/Core/Private/Math/UnitConversion.cpp
+++ b/Engine/Source/Runtime/Core/Private/Math/UnitConversion.cpp
@@ -98,6 +98,11 @@ FParseCandidate ParseCandidates[] = {
 	{ TEXT("Percent"),				EUnit::Percentage },			{ TEXT("%"),	EUnit::Percentage },
 
 	{ TEXT("times"),				EUnit::Multiplier },			{ TEXT("x"),	EUnit::Multiplier },			{ TEXT("multiplier"),		EUnit::Multiplier },
+
+	{ TEXT("Pascals"),				EUnit::Pascals },				{ TEXT("Pa"),	EUnit::Pascals},
+	{ TEXT("KiloPascals"),			EUnit::KiloPascals},			{ TEXT("kPa"),	EUnit::KiloPascals},
+	{ TEXT("MegaPascals"),			EUnit::MegaPascals},			{ TEXT("MPa"),	EUnit::MegaPascals},
+	{ TEXT("GigaPascals"),			EUnit::GigaPascals},			{ TEXT("GPa"),	EUnit::GigaPascals},
 };
 
 /** Static array of display strings that directly map to EUnit enumerations */
@@ -130,6 +135,8 @@ const TCHAR* const DisplayStrings[] = {
 	TEXT("%"),
 
 	TEXT("x"),
+
+	TEXT("Pa"), TEXT("kPa"), TEXT("MPa"), TEXT("GPa"),
 };
 
 const EUnitType UnitTypes[] = {
@@ -159,6 +166,8 @@ const EUnitType UnitTypes[] = {
 	EUnitType::PixelDensity,
 
 	EUnitType::Multipliers, EUnitType::Multipliers,
+
+	EUnitType::Stress, EUnitType::Stress, EUnitType::Stress, EUnitType::Stress,
 };
 
 
@@ -433,6 +442,7 @@ FUnitSettings::FUnitSettings()
 	DisplayUnits[(uint8)EUnitType::LuminousFlux].Add(EUnit::Lumens);
 	DisplayUnits[(uint8)EUnitType::ExposureValue].Add(EUnit::ExposureValue);
 	DisplayUnits[(uint8)EUnitType::Time].Add(EUnit::Seconds);
+	DisplayUnits[(uint8)EUnitType::Stress].Add(EUnit::MegaPascals);
 }
 
 bool FUnitSettings::ShouldDisplayUnits() const
@@ -680,6 +690,19 @@ namespace UnitConversion
 		}
 	}
 
+	double StressUnificationFactor(EUnit From)
+	{
+		// Convert to Pascals
+		switch (From)
+		{
+		case EUnit::GigaPascals:		return 1000000000.0;
+		case EUnit::MegaPascals:		return 1000000.0;
+		case EUnit::KiloPascals:		return 1000.0;
+		case EUnit::Pascals:			// fallthrough
+		default: 						return 1.0;
+		}
+	}
+
 	TValueOrError<FNumericUnit<double>, FText> TryParseExpression(const TCHAR* InExpression, EUnit From, const FNumericUnit<double>& InExistingValue)
 	{
 		const FUnitExpressionParser Parser(From);
@@ -716,6 +739,8 @@ namespace UnitConversion
 
 			TArray<FQuantizationInfo> Time;
 
+			TArray<FQuantizationInfo> Stress;
+
 			FStaticBounds()
 			{
 				MetricDistance.Emplace(EUnit::Micrometers,	1000.0f);
@@ -757,6 +782,11 @@ namespace UnitConversion
 				Time.Emplace(EUnit::Days,				365.242f / 12.0f);
 				Time.Emplace(EUnit::Months,			12.0f);
 				Time.Emplace(EUnit::Years,			0.0f);
+
+				Stress.Emplace(EUnit::Pascals,		1000.0f);
+				Stress.Emplace(EUnit::KiloPascals,	1000.0f);
+				Stress.Emplace(EUnit::MegaPascals,	1000.0f);
+				Stress.Emplace(EUnit::GigaPascals,	0.0f);
 			}
 		};
 
@@ -785,6 +815,9 @@ namespace UnitConversion
 		case EUnit::Milliseconds: case EUnit::Seconds: case EUnit::Minutes: case EUnit::Hours: case EUnit::Days: case EUnit::Months: case EUnit::Years:
 			return &Bounds.Time;
 
+		case EUnit::Pascals: case EUnit::KiloPascals: case EUnit::MegaPascals: case EUnit::GigaPascals:
+			return &Bounds.Stress;
+
 		default:
 			return TOptional<const TArray<FQuantizationInfo>*>();
 		}
diff --git a/Engine/Source/Runtime/Core/Public/Math/UnitConversion.h b/Engine/Source/Runtime/Core/Public/Math/UnitConversion.h
--- a/Engine/Source/Runtime/Core/Public/Math/UnitConversion.h
+++ b/Engine/Source/Runtime/Core/Public/Math/UnitConversion.h
@@ -56,6 +56,9 @@ enum class EUnit : uint8
 	/** Arbitrary multipliers */
 	Percentage,	Multiplier,
 
+	/** Stress units */
+	Pascals, KiloPascals, MegaPascals, GigaPascals,
+
 	/** Symbolic entry, not specifiable on meta data */
 	Unspecified
 };
@@ -63,7 +66,7 @@ enum class EUnit : uint8
 /** Enumeration that specifies particular classes of unit */
 enum class EUnitType
 {
-	Distance, Angle, Speed, Temperature, Mass, Force, Frequency, DataSize, LuminousFlux, LuminousIntensity, Illuminance, Luminance, Time, PixelDensity, Multipliers, ExposureValue,
+	Distance, Angle, Speed, Temperature, Mass, Force, Frequency, DataSize, LuminousFlux, LuminousIntensity, Illuminance, Luminance, Time, PixelDensity, Multipliers, ExposureValue, Stress,
 
 	// Symbolic entry - do not use directly
 	NumberOf,
diff --git a/Engine/Source/Runtime/Core/Public/Math/UnitConversion.inl b/Engine/Source/Runtime/Core/Public/Math/UnitConversion.inl
--- a/Engine/Source/Runtime/Core/Public/Math/UnitConversion.inl
+++ b/Engine/Source/Runtime/Core/Public/Math/UnitConversion.inl
@@ -38,6 +38,8 @@ namespace UnitConversion
 	CORE_API double TimeUnificationFactor(EUnit From);
 	/** Find the common quantization factor for the specified multiplier unit. Quantizes to 1.0 scale (where 1.0 == 100%). */
 	CORE_API double MultiplierUnificationFactor(EUnit From);
+	/** Find the common quantization factor for the specified multiplier unit. Quantizes to 1.0 scale (where 1.0 == 100%). */
+	CORE_API double StressUnificationFactor(EUnit From);
 
 	/** Attempt to parse an expression into a numeric unit */
 	CORE_API TValueOrError<FNumericUnit<double>, FText> TryParseExpression(const TCHAR* InExpression, EUnit DefaultUnit, const FNumericUnit<double>& InExistingValue);
@@ -86,6 +88,7 @@ T FUnitConversion::Convert(T InValue, EUnit From, EUnit To)
 		case EUnitType::LuminousFlux:		return InValue;
 		case EUnitType::Time:				return InValue * TimeUnificationFactor(From) 			* (1.0 / TimeUnificationFactor(To));
 		case EUnitType::Multipliers:		return InValue * MultiplierUnificationFactor(From) 		* (1.0 / MultiplierUnificationFactor(To));
+		case EUnitType::Stress:				return InValue * StressUnificationFactor(From)			* (1.0 / StressUnificationFactor(To));
 		// Temperature conversion is not just a simple multiplication, so needs special treatment
 		case EUnitType::Temperature:
 		{
diff --git a/Engine/Source/Runtime/Experimental/Chaos/Public/Chaos/Defines.h b/Engine/Source/Runtime/Experimental/Chaos/Public/Chaos/Defines.h
--- a/Engine/Source/Runtime/Experimental/Chaos/Public/Chaos/Defines.h
+++ b/Engine/Source/Runtime/Experimental/Chaos/Public/Chaos/Defines.h
@@ -39,6 +39,30 @@ typedef int32_t int32;
 
 namespace Chaos
 {
+	// Solver uses Kg for mass, Second for time and Cm for distance
+	// Since a MegaPascal is  force over an Area, and Force = Mass * Acc
+	// we get (Kg.cm)/(cm2.s2) => Kg/(cm.s2)
+	constexpr float MegaPascalToKgPerCmS2(float MegaPascals)
+	{
+		return MegaPascals * 10000;
+	}
+
+	struct FChaosPhysicsMaterialStrength
+	{
+
+		// using concrete as default ( lowest values of it )
+		FChaosPhysicsMaterialStrength()
+			: TensileStrength(MegaPascalToKgPerCmS2(2))
+			, CompressionStrength(MegaPascalToKgPerCmS2(20))
+			, ShearStrength(MegaPascalToKgPerCmS2(6))
+		{}
+
+		// Unit is Kg/(cm.s2)
+		float TensileStrength;
+		float CompressionStrength;
+		float ShearStrength;
+	};
+
 	class FChaosPhysicsMaterial
 	{
 	public:
@@ -70,6 +94,8 @@ namespace Chaos
 		ECombineMode FrictionCombineMode;
 		ECombineMode RestitutionCombineMode;
 
+		FChaosPhysicsMaterialStrength Strength;
+
 
 		FChaosPhysicsMaterial()
 			: Friction(0.5)
diff --git a/Engine/Source/Runtime/PhysicsCore/Private/ChaosEngineInterface.cpp b/Engine/Source/Runtime/PhysicsCore/Private/ChaosEngineInterface.cpp
--- a/Engine/Source/Runtime/PhysicsCore/Private/ChaosEngineInterface.cpp
+++ b/Engine/Source/Runtime/PhysicsCore/Private/ChaosEngineInterface.cpp
@@ -233,6 +233,9 @@ void FChaosEngineInterface::UpdateMaterial(FPhysicsMaterialHandle& InHandle,UPhy
 		Material->SleepingLinearThreshold = InMaterial->SleepLinearVelocityThreshold;
 		Material->SleepingAngularThreshold = InMaterial->SleepAngularVelocityThreshold;
 		Material->SleepCounterThreshold = InMaterial->SleepCounterThreshold;
+		Material->Strength.TensileStrength = Chaos::MegaPascalToKgPerCmS2(InMaterial->Strength.TensileStrength);
+		Material->Strength.CompressionStrength = Chaos::MegaPascalToKgPerCmS2(InMaterial->Strength.CompressionStrength);
+		Material->Strength.ShearStrength = Chaos::MegaPascalToKgPerCmS2(InMaterial->Strength.ShearStrength);
 	}
 
 	Chaos::FPhysicalMaterialManager::Get().UpdateMaterial(InHandle);
diff --git a/Engine/Source/Runtime/PhysicsCore/Private/PhysicalMaterial.cpp b/Engine/Source/Runtime/PhysicsCore/Private/PhysicalMaterial.cpp
--- a/Engine/Source/Runtime/PhysicsCore/Private/PhysicalMaterial.cpp
+++ b/Engine/Source/Runtime/PhysicsCore/Private/PhysicalMaterial.cpp
@@ -16,6 +16,14 @@ UDEPRECATED_PhysicalMaterialPropertyBase::UDEPRECATED_PhysicalMaterialPropertyBa
 {
 }
 
+FPhysicalMaterialStrength::FPhysicalMaterialStrength()
+{
+	// using concrete as default ( lowest values of it )
+	TensileStrength = 2;
+	CompressionStrength = 20;
+	ShearStrength = 6;
+}
+
 UPhysicalMaterial::UPhysicalMaterial(const FObjectInitializer& ObjectInitializer)
 	: Super(ObjectInitializer)
 {
@@ -28,11 +36,6 @@ UPhysicalMaterial::UPhysicalMaterial(const FObjectInitializer& ObjectInitializer
 	SleepAngularVelocityThreshold = 0.05f;
 	SleepCounterThreshold = 4;
 	bOverrideFrictionCombineMode = false;
-	// using concrete as default ( lowest values of it )
-	TensileStrength = 2;
-	CompressionStrength = 20;
-	ShearStrength = 6;
-
 	UserData = FChaosUserData(this);
 }
 
diff --git a/Engine/Source/Runtime/PhysicsCore/Public/PhysicalMaterials/PhysicalMaterial.h b/Engine/Source/Runtime/PhysicsCore/Public/PhysicalMaterials/PhysicalMaterial.h
--- a/Engine/Source/Runtime/PhysicsCore/Public/PhysicalMaterials/PhysicalMaterial.h
+++ b/Engine/Source/Runtime/PhysicsCore/Public/PhysicalMaterials/PhysicalMaterial.h
@@ -19,6 +19,39 @@
 
 struct FPropertyChangedEvent;
 
+/**
+ * Defines the directional strengths of a physical material in term of force per surface area
+ */
+USTRUCT(BlueprintType)
+struct FPhysicalMaterialStrength
+{
+	GENERATED_USTRUCT_BODY()
+
+	FPhysicalMaterialStrength();
+
+	/**
+	* Tensile strength of the material in MegaPascal ( 10^6 N/m2 )
+	* This amount of tension force per area the material can withstand before it fractures
+	*/
+	UPROPERTY(EditAnywhere, BlueprintReadOnly, Category = "PhysicalMaterial", meta = (ClampMin = 0, ForceUnits = "MPa"))
+	float TensileStrength;
+
+	/**
+	* Compression strength of the material in MegaPascal ( 10^6 N/m2 )
+	* This amount of compression force per area the material can withstand before it fractures, crumbles or buckles
+	*/
+	UPROPERTY(EditAnywhere, BlueprintReadOnly, Category = "PhysicalMaterial", meta = (ClampMin = 0, ForceUnits = "MPa"))
+	float CompressionStrength;
+
+	/**
+	* Shear strength of the material in MegaPascal ( 10^6 N/m2 )
+	* This amount of shear force per area the material can withstand before it fractures
+	*/
+	UPROPERTY(EditAnywhere, BlueprintReadOnly, Category = "PhysicalMaterial", meta = (ClampMin = 0, ForceUnits = "MPa"))
+	float ShearStrength;
+};
+
+
 /**
  * Physical materials are used to define the response of a physical object when interacting dynamically with the world.
  */
@@ -103,24 +136,8 @@ class PHYSICSCORE_API UPhysicalMaterial : public UObject
 	UPROPERTY(EditAnywhere, BlueprintReadOnly, Category = PhysicalProperties)
 	TEnumAsByte<EPhysicalSurface> SurfaceType;
 
-	/** 
-	* Tensile strength of the material in MegaPascal ( 10^6 N/m2 )
-	* This amount of tension force per area the material can withstand before it fractures
-	*/
-	float TensileStrength;
-
-	/**
-	* Compression strength of the material in MegaPascal ( 10^6 N/m2 )
-	* This amount of compression force per area the material can withstand before it fractures, crumbles or buckles
-	*/
-	float CompressionStrength;
-
-	/**
-	* Shear strength of the material in MegaPascal ( 10^6 N/m2 )
-	* This amount of shear force per area the material can withstand before it fractures
-	*/
-	float ShearStrength;
-
+	UPROPERTY(EditAnywhere, BlueprintReadOnly, Category = PhysicalProperties)
+	FPhysicalMaterialStrength Strength;
 
 public:
 
