I/O - simple file based FIFO cache

diff --git a/Engine/Source/Runtime/Core/Internal/IO/IoFileCache.h b/Engine/Source/Runtime/Core/Internal/IO/IoFileCache.h
new file mode 100644
--- /dev/null
+++ b/Engine/Source/Runtime/Core/Internal/IO/IoFileCache.h
@@ -0,0 +1,15 @@
+// Copyright Epic Games, Inc. All Rights Reserved.
+
+#pragma once
+
+#include "Templates/UniquePtr.h"
+
+class IIoCache;
+
+struct FFileIoCacheConfig
+{
+	uint64 DiskStorageSize;
+	uint64 MemoryStorageSize;
+};
+
+CORE_API TUniquePtr<IIoCache> MakeFileIoCache(const FFileIoCacheConfig& Config);
diff --git a/Engine/Source/Runtime/Core/Private/IO/IoCache.cpp b/Engine/Source/Runtime/Core/Private/IO/IoCache.cpp
--- a/Engine/Source/Runtime/Core/Private/IO/IoCache.cpp
+++ b/Engine/Source/Runtime/Core/Private/IO/IoCache.cpp
@@ -2,6 +2,8 @@
 
 #include "IO/IoCache.h"
 
+DEFINE_LOG_CATEGORY(LogIoCache);
+
 namespace UE::IO::Private
 {
 	FIoCacheRequestBase::FIoCacheRequestBase(FIoReadCallback&& ReadCallback)
diff --git a/Engine/Source/Runtime/Core/Private/IO/IoFileCache.cpp b/Engine/Source/Runtime/Core/Private/IO/IoFileCache.cpp
new file mode 100644
--- /dev/null
+++ b/Engine/Source/Runtime/Core/Private/IO/IoFileCache.cpp
@@ -0,0 +1,683 @@
+// Copyright Epic Games, Inc. All Rights Reserved.
+
+#include "IO/IoFileCache.h"
+#include "Async/AsyncFileHandle.h"
+#include "Containers/ContainersFwd.h"
+#include "Containers/IntrusiveDoubleLinkedList.h"
+#include "GenericPlatform/GenericPlatformFile.h"
+#include "IO/IoCache.h"
+#include "IO/IoDispatcher.h"
+#include "HAL/Event.h"
+#include "HAL/FileManager.h"
+#include "HAL/Runnable.h"
+#include "HAL/RunnableThread.h"
+#include "IO/IoHash.h"
+#include "Misc/ScopeLock.h"
+#include "Misc/Paths.h"
+#include "Misc/CommandLine.h"
+#include "ProfilingDebugging/CountersTrace.h"
+#include "ProfilingDebugging/CpuProfilerTrace.h"
+#include "Templates/UniquePtr.h"
+#include <atomic>
+
+TRACE_DECLARE_MEMORY_COUNTER(FFileIoCache_CachedBytes, TEXT("FileIoCache/TotalBytes"));
+TRACE_DECLARE_MEMORY_COUNTER(FFileIoCache_PendingBytes, TEXT("FileIoCache/PendingBytes"));
+TRACE_DECLARE_INT_COUNTER(FFileIoCache_GetCount, TEXT("FileIoCache/GetCount"));
+TRACE_DECLARE_INT_COUNTER(FFileIoCache_ErrorCount, TEXT("FileIoCache/ErrorCount"));
+TRACE_DECLARE_INT_COUNTER(FFileIoCache_PutCount, TEXT("FileIoCache/PutCount"));
+TRACE_DECLARE_INT_COUNTER(FFileIoCache_PutRejectCount, TEXT("FileIoCache/PutRejectCount"));
+TRACE_DECLARE_INT_COUNTER(FFileIoCache_PutExistingCount, TEXT("FileIoCache/PutExistingCount"));
+
+namespace UE::IO::Private
+{
+
+///////////////////////////////////////////////////////////////////////////////
+class FCacheFileToc
+{
+	struct FHeader
+	{
+		static constexpr uint32 ExpectedMagic = 0x2e696f; // .io
+
+		uint32 Magic = 0;
+		uint32 EntryCount = 0;
+		uint64 CursorPos = 0;
+	};
+
+public:
+	struct FTocEntry
+	{
+		FIoHash Key;
+		FIoHash Hash;
+		uint64 SerialOffset = 0;
+		uint64 SerialSize = 0;
+
+		friend FArchive& operator<<(FArchive& Ar, FTocEntry& Entry)
+		{
+			Ar << Entry.Key;
+			Ar << Entry.Hash;
+			Ar << Entry.SerialOffset;
+			Ar << Entry.SerialSize;
+
+			return Ar;
+		}
+	};
+
+	FCacheFileToc() = default;
+
+	void AddEntry(const FIoHash& Key, const FIoHash& Hash, uint64 SerialOffset, uint64 SerialSize);
+	FIoStatus Load(const FString& FilePath, uint64& OutCursorPos);
+	FIoStatus Save(const FString& FilePath, const uint64 CursorPos);
+
+	TConstArrayView<FTocEntry> GetEntries() const
+	{ 
+		return TocEntries;
+	}
+
+private:
+	TArray<FTocEntry> TocEntries;
+};
+
+void FCacheFileToc::AddEntry(const FIoHash& Key, const FIoHash& Hash, uint64 SerialOffset, uint64 SerialSize)
+{
+	TocEntries.Add(FTocEntry
+	{
+		Key,
+		Hash,
+		SerialOffset,
+		SerialSize
+	});
+}
+
+FIoStatus FCacheFileToc::Load(const FString& FilePath, uint64& OutCursorPos)
+{
+	IFileManager& FileMgr = IFileManager::Get();
+	TUniquePtr<FArchive> Ar(FileMgr.CreateFileReader(*FilePath));
+	
+	OutCursorPos = ~uint64(0);
+
+	if (!Ar.IsValid() || Ar->IsError())
+	{
+		return FIoStatus(EIoErrorCode::FileNotOpen);
+	}
+
+	FHeader Header;
+	Ar->Serialize(&Header, sizeof(FHeader));
+
+	if (Header.Magic != FHeader::ExpectedMagic)
+	{
+		return FIoStatus(EIoErrorCode::CorruptToc);
+	}
+
+	OutCursorPos = Header.CursorPos;
+	TocEntries.Empty();
+	TocEntries.Reserve(Header.EntryCount);
+	*Ar << TocEntries;
+
+	return FIoStatus(EIoErrorCode::Ok);
+}
+
+FIoStatus FCacheFileToc::Save(const FString& FilePath, const uint64 CursorPos)
+{
+	IFileManager& FileMgr = IFileManager::Get();
+	TUniquePtr<FArchive> Ar(FileMgr.CreateFileWriter(*FilePath));
+
+	if (!Ar.IsValid() || Ar->IsError())
+	{
+		return FIoStatus(EIoErrorCode::FileNotOpen);
+	}
+
+	FHeader Header;
+	Header.Magic = FHeader::ExpectedMagic;
+	Header.EntryCount = TocEntries.Num();
+	Header.CursorPos = CursorPos;
+
+	Ar->Serialize(&Header, sizeof(FHeader));
+	*Ar << TocEntries;
+
+	return FIoStatus(EIoErrorCode::Ok);
+}
+
+///////////////////////////////////////////////////////////////////////////////
+enum class ECacheEntryState : uint8
+{
+	None,
+	Pending,
+	Writing,
+	Persisted
+};
+
+///////////////////////////////////////////////////////////////////////////////
+struct FCacheEntry
+	: public TIntrusiveDoubleLinkedListNode<FCacheEntry>
+{
+	FIoHash Key;
+	FIoHash Hash;
+	uint64 SerialOffset = 0;
+	uint64 SerialSize = 0;
+	FIoBuffer Data;
+	ECacheEntryState State = ECacheEntryState::None;
+};
+
+using FCacheEntryList = TIntrusiveDoubleLinkedList<FCacheEntry>;
+
+///////////////////////////////////////////////////////////////////////////////
+class FFileIoCacheRequest
+	: public FIoCacheRequestBase
+{
+public:
+	FFileIoCacheRequest(const FString& CacheFilePath, FIoReadCallback&& ReadCallback)
+		: FIoCacheRequestBase(MoveTemp(ReadCallback))
+		, FilePath(CacheFilePath)
+	{ }
+
+	void Issue(FCacheEntry&& Entry, const FIoReadOptions& Options, EAsyncIOPriorityAndFlags Priority);
+	virtual void Wait() override;
+	virtual void Cancel() override;
+
+private:
+	FString FilePath;
+	TUniquePtr<IFileHandle> FileHandle;
+};
+
+void FFileIoCacheRequest::Issue(FCacheEntry&& Entry, const FIoReadOptions& Options, EAsyncIOPriorityAndFlags Priority)
+{
+	if (Entry.Data.GetSize() > 0)
+	{
+		TRACE_CPUPROFILER_EVENT_SCOPE(FFileIoCache::ReadPendingEntry);
+		const uint64 ReadOffset = Options.GetOffset();
+		const uint64 ReadSize = FMath::Min(Options.GetSize(), Entry.Data.GetSize());
+		FIoBuffer Buffer = Options.GetTargetVa() ? FIoBuffer(FIoBuffer::Wrap, Options.GetTargetVa(), ReadSize) : FIoBuffer(ReadSize);
+		Buffer.GetMutableView().CopyFrom(Entry.Data.GetView().RightChop(ReadOffset));
+		TRACE_COUNTER_INCREMENT(FFileIoCache_GetCount);
+		CompleteRequest(MoveTemp(Buffer));
+	}
+	else
+	{
+		TRACE_CPUPROFILER_EVENT_SCOPE(FFileIoCache::ReadPersistedEntry);
+		check(Entry.SerialSize > 0);
+		check(Entry.Hash != FIoHash::Zero);
+		//TODO: Make async
+
+		const uint64 ReadSize = FMath::Min(Options.GetSize(), Entry.SerialSize);
+		const uint64 ReadOffset = Entry.SerialOffset + Options.GetOffset();
+		FIoBuffer Buffer = Options.GetTargetVa() ? FIoBuffer(FIoBuffer::Wrap, Options.GetTargetVa(), ReadSize) : FIoBuffer(ReadSize);
+
+		IPlatformFile& Ipf = IPlatformFile::GetPlatformPhysical();
+		FileHandle.Reset(Ipf.OpenRead(*FilePath, true));
+		check(FileHandle);
+		FileHandle->Seek(int64(ReadOffset));
+		FileHandle->Read(Buffer.GetData(), ReadSize);
+
+		const FIoHash& ExpectedHash = Entry.Hash;
+		if (ExpectedHash == FIoHash::HashBuffer(Buffer.GetView()))
+		{
+			TRACE_COUNTER_INCREMENT(FFileIoCache_GetCount);
+			CompleteRequest(MoveTemp(Buffer));
+		}
+		else
+		{
+			TRACE_COUNTER_INCREMENT(FFileIoCache_ErrorCount);
+			CompleteRequest(EIoErrorCode::ReadError);
+		}
+
+		FileHandle.Reset();
+	}
+}
+
+void FFileIoCacheRequest::Wait()
+{
+}
+
+void FFileIoCacheRequest::Cancel()
+{
+}
+
+////////////////////////////////////////////////////////////////////////////////
+class FCacheMap
+{
+public:
+	void SetCacheLimits(uint64 MaxPendingBytes, uint64 MaxPersistedBytes);
+	void Reset();
+	bool Contains(FIoHash Key) const;
+	bool Get(const FIoHash& Key, FCacheEntry& OutEntry) const;
+	bool InsertPending(FIoHash Key, FMemoryView Data, bool& bAdded);
+	bool RemovePending(FCacheEntryList& OutPending);
+	void InsertPersisted(FCacheEntryList&& InPersisted, const uint64 CursorPos);
+	void RemovePersisted(const uint64 RequiredSize);
+	uint64 GetPendingBytes() const { return TotalPendingBytes; }
+	FIoStatus Load(const FString& FilePath, uint64& OutCursorPos);
+	FIoStatus Save(const FString& FilePath, const uint64 CursorPos);
+
+private:
+	FCacheEntryList Pending;
+	FCacheEntryList Persisted;
+	TMap<FIoHash, TUniquePtr<FCacheEntry>> Lookup;
+	mutable FCriticalSection Cs;
+	std::atomic_uint64_t TotalPendingBytes = 0;
+	std::atomic_uint64_t TotalPersistedBytes = 0;
+	uint64 MaxPersistedBytes = 0;
+	uint64 MaxPendingBytes = 0;
+};
+
+void FCacheMap::SetCacheLimits(uint64 InMaxPendingBytes, uint64 InMaxPersistedBytes)
+{
+	MaxPendingBytes = InMaxPendingBytes;
+	MaxPersistedBytes = InMaxPersistedBytes;
+}
+
+void FCacheMap::Reset()
+{
+	Pending.Reset();
+	Persisted.Reset();
+	Lookup.Reset();
+	TotalPendingBytes = 0;
+	TotalPersistedBytes = 0;
+}
+
+bool FCacheMap::Contains(FIoHash Key) const
+{
+	FScopeLock _(&Cs);
+	return Lookup.Contains(Key);
+}
+
+bool FCacheMap::Get(const FIoHash& Key, FCacheEntry& OutEntry) const
+{
+	FScopeLock _(&Cs);
+
+	if (const TUniquePtr<FCacheEntry>* Entry = Lookup.Find(Key))
+	{
+		OutEntry = *(Entry->Get());
+		return true;
+	}
+
+	return false;
+}
+
+bool FCacheMap::InsertPending(FIoHash Key, FMemoryView Data, bool& bAdded)
+{
+	check(Data.GetSize() > 0);
+
+	bAdded = false;
+	if (TotalPendingBytes + Data.GetSize() > MaxPendingBytes)
+	{
+		TRACE_COUNTER_INCREMENT(FFileIoCache_PutRejectCount);
+		return false;
+	}
+
+	FScopeLock _(&Cs);
+
+	if (Lookup.Contains(Key))
+	{
+		TRACE_COUNTER_INCREMENT(FFileIoCache_PutExistingCount);
+		return true;
+	}
+
+	TUniquePtr<FCacheEntry>& Entry = Lookup.FindOrAdd(Key);
+	Entry = MakeUnique<FCacheEntry>();
+	Entry->Key = Key;
+	Entry->Data = FIoBuffer(FIoBuffer::Clone, Data);
+	Entry->State = ECacheEntryState::Pending;
+
+	Pending.AddTail(Entry.Get());
+
+	TotalPendingBytes += Data.GetSize();
+
+	TRACE_COUNTER_ADD(FFileIoCache_PendingBytes, Data.GetSize());
+	TRACE_COUNTER_INCREMENT(FFileIoCache_PutCount);
+
+	return bAdded = true;
+}
+
+bool FCacheMap::RemovePending(FCacheEntryList& OutPending)
+{
+	FScopeLock _(&Cs);
+
+	if (Pending.IsEmpty())
+	{
+		return false;
+	}
+
+	OutPending.AddTail(MoveTemp(Pending));
+	Pending.Reset();
+	TotalPendingBytes = 0;
+	TRACE_COUNTER_SET(FFileIoCache_PendingBytes, 0);
+
+	return true;
+}
+
+void FCacheMap::InsertPersisted(FCacheEntryList&& InPersisted, const uint64 CursorPos)
+{
+	check(InPersisted.GetTail());
+	FCacheEntry& Tail = *InPersisted.GetTail();
+	const uint64 ExpectedCursosPos = (Tail.SerialOffset + Tail.SerialSize) % MaxPersistedBytes;
+	check(ExpectedCursosPos == CursorPos);
+
+	FScopeLock _(&Cs);
+
+	uint64 PersistedBytes = 0;
+	for (FCacheEntry& Entry : InPersisted)
+	{
+		check(Entry.SerialSize > 0);
+		Entry.State = ECacheEntryState::Persisted;
+		Entry.Data = FIoBuffer();
+		PersistedBytes += Entry.SerialSize;
+	}
+
+	Persisted.AddTail(MoveTemp(InPersisted));
+	TotalPersistedBytes += PersistedBytes;
+	TRACE_COUNTER_ADD(FFileIoCache_CachedBytes, PersistedBytes);
+}
+
+void FCacheMap::RemovePersisted(const uint64 RequiredSize)
+{
+	FScopeLock _(&Cs);
+
+	uint64 RemovedBytes = 0;
+	for (;;)
+	{
+		if ((TotalPersistedBytes - RemovedBytes) + RequiredSize < MaxPersistedBytes)
+		{
+			break;
+		}
+
+		FCacheEntry* Entry = Persisted.PopHead();
+		check(Entry);
+		const uint64 EntrySize = Entry->SerialSize;
+		const FIoHash Key = Entry->Key;
+		Lookup.Remove(Key);
+		RemovedBytes += EntrySize;
+	}
+
+	TotalPersistedBytes -= RemovedBytes;
+	TRACE_COUNTER_SUBTRACT(FFileIoCache_CachedBytes, RemovedBytes);
+}
+
+FIoStatus FCacheMap::Load(const FString& FilePath, uint64& OutCursorPos)
+{
+	FCacheFileToc CacheFileToc;
+	if (FIoStatus Status = CacheFileToc.Load(FilePath, OutCursorPos); !Status.IsOk())
+	{
+		return Status;
+	}
+
+	TConstArrayView<FCacheFileToc::FTocEntry> TocEntries = CacheFileToc.GetEntries();
+	for (const FCacheFileToc::FTocEntry& Entry : TocEntries)
+	{
+		check(!Lookup.Contains(Entry.Key));
+
+		TUniquePtr<FCacheEntry>& CacheEntry = Lookup.FindOrAdd(Entry.Key);
+		CacheEntry = MakeUnique<FCacheEntry>();
+
+		CacheEntry->Key = Entry.Key;
+		CacheEntry->Hash = Entry.Hash;
+		CacheEntry->SerialOffset = Entry.SerialOffset;
+		CacheEntry->SerialSize = Entry.SerialSize;
+		CacheEntry->State = ECacheEntryState::Persisted;
+		Persisted.AddTail(CacheEntry.Get());
+		TotalPersistedBytes += Entry.SerialSize;
+	}
+
+	if (FCacheEntry* Tail = Persisted.GetTail())
+	{
+		const uint64 ExpectedCursosPos = (Tail->SerialOffset + Tail->SerialSize) % MaxPersistedBytes;
+		check(ExpectedCursosPos == OutCursorPos);
+	}
+
+	TRACE_COUNTER_SET(FFileIoCache_CachedBytes, TotalPersistedBytes);
+
+	return FIoStatus(EIoErrorCode::Ok);
+}
+
+FIoStatus FCacheMap::Save(const FString& FilePath, const uint64 CursorPos)
+{
+	FCacheFileToc CacheFileToc;
+
+	if (FCacheEntry* Tail = Persisted.GetTail())
+	{
+		const uint64 ExpectedCursosPos = (Tail->SerialOffset + Tail->SerialSize) % MaxPersistedBytes;
+		check(ExpectedCursosPos == CursorPos);
+	}
+
+	for (FCacheEntry& Entry : Persisted)
+	{
+		CacheFileToc.AddEntry(Entry.Key, Entry.Hash, Entry.SerialOffset, Entry.SerialSize);
+	}
+
+	return CacheFileToc.Save(FilePath, CursorPos);
+}
+
+////////////////////////////////////////////////////////////////////////////////
+class FFileIoCache final
+	: public FRunnable
+	, public IIoCache
+{
+public:
+	explicit FFileIoCache(const FFileIoCacheConfig& Config);
+	virtual ~FFileIoCache();
+
+	virtual bool ContainsChunk(const FIoHash& Key) const override;
+	virtual FIoCacheRequest GetChunk(const FIoHash& Key, const FIoReadOptions& Options, FIoReadCallback&& Callback) override;
+	virtual FIoStatus PutChunk(const FIoHash& Key, FMemoryView Data) override;
+
+	// Runnable
+	virtual bool Init() override
+	{ 
+		return true;
+	}
+
+	virtual uint32 Run() override
+	{
+		return FileWriterThreadEntry();
+	}
+
+	virtual void Stop() override
+	{
+		bStopRequested = true;
+		TickWriterEvent->Trigger();
+	}
+
+private:
+	void Initialize();
+	void Shutdown();
+	uint32 FileWriterThreadEntry();
+
+	FFileIoCacheConfig CacheConfig;
+	FCacheMap CacheMap;
+	TUniquePtr<FRunnableThread> WriterThread;
+	TUniquePtr<IFileHandle> WriteFileHandle;
+	FEventRef TickWriterEvent;
+	FString CacheFilePath;
+	std::atomic_bool bStopRequested{false};
+};
+
+FFileIoCache::FFileIoCache(const FFileIoCacheConfig& Config)
+	: CacheConfig(Config)
+{
+	CacheMap.SetCacheLimits(Config.MemoryStorageSize, Config.DiskStorageSize);
+	Initialize();
+}
+
+FFileIoCache::~FFileIoCache()
+{
+	Shutdown();
+}
+
+bool FFileIoCache::ContainsChunk(const FIoHash& Key) const
+{
+	return CacheMap.Contains(Key);
+}
+
+FIoCacheRequest FFileIoCache::GetChunk(const FIoHash& Key, const FIoReadOptions& Options, FIoReadCallback&& Callback)
+{
+	FCacheEntry Entry;
+	if (CacheMap.Get(Key, Entry))
+	{
+		TUniquePtr<FFileIoCacheRequest> Request = MakeUnique<FFileIoCacheRequest>(CacheFilePath, MoveTemp(Callback));
+		Request->Issue(MoveTemp(Entry), Options, EAsyncIOPriorityAndFlags::AIOP_Normal);
+
+		return FIoCacheRequest(Request.Release());
+	}
+
+	return FIoCacheRequest();
+}
+	
+FIoStatus FFileIoCache::PutChunk(const FIoHash& Key, FMemoryView Data)
+{
+	bool bAdded = false;
+	if (CacheMap.InsertPending(Key, Data, bAdded))
+	{
+		if (bAdded)
+		{
+			TickWriterEvent->Trigger();
+		}
+		return FIoStatus::Ok;
+	}
+
+	return FIoStatus(EIoErrorCode::Unknown);
+}
+
+void FFileIoCache::Initialize()
+{
+	UE_LOG(LogIoCache, Log,
+		TEXT("Initializing file I/O cache, disk size %lluB, memory size %lluB"), CacheConfig.DiskStorageSize, CacheConfig.MemoryStorageSize);
+
+	const FString CacheDir = FPaths::ProjectPersistentDownloadDir() / TEXT("IoCache");
+	const FString CacheTocPath = CacheDir / TEXT("cache.utoc");
+	CacheFilePath = CacheDir / TEXT("cache.ucas");
+
+	IPlatformFile& Ipf = IPlatformFile::GetPlatformPhysical();
+	IFileManager& FileMgr = IFileManager::Get();
+	
+	FIoStatus CacheStatus(EIoErrorCode::Unknown);
+	if (FileMgr.FileExists(*CacheTocPath))
+	{
+		if (FParse::Param(FCommandLine::Get(), TEXT("ClearIoCache")))
+		{
+			UE_LOG(LogIoCache, Log, TEXT("Deleting cache file '%s'"), *CacheFilePath);
+			FileMgr.Delete(*CacheFilePath);
+		}
+		else
+		{
+			uint64 CursorPos = ~uint64(0);
+
+			CacheStatus = CacheMap.Load(CacheTocPath, CursorPos);
+			if (CacheStatus.IsOk())
+			{
+				check(CursorPos != ~uint64(0));
+				
+				UE_LOG(LogIoCache, Log, TEXT("Loaded TOC '%s'"), *CacheTocPath);
+				if (FileMgr.FileExists(*CacheFilePath))
+				{
+					WriteFileHandle.Reset(Ipf.OpenWrite(*CacheFilePath, true, true));
+					WriteFileHandle->Seek(CursorPos);
+					//TODO: Integrity check
+				}
+				else
+				{
+					UE_LOG(LogIoCache, Warning, TEXT("Failed to open cache file ''"), *CacheFilePath);
+					CacheStatus = FIoStatus(EIoErrorCode::FileNotOpen);
+				}
+			}
+			else
+			{
+				UE_LOG(LogIoCache, Warning, TEXT("Failed to load TOC '%s'"), *CacheTocPath);
+			}
+		}
+	}
+
+	if (!CacheStatus.IsOk())
+	{
+		CacheMap.Reset();
+		FileMgr.Delete(*CacheFilePath);
+
+		if (!FileMgr.DirectoryExists(*CacheDir))
+		{
+			FileMgr.MakeDirectory(*CacheDir, true);
+		}
+		WriteFileHandle.Reset(Ipf.OpenWrite(*CacheFilePath, true, true));
+	}
+
+	WriterThread.Reset(FRunnableThread::Create(this, TEXT("File I/O Cache"), 0, TPri_BelowNormal));
+}
+
+void FFileIoCache::Shutdown()
+{
+	if (bStopRequested)
+	{
+		return;
+	}
+
+	bStopRequested = true;
+	TickWriterEvent->Trigger();
+	WriterThread->Kill();
+
+	const FString CacheTocPath = FPaths::ProjectPersistentDownloadDir() / TEXT("IoCache") / TEXT("cache.utoc");
+	UE_LOG(LogIoCache, Log, TEXT("Saving TOC '%s'"), *CacheTocPath);
+	CacheMap.Save(CacheTocPath, WriteFileHandle->Tell());
+	WriteFileHandle.Reset();
+}
+
+uint32 FFileIoCache::FileWriterThreadEntry()
+{
+	while (!bStopRequested)
+	{
+		for (;;)
+		{
+			FCacheEntryList Entries;
+			if (CacheMap.RemovePending(Entries) == false)
+			{
+				break;
+			}
+
+			TRACE_CPUPROFILER_EVENT_SCOPE(FFileIoCache::WriteCacheEntry);
+
+			uint64 TotalPendingSize = 0;
+			for (FCacheEntry& Entry : Entries)
+			{
+				Entry.State = ECacheEntryState::Writing;
+				TotalPendingSize += Entry.Data.GetSize();
+			}
+			CacheMap.RemovePersisted(TotalPendingSize);
+
+			//TODO: Write bigger chunks
+			for (FCacheEntry& Entry : Entries)
+			{
+				const FIoBuffer& Data = Entry.Data;
+				check(Data.GetSize() > 0);
+
+				Entry.SerialOffset = WriteFileHandle->Tell();
+				Entry.SerialSize = Data.GetSize();
+				Entry.Hash = FIoHash::HashBuffer(Data.GetView());
+
+				const uint64 RemainingDiskSize = CacheConfig.DiskStorageSize - WriteFileHandle->Tell();
+				const uint64 ByteCount = FMath::Min(Entry.Data.GetSize(), RemainingDiskSize);
+				WriteFileHandle->Write(Entry.Data.GetData(), ByteCount);
+
+				const uint64 RemainingChunkSize = Entry.Data.GetSize() - ByteCount;
+				if (RemainingChunkSize > 0)
+				{
+					WriteFileHandle->Seek(0);
+					WriteFileHandle->Write(Entry.Data.GetData() + ByteCount, RemainingChunkSize);
+				}
+
+				WriteFileHandle->Flush();
+			}
+
+			CacheMap.InsertPersisted(MoveTemp(Entries), WriteFileHandle->Tell());
+		}
+
+		if (bStopRequested == false)
+		{
+			TickWriterEvent->Wait();
+		}
+	}
+
+	return 0;
+}
+
+} // namespace UE::IO::Private
+
+TUniquePtr<IIoCache> MakeFileIoCache(const FFileIoCacheConfig& Config)
+{
+	return MakeUnique<UE::IO::Private::FFileIoCache>(Config);
+}
diff --git a/Engine/Source/Runtime/Core/Public/IO/IoCache.h b/Engine/Source/Runtime/Core/Public/IO/IoCache.h
--- a/Engine/Source/Runtime/Core/Public/IO/IoCache.h
+++ b/Engine/Source/Runtime/Core/Public/IO/IoCache.h
@@ -7,6 +7,8 @@
 #include "Templates/RefCounting.h"
 #include <atomic>
 
+CORE_API DECLARE_LOG_CATEGORY_EXTERN(LogIoCache, Log, All);
+
 namespace UE::IO::Private
 {
 	/** An abstract base class for implementing cache requests. */
