Allow setting external profiler with environment variable UE_EXTERNAL_PROFILER

diff --git a/Engine/Source/Runtime/Core/Private/ProfilingDebugging/ExternalProfiler.cpp b/Engine/Source/Runtime/Core/Private/ProfilingDebugging/ExternalProfiler.cpp
--- a/Engine/Source/Runtime/Core/Private/ProfilingDebugging/ExternalProfiler.cpp
+++ b/Engine/Source/Runtime/Core/Private/ProfilingDebugging/ExternalProfiler.cpp
@@ -1,6 +1,7 @@
 // Copyright Epic Games, Inc. All Rights Reserved.
 
 #include "ProfilingDebugging/ExternalProfiler.h"
+#include "Algo/Find.h"
 #include "Logging/LogMacros.h"
 #include "Misc/Parse.h"
 #include "Misc/CommandLine.h"
@@ -80,6 +81,17 @@ FExternalProfiler* FActiveExternalProfilerBase::InitActiveProfiler()
 			UE_LOG(LogExternalProfiler, Log, TEXT("No external profilers were discovered.  External profiling features will not be available."));
 		}
 #endif
+		
+		if (ActiveProfiler == nullptr)
+		{
+			FString ProfilerName = FPlatformMisc::GetEnvironmentVariable(TEXT("UE_EXTERNAL_PROFILER"));
+			if (!ProfilerName.IsEmpty())
+			{
+				FExternalProfiler** Found = Algo::FindByPredicate(AvailableProfilers,
+					 [&ProfilerName](FExternalProfiler* Profiler) { return ProfilerName == Profiler->GetProfilerName();});
+				ActiveProfiler = Found ? *Found : nullptr;
+			}
+		}
 
 		// Don't try to initialize again this session
 		bDidInitialize = true;
