I/O - simplified cache interface and optimized chunk requests

- replaced cache async API with TTask for easier composition with HTTP request
- combine I/O request that fetches the same chunk into one request

diff --git a/Engine/Source/Runtime/Core/Private/IO/IoCache.cpp b/Engine/Source/Runtime/Core/Private/IO/IoCache.cpp
--- a/Engine/Source/Runtime/Core/Private/IO/IoCache.cpp
+++ b/Engine/Source/Runtime/Core/Private/IO/IoCache.cpp
@@ -3,67 +3,3 @@
 #include "IO/IoCache.h"
 
 DEFINE_LOG_CATEGORY(LogIoCache);
-
-namespace UE::IO::Private
-{
-	FIoCacheRequestBase::FIoCacheRequestBase(FIoReadCallback&& ReadCallback)
-		: Callback(MoveTemp(ReadCallback))
-	{
-	}
-
-	void FIoCacheRequestBase::CompleteRequest(FIoBuffer&& Buffer)
-	{
-		if (EIoErrorCode::Unknown == ErrorCode.load(std::memory_order_seq_cst))
-		{
-			Callback(TIoStatusOr<FIoBuffer>(MoveTemp(Buffer)));
-			ErrorCode = EIoErrorCode::Ok;
-		}
-	}
-
-	void FIoCacheRequestBase::CompleteRequest(EIoErrorCode Error)
-	{
-		if (EIoErrorCode::Unknown == ErrorCode.load(std::memory_order_seq_cst))
-		{
-			Callback(FIoStatus(Error));
-			ErrorCode = Error; 
-		}
-	}
-
-} // namespace UE::IO::Private
-
-FIoCacheRequest::FIoCacheRequest(UE::IO::Private::FIoCacheRequestBase* Base)
-	: Pimpl(Base)
-{
-}
-
-FIoCacheRequest::~FIoCacheRequest()
-{
-	if (IsValid())
-	{
-		if (!Status().IsCompleted())
-		{
-			Pimpl->Cancel();
-			Pimpl->Wait();
-		}
-	}
-}
-
-void FIoCacheRequest::Cancel()
-{
-	check(IsValid());
-
-	if (!Status().IsCompleted())
-	{
-		Pimpl->Cancel();
-	}
-}
-
-void FIoCacheRequest::Wait()
-{
-	check(IsValid());
-
-	if (!Status().IsCompleted())
-	{
-		Pimpl->Wait();
-	}
-}
diff --git a/Engine/Source/Runtime/Core/Private/IO/IoFileCache.cpp b/Engine/Source/Runtime/Core/Private/IO/IoFileCache.cpp
--- a/Engine/Source/Runtime/Core/Private/IO/IoFileCache.cpp
+++ b/Engine/Source/Runtime/Core/Private/IO/IoFileCache.cpp
@@ -1,8 +1,6 @@
 // Copyright Epic Games, Inc. All Rights Reserved.
 
 #include "IO/IoFileCache.h"
-#include "Async/AsyncFileHandle.h"
-#include "Containers/ContainersFwd.h"
 #include "Containers/IntrusiveDoubleLinkedList.h"
 #include "GenericPlatform/GenericPlatformFile.h"
 #include "IO/IoCache.h"
@@ -160,78 +158,6 @@ struct FCacheEntry
 
 using FCacheEntryList = TIntrusiveDoubleLinkedList<FCacheEntry>;
 
-///////////////////////////////////////////////////////////////////////////////
-class FFileIoCacheRequest
-	: public FIoCacheRequestBase
-{
-public:
-	FFileIoCacheRequest(const FString& CacheFilePath, FIoReadCallback&& ReadCallback)
-		: FIoCacheRequestBase(MoveTemp(ReadCallback))
-		, FilePath(CacheFilePath)
-	{ }
-
-	void Issue(FCacheEntry&& Entry, const FIoReadOptions& Options, EAsyncIOPriorityAndFlags Priority);
-	virtual void Wait() override;
-	virtual void Cancel() override;
-
-private:
-	FString FilePath;
-	TUniquePtr<IFileHandle> FileHandle;
-};
-
-void FFileIoCacheRequest::Issue(FCacheEntry&& Entry, const FIoReadOptions& Options, EAsyncIOPriorityAndFlags Priority)
-{
-	if (Entry.Data.GetSize() > 0)
-	{
-		TRACE_CPUPROFILER_EVENT_SCOPE(FFileIoCache::ReadPendingEntry);
-		const uint64 ReadOffset = Options.GetOffset();
-		const uint64 ReadSize = FMath::Min(Options.GetSize(), Entry.Data.GetSize());
-		FIoBuffer Buffer = Options.GetTargetVa() ? FIoBuffer(FIoBuffer::Wrap, Options.GetTargetVa(), ReadSize) : FIoBuffer(ReadSize);
-		Buffer.GetMutableView().CopyFrom(Entry.Data.GetView().RightChop(ReadOffset));
-		TRACE_COUNTER_INCREMENT(FFileIoCache_GetCount);
-		CompleteRequest(MoveTemp(Buffer));
-	}
-	else
-	{
-		TRACE_CPUPROFILER_EVENT_SCOPE(FFileIoCache::ReadPersistedEntry);
-		check(Entry.SerialSize > 0);
-		check(Entry.Hash != FIoHash::Zero);
-		//TODO: Make async
-
-		const uint64 ReadSize = FMath::Min(Options.GetSize(), Entry.SerialSize);
-		const uint64 ReadOffset = Entry.SerialOffset + Options.GetOffset();
-		FIoBuffer Buffer = Options.GetTargetVa() ? FIoBuffer(FIoBuffer::Wrap, Options.GetTargetVa(), ReadSize) : FIoBuffer(ReadSize);
-
-		IPlatformFile& Ipf = IPlatformFile::GetPlatformPhysical();
-		FileHandle.Reset(Ipf.OpenRead(*FilePath, true));
-		check(FileHandle);
-		FileHandle->Seek(int64(ReadOffset));
-		FileHandle->Read(Buffer.GetData(), ReadSize);
-
-		const FIoHash& ExpectedHash = Entry.Hash;
-		if (ExpectedHash == FIoHash::HashBuffer(Buffer.GetView()))
-		{
-			TRACE_COUNTER_INCREMENT(FFileIoCache_GetCount);
-			CompleteRequest(MoveTemp(Buffer));
-		}
-		else
-		{
-			TRACE_COUNTER_INCREMENT(FFileIoCache_ErrorCount);
-			CompleteRequest(EIoErrorCode::ReadError);
-		}
-
-		FileHandle.Reset();
-	}
-}
-
-void FFileIoCacheRequest::Wait()
-{
-}
-
-void FFileIoCacheRequest::Cancel()
-{
-}
-
 ////////////////////////////////////////////////////////////////////////////////
 class FCacheMap
 {
@@ -446,6 +372,8 @@ FIoStatus FCacheMap::Save(const FString& FilePath, const uint64 CursorPos)
 	return CacheFileToc.Save(FilePath, CursorPos);
 }
 
+using namespace UE::Tasks;
+
 ////////////////////////////////////////////////////////////////////////////////
 class FFileIoCache final
 	: public FRunnable
@@ -456,7 +384,7 @@ public:
 	virtual ~FFileIoCache();
 
 	virtual bool ContainsChunk(const FIoHash& Key) const override;
-	virtual FIoCacheRequest GetChunk(const FIoHash& Key, const FIoReadOptions& Options, FIoReadCallback&& Callback) override;
+	virtual TTask<TIoStatusOr<FIoBuffer>> GetChunk(const FIoHash& Key, const FIoReadOptions& Options, const FIoCancellationToken* CancellationToken) override;
 	virtual FIoStatus PutChunk(const FIoHash& Key, FMemoryView Data) override;
 
 	// Runnable
@@ -507,18 +435,62 @@ bool FFileIoCache::ContainsChunk(const FIoHash& Key) const
 	return CacheMap.Contains(Key);
 }
 
-FIoCacheRequest FFileIoCache::GetChunk(const FIoHash& Key, const FIoReadOptions& Options, FIoReadCallback&& Callback)
+TTask<TIoStatusOr<FIoBuffer>> FFileIoCache::GetChunk(const FIoHash& Key, const FIoReadOptions& Options, const FIoCancellationToken* CancellationToken)
 {
-	FCacheEntry Entry;
-	if (CacheMap.Get(Key, Entry))
-	{
-		TUniquePtr<FFileIoCacheRequest> Request = MakeUnique<FFileIoCacheRequest>(CacheFilePath, MoveTemp(Callback));
-		Request->Issue(MoveTemp(Entry), Options, EAsyncIOPriorityAndFlags::AIOP_Normal);
+	return Launch(UE_SOURCE_LOCATION,
+		[this, Key, Options, CancellationToken]()
+		{
+			FCacheEntry Entry;
+			if (!CacheMap.Get(Key, Entry))
+			{
+				return TIoStatusOr<FIoBuffer>(FIoStatus(EIoErrorCode::Unknown));
+			}
 
-		return FIoCacheRequest(Request.Release());
-	}
+			if (CancellationToken && CancellationToken->IsCancelled())
+			{
+				return TIoStatusOr<FIoBuffer>(FIoStatus(EIoErrorCode::Cancelled));
+			}
+
+			if (Entry.Data.GetSize() > 0)
+			{
+				TRACE_CPUPROFILER_EVENT_SCOPE(FFileIoCache::ReadPendingEntry);
+
+				const uint64 ReadOffset = Options.GetOffset();
+				const uint64 ReadSize = FMath::Min(Options.GetSize(), Entry.Data.GetSize());
+				FIoBuffer Buffer = Options.GetTargetVa() ? FIoBuffer(FIoBuffer::Wrap, Options.GetTargetVa(), ReadSize) : FIoBuffer(ReadSize);
+				Buffer.GetMutableView().CopyFrom(Entry.Data.GetView().RightChop(ReadOffset));
+				TRACE_COUNTER_INCREMENT(FFileIoCache_GetCount);
+
+				return TIoStatusOr<FIoBuffer>(Buffer);
+			}
+			else
+			{
+				TRACE_CPUPROFILER_EVENT_SCOPE(FFileIoCache::ReadPersistedEntry);
+
+				check(Entry.SerialSize > 0);
+				check(Entry.Hash != FIoHash::Zero);
 
-	return FIoCacheRequest();
+				const uint64 ReadSize = FMath::Min(Options.GetSize(), Entry.SerialSize);
+				const uint64 ReadOffset = Entry.SerialOffset + Options.GetOffset();
+				FIoBuffer Buffer = Options.GetTargetVa() ? FIoBuffer(FIoBuffer::Wrap, Options.GetTargetVa(), ReadSize) : FIoBuffer(ReadSize);
+
+				IPlatformFile& Ipf = IPlatformFile::GetPlatformPhysical();
+				TUniquePtr<IFileHandle> FileHandle(Ipf.OpenRead(*CacheFilePath, true));
+				check(FileHandle);
+				FileHandle->Seek(int64(ReadOffset));
+				FileHandle->Read(Buffer.GetData(), ReadSize);
+
+				const FIoHash& ExpectedHash = Entry.Hash;
+				if (ExpectedHash == FIoHash::HashBuffer(Buffer.GetView()))
+				{
+					TRACE_COUNTER_INCREMENT(FFileIoCache_GetCount);
+					return TIoStatusOr<FIoBuffer>(Buffer);
+				}
+
+				TRACE_COUNTER_INCREMENT(FFileIoCache_ErrorCount);
+				return TIoStatusOr<FIoBuffer>(FIoStatus(EIoErrorCode::ReadError));
+			}
+		});
 }
 	
 FIoStatus FFileIoCache::PutChunk(const FIoHash& Key, FMemoryView Data)
diff --git a/Engine/Source/Runtime/Core/Private/IO/IoMemoryCache.cpp b/Engine/Source/Runtime/Core/Private/IO/IoMemoryCache.cpp
--- a/Engine/Source/Runtime/Core/Private/IO/IoMemoryCache.cpp
+++ b/Engine/Source/Runtime/Core/Private/IO/IoMemoryCache.cpp
@@ -7,65 +7,13 @@
 #include "IO/IoDispatcher.h"
 #include "IO/IoHash.h"
 #include "Misc/ScopeLock.h"
-#include "Tasks/Task.h"
 
 #if !UE_BUILD_SHIPPING
 
 namespace UE::IO::Private
 {
 
-////////////////////////////////////////////////////////////////////////////////
-class FMemoryIoCacheRequest
-	: public FIoCacheRequestBase
-{
-public:
-	FMemoryIoCacheRequest(FIoReadCallback&& ReadCallback, FIoBuffer CachedBuffer);
-	virtual ~FMemoryIoCacheRequest() = default;
-	virtual void Wait() override;
-	virtual void Cancel() override;
-
-	void Issue();
-
-private:
-	FIoBuffer Buffer;
-	UE::Tasks::TTask<void> Task;
-	std::atomic_bool bCanceled{false};
-};
-
-FMemoryIoCacheRequest::FMemoryIoCacheRequest(FIoReadCallback&& ReadCallback, FIoBuffer CachedBuffer)
-	: FIoCacheRequestBase(MoveTemp(ReadCallback))
-	, Buffer(CachedBuffer)
-{
-}
-
-void FMemoryIoCacheRequest::Wait() 
-{
-	Task.Wait();
-}
-
-void FMemoryIoCacheRequest::Cancel() 
-{
-	bCanceled = true;
-}
-
-void FMemoryIoCacheRequest::Issue()
-{
-	Task = UE::Tasks::Launch(TEXT("Memory I/O Cache"), [this]()
-	{
-		if (bCanceled)
-		{
-			CompleteRequest(EIoErrorCode::Cancelled);
-		}
-		else if (Buffer.GetSize() > 0)
-		{
-			CompleteRequest(MoveTemp(Buffer));
-		}
-		else
-		{
-			CompleteRequest(EIoErrorCode::NotFound);
-		}
-	});
-}
+using namespace UE::Tasks;
 
 ////////////////////////////////////////////////////////////////////////////////
 class FMemoryIoCache
@@ -85,7 +33,7 @@ public:
 	explicit FMemoryIoCache(uint64 CacheSize);
 	virtual ~FMemoryIoCache();
 	virtual bool ContainsChunk(const FIoHash& Key) const override;
-	virtual FIoCacheRequest GetChunk(const FIoHash& Key, const FIoReadOptions& Options, FIoReadCallback&& Callback) override;
+	virtual TTask<TIoStatusOr<FIoBuffer>> GetChunk(const FIoHash& Key, const FIoReadOptions& Options, const FIoCancellationToken* CancellationToken) override;
 	virtual FIoStatus PutChunk(const FIoHash& Key, FMemoryView Data) override;
 
 private:
@@ -112,25 +60,33 @@ bool FMemoryIoCache::ContainsChunk(const FIoHash& Key) const
 	return Lookup.Contains(Key);
 }
 
-FIoCacheRequest FMemoryIoCache::GetChunk(const FIoHash& Key, const FIoReadOptions& Options, FIoReadCallback&& Callback)
+TTask<TIoStatusOr<FIoBuffer>> FMemoryIoCache::GetChunk(const FIoHash& Key, const FIoReadOptions& Options, const FIoCancellationToken* CancellationToken)
 {
-	FIoBuffer Buffer;
+	return UE::Tasks::Launch(UE_SOURCE_LOCATION, [this, Key, Options, CancellationToken]()
 	{
 		FScopeLock _(&CacheCS);
 
 		if (TUniquePtr<FCacheEntry>* Entry = Lookup.Find(Key))
 		{
 			FCacheEntry* CacheEntry = Entry->Get();
-			Buffer = CacheEntry->Buffer;
 			Lru.Remove(CacheEntry);
 			Lru.AddHead(CacheEntry);
-		}
-	}
 
-	TUniquePtr<FMemoryIoCacheRequest> Request = MakeUnique<FMemoryIoCacheRequest>(MoveTemp(Callback), Buffer);
-	Request->Issue();
+			if (CancellationToken && CancellationToken->IsCancelled())
+			{
+				return TIoStatusOr<FIoBuffer>(FIoStatus(EIoErrorCode::Cancelled));
+			}
 
-	return FIoCacheRequest(Request.Release());
+			const uint64 ReadOffset = Options.GetOffset();
+			const uint64 ReadSize = FMath::Min(Options.GetSize(), CacheEntry->Buffer.GetSize());
+			FIoBuffer Buffer = Options.GetTargetVa() ? FIoBuffer(FIoBuffer::Wrap, Options.GetTargetVa(), ReadSize) : FIoBuffer(ReadSize);
+			Buffer.GetMutableView().CopyFrom(CacheEntry->Buffer.GetView().RightChop(ReadOffset));
+
+			return TIoStatusOr<FIoBuffer>(Buffer);
+		}
+
+		return TIoStatusOr<FIoBuffer>(FIoStatus(EIoErrorCode::NotFound));
+	});
 }
 
 FIoStatus FMemoryIoCache::PutChunk(const FIoHash& Key, FMemoryView Data)
diff --git a/Engine/Source/Runtime/Core/Public/IO/IoCache.h b/Engine/Source/Runtime/Core/Public/IO/IoCache.h
--- a/Engine/Source/Runtime/Core/Public/IO/IoCache.h
+++ b/Engine/Source/Runtime/Core/Public/IO/IoCache.h
@@ -2,97 +2,32 @@
 
 #pragma once
 
-#include "IO/IoDispatcher.h"
+#include "IO/IoStatus.h"
+#include "IO/IoCancellationToken.h"
 #include "Memory/MemoryFwd.h"
-#include "Templates/RefCounting.h"
+#include "Tasks/Task.h"
 #include <atomic>
 
-CORE_API DECLARE_LOG_CATEGORY_EXTERN(LogIoCache, Log, All);
-
-namespace UE::IO::Private
-{
-	/** An abstract base class for implementing cache requests. */
-	class FIoCacheRequestBase
-	{
-	public:
-		virtual ~FIoCacheRequestBase() = default;
-
-		virtual void Wait() = 0;
-		virtual void Cancel() = 0;
-
-		inline FIoStatus Status()
-		{ 
-			return FIoStatus(ErrorCode.load(std::memory_order_relaxed));
-		}
-
-		inline uint32 GetRefCount() const
-		{
-			return RefCount.load(std::memory_order_relaxed);
-		}
+class FIoBuffer;
+class FIoReadOptions;
+struct FIoHash;
 
-		inline void AddRef()
-		{
-			RefCount.fetch_add(1, std::memory_order_relaxed);
-		}
-
-		inline void Release()
-		{
-			if (1 == RefCount.fetch_sub(1, std::memory_order_acq_rel))
-			{
-				delete this;
-			}
-		}
-
-	protected:
-		explicit FIoCacheRequestBase(FIoReadCallback&& ReadCallback);
-		void CompleteRequest(FIoBuffer&& Buffer);
-		void CompleteRequest(EIoErrorCode Error);
-
-		FIoReadCallback Callback;
-		std::atomic_uint32_t RefCount{0};
-		std::atomic<EIoErrorCode> ErrorCode{EIoErrorCode::Unknown};
-	};
-}
-
-/** Represents a pending I/O cache request. */
-class CORE_API FIoCacheRequest
-{
-public:
-	FIoCacheRequest() = default;
-	FIoCacheRequest(const FIoCacheRequest&) = delete;
-	FIoCacheRequest(FIoCacheRequest&& Other) = default;
-	~FIoCacheRequest();
-	
-	FIoCacheRequest& operator=(const FIoCacheRequest&) = delete;
-	FIoCacheRequest& operator=(FIoCacheRequest&& Other) = default;
-
-	explicit FIoCacheRequest(UE::IO::Private::FIoCacheRequestBase* Base);
-
-	bool IsValid() const
-	{ 
-		return Pimpl.IsValid();
-	}
-
-	FIoStatus Status()
-	{
-		return IsValid() ? Pimpl->Status() : FIoStatus(EIoErrorCode::InvalidCode);
-	}
-
-	void Wait();
-	void Cancel();
-
-private:
-	TRefCountPtr<UE::IO::Private::FIoCacheRequestBase> Pimpl;
-};
+CORE_API DECLARE_LOG_CATEGORY_EXTERN(LogIoCache, Log, All);
 
-/**
- * Interface for retrieving and storing I/O chunks identified by a 20 byte cache key.
- */
+/** Cache for binary blobs with a 20 byte cache key. */
 class IIoCache
 {
 public:
 	virtual ~IIoCache() = default;
+	/** Returns whether the specified cache key is present in the cache. */
 	virtual bool ContainsChunk(const FIoHash& Key) const = 0;
-	virtual FIoCacheRequest GetChunk(const FIoHash& Key, const FIoReadOptions& Options, FIoReadCallback&& Callback) = 0;
+
+	/** Get the chunk associated with the specified cache key. */
+	virtual UE::Tasks::TTask<TIoStatusOr<FIoBuffer>> GetChunk(
+		const FIoHash& Key,
+		const FIoReadOptions& Options,
+		const FIoCancellationToken* CancellationToken) = 0;
+
+	/** Insert a new chunk into the cache. */
 	virtual FIoStatus PutChunk(const FIoHash& Key, FMemoryView Data) = 0;
 };
diff --git a/Engine/Source/Runtime/Core/Public/IO/IoCancellationToken.h b/Engine/Source/Runtime/Core/Public/IO/IoCancellationToken.h
new file mode 100644
--- /dev/null
+++ b/Engine/Source/Runtime/Core/Public/IO/IoCancellationToken.h
@@ -0,0 +1,33 @@
+// Copyright Epic Games, Inc. All Rights Reserved.
+
+#pragma once
+
+#include <atomic>
+
+/**
+ * Experimental type for signaling cancellation of asynchronouse I/O operations.
+ */
+struct CORE_API FIoCancellationToken
+{
+	FIoCancellationToken() = default;
+	FIoCancellationToken(const FIoCancellationToken&) = delete;
+	FIoCancellationToken(FIoCancellationToken&&) = delete;
+
+	FIoCancellationToken& operator=(const FIoCancellationToken&) = delete;
+	FIoCancellationToken& operator=(FIoCancellationToken&&) = delete;
+
+	/** Returns true if the operation should cancel or not. */
+	bool IsCancelled() const
+	{ 
+		return bCancellationRequested;
+	}
+
+	/** Signal the asynchronouse operation to cancel. */
+	void Cancel()
+	{
+		bCancellationRequested = true;
+	}
+
+private:
+	std::atomic_bool bCancellationRequested{false};
+};
