Don't override max influences computed, otherwise alt skin weights may fail to upload.

#rnx

diff --git a/Engine/Source/Runtime/Core/Private/UObject/FortniteMainBranchObjectVersion.cpp b/Engine/Source/Runtime/Core/Private/UObject/FortniteMainBranchObjectVersion.cpp
--- a/Engine/Source/Runtime/Core/Private/UObject/FortniteMainBranchObjectVersion.cpp
+++ b/Engine/Source/Runtime/Core/Private/UObject/FortniteMainBranchObjectVersion.cpp
@@ -12,7 +12,7 @@ TMap<FGuid, FGuid> FFortniteMainBranchObjectVersion::GetSystemGuids()
 	SystemGuids.Add(DevGuids.NANITE_DERIVEDDATA_VER, FGuid("D28C003F825F4C178F698938FB0E47F7"));
 	SystemGuids.Add(DevGuids.NIAGARASHADERMAP_DERIVEDDATA_VER, FGuid("35D7960B72164B49B9B7B3FB5857E674"));
 	SystemGuids.Add(DevGuids.Niagara_LatestScriptCompileVersion, FGuid("D74C8B1657194273B102805F71528F94"));
-	SystemGuids.Add(DevGuids.SkeletalMeshDerivedDataVersion, FGuid("87620CAAE42ACAA2D32409D2AD939201"));
+	SystemGuids.Add(DevGuids.SkeletalMeshDerivedDataVersion, FGuid("518D7C51D31D9BE4A9F4B81542AC8119"));
 	SystemGuids.Add(DevGuids.STATICMESH_DERIVEDDATA_VER, FGuid("543843B103794E0E9BA4BE80FB602F79"));
 	return SystemGuids;
 }
