IoDispatcher: Improve shutdown logic.
First reset the GIoDispatcher pointer.
Then stop the IoDispatcher thread and wait for all pending requests to finish.
Then shutdown the backends.
Then destroy the IoDispatcher itself.

#rnx

diff --git a/Engine/Source/Runtime/Core/Private/IO/IoDispatcher.cpp b/Engine/Source/Runtime/Core/Private/IO/IoDispatcher.cpp
--- a/Engine/Source/Runtime/Core/Private/IO/IoDispatcher.cpp
+++ b/Engine/Source/Runtime/Core/Private/IO/IoDispatcher.cpp
@@ -346,15 +346,13 @@ public:
 
 	~FIoDispatcherImpl()
 	{
+		delete Thread;
 		for (const FBackendAndPriority& Backend : Backends)
 		{
 			Backend.Value->Shutdown();
 		}
 		FCoreDelegates::GetMemoryTrimDelegate().Remove(MemoryTrimDelegateHandle);
 		BackendContext->WakeUpDispatcherThreadDelegate.Unbind();
-		// when all mounted backends have been shutdown and the delegates have been cleared,
-		// then we can go ahead and delete the resolve thread
-		delete Thread;
 		FPlatformProcess::ReturnSynchEventToPool(DispatcherEvent);
 		RequestAllocator->ReleaseRef();
 	}
@@ -829,6 +827,12 @@ private:
 			ProcessIncomingRequests();
 			ProcessCompletedRequests();
 		}
+
+		ProcessIncomingRequests();
+		while (PendingIoRequestsCount > 0)
+		{
+			ProcessCompletedRequests();
+		}
 		return 0;
 	}
 
@@ -961,7 +965,7 @@ FIoDispatcher::InitializePostSettings()
 void
 FIoDispatcher::Shutdown()
 {
-	GIoDispatcher.Reset();
+	TUniquePtr<FIoDispatcher> LocalIoDispatcher(MoveTemp(GIoDispatcher));
 }
 
 FIoDispatcher&
diff --git a/Engine/Source/Runtime/PakFile/Private/IoDispatcherFileBackend.cpp b/Engine/Source/Runtime/PakFile/Private/IoDispatcherFileBackend.cpp
--- a/Engine/Source/Runtime/PakFile/Private/IoDispatcherFileBackend.cpp
+++ b/Engine/Source/Runtime/PakFile/Private/IoDispatcherFileBackend.cpp
@@ -1829,7 +1829,7 @@ FIoRequestImpl* FFileIoStore::GetCompletedRequests()
 		const bool bScatterAsync = bIsMultithreaded && (!BlockToDecompress->CompressionMethod.IsNone() || BlockToDecompress->EncryptionKey.IsValid() || BlockToDecompress->SignatureHash);
 		if (bScatterAsync)
 		{
-			BlockToDecompress->CompressionContext->Task = TGraphTask<FDecompressAsyncTask>::CreateTask().ConstructAndDispatchWhenReady(*this, BlockToDecompress);
+			TGraphTask<FDecompressAsyncTask>::CreateTask().ConstructAndDispatchWhenReady(*this, BlockToDecompress);
 		}
 		else
 		{
@@ -2058,15 +2058,6 @@ uint32 FFileIoStore::Run()
 			PlatformImpl->ServiceWait();
 		}
 	}
-
-	for (TUniquePtr<FFileIoStoreCompressionContext>& CompressionContext : CompressionContexts)
-	{
-		if (CompressionContext->Task)
-		{
-			CompressionContext->Task->Wait();
-		}
-	}
-
 	return 0;
 }
 
diff --git a/Engine/Source/Runtime/PakFile/Private/IoDispatcherFileBackend.h b/Engine/Source/Runtime/PakFile/Private/IoDispatcherFileBackend.h
--- a/Engine/Source/Runtime/PakFile/Private/IoDispatcherFileBackend.h
+++ b/Engine/Source/Runtime/PakFile/Private/IoDispatcherFileBackend.h
@@ -20,7 +20,6 @@ struct FFileIoStoreCompressionContext
 	FFileIoStoreCompressionContext* Next = nullptr;
 	uint64 UncompressedBufferSize = 0;
 	uint8* UncompressedBuffer = nullptr;
-	FGraphEventRef Task;
 };
 
 class FFileIoStoreReader
@@ -145,7 +144,7 @@ private:
 
 		FORCEINLINE static ESubsequentsMode::Type GetSubsequentsMode()
 		{
-			return ESubsequentsMode::TrackSubsequents;
+			return ESubsequentsMode::FireAndForget;
 		}
 
 		void DoTask(ENamedThreads::Type CurrentThread, const FGraphEventRef& MyCompletionGraphEvent)
