Add mult-cast delegates for engine and game data set and reset.

diff --git a/Engine/Source/Runtime/Core/Private/GenericPlatform/GenericPlatformCrashContext.cpp b/Engine/Source/Runtime/Core/Private/GenericPlatform/GenericPlatformCrashContext.cpp
--- a/Engine/Source/Runtime/Core/Private/GenericPlatform/GenericPlatformCrashContext.cpp
+++ b/Engine/Source/Runtime/Core/Private/GenericPlatform/GenericPlatformCrashContext.cpp
@@ -117,6 +117,12 @@ int32 FGenericCrashContext::StaticCrashContextIndex = 0;
 
 const FGuid FGenericCrashContext::ExecutionGuid = FGuid::NewGuid();
 
+FEngineDataResetDelegate FGenericCrashContext::OnEngineDataReset;
+FEngineDataSetDelegate FGenericCrashContext::OnEngineDataSet;
+
+FGameDataResetDelegate FGenericCrashContext::OnGameDataReset;
+FGameDataSetDelegate FGenericCrashContext::OnGameDataSet;
+
 #if WITH_ADDITIONAL_CRASH_CONTEXTS
 FAdditionalCrashContextDelegate FGenericCrashContext::AdditionalCrashContextDelegate;
 #endif //WITH_ADDITIONAL_CRASH_CONTEXTS
@@ -1102,6 +1108,7 @@ void FGenericCrashContext::SetEpicAccountId(const FString& EpicAccountId)
 void FGenericCrashContext::ResetEngineData()
 {
 	NCached::EngineData.Reset();
+	OnEngineDataReset.Broadcast();
 }
 
 void FGenericCrashContext::SetEngineData(const FString& Key, const FString& Value)
@@ -1130,6 +1137,8 @@ void FGenericCrashContext::SetEngineData(const FString& Key, const FString& Valu
 		});
 		OldVal = Value;
 	}
+
+	OnEngineDataSet.Broadcast(Key, Value);
 }
 
 const TMap<FString, FString>& FGenericCrashContext::GetEngineData()
@@ -1146,6 +1155,7 @@ const FString* FGenericCrashContext::GetEngineData(const FString& Key)
 void FGenericCrashContext::ResetGameData()
 {
 	NCached::GameData.Reset();
+	OnGameDataReset.Broadcast();
 }
 
 void FGenericCrashContext::SetGameData(const FString& Key, const FString& Value)
@@ -1174,6 +1184,8 @@ void FGenericCrashContext::SetGameData(const FString& Key, const FString& Value)
 		});
 		OldVal = Value;
 	}
+
+	OnGameDataSet.Broadcast(Key, Value);
 }
 
 const TMap<FString, FString>& FGenericCrashContext::GetGameData()
diff --git a/Engine/Source/Runtime/Core/Public/GenericPlatform/GenericPlatformCrashContext.h b/Engine/Source/Runtime/Core/Public/GenericPlatform/GenericPlatformCrashContext.h
--- a/Engine/Source/Runtime/Core/Public/GenericPlatform/GenericPlatformCrashContext.h
+++ b/Engine/Source/Runtime/Core/Public/GenericPlatform/GenericPlatformCrashContext.h
@@ -259,6 +259,13 @@ DECLARE_MULTICAST_DELEGATE_OneParam(FAdditionalCrashContextDelegate, FCrashConte
 
 #endif //WITH_ADDITIONAL_CRASH_CONTEXTS
 
+/** Delegates for engine and game data set / reset */
+DECLARE_MULTICAST_DELEGATE(FEngineDataResetDelegate);
+DECLARE_MULTICAST_DELEGATE_TwoParams(FEngineDataSetDelegate, const FString&, const FString&);
+
+DECLARE_MULTICAST_DELEGATE(FGameDataResetDelegate);
+DECLARE_MULTICAST_DELEGATE_TwoParams(FGameDataSetDelegate, const FString&, const FString&);
+
 /**
  *	Contains a runtime crash's properties that are common for all platforms.
  *	This may change in the future.
@@ -437,18 +444,30 @@ public:
 	/** Clears the engine data dictionary */
 	CORE_API static void ResetEngineData();
 
+	/** Accessor for engine data reset callback delegate */
+	CORE_API static FEngineDataResetDelegate& OnEngineDataResetDelegate() { return OnEngineDataReset; }
+
 	/** Updates (or adds if not already present) arbitrary engine data to the crash context (will remove the key if passed an empty string) */
 	CORE_API static void SetEngineData(const FString& Key, const FString& Value);
 
+	/** Accessor for engine data change callback delegate */
+	CORE_API static FEngineDataSetDelegate& OnEngineDataSetDelegate() { return OnEngineDataSet; }
+
 	/** Get the engine data dictionary */
 	CORE_API static const TMap<FString, FString>& GetEngineData();
 
 	/** Clears the game data dictionary */
 	CORE_API static void ResetGameData();
 
+	/** Accessor for game data reset callback delegate */
+	CORE_API static FGameDataResetDelegate& OnGameDataResetDelegate() { return OnGameDataReset; }
+
 	/** Updates (or adds if not already present) arbitrary game data to the crash context (will remove the key if passed an empty string) */
 	CORE_API static void SetGameData(const FString& Key, const FString& Value);
 
+	/** Accessor for game data change callback delegate */
+	CORE_API static FGameDataSetDelegate& OnGameDataSetDelegate() { return OnGameDataSet; }
+
 	/** Get the game data dictionary */
 	CORE_API static const TMap<FString, FString>& GetGameData();
 
@@ -660,6 +679,13 @@ private:
 	/**	Records which crash context we were using the StaticCrashContextIndex counter */
 	int32 CrashContextIndex;
 
+	/** Engine and game data set / reset delegates */
+	static FEngineDataResetDelegate OnEngineDataReset;
+	static FEngineDataSetDelegate OnEngineDataSet;
+
+	static FGameDataResetDelegate OnGameDataReset;
+	static FGameDataSetDelegate OnGameDataSet;
+
 	// FNoncopyable
 	FGenericCrashContext( const FGenericCrashContext& ) = delete;
 	FGenericCrashContext& operator=(const FGenericCrashContext&) = delete;
