Fix -csvCategories commandline option not working when multiple categories are specified

diff --git a/Engine/Source/Runtime/Core/Private/ProfilingDebugging/CsvProfiler.cpp b/Engine/Source/Runtime/Core/Private/ProfilingDebugging/CsvProfiler.cpp
--- a/Engine/Source/Runtime/Core/Private/ProfilingDebugging/CsvProfiler.cpp
+++ b/Engine/Source/Runtime/Core/Private/ProfilingDebugging/CsvProfiler.cpp
@@ -3681,7 +3681,7 @@ void FCsvProfiler::Init()
 	}
 
 	FString CsvCategoriesStr;
-	if (FParse::Value(FCommandLine::Get(), TEXT("csvCategories="), CsvCategoriesStr))
+	if (FParse::Value(FCommandLine::Get(), TEXT("csvCategories="), CsvCategoriesStr, /*bShouldStopOnSeparator*/false))
 	{
 		TArray<FString> CsvCategories;
 		CsvCategoriesStr.ParseIntoArray(CsvCategories, TEXT(","), true);
