fix for assert in catch because of invalid tags

diff --git a/Engine/Source/Runtime/Core/Tests/Async/TaskGraphTest.cpp b/Engine/Source/Runtime/Core/Tests/Async/TaskGraphTest.cpp
--- a/Engine/Source/Runtime/Core/Tests/Async/TaskGraphTest.cpp
+++ b/Engine/Source/Runtime/Core/Tests/Async/TaskGraphTest.cpp
@@ -1326,7 +1326,7 @@ PRAGMA_ENABLE_DEPRECATION_WARNINGS
 		}
 	}
 
-	TEST_CASE_NAMED(FTaskGraphOversubscriptionTest, "System::Core::Async::TaskGraph::Oversubscription", "[.][ApplicationContextMask][[EngineFilter][Disabled]")
+	TEST_CASE_NAMED(FTaskGraphOversubscriptionTest, "System::Core::Async::TaskGraph::Oversubscription", "[.][ApplicationContextMask][EngineFilter][Disabled]")
 	{
 		UE_BENCHMARK(5, OversubscriptionStressTest<10>);
 	}
