Removing deprecated network versions from archive to remove potential race condition with newer custom version

diff --git a/Engine/Source/Runtime/Core/Private/Serialization/Archive.cpp b/Engine/Source/Runtime/Core/Private/Serialization/Archive.cpp
--- a/Engine/Source/Runtime/Core/Private/Serialization/Archive.cpp
+++ b/Engine/Source/Runtime/Core/Private/Serialization/Archive.cpp
@@ -156,10 +156,6 @@ void FArchiveState::Reset()
 	ArUEVer								= GPackageFileUEVersion;
 	ArLicenseeUEVer						= GPackageFileLicenseeUEVersion;
 	ArEngineVer							= FEngineVersion::Current();
-	PRAGMA_DISABLE_DEPRECATION_WARNINGS
-	ArEngineNetVer						= FNetworkVersion::GetNetworkProtocolVersion(FEngineNetworkCustomVersion::Guid);
-	ArGameNetVer						= FNetworkVersion::GetNetworkProtocolVersion(FGameNetworkCustomVersion::Guid);
-	PRAGMA_ENABLE_DEPRECATION_WARNINGS
 	ArIsLoading							= false;
 	ArIsLoadingFromCookedPackage		= false;
 	ArIsSaving							= false;
@@ -218,10 +214,6 @@ void FArchiveState::CopyTrivialFArchiveStatusMembers(const FArchiveState& Archiv
 	ArUEVer                              = ArchiveToCopy.ArUEVer;
 	ArLicenseeUEVer                      = ArchiveToCopy.ArLicenseeUEVer;
 	ArEngineVer                          = ArchiveToCopy.ArEngineVer;
-	PRAGMA_DISABLE_DEPRECATION_WARNINGS
-	ArEngineNetVer                       = ArchiveToCopy.ArEngineNetVer;
-	ArGameNetVer                         = ArchiveToCopy.ArGameNetVer;
-	PRAGMA_ENABLE_DEPRECATION_WARNINGS
 	ArIsLoading                          = ArchiveToCopy.ArIsLoading;
 	ArIsLoadingFromCookedPackage         = ArchiveToCopy.ArIsLoadingFromCookedPackage;
 	ArIsSaving                           = ArchiveToCopy.ArIsSaving;
@@ -1420,9 +1412,6 @@ void FArchiveState::SetEngineVer(const FEngineVersionBase& InVer)
 
 void FArchiveState::SetEngineNetVer(const uint32 InEngineNetVer)
 {
-	PRAGMA_DISABLE_DEPRECATION_WARNINGS
-	ArEngineNetVer = InEngineNetVer;
-	PRAGMA_ENABLE_DEPRECATION_WARNINGS
 	SetCustomVersion(FEngineNetworkCustomVersion::Guid, InEngineNetVer, TEXT("EngineNetworkVersion"));
 }
 
@@ -1433,9 +1422,6 @@ uint32 FArchiveState::EngineNetVer() const
 
 void FArchiveState::SetGameNetVer(const uint32 InGameNetVer)
 {
-	PRAGMA_DISABLE_DEPRECATION_WARNINGS
-	ArGameNetVer = InGameNetVer;
-	PRAGMA_ENABLE_DEPRECATION_WARNINGS
 	SetCustomVersion(FGameNetworkCustomVersion::Guid, InGameNetVer, TEXT("GameNetworkVersion"));
 }
 
diff --git a/Engine/Source/Runtime/Core/Public/Serialization/Archive.h b/Engine/Source/Runtime/Core/Public/Serialization/Archive.h
--- a/Engine/Source/Runtime/Core/Public/Serialization/Archive.h
+++ b/Engine/Source/Runtime/Core/Public/Serialization/Archive.h
@@ -967,14 +967,6 @@ protected:
 	/** Holds the engine version. */
 	FEngineVersionBase ArEngineVer;
 
-	/** Holds the engine network protocol version. */
-	UE_DEPRECATED(5.2, "Engine net version moved to custom versions")
-	uint32 ArEngineNetVer;
-
-	/** Holds the game network protocol version. */
-	UE_DEPRECATED(5.2, "Game net version moved to custom versions")
-	uint32 ArGameNetVer;
-
 	/**
 	* All the custom versions stored in the archive.
 	* Stored as a pointer to a heap-allocated object because of a 3-way dependency between TArray, FCustomVersionContainer and FArchive, which is too much work to change right now.
@@ -2089,8 +2081,6 @@ private:
 	using FArchiveState::ArUEVer;
 	using FArchiveState::ArLicenseeUEVer;
 	using FArchiveState::ArEngineVer;
-	using FArchiveState::ArEngineNetVer;
-	using FArchiveState::ArGameNetVer;
 	using FArchiveState::CustomVersionContainer;
 protected:
 	using FArchiveState::SetShouldSkipCompilingAssets;
