[autortfm] some text related fixup

diff --git a/Engine/Source/Runtime/Core/Private/AutoRTFM/StdLibSupport.cpp b/Engine/Source/Runtime/Core/Private/AutoRTFM/StdLibSupport.cpp
--- a/Engine/Source/Runtime/Core/Private/AutoRTFM/StdLibSupport.cpp
+++ b/Engine/Source/Runtime/Core/Private/AutoRTFM/StdLibSupport.cpp
@@ -115,6 +115,8 @@ UE_AUTORTFM_REGISTER_SELF_FUNCTION(static_cast<const char* (*)(const char*, int)
 UE_AUTORTFM_REGISTER_SELF_FUNCTION(static_cast<const char* (*)(const char*, const char*)>(&strstr));
 UE_AUTORTFM_REGISTER_SELF_FUNCTION(strlen);
 
+UE_AUTORTFM_REGISTER_SELF_FUNCTION(wcscmp);
+
 // FIXME: This is only correct when:
 // - Str is newly allocated
 // - Format is either newly allocated or not mutated
diff --git a/Engine/Source/Runtime/Core/Private/Internationalization/TextFormatter.cpp b/Engine/Source/Runtime/Core/Private/Internationalization/TextFormatter.cpp
--- a/Engine/Source/Runtime/Core/Private/Internationalization/TextFormatter.cpp
+++ b/Engine/Source/Runtime/Core/Private/Internationalization/TextFormatter.cpp
@@ -28,6 +28,7 @@
 #include "Templates/UnrealTemplate.h"
 #include "Templates/ValueOrError.h"
 #include "Trace/Detail/Channel.h"
+#include "AutoRTFM/AutoRTFM.h"
 
 #define LOCTEXT_NAMESPACE "TextFormatter"
 
@@ -937,8 +938,13 @@ FTextFormatPatternDefinition::FTextFormatPatternDefinition()
 
 FTextFormatPatternDefinitionConstRef FTextFormatPatternDefinition::GetDefault()
 {
-	static FTextFormatPatternDefinitionConstRef DefaultFormatPatternDefinition = MakeShared<FTextFormatPatternDefinition, ESPMode::ThreadSafe>();
-	return DefaultFormatPatternDefinition;
+	FTextFormatPatternDefinitionConstRef* Result;
+	AutoRTFM::Open([&Result]()
+	{
+		static FTextFormatPatternDefinitionConstRef DefaultFormatPatternDefinition = MakeShared<FTextFormatPatternDefinition, ESPMode::ThreadSafe>();
+		Result = &DefaultFormatPatternDefinition;
+	});
+	return *Result;
 }
 
 const FTokenDefinitions& FTextFormatPatternDefinition::GetTextFormatDefinitions() const
