Optimized FName == EName comparison. It's faster to convert the EName to a FName and use FName == FName.

Boils down to a single compare vs two func calls and more.

diff --git a/Engine/Source/Runtime/Core/Public/UObject/NameTypes.h b/Engine/Source/Runtime/Core/Public/UObject/NameTypes.h
--- a/Engine/Source/Runtime/Core/Public/UObject/NameTypes.h
+++ b/Engine/Source/Runtime/Core/Public/UObject/NameTypes.h
@@ -1203,7 +1203,7 @@ private:
 
 	FORCEINLINE static bool Equals(FName A, EName B)
 	{
-		return A.GetComparisonIndex() == B && A.GetNumber() == NAME_NO_NUMBER_INTERNAL;
+		return A == FName(B);
 	}
 
 #if UE_FNAME_OUTLINE_NUMBER
