Parallelizing FBlockDecoder::TryDecompressTo

diff --git a/Engine/Source/Runtime/Core/Private/Compression/CompressedBuffer.cpp b/Engine/Source/Runtime/Core/Private/Compression/CompressedBuffer.cpp
--- a/Engine/Source/Runtime/Core/Private/Compression/CompressedBuffer.cpp
+++ b/Engine/Source/Runtime/Core/Private/Compression/CompressedBuffer.cpp
@@ -437,11 +437,13 @@ bool FBlockDecoder::TryDecompressTo(
 	const FHeader& Header,
 	const FMemoryView HeaderView,
 	const uint64 RawOffset,
-	FMutableMemoryView RawView) const
+	FMutableMemoryView DestinationRawView) const
 {
 	TRACE_CPUPROFILER_EVENT_SCOPE(FBlockDecoder::TryDecompressTo);
 
-	if (Header.TotalRawSize < RawOffset + RawView.GetSize())
+	// Don't allow an output buffer bigger than we can fill so callers can know it's
+	// all initialized data.
+	if (Header.TotalRawSize < RawOffset + DestinationRawView.GetSize())
 	{
 		return false;
 	}
@@ -449,70 +451,144 @@ bool FBlockDecoder::TryDecompressTo(
 	const uint64 BlockSize = uint64(1) << Header.BlockSizeExponent;
 	const uint64 LastBlockSize = BlockSize - (BlockSize * Header.BlockCount - Header.TotalRawSize);
 	const uint32 FirstBlockIndex = uint32(RawOffset / BlockSize);
-	const uint32 LastBlockIndex = uint32((RawOffset + RawView.GetSize() - 1) / BlockSize);
+	const uint32 LastBlockIndex = uint32((RawOffset + DestinationRawView.GetSize() - 1) / BlockSize);
 
 	uint64 RawBlockOffset = RawOffset % BlockSize;
 
-	const uint32* const CompressedBlockSizes = static_cast<const uint32*>(HeaderView.RightChop(sizeof(FHeader)).GetData());
-	uint64 CompressedOffset = sizeof(FHeader) + sizeof(uint32) * uint32(Header.BlockCount) +
-		Algo::TransformAccumulate(MakeArrayView(CompressedBlockSizes, FirstBlockIndex),
+	// There's an array of CompressedBlockSizes right after the header, BlockCount uint32s
+	TConstArrayView64<uint32> CompressedBlockSizes(static_cast<const uint32*>(HeaderView.RightChop(sizeof(FHeader)).GetData()), Header.BlockCount);
+
+	// Get the offset for the first block we care about by summing the sizes of the blocks we are ignoring.
+	uint64 CompressedOffset = sizeof(FHeader) + CompressedBlockSizes.Num()*CompressedBlockSizes.GetTypeSize() +
+		Algo::TransformAccumulate(CompressedBlockSizes.Left(FirstBlockIndex),
 			[](uint32 Size) -> uint64 { return NETWORK_ORDER32(Size); }, uint64(0));
 
-	// @todo : Parallel decoder needed!  This stream is chunked into blocks but this implementation is serial.  This is a bottleneck for large data loading.
+	// Virtually all calls are for "decompress entire buffer", so rather than fight with edge cases we
+	// just optimize the known usage.
+	if (RawOffset == 0 && Header.TotalRawSize == DestinationRawView.GetSize())
+	{
+		// Precompute the offsets to each block. 64 entries with 256kb blocks is 16mb uncompressed inline, otherwise we allocate.
+		TArray<uint64, TInlineAllocator<64>> CompressedOffsets;
+		CompressedOffsets.AddUninitialized((LastBlockIndex - FirstBlockIndex) + 1);
+		CompressedOffsets[0] = CompressedOffset;
+		if (LastBlockIndex != FirstBlockIndex)
+		{
+			for (uint32 BlockIndexInWorkingSet = 1; BlockIndexInWorkingSet <= (LastBlockIndex - FirstBlockIndex); BlockIndexInWorkingSet++)
+			{
+				CompressedOffsets[BlockIndexInWorkingSet] = CompressedOffsets[BlockIndexInWorkingSet-1] + NETWORK_ORDER32(CompressedBlockSizes[FirstBlockIndex + BlockIndexInWorkingSet - 1]);
+			}		
+		}
+
+		std::atomic_int32_t FailedBlockCount = 0;
+
+		ParallelFor(TEXT("FCompressedBuffer.DecompressTo.PF"), LastBlockIndex - FirstBlockIndex + 1, 1,
+			[FirstBlockIndex, LastBlockIndex, &Header, DestinationRawView, CompressedBlockSizes, &CompressedOffsets, &Source, &Context, &FailedBlockCount, this](int32 BlockIndexInWorkingSet)
+		{
+			const uint32 BlockIndex = FirstBlockIndex + BlockIndexInWorkingSet;
+			const uint64 BlockSize = uint64(1) << Header.BlockSizeExponent;
+			const uint64 LastBlockSize = BlockSize - (BlockSize * Header.BlockCount - Header.TotalRawSize);
+
+			// We know the output buffer is the same size as the uncompressed output, so this is a lot simpler than
+			// the general case below the parallelfor.
+			const uint64 ThisBlockSize = BlockIndex == Header.BlockCount - 1 ? LastBlockSize : BlockSize;
+			
+			FMutableMemoryView DecompressTo = DestinationRawView.Mid((BlockIndex - FirstBlockIndex) * BlockSize, ThisBlockSize);
+
+			const uint32 CompressedBlockSize = NETWORK_ORDER32(CompressedBlockSizes[BlockIndex]);
+			const bool bIsCompressed = CompressedBlockSize < ThisBlockSize;
+
+			if (bIsCompressed == false)
+			{
+				if (Source.Read(CompressedOffsets[BlockIndexInWorkingSet], DecompressTo) == false)
+				{
+					FailedBlockCount.fetch_add(1, std::memory_order_relaxed);
+					return;
+				}
+			}
+			else
+			{
+				const FMemoryView CompressedBlock = Source.ReadOrView(CompressedOffsets[BlockIndexInWorkingSet], CompressedBlockSize, Context);
+				if (CompressedBlock.IsEmpty() || !DecompressBlock(DecompressTo, CompressedBlock))
+				{
+					FailedBlockCount.fetch_add(1, std::memory_order_relaxed);
+					return;
+				}
+			}
+		});
+
+		return FailedBlockCount.load() == 0;
+	}
 
 	for (uint32 BlockIndex = FirstBlockIndex; BlockIndex <= LastBlockIndex; BlockIndex++)
 	{
-		const uint64 RawBlockSize = BlockIndex == Header.BlockCount - 1 ? LastBlockSize : BlockSize;
-		const uint64 RawBlockReadSize = FMath::Min(RawView.GetSize(), RawBlockSize - RawBlockOffset);
+		// How big we're decompressing to.
+		const uint64 ThisBlockSize = BlockIndex == Header.BlockCount - 1 ? LastBlockSize : BlockSize;
+
+		// How much we want from this block.
+		const uint64 RawBlockReadSize = FMath::Min(DestinationRawView.GetSize(), ThisBlockSize - RawBlockOffset);
+
 		const uint32 CompressedBlockSize = NETWORK_ORDER32(CompressedBlockSizes[BlockIndex]);
-		const bool bIsCompressed = CompressedBlockSize < RawBlockSize;
+		const bool bIsCompressed = CompressedBlockSize < ThisBlockSize;
 
 		if (bIsCompressed)
 		{
+			// Check if we're pulling from a block we've already decompressed.
 			if (Context.RawBlockIndex == BlockIndex)
 			{
-				RawView.Left(RawBlockReadSize).CopyFrom(Context.RawBlock.GetView().Mid(RawBlockOffset, RawBlockReadSize));
+				// Can just copy direct. (Left() is unnecessary?)
+				DestinationRawView.Left(RawBlockReadSize).CopyFrom(Context.RawBlock.GetView().Mid(RawBlockOffset, RawBlockReadSize));
 			}
 			else
 			{
-				FMutableMemoryView RawBlock;
-				if (RawBlockReadSize == RawBlockSize)
+				// Need to decompress a new block.
+				FMutableMemoryView DecompressTo;
+				if (RawBlockReadSize == ThisBlockSize)
 				{
-					RawBlock = RawView.Left(RawBlockSize);
+					// Decompressing a whole block, we can go direct to the output
+					DecompressTo = DestinationRawView.Left(ThisBlockSize);
 				}
 				else
 				{
-					if (Context.RawBlock.GetSize() < RawBlockSize)
+					// Otherwise we need to go to an intermediary as the output can't hold the whole
+					// decompressed block.
+					if (Context.RawBlock.GetSize() < ThisBlockSize)
 					{
 						Context.RawBlock = FUniqueBuffer::Alloc(BlockSize);
 					}
-					RawBlock = Context.RawBlock.GetView().Left(RawBlockSize);
+					DecompressTo = Context.RawBlock.GetView().Left(ThisBlockSize);
+
+					// keep track of what block we've stashed away.
 					Context.RawBlockIndex = BlockIndex;
 				}
 
 				const FMemoryView CompressedBlock = Source.ReadOrView(CompressedOffset, CompressedBlockSize, Context);
-				if (CompressedBlock.IsEmpty() || !DecompressBlock(RawBlock, CompressedBlock))
+				if (CompressedBlock.IsEmpty() || !DecompressBlock(DecompressTo, CompressedBlock))
 				{
 					return false;
 				}
 
-				if (RawBlockReadSize != RawBlockSize)
+				// If we didn't decompress direct, copy from our cache.
+				if (RawBlockReadSize != ThisBlockSize)
 				{
-					RawView.CopyFrom(RawBlock.Mid(RawBlockOffset, RawBlockReadSize));
+					DestinationRawView.CopyFrom(DecompressTo.Mid(RawBlockOffset, RawBlockReadSize));
 				}
 			}
 		}
 		else
 		{
-			Source.Read(CompressedOffset + RawBlockOffset, RawView.Left(RawBlockReadSize));
+			// uncompressed means it's just a direct copy.
+			if (Source.Read(CompressedOffset + RawBlockOffset, DestinationRawView.Left(RawBlockReadSize)) == false)
+			{
+				return false;
+			}
 		}
 
 		RawBlockOffset = 0;
 		CompressedOffset += CompressedBlockSize;
-		RawView += RawBlockReadSize;
+		DestinationRawView += RawBlockReadSize;
 	}
 
-	return RawView.GetSize() == 0;
+	return DestinationRawView.GetSize() == 0;
 }
 
 FCompositeBuffer FBlockDecoder::DecompressToComposite(
