Validation now uses structured logging to catch errors instead of specific logger.
#prefilght https://horde.devtools.epicgames.com/job/64503555805eff7b44b6ee46

diff --git a/Engine/Source/Runtime/Core/Private/Internationalization/Text.cpp b/Engine/Source/Runtime/Core/Private/Internationalization/Text.cpp
--- a/Engine/Source/Runtime/Core/Private/Internationalization/Text.cpp
+++ b/Engine/Source/Runtime/Core/Private/Internationalization/Text.cpp
@@ -21,6 +21,8 @@
 
 #include "Serialization/StructuredArchive.h"
 
+#include "Serialization/CompactBinaryValue.h"
+
 #include "UObject/EditorObjectVersion.h"
 #include "UObject/UE5ReleaseStreamObjectVersion.h"
 #include "HAL/PlatformProcess.h"
@@ -1321,6 +1323,46 @@ bool FFormatArgumentValue::IdenticalTo(const FFormatArgumentValue& Other, const
 	return false;
 }
 
+
+FFormatArgumentValue::FFormatArgumentValue(const class FCbValue& Value)
+{
+	switch (Value.GetType())
+	{
+	case ECbFieldType::String:
+		Type = EFormatArgumentType::Text;
+		TextValue = FText::FromString(FString(Value.AsString()));
+		break;
+	case ECbFieldType::IntegerPositive:
+		Type = EFormatArgumentType::UInt;
+		UIntValue = Value.AsIntegerPositive();
+		break;
+	case ECbFieldType::IntegerNegative:
+		Type = EFormatArgumentType::Int;
+		UIntValue = Value.AsIntegerNegative();
+		break;
+	case ECbFieldType::Float32:
+		Type = EFormatArgumentType::Float;
+		FloatValue = Value.AsFloat32();
+		break;
+	case ECbFieldType::Float64:
+		Type = EFormatArgumentType::Double;
+		DoubleValue = Value.AsFloat64();
+		break;
+	case ECbFieldType::BoolFalse:
+		Type = EFormatArgumentType::Int;
+		IntValue = 0;
+		break;
+	case ECbFieldType::BoolTrue:
+		Type = EFormatArgumentType::Int;
+		IntValue = 1;
+		break;
+	default:
+		Type = EFormatArgumentType::Text;
+		TextValue = FText::GetEmpty();
+	}
+
+}
+
 FString FFormatArgumentValue::ToFormattedString(const bool bInRebuildText, const bool bInRebuildAsSource) const
 {
 	FString Result;
diff --git a/Engine/Source/Runtime/Core/Public/Internationalization/Text.h b/Engine/Source/Runtime/Core/Public/Internationalization/Text.h
--- a/Engine/Source/Runtime/Core/Public/Internationalization/Text.h
+++ b/Engine/Source/Runtime/Core/Public/Internationalization/Text.h
@@ -821,6 +821,8 @@ public:
 	{
 	}
 
+	FFormatArgumentValue(const class FCbValue& Value);
+
 	FFormatArgumentValue(const int32 Value)
 		: Type(EFormatArgumentType::Int)
 	{
