Replaced some operator TArray new usage - enough to make BlankProgram compile if operator new is deprecated.

diff --git a/Engine/Source/Runtime/Core/Private/Containers/String.cpp b/Engine/Source/Runtime/Core/Private/Containers/String.cpp
--- a/Engine/Source/Runtime/Core/Private/Containers/String.cpp
+++ b/Engine/Source/Runtime/Core/Private/Containers/String.cpp
@@ -1388,7 +1388,7 @@ int32 FString::ParseIntoArray(TArray<FString>& OutArray, const TCHAR* const * De
 				if(!InCullEmpty || SubstringLength != 0)
 				{
 					// ... add new string from substring beginning up to the beginning of this delimiter.
-					new (OutArray) FString(SubstringEndIndex - SubstringBeginIndex, Start + SubstringBeginIndex);
+					OutArray.Emplace(SubstringEndIndex - SubstringBeginIndex, Start + SubstringBeginIndex);
 				}
 				// Next substring begins at the end of the discovered delimiter.
 				SubstringBeginIndex = SubstringEndIndex + DelimiterLength;
@@ -1406,7 +1406,7 @@ int32 FString::ParseIntoArray(TArray<FString>& OutArray, const TCHAR* const * De
 		if(!InCullEmpty || SubstringLength != 0)
 		{
 			// ... add new string from substring beginning up to the beginning of this delimiter.
-			new (OutArray) FString(Start + SubstringBeginIndex);
+			OutArray.Emplace(Start + SubstringBeginIndex);
 		}
 	}
 
diff --git a/Engine/Source/Runtime/Core/Private/Misc/BufferedOutputDevice.cpp b/Engine/Source/Runtime/Core/Private/Misc/BufferedOutputDevice.cpp
--- a/Engine/Source/Runtime/Core/Private/Misc/BufferedOutputDevice.cpp
+++ b/Engine/Source/Runtime/Core/Private/Misc/BufferedOutputDevice.cpp
@@ -13,7 +13,7 @@ void FBufferedOutputDevice::Serialize(const TCHAR* InData, ELogVerbosity::Type V
 	}
 
 	FScopeLock ScopeLock(&SynchronizationObject);
-	new(BufferedLines) FBufferedLine(InData, Category, Verbosity);
+	BufferedLines.Emplace(InData, Category, Verbosity);
 }
 
 void FBufferedOutputDevice::GetContents(TArray<FBufferedLine>& DestBuffer)
diff --git a/Engine/Source/Runtime/Core/Private/Misc/CommandLine.cpp b/Engine/Source/Runtime/Core/Private/Misc/CommandLine.cpp
--- a/Engine/Source/Runtime/Core/Private/Misc/CommandLine.cpp
+++ b/Engine/Source/Runtime/Core/Private/Misc/CommandLine.cpp
@@ -307,12 +307,12 @@ void FCommandLine::Parse(const TCHAR* InCmdLine, TArray<FString>& Tokens, TArray
 	{
 		if ((**NextToken == TCHAR('-')))
 		{
-			new(Switches) FString(NextToken.Mid(1));
-			new(Tokens) FString(NextToken.Right(NextToken.Len() - 1));
+			Switches.Add(NextToken.Mid(1));
+			Tokens.Add(NextToken.Right(NextToken.Len() - 1));
 		}
 		else
 		{
-			new(Tokens) FString(NextToken);
+			Tokens.Add(MoveTemp(NextToken));
 		}
 	}
 }
diff --git a/Engine/Source/Runtime/Core/Private/Misc/ConfigCacheIni.cpp b/Engine/Source/Runtime/Core/Private/Misc/ConfigCacheIni.cpp
--- a/Engine/Source/Runtime/Core/Private/Misc/ConfigCacheIni.cpp
+++ b/Engine/Source/Runtime/Core/Private/Misc/ConfigCacheIni.cpp
@@ -3002,7 +3002,7 @@ int32 FConfigCacheIni::GetSingleLineArray
 	FString NextToken;
 	while ( FParse::Token(RawString, NextToken, false) )
 	{
-		new(out_Arr) FString(NextToken);
+		out_Arr.Add(MoveTemp(NextToken));
 	}
 	return bValueExisted;
 }
@@ -3312,7 +3312,7 @@ struct FConfigMemoryData
 		SizeIndent = FMath::Max(SizeIndent, FString::FromInt((int32)TotalMem).Len());
 		MaxSizeIndent = FMath::Max(MaxSizeIndent, FString::FromInt((int32)MaxMem).Len());
 		
-		new(MemoryData) FConfigFileMemoryData( ConfigFilename, TotalMem, MaxMem );
+		MemoryData.Emplace( ConfigFilename, TotalMem, MaxMem );
 	}
 
 	void SortBySize()
diff --git a/Engine/Source/Runtime/Core/Private/Modules/ModuleManager.cpp b/Engine/Source/Runtime/Core/Private/Modules/ModuleManager.cpp
--- a/Engine/Source/Runtime/Core/Private/Modules/ModuleManager.cpp
+++ b/Engine/Source/Runtime/Core/Private/Modules/ModuleManager.cpp
@@ -832,7 +832,7 @@ void FModuleManager::UnloadModulesAtShutdown()
 			// Only if the module supports shutting down in this phase
 			if( ModuleInfo->Module->SupportsAutomaticShutdown() )
 			{
-				new (ModulesToUnload)FModulePair(ModuleIt.Key, ModuleIt.Value->LoadOrder, ModuleInfo->Module.Get());
+				ModulesToUnload.Emplace(ModuleIt.Key, ModuleIt.Value->LoadOrder, ModuleInfo->Module.Get());
 			}
 		}
 	}
diff --git a/Engine/Source/Runtime/Core/Private/ProfilingDebugging/ABTesting.cpp b/Engine/Source/Runtime/Core/Private/ProfilingDebugging/ABTesting.cpp
--- a/Engine/Source/Runtime/Core/Private/ProfilingDebugging/ABTesting.cpp
+++ b/Engine/Source/Runtime/Core/Private/ProfilingDebugging/ABTesting.cpp
@@ -123,7 +123,7 @@ const TCHAR* FABTest::TickAndGetCommand()
 		if (ABTestNumSamples < HistoryNum)
 		{
 			check(ABTestNumSamples == Samples.Num());
-			Sample = new (Samples)FSample();
+			Sample = &Samples.AddDefaulted_GetRef();
 		}
 		else
 		{
diff --git a/Engine/Source/Runtime/Core/Private/Serialization/MemoryImage.cpp b/Engine/Source/Runtime/Core/Private/Serialization/MemoryImage.cpp
--- a/Engine/Source/Runtime/Core/Private/Serialization/MemoryImage.cpp
+++ b/Engine/Source/Runtime/Core/Private/Serialization/MemoryImage.cpp
@@ -1858,22 +1858,22 @@ uint32 FMemoryImageSection::Flatten(FMemoryImageResult& OutResult) const
 	OutResult.VTables.Reserve(OutResult.VTables.Num() + VTables.Num());
 	for (const FMemoryImageVTablePointer& VTable : VTables)
 	{
-		FMemoryImageVTablePointer* ResultVTable = new(OutResult.VTables) FMemoryImageVTablePointer(VTable);
-		ResultVTable->Offset += AlignedOffset;
+		FMemoryImageVTablePointer& ResultVTable = OutResult.VTables.Add_GetRef(VTable);
+		ResultVTable.Offset += AlignedOffset;
 	}
 
 	OutResult.ScriptNames.Reserve(OutResult.ScriptNames.Num() + ScriptNames.Num());
 	for (const FMemoryImageNamePointer& Name : ScriptNames)
 	{
-		FMemoryImageNamePointer* ResultName = new(OutResult.ScriptNames) FMemoryImageNamePointer(Name);
-		ResultName->Offset += AlignedOffset;
+		FMemoryImageNamePointer& ResultName = OutResult.ScriptNames.Add_GetRef(Name);
+		ResultName.Offset += AlignedOffset;
 	}
 
 	OutResult.MemoryImageNames.Reserve(OutResult.MemoryImageNames.Num() + MemoryImageNames.Num());
 	for (const FMemoryImageNamePointer& Name : MemoryImageNames)
 	{
-		FMemoryImageNamePointer* ResultName = new(OutResult.MemoryImageNames) FMemoryImageNamePointer(Name);
-		ResultName->Offset += AlignedOffset;
+		FMemoryImageNamePointer& ResultName = OutResult.MemoryImageNames.Add_GetRef(Name);
+		ResultName.Offset += AlignedOffset;
 	}
 
 	return AlignedOffset;
diff --git a/Engine/Source/Runtime/Core/Public/Async/TaskGraphInterfaces.h b/Engine/Source/Runtime/Core/Public/Async/TaskGraphInterfaces.h
--- a/Engine/Source/Runtime/Core/Public/Async/TaskGraphInterfaces.h
+++ b/Engine/Source/Runtime/Core/Public/Async/TaskGraphInterfaces.h
@@ -960,7 +960,7 @@ public:
 	void DontCompleteUntil(FGraphEventRef EventToWaitFor)
 	{
 		checkThreadGraph(!IsComplete()); // it is not legal to add a DontCompleteUntil after the event has been completed. Basically, this is only legal within a task function.
-		new (EventsToWaitFor) FGraphEventRef(EventToWaitFor);
+		EventsToWaitFor.Emplace(EventToWaitFor);
 		TaskTrace::SubsequentAdded(EventToWaitFor->GetTraceId(), GetTraceId());
 	}
 
@@ -1866,7 +1866,7 @@ public:
 			// convert the pointer list to a list of handles
 			for (int32 Index = 0; Index < Pending.Num(); Index++)
 			{
-				new (PendingHandles) FGraphEventRef(Pending[Index]); 
+				PendingHandles.Emplace(Pending[Index]);
 				Pending[Index]->Release(); // remove the ref count we added when we added it to the lock free list
 			}
 			// start a new task that won't complete until all of these tasks have executed, plus any tasks that they create when they run
diff --git a/Engine/Source/Runtime/Core/Public/Containers/Array.h b/Engine/Source/Runtime/Core/Public/Containers/Array.h
--- a/Engine/Source/Runtime/Core/Public/Containers/Array.h
+++ b/Engine/Source/Runtime/Core/Public/Containers/Array.h
@@ -2515,7 +2515,7 @@ public:
 		Empty(Number);
 		for (SizeType Index = 0; Index < Number; ++Index)
 		{
-			new(*this) ElementType(Element);
+			Add(Element);
 		}
 	}
 
@@ -3488,7 +3488,7 @@ struct TIsTArray
 //
 // Array operator news.
 //
-template <typename T,typename AllocatorType> void* operator new( size_t Size, TArray<T,AllocatorType>& Array )
+template <typename T,typename AllocatorType> void* operator new(size_t Size, TArray<T, AllocatorType>& Array)
 {
 	check(Size == sizeof(T));
 	const auto Index = Array.AddUninitialized();
@@ -3557,7 +3557,7 @@ struct TArrayPrivateFriend
 				A.Empty(SerializeNum);
 				for (SizeType i=0; i<SerializeNum; i++)
 				{
-					Ar << *::new(A) ElementType;
+					Ar << A.AddDefaulted_GetRef();
 				}
 			}
 			else
@@ -3572,7 +3572,7 @@ struct TArrayPrivateFriend
 
 			for (SizeType i=0; i<SerializeNum; i++)
 			{
-				Ar << *::new(A) ElementType;
+				Ar << A.AddDefaulted_GetRef();
 			}
 		}
 		else
diff --git a/Engine/Source/Runtime/Core/Public/Containers/Map.h b/Engine/Source/Runtime/Core/Public/Containers/Map.h
--- a/Engine/Source/Runtime/Core/Public/Containers/Map.h
+++ b/Engine/Source/Runtime/Core/Public/Containers/Map.h
@@ -722,7 +722,7 @@ public:
 		OutArray.Empty(Pairs.Num());
 		for (typename ElementSetType::TConstIterator PairIt(Pairs); PairIt; ++PairIt)
 		{
-			new(OutArray) KeyType(PairIt->Key);
+			OutArray.Add(PairIt->Key);
 		}
 	}
 
@@ -736,7 +736,7 @@ public:
 		OutArray.Empty(Pairs.Num());
 		for (typename ElementSetType::TConstIterator PairIt(Pairs); PairIt; ++PairIt)
 		{
-			new(OutArray) ValueType(PairIt->Value);
+			OutArray.Add(PairIt->Value);
 		}
 	}
 
@@ -1408,7 +1408,7 @@ public:
 	{
 		for (typename Super::ElementSetType::TConstKeyIterator It(Super::Pairs, Key); It; ++It)
 		{
-			new(OutValues) ValueType(It->Value);
+			OutValues.Add(It->Value);
 		}
 
 		if (bMaintainOrder)
