Fix 1:st error generated when running PreSubmitTest.bat

The solution is taken from https://epicgames.slack.com/archives/C042RSGQVPD/p1677783304284859?thread_ts=1677487516.170069&cid=C042RSGQVPD
and is to "protected" includes of dbghelp.h by THIRD_PARTY_INCLUDES_START/THIRD_PARTY_INCLUDES_END.
There are more unprotected includes in the code, but for some reason this was enough to make PreSubmitTest.bat progress to the next unrelated error.
Which means that this is a necessary, but not complete, fix for PreSubmitTest.bat.

#rnx

diff --git a/Engine/Source/Runtime/Core/Private/Windows/WindowsPlatformCrashContext.cpp b/Engine/Source/Runtime/Core/Private/Windows/WindowsPlatformCrashContext.cpp
--- a/Engine/Source/Runtime/Core/Private/Windows/WindowsPlatformCrashContext.cpp
+++ b/Engine/Source/Runtime/Core/Private/Windows/WindowsPlatformCrashContext.cpp
@@ -37,6 +37,7 @@
 #include "Windows/AllowWindowsPlatformTypes.h"
 THIRD_PARTY_INCLUDES_START
 #include <strsafe.h>
+THIRD_PARTY_INCLUDES_START
 #include <dbghelp.h>
 #include <Shlwapi.h>
 #include <psapi.h>
