Fix for SOL-4430 - pass argc/argv into init functions for both UE and uLang so standard command line args are processed.

diff --git a/Engine/Source/Runtime/Core/Private/Misc/CommandLine.cpp b/Engine/Source/Runtime/Core/Private/Misc/CommandLine.cpp
--- a/Engine/Source/Runtime/Core/Private/Misc/CommandLine.cpp
+++ b/Engine/Source/Runtime/Core/Private/Misc/CommandLine.cpp
@@ -317,8 +317,8 @@ void FCommandLine::Parse(const TCHAR* InCmdLine, TArray<FString>& Tokens, TArray
 	}
 }
 
-
-FString FCommandLine::BuildFromArgV(const TCHAR* Prefix, int32 ArgC, TCHAR* ArgV[], const TCHAR* Suffix)
+template<typename CharType>
+FString BuildFromArgVImpl(const CharType* Prefix, int32 ArgC, CharType* ArgV[], const CharType* Suffix)
 {
 	FString Result;
 
@@ -351,13 +351,23 @@ FString FCommandLine::BuildFromArgV(const TCHAR* Prefix, int32 ArgC, TCHAR* ArgV
 	// add the prefix and suffix if provided
 	if (Prefix)
 	{
-		Result = FString::Printf(TEXT("%s %s"), Prefix, *Result);
+		Result = FString::Printf(TEXT("%s %s"), StringCast<TCHAR>(Prefix).Get(), *Result);
 	}
 
 	if (Suffix)
 	{
-		Result = FString::Printf(TEXT("%s %s"), *Result, Suffix);
+		Result = FString::Printf(TEXT("%s %s"), *Result, StringCast<TCHAR>(Suffix).Get());
 	}
 
 	return Result;
 }
+
+FString FCommandLine::BuildFromArgV(const WIDECHAR* Prefix, int32 ArgC, WIDECHAR* ArgV[], const WIDECHAR* Suffix)
+{
+	return BuildFromArgVImpl(Prefix, ArgC, ArgV, Suffix);
+}
+
+FString FCommandLine::BuildFromArgV(const ANSICHAR* Prefix, int32 ArgC, ANSICHAR* ArgV[], const ANSICHAR* Suffix)
+{
+	return BuildFromArgVImpl(Prefix, ArgC, ArgV, Suffix);
+}
diff --git a/Engine/Source/Runtime/Core/Public/Misc/CommandLine.h b/Engine/Source/Runtime/Core/Public/Misc/CommandLine.h
--- a/Engine/Source/Runtime/Core/Public/Misc/CommandLine.h
+++ b/Engine/Source/Runtime/Core/Public/Misc/CommandLine.h
@@ -106,7 +106,8 @@ struct FCommandLine
 	 * Builds a command line string from the ArgC/ArgV main() arguments, together with
 	 * an optional prefix or suffix, for adding additional command list arguments programmatically.
 	*/
-	CORE_API static FString BuildFromArgV(const TCHAR* Prefix, int32 ArgC, TCHAR* ArgV[], const TCHAR* Suffix);
+	CORE_API static FString BuildFromArgV(const WIDECHAR* Prefix, int32 ArgC, WIDECHAR* ArgV[], const WIDECHAR* Suffix);
+	CORE_API static FString BuildFromArgV(const ANSICHAR* Prefix, int32 ArgC, ANSICHAR* ArgV[], const ANSICHAR* Suffix);
 
 private:
 #if UE_COMMAND_LINE_USES_ALLOW_LIST
