Fix merge error

#rnx

diff --git a/Engine/Source/Runtime/Core/Private/Windows/WindowsPlatformCrashContext.cpp b/Engine/Source/Runtime/Core/Private/Windows/WindowsPlatformCrashContext.cpp
--- a/Engine/Source/Runtime/Core/Private/Windows/WindowsPlatformCrashContext.cpp
+++ b/Engine/Source/Runtime/Core/Private/Windows/WindowsPlatformCrashContext.cpp
@@ -37,7 +37,6 @@
 #include "Windows/AllowWindowsPlatformTypes.h"
 THIRD_PARTY_INCLUDES_START
 #include <strsafe.h>
-THIRD_PARTY_INCLUDES_START
 #include <dbghelp.h>
 #include <Shlwapi.h>
 #include <psapi.h>
