Core: Added a macro for the no_unique_address C++20 attribute

diff --git a/Engine/Source/Runtime/Core/Public/HAL/Platform.h b/Engine/Source/Runtime/Core/Public/HAL/Platform.h
--- a/Engine/Source/Runtime/Core/Public/HAL/Platform.h
+++ b/Engine/Source/Runtime/Core/Public/HAL/Platform.h
@@ -682,6 +682,16 @@
 	#define UE_CONSTEVAL constexpr
 #endif
 
+/* Use before a class data member declaration allow it to be overlapped with other non-static data members or base class subobjects of its class. */
+#if !defined(UE_NO_UNIQUE_ADDRESS) && defined(__has_cpp_attribute)
+	#if __has_cpp_attribute(no_unique_address)
+		#define UE_NO_UNIQUE_ADDRESS [[no_unique_address]]
+	#endif
+#endif
+#ifndef UE_NO_UNIQUE_ADDRESS
+	#define UE_NO_UNIQUE_ADDRESS
+#endif
+
 /* Wrap a function signature in these to indicate that the function never returns nullptr */
 #ifndef FUNCTION_NON_NULL_RETURN_START
 	#define FUNCTION_NON_NULL_RETURN_START
