Niagara - Merging 25284293 - Additional fixes for data interface routing. - Resubmit with fix for cook
+ Fix renderer data interface bindings.
+ Update the translator so that data interface reads will correctly find a previously written data interface in the case where the DI names don't match.
+ Update the translator so that a data interface that is written to multiple parameters generates mutiple entries in the cached default data interfaces.
+ Update the grid 2d and grid 3d collections to use the data interface itself to find attributes, rather than the data interface name.

diff --git a/Engine/Source/Runtime/Core/Private/UObject/FortniteMainBranchObjectVersion.cpp b/Engine/Source/Runtime/Core/Private/UObject/FortniteMainBranchObjectVersion.cpp
--- a/Engine/Source/Runtime/Core/Private/UObject/FortniteMainBranchObjectVersion.cpp
+++ b/Engine/Source/Runtime/Core/Private/UObject/FortniteMainBranchObjectVersion.cpp
@@ -11,7 +11,7 @@ TMap<FGuid, FGuid> FFortniteMainBranchObjectVersion::GetSystemGuids()
 	SystemGuids.Add(DevGuids.MATERIALSHADERMAP_DERIVEDDATA_VER, FGuid("CE7776E63E9E43FBB8C103EE1AB40B5D"));
 	SystemGuids.Add(DevGuids.NANITE_DERIVEDDATA_VER, FGuid("D28C003F825F4C178F698938FB0E47F7"));
 	SystemGuids.Add(DevGuids.NIAGARASHADERMAP_DERIVEDDATA_VER, FGuid("35D7960B72164B49B9B7B3FB5857E674"));
-	SystemGuids.Add(DevGuids.Niagara_LatestScriptCompileVersion, FGuid("D74C8B1657194273B102805F71528F94"));
+	SystemGuids.Add(DevGuids.Niagara_LatestScriptCompileVersion, FGuid("3DC7587E76B74CC3BA6FAD6B2F5EC31A"));
 	SystemGuids.Add(DevGuids.SkeletalMeshDerivedDataVersion, FGuid("518D7C51D31D9BE4A9F4B81542AC8119"));
 	SystemGuids.Add(DevGuids.STATICMESH_DERIVEDDATA_VER, FGuid("543843B103794E0E9BA4BE80FB602F79"));
 	return SystemGuids;
