Add the attended status to CrashContext
* Enable filtering unattended sessions easier from crash reporting tool

diff --git a/Engine/Source/Runtime/Core/Private/GenericPlatform/GenericPlatformCrashContext.cpp b/Engine/Source/Runtime/Core/Private/GenericPlatform/GenericPlatformCrashContext.cpp
--- a/Engine/Source/Runtime/Core/Private/GenericPlatform/GenericPlatformCrashContext.cpp
+++ b/Engine/Source/Runtime/Core/Private/GenericPlatform/GenericPlatformCrashContext.cpp
@@ -74,6 +74,18 @@ static bool NeedsEscape(FStringView Str)
 	return false;
 }
 
+static const TCHAR* AttendedStatusToString(EUnattendedStatus Status)
+{
+	switch(Status)
+	{
+	case EUnattendedStatus::Attended: return TEXT("Attended");
+	case EUnattendedStatus::Unattended: return TEXT("Unattended");
+	case EUnattendedStatus::Unknown: // fallthrough
+	default:
+		return TEXT("Unknown");
+	}
+}
+
 /*-----------------------------------------------------------------------------
 	FGenericCrashContext
 -----------------------------------------------------------------------------*/
@@ -192,6 +204,7 @@ void FGenericCrashContext::Initialize()
 
 	NCached::Set(NCached::Session.PlatformName, FPlatformProperties::PlatformName());
 	NCached::Set(NCached::Session.PlatformNameIni, FPlatformProperties::IniPlatformName());
+	NCached::Set(NCached::Session.AttendedStatus, AttendedStatusToString(EUnattendedStatus::Unknown));
 
 	// Information that cannot be gathered if command line is not initialized (e.g. crash during static init)
 	if (FCommandLine::IsInitialized())
@@ -221,7 +234,10 @@ void FGenericCrashContext::Initialize()
 			}
 		}
 
-		NCached::UserSettings.bNoDialog = FApp::IsUnattended() || IsRunningDedicatedServer();
+		const bool IsUnattended = FApp::IsUnattended();
+
+		NCached::UserSettings.bNoDialog = IsUnattended|| IsRunningDedicatedServer();
+		NCached::Set(NCached::Session.AttendedStatus, AttendedStatusToString(IsUnattended ? EUnattendedStatus::Unattended : EUnattendedStatus::Attended));
 	}
 
 	// Create a unique base guid for bug report ids
@@ -741,6 +757,7 @@ void FGenericCrashContext::SerializeContentToBuffer() const
 	AddCrashProperty( TEXT( "CrashType" ), GetCrashTypeString(Type) );
 	AddCrashProperty( TEXT( "ErrorMessage" ), ErrorMessage );
 	AddCrashProperty( TEXT( "CrashReporterMessage" ), NCached::Session.CrashReportClientRichText );
+	AddCrashProperty( TEXT( "AttendedStatus"), NCached::Session.AttendedStatus);
 
 	SerializeSessionContext(CommonBuffer);
 
diff --git a/Engine/Source/Runtime/Core/Public/GenericPlatform/GenericPlatformCrashContext.h b/Engine/Source/Runtime/Core/Public/GenericPlatform/GenericPlatformCrashContext.h
--- a/Engine/Source/Runtime/Core/Public/GenericPlatform/GenericPlatformCrashContext.h
+++ b/Engine/Source/Runtime/Core/Public/GenericPlatform/GenericPlatformCrashContext.h
@@ -128,6 +128,17 @@ enum class ECrashTrigger
 	Normal = 0
 };
 
+/**
+ * Tristate to identify a session which is attended or unattended (ie. usually automated testing)
+ * Determination requires command line arguments - therefore if not available, status is unknown 
+ */
+enum class EUnattendedStatus : uint8
+{
+	Unknown,
+	Attended,
+	Unattended
+};
+
 #define CR_MAX_ERROR_MESSAGE_CHARS 2048
 #define CR_MAX_DIRECTORY_CHARS 256
 #define CR_MAX_STACK_FRAMES 256
@@ -183,6 +194,7 @@ struct FSessionContext
 	TCHAR 					CrashReportClientRichText[CR_MAX_RICHTEXT_FIELD_CHARS];
 	TCHAR 					GameStateName[CR_MAX_GENERIC_FIELD_CHARS];
 	TCHAR 					CrashConfigFilePath[CR_MAX_DIRECTORY_CHARS];
+	TCHAR					AttendedStatus[CR_MAX_GENERIC_FIELD_CHARS];
 	char					PlatformName[CR_MAX_GENERIC_FIELD_CHARS];
 	char					PlatformNameIni[CR_MAX_GENERIC_FIELD_CHARS];
 	FPlatformMemoryStats	MemoryStats;
