Changed FErrorHistWriter to use FWordMutex to permit locking during static destructors being executed

diff --git a/Engine/Source/Runtime/Core/Private/Misc/AssertionMacros.cpp b/Engine/Source/Runtime/Core/Private/Misc/AssertionMacros.cpp
--- a/Engine/Source/Runtime/Core/Private/Misc/AssertionMacros.cpp
+++ b/Engine/Source/Runtime/Core/Private/Misc/AssertionMacros.cpp
@@ -7,6 +7,8 @@
 #include "Templates/Atomic.h"
 #include "Misc/CString.h"
 #include "Misc/Crc.h"
+#include "Async/UniqueLock.h"
+#include "Async/WordMutex.h"
 #include "Containers/UnrealString.h"
 #include "Containers/StringConv.h"
 #include "GenericPlatform/GenericPlatformStackWalk.h"
@@ -32,12 +34,9 @@ namespace
 	TAtomic<SIZE_T> NumEnsureFailures {0};
 	int32 ActiveEnsureCount = 0;
 
-	/** Lock used to synchronize the fail debug calls. */
-	static FCriticalSection& GetFailDebugCriticalSection()
-	{
-		static FCriticalSection FailDebugCriticalSection;
-		return FailDebugCriticalSection;
-	}
+	// Lock used to synchronize the fail debug calls.
+	// Using FWordMutex because it is zero-initialized and has no destructor.
+	static UE::FWordMutex FailDebugMutex;
 
 	struct FTempCommandLineScope
 	{
@@ -150,13 +149,13 @@ public:
 
 	inline FErrorHistWriter()
 	{
-		GetFailDebugCriticalSection().Lock();
+		FailDebugMutex.Lock();
 	}
 
 	inline ~FErrorHistWriter()
 	{
 		Terminate();
-		GetFailDebugCriticalSection().Unlock();
+		FailDebugMutex.Unlock();
 	}
 
 	inline void Terminate()
@@ -548,8 +547,7 @@ FORCENOINLINE void FDebug::EnsureFailed(const ANSICHAR* Expr, const ANSICHAR* Fi
 #endif
 
 #if PLATFORM_USE_REPORT_ENSURE
-			FScopeLock Lock(&GetFailDebugCriticalSection());
-
+			UE::TUniqueLock Lock(FailDebugMutex);
 			ReportEnsure(ErrorMsg, ProgramCounter);
 
 			GErrorHist[0] = TEXT('\0');
