Self register some math functions, may not be all but a start
Will likely need to add more as they come

diff --git a/Engine/Source/Runtime/Core/Private/AutoRTFM/StdLibSupport.cpp b/Engine/Source/Runtime/Core/Private/AutoRTFM/StdLibSupport.cpp
--- a/Engine/Source/Runtime/Core/Private/AutoRTFM/StdLibSupport.cpp
+++ b/Engine/Source/Runtime/Core/Private/AutoRTFM/StdLibSupport.cpp
@@ -117,6 +117,39 @@ UE_AUTORTFM_REGISTER_SELF_FUNCTION(strlen);
 
 UE_AUTORTFM_REGISTER_SELF_FUNCTION(wcscmp);
 
+// Self register Math functions
+UE_AUTORTFM_REGISTER_SELF_FUNCTION(sqrt);
+UE_AUTORTFM_REGISTER_SELF_FUNCTION(sqrtf);
+UE_AUTORTFM_REGISTER_SELF_FUNCTION(sin);
+UE_AUTORTFM_REGISTER_SELF_FUNCTION(sinf);
+UE_AUTORTFM_REGISTER_SELF_FUNCTION(cos);
+UE_AUTORTFM_REGISTER_SELF_FUNCTION(cosf);
+UE_AUTORTFM_REGISTER_SELF_FUNCTION(tan);
+UE_AUTORTFM_REGISTER_SELF_FUNCTION(tanf);
+UE_AUTORTFM_REGISTER_SELF_FUNCTION(asin);
+UE_AUTORTFM_REGISTER_SELF_FUNCTION(asinf);
+UE_AUTORTFM_REGISTER_SELF_FUNCTION(acos);
+UE_AUTORTFM_REGISTER_SELF_FUNCTION(acosf);
+UE_AUTORTFM_REGISTER_SELF_FUNCTION(atan);
+UE_AUTORTFM_REGISTER_SELF_FUNCTION(atanf);
+UE_AUTORTFM_REGISTER_SELF_FUNCTION(atan2);
+UE_AUTORTFM_REGISTER_SELF_FUNCTION(sinh);
+UE_AUTORTFM_REGISTER_SELF_FUNCTION(sinhf);
+UE_AUTORTFM_REGISTER_SELF_FUNCTION(cosh);
+UE_AUTORTFM_REGISTER_SELF_FUNCTION(coshf);
+UE_AUTORTFM_REGISTER_SELF_FUNCTION(tanh);
+UE_AUTORTFM_REGISTER_SELF_FUNCTION(tanhf);
+UE_AUTORTFM_REGISTER_SELF_FUNCTION(asinh);
+UE_AUTORTFM_REGISTER_SELF_FUNCTION(acosh);
+UE_AUTORTFM_REGISTER_SELF_FUNCTION(atanh);
+UE_AUTORTFM_REGISTER_SELF_FUNCTION(exp);
+UE_AUTORTFM_REGISTER_SELF_FUNCTION(expf);
+UE_AUTORTFM_REGISTER_SELF_FUNCTION(log);
+UE_AUTORTFM_REGISTER_SELF_FUNCTION(logf);
+UE_AUTORTFM_REGISTER_SELF_FUNCTION(pow);
+UE_AUTORTFM_REGISTER_SELF_FUNCTION(powf);
+UE_AUTORTFM_REGISTER_SELF_FUNCTION(llrint);
+
 // FIXME: This is only correct when:
 // - Str is newly allocated
 // - Format is either newly allocated or not mutated
