Add PlatformFullName inside the CrashContext

diff --git a/Engine/Source/Runtime/Core/Private/GenericPlatform/GenericPlatformCrashContext.cpp b/Engine/Source/Runtime/Core/Private/GenericPlatform/GenericPlatformCrashContext.cpp
--- a/Engine/Source/Runtime/Core/Private/GenericPlatform/GenericPlatformCrashContext.cpp
+++ b/Engine/Source/Runtime/Core/Private/GenericPlatform/GenericPlatformCrashContext.cpp
@@ -657,6 +657,7 @@ void FGenericCrashContext::SerializeSessionContext(FString& Buffer)
 	AddCrashPropertyInternal(Buffer, TEXT("GameSessionID"), NCached::Session.GameSessionID);
 
 	AddCrashPropertyInternal(Buffer, TEXT("PlatformName"), NCached::Session.PlatformName);
+	AddCrashPropertyInternal(Buffer, TEXT("PlatformFullName"), NCached::Session.PlatformName);
 	AddCrashPropertyInternal(Buffer, TEXT("PlatformNameIni"), NCached::Session.PlatformNameIni);
 	AddCrashPropertyInternal(Buffer, TEXT("EngineMode"), NCached::Session.EngineMode);
 	AddCrashPropertyInternal(Buffer, TEXT("EngineModeEx"), NCached::Session.EngineModeEx);
