Added comment to FName comparison optimization in CL 25308282.

diff --git a/Engine/Source/Runtime/Core/Public/UObject/NameTypes.h b/Engine/Source/Runtime/Core/Public/UObject/NameTypes.h
--- a/Engine/Source/Runtime/Core/Public/UObject/NameTypes.h
+++ b/Engine/Source/Runtime/Core/Public/UObject/NameTypes.h
@@ -1203,6 +1203,9 @@ private:
 
 	FORCEINLINE static bool Equals(FName A, EName B)
 	{
+		// With UE_FNAME_OUTLINE_NUMBER 1, FName == FName(EName) is
+		// faster than extracting index and number for direct EName comparison.
+		// return A.GetComparisonIndex() == B && A.GetNumber() == NAME_NO_NUMBER_INTERNAL;
 		return A == FName(B);
 	}
 
