Fix TSAN warning by avoiding overwriting a value with FArchive_Serialize_BitfieldBool when it is the same during save

#rnx

diff --git a/Engine/Source/Runtime/Core/Public/Serialization/Archive.h b/Engine/Source/Runtime/Core/Public/Serialization/Archive.h
--- a/Engine/Source/Runtime/Core/Public/Serialization/Archive.h
+++ b/Engine/Source/Runtime/Core/Public/Serialization/Archive.h
@@ -57,8 +57,10 @@ struct FArchiveCookContext;
 
 #define DEVIRTUALIZE_FLinkerLoad_Serialize (!WITH_EDITORONLY_DATA)
 
-// Helper macro to make serializing a bitpacked boolean in an archive easier
-#define FArchive_Serialize_BitfieldBool(ARCHIVE, BITFIELD_BOOL) { bool TEMP_BITFIELD_BOOL = BITFIELD_BOOL; ARCHIVE << TEMP_BITFIELD_BOOL; BITFIELD_BOOL = TEMP_BITFIELD_BOOL; }
+// Helper macro to make serializing a bitpacked boolean in an archive easier. 
+// NOTE: The condition is there to avoid overwriting a value that is the same, especially important to make saving an immutable operation and avoid dirtying cachelines for nothing.
+//       This will also make TSAN happy when multiple threads are hitting the same bitfield during concurrent saves.
+#define FArchive_Serialize_BitfieldBool(ARCHIVE, BITFIELD_BOOL) { bool TEMP_BITFIELD_BOOL = BITFIELD_BOOL; ARCHIVE << TEMP_BITFIELD_BOOL; if (BITFIELD_BOOL != TEMP_BITFIELD_BOOL) { BITFIELD_BOOL = TEMP_BITFIELD_BOOL; } }
 
 struct CORE_API FArchiveState
 {
