#UE Add commandline parameter to allow buildmachines to use crashreportclient to upload crashes. Used when testing crash reports

diff --git a/Engine/Source/Runtime/Core/Private/Windows/WindowsPlatformCrashContext.cpp b/Engine/Source/Runtime/Core/Private/Windows/WindowsPlatformCrashContext.cpp
--- a/Engine/Source/Runtime/Core/Private/Windows/WindowsPlatformCrashContext.cpp
+++ b/Engine/Source/Runtime/Core/Private/Windows/WindowsPlatformCrashContext.cpp
@@ -836,7 +836,7 @@ int32 ReportCrashUsingCrashReportClient(FWindowsPlatformCrashContext& InContext,
 			FGenericCrashContext::DumpLog(CrashFolderAbsolute);
 
 			// Build machines do not upload these automatically since it is not okay to have lingering processes after the build completes.
-			if (GIsBuildMachine)
+			if (GIsBuildMachine && !FParse::Param(FCommandLine::Get(), TEXT("AllowCrashReportClientOnBuildMachine")))
 			{
 				return EXCEPTION_CONTINUE_EXECUTION;
 			}
