UE::Tasks: fixed compilation of busy-waiting for multiple tasks

diff --git a/Engine/Source/Runtime/Core/Public/Tasks/Task.h b/Engine/Source/Runtime/Core/Public/Tasks/Task.h
--- a/Engine/Source/Runtime/Core/Public/Tasks/Task.h
+++ b/Engine/Source/Runtime/Core/Public/Tasks/Task.h
@@ -78,7 +78,7 @@ namespace UE::Tasks
 			// waits for task's completion with timeout. Tries to retract the task and execute it in-place, if failed - blocks until the task 
 			// is completed by another thread. If timeout is zero, tries to retract the task and returns immedially after that.
 			// @return true if the task is completed
-			bool Wait(FTimespan Timeout = FTimespan::MaxValue())
+			bool Wait(FTimespan Timeout = FTimespan::MaxValue()) const
 			{
 				return !IsValid() || Pimpl->Wait(FTimeout{ Timeout });
 			}
@@ -86,7 +86,7 @@ namespace UE::Tasks
 			// waits for task's completion for at least the specified amount of time, while executing other tasks.
 			// the call can return much later than the given timeout
 			// @return true if the task is completed
-			bool BusyWait(FTimespan Timeout = FTimespan::MaxValue())
+			bool BusyWait(FTimespan Timeout = FTimespan::MaxValue()) const
 			{
 				return !IsValid() || Pimpl->BusyWait(FTimeout{ Timeout });
 			}
@@ -95,7 +95,7 @@ namespace UE::Tasks
 			// the call can return much later than the given condition became true
 			// @return true if the task is completed
 			template<typename ConditionType>
-			bool BusyWait(ConditionType&& Condition)
+			bool BusyWait(ConditionType&& Condition) const
 			{
 				return !IsValid() || Pimpl->BusyWait(Forward<ConditionType>(Condition));
 			}
diff --git a/Engine/Source/Runtime/Core/Tests/Tasks/TasksTest.cpp b/Engine/Source/Runtime/Core/Tests/Tasks/TasksTest.cpp
--- a/Engine/Source/Runtime/Core/Tests/Tasks/TasksTest.cpp
+++ b/Engine/Source/Runtime/Core/Tests/Tasks/TasksTest.cpp
@@ -84,6 +84,24 @@ namespace UE { namespace TasksTests
 			verify(Event.BusyWait(FTimespan::Zero()));
 		}
 
+		{	// busy-waiting for multiple tasks
+			TArray<FTask> Tasks
+			{
+				Launch(UE_SOURCE_LOCATION,[] {}),
+				Launch(UE_SOURCE_LOCATION,[] {})
+			};
+			BusyWait(Tasks);
+		}
+
+		{	// busy-waiting for multiple tasks
+			TArray<FTask> Tasks
+			{
+				Launch(UE_SOURCE_LOCATION,[] {}),
+				Launch(UE_SOURCE_LOCATION,[] {})
+			};
+			BusyWait(Tasks, FTimespan::FromMilliseconds(10));
+		}
+
 		{	// basic use-case, postpone waiting so the task is executed first
 			std::atomic<bool> Done{ false };
 			FTask Task = Launch(UE_SOURCE_LOCATION, [&Done] { Done = true; });
