I/O - fixed read issues from local file cache on Switch

#rnx

diff --git a/Engine/Source/Runtime/Core/Private/IO/IoFileCache.cpp b/Engine/Source/Runtime/Core/Private/IO/IoFileCache.cpp
--- a/Engine/Source/Runtime/Core/Private/IO/IoFileCache.cpp
+++ b/Engine/Source/Runtime/Core/Private/IO/IoFileCache.cpp
@@ -20,6 +20,7 @@
 
 TRACE_DECLARE_MEMORY_COUNTER(FFileIoCache_CachedBytes, TEXT("FileIoCache/TotalBytes"));
 TRACE_DECLARE_MEMORY_COUNTER(FFileIoCache_PendingBytes, TEXT("FileIoCache/PendingBytes"));
+TRACE_DECLARE_MEMORY_COUNTER(FFileIoCache_ReadBytes, TEXT("FileIoCache/TotalReadBytes"));
 TRACE_DECLARE_INT_COUNTER(FFileIoCache_GetCount, TEXT("FileIoCache/GetCount"));
 TRACE_DECLARE_INT_COUNTER(FFileIoCache_ErrorCount, TEXT("FileIoCache/ErrorCount"));
 TRACE_DECLARE_INT_COUNTER(FFileIoCache_PutCount, TEXT("FileIoCache/PutCount"));
@@ -412,9 +413,9 @@ private:
 	FFileIoCacheConfig CacheConfig;
 	FCacheMap CacheMap;
 	TUniquePtr<FRunnableThread> WriterThread;
-	TUniquePtr<IFileHandle> WriteFileHandle;
 	FEventRef TickWriterEvent;
 	FString CacheFilePath;
+	uint64 WriteCursorPos = 0;
 	std::atomic_bool bStopRequested{false};
 };
 
@@ -475,20 +476,35 @@ TTask<TIoStatusOr<FIoBuffer>> FFileIoCache::GetChunk(const FIoHash& Key, const F
 				FIoBuffer Buffer = Options.GetTargetVa() ? FIoBuffer(FIoBuffer::Wrap, Options.GetTargetVa(), ReadSize) : FIoBuffer(ReadSize);
 
 				IPlatformFile& Ipf = IPlatformFile::GetPlatformPhysical();
-				TUniquePtr<IFileHandle> FileHandle(Ipf.OpenRead(*CacheFilePath, true));
-				check(FileHandle);
-				FileHandle->Seek(int64(ReadOffset));
-				FileHandle->Read(Buffer.GetData(), ReadSize);
-
-				const FIoHash& ExpectedHash = Entry.Hash;
-				if (ExpectedHash == FIoHash::HashBuffer(Buffer.GetView()))
+				if (TUniquePtr<IFileHandle> FileHandle(Ipf.OpenRead(*CacheFilePath, false)); FileHandle.IsValid())
 				{
-					TRACE_COUNTER_INCREMENT(FFileIoCache_GetCount);
-					return TIoStatusOr<FIoBuffer>(Buffer);
+					UE_LOG(LogIoCache, VeryVerbose, TEXT("Read chunk, Key='%s', Hash='%s', File='%s', Offset='%llu', Size='%llu'"),
+						*LexToString(Entry.Key), *LexToString(Entry.Hash), *CacheFilePath, Entry.SerialOffset, Entry.SerialSize);
+					
+					FileHandle->Seek(int64(ReadOffset));
+					FileHandle->Read(Buffer.GetData(), ReadSize);
+
+					const FIoHash& ExpectedHash = Entry.Hash;
+					const FIoHash Hash = FIoHash::HashBuffer(Buffer.GetView());
+					if (Hash == ExpectedHash) 
+					{
+						TRACE_COUNTER_INCREMENT(FFileIoCache_GetCount);
+						TRACE_COUNTER_ADD(FFileIoCache_ReadBytes, ReadSize);
+						return TIoStatusOr<FIoBuffer>(Buffer);
+					}
+
+					TRACE_COUNTER_INCREMENT(FFileIoCache_ErrorCount);
+					UE_LOG(LogIoCache, Verbose, TEXT("Read chunk failed, hash mismatch, Key='%s', Hash='%s', ExpectedHash='%s', File='%s', Offset='%llu', Size='%llu'"),
+						*LexToString(Entry.Key), *LexToString(Hash), *LexToString(ExpectedHash), *CacheFilePath, ReadOffset, ReadSize);
+
+					return TIoStatusOr<FIoBuffer>(FIoStatus(EIoErrorCode::ReadError));
+				}
+				else
+				{
+					TRACE_COUNTER_INCREMENT(FFileIoCache_ErrorCount);
+					UE_LOG(LogIoCache, Warning, TEXT("Read chunk failed, unable to open cache file '%s' for reading"), *CacheFilePath);
+					return TIoStatusOr<FIoBuffer>(FIoStatus(EIoErrorCode::ReadError));
 				}
-
-				TRACE_COUNTER_INCREMENT(FFileIoCache_ErrorCount);
-				return TIoStatusOr<FIoBuffer>(FIoStatus(EIoErrorCode::ReadError));
 			}
 		});
 }
@@ -516,6 +532,7 @@ void FFileIoCache::Initialize()
 	const FString CacheDir = FPaths::ProjectPersistentDownloadDir() / TEXT("IoCache");
 	const FString CacheTocPath = CacheDir / TEXT("cache.utoc");
 	CacheFilePath = CacheDir / TEXT("cache.ucas");
+	WriteCursorPos = 0;
 
 	IPlatformFile& Ipf = IPlatformFile::GetPlatformPhysical();
 	IFileManager& FileMgr = IFileManager::Get();
@@ -530,19 +547,15 @@ void FFileIoCache::Initialize()
 		}
 		else
 		{
-			uint64 CursorPos = ~uint64(0);
-
-			CacheStatus = CacheMap.Load(CacheTocPath, CursorPos);
+			CacheStatus = CacheMap.Load(CacheTocPath, WriteCursorPos);
 			if (CacheStatus.IsOk())
 			{
-				check(CursorPos != ~uint64(0));
+				check(WriteCursorPos != ~uint64(0));
 				
 				UE_LOG(LogIoCache, Log, TEXT("Loaded TOC '%s'"), *CacheTocPath);
 				if (FileMgr.FileExists(*CacheFilePath))
 				{
-					WriteFileHandle.Reset(Ipf.OpenWrite(*CacheFilePath, true, true));
-					WriteFileHandle->Seek(CursorPos);
-					//TODO: Integrity check
+					//TODO: Integrity check?
 				}
 				else
 				{
@@ -559,6 +572,7 @@ void FFileIoCache::Initialize()
 
 	if (!CacheStatus.IsOk())
 	{
+		WriteCursorPos = 0;
 		CacheMap.Reset();
 		FileMgr.Delete(*CacheFilePath);
 
@@ -566,7 +580,6 @@ void FFileIoCache::Initialize()
 		{
 			FileMgr.MakeDirectory(*CacheDir, true);
 		}
-		WriteFileHandle.Reset(Ipf.OpenWrite(*CacheFilePath, true, true));
 	}
 
 	WriterThread.Reset(FRunnableThread::Create(this, TEXT("File I/O Cache"), 0, TPri_BelowNormal));
@@ -585,12 +598,13 @@ void FFileIoCache::Shutdown()
 
 	const FString CacheTocPath = FPaths::ProjectPersistentDownloadDir() / TEXT("IoCache") / TEXT("cache.utoc");
 	UE_LOG(LogIoCache, Log, TEXT("Saving TOC '%s'"), *CacheTocPath);
-	CacheMap.Save(CacheTocPath, WriteFileHandle->Tell());
-	WriteFileHandle.Reset();
+	CacheMap.Save(CacheTocPath, WriteCursorPos);
 }
 
 uint32 FFileIoCache::FileWriterThreadEntry()
 {
+	IPlatformFile& Ipf = IPlatformFile::GetPlatformPhysical();
+
 	while (!bStopRequested)
 	{
 		for (;;)
@@ -601,6 +615,15 @@ uint32 FFileIoCache::FileWriterThreadEntry()
 				break;
 			}
 
+			TUniquePtr<IFileHandle> WriteFileHandle(Ipf.OpenWrite(*CacheFilePath, true, true));
+			if (!WriteFileHandle.IsValid())
+			{
+				UE_LOG(LogIoCache, Warning, TEXT("Write chunks failed, unable to open file '%s' for writing"), *CacheFilePath);
+				break;
+			}
+
+			WriteFileHandle->Seek(WriteCursorPos);
+
 			TRACE_CPUPROFILER_EVENT_SCOPE(FFileIoCache::WriteCacheEntry);
 
 			uint64 TotalPendingSize = 0;
@@ -621,6 +644,9 @@ uint32 FFileIoCache::FileWriterThreadEntry()
 				Entry.SerialSize = Data.GetSize();
 				Entry.Hash = FIoHash::HashBuffer(Data.GetView());
 
+				UE_LOG(LogIoCache, VeryVerbose, TEXT("Write chunk, Key='%s', Hash='%s', File='%s', Offset='%llu', Size='%llu'"),
+					*LexToString(Entry.Key), *LexToString(Entry.Hash), *CacheFilePath, Entry.SerialOffset, Entry.SerialSize);
+
 				const uint64 RemainingDiskSize = CacheConfig.DiskStorageSize - WriteFileHandle->Tell();
 				const uint64 ByteCount = FMath::Min(Entry.Data.GetSize(), RemainingDiskSize);
 				WriteFileHandle->Write(Entry.Data.GetData(), ByteCount);
@@ -628,13 +654,15 @@ uint32 FFileIoCache::FileWriterThreadEntry()
 				const uint64 RemainingChunkSize = Entry.Data.GetSize() - ByteCount;
 				if (RemainingChunkSize > 0)
 				{
+					WriteFileHandle->Flush();
 					WriteFileHandle->Seek(0);
 					WriteFileHandle->Write(Entry.Data.GetData() + ByteCount, RemainingChunkSize);
 				}
-
-				WriteFileHandle->Flush();
 			}
 
+			WriteFileHandle->Flush();
+			WriteCursorPos = WriteFileHandle->Tell();
+
 			CacheMap.InsertPersisted(MoveTemp(Entries), WriteFileHandle->Tell());
 		}
 
