Added helper method GetGuidFromHash in FArchiveMD5.

diff --git a/Engine/Source/Runtime/Core/Public/Misc/ArchiveMD5.h b/Engine/Source/Runtime/Core/Public/Misc/ArchiveMD5.h
--- a/Engine/Source/Runtime/Core/Public/Misc/ArchiveMD5.h
+++ b/Engine/Source/Runtime/Core/Public/Misc/ArchiveMD5.h
@@ -50,6 +50,13 @@ public:
 		Hash.Set(MD5);
 	}
 
+	FGuid GetGuidFromHash()
+	{
+		FMD5Hash MD5Hash;
+		GetHash(MD5Hash);
+		return MD5HashToGuid(MD5Hash);
+	}
+
 protected:
 	FMD5 MD5;
 };
\ No newline at end of file
