Adding printf handler to oodle data.

diff --git a/Engine/Source/Runtime/Core/Private/Compression/OodleDataCompression.cpp b/Engine/Source/Runtime/Core/Private/Compression/OodleDataCompression.cpp
--- a/Engine/Source/Runtime/Core/Private/Compression/OodleDataCompression.cpp
+++ b/Engine/Source/Runtime/Core/Private/Compression/OodleDataCompression.cpp
@@ -574,6 +574,18 @@ static void OODLE_CALLBACK OodleFree(void* Ptr)
 	FMemory::Free(Ptr);
 }
 
+static void OODLE_CALLBACK OodlePrintf(int verboseLevel, const char* file, int line, const char* fmt, ...)
+{
+	TAnsiStringBuilder<256> OodleOutput;
+
+	va_list Args;
+	va_start(Args, fmt);
+	OodleOutput.AppendV(fmt, Args);
+	va_end(Args);
+	
+	UE_LOG(OodleDataCompression, Display, TEXT("Oodle: %hs"), *OodleOutput);
+}
+
 void CORE_API StartupPreInit(void)
 {
 	// called from LaunchEngineLoop at "PreInit" time
@@ -590,6 +602,7 @@ void CORE_API StartupPreInit(void)
 	// 
 	// install FMemory for Oodle allocators :
 	OodleCore_Plugins_SetAllocators(OodleAlloc, OodleFree);
+	OodleCore_Plugins_SetPrintf(OodlePrintf);
 }
 
 
