Expose is connected in interface

Exposes a variant of the the IsConnected method that also outputs the current session and trace guid. Also prints the current session and trace guids in the status message.

#rnx

diff --git a/Engine/Source/Runtime/Core/Private/ProfilingDebugging/TraceAuxiliary.cpp b/Engine/Source/Runtime/Core/Private/ProfilingDebugging/TraceAuxiliary.cpp
--- a/Engine/Source/Runtime/Core/Private/ProfilingDebugging/TraceAuxiliary.cpp
+++ b/Engine/Source/Runtime/Core/Private/ProfilingDebugging/TraceAuxiliary.cpp
@@ -83,6 +83,7 @@ class FTraceAuxiliaryImpl
 public:
 	const TCHAR* GetDest() const;
 	bool IsConnected() const;
+	bool IsConnected(FGuid& OutSessionGuid, FGuid& OutTraceGuid);
 	FTraceAuxiliary::EConnectionType GetConnectionType() const;
 	void GetActiveChannelsString(FStringBuilderBase& String) const;
 	void AddCommandlineChannels(const TCHAR* ChannelList);
@@ -619,6 +620,20 @@ bool FTraceAuxiliaryImpl::IsConnected() const
 	return State.load() == EState::Tracing;
 }
 
+////////////////////////////////////////////////////////////////////////////////
+bool FTraceAuxiliaryImpl::IsConnected(FGuid& OutSessionGuid, FGuid& OutTraceGuid)
+{
+	uint32 SessionGuid[4];
+	uint32 TraceGuid[4];
+	if (UE::Trace::IsTracingTo(SessionGuid, TraceGuid))
+	{
+		OutSessionGuid = *(FGuid*)SessionGuid;
+		OutTraceGuid = *(FGuid*)TraceGuid;
+		return true;
+	}
+	return false;
+}
+
 ////////////////////////////////////////////////////////////////////////////////
 FTraceAuxiliary::EConnectionType FTraceAuxiliaryImpl::GetConnectionType() const
 {
@@ -835,12 +850,14 @@ static void TraceAuxiliaryStatus()
 
 	// Status of data connection
 	TStringBuilder<256> ConnectionStr;
-	if (UE::Trace::IsTracing())
+	FGuid SessionGuid, TraceGuid;
+	if (GTraceAuxiliary.IsConnected(SessionGuid, TraceGuid))
 	{
 		const TCHAR* Dest = GTraceAuxiliary.GetDest();
 		if (Dest && FCString::Strlen(Dest) > 0)
 		{
-			ConnectionStr.Appendf(TEXT("Tracing to '%s'"), Dest);
+			ConnectionStr.Appendf(TEXT("Tracing to '%s', "), Dest);
+			ConnectionStr.Appendf(TEXT("session %s trace %s"), *SessionGuid.ToString(), *TraceGuid.ToString());
 		}
 		else
 		{
@@ -1694,6 +1711,15 @@ bool FTraceAuxiliary::IsConnected()
 	return false;
 }
 
+bool FTraceAuxiliary::IsConnected(FGuid& OutSessionGuid, FGuid& OutTraceGuid)
+{
+#if UE_TRACE_ENABLED
+	return GTraceAuxiliary.IsConnected(OutSessionGuid, OutTraceGuid);
+#else
+	return false;
+#endif
+}
+
 FTraceAuxiliary::EConnectionType FTraceAuxiliary::GetConnectionType()
 {
 #if UE_TRACE_ENABLED
diff --git a/Engine/Source/Runtime/Core/Public/ProfilingDebugging/TraceAuxiliary.h b/Engine/Source/Runtime/Core/Public/ProfilingDebugging/TraceAuxiliary.h
--- a/Engine/Source/Runtime/Core/Public/ProfilingDebugging/TraceAuxiliary.h
+++ b/Engine/Source/Runtime/Core/Public/ProfilingDebugging/TraceAuxiliary.h
@@ -157,6 +157,15 @@ public:
 	 */
 	static bool	IsConnected();
 
+
+	/**
+	 * Returns whether the trace system is currently connected. If connected writes the session/trace identifiers.
+	 * @param OutSessionGuid If connected, the session guid
+	 * @param OutTraceGuid If connected, the trace guid
+	 * @return True if connected
+	 */
+	static bool IsConnected(FGuid& OutSessionGuid, FGuid& OutTraceGuid);
+
 	/**
 	*  Returns the current connection type.
 	*/
