Fix TMap::FilterByPredicate discarding template arguments in return type.

diff --git a/Engine/Source/Runtime/Core/Public/Containers/Map.h b/Engine/Source/Runtime/Core/Public/Containers/Map.h
--- a/Engine/Source/Runtime/Core/Public/Containers/Map.h
+++ b/Engine/Source/Runtime/Core/Public/Containers/Map.h
@@ -516,9 +516,9 @@ public:
 	 *          the subset of elements for which the functor returns true.
 	 */
 	template <typename Predicate>
-	TMap<KeyType, ValueType> FilterByPredicate(Predicate Pred) const
+	TMap<KeyType, ValueType, SetAllocator, KeyFuncs> FilterByPredicate(Predicate Pred) const
 	{
-		TMap<KeyType, ValueType> FilterResults;
+		TMap<KeyType, ValueType, SetAllocator, KeyFuncs> FilterResults;
 		FilterResults.Reserve(Pairs.Num());
 		for (const ElementType& Pair : Pairs)
 		{
