Fixed mismatch string formatting argument in core test (PathViewTest).

diff --git a/Engine/Source/Runtime/Core/Tests/Misc/PathViewsTest.cpp b/Engine/Source/Runtime/Core/Tests/Misc/PathViewsTest.cpp
--- a/Engine/Source/Runtime/Core/Tests/Misc/PathViewsTest.cpp
+++ b/Engine/Source/Runtime/Core/Tests/Misc/PathViewsTest.cpp
@@ -785,7 +785,7 @@ TEST_CASE_NAMED(FPathViewsGetMountPointNameTest, "System::Core::Misc::PathViews:
 		Result = FPathViews::GetMountPointNameFromPath(Base, &bHadClassPrefix, false);
 		if (Result != FStringView(ExpectedWithSlash))
 		{
-			FAIL_CHECK(FString::Printf(TEXT("AppendPath(\"%s\", \"%s\") == \"%.*s\", expected \"%s\"."),
+			FAIL_CHECK(FString::Printf(TEXT("GetMountPointNameFromPath(\"%s\", &bHadClassPrefix, false) == \"%.*s\", expected \"%s\"."),
 				Base, Result.Len(), Result.GetData(), ExpectedWithSlash));
 		}
 	};
