Fixed typos in TUniquePtr comments about TPimplPtr.

diff --git a/Engine/Source/Runtime/Core/Public/Templates/UniquePtr.h b/Engine/Source/Runtime/Core/Public/Templates/UniquePtr.h
--- a/Engine/Source/Runtime/Core/Public/Templates/UniquePtr.h
+++ b/Engine/Source/Runtime/Core/Public/Templates/UniquePtr.h
@@ -59,8 +59,8 @@ struct TDefaultDelete
 		//   .gen.cpp file where your pimpl type doesn't exist.  It cannot be defaulted, but it need not
 		//   contain any particular implementation; the object just needs to be garbage collectable.
 		//
-		// If this is efficiency is less important than simplicity, you may want to consider
-		// using a TPimplPtr instead, though this is also pretty efficient too.
+		// If efficiency is less important than simplicity, you may want to consider
+		// using a TPimplPtr instead, though it is also pretty efficient.
 		delete Ptr;
 	}
 };
