Adding some logging to UnixSignalHeartBeat.cpp to help diagnose issue

diff --git a/Engine/Source/Runtime/Core/Private/Unix/UnixSignalHeartBeat.cpp b/Engine/Source/Runtime/Core/Private/Unix/UnixSignalHeartBeat.cpp
--- a/Engine/Source/Runtime/Core/Private/Unix/UnixSignalHeartBeat.cpp
+++ b/Engine/Source/Runtime/Core/Private/Unix/UnixSignalHeartBeat.cpp
@@ -56,6 +56,7 @@ namespace
 	void SignalHitchHandler(int Signal)
 	{
 #if USE_HITCH_DETECTION
+		UE_LOG(LogUnixHeartBeat, Log, TEXT("SignalHitchHandler"));
 		GHitchDetected = true;
 #endif
 	}
@@ -149,6 +150,9 @@ void FUnixSignalGameHitchHeartBeat::InitSettings()
 		{
 			SuspendCount = 1;
 		}
+
+		UE_LOG(LogUnixHeartBeat, Log, TEXT("bStartSuspended:%d HitchThresholdS:%f"),
+			bStartSuspended ? 1 : 0, HitchThresholdS);
 	}
 }
 
@@ -208,6 +212,7 @@ void FUnixSignalGameHitchHeartBeat::SuspendHeartBeat()
 	}
 
 	SuspendCount++;
+	UE_LOG(LogUnixHeartBeat, Log, TEXT("SuspendCount:%d"), SuspendCount);
 
 	if (TimerId)
 	{
@@ -234,6 +239,7 @@ void FUnixSignalGameHitchHeartBeat::ResumeHeartBeat()
 	if( SuspendCount > 0)
 	{
 		SuspendCount--;
+		UE_LOG(LogUnixHeartBeat, Log, TEXT("SuspendCount:%d"), SuspendCount);
 
 		FrameStart(true);
 	}
