Made FBufferedLine movable.

diff --git a/Engine/Source/Developer/OutputLog/Private/SDeviceOutputLog.cpp b/Engine/Source/Developer/OutputLog/Private/SDeviceOutputLog.cpp
--- a/Engine/Source/Developer/OutputLog/Private/SDeviceOutputLog.cpp
+++ b/Engine/Source/Developer/OutputLog/Private/SDeviceOutputLog.cpp
@@ -170,7 +170,7 @@ void SDeviceOutputLog::Tick(const FGeometry& AllottedGeometry, const double InCu
 	{
 		for (const FBufferedLine& Line : BufferedLines)
 		{
-			MessagesTextMarshaller->AppendPendingMessage(Line.Data, Line.Verbosity, Line.Category);
+			MessagesTextMarshaller->AppendPendingMessage(Line.Data.Get(), Line.Verbosity, Line.Category);
 		}
 		BufferedLines.Empty(32);
 	}
diff --git a/Engine/Source/Runtime/Core/Private/Misc/BufferedOutputDevice.cpp b/Engine/Source/Runtime/Core/Private/Misc/BufferedOutputDevice.cpp
--- a/Engine/Source/Runtime/Core/Private/Misc/BufferedOutputDevice.cpp
+++ b/Engine/Source/Runtime/Core/Private/Misc/BufferedOutputDevice.cpp
@@ -27,6 +27,6 @@ void FBufferedOutputDevice::RedirectTo(FOutputDevice& Ar)
 	FScopeLock ScopeLock(&SynchronizationObject);
 	for (const FBufferedLine& BufferedLine : BufferedLines)
 	{
-		Ar.Serialize(BufferedLine.Data, BufferedLine.Verbosity, BufferedLine.Category);
+		Ar.Serialize(BufferedLine.Data.Get(), BufferedLine.Verbosity, BufferedLine.Category);
 	}
 }
diff --git a/Engine/Source/Runtime/Core/Private/Misc/OutputDeviceRedirector.cpp b/Engine/Source/Runtime/Core/Private/Misc/OutputDeviceRedirector.cpp
--- a/Engine/Source/Runtime/Core/Private/Misc/OutputDeviceRedirector.cpp
+++ b/Engine/Source/Runtime/Core/Private/Misc/OutputDeviceRedirector.cpp
@@ -28,13 +28,8 @@ FBufferedLine::FBufferedLine(const TCHAR* InData, const FName& InCategory, ELogV
 	, Verbosity(InVerbosity)
 {
 	int32 NumChars = FCString::Strlen(InData) + 1;
-	void* Dest = FMemory::Malloc(sizeof(TCHAR) * NumChars);
-	Data = (TCHAR*)FMemory::Memcpy(Dest, InData, sizeof(TCHAR) * NumChars);
-}
-
-FBufferedLine::~FBufferedLine()
-{
-	FMemory::Free(const_cast<TCHAR*>(Data));
+	Data = MakeUniqueForOverwrite<TCHAR[]>(sizeof(TCHAR) * NumChars);
+	FMemory::Memcpy(Data.Get(), InData, sizeof(TCHAR) * NumChars);
 }
 
 namespace UE::Private
@@ -578,7 +573,7 @@ void FOutputDeviceRedirector::SerializeBacklog(FOutputDevice* OutputDevice)
 	FReadScopeLock ScopeLock(State->BacklogLock);
 	for (const FBufferedLine& BacklogLine : State->BacklogLines)
 	{
-		OutputDevice->Serialize(BacklogLine.Data, BacklogLine.Verbosity, BacklogLine.Category, BacklogLine.Time);
+		OutputDevice->Serialize(BacklogLine.Data.Get(), BacklogLine.Verbosity, BacklogLine.Category, BacklogLine.Time);
 	}
 }
 
diff --git a/Engine/Source/Runtime/Core/Public/Misc/OutputDeviceRedirector.h b/Engine/Source/Runtime/Core/Public/Misc/OutputDeviceRedirector.h
--- a/Engine/Source/Runtime/Core/Public/Misc/OutputDeviceRedirector.h
+++ b/Engine/Source/Runtime/Core/Public/Misc/OutputDeviceRedirector.h
@@ -9,6 +9,7 @@
 #include "Misc/EnumClassFlags.h"
 #include "Misc/OutputDevice.h"
 #include "Templates/PimplPtr.h"
+#include "Templates/UniquePtr.h"
 #include "UObject/NameTypes.h"
 
 /*-----------------------------------------------------------------------------
@@ -20,34 +21,13 @@ namespace UE::Private { struct FOutputDeviceRedirectorState; }
 /** The type of lines buffered by secondary threads. */
 struct CORE_API FBufferedLine
 {
-	enum EBufferedLineInit
-	{
-		EMoveCtor = 0
-	};
-
-	const TCHAR* Data;
+	TUniquePtr<TCHAR[]> Data;
 	const FLazyName Category;
 	const double Time;
 	const ELogVerbosity::Type Verbosity;
 	bool bExternalAllocation;
 
-	FBufferedLine(const TCHAR* InData, const FName& InCategory, ELogVerbosity::Type InVerbosity, const double InTime = -1);
-
-	FBufferedLine(FBufferedLine& InBufferedLine, EBufferedLineInit Unused)
-		: Data(InBufferedLine.Data)
-		, Category(InBufferedLine.Category)
-		, Time(InBufferedLine.Time)
-		, Verbosity(InBufferedLine.Verbosity)
-		, bExternalAllocation(InBufferedLine.bExternalAllocation)
-	{
-		InBufferedLine.Data = nullptr;
-		InBufferedLine.bExternalAllocation = false;
-	}
-
-	/** Noncopyable for now, could be made movable */
-	FBufferedLine(const FBufferedLine&) = delete;
-	FBufferedLine& operator=(const FBufferedLine&) = delete;
-	~FBufferedLine();
+	explicit FBufferedLine(const TCHAR* InData, const FName& InCategory, ELogVerbosity::Type InVerbosity, const double InTime = -1);
 };
 
 enum class EOutputDeviceRedirectorFlushOptions : uint32
diff --git a/Engine/Source/Runtime/Engine/Private/UnrealEngine.cpp b/Engine/Source/Runtime/Engine/Private/UnrealEngine.cpp
--- a/Engine/Source/Runtime/Engine/Private/UnrealEngine.cpp
+++ b/Engine/Source/Runtime/Engine/Private/UnrealEngine.cpp
@@ -11275,7 +11275,7 @@ bool UEngine::FErrorsAndWarningsCollector::Tick(float Seconds)
 		// Remove any dupes and count them
 		do 
 		{
-			uint32 ThisHash = FCrc::StrCrc32(ActiveLines[DupeCount].Data);
+			uint32 ThisHash = FCrc::StrCrc32(ActiveLines[DupeCount].Data.Get());
 
 			if (CurrentHash && ThisHash != CurrentHash)
 			{
@@ -11288,7 +11288,7 @@ bool UEngine::FErrorsAndWarningsCollector::Tick(float Seconds)
 		} while (DupeCount < ActiveLines.Num());
 
 		// Save off properties
-		FString Msg = ActiveLines[0].Data;
+		FString Msg = ActiveLines[0].Data.Get();
 		ELogVerbosity::Type Verbosity = ActiveLines[0].Verbosity;
 
 		// Remove any lines we condensed
