Fixed incorrect size being reported by TArray's overflow checker.
Fixed negative int32 sizes being reported as really large uint64 sizes - now they're reported as moderately large uint32 sizes.

diff --git a/Engine/Source/Runtime/Core/Public/Containers/Array.h b/Engine/Source/Runtime/Core/Public/Containers/Array.h
--- a/Engine/Source/Runtime/Core/Public/Containers/Array.h
+++ b/Engine/Source/Runtime/Core/Public/Containers/Array.h
@@ -419,7 +419,8 @@ public:
 	{
 		if (Count < 0)
 		{
-			UE::Core::Private::OnInvalidArrayNum((unsigned long long)Count);
+			// Cast to USizeType first to prevent sign extension on negative sizes, producing unusually large values.
+			UE::Core::Private::OnInvalidArrayNum((unsigned long long)(USizeType)Count);
 		}
 
 		check(Ptr != nullptr || Count == 0);
@@ -566,7 +567,8 @@ private:
 			{
 				if (ToArray.ArrayNum != FromArray.ArrayNum || ToArray.ArrayMax != FromArray.ArrayMax)
 				{
-					UE::Core::Private::OnInvalidArrayNum((unsigned long long)ToArray.ArrayNum);
+					// Cast to USizeType first to prevent sign extension on negative sizes, producing unusually large values.
+					UE::Core::Private::OnInvalidArrayNum((unsigned long long)(USizeType)ToArray.ArrayNum);
 				}
 			}
 
@@ -1808,7 +1810,8 @@ public:
 	{
 		if (NewSize < 0)
 		{
-			UE::Core::Private::OnInvalidArrayNum((unsigned long long)NewSize);
+			// Cast to USizeType first to prevent sign extension on negative sizes, producing unusually large values.
+			UE::Core::Private::OnInvalidArrayNum((unsigned long long)(USizeType)NewSize);
 		}
 
 		// If we have space to hold the excepted size, then don't reallocate
@@ -1832,7 +1835,8 @@ public:
 	{
 		if (Slack < 0)
 		{
-			UE::Core::Private::OnInvalidArrayNum((unsigned long long)Slack);
+			// Cast to USizeType first to prevent sign extension on negative sizes, producing unusually large values.
+			UE::Core::Private::OnInvalidArrayNum((unsigned long long)(USizeType)Slack);
 		}
 
 		DestructItems(GetData(), ArrayNum);
@@ -1862,7 +1866,8 @@ public:
 		}
 		else if (NewNum < 0)
 		{
-			UE::Core::Private::OnInvalidArrayNum((unsigned long long)NewNum);
+			// Cast to USizeType first to prevent sign extension on negative sizes, producing unusually large values.
+			UE::Core::Private::OnInvalidArrayNum((unsigned long long)(USizeType)NewNum);
 		}
 		else if (NewNum < Num())
 		{
@@ -1885,7 +1890,8 @@ public:
 		}
 		else if (NewNum < 0)
 		{
-			UE::Core::Private::OnInvalidArrayNum((unsigned long long)NewNum);
+			// Cast to USizeType first to prevent sign extension on negative sizes, producing unusually large values.
+			UE::Core::Private::OnInvalidArrayNum((unsigned long long)(USizeType)NewNum);
 		}
 		else if (NewNum < Num())
 		{
@@ -1906,7 +1912,8 @@ public:
 		}
 		else if (NewNum < 0)
 		{
-			UE::Core::Private::OnInvalidArrayNum((unsigned long long)NewNum);
+			// Cast to USizeType first to prevent sign extension on negative sizes, producing unusually large values.
+			UE::Core::Private::OnInvalidArrayNum((unsigned long long)(USizeType)NewNum);
 		}
 		else if (NewNum < Num())
 		{
@@ -2496,7 +2503,8 @@ public:
 		checkSlow(Number >= 0);
 		if (Number < 0)
 		{
-			UE::Core::Private::OnInvalidArrayNum((unsigned long long)Number);
+			// Cast to USizeType first to prevent sign extension on negative sizes, producing unusually large values.
+			UE::Core::Private::OnInvalidArrayNum((unsigned long long)(USizeType)Number);
 		}
 		else if (Number > ArrayMax)
 		{
@@ -2972,7 +2980,8 @@ private:
 		// This should only happen when we've underflowed or overflowed SizeType in the caller
 		if (LocalArrayNum < OldNum)
 		{
-			UE::Core::Private::OnInvalidArrayNum((unsigned long long)LocalArrayNum - (unsigned long long)OldNum);
+			// Cast to USizeType first to prevent sign extension on negative sizes, producing unusually large values.
+			UE::Core::Private::OnInvalidArrayNum((unsigned long long)(USizeType)LocalArrayNum);
 		}
 		ArrayMax = AllocatorCalculateSlackGrow(LocalArrayNum, ArrayMax);
 		AllocatorResizeAllocation(OldNum, ArrayMax);
