Fix a PVS-Studio code analysis false positive.

diff --git a/Engine/Source/Runtime/Core/Public/HAL/FMemory.inl b/Engine/Source/Runtime/Core/Public/HAL/FMemory.inl
--- a/Engine/Source/Runtime/Core/Public/HAL/FMemory.inl
+++ b/Engine/Source/Runtime/Core/Public/HAL/FMemory.inl
@@ -43,7 +43,12 @@ FMEMORY_INLINE_FUNCTION_DECORATOR void* FMemory::Malloc(SIZE_T Count, uint32 Ali
 	// AutoRTFM: This is a no-op for non-transactional code.
 	// For transactional code, this defers a call to Free if the transaction aborts,
 	// so that rolling back this allocation will end up freeing the memory.
-	AutoRTFM::OpenAbort([Ptr] { Free(Ptr); });
+	AutoRTFM::OpenAbort([Ptr]
+	{
+		// Disable the code analysis warning that complains that Free is being passed
+		// a pointer that may be null. Free explicitly handles this case already.
+		Free(Ptr); //-V575
+	});
 
 	return AutoRTFM::DidAllocate(Ptr, Count);
 }
