Early return from ParallelFor when there is no work to do.
Avoids making very small profiling scopes in some cases.

diff --git a/Engine/Source/Runtime/Core/Public/Async/ParallelFor.h b/Engine/Source/Runtime/Core/Public/Async/ParallelFor.h
--- a/Engine/Source/Runtime/Core/Public/Async/ParallelFor.h
+++ b/Engine/Source/Runtime/Core/Public/Async/ParallelFor.h
@@ -114,6 +114,10 @@ namespace ParallelForImpl
 	template<typename BodyType, typename PreWorkType, typename ContextType>
 	inline void ParallelForInternal(const TCHAR* DebugName, int32 Num, int32 MinBatchSize, BodyType Body, PreWorkType CurrentThreadWorkToDoBeforeHelping, EParallelForFlags Flags, const TArrayView<ContextType>& Contexts)
 	{
+		if (Num == 0) 
+		{
+			return;
+		}
 		SCOPE_CYCLE_COUNTER(STAT_ParallelFor);
 		TRACE_CPUPROFILER_EVENT_SCOPE(ParallelFor);
 		check(Num >= 0);
