Reinstate deleted macros to make sure the correct memory allocator is in use on Apple platforms

diff --git a/Engine/Source/Runtime/Core/Private/Apple/ApplePlatformMemory.cpp b/Engine/Source/Runtime/Core/Private/Apple/ApplePlatformMemory.cpp
--- a/Engine/Source/Runtime/Core/Private/Apple/ApplePlatformMemory.cpp
+++ b/Engine/Source/Runtime/Core/Private/Apple/ApplePlatformMemory.cpp
@@ -287,47 +287,65 @@ void FApplePlatformMemory::Init()
 	
 }
 
-// Set rather to use BinnedMalloc2 for binned malloc, can be overridden below
-#define USE_MALLOC_BINNED2 (PLATFORM_MAC)
+// Use MallocBinned2 as default, can be overriden below.
+#define USE_MALLOC_BINNED2 1
 
-FMalloc* FApplePlatformMemory::BaseAllocator()
+void FApplePlatformMemory::SetAllocatorToUse()
 {
-	static FMalloc* Instance = nullptr;
-	if (Instance != nullptr)
-	{
-		return Instance;
-	}
-
-#if ENABLE_LOW_LEVEL_MEM_TRACKER
-	FPlatformMemoryStats MemStats = FApplePlatformMemory::GetStats();
-	FLowLevelMemTracker::Get().SetProgramSize(MemStats.UsedPhysical);
-#endif
-
-#if PLATFORM_IOS || PLATFORM_TVOS
-    if (FIOSPlatformMisc::IsEntitlementEnabled("com.apple.developer.kernel.extended-virtual-addressing"))
+    // force Ansi allocator in particular cases
+    if(getenv("UE4_FORCE_MALLOC_ANSI") != nullptr)
     {
-        UE_LOG(LogTemp, Display, TEXT("Virtual Address Space entitlement found. Using MallocBinned2 to allocate memory"));
-        AllocatorToUse = EMemoryAllocatorToUse::Binned2;
-#undef USE_MALLOC_BINNED2
-#define USE_MALLOC_BINNED2 1
+        UE_LOG(LogTemp, Display, TEXT("Using Ansi allocator."));
+        AllocatorToUse = EMemoryAllocatorToUse::Ansi;
+        return;
     }
-    else
+    if (FORCE_ANSI_ALLOCATOR)
     {
-        UE_LOG(LogTemp, Warning, TEXT("Virtual Address Space entitlement NOT found. Defaulting to Ansi allocator"));
+        UE_LOG(LogTemp, Display, TEXT("Using Ansi allocator."));
+        AllocatorToUse = EMemoryAllocatorToUse::Ansi;
+        return;
     }
-#elif USE_MALLOC_BINNED2
     if (USE_MALLOC_BINNED2)
     {
-        UE_LOG(LogTemp, Display, TEXT("Using MallocBinned2 to allocate memory"));
+ #if PLATFORM_IOS || PLATFORM_TVOS
+        if(!FIOSPlatformMisc::IsEntitlementEnabled("com.apple.developer.kernel.extended-virtual-addressing"))
+        {
+            UE_LOG(LogTemp, Warning, TEXT("MallocBinned2 requested but Virtual Address Space entitlement not found. Check your entitlements. Falling back to Ansi."));
+            AllocatorToUse = EMemoryAllocatorToUse::Ansi;
+            return;
+        }
+        else
+        {
+            UE_LOG(LogTemp, Warning, TEXT("Virtual Address Space entitlement found. Using MallocBinned2 allocator"));
+        }
+#endif
+        UE_LOG(LogTemp, Display, TEXT("Using MallocBinned2 allocator."));
         AllocatorToUse = EMemoryAllocatorToUse::Binned2;
+        return;
     }
     else
     {
-        UE_LOG(LogTemp, Warning, TEXT("Defaulting to Ansi allocator"));
+        UE_LOG(LogTemp, Display, TEXT("Defaulting to Ansi allocator."));
         AllocatorToUse = EMemoryAllocatorToUse::Ansi;
+        return;
     }
+}
+
+FMalloc* FApplePlatformMemory::BaseAllocator()
+{
+	static FMalloc* Instance = nullptr;
+	if (Instance != nullptr)
+	{
+		return Instance;
+	}
+
+#if ENABLE_LOW_LEVEL_MEM_TRACKER
+	FPlatformMemoryStats MemStats = FApplePlatformMemory::GetStats();
+	FLowLevelMemTracker::Get().SetProgramSize(MemStats.UsedPhysical);
 #endif
     
+    SetAllocatorToUse();
+    
     switch (AllocatorToUse)
     {
         case EMemoryAllocatorToUse::Ansi:
diff --git a/Engine/Source/Runtime/Core/Public/Apple/ApplePlatformMemory.h b/Engine/Source/Runtime/Core/Public/Apple/ApplePlatformMemory.h
--- a/Engine/Source/Runtime/Core/Public/Apple/ApplePlatformMemory.h
+++ b/Engine/Source/Runtime/Core/Public/Apple/ApplePlatformMemory.h
@@ -75,6 +75,7 @@ struct CORE_API FApplePlatformMemory : public FGenericPlatformMemory
 	static bool PtrIsFromNanoMalloc( void* Ptr);
 	static bool IsNanoMallocAvailable();
 	static void NanoMallocInit();
+    static void SetAllocatorToUse();
 
 	class FPlatformVirtualMemoryBlock : public FBasicVirtualMemoryBlock
 	{
