[autortfm] Use UE_LOG with a new LogAutoRTFM category for all our logging.

diff --git a/Engine/Source/Programs/AutoRTFMTests/Private/API.cpp b/Engine/Source/Programs/AutoRTFMTests/Private/API.cpp
--- a/Engine/Source/Programs/AutoRTFMTests/Private/API.cpp
+++ b/Engine/Source/Programs/AutoRTFMTests/Private/API.cpp
@@ -28,7 +28,7 @@ TEST_CASE("API.autortfm_result")
 
     REQUIRE(42 == Answer);
 
-    REQUIRE(autortfm_aborted_by_language == autortfm_transact([](void* const Arg)
+    REQUIRE(autortfm_aborted_by_request == autortfm_transact([](void* const Arg)
     {
         *static_cast<int* const>(Arg) = 13;
         AutoRTFM::AbortIfTransactional();
@@ -167,7 +167,7 @@ TEST_CASE("API.autortfm_abort_if_transactional")
     REQUIRE(true == BeforeNest);
     REQUIRE(false == InNest);
     REQUIRE(true == AfterNest);
-    REQUIRE(AutoRTFM::ETransactionResult::AbortedByLanguage == NestResult);
+    REQUIRE(AutoRTFM::ETransactionResult::AbortedByRequest == NestResult);
 }
 
 TEST_CASE("API.autortfm_abort_if_closed")
@@ -179,7 +179,7 @@ TEST_CASE("API.autortfm_abort_if_closed")
     bool InNest = false;
     bool AfterNest = false;
 
-    REQUIRE(AutoRTFM::ETransactionResult::AbortedByLanguage == AutoRTFM::Transact([&]
+    REQUIRE(AutoRTFM::ETransactionResult::AbortedByRequest == AutoRTFM::Transact([&]
     {
         BeforeNest = true;
 
@@ -623,7 +623,7 @@ TEST_CASE("API.ETransactionResult")
 
     REQUIRE(42 == Answer);
 
-    REQUIRE(AutoRTFM::ETransactionResult::AbortedByLanguage == AutoRTFM::Transact([&]
+    REQUIRE(AutoRTFM::ETransactionResult::AbortedByRequest == AutoRTFM::Transact([&]
     {
         Answer = 13;
         AutoRTFM::AbortIfTransactional();
@@ -762,7 +762,7 @@ TEST_CASE("API.AbortIfTransactional")
     REQUIRE(true == BeforeNest);
     REQUIRE(false == InNest);
     REQUIRE(true == AfterNest);
-    REQUIRE(AutoRTFM::ETransactionResult::AbortedByLanguage == NestResult);
+    REQUIRE(AutoRTFM::ETransactionResult::AbortedByRequest == NestResult);
 }
 
 TEST_CASE("API.AbortIfClosed")
@@ -774,7 +774,7 @@ TEST_CASE("API.AbortIfClosed")
     bool InNest = false;
     bool AfterNest = false;
 
-    REQUIRE(AutoRTFM::ETransactionResult::AbortedByLanguage == AutoRTFM::Transact([&]
+    REQUIRE(AutoRTFM::ETransactionResult::AbortedByRequest == AutoRTFM::Transact([&]
     {
         BeforeNest = true;
 
diff --git a/Engine/Source/Programs/AutoRTFMTests/Private/Tests.cpp b/Engine/Source/Programs/AutoRTFMTests/Private/Tests.cpp
--- a/Engine/Source/Programs/AutoRTFMTests/Private/Tests.cpp
+++ b/Engine/Source/Programs/AutoRTFMTests/Private/Tests.cpp
@@ -8,6 +8,10 @@
 #include <atomic>
 #include <vector>
 #include <cstring>
+#include "Logging/LogMacros.h"
+
+DECLARE_LOG_CATEGORY_EXTERN(LogAutoRTFMTests, Display, All)
+DEFINE_LOG_CATEGORY(LogAutoRTFMTests)
 
 TEST_CASE("Tests.WriteInt")
 {
@@ -18,6 +22,14 @@ TEST_CASE("Tests.WriteInt")
     REQUIRE(X == 2);
 }
 
+TEST_CASE("Tests.UE_LOG")
+{
+	AutoRTFM::Commit([&]
+	{
+		UE_LOG(LogAutoRTFMTests, Display, TEXT("Testing this works!"));
+	});
+}
+
 // This test ensures that if you have STM and non-STM modifying data that is
 // adjacent in memory, the STM code won't lose modifications to data that
 // happens to fall into the same STM line.
diff --git a/Engine/Source/Runtime/Core/Private/AutoRTFM/API.cpp b/Engine/Source/Runtime/Core/Private/AutoRTFM/API.cpp
--- a/Engine/Source/Runtime/Core/Private/AutoRTFM/API.cpp
+++ b/Engine/Source/Runtime/Core/Private/AutoRTFM/API.cpp
@@ -8,7 +8,6 @@
 #include "Context.h"
 #include "ContextInlines.h"
 #include "ContextStatus.h"
-#include "Debug.h"
 #include "FunctionMapInlines.h"
 #include "TransactionInlines.h"
 #include "Utils.h"
@@ -40,63 +39,31 @@ extern "C" autortfm_result autortfm_transact(void (*Work)(void* Arg), void* Arg)
 extern "C" void autortfm_commit(void (*Work)(void* Arg), void* Arg)
 {
     autortfm_result Result = autortfm_transact(Work, Arg);
-    if (Result != autortfm_committed)
-    {
-        fprintf(stderr, "FATAL: Unexpected transaction result: %u.\n", Result);
-        abort();
-    }
+	UE_CLOG(Result != autortfm_committed, LogAutoRTFM, Fatal, TEXT("Unexpected transaction result: %u."), Result);
 }
 
 extern "C" void autortfm_abort()
 {
-    if (FContext::IsTransactional())
-    {
-        FContext::Get()->AbortByRequestAndThrow();
-    }
-    else
-    {
-        fprintf(stderr, "autortfm_abort called from outside a transaction.\n");
-        abort();
-    }
+	UE_CLOG(!FContext::IsTransactional(), LogAutoRTFM, Fatal, TEXT("The function `autortfm_abort` was called from outside a transaction."));
+	FContext::Get()->AbortByRequestAndThrow();
 }
 
 extern "C" bool autortfm_start_transaction()
 {
-	if (FContext::IsTransactional())
-	{
-		return FContext::Get()->StartTransaction();
-	}
-	else
-	{
-		fprintf(stderr, "autortfm_start_transaction called from outside a transact.\n");
-		abort();
-	}
+	UE_CLOG(!FContext::IsTransactional(), LogAutoRTFM, Fatal, TEXT("The function `autortfm_start_transaction` was called from outside a transact."));
+	return FContext::Get()->StartTransaction();
 }
 
 extern "C" autortfm_result autortfm_commit_transaction()
 {
-	if (FContext::IsTransactional())
-	{
-		return static_cast<autortfm_result>(FContext::Get()->CommitTransaction());
-	}
-	else
-	{
-		fprintf(stderr, "autortfm_commit_transaction called from outside a transaction.\n");
-		abort();
-	}
+	UE_CLOG(!FContext::IsTransactional(), LogAutoRTFM, Fatal, TEXT("The function `autortfm_commit_transaction` was called from outside a transact."));
+	return static_cast<autortfm_result>(FContext::Get()->CommitTransaction());
 }
 
 extern "C" autortfm_result autortfm_abort_transaction()
 {
-	if (FContext::IsTransactional())
-	{
-		return static_cast<autortfm_result>(FContext::Get()->AbortTransaction(false));
-	}
-	else
-	{
-		fprintf(stderr, "autortfm_abort_transaction called from outside a transaction.\n");
-		abort();
-	}
+	UE_CLOG(!FContext::IsTransactional(), LogAutoRTFM, Fatal, TEXT("The function `autortfm_abort_transaction` was called from outside a transact."));
+	return static_cast<autortfm_result>(FContext::Get()->AbortTransaction(false));
 }
 
 extern "C" void autortfm_clear_transaction_status()
@@ -118,11 +85,7 @@ extern "C" bool autortfm_current_nest_throw()
 
 extern "C" void autortfm_abort_if_transactional()
 {
-    if (FContext::IsTransactional())
-    {
-        fprintf(stderr, "autortfm_abort_if_transactional called from an open nest inside a transaction.\n");
-        abort();
-    }
+	UE_CLOG(FContext::IsTransactional(), LogAutoRTFM, Fatal, TEXT("The function `autortfm_abort_if_transactional` was called from an open inside a transaction."));
 }
 
 extern "C" void autortfm_abort_if_closed()
@@ -136,21 +99,15 @@ extern "C" void autortfm_open(void (*Work)(void* Arg), void* Arg)
 
 extern "C" autortfm_status autortfm_close(void (*Work)(void* Arg), void* Arg)
 {
+	UE_CLOG(!FContext::IsTransactional(), LogAutoRTFM, Fatal, TEXT("Close called from an outside a transaction."));
+
 	autortfm_status Result = autortfm_status_ontrack;
 
-    if (FContext::IsTransactional())
-    {
-        FContext* Context = FContext::Get();
-        void (*WorkClone)(void* Arg, FContext* Context) = FunctionMapLookup(Work, Context, "autortfm_close");
-        if (WorkClone)
-        {
-			Result = static_cast<autortfm_status>(Context->CallClosedNest(WorkClone, Arg));
-        }
-    }
-    else
+    FContext* Context = FContext::Get();
+    void (*WorkClone)(void* Arg, FContext* Context) = FunctionMapLookup(Work, Context, "autortfm_close");
+    if (WorkClone)
     {
-        fprintf(stderr, "autortfm_close called from outside a transaction.\n");
-        abort();
+		Result = static_cast<autortfm_status>(Context->CallClosedNest(WorkClone, Arg));
     }
 
 	return Result;
@@ -163,11 +120,7 @@ extern "C" void autortfm_record_open_write(void* Ptr, size_t Size)
 
 extern "C" void autortfm_register_open_function(void* OriginalFunction, void* NewFunction)
 {
-    constexpr bool bVerbose = false;
-    if (bVerbose)
-    {
-        fprintf(GetLogFile(), "Registering open %p->%p\n", OriginalFunction, NewFunction);
-    }
+	UE_LOG(LogAutoRTFM, Verbose, TEXT("Registering open %p->%p"), OriginalFunction, NewFunction);
     FunctionMapAdd(OriginalFunction, NewFunction);
 }
 
@@ -271,115 +224,86 @@ extern "C" void autortfm_check_abi(void* const Ptr, const size_t Size)
 		}
     } RuntimeConstants;
 
-    if (sizeof(FConstants) != Size)
-    {
-        fprintf(GetLogFile(), "Fatal: found ABI error between AutoRTFM compiler and runtime\n");
-        abort();
-    }
+	UE_CLOG(sizeof(FConstants) != Size, LogAutoRTFM, Fatal, TEXT("ABI error between AutoRTFM compiler and runtime."));
 
     const FConstants* const CompilerConstants = static_cast<FConstants*>(Ptr);
 
-    if (RuntimeConstants != *CompilerConstants)
-    {
-        fprintf(GetLogFile(), "Fatal: found ABI error between AutoRTFM compiler and runtime\n");
-        abort();
-    }
+	UE_CLOG(RuntimeConstants != *CompilerConstants, LogAutoRTFM, Fatal, TEXT("ABI error between AutoRTFM compiler and runtime."));
 }
 
 // Second Part - the same API exposed inside transactions. Note that we don't expose all of the API
 // to transactions! That's intentional. However, things like autortfm_defer_until_commit can be called
 // from an open nest in a transaction.
-bool STM_autortfm_is_transactional(FContext* Context)
+bool RTFM_autortfm_is_transactional(FContext* Context)
 {
     return true;
 }
 UE_AUTORTFM_REGISTER_OPEN_FUNCTION(autortfm_is_transactional);
 
-bool STM_autortfm_is_closed(FContext* Context)
+bool RTFM_autortfm_is_closed(FContext* Context)
 {
     return true;
 }
 UE_AUTORTFM_REGISTER_OPEN_FUNCTION(autortfm_is_closed);
 
-autortfm_result STM_autortfm_transact(void (*Work)(void* Arg), void* Arg, FContext* Context)
+autortfm_result RTFM_autortfm_transact(void (*Work)(void* Arg), void* Arg, FContext* Context)
 {
     return static_cast<autortfm_result>(Context->Transact(Work, Arg));
 }
 UE_AUTORTFM_REGISTER_OPEN_FUNCTION(autortfm_transact);
 
-void STM_autortfm_commit(void (*Work)(void* Arg), void* Arg, FContext*)
+void RTFM_autortfm_commit(void (*Work)(void* Arg), void* Arg, FContext*)
 {
     autortfm_result Result = autortfm_transact(Work, Arg);
-    if (Result != autortfm_committed)
-    {
-        fprintf(stderr, "FATAL: Unexpected transaction result: %u.\n", Result);
-        abort();
-    }
+	UE_CLOG(Result != autortfm_committed, LogAutoRTFM, Fatal, TEXT("Unexpected transaction result: %u."), Result);
 }
 UE_AUTORTFM_REGISTER_OPEN_FUNCTION(autortfm_commit);
 
-void STM_autortfm_abort(FContext* Context)
+void RTFM_autortfm_abort(FContext* Context)
 {
     Context->AbortByRequestAndThrow();
 }
 UE_AUTORTFM_REGISTER_OPEN_FUNCTION(autortfm_abort);
 
-void STM_autortfm_start_transaction(FContext* Context)
+void RTFM_autortfm_start_transaction(FContext* Context)
 {
-	fprintf(stderr, "autortfm_start_transaction called from closed code.\n");
-	abort();
+	UE_LOG(LogAutoRTFM, Fatal, TEXT("The function `autortfm_start_transaction` was called from closed code."));
 }
 UE_AUTORTFM_REGISTER_OPEN_FUNCTION(autortfm_start_transaction);
 
-void STM_autortfm_commit_transaction(FContext* Context)
+void RTFM_autortfm_commit_transaction(FContext* Context)
 {
-	fprintf(stderr, "autostm_commit_transaction called from closed code.\n");
-	abort();
+	UE_LOG(LogAutoRTFM, Fatal, TEXT("The function `RTFM_autortfm_commit_transaction` was called from closed code."));
 }
 UE_AUTORTFM_REGISTER_OPEN_FUNCTION(autortfm_commit_transaction);
 
-autortfm_result STM_autortfm_abort_transaction(FContext* Context)
+autortfm_result RTFM_autortfm_abort_transaction(FContext* Context)
 {
-	if (FContext::IsTransactional())
-	{
-		return static_cast<autortfm_result>(Context->AbortTransaction(true));
-	}
-	else
-	{
-		fprintf(stderr, "autortfm_abort_transaction called from outside a transaction.\n");
-		abort();
-	}
+	UE_CLOG(!FContext::IsTransactional(), LogAutoRTFM, Fatal, TEXT("The function `autortfm_abort_transaction` was called from outside a transaction"));
+	return static_cast<autortfm_result>(Context->AbortTransaction(true));
 }
 UE_AUTORTFM_REGISTER_OPEN_FUNCTION(autortfm_abort_transaction);
 
-void STM_autortfm_clear_transaction_status(FContext* Context)
+void RTFM_autortfm_clear_transaction_status(FContext* Context)
 {
-	fprintf(stderr, "autortfm_clear_transaction_status called from closed code.\n");
-	abort();
+	UE_LOG(LogAutoRTFM, Fatal, TEXT("The function `autortfm_clear_transaction_status` was called from closed code."));
 }
 
-void STM_autortfm_abort_if_transactional(FContext* Context)
+void RTFM_autortfm_abort_if_transactional(FContext* Context)
 {
-    if (FDebug::bVerbose)
-    {
-        fprintf(stderr, "autortfm_abort_if_transactional called from inside a transaction.\n");
-    }
-    Context->AbortByLanguageAndThrow();
+    UE_LOG(LogAutoRTFM, Verbose, TEXT("The function `autortfm_abort_if_transactional` was called from inside a transaction."));
+    Context->AbortByRequestAndThrow();
 }
 UE_AUTORTFM_REGISTER_OPEN_FUNCTION(autortfm_abort_if_transactional);
 
-void STM_autortfm_abort_if_closed(FContext* Context)
+void RTFM_autortfm_abort_if_closed(FContext* Context)
 {
-    if (FDebug::bVerbose)
-    {
-        fprintf(stderr, "autortfm_abort_if_closed called from a closed nest inside a transaction.\n");
-    }
-    
-    Context->AbortByLanguageAndThrow();
+	UE_LOG(LogAutoRTFM, Verbose, TEXT("The function `autortfm_abort_if_closed` was called from closed inside a transaction."));
+    Context->AbortByRequestAndThrow();
 }
 UE_AUTORTFM_REGISTER_OPEN_FUNCTION(autortfm_abort_if_closed);
 
-void STM_autortfm_open(void (*Work)(void* Arg), void* Arg, FContext* Context)
+void RTFM_autortfm_open(void (*Work)(void* Arg), void* Arg, FContext* Context)
 {
 	Work(Arg);
 
@@ -390,9 +314,9 @@ void STM_autortfm_open(void (*Work)(void* Arg), void* Arg, FContext* Context)
 }
 UE_AUTORTFM_REGISTER_OPEN_FUNCTION(autortfm_open);
 
-autortfm_status STM_autortfm_close(void (*Work)(void* Arg), void* Arg, FContext* Context)
+autortfm_status RTFM_autortfm_close(void (*Work)(void* Arg), void* Arg, FContext* Context)
 {
-    void (*WorkClone)(void* Arg, FContext* Context) = FunctionMapLookup(Work, Context, "STM_autortfm_close");
+    void (*WorkClone)(void* Arg, FContext* Context) = FunctionMapLookup(Work, Context, "RTFM_autortfm_close");
     if (WorkClone)
     {
         WorkClone(Arg, Context);
@@ -402,73 +326,72 @@ autortfm_status STM_autortfm_close(void (*Work)(void* Arg), void* Arg, FContext*
 }
 UE_AUTORTFM_REGISTER_OPEN_FUNCTION(autortfm_close);
 
-extern "C" void STM_autortfm_record_open_write(void*, size_t, FContext*)
+extern "C" void RTFM_autortfm_record_open_write(void*, size_t, FContext*)
 {
-	fprintf(stderr, "FATAL: autortfm_record_open_write called from closed code.\n");
-	abort();
+	UE_LOG(LogAutoRTFM, Fatal, TEXT("The function `autortfm_record_open_write` was called from closed code."));
 }
 UE_AUTORTFM_REGISTER_OPEN_FUNCTION(autortfm_record_open_write);
 
-void STM_DeferUntilCommit(TFunction<void()>&& Work, FContext* Context)
+void RTFM_DeferUntilCommit(TFunction<void()>&& Work, FContext* Context)
 {
     ASSERT(Context->GetStatus() == EContextStatus::OnTrack);
     Context->GetCurrentTransaction()->DeferUntilCommit(MoveTemp(Work));
 }
 UE_AUTORTFM_REGISTER_OPEN_FUNCTION(DeferUntilCommit);
 
-void STM_DeferUntilAbort(TFunction<void()>&& Work, FContext* Context)
+void RTFM_DeferUntilAbort(TFunction<void()>&& Work, FContext* Context)
 {
     ASSERT(Context->GetStatus() == EContextStatus::OnTrack);
     Context->GetCurrentTransaction()->DeferUntilAbort(MoveTemp(Work));
 }
 UE_AUTORTFM_REGISTER_OPEN_FUNCTION(DeferUntilAbort);
 
-void STM_OpenCommit(TFunction<void()>&& Work, FContext* Context)
+void RTFM_OpenCommit(TFunction<void()>&& Work, FContext* Context)
 {
     ASSERT(Context->GetStatus() == EContextStatus::OnTrack);
     Context->GetCurrentTransaction()->DeferUntilCommit(MoveTemp(Work));
 }
 UE_AUTORTFM_REGISTER_OPEN_FUNCTION(OpenCommit);
 
-void STM_OpenAbort(TFunction<void()>&& Work, FContext* Context)
+void RTFM_OpenAbort(TFunction<void()>&& Work, FContext* Context)
 {
     ASSERT(Context->GetStatus() == EContextStatus::OnTrack);
     Context->GetCurrentTransaction()->DeferUntilAbort(MoveTemp(Work));
 }
 UE_AUTORTFM_REGISTER_OPEN_FUNCTION(OpenAbort);
 
-extern "C" void STM_autortfm_defer_until_commit(void (*Work)(void* Arg), void* Arg, FContext* Context)
+extern "C" void RTFM_autortfm_defer_until_commit(void (*Work)(void* Arg), void* Arg, FContext* Context)
 {
-    STM_DeferUntilCommit([Work, Arg] { Work(Arg); }, Context);
+    RTFM_DeferUntilCommit([Work, Arg] { Work(Arg); }, Context);
 }
 UE_AUTORTFM_REGISTER_OPEN_FUNCTION(autortfm_defer_until_commit);
 
-extern "C" void STM_autortfm_defer_until_abort(void (*Work)(void* arg), void* Arg, FContext* Context)
+extern "C" void RTFM_autortfm_defer_until_abort(void (*Work)(void* arg), void* Arg, FContext* Context)
 {
-    STM_DeferUntilAbort([Work, Arg] { Work(Arg); }, Context);
+    RTFM_DeferUntilAbort([Work, Arg] { Work(Arg); }, Context);
 }
 UE_AUTORTFM_REGISTER_OPEN_FUNCTION(autortfm_defer_until_abort);
 
-extern "C" void STM_autortfm_open_commit(void (*Work)(void* Arg), void* Arg, FContext* Context)
+extern "C" void RTFM_autortfm_open_commit(void (*Work)(void* Arg), void* Arg, FContext* Context)
 {
-    STM_OpenCommit([Work, Arg] { Work(Arg); }, Context);
+    RTFM_OpenCommit([Work, Arg] { Work(Arg); }, Context);
 }
 UE_AUTORTFM_REGISTER_OPEN_FUNCTION(autortfm_open_commit);
 
-extern "C" void STM_autortfm_open_abort(void (*Work)(void* arg), void* Arg, FContext* Context)
+extern "C" void RTFM_autortfm_open_abort(void (*Work)(void* arg), void* Arg, FContext* Context)
 {
-    STM_OpenAbort([Work, Arg] { Work(Arg); }, Context);
+    RTFM_OpenAbort([Work, Arg] { Work(Arg); }, Context);
 }
 UE_AUTORTFM_REGISTER_OPEN_FUNCTION(autortfm_open_abort);
 
-void* STM_autortfm_did_allocate(void* Ptr, size_t Size, FContext* Context)
+void* RTFM_autortfm_did_allocate(void* Ptr, size_t Size, FContext* Context)
 {
     Context->DidAllocate(Ptr, Size);
     return Ptr;
 }
 UE_AUTORTFM_REGISTER_OPEN_FUNCTION(autortfm_did_allocate);
 
-void STM_autortfm_check_consistency_assuming_no_races(FContext* Context)
+void RTFM_autortfm_check_consistency_assuming_no_races(FContext* Context)
 {
 }
 UE_AUTORTFM_REGISTER_OPEN_FUNCTION(autortfm_check_consistency_assuming_no_races);
diff --git a/Engine/Source/Runtime/Core/Private/AutoRTFM/CompilerInterface.cpp b/Engine/Source/Runtime/Core/Private/AutoRTFM/CompilerInterface.cpp
--- a/Engine/Source/Runtime/Core/Private/AutoRTFM/CompilerInterface.cpp
+++ b/Engine/Source/Runtime/Core/Private/AutoRTFM/CompilerInterface.cpp
@@ -3,7 +3,6 @@
 #if (defined(__AUTORTFM) && __AUTORTFM)
 
 #include "ContextInlines.h"
-#include "Debug.h"
 #include "FunctionMapInlines.h"
 #include "Memcpy.h"
 
@@ -41,25 +40,21 @@ extern "C" void autortfm_record_write(FContext* Context, void* Ptr, size_t Size)
 
 extern "C" void* autortfm_lookup_function(FContext* Context, void* OriginalFunction, const char* Where)
 {
-    FDebug Debug(Context, OriginalFunction, nullptr, 0, 0, __FUNCTION__);
     return FunctionMapLookup(OriginalFunction, Context, Where);
 }
 
 extern "C" void autortfm_memcpy(void* Dst, const void* Src, size_t Size, FContext* Context)
 {
-    FDebug Debug(Context, Dst, Src, Size, 0, __FUNCTION__);
     Memcpy(Dst, Src, Size, Context);
 }
 
 extern "C" void autortfm_memmove(void* Dst, const void* Src, size_t Size, FContext* Context)
 {
-    FDebug Debug(Context, Dst, Src, Size, 0, __FUNCTION__);
     Memmove(Dst, Src, Size, Context);
 }
 
 extern "C" void autortfm_memset(void* Dst, int Value, size_t Size, FContext* Context)
 {
-    FDebug Debug(Context, Dst, nullptr, Size, 0, __FUNCTION__);
     Memset(Dst, Value, Size, Context);
 }
 
@@ -67,12 +62,13 @@ extern "C" void autortfm_llvm_fail(FContext* Context, const char* Message)
 {
     if (Message)
     {
-        fprintf(stderr, "Transaction failing because of language issue:\n%s\n", Message);
+		UE_LOG(LogAutoRTFM, Warning, TEXT("Transaction failing because of language issue '%s'."), Message);
     }
     else
     {
-        fprintf(stderr, "Transaction failing because of language issue.\n");
-    }
+		UE_LOG(LogAutoRTFM, Warning, TEXT("Transaction failing because of language issue."));
+	}
+
     Context->AbortByLanguageAndThrow();
 }
 
@@ -83,15 +79,14 @@ extern "C" void autortfm_llvm_alignment_error(FContext* Context, void* Ptr, size
 
 extern "C" void autortfm_llvm_error(FContext* Context, const char* Message)
 {
-    if (Message)
-    {
-        fprintf(stderr, "Aborting because LLVM error:\n%s\n", Message);
-    }
-    else
-    {
-        fprintf(stderr, "Aborting because LLVM error.\n");
-    }
-    abort();
+	if (Message)
+	{
+		UE_LOG(LogAutoRTFM, Fatal, TEXT("Transaction failing because of LLVM issue '%s'."), Message);
+	}
+	else
+	{
+		UE_LOG(LogAutoRTFM, Fatal, TEXT("Transaction failing because of LLVM issue."));
+	}
 }
 
 } // namespace AutoRTFM
diff --git a/Engine/Source/Runtime/Core/Private/AutoRTFM/Context.cpp b/Engine/Source/Runtime/Core/Private/AutoRTFM/Context.cpp
--- a/Engine/Source/Runtime/Core/Private/AutoRTFM/Context.cpp
+++ b/Engine/Source/Runtime/Core/Private/AutoRTFM/Context.cpp
@@ -95,12 +95,9 @@ ETransactionResult FContext::CommitTransaction()
 	}
 	else
 	{
-		if (bVerbose)
-		{
-			fprintf(GetLogFile(), "About to commit; my state is:\n");
-			DumpState();
-			fprintf(GetLogFile(), "Committing...\n");
-		}
+		UE_LOG(LogAutoRTFM, Verbose, TEXT("About to commit; my state is:"));
+		DumpState();
+		UE_LOG(LogAutoRTFM, Verbose, TEXT("Committing..."));
 
 		if (CurrentTransaction->AttemptToCommit())
 		{
@@ -108,10 +105,7 @@ ETransactionResult FContext::CommitTransaction()
 		}
 		else
 		{
-			if (bVerbose)
-			{
-				fprintf(GetLogFile(), "Commit failed!\n");
-			}
+			UE_LOG(LogAutoRTFM, Verbose, TEXT("Commit failed!"));
 			ASSERT(Status != EContextStatus::OnTrack);
 			ASSERT(Status != EContextStatus::Idle);
 		}
@@ -259,7 +253,7 @@ ETransactionResult FContext::Transact(void (*Function)(void* Arg), void* Arg)
     void (*ClonedFunction)(void* Arg, FContext* Context) = FunctionMapTryLookup(Function);
     if (!ClonedFunction)
     {
-        fprintf(stderr, "Could not find function %p (%s) in AutoRTFM::FContext::Transact.\n", Function, TCHAR_TO_ANSI(*GetFunctionDescription(Function)));
+		UE_LOG(LogAutoRTFM, Warning, TEXT("Could not find function %p (%s) in AutoRTFM::FContext::Transact."), Function, *GetFunctionDescription(Function));
         return ETransactionResult::AbortedByLanguage;
     }
     
@@ -295,12 +289,9 @@ ETransactionResult FContext::Transact(void (*Function)(void* Arg), void* Arg)
 
             if (Status == EContextStatus::OnTrack)
             {
-                if (bVerbose)
-                {
-                    fprintf(GetLogFile(), "About to commit; my state is:\n");
-                    DumpState();
-                    fprintf(GetLogFile(), "Committing...\n");
-                }
+				UE_LOG(LogAutoRTFM, Verbose, TEXT("About to commit; my state is:"));
+				DumpState();
+				UE_LOG(LogAutoRTFM, Verbose, TEXT("Committing..."));
 
                 if (CurrentTransaction->AttemptToCommit())
                 {
@@ -308,33 +299,25 @@ ETransactionResult FContext::Transact(void (*Function)(void* Arg), void* Arg)
                     break;
                 }
 
-                if (bVerbose)
-                {
-                    fprintf(GetLogFile(), "Commit failed!\n");
-                }
+				UE_LOG(LogAutoRTFM, Verbose, TEXT("Commit failed!"));
 
                 ASSERT(Status != EContextStatus::OnTrack);
                 ASSERT(Status != EContextStatus::Idle);
             }
+
             if (Status == EContextStatus::AbortedByRequest)
             {
                 Result = ETransactionResult::AbortedByRequest;
                 break;
             }
+
             if (Status == EContextStatus::AbortedByLanguage)
             {
                 Result = ETransactionResult::AbortedByLanguage;
                 break;
             }
-            ASSERT(Status == EContextStatus::AbortedByFailedLockAcquisition);
 
-            if (bVerbose)
-            {
-                fprintf(GetLogFile(), "About to prelock some locks!\n");
-            }
-            ASSERT(!bLocksToHoldAreHeld);
-            
-            bLocksToHoldAreHeld = true;
+            ASSERT(Status == EContextStatus::AbortedByFailedLockAcquisition);
         }
 
 		NewTransaction->SetIsDone();
@@ -387,13 +370,8 @@ void FContext::AbortByRequestWithoutThrowing()
 
 void FContext::AbortByLanguageAndThrow()
 {
-    constexpr bool bAbortProgram = false;
+	UE_DEBUG_BREAK();
     ASSERT(Status == EContextStatus::OnTrack);
-    if (bAbortProgram)
-    {
-        fprintf(stderr, "FATAL: Unexpected language abort.\n");
-        abort();
-    }
     Status = EContextStatus::AbortedByLanguage;
     CurrentTransaction->AbortAndThrow();
 }
@@ -436,8 +414,7 @@ void FContext::Throw()
 
 void FContext::DumpState() const
 {
-    fprintf(GetLogFile(), "Context at %p.\n", this);
-    fprintf(GetLogFile(), "Transaction stack: %p...%p\n", StackBegin, OuterTransactStackAddress);
+	UE_LOG(LogAutoRTFM, Verbose, TEXT("Context at %p, transaction stack: %p..%p."), this, StackBegin, OuterTransactStackAddress);
 }
 
 } // namespace AutoRTFM
diff --git a/Engine/Source/Runtime/Core/Private/AutoRTFM/Context.h b/Engine/Source/Runtime/Core/Private/AutoRTFM/Context.h
--- a/Engine/Source/Runtime/Core/Private/AutoRTFM/Context.h
+++ b/Engine/Source/Runtime/Core/Private/AutoRTFM/Context.h
@@ -78,7 +78,6 @@ private:
     void* StackEnd{nullptr};
     void* OuterTransactStackAddress{nullptr};
     void* CurrentTransactStackAddress{nullptr};
-    bool bLocksToHoldAreHeld{false};
     EContextStatus Status{EContextStatus::Idle};
 };
 
diff --git a/Engine/Source/Runtime/Core/Private/AutoRTFM/Debug.h b/Engine/Source/Runtime/Core/Private/AutoRTFM/Debug.h
deleted file mode 100644
diff --git a/Engine/Source/Runtime/Core/Private/AutoRTFM/FunctionMapInlines.h b/Engine/Source/Runtime/Core/Private/AutoRTFM/FunctionMapInlines.h
--- a/Engine/Source/Runtime/Core/Private/AutoRTFM/FunctionMapInlines.h
+++ b/Engine/Source/Runtime/Core/Private/AutoRTFM/FunctionMapInlines.h
@@ -15,11 +15,15 @@ inline void* FunctionMapLookup(void* OldFunction, FContext* Context, const char*
     void* Result = FunctionMapTryLookup(OldFunction);
     if (!Result)
     {
-        fprintf(stderr, "Could not find function %p (%s)\n", OldFunction, TCHAR_TO_ANSI(*GetFunctionDescription(OldFunction)));
-        if (Where)
-        {
-            fprintf(stderr, "%s\n", Where);
-        }
+		if (Where)
+		{
+			UE_LOG(LogAutoRTFM, Warning, TEXT("Could not find function %p '%s' where '%s'."), OldFunction, *GetFunctionDescription(OldFunction), Where);
+		}
+		else
+		{
+			UE_LOG(LogAutoRTFM, Warning, TEXT("Could not find function %p '%s'."), OldFunction, *GetFunctionDescription(OldFunction));
+		}
+
         Context->AbortByLanguageAndThrow();
     }
     return Result;
diff --git a/Engine/Source/Runtime/Core/Private/AutoRTFM/Memcpy.cpp b/Engine/Source/Runtime/Core/Private/AutoRTFM/Memcpy.cpp
--- a/Engine/Source/Runtime/Core/Private/AutoRTFM/Memcpy.cpp
+++ b/Engine/Source/Runtime/Core/Private/AutoRTFM/Memcpy.cpp
@@ -3,59 +3,34 @@
 #if (defined(__AUTORTFM) && __AUTORTFM)
 #include "Memcpy.h"
 #include "ContextInlines.h"
-#include "Debug.h"
 
 namespace AutoRTFM
 {
 
-constexpr bool bVerbose = FDebug::bVerbose;
-
 void* MemcpyToNew(void* InDst, const void* InSrc, size_t Size, FContext* Context)
 {
-    if (bVerbose)
-    {
-        fprintf(GetLogFile(), "MemcpyToNew(%p, %p, %zu)\n", InDst, InSrc, Size);
-    }
-    FDebug Debug(Context, InDst, InSrc, Size, 0, __FUNCTION__);
+	UE_LOG(LogAutoRTFM, Verbose, TEXT("MemcpyToNew(%p, %p, %zu)"), InDst, InSrc, Size);
     AutoRTFM::Unreachable();
 }
 
 void* Memcpy(void* InDst, const void* InSrc, size_t Size, FContext* Context)
 {
-    if (bVerbose)
-    {
-        fprintf(GetLogFile(), "Memcpy(%p, %p, %zu)\n", InDst, InSrc, Size);
-    }
-    FDebug Debug(Context, InDst, InSrc, Size, 0, __FUNCTION__);
-    
+	UE_LOG(LogAutoRTFM, Verbose, TEXT("Memcpy(%p, %p, %zu)"), InDst, InSrc, Size);
     Context->RecordWrite(InDst, Size, true);
-
     return memcpy(InDst, InSrc, Size);
 }
 
 void* Memmove(void* InDst, const void* InSrc, size_t Size, FContext* Context)
 {
-    if (bVerbose)
-    {
-        fprintf(GetLogFile(), "Memmove(%p, %p, %zu)\n", InDst, InSrc, Size);
-    }
-    FDebug Debug(Context, InDst, InSrc, Size, 0, __FUNCTION__);
-
+	UE_LOG(LogAutoRTFM, Verbose, TEXT("Memmove(%p, %p, %zu)"), InDst, InSrc, Size);
 	Context->RecordWrite(InDst, Size, true);
-
 	return memmove(InDst, InSrc, Size);
 }
 
 void* Memset(void* InDst, int Value, size_t Size, FContext* Context)
 {
-    if (bVerbose)
-    {
-        fprintf(GetLogFile(), "Memset(%p, %d, %zu)\n", InDst, Value, Size);
-    }
-    FDebug Debug(Context, InDst, nullptr, Size, 0, __FUNCTION__);
-
+	UE_LOG(LogAutoRTFM, Verbose, TEXT("Memset(%p, %d, %zu)"), InDst, Value, Size);
 	Context->RecordWrite(InDst, Size, true);
-
 	return memset(InDst, Value, Size);
 }
 
diff --git a/Engine/Source/Runtime/Core/Private/AutoRTFM/StdLibSupport.cpp b/Engine/Source/Runtime/Core/Private/AutoRTFM/StdLibSupport.cpp
--- a/Engine/Source/Runtime/Core/Private/AutoRTFM/StdLibSupport.cpp
+++ b/Engine/Source/Runtime/Core/Private/AutoRTFM/StdLibSupport.cpp
@@ -28,7 +28,7 @@ UE_AUTORTFM_REGISTER_OPEN_FUNCTION_EXPLICIT(memcpy, Memcpy);
 UE_AUTORTFM_REGISTER_OPEN_FUNCTION_EXPLICIT(memmove, Memmove);
 UE_AUTORTFM_REGISTER_OPEN_FUNCTION_EXPLICIT(memset, Memset);
 
-void* STM_malloc(size_t Size, FContext* Context)
+void* RTFM_malloc(size_t Size, FContext* Context)
 {
     void* Result = malloc(Size);
     Context->GetCurrentTransaction()->DeferUntilAbort([Result]
@@ -41,7 +41,7 @@ void* STM_malloc(size_t Size, FContext* Context)
 
 UE_AUTORTFM_REGISTER_OPEN_FUNCTION(malloc);
 
-void STM_free(void* Ptr, FContext* Context)
+void RTFM_free(void* Ptr, FContext* Context)
 {
     if (Ptr)
     {
@@ -53,9 +53,9 @@ void STM_free(void* Ptr, FContext* Context)
 }
 UE_AUTORTFM_REGISTER_OPEN_FUNCTION(free);
 
-void* STM_realloc(void* Ptr, size_t Size, FContext* Context)
+void* RTFM_realloc(void* Ptr, size_t Size, FContext* Context)
 {
-    void* NewObject = STM_malloc(Size, Context);
+    void* NewObject = RTFM_malloc(Size, Context);
     if (Ptr)
     {
 #if defined(__APPLE__)
@@ -66,13 +66,13 @@ void* STM_realloc(void* Ptr, size_t Size, FContext* Context)
 		const size_t OldSize = malloc_usable_size(Ptr);
 #endif
         MemcpyToNew(NewObject, Ptr,  FMath::Min(OldSize, Size), Context);
-        STM_free(Ptr, Context);
+        RTFM_free(Ptr, Context);
     }
     return NewObject;
 }
 UE_AUTORTFM_REGISTER_OPEN_FUNCTION(realloc);
 
-char* STM_strcpy(char* const Dst, const char* const Src, FContext* const Context)
+char* RTFM_strcpy(char* const Dst, const char* const Src, FContext* const Context)
 {
     const size_t SrcLen = strlen(Src);
 
@@ -81,14 +81,14 @@ char* STM_strcpy(char* const Dst, const char* const Src, FContext* const Context
 }
 UE_AUTORTFM_REGISTER_OPEN_FUNCTION(strcpy);
 
-char* STM_strncpy(char* const Dst, const char* const Src, const size_t Num, FContext* const Context)
+char* RTFM_strncpy(char* const Dst, const char* const Src, const size_t Num, FContext* const Context)
 {
     Context->RecordWrite(Dst, Num, true);
     return strncpy(Dst, Src, Num);
 }
 UE_AUTORTFM_REGISTER_OPEN_FUNCTION(strncpy);
 
-char* STM_strcat(char* const Dst, const char* const Src, FContext* const Context)
+char* RTFM_strcat(char* const Dst, const char* const Src, FContext* const Context)
 {
     const size_t DstLen = strlen(Dst);
     const size_t SrcLen = strlen(Src);
@@ -98,7 +98,7 @@ char* STM_strcat(char* const Dst, const char* const Src, FContext* const Context
 }
 UE_AUTORTFM_REGISTER_OPEN_FUNCTION(strcat);
 
-char* STM_strncat(char* const Dst, const char* const Src, const size_t Num, FContext* const Context)
+char* RTFM_strncat(char* const Dst, const char* const Src, const size_t Num, FContext* const Context)
 {
     const size_t DstLen = strlen(Dst);
 
@@ -154,7 +154,7 @@ UE_AUTORTFM_REGISTER_SELF_FUNCTION(llrint);
 // - Str is newly allocated
 // - Format is either newly allocated or not mutated
 // - any strings passed as arguments are either newly allocated or not mutated
-int STM_snprintf(char* Str, size_t Size, char* Format, FContext* Context, ...)
+int RTFM_snprintf(char* Str, size_t Size, char* Format, FContext* Context, ...)
 {
     va_list ArgList;
     va_start(ArgList, Context);
@@ -164,7 +164,7 @@ int STM_snprintf(char* Str, size_t Size, char* Format, FContext* Context, ...)
 }
 UE_AUTORTFM_REGISTER_OPEN_FUNCTION(snprintf);
 
-int STM_printf(const char* Format, FContext* Context, ...)
+int RTFM_printf(const char* Format, FContext* Context, ...)
 {
     va_list ArgList;
     va_start(ArgList, Context);
@@ -174,13 +174,13 @@ int STM_printf(const char* Format, FContext* Context, ...)
 }
 UE_AUTORTFM_REGISTER_OPEN_FUNCTION(printf);
 
-int STM_putchar(int Char, FContext* Context)
+int RTFM_putchar(int Char, FContext* Context)
 {
     return putchar(Char);
 }
 UE_AUTORTFM_REGISTER_OPEN_FUNCTION(putchar);
 
-int STM_puts(const char* Str, FContext* Context)
+int RTFM_puts(const char* Str, FContext* Context)
 {
     return puts(Str);
 }
@@ -188,7 +188,7 @@ UE_AUTORTFM_REGISTER_OPEN_FUNCTION(puts);
 
 #if PLATFORM_WINDOWS
 
-FILE* STM___acrt_iob_func(int Index, FContext* Context)
+FILE* RTFM___acrt_iob_func(int Index, FContext* Context)
 {
     switch (Index)
     {
@@ -196,7 +196,7 @@ FILE* STM___acrt_iob_func(int Index, FContext* Context)
     case 2:
         return __acrt_iob_func(Index);
     default:
-        fprintf(stderr, "Attempt to get file descriptor %d (not 1 or 2) in __acrt_iob_func.\n", Index);
+		UE_LOG(LogAutoRTFM, Warning, TEXT("Attempt to get file descriptor %d (not 1 or 2) in __acrt_iob_func."), Index);
         Context->AbortByLanguageAndThrow();
         return NULL;
     }
@@ -227,7 +227,7 @@ UE_AUTORTFM_REGISTER_SELF_FUNCTION(GetCurrentThreadId);
 
 UE_AUTORTFM_REGISTER_SELF_FUNCTION(TlsGetValue);
 
-BOOL STM_TlsSetValue(DWORD dwTlsIndex, LPVOID lpTlsValue, FContext* Context)
+BOOL RTFM_TlsSetValue(DWORD dwTlsIndex, LPVOID lpTlsValue, FContext* Context)
 {
 	LPVOID CurrentValue = TlsGetValue(dwTlsIndex);
 
@@ -241,7 +241,7 @@ BOOL STM_TlsSetValue(DWORD dwTlsIndex, LPVOID lpTlsValue, FContext* Context)
 UE_AUTORTFM_REGISTER_OPEN_FUNCTION(TlsSetValue);
 #endif // PLATFORM_WINDOWS
 
-wchar_t* STM_wcsncpy(wchar_t* Dst, const wchar_t* Src, size_t Count, FContext* Context)
+wchar_t* RTFM_wcsncpy(wchar_t* Dst, const wchar_t* Src, size_t Count, FContext* Context)
 {
 	AutoRTFM::Unreachable();
 }
@@ -261,7 +261,7 @@ UE_AUTORTFM_REGISTER_OPEN_FUNCTION(wcsncpy);
 #pragma warning(pop)
 #endif
 
-int STM_atexit(void(__cdecl*Callback)(void), FContext* Context)
+int RTFM_atexit(void(__cdecl*Callback)(void), FContext* Context)
 {
 	Context->GetCurrentTransaction()->DeferUntilCommit([Callback]
 		{
diff --git a/Engine/Source/Runtime/Core/Private/AutoRTFM/Transaction.cpp b/Engine/Source/Runtime/Core/Private/AutoRTFM/Transaction.cpp
--- a/Engine/Source/Runtime/Core/Private/AutoRTFM/Transaction.cpp
+++ b/Engine/Source/Runtime/Core/Private/AutoRTFM/Transaction.cpp
@@ -4,7 +4,6 @@
 #include "Transaction.h"
 #include "TransactionInlines.h"
 #include "CallNestInlines.h"
-#include "Debug.h"
 #include "GlobalData.h"
 
 namespace AutoRTFM
@@ -26,14 +25,13 @@ bool FTransaction::IsFresh() const
 
 void FTransaction::AbortWithoutThrowing()
 {
-    if (FDebug::bVerbose)
-    {
-        fprintf(GetLogFile(), "Aborting (%s)\n", GetContextStatusName(Context->GetStatus()));
-    }
+	UE_LOG(LogAutoRTFM, Verbose, TEXT("Aborting '%s'!"), GetContextStatusName(Context->GetStatus()));
+
     ASSERT(Context->GetStatus() == EContextStatus::AbortedByFailedLockAcquisition
            || Context->GetStatus() == EContextStatus::AbortedByLanguage
            || Context->GetStatus() == EContextStatus::AbortedByRequest);
     ASSERT(Context->GetCurrentTransaction() == this);
+
     if (IsNested())
     {
         AbortNested();
@@ -71,14 +69,10 @@ bool FTransaction::AttemptToCommit()
 
 void FTransaction::Undo()
 {
-	const bool bVerboseUndo = false;
-
-	if (bVerboseUndo)
-	{
-		fprintf(GetLogFile(), "[FTransaction::Undo] START\n");
-	}
+	UE_LOG(LogAutoRTFM, Verbose, TEXT("Undoing a transaction..."));
 
 	int VerboseCounter = 0;
+
 	for(auto Iter = WriteLog.rbegin(); Iter != WriteLog.rend(); ++Iter)
     {
 		FWriteLogEntry& Entry = *Iter;
@@ -94,29 +88,30 @@ void FTransaction::Undo()
         const size_t Size = Entry.OriginalAndSize.GetTopTag();
         void* const Copy = Entry.Copy;
 
-		if (bVerboseUndo)
+		if (UE_LOG_ACTIVE(LogAutoRTFM, Verbose))
 		{
-			FILE* LogFile = GetLogFile();
-			fprintf(LogFile, "%4d [UNDO] %p %4llu : ", VerboseCounter, Original, Size);
+			TStringBuilder<1024> Builder;
+
+			Builder.Appendf(TEXT("%4d [UNDO] %p %4llu : [ "), VerboseCounter, Original, Size);
 
 			unsigned char* Current = (unsigned char*)Original;
 			unsigned char* Old = (unsigned char*)Copy;
 
-			fprintf(LogFile, "[");
-
-			for(size_t i = 0; i < Size; i++)
+			for (size_t i = 0; i < Size; i++)
 			{
-				fprintf(LogFile, "%02X ", Current[i]);
+				Builder.Appendf(TEXT("%02X "), Current[i]);
 			}
 
-			fprintf(LogFile, "] -> [");
+			Builder << TEXT("] -> [ ");
 
-			for(size_t i = 0; i < Size; i++)
+			for (size_t i = 0; i < Size; i++)
 			{
-				fprintf(LogFile, "%02X ", Old[i]);
+				Builder.Appendf(TEXT("%02X "), Old[i]);
 			}
 
-			fprintf(LogFile, "]\n");
+			Builder << TEXT("]");
+
+			UE_LOG(LogAutoRTFM, Verbose, TEXT("%s"), Builder.ToString());
 		}
 
         memcpy(Original, Copy, Size);
@@ -124,10 +119,7 @@ void FTransaction::Undo()
 		VerboseCounter++;
     }
 
-	if (bVerboseUndo)
-	{
-		fprintf(GetLogFile(), "[FTransaction::Undo] END\n");
-	}
+	UE_LOG(LogAutoRTFM, Verbose, TEXT("Undone a transaction!"));
 }
 
 void FTransaction::AbortNested()
@@ -190,12 +182,9 @@ bool FTransaction::AttemptToCommitOuterNest()
 {
     ASSERT(!Parent);
 
-    if (FDebug::bVerbose)
-    {
-        fprintf(GetLogFile(), "About to run commit tasks!\n");
-        Context->DumpState();
-        fprintf(GetLogFile(), "Running commit tasks...\n");
-    }
+	UE_LOG(LogAutoRTFM, Verbose, TEXT("About to run commit tasks!"));
+	Context->DumpState();
+	UE_LOG(LogAutoRTFM, Verbose, TEXT("Running commit tasks..."));
 
     CommitTasks.ForEachForward([] (const TFunction<void()>& Task) -> bool { Task(); return true; });
 
diff --git a/Engine/Source/Runtime/Core/Private/AutoRTFM/TransactionInlines.h b/Engine/Source/Runtime/Core/Private/AutoRTFM/TransactionInlines.h
--- a/Engine/Source/Runtime/Core/Private/AutoRTFM/TransactionInlines.h
+++ b/Engine/Source/Runtime/Core/Private/AutoRTFM/TransactionInlines.h
@@ -51,9 +51,8 @@ inline void FTransaction::RecordWrite(void* LogicalAddress, size_t Size, bool bI
 		&& !bIsClosed
 		&& Context->IsInnerTransactionStack(LogicalAddress))
 	{
-		// It is an error to write to stack variables within an open nest
-		fprintf(stderr, "FATAL: Writing to local stack memory from an unscoped transaction is not allowed.\n");
-		abort();
+		// It is an error to write to stack variables within an open nest.
+		UE_LOG(LogAutoRTFM, Fatal, TEXT("Writing to local stack memory from an unscoped transaction is not allowed."));
 	}
 
     uint8_t* Address = reinterpret_cast<uint8_t*>(LogicalAddress);
diff --git a/Engine/Source/Runtime/Core/Private/AutoRTFM/Utils.cpp b/Engine/Source/Runtime/Core/Private/AutoRTFM/Utils.cpp
--- a/Engine/Source/Runtime/Core/Private/AutoRTFM/Utils.cpp
+++ b/Engine/Source/Runtime/Core/Private/AutoRTFM/Utils.cpp
@@ -16,40 +16,10 @@
 #include "Containers/StringConv.h"
 #include "Misc/CoreMiscDefines.h"
 
-namespace AutoRTFM
-{
-
-// To log to a file, comment out LogFilename = nullptr and uncomment the other one.
-constexpr const char* LogFilename = nullptr;
-// constexpr const char* LogFilename = "log.txt";
-
-constexpr const char* LogFileMode = "wt";
-
-FILE* LogFile;
+DEFINE_LOG_CATEGORY(LogAutoRTFM)
 
-FILE* GetLogFile()
+namespace AutoRTFM
 {
-    UE_CALL_ONCE([]
-    {
-        if (LogFilename)
-        {
-            ASSERT(LogFileMode);
-
-            LogFile = fopen(LogFilename, LogFileMode);
-
-            if (!LogFile)
-            {
-                fprintf(stderr, "Could not open %s: %s\n", LogFilename, strerror(errno));
-            }
-        }
-        if (!LogFile)
-        {
-            LogFile = stderr;
-        }
-        setvbuf(LogFile, NULL, _IONBF, 0);
-    });
-    return LogFile;
-}
 
 FString GetFunctionDescription(void* FunctionPtr)
 {
diff --git a/Engine/Source/Runtime/Core/Private/AutoRTFM/Utils.h b/Engine/Source/Runtime/Core/Private/AutoRTFM/Utils.h
--- a/Engine/Source/Runtime/Core/Private/AutoRTFM/Utils.h
+++ b/Engine/Source/Runtime/Core/Private/AutoRTFM/Utils.h
@@ -9,20 +9,18 @@
 #include <stdlib.h>
 
 #include "Containers/UnrealString.h"
+#include "Logging/LogMacros.h"
 
-#define ASSERT(exp) do { \
-    if (!(exp)) { \
-        fprintf(stderr, "%s:%d:%s: assertion %s failed.\n", __FILE__, __LINE__, __PRETTY_FUNCTION__, #exp); \
-        abort(); \
-    } \
-} while (false)
+DECLARE_LOG_CATEGORY_EXTERN(LogAutoRTFM, NoLogging, All)
+
+#define ASSERT(exp) UE_CLOG(!(exp), LogAutoRTFM, Fatal, TEXT("%s:%d:%s: assertion %s failed."), __FILE__, __LINE__, __PRETTY_FUNCTION__, #exp)
 
 namespace AutoRTFM
 {
 
 [[noreturn]] inline void Unreachable()
 {
-    fprintf(stderr, "Unreachable encountered!\n");
+	UE_LOG(LogAutoRTFM, Fatal, TEXT("Unreachable encountered!"));
 
 #if PLATFORM_WINDOWS
     __assume(false);
@@ -31,8 +29,6 @@ namespace AutoRTFM
 #endif // PLATFORM_WINDOWS
 }
 
-FILE* GetLogFile();
-
 FString GetFunctionDescription(void* FunctionPtr);
 
 template<typename TReturnType, typename... TParameterTypes>
diff --git a/Engine/Source/Runtime/Core/Private/Logging/LogMacros.cpp b/Engine/Source/Runtime/Core/Private/Logging/LogMacros.cpp
--- a/Engine/Source/Runtime/Core/Private/Logging/LogMacros.cpp
+++ b/Engine/Source/Runtime/Core/Private/Logging/LogMacros.cpp
@@ -73,7 +73,12 @@ void FMsg::LogV(const ANSICHAR* File, int32 Line, const FLogCategoryName& Catego
 		default:
 			break;
 		}
-		(OutputDevice ? OutputDevice : GLog)->Serialize(Message, Verbosity, Category);
+
+		// Logging is always done in the open as we want logs even with transactionalized code.
+		AutoRTFM::Open([OutputDevice, Message, Verbosity, Category]
+		{
+			(OutputDevice ? OutputDevice : GLog)->Serialize(Message, Verbosity, Category);
+		});
 	}
 	else
 	{
diff --git a/Engine/Source/Runtime/Core/Public/AutoRTFM/AutoRTFM.h b/Engine/Source/Runtime/Core/Public/AutoRTFM/AutoRTFM.h
--- a/Engine/Source/Runtime/Core/Public/AutoRTFM/AutoRTFM.h
+++ b/Engine/Source/Runtime/Core/Public/AutoRTFM/AutoRTFM.h
@@ -566,7 +566,7 @@ struct FRegisterOpenFunction
 
 #if UE_AUTORTFM
 #define UE_AUTORTFM_REGISTER_OPEN_FUNCTION_EXPLICIT(OriginalFunction, NewFunction) static const FRegisterOpenFunction UE_AUTORTFM_CONCAT(AutoRTFMFunctionRegistration, __COUNTER__)(reinterpret_cast<void*>(OriginalFunction), reinterpret_cast<void*>(NewFunction))
-#define UE_AUTORTFM_REGISTER_OPEN_FUNCTION(OriginalFunction) UE_AUTORTFM_REGISTER_OPEN_FUNCTION_EXPLICIT(OriginalFunction, STM_ ## OriginalFunction)
+#define UE_AUTORTFM_REGISTER_OPEN_FUNCTION(OriginalFunction) UE_AUTORTFM_REGISTER_OPEN_FUNCTION_EXPLICIT(OriginalFunction, RTFM_ ## OriginalFunction)
 #define UE_AUTORTFM_REGISTER_SELF_FUNCTION(OriginalFunction) UE_AUTORTFM_REGISTER_OPEN_FUNCTION_EXPLICIT(OriginalFunction, OriginalFunction)
 #else
 #define UE_AUTORTFM_REGISTER_OPEN_FUNCTION_EXPLICIT(OriginalFunction, NewFunction)
