LLM tags for tasks

diff --git a/Engine/Source/Runtime/Core/Private/Async/TaskGraph.cpp b/Engine/Source/Runtime/Core/Private/Async/TaskGraph.cpp
--- a/Engine/Source/Runtime/Core/Private/Async/TaskGraph.cpp
+++ b/Engine/Source/Runtime/Core/Private/Async/TaskGraph.cpp
@@ -1037,7 +1037,7 @@ private:
 	**/
 	uint64 ProcessTasks()
 	{
-		LLM_SCOPE(ELLMTag::TaskGraphTasksMisc);
+		LLM_SCOPE_BYNAME(TEXT("Tasks/AnyThread/ProcessTasks"));
 
 		TStatId StallStatId;
 		bool bCountAsStall = true;
@@ -1951,6 +1951,8 @@ public:
 private:
 	void QueueTask(class FBaseGraphTask* Task, bool bWakeUpWorker, ENamedThreads::Type InThreadToExecuteOn, ENamedThreads::Type InCurrentThreadIfKnown) override
 	{
+		LLM_SCOPE_BYNAME(TEXT("Tasks/Scheduler/QueueTask"));
+
 		if (ENamedThreads::GetThreadIndex(InThreadToExecuteOn) == ENamedThreads::AnyThread)
 		{
 #if TASKGRAPH_NEW_FRONTEND
@@ -2428,6 +2430,7 @@ static TLockFreeClassAllocator_TLSCache<FGraphEvent, PLATFORM_CACHE_LINE_SIZE>&
 
 FGraphEventRef FGraphEvent::CreateGraphEvent()
 {
+	LLM_SCOPE_BYNAME(TEXT("Tasks/FGraphEvent/CreateGraphEvent"));
 	FGraphEvent* Instance = new(GetGraphEventAllocator().Allocate()) FGraphEvent{};
 	return Instance;
 }
diff --git a/Engine/Source/Runtime/Core/Private/HAL/Thread.cpp b/Engine/Source/Runtime/Core/Private/HAL/Thread.cpp
--- a/Engine/Source/Runtime/Core/Private/HAL/Thread.cpp
+++ b/Engine/Source/Runtime/Core/Private/HAL/Thread.cpp
@@ -63,7 +63,6 @@ private:
 	{
 		// the thread can be started before `RunnableThread` member is initialized
 		//check(RunnableThread.IsValid());
-		LLM_SCOPE(ELLMTag::TaskGraphTasksMisc);
 		ThreadFunction();
 
 		return 0;
diff --git a/Engine/Source/Runtime/Core/Public/Async/TaskGraphInterfaces.h b/Engine/Source/Runtime/Core/Public/Async/TaskGraphInterfaces.h
--- a/Engine/Source/Runtime/Core/Public/Async/TaskGraphInterfaces.h
+++ b/Engine/Source/Runtime/Core/Public/Async/TaskGraphInterfaces.h
@@ -606,6 +606,7 @@ public:
 		template<typename...T>
 		FORCEINLINE_DEBUGGABLE TGraphTask* ConstructAndHoldImpl(T&&... Args)
 		{
+			LLM_SCOPE_BYNAME(TEXT("Tasks/TGraphTask/ConstructAndHoldImpl"));
 			TGraphTask* Task = new TGraphTask(Prerequisites);
 			TTask* TaskObject = new(&Task->TaskStorage) TTask(Forward<T>(Args)...);
 
@@ -720,6 +721,8 @@ inline void FGraphEventImpl::operator delete(void* Ptr)
 
 inline FGraphEventRef FBaseGraphTask::CreateGraphEvent()
 {
+	LLM_SCOPE_BYNAME(TEXT("Tasks/FGraphEvent/CreateGraphEvent"));
+
 	FGraphEventImpl* GraphEvent = new FGraphEventImpl;
 	return FGraphEventRef{ GraphEvent, /*bAddRef = */ false };
 }
@@ -936,6 +939,7 @@ public:
 	**/
 	bool AddSubsequent(class FBaseGraphTask* Subsequent)
 	{
+		LLM_SCOPE_BYNAME(TEXT("Tasks/FGraphEvent/AddSubsequent"));
 		bool bSucceeded = SubsequentList.PushIfNotClosed(Subsequent);
 		if (bSucceeded)
 		{
@@ -959,6 +963,7 @@ public:
 	**/
 	void DontCompleteUntil(FGraphEventRef EventToWaitFor)
 	{
+		LLM_SCOPE_BYNAME(TEXT("Tasks/FGraphEvent/DontCompleteUntil"));
 		checkThreadGraph(!IsComplete()); // it is not legal to add a DontCompleteUntil after the event has been completed. Basically, this is only legal within a task function.
 		EventsToWaitFor.Emplace(EventToWaitFor);
 		TaskTrace::SubsequentAdded(EventToWaitFor->GetTraceId(), GetTraceId());
@@ -1215,8 +1220,12 @@ public:
 	**/
 	static FConstructor CreateTask(const FGraphEventArray* Prerequisites = NULL, ENamedThreads::Type CurrentThreadIfKnown = ENamedThreads::AnyThread)
 	{
+		FGraphEventRef GraphEvent = TTask::GetSubsequentsMode() == ESubsequentsMode::FireAndForget ? NULL : FGraphEvent::CreateGraphEvent();
+
+		LLM_SCOPE_BYNAME(TEXT("Tasks/TGraphTask/CreateTask"));
+
 		int32 NumPrereq = Prerequisites ? Prerequisites->Num() : 0;
-		return FConstructor(new TGraphTask(TTask::GetSubsequentsMode() == ESubsequentsMode::FireAndForget ? NULL : FGraphEvent::CreateGraphEvent(), NumPrereq), Prerequisites, CurrentThreadIfKnown);
+		return FConstructor(new TGraphTask(MoveTemp(GraphEvent), NumPrereq), Prerequisites, CurrentThreadIfKnown);
 	}
 
 	void Unlock(ENamedThreads::Type CurrentThreadIfKnown = ENamedThreads::AnyThread)
diff --git a/Engine/Source/Runtime/Core/Public/Tasks/TaskPrivate.h b/Engine/Source/Runtime/Core/Public/Tasks/TaskPrivate.h
--- a/Engine/Source/Runtime/Core/Public/Tasks/TaskPrivate.h
+++ b/Engine/Source/Runtime/Core/Public/Tasks/TaskPrivate.h
@@ -216,6 +216,8 @@ namespace UE::Tasks
 					return false;
 				}
 
+				LLM_SCOPE_BYNAME(TEXT("Tasks/FTaskBase/AddPrerequisites"));
+
 				Prerequisite.AddRef(); // keep it alive until this task's execution
 				Prerequisites.Push(&Prerequisite); // release memory order
 				return true;
@@ -272,6 +274,7 @@ namespace UE::Tasks
 
 					if (Prerequisite->AddSubsequent(*this)) // acq_rel memory order
 					{
+						LLM_SCOPE_BYNAME(TEXT("Tasks/FTaskBase/AddPrerequisites"));
 						Prerequisite->AddRef(); // keep it alive until this task's execution
 						Prerequisites.Push(Prerequisite); // release memory order
 					}
@@ -290,6 +293,7 @@ namespace UE::Tasks
 			// returns false if the task is already completed and the subsequent wasn't added
 			bool AddSubsequent(FTaskBase& Subsequent)
 			{
+				LLM_SCOPE_BYNAME(TEXT("Tasks/FTaskBase/AddSubsequent"));
 				TaskTrace::SubsequentAdded(GetTraceId(), Subsequent.GetTraceId()); // doesn't matter if we suceeded below, we need to record task dependency
 				return Subsequents.PushIfNotClosed(&Subsequent);
 			}
@@ -347,6 +351,7 @@ namespace UE::Tasks
 
 				if (Nested.AddSubsequent(*this)) // "release" memory order
 				{
+					LLM_SCOPE_BYNAME(TEXT("Tasks/FTaskBase/AddNested"));
 					Nested.AddRef(); // keep it alive as we store it in `Prerequisites` and we can need it to try to retract it. it's released on closing the task
 					Prerequisites.Push(&Nested);
 				}
@@ -523,6 +528,7 @@ namespace UE::Tasks
 							FTaskBase* PrevPipedTask = TryPushIntoPipe();
 							if (PrevPipedTask != nullptr) // the pipe is blocked
 							{
+								LLM_SCOPE_BYNAME(TEXT("Tasks/FTaskBase/AddPrerequisites"));
 								// the prev task in pipe's chain becomes this task's prerequisite, to enabled piped task retraction.
 								// its ref count already accounted for this ref. the ref will be released when the prereq is not needed anymore
 								Prerequisites.Push(PrevPipedTask);
@@ -750,6 +756,7 @@ namespace UE::Tasks
 			// a helper that deduces the template argument
 			static TExecutableTask* Create(const TCHAR* InDebugName, TaskBodyType&& TaskBody, ETaskPriority InPriority, EExtendedTaskPriority InExtendedPriority)
 			{
+				LLM_SCOPE_BYNAME(TEXT("Tasks/TExecutableTask/Create"));
 				return new TExecutableTask(InDebugName, MoveTemp(TaskBody), InPriority, InExtendedPriority);
 			}
 
@@ -775,6 +782,7 @@ namespace UE::Tasks
 		public:
 			static FTaskEventBase* Create(const TCHAR* DebugName)
 			{
+				LLM_SCOPE_BYNAME(TEXT("Tasks/FTaskEventBase/Create"));
 				return new FTaskEventBase(DebugName);
 			}
 
