Tasks System: fixed tasks integration that rotted a bit since it was used last time

diff --git a/Engine/Source/Runtime/Core/Private/Tasks/TaskPrivate.cpp b/Engine/Source/Runtime/Core/Private/Tasks/TaskPrivate.cpp
--- a/Engine/Source/Runtime/Core/Private/Tasks/TaskPrivate.cpp
+++ b/Engine/Source/Runtime/Core/Private/Tasks/TaskPrivate.cpp
@@ -274,7 +274,7 @@ namespace UE::Tasks
 				// execute other tasks of this named thread while waiting
 				ETaskPriority Dummy;
 				EExtendedTaskPriority ExtendedPriority;
-				FBaseGraphTask::TranslatePriority(CurrentThread, Dummy, ExtendedPriority);
+				TranslatePriority(CurrentThread, Dummy, ExtendedPriority);
 
 				auto TaskBody = [CurrentThread, &TaskGraph] { TaskGraph.RequestReturn(CurrentThread); };
 				using FReturnFromNamedThreadTask = TExecutableTask<decltype(TaskBody)>;
@@ -283,7 +283,8 @@ namespace UE::Tasks
 				ReturnTask.TryLaunch(sizeof(ReturnTask)); // the result doesn't matter
 
 				TaskGraph.ProcessThreadUntilRequestReturn(CurrentThread);
-				return IsCompleted();
+				check(Task.IsCompleted());
+				return true;
 			}
 #endif
 
@@ -458,7 +459,31 @@ namespace UE::Tasks
 #endif
 		}
 
-#if TASKGRAPH_NEW_FRONTEND
+		ENamedThreads::Type TranslatePriority(ETaskPriority Priority)
+		{
+			checkSlow(Priority < ETaskPriority::Count);
+
+			ENamedThreads::Type ConversionMap[] =
+			{
+				ENamedThreads::AnyHiPriThreadNormalTask,
+				ENamedThreads::AnyNormalThreadNormalTask,
+				ENamedThreads::AnyBackgroundHiPriTask,
+				ENamedThreads::AnyBackgroundThreadNormalTask,
+				ENamedThreads::AnyBackgroundThreadNormalTask // same as above
+			};
+
+			return ConversionMap[(int32)Priority];
+		}
+
+#if !TASKGRAPH_NEW_FRONTEND
+
+		ENamedThreads::Type TranslatePriority(ETaskPriority Priority, EExtendedTaskPriority /*ExtendedPriority*/)
+		{
+			checkSlow(Priority <= ETaskPriority::Count);
+			return TranslatePriority(Priority);
+		}
+
+#else // !TASKGRAPH_NEW_FRONTEND
 
 		// task priority translation from the new API to the old API
 		ENamedThreads::Type TranslatePriority(EExtendedTaskPriority Priority)
@@ -486,22 +511,14 @@ namespace UE::Tasks
 			return (ENamedThreads::Type)ConversionMap[(int32)Priority - (int32)EExtendedTaskPriority::GameThreadNormalPri];
 		}
 
-#endif
-
-		ENamedThreads::Type TranslatePriority(ETaskPriority Priority)
+		ENamedThreads::Type TranslatePriority(ETaskPriority Priority, EExtendedTaskPriority ExtendedPriority)
 		{
-			checkSlow(Priority < ETaskPriority::Count);
+			checkSlow(Priority <= ETaskPriority::Count);
+			checkSlow(ExtendedPriority <= EExtendedTaskPriority::Count);
 
-			ENamedThreads::Type ConversionMap[] =
-			{
-				ENamedThreads::AnyHiPriThreadNormalTask,
-				ENamedThreads::AnyNormalThreadNormalTask,
-				ENamedThreads::AnyBackgroundHiPriTask,
-				ENamedThreads::AnyBackgroundThreadNormalTask,
-				ENamedThreads::AnyBackgroundThreadNormalTask // same as above
-			};
-
-			return ConversionMap[(int32)Priority];
+			return ExtendedPriority < EExtendedTaskPriority::GameThreadNormalPri || ExtendedPriority == EExtendedTaskPriority::Count ? TranslatePriority(Priority) : TranslatePriority(ExtendedPriority);
 		}
+
+#endif // !TASKGRAPH_NEW_FRONTEND
 	}
 }
diff --git a/Engine/Source/Runtime/Core/Public/Async/TaskGraphInterfaces.h b/Engine/Source/Runtime/Core/Public/Async/TaskGraphInterfaces.h
--- a/Engine/Source/Runtime/Core/Public/Async/TaskGraphInterfaces.h
+++ b/Engine/Source/Runtime/Core/Public/Async/TaskGraphInterfaces.h
@@ -23,6 +23,7 @@
 #include "Containers/LockFreeFixedSizeAllocator.h"
 #include "Experimental/ConcurrentLinearAllocator.h"
 #include "Misc/MemStack.h"
+#include "Misc/Timeout.h"
 #include "Templates/Atomic.h"
 #include "ProfilingDebugging/MetadataTrace.h"
 
@@ -551,14 +552,13 @@ public:
 			return FTaskGraphInterface::Get().WaitUntilTaskCompletes(this, CurrentThreadIfKnown);
 		}
 
-		return FTaskBase::Wait();
+		FTaskBase::Wait(UE::FTimeout::Never());
 	}
 
 	ENamedThreads::Type GetThreadToExecuteOn() const
 	{
-		return UE::Tasks::Private::TranslatePriority(GetExtendedPriority());
+		return UE::Tasks::Private::TranslatePriority(GetPriority(), GetExtendedPriority());
 	}
-
 };
 
 static constexpr int32 SmallTaskSize = 256;
@@ -692,7 +692,7 @@ public:
 		: FBaseGraphTask(nullptr)
 	{
 		TaskTrace::Created(GetTraceId(), sizeof(*this));
-		Init((TEXT("GraphEvent"), UE::Tasks::ETaskPriority::Normal, UE::Tasks::EExtendedTaskPriority::TaskEvent);
+		Init(TEXT("GraphEvent"), UE::Tasks::ETaskPriority::Normal, UE::Tasks::EExtendedTaskPriority::TaskEvent);
 	}
 
 	static void* operator new(size_t Size);
diff --git a/Engine/Source/Runtime/Core/Public/Tasks/TaskPrivate.h b/Engine/Source/Runtime/Core/Public/Tasks/TaskPrivate.h
--- a/Engine/Source/Runtime/Core/Public/Tasks/TaskPrivate.h
+++ b/Engine/Source/Runtime/Core/Public/Tasks/TaskPrivate.h
@@ -76,10 +76,7 @@ namespace UE::Tasks
 	namespace Private
 	{
 		CORE_API void TranslatePriority(ENamedThreads::Type ThreadType, ETaskPriority& OutPriority, EExtendedTaskPriority& OutExtendedPriority);
-		CORE_API ENamedThreads::Type TranslatePriority(ETaskPriority Priority);
-#if TASKGRAPH_NEW_FRONTEND
-		CORE_API ENamedThreads::Type TranslatePriority(EExtendedTaskPriority Priority);
-#endif
+		CORE_API ENamedThreads::Type TranslatePriority(ETaskPriority Priority, EExtendedTaskPriority ExtendedPriority);
 	}
 
 	class FPipe;
@@ -129,7 +126,7 @@ namespace UE::Tasks
 				}
 			}
 
-			uint32 GetRefCount(std::memory_order MemoryOrder) const
+			uint32 GetRefCount(std::memory_order MemoryOrder = std::memory_order_relaxed) const
 			{
 				return RefCount.load(MemoryOrder);
 			}
@@ -190,6 +187,11 @@ namespace UE::Tasks
 			}
 #endif
 
+			ETaskPriority GetPriority() const
+			{
+				return LowLevelTask.GetPriority();
+			}
+
 			EExtendedTaskPriority GetExtendedPriority() const
 			{
 				return ExtendedPriority;
@@ -317,7 +319,7 @@ namespace UE::Tasks
 			// In this case the task will be automatically scheduled when all dependencies are completed.
 			bool TryLaunch(uint64 TaskSize)
 			{
-				TaskTrace::Launched(GetTraceId(), LowLevelTask.GetDebugName(), true, TranslatePriority(LowLevelTask.GetPriority()), TaskSize);
+				TaskTrace::Launched(GetTraceId(), LowLevelTask.GetDebugName(), true, TranslatePriority(LowLevelTask.GetPriority(), ExtendedPriority), TaskSize);
 				return TryUnlock();
 			}
 
