Support for chunked & threaded hashing with xxhash128, upwards of around 120 gb/s on my threadripper.

diff --git a/Engine/Source/Runtime/Core/Private/Hash/xxhash.cpp b/Engine/Source/Runtime/Core/Private/Hash/xxhash.cpp
--- a/Engine/Source/Runtime/Core/Private/Hash/xxhash.cpp
+++ b/Engine/Source/Runtime/Core/Private/Hash/xxhash.cpp
@@ -2,10 +2,13 @@
 
 #include "Hash/xxhash.h"
 
+#include "Async/ParallelFor.h"
 #include "Containers/ContainersFwd.h"
 #include "Memory/CompositeBuffer.h"
 #include "Memory/MemoryView.h"
 #include "Memory/SharedBuffer.h"
+#include "Serialization/CompactBinaryWriter.h"
+#include "Tasks/Task.h"
 
 THIRD_PARTY_INCLUDES_START
 #define XXH_INLINE_ALL
@@ -29,6 +32,51 @@ FXxHash64 FXxHash64::HashBuffer(const FCompositeBuffer& Buffer)
 	return Builder.Finalize();
 }
 
+
+/* static */ UE::Tasks::TTask<FXxHash64> FXxHash64::HashBufferChunkedAsync(FMemoryView View, uint64 ChunkSize)
+{
+	if (View.GetSize() <= ChunkSize || ChunkSize == 0)
+	{
+		return UE::Tasks::Launch(TEXT("XxHash64.HashBufferChunkedAsync.Single"), [View]()
+		{
+			return FXxHash64::HashBuffer(View);
+		});
+	}
+
+	return UE::Tasks::Launch(TEXT("XxHash64.HashBufferChunkedAsync.Multi"), [View, ChunkSize]()
+	{
+		uint64 ChunkCount = (View.GetSize() + ChunkSize - 1) / ChunkSize;
+
+		// Due to limitations of ParallelFor, we can't hash anything larger than uint32_max * chunksize;
+		uint32 ChunkCount32 = IntCastChecked<uint32, uint64>(ChunkCount);
+
+		TArray<FXxHash64, TInlineAllocator<64>> ChunkHashes;
+		ChunkHashes.AddDefaulted(ChunkCount32);
+
+		ParallelFor(TEXT("XxHash64.PF"), ChunkCount32, 1, [View, ChunkSize, ChunkCount32, &ChunkHashes](int32 Index)
+		{
+			FMemoryView ChunkView = View.Mid(Index * ChunkSize, ChunkSize);
+			ChunkHashes[Index] = FXxHash64::HashBuffer(ChunkView);
+		});
+
+		FXxHash64Builder Accumulator;
+		for (FXxHash64& ChunkHash : ChunkHashes)
+		{
+			Accumulator.Update(&ChunkHash.Hash, sizeof(ChunkHash.Hash));
+		}
+		uint64 ViewSize = View.GetSize();
+		Accumulator.Update(&ViewSize, sizeof(ViewSize));
+
+		return Accumulator.Finalize();
+	});
+}
+
+/* static */ UE::Tasks::TTask<FXxHash64> FXxHash64::HashBufferChunkedAsync(const void* Data, uint64 Size, uint64 ChunkSize)
+{
+	return HashBufferChunkedAsync(FMemoryView(Data, Size), ChunkSize);
+}
+
+
 FXxHash128 FXxHash128::HashBuffer(FMemoryView View)
 {
 	const XXH128_hash_t Value = XXH3_128bits(View.GetData(), View.GetSize());
@@ -41,6 +89,52 @@ FXxHash128 FXxHash128::HashBuffer(const void* Data, uint64 Size)
 	return {Value.low64, Value.high64};
 }
 
+/* static */ UE::Tasks::TTask<FXxHash128> FXxHash128::HashBufferChunkedAsync(FMemoryView View, uint64 ChunkSize)
+{
+	if (View.GetSize() <= ChunkSize || ChunkSize == 0)
+	{
+		return UE::Tasks::Launch(TEXT("XxHash128.HashBufferChunkedAsync.Single"), [View]()
+		{
+			return FXxHash128::HashBuffer(View);
+		});
+	}
+
+	return UE::Tasks::Launch(TEXT("XxHash128.HashBufferChunkedAsync.Multi"), [View, ChunkSize]()
+	{
+		uint64 ChunkCount = (View.GetSize() + ChunkSize - 1) / ChunkSize;
+		
+		// Due to limitations of ParallelFor, we can't hash anything larger than uint32_max * chunksize;
+		uint32 ChunkCount32 = IntCastChecked<uint32, uint64>(ChunkCount);
+
+		TArray<FXxHash128, TInlineAllocator<64>> ChunkHashes;
+		ChunkHashes.AddDefaulted(ChunkCount32);
+
+		ParallelFor(TEXT("XxHash128.PF"), ChunkCount32, 1, [View, ChunkSize, ChunkCount32, &ChunkHashes](int32 Index)
+		{
+			FMemoryView ChunkView = View.Mid(Index * ChunkSize, ChunkSize);
+			ChunkHashes[Index] = FXxHash128::HashBuffer(ChunkView);
+		});
+
+		FXxHash128Builder Accumulator;
+		for (FXxHash128& ChunkHash : ChunkHashes)
+		{
+			uint64 Hash[2] = {ChunkHash.HashLow, ChunkHash.HashHigh};
+			Accumulator.Update(Hash, sizeof(Hash));
+		}
+		uint64 ViewSize = View.GetSize();
+		Accumulator.Update(&ViewSize, sizeof(ViewSize));
+		
+		return Accumulator.Finalize();
+	});
+}
+
+/* static */ UE::Tasks::TTask<FXxHash128> FXxHash128::HashBufferChunkedAsync(const void* Data, uint64 Size, uint64 ChunkSize)
+{
+	return HashBufferChunkedAsync(FMemoryView(Data, Size), ChunkSize);
+}
+
+
+
 FXxHash128 FXxHash128::HashBuffer(const FCompositeBuffer& Buffer)
 {
 	FXxHash128Builder Builder;
diff --git a/Engine/Source/Runtime/Core/Public/Hash/xxhash.h b/Engine/Source/Runtime/Core/Public/Hash/xxhash.h
--- a/Engine/Source/Runtime/Core/Public/Hash/xxhash.h
+++ b/Engine/Source/Runtime/Core/Public/Hash/xxhash.h
@@ -11,8 +11,15 @@
 #include "Serialization/Archive.h"
 #include "String/BytesToHex.h"
 
+class FCbFieldView; 
+class FCbWriter;
 class FCompositeBuffer;
 template <typename CharType> class TStringBuilderBase;
+namespace UE::Tasks 
+{
+	template<typename ResultType>
+	class TTask;
+}
 
 /** A 64-bit hash from XXH3. */
 struct FXxHash64
@@ -21,6 +28,11 @@ struct FXxHash64
 	[[nodiscard]] CORE_API static FXxHash64 HashBuffer(const void* Data, uint64 Size);
 	[[nodiscard]] CORE_API static FXxHash64 HashBuffer(const FCompositeBuffer& Buffer);
 
+	// Kick off a task to hash the given memory. Buffers are hashed in ChunkSize independent pieces, with the results
+	// hashed together in to a final result. ChunkSize should be fairly large to cover the cost of dispatching tasks. (e.g. 256kb+)
+	[[nodiscard]] CORE_API static UE::Tasks::TTask<FXxHash64> HashBufferChunkedAsync(FMemoryView View, uint64 ChunkSize);
+	[[nodiscard]] CORE_API static UE::Tasks::TTask<FXxHash64> HashBufferChunkedAsync(const void* Data, uint64 Size, uint64 ChunkSize);
+
 	/** Load the hash from its canonical (big-endian) representation. */
 	static inline FXxHash64 FromByteArray(uint8 (&Bytes)[sizeof(uint64)])
 	{
@@ -51,6 +63,11 @@ struct FXxHash64
 		return Hash < B.Hash;
 	}
 
+	inline bool IsZero() const
+	{
+		return Hash == 0;
+	}
+
 	friend inline uint32 GetTypeHash(const FXxHash64& InHash)
 	{
 		return uint32(InHash.Hash);
@@ -87,6 +104,11 @@ public:
 	[[nodiscard]] CORE_API static FXxHash128 HashBuffer(const void* Data, uint64 Size);
 	[[nodiscard]] CORE_API static FXxHash128 HashBuffer(const FCompositeBuffer& Buffer);
 
+	// Kick off a task to hash the given memory. Buffers are hashed in ChunkSize independent pieces, with the results
+	// hashed together in to a final result. ChunkSize should be fairly large to cover the cost of dispatching tasks. (e.g. 256kb+)
+	[[nodiscard]] CORE_API static UE::Tasks::TTask<FXxHash128> HashBufferChunkedAsync(FMemoryView View, uint64 ChunkSize);
+	[[nodiscard]] CORE_API static UE::Tasks::TTask<FXxHash128> HashBufferChunkedAsync(const void* Data, uint64 Size, uint64 ChunkSize);
+
 	/** Load the hash from its canonical (big-endian) representation. */
 	static inline FXxHash128 FromByteArray(uint8 (&Bytes)[sizeof(uint64[2])])
 	{
@@ -130,6 +152,11 @@ public:
 		return HashHigh != B.HashHigh ? HashHigh < B.HashHigh : HashLow < B.HashLow;
 	}
 
+	inline bool IsZero() const
+	{
+		return HashHigh == 0 && HashLow == 0;
+	}
+
 	friend inline uint32 GetTypeHash(const FXxHash128& Hash)
 	{
 		return uint32(Hash.HashLow);
