Fixes potential memory leak in gzip compression

diff --git a/Engine/Source/Runtime/Core/Private/Misc/Compression.cpp b/Engine/Source/Runtime/Core/Private/Misc/Compression.cpp
--- a/Engine/Source/Runtime/Core/Private/Misc/Compression.cpp
+++ b/Engine/Source/Runtime/Core/Private/Misc/Compression.cpp
@@ -157,16 +157,13 @@ static bool appCompressMemoryGZIP(void* CompressedBuffer, int32& CompressedSize,
 	gzipstream.avail_out = CompressedSize;
 
 	int status = 0;
-	bool bOperationSucceeded = false;
 	while ((status = deflate(&gzipstream, Z_FINISH)) == Z_OK);
-	if (status == Z_STREAM_END)
-	{
-		bOperationSucceeded = true;
-		deflateEnd(&gzipstream);
-	}
+	deflateEnd(&gzipstream);
 
 	// Propagate compressed size from intermediate variable back into out variable.
 	CompressedSize = gzipstream.total_out;
+
+	bool bOperationSucceeded = (status == Z_STREAM_END);
 	return bOperationSucceeded;
 }
 
