- Allow for platform plugin extensions to have config files (PlatformMyPlugin.ini)
- Added -dumpiniloads commandline option to show what ini files are looked for, and what are found

diff --git a/Engine/Plugins/Experimental/GameFeatures/Source/GameFeatures/Private/GameFeatureData.cpp b/Engine/Plugins/Experimental/GameFeatures/Source/GameFeatures/Private/GameFeatureData.cpp
--- a/Engine/Plugins/Experimental/GameFeatures/Source/GameFeatures/Private/GameFeatureData.cpp
+++ b/Engine/Plugins/Experimental/GameFeatures/Source/GameFeatures/Private/GameFeatureData.cpp
@@ -10,6 +10,7 @@
 #include "GameFeatureAction_AddWPContent.h"
 #include "WorldPartition/ContentBundle/ContentBundleDescriptor.h"
 #include "AssetRegistry/AssetData.h"
+#include "Interfaces/IPluginManager.h"
 
 #if WITH_EDITOR
 #include "Misc/DataValidation.h"
@@ -82,7 +83,8 @@ void UGameFeatureData::InitializeBasePluginIniFile(const FString& PluginInstalle
 	if (!FConfigCacheIni::LoadExternalIniFile(PluginConfig, *PluginName, *EngineConfigDir, *PluginConfigDir, bIsBaseIniName, nullptr, bForceReloadFromDisk, bWriteDestIni))
 	{
 		// Now try the same rules as PluginManager using Default and the config hierarchy
-		FConfigContext Context = FConfigContext::ReadIntoPluginFile(PluginConfig, *FPaths::GetPath(PluginInstalledFilename));
+		FConfigContext Context = FConfigContext::ReadIntoPluginFile(PluginConfig, *FPaths::GetPath(PluginInstalledFilename),
+				IPluginManager::Get().FindPluginFromPath(PluginName)->GetExtensionBaseDirs());
 
 		if (!Context.Load(*PluginName))
 		{
diff --git a/Engine/Source/Runtime/Core/Private/Misc/ConfigContext.cpp b/Engine/Source/Runtime/Core/Private/Misc/ConfigContext.cpp
--- a/Engine/Source/Runtime/Core/Private/Misc/ConfigContext.cpp
+++ b/Engine/Source/Runtime/Core/Private/Misc/ConfigContext.cpp
@@ -110,12 +110,33 @@ const FConfigContext::FPerPlatformDirs& FConfigContext::GetPerPlatformDirs(const
 	FConfigContext::FPerPlatformDirs* Dirs = FConfigContext::PerPlatformDirs.Find(PlatformName);
 	if (Dirs == nullptr)
 	{
+		FString PluginExtDir;
+		if (bIsForPlugin)
+		{
+ 			if (DestIniFilename.Contains(TEXT("Enhanced")))
+			{
+				UE_LOG(LogConfig, Display, TEXT("paper2d"));
+			}
+			// look if there's a plugin extension for this platform, it will have the platform name in the path
+			for (const FString& ChildDir : ChildPluginBaseDirs)
+			{
+				UE_LOG(LogConfig, Display, TEXT("HAVE CHILD PLUGIN DIR: %s"), *ChildDir);
+				if (ChildDir.Contains(*FString::Printf(TEXT("/%s/"), *PlatformName)))
+				{
+					PluginExtDir = ChildDir;
+					break;
+				}
+			}
+		}
+		
 		Dirs = &PerPlatformDirs.Emplace(PlatformName, FConfigContext::FPerPlatformDirs
 			{
 				// PlatformExtensionEngineDir
 				FPaths::Combine(*FPaths::EnginePlatformExtensionsDir(), *PlatformName).Replace(*FPaths::EngineDir(), *(EngineRootDir + "/")),
 				// PlatformExtensionProjectDir
-				FPaths::Combine(*FPaths::ProjectPlatformExtensionsDir(), *PlatformName).Replace(*FPaths::ProjectDir(), *(ProjectRootDir + "/"))
+				FPaths::Combine(*FPaths::ProjectPlatformExtensionsDir(), *PlatformName).Replace(*FPaths::ProjectDir(), *(ProjectRootDir + "/")),
+				// PluginExtensionDir
+				PluginExtDir,
 			});
 	}
 	return *Dirs;
@@ -456,6 +477,7 @@ FString FConfigContext::PerformFinalExpansions(const FString& InString, const FS
 	if (bIsForPlugin)
 	{
 		OutString = OutString.Replace(TEXT("{PLUGIN}"), *PluginRootDir);
+		OutString = OutString.Replace(TEXT("{EXTPLUGIN}"), *GetPerPlatformDirs(InPlatform).PlatformExtensionPluginDir);
 	}
 
 	return OutString;
@@ -611,6 +633,8 @@ void FConfigContext::AddStaticLayersToHierarchy()
  **/
 static bool LoadIniFileHierarchy(const FConfigFileHierarchy& HierarchyToLoad, FConfigFile& ConfigFile, bool bUseCache, const TSet<FString>* IniCacheSet)
 {
+	static bool bDumpIniLoadInfo = FParse::Param(FCommandLine::Get(), TEXT("dumpiniloads"));
+	
 	TRACE_CPUPROFILER_EVENT_SCOPE(LoadIniFileHierarchy);
 	// Traverse ini list back to front, merging along the way.
 	for (const TPair<int32, FString>& HierarchyIt : HierarchyToLoad)
@@ -618,6 +642,8 @@ static bool LoadIniFileHierarchy(const FConfigFileHierarchy& HierarchyToLoad, FC
 		bool bDoCombine = (HierarchyIt.Key != 0);
 		const FString& IniFileName = HierarchyIt.Value;
 
+		UE_CLOG(bDumpIniLoadInfo, LogConfig, Display, TEXT("Looking for file: %s"), *IniFileName);
+		
 		// skip non-existant files
 		if (IsUsingLocalIniFile(*IniFileName, nullptr) && !DoesConfigFileExistWrapper(*IniFileName, IniCacheSet))
 		{
@@ -627,6 +653,8 @@ static bool LoadIniFileHierarchy(const FConfigFileHierarchy& HierarchyToLoad, FC
 		bool bDoEmptyConfig = false;
 		//UE_LOG(LogConfig, Log,  TEXT( "Combining configFile: %s" ), *IniList(IniIndex) );
 		ProcessIniContents(*IniFileName, *IniFileName, &ConfigFile, bDoEmptyConfig, bDoCombine);
+
+		UE_CLOG(bDumpIniLoadInfo, LogConfig, Display, TEXT("   Found!"));
 	}
 
 	// Set this configs files source ini hierarchy to show where it was loaded from.
diff --git a/Engine/Source/Runtime/Core/Public/Misc/ConfigContext.h b/Engine/Source/Runtime/Core/Public/Misc/ConfigContext.h
--- a/Engine/Source/Runtime/Core/Public/Misc/ConfigContext.h
+++ b/Engine/Source/Runtime/Core/Public/Misc/ConfigContext.h
@@ -92,11 +92,13 @@ public:
 	/**
 	 * Create a context to read a plugin's ini file named for the plugin. This is not used for inserting, say, Engine.ini into GConfig
 	 */
-	static FConfigContext ReadIntoPluginFile(FConfigFile& DestConfigFile, const FString& PluginRootDir)
+	static FConfigContext ReadIntoPluginFile(FConfigFile& DestConfigFile, const FString& PluginRootDir, const TArray<FString>& ChildPluginsBaseDirs)
 	{
 		FConfigContext Context(nullptr, true, FString(), &DestConfigFile);
 		Context.bIsForPlugin = true;
 		Context.PluginRootDir = PluginRootDir;
+		Context.ChildPluginBaseDirs = ChildPluginsBaseDirs;
+
 
 		// plugins are read in parallel, so we are reading into a file, but not touching GConfig, so bWriteDest would be false, but we
 		// want to write them out as if we had been using GConfig
@@ -140,6 +142,7 @@ public:
 	{
 		FString PlatformExtensionEngineDir;
 		FString PlatformExtensionProjectDir;
+		FString PlatformExtensionPluginDir;
 	};
 	/**
 	 * Return the paths to use to find hierarchical config files for the given platform (note that these are independent of the ini name)
@@ -163,7 +166,8 @@ public:
 	FString ProjectConfigDir;
 	FString ProjectRootDir;
 	FString PluginRootDir;
-
+	TArray<FString> ChildPluginBaseDirs;
+	
 	// useful strings that are used alot when walking the hierarchy
 	FString ProjectNotForLicenseesDir;
 	FString ProjectNoRedistDir;
@@ -199,4 +203,4 @@ protected:
 	FString PerformFinalExpansions(const FString& InString, const FString& Platform);
 };
 
-bool DoesConfigFileExistWrapper(const TCHAR* IniFile, const TSet<FString>* IniCacheSet = nullptr);
\ No newline at end of file
+bool DoesConfigFileExistWrapper(const TCHAR* IniFile, const TSet<FString>* IniCacheSet = nullptr);
diff --git a/Engine/Source/Runtime/Core/Public/Misc/ConfigHierarchy.h b/Engine/Source/Runtime/Core/Public/Misc/ConfigHierarchy.h
--- a/Engine/Source/Runtime/Core/Public/Misc/ConfigHierarchy.h
+++ b/Engine/Source/Runtime/Core/Public/Misc/ConfigHierarchy.h
@@ -89,9 +89,16 @@ inline FConfigLayerExpansion GConfigExpansions[] =
 
 	// Platform Extensions
 	{
-		TEXT("{ENGINE}/Config/{PLATFORM}/"),	TEXT("{EXTENGINE}/Config/"),	
-		TEXT("{PROJECT}/Config/{PLATFORM}/"),	TEXT("{EXTPROJECT}/Config/"), 
-		EConfigExpansionFlags::ForUncooked | EConfigExpansionFlags::ForCooked | EConfigExpansionFlags::ForPlugin
+		TEXT("{ENGINE}/Config/{PLATFORM}/"),	TEXT("{EXTENGINE}/Config/"),
+		TEXT("{PROJECT}/Config/{PLATFORM}/"),	TEXT("{EXTPROJECT}/Config/"),
+		EConfigExpansionFlags::ForUncooked | EConfigExpansionFlags::ForCooked,
+	},
+
+	// Plugin Platform Extensions
+	{
+		TEXT("{PLUGIN}/Config/{PLATFORM}/"),	TEXT("{EXTPLUGIN}/Config/"),
+		TEXT("{PROJECT}/Config/{PLATFORM}/"),	TEXT("{EXTPROJECT}/Config/"),
+		EConfigExpansionFlags::ForPlugin,
 	},
 
 	// Platform Extensions in Restricted Locations
diff --git a/Engine/Source/Runtime/Projects/Private/PluginManager.cpp b/Engine/Source/Runtime/Projects/Private/PluginManager.cpp
--- a/Engine/Source/Runtime/Projects/Private/PluginManager.cpp
+++ b/Engine/Source/Runtime/Projects/Private/PluginManager.cpp
@@ -1713,7 +1713,7 @@ bool FPluginManager::ConfigureEnabledPlugins()
 
 					FConfigFile& PluginConfig = GConfig->Add(PluginConfigFilename, FConfigFile());
 
-					FConfigContext Context = FConfigContext::ReadIntoPluginFile(PluginConfig, FPaths::GetPath(Plugin.FileName));
+					FConfigContext Context = FConfigContext::ReadIntoPluginFile(PluginConfig, FPaths::GetPath(Plugin.FileName), Plugin.GetExtensionBaseDirs());
 
 					if (PluginSystemDefs::IsCachingIniFilesForProcessing())
 					{
