Integrating some missed feedback.

diff --git a/Engine/Source/Runtime/Core/Private/Hash/xxhash.cpp b/Engine/Source/Runtime/Core/Private/Hash/xxhash.cpp
--- a/Engine/Source/Runtime/Core/Private/Hash/xxhash.cpp
+++ b/Engine/Source/Runtime/Core/Private/Hash/xxhash.cpp
@@ -47,13 +47,13 @@ FXxHash64 FXxHash64::HashBuffer(const FCompositeBuffer& Buffer)
 	{
 		uint64 ChunkCount = (View.GetSize() + ChunkSize - 1) / ChunkSize;
 
-		// Due to limitations of ParallelFor, we can't hash anything larger than uint32_max * chunksize;
-		uint32 ChunkCount32 = IntCastChecked<uint32, uint64>(ChunkCount);
+		// Due to limitations of ParallelFor, we can't hash anything larger than int32_max * chunksize;
+		int32 ChunkCount32 = IntCastChecked<int32>(ChunkCount);
 
 		TArray<FXxHash64, TInlineAllocator<64>> ChunkHashes;
 		ChunkHashes.AddDefaulted(ChunkCount32);
 
-		ParallelFor(TEXT("XxHash64.PF"), ChunkCount32, 1, [View, ChunkSize, ChunkCount32, &ChunkHashes](int32 Index)
+		ParallelFor(TEXT("XxHash64.PF"), ChunkCount32, 1, [View, ChunkSize, &ChunkHashes](int32 Index)
 		{
 			FMemoryView ChunkView = View.Mid(Index * ChunkSize, ChunkSize);
 			ChunkHashes[Index] = FXxHash64::HashBuffer(ChunkView);
@@ -62,7 +62,9 @@ FXxHash64 FXxHash64::HashBuffer(const FCompositeBuffer& Buffer)
 		FXxHash64Builder Accumulator;
 		for (FXxHash64& ChunkHash : ChunkHashes)
 		{
-			Accumulator.Update(&ChunkHash.Hash, sizeof(ChunkHash.Hash));
+			uint8 HashBytes[8];
+			ChunkHash.ToByteArray(HashBytes);
+			Accumulator.Update(HashBytes, sizeof(HashBytes));
 		}
 		uint64 ViewSize = View.GetSize();
 		Accumulator.Update(&ViewSize, sizeof(ViewSize));
@@ -103,13 +105,13 @@ FXxHash128 FXxHash128::HashBuffer(const void* Data, uint64 Size)
 	{
 		uint64 ChunkCount = (View.GetSize() + ChunkSize - 1) / ChunkSize;
 		
-		// Due to limitations of ParallelFor, we can't hash anything larger than uint32_max * chunksize;
-		uint32 ChunkCount32 = IntCastChecked<uint32, uint64>(ChunkCount);
+		// Due to limitations of ParallelFor, we can't hash anything larger than int32_max * chunksize;
+		int32 ChunkCount32 = IntCastChecked<int32>(ChunkCount);
 
 		TArray<FXxHash128, TInlineAllocator<64>> ChunkHashes;
 		ChunkHashes.AddDefaulted(ChunkCount32);
 
-		ParallelFor(TEXT("XxHash128.PF"), ChunkCount32, 1, [View, ChunkSize, ChunkCount32, &ChunkHashes](int32 Index)
+		ParallelFor(TEXT("XxHash128.PF"), ChunkCount32, 1, [View, ChunkSize, &ChunkHashes](int32 Index)
 		{
 			FMemoryView ChunkView = View.Mid(Index * ChunkSize, ChunkSize);
 			ChunkHashes[Index] = FXxHash128::HashBuffer(ChunkView);
@@ -118,8 +120,9 @@ FXxHash128 FXxHash128::HashBuffer(const void* Data, uint64 Size)
 		FXxHash128Builder Accumulator;
 		for (FXxHash128& ChunkHash : ChunkHashes)
 		{
-			uint64 Hash[2] = {ChunkHash.HashLow, ChunkHash.HashHigh};
-			Accumulator.Update(Hash, sizeof(Hash));
+			uint8 HashBytes[16];
+			ChunkHash.ToByteArray(HashBytes);
+			Accumulator.Update(HashBytes, sizeof(HashBytes));
 		}
 		uint64 ViewSize = View.GetSize();
 		Accumulator.Update(&ViewSize, sizeof(ViewSize));
diff --git a/Engine/Source/Runtime/Core/Public/Hash/xxhash.h b/Engine/Source/Runtime/Core/Public/Hash/xxhash.h
--- a/Engine/Source/Runtime/Core/Public/Hash/xxhash.h
+++ b/Engine/Source/Runtime/Core/Public/Hash/xxhash.h
@@ -11,8 +11,6 @@
 #include "Serialization/Archive.h"
 #include "String/BytesToHex.h"
 
-class FCbFieldView; 
-class FCbWriter;
 class FCompositeBuffer;
 template <typename CharType> class TStringBuilderBase;
 namespace UE::Tasks 
