Adding some logging to UnixSignalHeartBeat.cpp to help diagnose issue

diff --git a/Engine/Source/Runtime/Core/Private/Unix/UnixSignalHeartBeat.cpp b/Engine/Source/Runtime/Core/Private/Unix/UnixSignalHeartBeat.cpp
--- a/Engine/Source/Runtime/Core/Private/Unix/UnixSignalHeartBeat.cpp
+++ b/Engine/Source/Runtime/Core/Private/Unix/UnixSignalHeartBeat.cpp
@@ -150,10 +150,9 @@ void FUnixSignalGameHitchHeartBeat::InitSettings()
 		{
 			SuspendCount = 1;
 		}
-
-		UE_LOG(LogUnixHeartBeat, Log, TEXT("bStartSuspended:%d HitchThresholdS:%f"),
-			bStartSuspended ? 1 : 0, HitchThresholdS);
 	}
+
+	UE_LOG(LogUnixHeartBeat, Log, TEXT("HitchThresholdS:%f"), HitchThresholdS);
 }
 
 void FUnixSignalGameHitchHeartBeat::FrameStart(bool bSkipThisFrame)
@@ -248,6 +247,7 @@ void FUnixSignalGameHitchHeartBeat::ResumeHeartBeat()
 
 void FUnixSignalGameHitchHeartBeat::Restart()
 {
+	UE_LOG(LogUnixHeartBeat, Log, TEXT("Restart"));
 	bDisabled = false;
 
 	// If we still have a valid handle on the timer_t clean it up
@@ -262,6 +262,7 @@ void FUnixSignalGameHitchHeartBeat::Restart()
 
 void FUnixSignalGameHitchHeartBeat::Stop()
 {
+	UE_LOG(LogUnixHeartBeat, Log, TEXT("Stop"));
 	SuspendHeartBeat();
 	bDisabled = true;
 }
