Non-unity fix

#rnx

diff --git a/Engine/Source/Runtime/Core/Internal/IO/IoAllocators.h b/Engine/Source/Runtime/Core/Internal/IO/IoAllocators.h
--- a/Engine/Source/Runtime/Core/Internal/IO/IoAllocators.h
+++ b/Engine/Source/Runtime/Core/Internal/IO/IoAllocators.h
@@ -3,6 +3,7 @@
 #pragma once
 
 #include "Templates/TypeCompatibleBytes.h"
+#include "Templates/UnrealTemplate.h"
 
 template <typename T, uint16 SlabSize = 4096>
 class TSingleThreadedSlabAllocator
