Set array slack growth factor to 1.25x if AGGRESSIVE_MEMORY_SAVING is enabled

diff --git a/Engine/Source/Runtime/Core/Public/Containers/ContainerAllocationPolicies.h b/Engine/Source/Runtime/Core/Public/Containers/ContainerAllocationPolicies.h
--- a/Engine/Source/Runtime/Core/Public/Containers/ContainerAllocationPolicies.h
+++ b/Engine/Source/Runtime/Core/Public/Containers/ContainerAllocationPolicies.h
@@ -75,7 +75,6 @@ FORCEINLINE SizeType DefaultCalculateSlackGrow(SizeType NumElements, SizeType Nu
 #endif
 #if AGGRESSIVE_MEMORY_SAVING
 	const SIZE_T FirstGrow = 1;
-	const SIZE_T ConstantGrow = 0;
 #else
 	const SIZE_T FirstGrow = 4;
 	const SIZE_T ConstantGrow = 16;
@@ -90,7 +89,11 @@ FORCEINLINE SizeType DefaultCalculateSlackGrow(SizeType NumElements, SizeType Nu
 	if (NumAllocatedElements)
 	{
 		// Allocate slack for the array proportional to its size.
+#if AGGRESSIVE_MEMORY_SAVING
+		Grow = SIZE_T(NumElements) + SIZE_T(NumElements) / 4;
+#else
 		Grow = SIZE_T(NumElements) + 3 * SIZE_T(NumElements) / 8 + ConstantGrow;
+#endif
 	}
 	else if (SIZE_T(NumElements) > Grow)
 	{
@@ -100,9 +103,14 @@ FORCEINLINE SizeType DefaultCalculateSlackGrow(SizeType NumElements, SizeType Nu
 	if (NumAllocatedElements || SIZE_T(NumElements) > Grow)
 	{
 		// Allocate slack for the array proportional to its size.
+#if AGGRESSIVE_MEMORY_SAVING
+		Grow = SIZE_T(NumElements) + SIZE_T(NumElements) / 4;
+#else
 		Grow = SIZE_T(NumElements) + 3 * SIZE_T(NumElements) / 8 + ConstantGrow;
+#endif
 	}
 #endif
+	
 	if (bAllowQuantize)
 	{
 		Retval = (SizeType)(FMemory::QuantizeSize(Grow * BytesPerElement, Alignment) / BytesPerElement);
