More logging for UnixSignalHeartBeat.cpp

diff --git a/Engine/Source/Runtime/Core/Private/Unix/UnixSignalHeartBeat.cpp b/Engine/Source/Runtime/Core/Private/Unix/UnixSignalHeartBeat.cpp
--- a/Engine/Source/Runtime/Core/Private/Unix/UnixSignalHeartBeat.cpp
+++ b/Engine/Source/Runtime/Core/Private/Unix/UnixSignalHeartBeat.cpp
@@ -56,10 +56,24 @@ namespace
 	void SignalHitchHandler(int Signal)
 	{
 #if USE_HITCH_DETECTION
-		UE_LOG(LogUnixHeartBeat, Log, TEXT("SignalHitchHandler"));
+		UE_LOG(LogUnixHeartBeat, Verbose, TEXT("SignalHitchHandler"));
 		GHitchDetected = true;
 #endif
 	}
+
+	timespec TimerGetTime(timer_t TimerId)
+	{
+		struct itimerspec HeartBeatTime;
+		FMemory::Memzero(HeartBeatTime);
+
+		if (timer_gettime(TimerId, &HeartBeatTime) == -1)
+		{
+			int Errno = errno;
+			UE_LOG(LogUnixHeartBeat, Warning, TEXT("Failed to timer_gettime() errno=%d (%s)"), Errno, UTF8_TO_TCHAR(strerror(Errno)));
+		}
+
+		return HeartBeatTime.it_value;
+	}
 }
 
 FUnixSignalGameHitchHeartBeat::FUnixSignalGameHitchHeartBeat()
@@ -152,7 +166,7 @@ void FUnixSignalGameHitchHeartBeat::InitSettings()
 		}
 	}
 
-	UE_LOG(LogUnixHeartBeat, Log, TEXT("HitchThresholdS:%f"), HitchThresholdS);
+	UE_LOG(LogUnixHeartBeat, Verbose, TEXT("HitchThresholdS:%f"), HitchThresholdS);
 }
 
 void FUnixSignalGameHitchHeartBeat::FrameStart(bool bSkipThisFrame)
@@ -160,8 +174,15 @@ void FUnixSignalGameHitchHeartBeat::FrameStart(bool bSkipThisFrame)
 #if USE_HITCH_DETECTION
 	check(IsInGameThread());
 
+	UE_LOG(LogUnixHeartBeat, VeryVerbose, TEXT("bDisabled:%s SuspendCount:%d TimerId:0x%x timer:%d"),
+		bDisabled ? TEXT("true") : TEXT("false"), SuspendCount, TimerId, 
+		TimerId ? TimerGetTime(TimerId).tv_nsec : 0);
+
 	if (!bDisabled && SuspendCount == 0 && TimerId)
 	{
+		UE_LOG(LogUnixHeartBeat, VeryVerbose, TEXT("HitchThresholdS:%f MinimalHitchThreashold:%f"),
+			HitchThresholdS, MinimalHitchThreashold);
+
 		if (HitchThresholdS > MinimalHitchThreashold)
 		{
 			struct itimerspec HeartBeatTime;
@@ -211,7 +232,7 @@ void FUnixSignalGameHitchHeartBeat::SuspendHeartBeat()
 	}
 
 	SuspendCount++;
-	UE_LOG(LogUnixHeartBeat, Log, TEXT("SuspendCount:%d"), SuspendCount);
+	UE_LOG(LogUnixHeartBeat, Verbose, TEXT("SuspendCount:%d"), SuspendCount);
 
 	if (TimerId)
 	{
@@ -238,7 +259,7 @@ void FUnixSignalGameHitchHeartBeat::ResumeHeartBeat()
 	if( SuspendCount > 0)
 	{
 		SuspendCount--;
-		UE_LOG(LogUnixHeartBeat, Log, TEXT("SuspendCount:%d"), SuspendCount);
+		UE_LOG(LogUnixHeartBeat, Verbose, TEXT("SuspendCount:%d"), SuspendCount);
 
 		FrameStart(true);
 	}
@@ -247,7 +268,7 @@ void FUnixSignalGameHitchHeartBeat::ResumeHeartBeat()
 
 void FUnixSignalGameHitchHeartBeat::Restart()
 {
-	UE_LOG(LogUnixHeartBeat, Log, TEXT("Restart"));
+	UE_LOG(LogUnixHeartBeat, Verbose, TEXT("Restart"));
 	bDisabled = false;
 
 	// If we still have a valid handle on the timer_t clean it up
@@ -262,7 +283,7 @@ void FUnixSignalGameHitchHeartBeat::Restart()
 
 void FUnixSignalGameHitchHeartBeat::Stop()
 {
-	UE_LOG(LogUnixHeartBeat, Log, TEXT("Stop"));
+	UE_LOG(LogUnixHeartBeat, Verbose, TEXT("Stop"));
 	SuspendHeartBeat();
 	bDisabled = true;
 }
