Make AutoRTFM 4x faster with the DeltaBlue benchmark.

I've:

- Made sure the write log pages are uninitialized (rather than the default of memset 0), as that is safe and that change alone got us 2x perf.
- Noticed that we do tons of tiny transactions, meaning my default write log allocator page size was far too high. Reducing it to 128 made us another 2x faster.

diff --git a/Engine/Source/Runtime/Core/Private/AutoRTFM/TransactionInlines.h b/Engine/Source/Runtime/Core/Private/AutoRTFM/TransactionInlines.h
--- a/Engine/Source/Runtime/Core/Private/AutoRTFM/TransactionInlines.h
+++ b/Engine/Source/Runtime/Core/Private/AutoRTFM/TransactionInlines.h
@@ -10,8 +10,6 @@ namespace AutoRTFM
 
 inline void FTransaction::RecordWriteMaxPageSized(void* LogicalAddress, size_t Size)
 {
-    ASSERT(Size <= UINT16_MAX);
-
     FMemoryLocation Key(LogicalAddress);
     Key.SetTopTag(static_cast<uint16_t>(Size));
 
@@ -59,9 +57,9 @@ inline void FTransaction::RecordWrite(void* LogicalAddress, size_t Size, bool bI
 
     size_t I = 0;
 
-    for (; (I + UINT16_MAX) < Size; I += UINT16_MAX)
+    for (; (I + FWriteLogBumpAllocator::MaxSize) < Size; I += FWriteLogBumpAllocator::MaxSize)
     {
-        RecordWriteMaxPageSized(Address + I, UINT16_MAX);
+        RecordWriteMaxPageSized(Address + I, FWriteLogBumpAllocator::MaxSize);
     }
 
     // Remainder at the end of the memcpy.
@@ -70,7 +68,7 @@ inline void FTransaction::RecordWrite(void* LogicalAddress, size_t Size, bool bI
 
 inline void FTransaction::DidAllocate(void* LogicalAddress, size_t Size)
 {
-    if (Size <= UINT16_MAX)
+    if (Size <= FWriteLogBumpAllocator::MaxSize)
     {
         FMemoryLocation Key(LogicalAddress);
         Key.SetTopTag(static_cast<uint16_t>(Size));
diff --git a/Engine/Source/Runtime/Core/Private/AutoRTFM/WriteLogBumpAllocator.h b/Engine/Source/Runtime/Core/Private/AutoRTFM/WriteLogBumpAllocator.h
--- a/Engine/Source/Runtime/Core/Private/AutoRTFM/WriteLogBumpAllocator.h
+++ b/Engine/Source/Runtime/Core/Private/AutoRTFM/WriteLogBumpAllocator.h
@@ -11,8 +11,12 @@ namespace AutoRTFM
 	class FWriteLogBumpAllocator final
 	{
 	public:
+		constexpr static size_t MaxSize = 128;
+
 		FWriteLogBumpAllocator()
 		{
+			// The max we can support in our tagged pointers is 16 bits of data.
+			static_assert(MaxSize <= UINT16_MAX);
 			Reset();
 		}
 
@@ -25,7 +29,7 @@ namespace AutoRTFM
 
 		void* Allocate(size_t Bytes)
 		{
-			ASSERT(Bytes <= FPage::MaxSize);
+			ASSERT(Bytes <= MaxSize);
 
 			if (nullptr == Start)
 			{
@@ -33,7 +37,7 @@ namespace AutoRTFM
 				Current = Start;
 			}
 
-			if (Bytes <= (FPage::MaxSize - Current->Size))
+			if (Bytes <= (MaxSize - Current->Size))
 			{
 				void* Result = &Current->Bytes[Current->Size];
 				Current->Size += Bytes;
@@ -80,13 +84,10 @@ namespace AutoRTFM
 
 		struct FPage final
 		{
-			// TODO: This is because we used tagged pointers elsewhere and so
-			// the max size of an allocation we will track in the bump allocator
-			// is 16 bits.
-			constexpr static size_t MaxSize = UINT16_MAX;
-
-			uint8_t Bytes[MaxSize];
+			// Specify a constructor so that `Bytes` below will be left uninitialized.
+			explicit FPage() {}
 
+			uint8_t Bytes[FWriteLogBumpAllocator::MaxSize];
 			FPage* Next = nullptr;
 			size_t Size = 0;
 		};
