Reworking of ParseValue to buffer
Addeed basic tests

diff --git a/Engine/Source/Runtime/Core/Private/Misc/Parse.cpp b/Engine/Source/Runtime/Core/Private/Misc/Parse.cpp
--- a/Engine/Source/Runtime/Core/Private/Misc/Parse.cpp
+++ b/Engine/Source/Runtime/Core/Private/Misc/Parse.cpp
@@ -238,6 +238,12 @@ bool FParse::Value(
 	const TCHAR**	OptStreamGotTo
 )
 {
+	if (MaxLen == 0)
+	{
+		return false;
+	}
+	check(Value && MaxLen > 0);
+
 	bool bSuccess = false;
 	int32 MatchLen = FCString::Strlen(Match);
 
@@ -246,58 +252,56 @@ bool FParse::Value(
 		*OptStreamGotTo = nullptr;
 	}
 
-	for (const TCHAR* Found = FCString::Strifind(Stream, Match, true); Found != nullptr; Found = FCString::Strifind(Found + MatchLen, Match, true))
+	const TCHAR* FoundInStream = FCString::Strifind(Stream, Match, true);
+	if (FoundInStream == nullptr)
 	{
-		const TCHAR* Start = Found + MatchLen;
+		Value[0] = TCHAR('\0');
+		return false;
+	}
 
-		// Check for quoted arguments' string with spaces
-		// -Option="Value1 Value2"
-		//         ^~~~Start
-		bool bArgumentsQuoted = *Start == '"';
+	const TCHAR* ValueStartInStream = FoundInStream + MatchLen;
+	const TCHAR* ValueEndInStream;
 
-		if (bArgumentsQuoted)
-		{
-			// Skip quote character if only params were quoted.
-			int32 QuoteCharactersToSkip = 1;
-			FCString::Strncpy(Value, Start + QuoteCharactersToSkip, MaxLen);
+	// Check for quoted arguments' string with spaces
+	// -Option="Value1 Value2"
+	//         ^~~~Start
+	const bool bArgumentsQuoted = *ValueStartInStream == '"';
 
-			Value[MaxLen - 1] = TCHAR('\0');
-			TCHAR* Temp = FCString::Strstr( Value, TEXT("\x22") );
-			if (Temp != nullptr)
-			{
-				*Temp = TCHAR('\0');
-			}
-		}
-		else
-		{
-			// Skip initial whitespace
-			Start += FCString::Strspn(Start, TEXT(" \r\n\t"));
-
-			// Non-quoted string without spaces.
-			FCString::Strncpy( Value, Start, MaxLen );
-			Value[MaxLen - 1]= TCHAR('\0');
-			TCHAR* Temp;
-			Temp = FCString::Strstr( Value, TEXT(" ")  ); if( Temp ) *Temp = TCHAR('\0');
-			Temp = FCString::Strstr( Value, TEXT("\r") ); if( Temp ) *Temp = TCHAR('\0');
-			Temp = FCString::Strstr( Value, TEXT("\n") ); if( Temp ) *Temp = TCHAR('\0');
-			Temp = FCString::Strstr( Value, TEXT("\t") ); if( Temp ) *Temp = TCHAR('\0');
-			if (bShouldStopOnSeparator)
-			{
-				Temp = FCString::Strstr( Value, TEXT(",")  ); if( Temp ) *Temp = TCHAR('\0');
-				Temp = FCString::Strstr( Value, TEXT(")")  ); if( Temp ) *Temp = TCHAR('\0');
-			}
-		}
+	if (bArgumentsQuoted)
+	{
+		// Skip quote character if only params were quoted.
+		ValueStartInStream += 1;
+		ValueEndInStream = FCString::Strstr(ValueStartInStream, TEXT("\x22"));
 
-		if (OptStreamGotTo) 
+		if (ValueEndInStream == nullptr)
 		{
-			*OptStreamGotTo = Start + FCString::Strlen(Value);
+			// this should probably log a warning if bArgumentsQuoted is true, as we started with a '"' and didn't find the terminating one.
+			ValueEndInStream = FoundInStream + FCString::Strlen(FoundInStream);
 		}
+	}
+	else
+	{
+		// Skip initial whitespace
+		const TCHAR* WhiteSpaceChars = TEXT(" \r\n\t");
+		ValueStartInStream += FCString::Strspn(ValueStartInStream, WhiteSpaceChars);
 
-		bSuccess = true;
-		break;
+		// Non-quoted string without spaces.
+		const TCHAR* TerminatingChars = bShouldStopOnSeparator ? TEXT(",) \r\n\t") : WhiteSpaceChars;
+		ValueEndInStream = ValueStartInStream + FCString::Strcspn(ValueStartInStream, TerminatingChars);
 	}
 
-	return bSuccess;
+	int32 ValueLength = FMath::Min<int32>(MaxLen - 1, UE_PTRDIFF_TO_INT32(ValueEndInStream - ValueStartInStream));
+	// It is possible for ValueLength to be 0.
+	// FCString::Strncpy asserts that its copying at least 1 char, memcpy has no such constraint.
+	FMemory::Memcpy(Value, ValueStartInStream, sizeof(Value[0]) * ValueLength);
+	Value[ValueLength] = TCHAR('\0');
+
+	if (OptStreamGotTo)
+	{
+		*OptStreamGotTo = ValueEndInStream;
+	}
+
+	return true;
 }
 
 //
diff --git a/Engine/Source/Runtime/Core/Tests/Misc/ParseTests.cpp b/Engine/Source/Runtime/Core/Tests/Misc/ParseTests.cpp
new file mode 100644
--- /dev/null
+++ b/Engine/Source/Runtime/Core/Tests/Misc/ParseTests.cpp
@@ -0,0 +1,121 @@
+// Copyright Epic Games, Inc. All Rights Reserved.
+
+#if WITH_LOW_LEVEL_TESTS
+
+#include "CoreMinimal.h"
+#include "Misc/AutomationTest.h"
+#include "Misc/Parse.h"
+
+#include "Tests/TestHarnessAdapter.h"
+
+TEST_CASE("Parse::Value::ToBuffer", "[Parse][Smoke]")
+{
+	TCHAR Buffer[256];
+
+	SECTION("Basic Usage") 
+	{
+		const TCHAR* Line = TEXT("a=a1 b=b2 c=c3");
+
+		CHECK(FParse::Value(Line, TEXT("a="), Buffer, 256));
+		CHECK(FCString::Strcmp(Buffer, TEXT("a1")) == 0);
+
+		CHECK(FParse::Value(Line, TEXT("b="), Buffer, 256));
+		CHECK(FCString::Strcmp(Buffer, TEXT("b2")) == 0);
+
+		CHECK(FParse::Value(Line, TEXT("c="), Buffer, 256));
+		CHECK(FCString::Strcmp(Buffer, TEXT("c3")) == 0);
+
+		CHECK(false == FParse::Value(Line, TEXT("not_there="), Buffer, 256));
+		CHECK(Buffer[0] == TCHAR(0));
+	}
+
+	SECTION("Quoted Values")
+	{
+		CHECK(FParse::Value(TEXT("a=a1 b=\"value with a space, and commas\" c=c3"), TEXT("b="), Buffer, 256));
+		CHECK(FCString::Strcmp(Buffer, TEXT("value with a space, and commas")) == 0);
+	}
+
+	SECTION("Value may (not)? have a delimiter")
+	{
+		const TCHAR* Line = TEXT("a=a1,a2");
+
+		CHECK(FParse::Value(Line, TEXT("a="), Buffer, 256, true));
+		CHECK(FCString::Strcmp(Buffer, TEXT("a1")) == 0);
+
+		CHECK(FParse::Value(Line, TEXT("a="), Buffer, 256, false)); // false = don't stop on , or )
+		CHECK(FCString::Strcmp(Buffer, TEXT("a1,a2")) == 0);
+	}
+
+	SECTION("Value may have spaces on its left")
+	{
+		CHECK(FParse::Value(TEXT("a=   value"), TEXT("a="), Buffer, 256));
+		CHECK(FCString::Strcmp(Buffer, TEXT("value")) == 0);
+	}
+
+	SECTION("Value could be a key value pair")
+	{
+		CHECK(FParse::Value(TEXT("a=  b=value"), TEXT("a="), Buffer, 256));
+		CHECK(FCString::Strcmp(Buffer, TEXT("b=value")) == 0);
+
+		CHECK(FParse::Value(TEXT("a=  b=  value"), TEXT("a="), Buffer, 256));
+		CHECK(FCString::Strcmp(Buffer, TEXT("b=")) == 0);
+		CHECK(FParse::Value(TEXT("a=  b=  value"), TEXT("b="), Buffer, 256));
+		CHECK(FCString::Strcmp(Buffer, TEXT("value")) == 0);
+	}
+
+	SECTION("Key may appear mutiple times")
+	{
+		const TCHAR* Line = TEXT("rep=a1 rep=b2 rep=c3");
+		const TCHAR* ExpectedResults[] = { TEXT("a1"), TEXT("b2"), TEXT("c3") };
+
+		const TCHAR* Cursor = Line;
+		for (int Loop = 0; Loop < 4; ++Loop)
+		{
+			CHECK(Cursor != nullptr);
+
+			bool bFound = FParse::Value(Cursor, TEXT("rep="), Buffer, 256, true, &Cursor);
+
+			if (Loop < 3) 
+			{
+				CHECK(bFound);
+				CHECK(FCString::Strcmp(Buffer, ExpectedResults[Loop]) == 0);
+			}
+			else
+			{
+				CHECK(!bFound);
+				CHECK(Buffer[0] == TCHAR(0));
+				CHECK(Cursor == nullptr);
+			}
+		}
+	}
+	
+	SECTION("Key may have no value, It is found but Value is empty")
+	{
+		CHECK(FParse::Value(TEXT("a=   "), TEXT("a="), Buffer, 256));
+		CHECK(Buffer[0] == TCHAR(0));
+	}
+
+	SECTION("Key with unbalanced quote, It is found but Value is empty")
+	{
+		for (TCHAR& C : Buffer)
+		{
+			C = TCHAR('*');
+		}
+		CHECK(FParse::Value(TEXT("a=\"   "), TEXT("a="), Buffer, 256));
+		CHECK(FCString::Strchr(Buffer, TCHAR('*')) == nullptr);
+	}
+
+	SECTION("Key may have no value, It is found but Value is empty")
+	{
+		CHECK(FParse::Value(TEXT("a=   "), TEXT("a="), Buffer, 256));
+		CHECK(Buffer[0] == TCHAR(0));
+	}
+
+	SECTION("Output var sanity")
+	{
+		CHECK(false == FParse::Value(TEXT("a=   "), TEXT("a="), Buffer, 0));
+	}
+}
+
+#endif
+	
\ No newline at end of file
