CompressedBuffer: Added parallel decoding for targets with editor-only data

diff --git a/Engine/Source/Runtime/Core/Private/Compression/CompressedBuffer.cpp b/Engine/Source/Runtime/Core/Private/Compression/CompressedBuffer.cpp
--- a/Engine/Source/Runtime/Core/Private/Compression/CompressedBuffer.cpp
+++ b/Engine/Source/Runtime/Core/Private/Compression/CompressedBuffer.cpp
@@ -7,6 +7,7 @@
 #include "Async/ParallelFor.h"
 #include "Compression/OodleDataCompression.h"
 #include "Containers/ArrayView.h"
+#include "HAL/PlatformMisc.h"
 #include "Hash/Blake3.h"
 #include "IO/IoHash.h"
 #include "Math/UnrealMathUtility.h"
@@ -28,6 +29,8 @@ namespace UE::CompressedBuffer::Private
 static constexpr uint64 DefaultBlockSize = 256 * 1024;
 static constexpr uint64 DefaultHeaderSize = 4 * 1024;
 
+static constexpr uint64 ParallelDecodeMinRawSize = 4 * 1024 * 1024;
+
 ///////////////////////////////////////////////////////////////////////////////////////////////////////////////////////
 
 /** Method used to compress the data in a compressed buffer. */
@@ -154,6 +157,7 @@ static_assert(sizeof(FHeader) == 64, "FHeader is the wrong size.");
 class FDecoderSource
 {
 public:
+	virtual bool SupportsParallelRead() const = 0;
 	virtual bool Read(uint64 Offset, FMutableMemoryView Data) const = 0;
 	virtual FMemoryView ReadOrView(uint64 Offset, uint64 Size, FDecoderContext& Context) const = 0;
 	virtual FCompositeBuffer ReadToComposite(uint64 Offset, uint64 Size) const = 0;
@@ -429,6 +433,11 @@ public:
 
 protected:
 	virtual bool DecompressBlock(FMutableMemoryView RawData, FMemoryView CompressedData) const = 0;
+
+#if WITH_EDITORONLY_DATA
+private:
+	bool TryParallelDecompressTo(FDecoderContext& Context, const FDecoderSource& Source, const FHeader& Header, FMemoryView HeaderView, uint64 RawOffset, FMutableMemoryView RawView) const;
+#endif
 };
 
 bool FBlockDecoder::TryDecompressTo(
@@ -437,33 +446,38 @@ bool FBlockDecoder::TryDecompressTo(
 	const FHeader& Header,
 	const FMemoryView HeaderView,
 	const uint64 RawOffset,
-	FMutableMemoryView RawView) const
+	const FMutableMemoryView RawView) const
 {
-	TRACE_CPUPROFILER_EVENT_SCOPE(FBlockDecoder::TryDecompressTo);
-
 	if (Header.TotalRawSize < RawOffset + RawView.GetSize())
 	{
 		return false;
 	}
 
-	const uint64 BlockSize = uint64(1) << Header.BlockSizeExponent;
-	const uint64 LastBlockSize = BlockSize - (BlockSize * Header.BlockCount - Header.TotalRawSize);
-	const uint32 FirstBlockIndex = uint32(RawOffset / BlockSize);
-	const uint32 LastBlockIndex = uint32((RawOffset + RawView.GetSize() - 1) / BlockSize);
+	const uint32 FirstBlockIndex = uint32(RawOffset >> Header.BlockSizeExponent);
+	const uint32 LastBlockIndex = uint32((RawOffset + RawView.GetSize() - 1) >> Header.BlockSizeExponent);
 
-	uint64 RawBlockOffset = RawOffset % BlockSize;
+#if WITH_EDITORONLY_DATA
+	if (RawView.GetSize() >= ParallelDecodeMinRawSize && LastBlockIndex - FirstBlockIndex > 1 && Source.SupportsParallelRead())
+	{
+		return TryParallelDecompressTo(Context, Source, Header, HeaderView, RawOffset, RawView);
+	}
+#endif
+
+	TRACE_CPUPROFILER_EVENT_SCOPE(FBlockDecoder::TryDecompressTo);
 
 	const uint32* const CompressedBlockSizes = static_cast<const uint32*>(HeaderView.RightChop(sizeof(FHeader)).GetData());
 	uint64 CompressedOffset = sizeof(FHeader) + sizeof(uint32) * uint32(Header.BlockCount) +
 		Algo::TransformAccumulate(MakeArrayView(CompressedBlockSizes, FirstBlockIndex),
 			[](uint32 Size) -> uint64 { return NETWORK_ORDER32(Size); }, uint64(0));
 
-	// @todo : Parallel decoder needed!  This stream is chunked into blocks but this implementation is serial.  This is a bottleneck for large data loading.
-
-	for (uint32 BlockIndex = FirstBlockIndex; BlockIndex <= LastBlockIndex; BlockIndex++)
+	for (uint32 BlockIndex = FirstBlockIndex; BlockIndex <= LastBlockIndex; ++BlockIndex)
 	{
-		const uint64 RawBlockSize = BlockIndex == Header.BlockCount - 1 ? LastBlockSize : BlockSize;
-		const uint64 RawBlockReadSize = FMath::Min(RawView.GetSize(), RawBlockSize - RawBlockOffset);
+		const auto BlocksToBytes = [&Header](uint64 BlockCount) -> uint64 { return BlockCount << Header.BlockSizeExponent; };
+		const bool bFirstBlock = (BlockIndex == FirstBlockIndex);
+		const uint64 BlockSize = BlocksToBytes(1);
+		const uint64 RawBlockOffset = bFirstBlock ? RawOffset & (BlockSize - 1) : 0;
+		const uint64 RawBlockSize = FMath::Min(BlockSize, Header.TotalRawSize - BlocksToBytes(BlockIndex));
+		const FMutableMemoryView RawTarget = RawView.Mid(bFirstBlock ? 0 : BlocksToBytes(BlockIndex) - RawOffset, RawBlockSize - RawBlockOffset);
 		const uint32 CompressedBlockSize = NETWORK_ORDER32(CompressedBlockSizes[BlockIndex]);
 		const bool bIsCompressed = CompressedBlockSize < RawBlockSize;
 
@@ -471,14 +485,14 @@ bool FBlockDecoder::TryDecompressTo(
 		{
 			if (Context.RawBlockIndex == BlockIndex)
 			{
-				RawView.Left(RawBlockReadSize).CopyFrom(Context.RawBlock.GetView().Mid(RawBlockOffset, RawBlockReadSize));
+				RawTarget.CopyFrom(Context.RawBlock.GetView().Mid(RawBlockOffset, RawTarget.GetSize()));
 			}
 			else
 			{
 				FMutableMemoryView RawBlock;
-				if (RawBlockReadSize == RawBlockSize)
+				if (RawTarget.GetSize() == RawBlockSize)
 				{
-					RawBlock = RawView.Left(RawBlockSize);
+					RawBlock = RawTarget;
 				}
 				else
 				{
@@ -496,24 +510,141 @@ bool FBlockDecoder::TryDecompressTo(
 					return false;
 				}
 
-				if (RawBlockReadSize != RawBlockSize)
+				if (RawTarget.GetSize() != RawBlockSize)
 				{
-					RawView.CopyFrom(RawBlock.Mid(RawBlockOffset, RawBlockReadSize));
+					RawTarget.CopyFrom(RawBlock.Mid(RawBlockOffset, RawTarget.GetSize()));
 				}
 			}
 		}
 		else
 		{
-			Source.Read(CompressedOffset + RawBlockOffset, RawView.Left(RawBlockReadSize));
+			Source.Read(CompressedOffset + RawBlockOffset, RawTarget);
 		}
 
-		RawBlockOffset = 0;
 		CompressedOffset += CompressedBlockSize;
-		RawView += RawBlockReadSize;
 	}
 
-	return RawView.GetSize() == 0;
+	return true;
+}
+
+#if WITH_EDITORONLY_DATA
+FORCENOINLINE bool FBlockDecoder::TryParallelDecompressTo(
+	FDecoderContext& Context,
+	const FDecoderSource& Source,
+	const FHeader& Header,
+	const FMemoryView HeaderView,
+	const uint64 RawOffset,
+	const FMutableMemoryView RawView) const
+{
+	if (Header.TotalRawSize < RawOffset + RawView.GetSize())
+	{
+		return false;
+	}
+
+	const uint32 FirstBlockIndex = uint32(RawOffset >> Header.BlockSizeExponent);
+	const uint32 LastBlockIndex = uint32((RawOffset + RawView.GetSize() - 1) >> Header.BlockSizeExponent);
+
+	TRACE_CPUPROFILER_EVENT_SCOPE(FBlockDecoder::TryParallelDecompressTo);
+
+	const uint32* const CompressedBlockSizes = static_cast<const uint32*>(HeaderView.RightChop(sizeof(FHeader)).GetData());
+	TArray<uint64, TInlineAllocator<64>> CompressedOffsets;
+	CompressedOffsets.Reserve(LastBlockIndex - FirstBlockIndex + 1);
+	CompressedOffsets.Add(sizeof(FHeader) + sizeof(uint32) * uint32(Header.BlockCount) +
+		Algo::TransformAccumulate(MakeArrayView(CompressedBlockSizes, FirstBlockIndex),
+			[](uint32 Size) -> uint64 { return NETWORK_ORDER32(Size); }, uint64(0)));
+	for (uint32 BlockIndex = FirstBlockIndex; BlockIndex < LastBlockIndex; ++BlockIndex)
+	{
+		CompressedOffsets.Add(CompressedOffsets.Last() + NETWORK_ORDER32(CompressedBlockSizes[BlockIndex]));
+	}
+
+	const auto TryDecompressBlock = [this, &Source, &Header, RawOffset, RawView, CompressedBlockSizes, CompressedOffsets = CompressedOffsets.GetData(), FirstBlockIndex](FDecoderContext& Context, const uint32 BlockIndex)
+	{
+		const auto BlocksToBytes = [&Header](uint64 BlockCount) -> uint64 { return BlockCount << Header.BlockSizeExponent; };
+		const bool bFirstBlock = (BlockIndex == FirstBlockIndex);
+		const uint64 BlockSize = BlocksToBytes(1);
+		const uint64 RawBlockOffset = bFirstBlock ? RawOffset & (BlockSize - 1) : 0;
+		const uint64 RawBlockSize = FMath::Min(BlockSize, Header.TotalRawSize - BlocksToBytes(BlockIndex));
+		const FMutableMemoryView RawTarget = RawView.Mid(bFirstBlock ? 0 : BlocksToBytes(BlockIndex) - RawOffset, RawBlockSize - RawBlockOffset);
+		const uint64 CompressedOffset = CompressedOffsets[BlockIndex - FirstBlockIndex];
+		const uint32 CompressedBlockSize = NETWORK_ORDER32(CompressedBlockSizes[BlockIndex]);
+		const bool bIsCompressed = CompressedBlockSize < RawBlockSize;
+
+		if (bIsCompressed)
+		{
+			if (Context.RawBlockIndex == BlockIndex)
+			{
+				RawTarget.CopyFrom(Context.RawBlock.GetView().Mid(RawBlockOffset, RawTarget.GetSize()));
+			}
+			else
+			{
+				FMutableMemoryView RawBlock;
+				if (RawTarget.GetSize() == RawBlockSize)
+				{
+					RawBlock = RawTarget;
+				}
+				else
+				{
+					if (Context.RawBlock.GetSize() < RawBlockSize)
+					{
+						Context.RawBlock = FUniqueBuffer::Alloc(BlockSize);
+					}
+					RawBlock = Context.RawBlock.GetView().Left(RawBlockSize);
+					Context.RawBlockIndex = BlockIndex;
+				}
+
+				const FMemoryView CompressedBlock = Source.ReadOrView(CompressedOffset, CompressedBlockSize, Context);
+				if (CompressedBlock.IsEmpty() || !DecompressBlock(RawBlock, CompressedBlock))
+				{
+					return false;
+				}
+
+				if (RawTarget.GetSize() != RawBlockSize)
+				{
+					RawTarget.CopyFrom(RawBlock.Mid(RawBlockOffset, RawTarget.GetSize()));
+				}
+			}
+		}
+		else
+		{
+			Source.Read(CompressedOffset + RawBlockOffset, RawTarget);
+		}
+
+		return true;
+	};
+
+	// Track whether any block failed to decompress.
+	std::atomic<bool> bSuccess = true;
+	// Share decoder contexts among blocks in the same task. These will be unused for most inputs.
+	TArray<FDecoderContext, TInlineAllocator<16>> ParallelContexts;
+	// Decompress the first and last blocks as pre-work because they are the two that may use Context.RawBlock.
+	const int32 ParallelBlockCount = int32(LastBlockIndex - FirstBlockIndex - 1);
+	// Require a minimum batch of 1 MiB raw blocks to avoid tasks that are too tiny to be worthwhile.
+	const int32 ParallelMinBatchSize = (Header.BlockSizeExponent < 20) ? (1 << (20 - Header.BlockSizeExponent)) : 1;
+	// Limit task count because decoding saturates memory bandwidth with very few threads.
+	static const int32 MaxTaskCount = FMath::Max(3, FPlatformMisc::NumberOfCoresIncludingHyperthreads() / 4);
+
+	ParallelContexts.SetNum(FMath::Min(MaxTaskCount, ParallelBlockCount));
+	ParallelForWithPreWorkWithExistingTaskContext(TEXT("BlockDecoder.TryDecompressTo.PF"),
+		MakeArrayView(ParallelContexts), ParallelBlockCount, ParallelMinBatchSize,
+		[&TryDecompressBlock, &bSuccess, FirstBlockIndex](FDecoderContext& TaskContext, const int32 ParallelBlockIndex)
+		{
+			if (!TryDecompressBlock(TaskContext, uint32(ParallelBlockIndex + FirstBlockIndex + 1)))
+			{
+				bSuccess.store(false, std::memory_order_relaxed);
+			}
+		},
+		[&TryDecompressBlock, &bSuccess, &Context, FirstBlockIndex, LastBlockIndex]
+		{
+			if (!TryDecompressBlock(Context, FirstBlockIndex) || !TryDecompressBlock(Context, LastBlockIndex))
+			{
+				bSuccess.store(false, std::memory_order_relaxed);
+			}
+		},
+		EParallelForFlags::Unbalanced);
+
+	return bSuccess.load(std::memory_order_relaxed);
 }
+#endif // WITH_EDITORONLY_DATA
 
 FCompositeBuffer FBlockDecoder::DecompressToComposite(
 	FDecoderContext& Context,
@@ -781,6 +912,11 @@ public:
 	{
 	}
 
+	bool SupportsParallelRead() const final
+	{
+		return false;
+	}
+
 	bool Read(uint64 Offset, FMutableMemoryView Data) const final
 	{
 		Archive.Seek(int64(BaseOffset + Offset));
@@ -821,6 +957,11 @@ public:
 	{
 	}
 
+	bool SupportsParallelRead() const final
+	{
+		return true;
+	}
+
 	bool Read(const uint64 Offset, const FMutableMemoryView Data) const final
 	{
 		if (Offset + Data.GetSize() <= Buffer.GetSize())
diff --git a/Engine/Source/Runtime/Core/Public/Async/ParallelFor.h b/Engine/Source/Runtime/Core/Public/Async/ParallelFor.h
--- a/Engine/Source/Runtime/Core/Public/Async/ParallelFor.h
+++ b/Engine/Source/Runtime/Core/Public/Async/ParallelFor.h
@@ -114,7 +114,7 @@ namespace ParallelForImpl
 	template<typename BodyType, typename PreWorkType, typename ContextType>
 	inline void ParallelForInternal(const TCHAR* DebugName, int32 Num, int32 MinBatchSize, BodyType Body, PreWorkType CurrentThreadWorkToDoBeforeHelping, EParallelForFlags Flags, const TArrayView<ContextType>& Contexts)
 	{
-		if (Num == 0) 
+		if (Num == 0)
 		{
 			return;
 		}
@@ -625,6 +625,30 @@ inline void ParallelForWithPreWorkWithTaskContext(
 	}
 }
 
+/** 
+ * General purpose parallel for that uses the taskgraph
+ * @param DebugName; ProfilingScope and DebugName
+ * @param Contexts; User-privided array of user-defined task-level context objects
+ * @param Num; number of calls of Body; Body(0), Body(1), ..., Body(Num - 1)
+ * @param MinBatchSize; Minimum Size of a Batch (will only launch DivUp(Num, MinBatchSize) Workers 
+ * @param Body; Function to call from multiple threads
+ * @param CurrentThreadWorkToDoBeforeHelping; The work is performed on the main thread before it starts helping with the ParallelFor proper
+ * @param Flags; Used to customize the behavior of the ParallelFor if needed.
+ * Notes: Please add stats around to calls to parallel for and within your lambda as appropriate. Do not clog the task graph with long running tasks or tasks that block.
+ */
+template <typename ContextType, typename BodyType, typename PreWorkType>
+inline void ParallelForWithPreWorkWithExistingTaskContext(
+	const TCHAR* DebugName,
+	TArrayView<ContextType> Contexts,
+	int32 Num,
+	int32 MinBatchSize,
+	BodyType&& Body,
+	PreWorkType&& CurrentThreadWorkToDoBeforeHelping,
+	EParallelForFlags Flags = EParallelForFlags::None)
+{
+	ParallelForImpl::ParallelForInternal(DebugName, Num, MinBatchSize, Forward<BodyType>(Body), Forward<PreWorkType>(CurrentThreadWorkToDoBeforeHelping), Flags, Contexts);
+}
+
 /** 
 	*	General purpose parallel for that uses the taskgraph. This variant constructs for the caller a user-defined context
 	* 	object for each task that may get spawned to do work, and passes it on to the loop body to give it a task-local
@@ -743,10 +767,7 @@ inline void ParallelForWithTaskContext(const TCHAR* DebugName, TArray<ContextTyp
 template <typename ContextType, typename FunctionType>
 inline void ParallelForWithExistingTaskContext(TArrayView<ContextType> Contexts, int32 Num, int32 MinBatchSize, const FunctionType& Body, EParallelForFlags Flags = EParallelForFlags::None)
 {
-	if (Num > 0)
-	{
-		ParallelForImpl::ParallelForInternal(TEXT("ParallelFor Task"), Num, MinBatchSize, Body, [](){}, Flags, Contexts);
-	}
+	ParallelForImpl::ParallelForInternal(TEXT("ParallelFor Task"), Num, MinBatchSize, Body, [](){}, Flags, Contexts);
 }
 
 /**
@@ -764,8 +785,5 @@ inline void ParallelForWithExistingTaskContext(TArrayView<ContextType> Contexts,
 template <typename ContextType, typename FunctionType>
 inline void ParallelForWithExistingTaskContext(const TCHAR* DebugName, TArrayView<ContextType> Contexts, int32 Num, int32 MinBatchSize, const FunctionType& Body, EParallelForFlags Flags = EParallelForFlags::None)
 {
-	if (Num > 0)
-	{
-		ParallelForImpl::ParallelForInternal(DebugName, Num, MinBatchSize, Body, [](){}, Flags, Contexts);
-	}
+	ParallelForImpl::ParallelForInternal(DebugName, Num, MinBatchSize, Body, [](){}, Flags, Contexts);
 }
