Removed bad module includes in Core.

diff --git a/Engine/Source/Runtime/Core/Core.Build.cs b/Engine/Source/Runtime/Core/Core.Build.cs
--- a/Engine/Source/Runtime/Core/Core.Build.cs
+++ b/Engine/Source/Runtime/Core/Core.Build.cs
@@ -33,12 +33,7 @@ public class Core : ModuleRules
 
 		PrivateIncludePathModuleNames.AddRange(
 			new string[] {
-				"Analytics",
 				"DerivedDataCache",
-				"Engine",
-				"InputDevice",
-				"RHI",
-				"SynthBenchmark",
 				"TargetPlatform",
 			}
 			);
