Allowed TIsCharType to work with cv-qualified types.

diff --git a/Engine/Source/Runtime/Core/Public/Traits/IsCharType.h b/Engine/Source/Runtime/Core/Public/Traits/IsCharType.h
--- a/Engine/Source/Runtime/Core/Public/Traits/IsCharType.h
+++ b/Engine/Source/Runtime/Core/Public/Traits/IsCharType.h
@@ -20,5 +20,9 @@ template<>           struct TIsCharType<UTF32CHAR> { enum { Value = true  }; };
 template<>           struct TIsCharType<wchar_t>   { enum { Value = true  }; };
 #endif
 
+template <typename T> struct TIsCharType<const          T> { enum { Value = TIsCharType<T>::Value }; };
+template <typename T> struct TIsCharType<      volatile T> { enum { Value = TIsCharType<T>::Value }; };
+template <typename T> struct TIsCharType<const volatile T> { enum { Value = TIsCharType<T>::Value }; };
+
 template <typename T>
 constexpr inline bool TIsCharType_V = TIsCharType<T>::Value;
