UE_FNAME_OUTLINE_NUMBER memory optimization

Drop FNameEntry alignment to 2 and align numbered entries to 4 manually. Applies when WITH_CASE_PRESERVING_NAME == 0.

Saved 4.5MB for an internal project on Gen5 Test Client (87.5 -> 83.0MB)

diff --git a/Engine/Source/Runtime/Core/Private/UObject/UnrealNames.cpp b/Engine/Source/Runtime/Core/Private/UObject/UnrealNames.cpp
--- a/Engine/Source/Runtime/Core/Private/UObject/UnrealNames.cpp
+++ b/Engine/Source/Runtime/Core/Private/UObject/UnrealNames.cpp
@@ -422,6 +422,18 @@ private:
 	uint32 IdAndHash = 0;
 };
 
+// Stores a (string, number) fname pair to construct the string "string_(number-1)" on demand.
+// Id is a reference to another entry with Header.Len != 0 (no recursion).
+struct FNumberedEntry
+{
+#if UE_FNAME_OUTLINE_NUMBER // Let empty type exist for clang static analyzer
+	FNameEntryId	Id;
+	uint32			Number;
+
+	bool IsEmpty() const { return Id.IsNone() && Number == 0; }
+#endif
+};
+
 /**
  * Thread-safe paged FNameEntry allocator
  */
@@ -463,7 +475,7 @@ public:
 	 * @param   Size  Size in bytes to allocate, 
 	 * @return  Allocation of passed in size cast to a FNameEntry pointer.
 	 */
-	template <class ScopeLock>
+	template <class ScopeLock, bool bAlignNumberedName = false>
 	FNameEntryHandle Allocate(uint32 Bytes)
 	{
 		Bytes = Align(Bytes, alignof(FNameEntry));
@@ -471,6 +483,14 @@ public:
 
 		ScopeLock _(Lock);
 
+		// Pad up FNameEntry::NumberedName to be 4B-aligned if FNameEntry is only 2B-aligned
+#if UE_FNAME_OUTLINE_NUMBER
+		if constexpr (bAlignNumberedName)
+		{
+			CurrentByteCursor = Align(CurrentByteCursor + offsetof(FNameEntry, NumberedName), alignof(FNumberedEntry)) - offsetof(FNameEntry, NumberedName);
+		}
+#endif
+
 		// Allocate a new pool if current one is exhausted. We don't worry about a little bit
 		// of waste at the end given the relative size of pool to average and max allocation.
 		if (BlockSizeBytes - CurrentByteCursor < Bytes)
@@ -519,8 +539,10 @@ public:
 	template<class ScopeLock>
 	FNameEntryHandle CreateWithNumber(const FNameEntryId& StringPart, uint32 NumberPart, TOptional<FNameEntryId> ComparisonId, FNameEntryHeader Header)
 	{
+		static_assert(sizeof(FNameEntry::NumberedName) == sizeof(FNumberedEntry));
+		static constexpr bool bAlignNumberedName = alignof(FNumberedEntry) < alignof(FNameEntry);
 		FPlatformMisc::Prefetch(Blocks[CurrentBlock]);
-		FNameEntryHandle Handle = Allocate<ScopeLock>(FNameEntry::GetDataOffset() + sizeof(FNameEntry::NumberedName));
+		FNameEntryHandle Handle = Allocate<ScopeLock, bAlignNumberedName>(FNameEntry::GetDataOffset() + sizeof(FNameEntry::NumberedName));
 		FNameEntry& Entry = Resolve(Handle);
 
 #if WITH_CASE_PRESERVING_NAME
@@ -528,9 +550,8 @@ public:
 #endif
 
 		Entry.Header = Header;
-
-		Entry.NumberedName.Id = StringPart;
-		Entry.NumberedName.Number = NumberPart;
+		check(reinterpret_cast<int64>(Entry.NumberedName) % alignof(FNumberedEntry) == 0);
+		*reinterpret_cast<FNumberedEntry*>(Entry.NumberedName) = {StringPart, NumberPart};
 
 		return Handle;
 	}
@@ -591,7 +612,7 @@ public:
 					}
 					else // Null-terminator or numbered-entry
 					{
-						if ((End - It) < FNameEntry::GetNumberedEntrySize() || (Entry->NumberedName.Id == FNameEntryId() && Entry->NumberedName.Number == 0))
+						if ((End - It) < FNameEntry::GetNumberedEntrySize() || Entry->GetNumberedName().IsEmpty())
 						{
 							break;
 						}
@@ -641,8 +662,7 @@ private:
 			else // Null-terminator or numbered-entry
 			{
 #if UE_FNAME_OUTLINE_NUMBER
-				if ((End - It) < FNameEntry::GetNumberedEntrySize()
-				|| (Entry->NumberedName.Id == FNameEntryId() && Entry->NumberedName.Number == 0))
+				if ((End - It) < FNameEntry::GetNumberedEntrySize() || Entry->GetNumberedName().IsEmpty())
 				{
 					break;
 				}
@@ -670,8 +690,7 @@ private:
 		{
 			FNameEntry* Terminator = (FNameEntry*)(Blocks[CurrentBlock] + CurrentByteCursor);
 			Terminator->Header.Len = 0;
-			Terminator->NumberedName.Id = FNameEntryId();
-			Terminator->NumberedName.Number = 0;
+			*reinterpret_cast<FNumberedEntry*>(Terminator->NumberedName) = {};
 		}
 #else
 		if (CurrentByteCursor + FNameEntry::GetDataOffset() <= BlockSizeBytes)
@@ -1052,7 +1071,7 @@ protected:
 		//	so the headers will compare unequal
 		if( Entry.Header == Value.Hash.EntryProbeHeader )
 		{
-			return Entry.NumberedName.Id == Value.StringPart && Entry.NumberedName.Number == Value.NumberPart;
+			return Entry.GetNumberedName().Id == Value.StringPart && Entry.GetNumberedName().Number == Value.NumberPart;
 		}
 		return false;
 	}
@@ -1062,7 +1081,7 @@ protected:
 	{
 		if (Entry.Header.Len == 0)
 		{
-			OutValue = FNumberedNameValue<Sensitivity>(Entry.NumberedName.Id, Entry.NumberedName.Number);
+			OutValue = FNumberedNameValue<Sensitivity>(Entry.GetNumberedName().Id, Entry.GetNumberedName().Number);
 			return true;
 		}
 		return false;
@@ -2348,7 +2367,7 @@ int32 FNameEntry::GetNumberedEntrySize()
 uint32 FNameEntry::GetNumber() const
 {
 	checkName(Header.Len == 0);
-	return NumberedName.Number;
+	return GetNumberedName().Number;
 }
 #endif
 
@@ -2910,7 +2929,7 @@ struct FNameHelper
 
 	#if UE_FNAME_OUTLINE_NUMBER
 		int32 Number = Entry->IsNumbered() ? Entry->GetNumber() : NAME_NO_NUMBER_INTERNAL;
-		Entry = Entry->IsNumbered() ? FName::ResolveEntry(Entry->NumberedName.Id) : Entry; 
+		Entry = Entry->IsNumbered() ? FName::ResolveEntry(Entry->GetNumberedName().Id) : Entry; 
 		check(!Entry->IsNumbered());
 	#else
 		int32 Number = Name.GetNumber();
@@ -3103,7 +3122,7 @@ FNameEntryId FName::GetComparisonIndex() const
 	const FNameEntry& Entry = GetNamePool().Resolve(ComparisonIndex);
 	if (Entry.Header.Len == 0)
 	{
-		return Entry.NumberedName.Id;
+		return Entry.GetNumberedName().Id;
 	}
 	else
 	{
@@ -3117,7 +3136,7 @@ FNameEntryId FName::GetDisplayIndex() const
 	const FNameEntry& Entry = GetNamePool().Resolve(DisplayIndex);
 	if (Entry.Header.Len == 0)
 	{
-		return Entry.NumberedName.Id;
+		return Entry.GetNumberedName().Id;
 	}
 	else
 	{
@@ -3140,7 +3159,7 @@ int32 FName::GetNumber() const
 	const FNameEntry* Entry = ResolveEntry(Id);
 	if (Entry->Header.Len == 0)
 	{
-		return Entry->NumberedName.Number;
+		return Entry->GetNumberedName().Number;
 	}
 	else
 	{
@@ -3167,7 +3186,7 @@ const FNameEntry* FName::ResolveEntryRecursive(FNameEntryId LookupId)
 #if UE_FNAME_OUTLINE_NUMBER
 	if (Entry->Header.Len == 0)
 	{
-		return ResolveEntry(Entry->NumberedName.Id); // Should only ever recurse one level
+		return ResolveEntry(Entry->GetNumberedName().Id); // Should only ever recurse one level
 	}
 	else
 #endif
@@ -3346,7 +3365,7 @@ void FName::ToString(FString& Out) const
 	}
 	else
 	{
-		const FNameEntry* BaseEntry = ResolveEntry(ThisNameEntry->NumberedName.Id);
+		const FNameEntry* BaseEntry = ResolveEntry(ThisNameEntry->GetNumberedName().Id);
 		Out.Reset(BaseEntry->GetNameLength() + 6);
 		BaseEntry->AppendNameToString(Out);
 
@@ -3454,7 +3473,7 @@ void FName::AppendStringInternal(StringBuilderType& Out) const
 	}
 	else
 	{
-		const FNameEntry* BaseEntry = ResolveEntry(ThisNameEntry->NumberedName.Id);
+		const FNameEntry* BaseEntry = ResolveEntry(ThisNameEntry->GetNumberedName().Id);
 		BaseEntry->AppendNameToString(Out);
 		Out << '_' << NAME_INTERNAL_TO_EXTERNAL(ThisNameEntry->GetNumber());
 	}
@@ -4049,9 +4068,9 @@ void FNameEntry::DebugDump(FOutputDevice& Out) const
 	if (IsNumbered())
 	{
 		TCHAR Buffer[NAME_SIZE];
-		const FNameEntry& BaseEntry = GetNamePool().Resolve(NumberedName.Id);
+		const FNameEntry& BaseEntry = GetNamePool().Resolve(GetNumberedName().Id);
 		BaseEntry.GetName(Buffer);
-		Out.Logf(TEXT("%s_%d"), Buffer, NAME_INTERNAL_TO_EXTERNAL(NumberedName.Number));
+		Out.Logf(TEXT("%s_%d"), Buffer, NAME_INTERNAL_TO_EXTERNAL(GetNumberedName().Number));
 	}
 	else
 #endif
diff --git a/Engine/Source/Runtime/Core/Public/UObject/NameTypes.h b/Engine/Source/Runtime/Core/Public/UObject/NameTypes.h
--- a/Engine/Source/Runtime/Core/Public/UObject/NameTypes.h
+++ b/Engine/Source/Runtime/Core/Public/UObject/NameTypes.h
@@ -44,6 +44,7 @@ enum {NAME_SIZE	= 1024};
 
 struct FMinimalName;
 struct FScriptName;
+struct FNumberedEntry;
 class FName;
 
 /** Opaque id to a deduplicated name */
@@ -212,16 +213,7 @@ private:
 	{
 		ANSICHAR	AnsiName[NAME_SIZE];
 		WIDECHAR	WideName[NAME_SIZE];
-#if UE_FNAME_OUTLINE_NUMBER
-		// These fields are valid when Header.Len == 0.
-		// Stores a (string, number) fname pair to construct the string "string_(number-1)" on demand.
-		// Id is a reference to another entry with Header.Len != 0 (no recursion).
-		struct 
-		{
-			FNameEntryId	Id;
-			uint32			Number;
-		} NumberedName;
-#endif // UE_FNAME_OUTLINE_NUMBER
+		uint8		NumberedName[8 * UE_FNAME_OUTLINE_NUMBER];
 	};
 
 	FNameEntry(struct FClangKeepDebugInfo);
@@ -232,24 +224,11 @@ private:
 
 public:
 	/** Returns whether this name entry is represented via WIDECHAR or ANSICHAR. */
-	FORCEINLINE bool IsWide() const
-	{
-		return Header.bIsWide;
-	}
-
-	FORCEINLINE int32 GetNameLength() const
-	{
-		return Header.Len;
-	}
-
-	FORCEINLINE bool IsNumbered() const 
-	{
-#if UE_FNAME_OUTLINE_NUMBER
-		return Header.Len == 0;
-#else
-		return false;
-#endif
-	}
+	FORCEINLINE bool IsWide() const								{ return Header.bIsWide; }
+	FORCEINLINE int32 GetNameLength() const 					{ return Header.Len; }
+	FORCEINLINE bool IsNumbered() const 						{ return !!UE_FNAME_OUTLINE_NUMBER && Header.Len == 0; }
+	// @pre IsNumbered()
+	FORCEINLINE const FNumberedEntry& GetNumberedName() const	{ return reinterpret_cast<const FNumberedEntry&>(NumberedName[0]); }
 
 	/**
 	 * Copy unterminated name to TCHAR buffer without allocating.
