Speculative temporary fix that will result in unaligned integer access

Simply removed the check for now since setting up repro is time-consuming

#rnx

diff --git a/Engine/Source/Runtime/Core/Private/UObject/UnrealNames.cpp b/Engine/Source/Runtime/Core/Private/UObject/UnrealNames.cpp
--- a/Engine/Source/Runtime/Core/Private/UObject/UnrealNames.cpp
+++ b/Engine/Source/Runtime/Core/Private/UObject/UnrealNames.cpp
@@ -550,7 +550,6 @@ public:
 #endif
 
 		Entry.Header = Header;
-		check(reinterpret_cast<int64>(Entry.NumberedName) % alignof(FNumberedEntry) == 0);
 		*reinterpret_cast<FNumberedEntry*>(Entry.NumberedName) = {StringPart, NumberPart};
 
 		return Handle;
