Compute a better estimate of AvailableVirtual.

diff --git a/Engine/Source/Editor/UnrealEd/Private/CookOnTheFlyServer.cpp b/Engine/Source/Editor/UnrealEd/Private/CookOnTheFlyServer.cpp
--- a/Engine/Source/Editor/UnrealEd/Private/CookOnTheFlyServer.cpp
+++ b/Engine/Source/Editor/UnrealEd/Private/CookOnTheFlyServer.cpp
@@ -4670,15 +4670,12 @@ bool UCookOnTheFlyServer::PumpHasExceededMaxMemory(uint32& OutResultFlags)
 		// you can get that by setting the MemoryMinFreeVirtual and MemoryMinFreePhysical config to be the same
 
 		// AvailableVirtual is actually ullAvailPageFile (commit charge available)
-		// Temporarily ignore MemoryMinFreeVirtual on Linux until we have fixed the value of AvailableVirtual in FUnixPlatformMemory::GetStats
-#if !PLATFORM_LINUX
 		if (MemoryMinFreeVirtual > 0 && MemStats.AvailableVirtual < MemoryMinFreeVirtual)
 		{
 			TriggerMessages.Appendf(TEXT("\n  CookSettings.MemoryMinFreeVirtual: Available virtual memory %dMiB is less than %dMiB."),
 				static_cast<uint32>(MemStats.AvailableVirtual / 1024 / 1024), static_cast<uint32>(MemoryMinFreeVirtual / 1024 / 1024));
 			bMinFreeTriggered = true;
 		}
-#endif
 		if (MemoryMinFreePhysical > 0 && MemStats.AvailablePhysical < MemoryMinFreePhysical)
 		{
 			TriggerMessages.Appendf(TEXT("\n  CookSettings.MemoryMinFreePhysical: Available physical memory %dMiB is less than %dMiB."),
diff --git a/Engine/Source/Runtime/Core/Private/Unix/UnixPlatformMemory.cpp b/Engine/Source/Runtime/Core/Private/Unix/UnixPlatformMemory.cpp
--- a/Engine/Source/Runtime/Core/Private/Unix/UnixPlatformMemory.cpp
+++ b/Engine/Source/Runtime/Core/Private/Unix/UnixPlatformMemory.cpp
@@ -795,6 +795,38 @@ static uint32 ReadProcFields(const char *FileName, FProcField *ProcFields, uint3
 	return NumFieldsFound;
 }
 
+static float ReadOvercommitRatio()
+{
+	float OutVal = 0.0f;
+	int Fd = open("/proc/sys/vm/overcommit_ratio", O_RDONLY);
+
+	if (Fd < 0)
+	{
+		static bool bLogOnceFileNotFound = false;
+		if (!bLogOnceFileNotFound)
+		{
+			fprintf(stderr, "Warning: ReadOvercommitRatio failed to open /proc/sys/vm/overcommit_ratio (Err %d).\n", errno);
+			fflush(stderr);
+			bLogOnceFileNotFound = true;
+		}
+	}
+	else
+	{
+		// The overcommit_ratio file always contains a number from 0 to 100 and nothing else
+		char Buffer[512] = { 0 };
+		ssize_t ReadBytes = read(Fd, Buffer, sizeof(Buffer) - 1);
+		if (ReadBytes > 0)
+		{
+			OutVal = FCStringAnsi::Atof(Buffer) / 100.0f;
+		}
+
+		close(Fd);
+	}
+
+	return OutVal;
+
+}
+
 // struct CORE_API FGenericPlatformMemoryStats : public FPlatformMemoryConstants
 //   uint64 AvailablePhysical; /** The amount of physical memory currently available, in bytes. */			MemAvailable (or MemFree + Cached)
 //   uint64 AvailableVirtual;  /** The amount of virtual memory currently available, in bytes. */			SwapFree
@@ -806,16 +838,20 @@ FPlatformMemoryStats FUnixPlatformMemory::GetStats()
 {
 	uint64 MemFree = 0;
 	uint64 Cached = 0;
+	uint64 MemTotal = 0;
+	uint64 SwapTotal = 0;
+	uint64 CommittedAS = 0;
 	FPlatformMemoryStats MemoryStats;
 
 	FProcField SMapsFields[] =
 	{
 		// An estimate of how much memory is available for starting new applications, without swapping.
 		{ "MemAvailable:", &MemoryStats.AvailablePhysical },
-		// Amount of swap space that is currently unused.
-		{ "SwapFree:",     &MemoryStats.AvailableVirtual },
 		{ "MemFree:",      &MemFree },
 		{ "Cached:",       &Cached },
+		{ "MemTotal:",     &MemTotal },
+		{ "SwapTotal:",    &SwapTotal },
+		{ "Committed_AS:", &CommittedAS },
 	};
 	if (ReadProcFields("/proc/meminfo", SMapsFields, UE_ARRAY_COUNT(SMapsFields)))
 	{
@@ -824,6 +860,13 @@ FPlatformMemoryStats FUnixPlatformMemory::GetStats()
 		{
 			MemoryStats.AvailablePhysical = FMath::Min(MemFree + Cached, MemoryStats.TotalPhysical);
 		}
+
+		// OS alloted percentage of physical RAM (used here as a ratio) to overcommit
+		// https://www.kernel.org/doc/Documentation/vm/overcommit-accounting
+		float OvercommitRatio = ReadOvercommitRatio();
+
+		// Total memory * commit percentage (which can be greater than 100%) - committed memory
+		MemoryStats.AvailableVirtual = static_cast<uint64>(static_cast<float>(MemTotal + SwapTotal)	* (1.0f + OvercommitRatio)) - CommittedAS;
 	}
 
 	FProcField SMapsFields2[] =
