Don't read the source buffer in iostore if we are just copying the chunk from the reference database.

diff --git a/Engine/Source/Developer/IoStoreUtilities/Private/IoStoreUtilities.cpp b/Engine/Source/Developer/IoStoreUtilities/Private/IoStoreUtilities.cpp
--- a/Engine/Source/Developer/IoStoreUtilities/Private/IoStoreUtilities.cpp
+++ b/Engine/Source/Developer/IoStoreUtilities/Private/IoStoreUtilities.cpp
@@ -107,7 +107,7 @@ public:
 	int32 ContainerNotFound = 0;
 	int64 FulfillBytes = 0;
 	int64 FulfillBytesPerChunk[(int8)EIoChunkType::MAX] = {};
-	
+	uint32 CompressionBlockSize = 0;
 
 	bool Init(const FString& InGlobalContainerFileName, const FKeyChain& InDecryptionKeychain)
 	{
@@ -131,6 +131,7 @@ public:
 			ContainerFilePaths.Emplace(Directory / Filename);
 		}
 
+		CompressionBlockSize = 0;
 		int64 IoChunkCount = 0;		
 		for (const FString& ContainerFilePath : ContainerFilePaths)
 		{
@@ -148,6 +149,16 @@ public:
 				return true;
 			});
 
+			if (Readers.Num() == 0)
+			{
+				CompressionBlockSize = Reader->GetCompressionBlockSize();
+			}
+			else if (Reader->GetCompressionBlockSize() != CompressionBlockSize)
+			{
+				UE_LOG(LogIoStore, Error, TEXT("Reference chunk containers had different compression block sizes, failing to init reference db (%u and %u)"), CompressionBlockSize, Reader->GetCompressionBlockSize());
+				return false;
+			}
+
 			IoChunkCount += ReaderChunks.Chunks.Num();
 			ReaderChunks.ReaderIndex = Readers.Num();
 			Readers.Add(MoveTemp(Reader));
@@ -157,8 +168,36 @@ public:
 		return true;
 	}
 
+	virtual uint32 GetCompressionBlockSize() const override
+	{
+		return CompressionBlockSize;
+	}
+
+	virtual bool ChunkExists(const TPair<FIoContainerId, FIoChunkHash>& InChunkKey, uint32& OutNumChunkBlocks)
+	{
+		FReaderChunks* ReaderChunks = ChunkDatabase.Find(InChunkKey.Key);
+		if (ReaderChunks == nullptr)
+		{
+			// Container doesn't exist - likely provided the path to a different project. Mark this
+			// error as happening once so we can log at the end.
+			ContainerNotFound++;
+			return false;
+		}
+
+		FIoChunkId* ChunkId = ReaderChunks->Chunks.Find(InChunkKey.Value);
+		if (ChunkId == nullptr)
+		{
+			// No exact chunk data match - this is a normal exit condition for a changed block.
+			return false;
+		}
+
+		OutNumChunkBlocks = Readers[ReaderChunks->ReaderIndex]->GetChunkInfo(*ChunkId).ValueOrDie().NumCompressedBlocks;
+		return true;
+	}
+
+
 	// Not thread safe, called from the BeginCompress dispatch thread.
-	virtual bool RetrieveChunk(const TPair<FIoContainerId, FIoChunkHash>& InChunkKey, const FName& InCompressionMethod, uint64 InUncompressedSize, uint64 InNumChunkBlocks, TUniqueFunction<void(TIoStatusOr<FIoStoreCompressedReadResult>)> InCompleteCallback)
+	virtual bool RetrieveChunk(const TPair<FIoContainerId, FIoChunkHash>& InChunkKey, TUniqueFunction<void(TIoStatusOr<FIoStoreCompressedReadResult>)> InCompleteCallback)
 	{
 		RequestCount++;
 
@@ -189,18 +228,6 @@ public:
 			return true;
 		});
 
-		if (TotalUncompressedSize != InUncompressedSize)
-		{
-			// Shocked if this happens - hash match with different data!
-			return false;
-		}
-
-		if (CompressedBlockCount != InNumChunkBlocks)
-		{
-			// Different CompressionBlockSize between the builds
-			return false;
-		}
-
 		FulfillBytesPerChunk[(int8)ChunkId->GetChunkType()] += TotalCompressedSize;
 		FulfillBytes += TotalCompressedSize;
 		FulfillCount++;
@@ -2121,6 +2148,8 @@ static void ProcessShaderLibraries(const FIoStoreArguments& Arguments, TArray<FC
 					}
 					TargetFile.ChunkId = LibraryChunk.Key;
 					TargetFile.SourceBuffer.Emplace(LibraryChunk.Value);
+					TargetFile.SourceSize = LibraryChunk.Value.GetSize();
+
 					const bool bIsGlobalShaderLibrary = FPaths::GetCleanFilename(TargetFile.NormalizedSourcePath).StartsWith(TEXT("ShaderArchive-Global-"));
 					const FShaderInfo::EShaderType ShaderType = bIsGlobalShaderLibrary ? FShaderInfo::Global : FShaderInfo::Normal;
 					for (const TTuple<FIoChunkId, FIoBuffer, uint32>& CodeChunk : CodeChunks)
@@ -2822,6 +2851,10 @@ private:
 		FIoStoreWriteRequestManager& Manager;
 		const FContainerTargetFile& TargetFile;
 		TArray<FFileRegion> FileRegions;
+
+		// Note -- this is filled with the TargetFile.SourceSize value which is the size of the buffer 
+		// used for IO, however it's not necessarily the size of the resulting input to iostore as
+		// the buffer can be post-processed after i/o (e.g. CreateOptimizedPackage).
 		uint64 SourceBufferSize;
 		FGraphEventRef CompletionEvent;
 		FIoBuffer SourceBuffer;
diff --git a/Engine/Source/Runtime/Core/Private/IO/IoStore.cpp b/Engine/Source/Runtime/Core/Private/IO/IoStore.cpp
--- a/Engine/Source/Runtime/Core/Private/IO/IoStore.cpp
+++ b/Engine/Source/Runtime/Core/Private/IO/IoStore.cpp
@@ -184,7 +184,10 @@ struct FIoStoreWriteQueueEntry
 	FIoChunkId ChunkId;
 	FIoChunkHash ChunkHash;
 	uint64 Sequence = 0;
-	uint64 UncompressedSize = 0;
+	
+	// We make this optional because at the latest it might not be valid until FinishCompressionBarrior
+	// completes and we'd like to have a check() on that.
+	TOptional<uint64> UncompressedSize = 0;
 	uint64 CompressedSize = 0;
 	uint64 Padding = 0;
 	uint64 Offset = 0;
@@ -204,6 +207,11 @@ struct FIoStoreWriteQueueEntry
 	bool bAdded = false;
 	bool bModified = false;
 	bool bStoreCompressedDataInDDC = false;
+	
+	bool bLoadingFromReferenceDb = false;
+	// When we know we're loading from the reference chunk db, we don't read the source
+	// buffer but we still need to know the number of chunks which we get from the refdb.
+	uint32 NumChunkBlocksFromRefDb = 0; 
 };
 
 class FIoStoreWriteQueue
@@ -393,7 +401,15 @@ public:
 	void ScheduleCompression(FIoStoreWriteQueueEntry* QueueEntry)
 	{
 		BeginCompressionQueue.Enqueue(QueueEntry);
-		QueueEntry->Request->PrepareSourceBufferAsync(QueueEntry->BeginCompressionBarrier);
+		if (QueueEntry->bLoadingFromReferenceDb == false)
+		{
+			QueueEntry->Request->PrepareSourceBufferAsync(QueueEntry->BeginCompressionBarrier);
+		}
+		else
+		{
+			// We don't need to wait on a read so we can kick directly.
+			QueueEntry->BeginCompressionBarrier->DispatchSubsequents();
+		}
 	}
 
 	//
@@ -694,6 +710,12 @@ public:
 
 	void SetReferenceChunkDatabase(TSharedPtr<IIoStoreWriterReferenceChunkDatabase> InReferenceChunkDatabase)
 	{
+		if (InReferenceChunkDatabase.IsValid() && InReferenceChunkDatabase->GetCompressionBlockSize() != WriterContext->GetSettings().CompressionBlockSize)
+		{
+			UE_LOG(LogIoStore, Warning, TEXT("Reference chunk database has a different compression block size than the current writer!"));
+			UE_LOG(LogIoStore, Warning, TEXT("No chunks will match, so ignoring. ReferenceChunkDb: %d, IoStoreWriter: %d"), InReferenceChunkDatabase->GetCompressionBlockSize(), WriterContext->GetSettings().CompressionBlockSize);
+			return;
+		}
 		ReferenceChunkDatabase = InReferenceChunkDatabase;
 	}
 	void SetHashDatabase(TSharedPtr<IIoStoreWriterHashDatabase> InHashDatabase, bool bInVerifyHashDatabase)
@@ -823,6 +845,12 @@ public:
 				WriterContext->HashDbChunksCount.IncrementExchange();
 				WriterContext->HashDbChunksByType[(int8)Entry->ChunkId.GetChunkType()].IncrementExchange();
 				WriterContext->HashedChunksCount.IncrementExchange();
+
+				if (ReferenceChunkDatabase.IsValid() && CompressionMethodForEntry(Entry) != NAME_None)
+				{
+					TPair<FIoContainerId, FIoChunkHash> ChunkKey(ContainerSettings.ContainerId, Entry->ChunkHash);
+					Entry->bLoadingFromReferenceDb = ReferenceChunkDatabase->ChunkExists(ChunkKey, Entry->NumChunkBlocksFromRefDb);
+				}
 				return;
 			}
 
@@ -852,6 +880,13 @@ public:
 				}
 			}
 
+
+			if (ReferenceChunkDatabase.IsValid() && CompressionMethodForEntry(Entry) != NAME_None)
+			{
+				TPair<FIoContainerId, FIoChunkHash> ChunkKey(ContainerSettings.ContainerId, Entry->ChunkHash);
+				Entry->bLoadingFromReferenceDb = ReferenceChunkDatabase->ChunkExists(ChunkKey, Entry->NumChunkBlocksFromRefDb);
+			}
+
 			// Release the source data buffer, it will be reloaded later when we start compressing the chunk
 			Entry->Request->FreeSourceBuffer();
 		}, TStatId(), &HashPrereqs, ENamedThreads::AnyHiPriThreadHiPriTask);
@@ -1478,64 +1513,57 @@ private:
 		return !Ar.IsError();
 	}
 
-	void BeginCompress(FIoStoreWriteQueueEntry* Entry)
+	FName CompressionMethodForEntry(FIoStoreWriteQueueEntry* Entry) const
 	{
-		WriterContext->BeginCompressChunksByType[(int8)Entry->ChunkId.GetChunkType()].IncrementExchange();
-
 		FName CompressionMethod = NAME_None;
 		const FIoStoreWriterSettings& WriterSettings = WriterContext->WriterSettings;
 		if (ContainerSettings.IsCompressed() && !Entry->Options.bForceUncompressed && !Entry->Options.bIsMemoryMapped)
 		{
 			CompressionMethod = WriterSettings.CompressionMethod;
 		}
+		return CompressionMethod;
+	}
 
-		const FIoBuffer* SourceBuffer = Entry->Request->GetSourceBuffer();
-		Entry->UncompressedSize = SourceBuffer->DataSize();
 
-		check(WriterSettings.CompressionBlockSize > 0);
-		const uint64 NumChunkBlocks = Align(Entry->UncompressedSize, WriterSettings.CompressionBlockSize) / WriterSettings.CompressionBlockSize;
-		if (NumChunkBlocks == 0)
-		{
-			Entry->FinishCompressionBarrier->DispatchSubsequents();
-			return;
-		}
+
+	void BeginCompress(FIoStoreWriteQueueEntry* Entry)
+	{
+		WriterContext->BeginCompressChunksByType[(int8)Entry->ChunkId.GetChunkType()].IncrementExchange();
 		
-		Entry->ChunkBlocks.SetNum(int32(NumChunkBlocks));
+		if (Entry->bLoadingFromReferenceDb)
 		{
-			// We must allocate resources for our tasks up front to prevent resource deadlock.
-			// Note that for Reference Chunk loaded blocks this will reserve more than actually
-			// needed.
-			uint64 BytesToProcess = Entry->UncompressedSize;
-			const uint8* UncompressedData = SourceBuffer->Data();
-			for (int32 BlockIndex = 0; BlockIndex < NumChunkBlocks; ++BlockIndex)
+			if (Entry->NumChunkBlocksFromRefDb == 0)
+			{
+				Entry->FinishCompressionBarrier->DispatchSubsequents();
+				return;
+			}
+
+			// Allocate resources before launching the read tasks to prevent deadlock. Note this will
+			// allocate iobuffers big enough for uncompressed size, when we only actually need it for
+			// compressed size.
+			int32 NumChunkBlocks32 = IntCastChecked<int32>(Entry->NumChunkBlocksFromRefDb);
+			Entry->ChunkBlocks.SetNum(NumChunkBlocks32);
+			for (int32 BlockIndex = 0; BlockIndex < NumChunkBlocks32; ++BlockIndex)
 			{
 				FChunkBlock& Block = Entry->ChunkBlocks[BlockIndex];
-				Block.IoBuffer = WriterContext->AllocCompressionBuffer(int32(NumChunkBlocks));
-				Block.CompressionMethod = CompressionMethod;
-				Block.UncompressedSize = FMath::Min(BytesToProcess, WriterSettings.CompressionBlockSize);
-				Block.UncompressedData = UncompressedData;
-				BytesToProcess -= Block.UncompressedSize;
-				UncompressedData += Block.UncompressedSize;
+				Block.IoBuffer = WriterContext->AllocCompressionBuffer(NumChunkBlocks32);
+				// Everything else in a block gets filled out from the refdb.
 			}
-		}
 
-		//
-		// Check if this chunk exists in the reference cache.
-		//
-		if (ReferenceChunkDatabase.IsValid() && CompressionMethod != NAME_None)
-		{
 			TPair<FIoContainerId, FIoChunkHash> ChunkKey(ContainerSettings.ContainerId, Entry->ChunkHash);
 
 			// Valid chunks must create the same decompressed bits, but can have different compressed bits.
 			// Since we are on a lightweight dispatch thread, the actual read is async, as is the processing
 			// of the results.
-			bool GotChunk = ReferenceChunkDatabase->RetrieveChunk(ChunkKey, CompressionMethod, Entry->UncompressedSize, NumChunkBlocks, [this, Entry](TIoStatusOr<FIoStoreCompressedReadResult> InReadResult)
+			bool bChunkExists = ReferenceChunkDatabase->RetrieveChunk(ChunkKey, [this, Entry](TIoStatusOr<FIoStoreCompressedReadResult> InReadResult)
 			{
+
 				// If we fail here, in order to recover we effectively need to re-kick this chunk's
-				// BeginCompress()... however, this is just a direct read and should only fail
+				// BeginCompress() as well as source buffer read... however, this is just a direct read and should only fail
 				// in catastrophic scenarios (loss of connection on a network drive?).
 				FIoStoreCompressedReadResult ReadResult = InReadResult.ValueOrDie();
 
+				uint64 TotalUncompressedSize = 0;
 				uint8* ReferenceData = ReadResult.IoBuffer.GetData();
 				for (int32 BlockIndex = 0; BlockIndex < ReadResult.Blocks.Num(); ++BlockIndex)
 				{
@@ -1543,6 +1571,8 @@ private:
 					FChunkBlock& Block = Entry->ChunkBlocks[BlockIndex];
 					Block.CompressionMethod = ReferenceBlock.CompressionMethod;
 					Block.CompressedSize = ReferenceBlock.CompressedSize;
+					Block.UncompressedSize = ReferenceBlock.UncompressedSize;
+					TotalUncompressedSize += ReferenceBlock.UncompressedSize;
 
 					// Future optimization: ReadCompressed returns the memory ready to encrypt in one
 					// large contiguous buffer (i.e. padded). We could use the FIoBuffer functionality of referencing a 
@@ -1557,15 +1587,49 @@ private:
 					ReferenceData += ReferenceBlock.AlignedSize;
 				}
 
+				Entry->UncompressedSize.Emplace(TotalUncompressedSize);
 				Entry->FinishCompressionBarrier->DispatchSubsequents();
 			});
 
-			if (GotChunk)
+			check(bChunkExists); // Sanity - should never return false as we can only get here if ChunkExists() returns true.
+
+			WriterContext->RefDbChunksCount.IncrementExchange();
+			WriterContext->RefDbChunksByType[(int8)Entry->ChunkId.GetChunkType()].IncrementExchange();
+			// Lambda handles dispatch subsequents
+			return;
+		}
+
+		const FIoStoreWriterSettings& WriterSettings = WriterContext->WriterSettings;
+
+		FName CompressionMethod = CompressionMethodForEntry(Entry);
+
+		const FIoBuffer* SourceBuffer = Entry->Request->GetSourceBuffer();
+		Entry->UncompressedSize.Emplace(SourceBuffer->DataSize());
+
+		check(WriterSettings.CompressionBlockSize > 0);
+		const uint64 NumChunkBlocks64 = Align(Entry->UncompressedSize.GetValue(), WriterSettings.CompressionBlockSize) / WriterSettings.CompressionBlockSize;
+		if (NumChunkBlocks64 == 0)
+		{
+			Entry->FinishCompressionBarrier->DispatchSubsequents();
+			return;
+		}
+
+		int32 NumChunkBlocks = IntCastChecked<int32>(NumChunkBlocks64);
+		
+		Entry->ChunkBlocks.SetNum(NumChunkBlocks);
+		{
+			// We must allocate resources for our tasks up front to prevent resource deadlock.
+			uint64 BytesToProcess = Entry->UncompressedSize.GetValue();
+			const uint8* UncompressedData = SourceBuffer->Data();
+			for (int32 BlockIndex = 0; BlockIndex < NumChunkBlocks; ++BlockIndex)
 			{
-				WriterContext->RefDbChunksCount.IncrementExchange();
-				WriterContext->RefDbChunksByType[(int8)Entry->ChunkId.GetChunkType()].IncrementExchange();
-				// Lambda handles dispatch subsequents
-				return;
+				FChunkBlock& Block = Entry->ChunkBlocks[BlockIndex];
+				Block.IoBuffer = WriterContext->AllocCompressionBuffer(NumChunkBlocks);
+				Block.CompressionMethod = CompressionMethod;
+				Block.UncompressedSize = FMath::Min(BytesToProcess, WriterSettings.CompressionBlockSize);
+				Block.UncompressedData = UncompressedData;
+				BytesToProcess -= Block.UncompressedSize;
+				UncompressedData += Block.UncompressedSize;
 			}
 		}
 
@@ -1584,7 +1648,7 @@ private:
 		bool bUseDDCCompression =
 			WriterContext->DDC &&
 			WriterSettings.bCompressionEnableDDC &&
-			Entry->UncompressedSize > WriterSettings.CompressionMinSizeToConsiderDDC;
+			Entry->UncompressedSize.GetValue() > WriterSettings.CompressionMinSizeToConsiderDDC;
 		if (bUseDDCCompression)
 		{
 			TStringBuilder<256> CacheKeySuffix;
@@ -1767,7 +1831,7 @@ private:
 		
 		FIoOffsetAndLength OffsetLength;
 		OffsetLength.SetOffset(UncompressedFileOffset);
-		OffsetLength.SetLength(Entry->UncompressedSize);
+		OffsetLength.SetLength(Entry->UncompressedSize.GetValue());
 
 		FIoStoreTocEntryMeta ChunkMeta{ Entry->ChunkHash, FIoStoreTocEntryMetaFlags::None };
 		if (Entry->Options.bIsMemoryMapped)
@@ -1807,8 +1871,8 @@ private:
 
 		const uint64 RegionStartOffset = TargetPartition->Offset;
 		TargetPartition->Offset += Entry->CompressedSize;
-		UncompressedFileOffset += Align(Entry->UncompressedSize, WriterSettings.CompressionBlockSize);
-		TotalEntryUncompressedSize += Entry->UncompressedSize;
+		UncompressedFileOffset += Align(Entry->UncompressedSize.GetValue(), WriterSettings.CompressionBlockSize);
+		TotalEntryUncompressedSize += Entry->UncompressedSize.GetValue();
 
 		if (WriterSettings.bEnableFileRegions)
 		{
@@ -2001,7 +2065,10 @@ void FIoStoreWriterContextImpl::BeginEncryptionAndSigningThreadFunc()
 		{
 			FIoStoreWriteQueueEntry* Next = Entry->Next;
 			Entry->FinishCompressionBarrier->Wait();
-			Entry->Request->FreeSourceBuffer();
+			if (Entry->bLoadingFromReferenceDb == false)
+			{
+				Entry->Request->FreeSourceBuffer();
+			}
 			if (Entry->bStoreCompressedDataInDDC)
 			{
 				TRACE_CPUPROFILER_EVENT_SCOPE(StoreInDDC);
@@ -2675,7 +2742,7 @@ public:
 		return BlockIndex >= FirstBlockIndex && BlockIndex <= LastBlockIndex;
 	}
 
-	uint64 GetCompressionBlockSize() const
+	uint32 GetCompressionBlockSize() const
 	{
 		return Toc.GetTocResource().Header.CompressionBlockSize;
 	}
@@ -2819,7 +2886,7 @@ const FIoDirectoryIndexReader& FIoStoreReader::GetDirectoryIndexReader() const
 	return Impl->GetDirectoryIndexReader();
 }
 
-uint64 FIoStoreReader::GetCompressionBlockSize() const
+uint32 FIoStoreReader::GetCompressionBlockSize() const
 {
 	return Impl->GetCompressionBlockSize();
 }
diff --git a/Engine/Source/Runtime/Core/Public/IO/IoDispatcher.h b/Engine/Source/Runtime/Core/Public/IO/IoDispatcher.h
--- a/Engine/Source/Runtime/Core/Public/IO/IoDispatcher.h
+++ b/Engine/Source/Runtime/Core/Public/IO/IoDispatcher.h
@@ -65,6 +65,7 @@ class IMappedFileRegion;
 struct FFileRegion;
 struct IIoDispatcherBackend;
 template <typename CharType> class TStringBuilderBase;
+template <typename OptionalType> struct TOptional;
 
 CORE_API DECLARE_LOG_CATEGORY_EXTERN(LogIoDispatcher, Log, All);
 
@@ -562,10 +563,16 @@ class IIoStoreWriteRequest
 {
 public:
 	virtual ~IIoStoreWriteRequest() = default;
+
+	// Launches any async operations necessary in order to access the buffer. CompletionEvent is set once it's ready, which may be immediate.
 	virtual void PrepareSourceBufferAsync(FGraphEventRef CompletionEvent) = 0;
 	virtual uint64 GetOrderHint() = 0;
 	virtual TArrayView<const FFileRegion> GetRegions() = 0;
+
+	// Only valid after the completion event passed to PrepareSourceBufferAsync has fired.
 	virtual const FIoBuffer* GetSourceBuffer() = 0;
+
+	// Can't be called between PrepareSourceBufferAsync and its completion!
 	virtual void FreeSourceBuffer() = 0;
 };
 
@@ -632,12 +639,12 @@ public:
 	* Used by IIoStoreWriter to check and see if there's a reference chunk that matches the data that
 	* IoStoreWriter wants to compress and write. Validity checks must be synchronous - if a chunk can't be
 	* used for some reason (no matching chunk exists or otherwise), this function must return false and not 
-	* call InCompletionCallback. The input parameters are all for match finding and validating.
+	* call InCompletionCallback.
 	* 
 	* Once a matching chunk is found, it is read from the source iostore container asynchronously, and upon
 	* completion InCompletionCallback is called with the raw output from FIoStoreReader::ReadCompressed (i.e.
-	* FIoStoreCompressedReadResult). Failures once the async read process has started are currently fatal due to needing
-	* to rekick the BeginCompress for the chunk.
+	* FIoStoreCompressedReadResult). Failures once the async read process has started are currently fatal due to
+	* difficulties in rekicking a read.
 	* 
 	* For the moment, changes in compression method are allowed.
 	* 
@@ -646,7 +653,19 @@ public:
 	* Chunks provided *MUST* decompress to bits that hash to the exact value provided in InChunkKey (i.e. be exactly the same bits),
 	* and also be the same number of blocks (i.e. same CompressionBlockSize)
 	*/
-	CORE_API virtual bool RetrieveChunk(const TPair<FIoContainerId, FIoChunkHash>& InChunkKey, const FName& InCompressionMethod, uint64 InUncompressedSize, uint64 InNumChunkBlocks, TUniqueFunction<void(TIoStatusOr<FIoStoreCompressedReadResult>)> InCompletionCallback) = 0;
+	CORE_API virtual bool RetrieveChunk(const TPair<FIoContainerId, FIoChunkHash>& InChunkKey, TUniqueFunction<void(TIoStatusOr<FIoStoreCompressedReadResult>)> InCompletionCallback) = 0;
+
+	/* 
+	* Quick synchronous existence check that returns the number of blocks for the chunk. This is used to set up
+	* the necessary structures without needing to read the source data for the chunk.
+	*/
+	CORE_API virtual bool ChunkExists(const TPair<FIoContainerId, FIoChunkHash>& InChunkKey, uint32& OutNumChunkBlocks) = 0;
+
+	/*
+	* Returns the compression block size that was used to break up the IoChunks in the source containers. If this is different than what we want, 
+	* then none of the chunks will ever match. Knowing this up front allows us to only match on hash
+	*/
+	CORE_API virtual uint32 GetCompressionBlockSize() const = 0;
 };
 
 /**
@@ -713,7 +732,7 @@ public:
 	CORE_API void GetFilenamesByBlockIndex(const TArray<int32>& InBlockIndexList, TArray<FString>& OutFileList) const;
 	CORE_API void GetFilenames(TArray<FString>& OutFileList) const;
 
-	CORE_API uint64 GetCompressionBlockSize() const;
+	CORE_API uint32 GetCompressionBlockSize() const;
 	CORE_API const TArray<FName>& GetCompressionMethods() const;
 	CORE_API void EnumerateCompressedBlocks(TFunction<bool(const FIoStoreTocCompressedBlockInfo&)>&& Callback) const;
 	CORE_API void EnumerateCompressedBlocksForChunk(const FIoChunkId& Chunk, TFunction<bool(const FIoStoreTocCompressedBlockInfo&)>&& Callback) const;
