Added the Localization.DumpLiveTable command to dump the current live table state to the log

This can be optionally filtered using any of the following wildcard arguments:
* Namespace - The namespace to filter against
* Key - The key to filter against
* DisplayString - The display string to filter against

Examples:
* Localization.DumpLiveTable -Key=Foo
    * Matches any live table entries with the key "Foo"
* Localization.DumpLiveTable -DisplayString="This is some text"
    * Matches any live table entries with the display string "This is some text"
* Localization.DumpLiveTable -Key=Bar* -DisplayString="This is some other text"
    * Matches any live table entries with a key that starts with "Bar" and with the display string "This is some other text"

diff --git a/Engine/Source/Runtime/Core/Private/Internationalization/TextLocalizationManager.cpp b/Engine/Source/Runtime/Core/Private/Internationalization/TextLocalizationManager.cpp
--- a/Engine/Source/Runtime/Core/Private/Internationalization/TextLocalizationManager.cpp
+++ b/Engine/Source/Runtime/Core/Private/Internationalization/TextLocalizationManager.cpp
@@ -50,6 +50,41 @@ static FAutoConsoleVariableRef CVarAsyncLoadLocalizationData(TEXT("Localization.
  */
 static bool AsyncLoadLocalizationDataOnLanguageChange = false;
 static FAutoConsoleVariableRef CVarAsyncLoadLocalizationDataOnLanguageChange(TEXT("Localization.AsyncLoadLocalizationDataOnLanguageChange"), AsyncLoadLocalizationDataOnLanguageChange, TEXT("True to load localization data asynchronously (non-blocking) when the language changes, or False to load it synchronously (blocking)"));
+
+#if ENABLE_LOC_TESTING
+static FAutoConsoleCommand CmdDumpLiveTable(
+	TEXT("Localization.DumpLiveTable"), 
+	TEXT("Dumps the current live table state to the log, optionally filtering it based on wildcard arguments for 'Namespace', 'Key', or 'DisplayString', eg) -Key=Foo, or -DisplayString=\"This is some text\", or -Key=Bar*Baz -DisplayString=\"This is some other text\""), 
+	FConsoleCommandWithArgsDelegate::CreateStatic([](const TArray<FString>& Args)
+	{
+		auto ParseOptionalStringArg = [](const TCHAR* Arg, const TCHAR* TokenName, TOptional<FString>& OutResult)
+		{
+			FString TmpResult;
+			if (FParse::Value(Arg, TokenName, TmpResult))
+			{
+				OutResult = MoveTemp(TmpResult);
+				return true;
+			}
+			return false;
+		};
+
+		TOptional<FString> NamespaceFilter;
+		TOptional<FString> KeyFilter;
+		TOptional<FString> DisplayStringFilter;
+
+		for (const FString& Arg : Args)
+		{
+			if (!ParseOptionalStringArg(*Arg, TEXT("Namespace="), NamespaceFilter) &&
+				!ParseOptionalStringArg(*Arg, TEXT("Key="), KeyFilter) &&
+				!ParseOptionalStringArg(*Arg, TEXT("DisplayString="), DisplayStringFilter))
+			{
+				UE_LOG(LogLocalization, Warning, TEXT("Unknown argument '%s' passed to Localization.DumpLiveTable!"), *Arg);
+			}
+		}
+
+		FTextLocalizationManager::Get().DumpLiveTable(NamespaceFilter.GetPtrOrNull(), KeyFilter.GetPtrOrNull(), DisplayStringFilter.GetPtrOrNull());
+	}));
+#endif
 }
 
 enum class ERequestedCultureOverrideLevel : uint8
@@ -531,6 +566,31 @@ void FTextLocalizationManager::CompactDataStructures()
 	UE_LOG(LogTextLocalizationManager, Log, TEXT("Compacting localization data took %6.2fms"), 1000.0 * (FPlatformTime::Seconds() - StartTime));
 }
 
+#if ENABLE_LOC_TESTING
+void FTextLocalizationManager::DumpLiveTable(const FString* NamespaceFilter, const FString* KeyFilter, const FString* DisplayStringFilter) const
+{
+	auto PassesFilter = [](const FString& Str, const FString* Filter)
+	{
+		return !Filter || Str.MatchesWildcard(*Filter, ESearchCase::IgnoreCase); // Note: This is case insensitive since its used from a debug command
+	};
+
+	UE_LOG(LogLocalization, Display, TEXT("----------------------------------------------------------------------"));
+
+	FScopeLock ScopeLock(&DisplayStringLookupTableCS);
+	for (const auto& DisplayStringPair : DisplayStringLookupTable)
+	{
+		if (PassesFilter(DisplayStringPair.Key.GetNamespace().GetChars(), NamespaceFilter) &&
+			PassesFilter(DisplayStringPair.Key.GetKey().GetChars(), KeyFilter) &&
+			PassesFilter(**DisplayStringPair.Value.DisplayString, DisplayStringFilter))
+		{
+			UE_LOG(LogLocalization, Display, TEXT("LiveTableEntry: Namespace: '%s', Key: '%s', DisplayString: '%s'"), DisplayStringPair.Key.GetNamespace().GetChars(), DisplayStringPair.Key.GetKey().GetChars(), **DisplayStringPair.Value.DisplayString);
+		}
+	}
+
+	UE_LOG(LogLocalization, Display, TEXT("----------------------------------------------------------------------"));
+}
+#endif
+
 FString FTextLocalizationManager::GetRequestedLanguageName() const
 {
 	ERequestedCultureOverrideLevel LanguageOverrideLevel = ERequestedCultureOverrideLevel::Defaults;
diff --git a/Engine/Source/Runtime/Core/Public/Internationalization/TextLocalizationManager.h b/Engine/Source/Runtime/Core/Public/Internationalization/TextLocalizationManager.h
--- a/Engine/Source/Runtime/Core/Public/Internationalization/TextLocalizationManager.h
+++ b/Engine/Source/Runtime/Core/Public/Internationalization/TextLocalizationManager.h
@@ -115,6 +115,14 @@ public:
 	void DumpMemoryInfo() const;
 	void CompactDataStructures();
 
+#if ENABLE_LOC_TESTING
+	/**
+	 * Dumps the current live table state to the log, optionally filtering it based on the given wildcard arguments.
+	 * @note Calling this function with no filters specified will dump the entire live table.
+	 */
+	void DumpLiveTable(const FString* NamespaceFilter = nullptr, const FString* KeyFilter = nullptr, const FString* DisplayStringFilter = nullptr) const;
+#endif
+
 	/**
 	 * Get the language that will be requested during localization initialization, based on the hierarchy of: command line -> configs -> OS default.
 	 */
