Add stats to AutoRTFM.

diff --git a/Engine/Source/Runtime/Core/Private/AutoRTFM/Context.cpp b/Engine/Source/Runtime/Core/Private/AutoRTFM/Context.cpp
--- a/Engine/Source/Runtime/Core/Private/AutoRTFM/Context.cpp
+++ b/Engine/Source/Runtime/Core/Private/AutoRTFM/Context.cpp
@@ -7,6 +7,7 @@
 #include "FunctionMap.h"
 #include "GlobalData.h"
 #include "ScopedGuard.h"
+#include "Stats.h"
 #include "TransactionInlines.h"
 
 #include "Templates/UniquePtr.h"
@@ -189,6 +190,9 @@ void FContext::PushTransaction(FTransaction* NewTransaction)
 	
 	NewTransaction->SetParent(CurrentTransaction);
 	CurrentTransaction = NewTransaction;
+
+	// Collect stats that we've got a new transaction.
+	Stats.Collect<EStatsKind::Transaction>();
 }
 
 void FContext::PopTransaction()
diff --git a/Engine/Source/Runtime/Core/Private/AutoRTFM/Stats.cpp b/Engine/Source/Runtime/Core/Private/AutoRTFM/Stats.cpp
new file mode 100644
--- /dev/null
+++ b/Engine/Source/Runtime/Core/Private/AutoRTFM/Stats.cpp
@@ -0,0 +1,131 @@
+// Copyright Epic Games, Inc. All Rights Reserved.
+
+#if (defined(__AUTORTFM) && __AUTORTFM)
+#include "Stats.h"
+
+using namespace AutoRTFM;
+
+template<> void AutoRTFM::FStats::Report<EStatsKind::Transaction>(const uint64_t Data) const
+{
+	UE_LOG(LogAutoRTFM, Display, TEXT("  Total transactions:        %11u"), Data);
+}
+
+template<> void AutoRTFM::FStats::Report<EStatsKind::Commit>(const uint64_t Data) const
+{
+	UE_LOG(LogAutoRTFM, Display, TEXT("  Total commits:             %11u"), Data);
+}
+
+template<> void AutoRTFM::FStats::Report<EStatsKind::Abort>(const uint64_t Data) const
+{
+	UE_LOG(LogAutoRTFM, Display, TEXT("  Total aborts:              %11u"), Data);
+}
+
+template<> void AutoRTFM::FStats::Report<EStatsKind::AverageTransactionDepth>(const uint64_t Data) const
+{
+	const uint64_t TotalTransactions = Datas[static_cast<size_t>(EStatsKind::Transaction)];
+	UE_LOG(LogAutoRTFM, Display, TEXT("  Average transaction depth: %15.3f"), (static_cast<double>(Data) / static_cast<double>(TotalTransactions)));
+}
+
+template<> void AutoRTFM::FStats::Report<EStatsKind::MaximumTransactionDepth>(const uint64_t Data) const
+{
+    UE_LOG(LogAutoRTFM, Display, TEXT("  Maximum transaction depth: %11u"), Data);
+}
+
+template<> void AutoRTFM::FStats::Report<EStatsKind::AverageWriteLogEntries>(const uint64_t Data) const
+{
+    const uint64_t TotalTransactions = Datas[static_cast<size_t>(EStatsKind::Transaction)];
+    UE_LOG(LogAutoRTFM, Display, TEXT("  Average write log entries: %15.3f"), (static_cast<double>(Data) / static_cast<double>(TotalTransactions)));
+}
+
+template<> void AutoRTFM::FStats::Report<EStatsKind::MaximumWriteLogEntries>(const uint64_t Data) const
+{
+    UE_LOG(LogAutoRTFM, Display, TEXT("  Maximum write log entries: %11u"), Data);
+}
+
+template<> void AutoRTFM::FStats::Report<EStatsKind::AverageWriteLogBytes>(const uint64_t Data) const
+{
+    const uint64_t TotalTransactions = Datas[static_cast<size_t>(EStatsKind::Transaction)];
+    UE_LOG(LogAutoRTFM, Display, TEXT("  Average write log bytes:   %15.3f"), (static_cast<double>(Data) / static_cast<double>(TotalTransactions)));
+}
+
+template<> void AutoRTFM::FStats::Report<EStatsKind::MaximumWriteLogBytes>(const uint64_t Data) const
+{
+    UE_LOG(LogAutoRTFM, Display, TEXT("  Maximum write log bytes:   %11u"), Data);
+}
+
+template<> void AutoRTFM::FStats::Report<EStatsKind::HitSetHit>(const uint64_t Data) const
+{
+    UE_LOG(LogAutoRTFM, Display, TEXT("  HitSet hits:               %11u"), Data);
+}
+
+template<> void AutoRTFM::FStats::Report<EStatsKind::HitSetMiss>(const uint64_t Data) const
+{
+    UE_LOG(LogAutoRTFM, Display, TEXT("  HitSet misses:             %11u"), Data);
+}
+
+template<> void AutoRTFM::FStats::Report<EStatsKind::HitSetSkippedBecauseOfStackLocalMemory>(const uint64_t Data) const
+{
+    UE_LOG(LogAutoRTFM, Display, TEXT("  HitSet skip (stack local): %11u"), Data);
+}
+
+template<> void AutoRTFM::FStats::Report<EStatsKind::AverageCommitTasks>(const uint64_t Data) const
+{
+    const uint64_t TotalTransactions = Datas[static_cast<size_t>(EStatsKind::Transaction)];
+    UE_LOG(LogAutoRTFM, Display, TEXT("  Average commit tasks:      %15.3f"), (static_cast<double>(Data) / static_cast<double>(TotalTransactions)));
+}
+
+template<> void AutoRTFM::FStats::Report<EStatsKind::MaximumCommitTasks>(const uint64_t Data) const
+{
+    UE_LOG(LogAutoRTFM, Display, TEXT("  Maximum commit tasks:      %11u"), Data);
+}
+
+template<> void AutoRTFM::FStats::Report<EStatsKind::AverageAbortTasks>(const uint64_t Data) const
+{
+    const uint64_t TotalTransactions = Datas[static_cast<size_t>(EStatsKind::Transaction)];
+    UE_LOG(LogAutoRTFM, Display, TEXT("  Average abort tasks:       %15.3f"), (static_cast<double>(Data) / static_cast<double>(TotalTransactions)));
+}
+
+template<> void AutoRTFM::FStats::Report<EStatsKind::MaximumAbortTasks>(const uint64_t Data) const
+{
+    UE_LOG(LogAutoRTFM, Display, TEXT("  Maximum abort tasks:       %11u"), Data);
+}
+
+void AutoRTFM::FStats::Report() const
+{
+	if constexpr (bCollectStats)
+	{
+		UE_LOG(LogAutoRTFM, Display, TEXT("AutoRTFM Statistics:"));
+
+		for (size_t I = 0; I < static_cast<size_t>(EStatsKind::Total); I++)
+		{
+			switch (static_cast<EStatsKind>(I))
+			{
+			default: AutoRTFM::Unreachable(); break;
+#define REPORT_CASE(x) case (x): Report<x>(Datas[static_cast<size_t>(x)]); break
+			REPORT_CASE(EStatsKind::Transaction);
+			REPORT_CASE(EStatsKind::Commit);
+			REPORT_CASE(EStatsKind::Abort);
+			REPORT_CASE(EStatsKind::AverageTransactionDepth);
+            REPORT_CASE(EStatsKind::MaximumTransactionDepth);
+            REPORT_CASE(EStatsKind::AverageWriteLogEntries);
+            REPORT_CASE(EStatsKind::MaximumWriteLogEntries);
+			REPORT_CASE(EStatsKind::AverageWriteLogBytes);
+			REPORT_CASE(EStatsKind::MaximumWriteLogBytes);
+            REPORT_CASE(EStatsKind::HitSetHit);
+            REPORT_CASE(EStatsKind::HitSetMiss);
+            REPORT_CASE(EStatsKind::HitSetSkippedBecauseOfStackLocalMemory);
+            REPORT_CASE(EStatsKind::AverageCommitTasks);
+            REPORT_CASE(EStatsKind::MaximumCommitTasks);
+            REPORT_CASE(EStatsKind::AverageAbortTasks);
+            REPORT_CASE(EStatsKind::MaximumAbortTasks);
+#undef REPORT_CASE
+			}
+		}
+
+		UE_DEBUG_BREAK();
+	}
+}
+
+AutoRTFM::FStats AutoRTFM::Stats;
+
+#endif // defined(__AUTORTFM) && __AUTORTFM
diff --git a/Engine/Source/Runtime/Core/Private/AutoRTFM/Stats.h b/Engine/Source/Runtime/Core/Private/AutoRTFM/Stats.h
new file mode 100644
--- /dev/null
+++ b/Engine/Source/Runtime/Core/Private/AutoRTFM/Stats.h
@@ -0,0 +1,172 @@
+// Copyright Epic Games, Inc. All Rights Reserved.
+
+#pragma once
+
+#include "Utils.h"
+
+namespace AutoRTFM
+{
+
+static constexpr bool bCollectStats = false;
+
+enum class EStatsKind : uint8_t
+{
+	Transaction = 0,
+	Commit,
+	Abort,
+	AverageTransactionDepth,
+	MaximumTransactionDepth,
+	AverageWriteLogEntries,
+	MaximumWriteLogEntries,
+	AverageWriteLogBytes,
+	MaximumWriteLogBytes,
+	HitSetHit,
+	HitSetMiss,
+	HitSetSkippedBecauseOfStackLocalMemory,
+	AverageCommitTasks,
+	MaximumCommitTasks,
+	AverageAbortTasks,
+	MaximumAbortTasks,
+	Total
+};
+
+template<EStatsKind Kind> constexpr bool IsStatsKindNoArgs()
+{
+	switch (Kind)
+	{
+	default:
+		return false;
+	case EStatsKind::Transaction:
+	case EStatsKind::Commit:
+	case EStatsKind::Abort:
+	case EStatsKind::HitSetHit:
+	case EStatsKind::HitSetMiss:
+	case EStatsKind::HitSetSkippedBecauseOfStackLocalMemory:
+		return true;
+	}
+}
+
+template<EStatsKind Kind> constexpr bool IsStatsKindOneArg()
+{
+	switch (Kind)
+	{
+	default:
+		return false;
+	case EStatsKind::AverageTransactionDepth:
+	case EStatsKind::MaximumTransactionDepth:
+	case EStatsKind::AverageWriteLogEntries:
+	case EStatsKind::MaximumWriteLogEntries:
+	case EStatsKind::AverageWriteLogBytes:
+	case EStatsKind::MaximumWriteLogBytes:
+	case EStatsKind::AverageCommitTasks:
+	case EStatsKind::MaximumCommitTasks:
+	case EStatsKind::AverageAbortTasks:
+	case EStatsKind::MaximumAbortTasks:
+		return true;
+	}
+}
+
+template<EStatsKind Kind> constexpr bool IsStatsKindAverage()
+{
+	switch (Kind)
+	{
+	default:
+		return false;
+	case EStatsKind::AverageTransactionDepth:
+	case EStatsKind::AverageWriteLogEntries:
+	case EStatsKind::AverageWriteLogBytes:
+	case EStatsKind::AverageCommitTasks:
+	case EStatsKind::AverageAbortTasks:
+		return true;
+	}
+}
+
+template<EStatsKind Kind> constexpr bool IsStatsKindMaximum()
+{
+	switch (Kind)
+	{
+	default:
+		return false;
+	case EStatsKind::MaximumTransactionDepth:
+	case EStatsKind::MaximumWriteLogEntries:
+	case EStatsKind::MaximumWriteLogBytes:
+	case EStatsKind::MaximumCommitTasks:
+	case EStatsKind::MaximumAbortTasks:
+		return true;
+	}
+}
+
+struct FStats final
+{
+	FStats()
+	{
+		FMemory::Memset(Datas, 0, sizeof(Datas));
+	}
+
+    ~FStats()
+    {
+		if constexpr (bCollectStats)
+		{
+			Report();
+		}
+    }
+
+    void Report() const;
+
+	template<EStatsKind Kind> UE_AUTORTFM_FORCEINLINE void Collect()
+	{
+		static_assert(IsStatsKindNoArgs<Kind>());
+
+		if constexpr (bCollectStats)
+		{
+			Datas[static_cast<size_t>(Kind)] += 1;
+		}
+	}
+
+	template<EStatsKind Kind> UE_AUTORTFM_FORCEINLINE void Collect(uint64_t Data)
+	{
+		static_assert(IsStatsKindOneArg<Kind>());
+
+		if constexpr (bCollectStats)
+		{
+			if constexpr (IsStatsKindAverage<Kind>())
+			{
+				Datas[static_cast<size_t>(Kind)] += Data;
+			}
+			else if constexpr (IsStatsKindMaximum<Kind>())
+			{
+				uint64_t* const Ptr = &Datas[static_cast<size_t>(Kind)];
+				*Ptr = std::max(*Ptr, Data);
+			}
+		}
+	}
+
+private:
+	uint64_t Datas[static_cast<size_t>(EStatsKind::Total)];
+
+	template<EStatsKind Kind> void Report(const uint64_t Data) const;
+};
+
+extern FStats Stats;
+
+template<typename T, bool bActive = bCollectStats> struct TStatStorage;
+
+template<typename T> struct TStatStorage<T, false> final
+{
+	UE_AUTORTFM_FORCEINLINE TStatStorage(const T) { /* ignored */ }
+	UE_AUTORTFM_FORCEINLINE void operator=(const T) { /* ignored */ }
+	UE_AUTORTFM_FORCEINLINE void operator+=(const T) { /* ignored */ }
+	UE_AUTORTFM_FORCEINLINE operator T() const { return T(0); }
+};
+
+template<typename T> struct TStatStorage<T, true> final
+{
+	UE_AUTORTFM_FORCEINLINE TStatStorage(const T t) : Data(t) {}
+	UE_AUTORTFM_FORCEINLINE void operator=(const T t) { Data = t; }
+	UE_AUTORTFM_FORCEINLINE void operator+=(const T t) { Data += t; }
+	UE_AUTORTFM_FORCEINLINE operator T() const { return Data; }
+private:
+	T Data;
+};
+
+} // namespace AutoRTFM
diff --git a/Engine/Source/Runtime/Core/Private/AutoRTFM/TaskArray.h b/Engine/Source/Runtime/Core/Private/AutoRTFM/TaskArray.h
--- a/Engine/Source/Runtime/Core/Private/AutoRTFM/TaskArray.h
+++ b/Engine/Source/Runtime/Core/Private/AutoRTFM/TaskArray.h
@@ -20,6 +20,7 @@ public:
     {
         Latest.Push(MoveTemp(value));
     }
+
     void Add(const T& value)
     {
         Latest.Push(value);
@@ -82,6 +83,7 @@ public:
                 return false;
             }
         }
+
         for (size_t IndexInStash = Stash.Num(); IndexInStash--;)
         {
             const TArray<T>& StashedVector = *Stash[IndexInStash];
@@ -102,6 +104,19 @@ public:
         Stash.Empty();
     }
 
+    size_t Num() const
+    {
+        size_t Result = Latest.Num();
+
+        for (size_t Index = 0; Index < Stash.Num(); Index++)
+        {
+            const TArray<T>& StashedVector = *Stash[Index];
+            Result += StashedVector.Num();
+        }
+
+        return Result;
+    }
+
 private:
     // We don't want to do this too often; currently we just do it where it's asymptitically relevant like AddAll. This doesn't
     // logically change the TaskArray but it changes its internal representation. Hence the use of `mutable` and hence why this
diff --git a/Engine/Source/Runtime/Core/Private/AutoRTFM/Transaction.cpp b/Engine/Source/Runtime/Core/Private/AutoRTFM/Transaction.cpp
--- a/Engine/Source/Runtime/Core/Private/AutoRTFM/Transaction.cpp
+++ b/Engine/Source/Runtime/Core/Private/AutoRTFM/Transaction.cpp
@@ -32,6 +32,9 @@ void FTransaction::AbortWithoutThrowing()
            || Context->GetStatus() == EContextStatus::AbortedByRequest);
     ASSERT(Context->GetCurrentTransaction() == this);
 
+    Stats.Collect<EStatsKind::Abort>();
+    CollectStats();
+
     if (IsNested())
     {
         AbortNested();
@@ -53,6 +56,10 @@ bool FTransaction::AttemptToCommit()
 {
     ASSERT(Context->GetStatus() == EContextStatus::OnTrack);
     ASSERT(Context->GetCurrentTransaction() == this);
+
+    Stats.Collect<EStatsKind::Commit>();
+    CollectStats();
+
     bool bResult;
     if (IsNested())
     {
diff --git a/Engine/Source/Runtime/Core/Private/AutoRTFM/Transaction.h b/Engine/Source/Runtime/Core/Private/AutoRTFM/Transaction.h
--- a/Engine/Source/Runtime/Core/Private/AutoRTFM/Transaction.h
+++ b/Engine/Source/Runtime/Core/Private/AutoRTFM/Transaction.h
@@ -4,6 +4,7 @@
 
 #include "HitSet.h"
 #include "LongJump.h"
+#include "Stats.h"
 #include "TaggedPtr.h"
 #include "TaskArray.h"
 #include "WriteLog.h"
@@ -20,7 +21,20 @@ public:
     
     bool IsNested() const { return !!Parent; }
     FTransaction* GetParent() const { return Parent; }
-	void SetParent(FTransaction* NewParent) { Parent = NewParent; }
+
+	void SetParent(FTransaction* NewParent)
+    {
+        Parent = NewParent;
+
+        // For stats, record the nested depth of the transaction.
+		if (NewParent)
+		{
+			StatDepth = NewParent->StatDepth + 1;
+		}
+
+		Stats.Collect<EStatsKind::AverageTransactionDepth>(StatDepth);
+		Stats.Collect<EStatsKind::MaximumTransactionDepth>(StatDepth);
+    }
 
     // This should just use type displays or ranges. Maybe ranges could even work out great.
     bool IsNestedWithin(const FTransaction* Other) const
@@ -67,6 +81,8 @@ private:
     bool AttemptToCommitOuterNest();
 
     void Reset(); // Frees memory and sets us up for possible retry.
+
+    void CollectStats() const;
     
     FContext* Context;
     
@@ -83,6 +99,7 @@ private:
     FHitSet HitSet;
     FWriteLog WriteLog;
     FWriteLogBumpAllocator WriteLogBumpAllocator;
+    TStatStorage<uint64_t> StatDepth = 1;
 };
 
 } // namespace AutoRTFM
diff --git a/Engine/Source/Runtime/Core/Private/AutoRTFM/TransactionInlines.h b/Engine/Source/Runtime/Core/Private/AutoRTFM/TransactionInlines.h
--- a/Engine/Source/Runtime/Core/Private/AutoRTFM/TransactionInlines.h
+++ b/Engine/Source/Runtime/Core/Private/AutoRTFM/TransactionInlines.h
@@ -8,29 +8,8 @@
 namespace AutoRTFM
 {
 
-inline void FTransaction::RecordWriteMaxPageSized(void* LogicalAddress, size_t Size)
+UE_AUTORTFM_FORCEINLINE void FTransaction::RecordWriteMaxPageSized(void* LogicalAddress, size_t Size)
 {
-    FMemoryLocation Key(LogicalAddress);
-    Key.SetTopTag(static_cast<uint16_t>(Size));
-
-    // If we are recording a stack address that is relative to our current
-    // transactions stack location, we do not need to record the data in the
-    // write log because if that transaction aborted, that memory will cease to
-    // be meaningful anyway!
-    // TODO: with more advanced LLVM-skillz we could just detect writes to
-    // alloca's in the LLVM code and that should probably catch all or nearly
-    // all of the cases this checks would catch, so we could revisit doing this
-    // check entirely in future!
-    if (Context->IsInnerTransactionStack(LogicalAddress))
-    {
-        return;
-    }
-
-    if (!HitSet.Insert(Key))
-    {
-        return;
-    }
-
     void* CopyAddress = WriteLogBumpAllocator.Allocate(Size);
     memcpy(CopyAddress, LogicalAddress, Size);
 
@@ -53,7 +32,31 @@ inline void FTransaction::RecordWrite(void* LogicalAddress, size_t Size, bool bI
 		UE_LOG(LogAutoRTFM, Fatal, TEXT("Writing to local stack memory from an unscoped transaction is not allowed."));
 	}
 
-    uint8_t* Address = reinterpret_cast<uint8_t*>(LogicalAddress);
+    // If we are recording a stack address that is relative to our current
+    // transactions stack location, we do not need to record the data in the
+    // write log because if that transaction aborted, that memory will cease to
+    // be meaningful anyway!
+    if (Context->IsInnerTransactionStack(LogicalAddress))
+    {
+        Stats.Collect<EStatsKind::HitSetSkippedBecauseOfStackLocalMemory>();
+        return;
+    }
+
+    if (Size <= FWriteLogBumpAllocator::MaxSize)
+    {
+        FMemoryLocation Key(LogicalAddress);
+        Key.SetTopTag(static_cast<uint16_t>(Size));
+
+        if (!HitSet.Insert(Key))
+        {
+            Stats.Collect<EStatsKind::HitSetHit>();
+            return;
+        }
+    }
+
+    Stats.Collect<EStatsKind::HitSetMiss>();
+
+    uint8_t* const Address = reinterpret_cast<uint8_t*>(LogicalAddress);
 
     size_t I = 0;
 
@@ -88,4 +91,19 @@ inline void FTransaction::DeferUntilAbort(TFunction<void()>&& Callback)
     AbortTasks.Add(MoveTemp(Callback));
 }
 
+UE_AUTORTFM_FORCEINLINE void FTransaction::CollectStats() const
+{
+    Stats.Collect<EStatsKind::AverageWriteLogEntries>(WriteLog.Num());
+    Stats.Collect<EStatsKind::MaximumWriteLogEntries>(WriteLog.Num());
+
+    Stats.Collect<EStatsKind::AverageWriteLogBytes>(WriteLogBumpAllocator.StatTotalSize);
+    Stats.Collect<EStatsKind::MaximumWriteLogBytes>(WriteLogBumpAllocator.StatTotalSize);
+
+    Stats.Collect<EStatsKind::AverageCommitTasks>(CommitTasks.Num());
+    Stats.Collect<EStatsKind::MaximumCommitTasks>(CommitTasks.Num());
+
+    Stats.Collect<EStatsKind::AverageAbortTasks>(AbortTasks.Num());
+    Stats.Collect<EStatsKind::MaximumAbortTasks>(AbortTasks.Num());
+}
+
 } // namespace AutoRTFM
diff --git a/Engine/Source/Runtime/Core/Private/AutoRTFM/Utils.h b/Engine/Source/Runtime/Core/Private/AutoRTFM/Utils.h
--- a/Engine/Source/Runtime/Core/Private/AutoRTFM/Utils.h
+++ b/Engine/Source/Runtime/Core/Private/AutoRTFM/Utils.h
@@ -11,7 +11,7 @@
 #include "Containers/UnrealString.h"
 #include "Logging/LogMacros.h"
 
-DECLARE_LOG_CATEGORY_EXTERN(LogAutoRTFM, NoLogging, All)
+DECLARE_LOG_CATEGORY_EXTERN(LogAutoRTFM, Display, All)
 
 #define ASSERT(exp) UE_CLOG(!(exp), LogAutoRTFM, Fatal, TEXT("%s:%d:%s: assertion %s failed."), __FILE__, __LINE__, __PRETTY_FUNCTION__, #exp)
 
diff --git a/Engine/Source/Runtime/Core/Private/AutoRTFM/WriteLog.h b/Engine/Source/Runtime/Core/Private/AutoRTFM/WriteLog.h
--- a/Engine/Source/Runtime/Core/Private/AutoRTFM/WriteLog.h
+++ b/Engine/Source/Runtime/Core/Private/AutoRTFM/WriteLog.h
@@ -160,13 +160,8 @@ namespace AutoRTFM
             TotalSize = 0;
         }
 
-        bool IsEmpty() const
-        {
-            return 0 == TotalSize;
-        }
-
-		size_t Num() const { return TotalSize; }
-
+        UE_AUTORTFM_FORCEINLINE bool IsEmpty() const { return 0 == TotalSize; }
+		UE_AUTORTFM_FORCEINLINE size_t Num() const { return TotalSize; }
     private:
 
         FWriteLogEntryBucket* Start = nullptr;
diff --git a/Engine/Source/Runtime/Core/Private/AutoRTFM/WriteLogBumpAllocator.h b/Engine/Source/Runtime/Core/Private/AutoRTFM/WriteLogBumpAllocator.h
--- a/Engine/Source/Runtime/Core/Private/AutoRTFM/WriteLogBumpAllocator.h
+++ b/Engine/Source/Runtime/Core/Private/AutoRTFM/WriteLogBumpAllocator.h
@@ -2,8 +2,7 @@
 
 #pragma once
 
-#include <stdint.h>
-
+#include "Stats.h"
 #include "Utils.h"
 
 namespace AutoRTFM
@@ -30,6 +29,7 @@ namespace AutoRTFM
 		void* Allocate(size_t Bytes)
 		{
 			ASSERT(Bytes <= MaxSize);
+			StatTotalSize += Bytes;
 
 			if (nullptr == Start)
 			{
@@ -80,6 +80,8 @@ namespace AutoRTFM
 			Other.Start = nullptr;
 			Other.Current = nullptr;
 		}
+		
+		TStatStorage<size_t> StatTotalSize = 0;
 	private:
 
 		struct FPage final
