Move from std::max to FMath::Max

diff --git a/Engine/Source/Runtime/Core/Private/AutoRTFM/Stats.h b/Engine/Source/Runtime/Core/Private/AutoRTFM/Stats.h
--- a/Engine/Source/Runtime/Core/Private/AutoRTFM/Stats.h
+++ b/Engine/Source/Runtime/Core/Private/AutoRTFM/Stats.h
@@ -136,7 +136,7 @@ struct FStats final
 			else if constexpr (IsStatsKindMaximum<Kind>())
 			{
 				uint64_t* const Ptr = &Datas[static_cast<size_t>(Kind)];
-				*Ptr = std::max(*Ptr, Data);
+				*Ptr = FMath::Max(*Ptr, Data);
 			}
 		}
 	}
