ParkingLot: Gracefully support waiting even after FThreadLocalData is destroyed

#rnx

diff --git a/Engine/Source/Runtime/Core/Private/Async/ParkingLot.cpp b/Engine/Source/Runtime/Core/Private/Async/ParkingLot.cpp
--- a/Engine/Source/Runtime/Core/Private/Async/ParkingLot.cpp
+++ b/Engine/Source/Runtime/Core/Private/Async/ParkingLot.cpp
@@ -103,24 +103,30 @@ class FThreadLocalData
 {
 public:
 	// TODO: This version reserves the table as threads spawn. The function-local thread_local
-	//       reserves the tables as threads wait for the first time. Not sure which is better.
+	//       reserves the table as threads wait for the first time. Not sure which is better.
 	//static thread_local FThreadLocalData ThreadLocalData;
 
-	static FThread& Get()
+	static TRefCountPtr<FThread> Get()
 	{
 		static thread_local FThreadLocalData ThreadLocalData;
 		FThreadLocalData& LocalThreadLocalData = ThreadLocalData;
-		if (!LocalThreadLocalData.Thread)
+		if (LIKELY(LocalThreadLocalData.Thread))
 		{
-			LocalThreadLocalData.Thread = FThread::New();
+			return LocalThreadLocalData.Thread;
 		}
-		return *LocalThreadLocalData.Thread;
+		if (LocalThreadLocalData.bDestroyed)
+		{
+			return FThread::New();
+		}
+		LocalThreadLocalData.Thread = FThread::New();
+		return LocalThreadLocalData.Thread;
 	}
 
 private:
 	inline static std::atomic<uint32> ThreadCount = 0;
 
 	TRefCountPtr<FThread> Thread;
+	bool bDestroyed = false;
 
 	FThreadLocalData();
 	~FThreadLocalData();
@@ -554,6 +560,8 @@ FThreadLocalData::FThreadLocalData()
 FThreadLocalData::~FThreadLocalData()
 {
 	ThreadCount.fetch_sub(1, std::memory_order_relaxed);
+	Thread.SafeRelease();
+	bDestroyed = true;
 }
 
 ///////////////////////////////////////////////////////////////////////////////////////////////////////////////////////
@@ -574,10 +582,10 @@ FWaitState WaitUntil(const void* Address, const TFunctionRef<bool()>& CanWait, c
 
 	check(!WaitTime.IsNaN());
 
-	FThread& Self = FThreadLocalData::Get();
+	TRefCountPtr<FThread> Self = FThreadLocalData::Get();
 
-	checkfSlow(!Self.WaitAddress, TEXT("WaitAddress must be null. This can happen if Wait is called by BeforeWait."));
-	checkfSlow(Self.WakeToken == 0, TEXT("WakeToken must be 0. This is an error in ParkingLot."));
+	checkfSlow(!Self->WaitAddress, TEXT("WaitAddress must be null. This can happen if Wait is called by BeforeWait."));
+	checkfSlow(Self->WakeToken == 0, TEXT("WakeToken must be 0. This is an error in ParkingLot."));
 
 	FWaitState State;
 
@@ -590,23 +598,23 @@ FWaitState WaitUntil(const void* Address, const TFunctionRef<bool()>& CanWait, c
 		{
 			return State;
 		}
-		Self.WaitAddress = Address;
-		Self.Event.Reset();
-		Bucket.Enqueue(&Self);
+		Self->WaitAddress = Address;
+		Self->Event.Reset();
+		Bucket.Enqueue(&*Self);
 	}
 
 	// BeforeWait must be invoked after the bucket is unlocked.
 	BeforeWait();
 
 	// Wait until the timeout or until the thread has been dequeued.
-	Self.Event.WaitUntil(WaitTime);
+	Self->Event.WaitUntil(WaitTime);
 
 	// WaitAddress is reset when the thread is dequeued.
-	if (!Self.WaitAddress)
+	if (!Self->WaitAddress)
 	{
 		State.bDidWake = true;
-		State.WakeToken = Self.WakeToken;
-		Self.WakeToken = 0;
+		State.WakeToken = Self->WakeToken;
+		Self->WakeToken = 0;
 		return State;
 	}
 
@@ -615,7 +623,7 @@ FWaitState WaitUntil(const void* Address, const TFunctionRef<bool()>& CanWait, c
 	bool bDequeued = false;
 	if (TDynamicUniqueLock<FWordMutex> BucketLock; FBucket* Bucket = FTable::FindBucket(Address, BucketLock))
 	{
-		Bucket->DequeueIf([Self = &Self, &bDequeued](FThread* Thread)
+		Bucket->DequeueIf([Self = &*Self, &bDequeued](FThread* Thread)
 		{
 			if (Thread == Self)
 			{
@@ -631,10 +639,10 @@ FWaitState WaitUntil(const void* Address, const TFunctionRef<bool()>& CanWait, c
 	// has finished waking this thread by setting its wait address to null.
 	if (!bDequeued)
 	{
-		Self.Event.Wait();
+		Self->Event.Wait();
 		State.bDidWake = true;
-		State.WakeToken = Self.WakeToken;
-		Self.WakeToken = 0;
+		State.WakeToken = Self->WakeToken;
+		Self->WakeToken = 0;
 	}
 
 	return State;
