Switched TScriptDelegate::ProcessDelegate from a write scope to a read scope, to allow concurrent read access.

diff --git a/Engine/Source/Runtime/Core/Public/UObject/ScriptDelegates.h b/Engine/Source/Runtime/Core/Public/UObject/ScriptDelegates.h
--- a/Engine/Source/Runtime/Core/Public/UObject/ScriptDelegates.h
+++ b/Engine/Source/Runtime/Core/Public/UObject/ScriptDelegates.h
@@ -419,10 +419,8 @@ public:
 		UObjectTemplate* ObjectPtr;
 		UFunction* Function;
 
-		{	// to avoid MT access check if the delegate is deleted from inside of its callback, we don't cover the callback execution
-			// by access protection scope
-			// the `const` on the method is a lie
-			FWriteAccessScope WriteScope = const_cast<TScriptDelegate*>(this)->GetWriteAccessScope();
+		{	// don't cover delegate execution by access protection
+			FReadAccessScope ReadScope = GetReadAccessScope();
 
 			checkf( Object.IsValid() != false, TEXT( "ProcessDelegate() called with no object bound to delegate!" ) );
 			checkf( FunctionName != NAME_None, TEXT( "ProcessDelegate() called with no function name set!" ) );
