Mac/Android/iOS: Replace UE4 string references with UE.

diff --git a/Engine/Source/Runtime/Core/Private/Android/AndroidPlatformCrashContext.cpp b/Engine/Source/Runtime/Core/Private/Android/AndroidPlatformCrashContext.cpp
--- a/Engine/Source/Runtime/Core/Private/Android/AndroidPlatformCrashContext.cpp
+++ b/Engine/Source/Runtime/Core/Private/Android/AndroidPlatformCrashContext.cpp
@@ -175,9 +175,8 @@ const ANSICHAR* FAndroidCrashContext::ItoANSI(uint64 Val, uint64 Base, uint32 Le
 void FAndroidCrashContext::GenerateReportDirectoryName(char(&DirectoryNameOUT)[CrashReportMaxPathSize])
 {
 	FGuid ReportGUID = FGuid::NewGuid();
-	//FCStringAnsi::Snprintf(DirectoryNameOUT, CrashReportMaxPathSize, ("%s/CrashReport-UE4-%s-pid-%d-%08X%08X%08X%08X"), GAndroidCrashInfo.AndroidCrashReportPath, GAndroidCrashInfo.ProjectNameUTF8, (uint32)getpid(), ReportGUID.A, ReportGUID.B, ReportGUID.C, ReportGUID.D);
 	FCStringAnsi::Strncpy(DirectoryNameOUT, GAndroidCrashInfo.AndroidCrashReportPath, CrashReportMaxPathSize);
-	FCStringAnsi::Strcat(DirectoryNameOUT, CrashReportMaxPathSize, "/CrashReport-UE4-");
+	FCStringAnsi::Strcat(DirectoryNameOUT, CrashReportMaxPathSize, "/CrashReport-UE-");
 	FCStringAnsi::Strcat(DirectoryNameOUT, CrashReportMaxPathSize, GAndroidCrashInfo.ProjectNameUTF8);
 	FCStringAnsi::Strcat(DirectoryNameOUT, CrashReportMaxPathSize, "-pid-");
 	FCStringAnsi::Strcat(DirectoryNameOUT, CrashReportMaxPathSize, ItoANSI((uint32)getpid(), 10));
diff --git a/Engine/Source/Runtime/Core/Private/IOS/IOSPlatformCrashContext.cpp b/Engine/Source/Runtime/Core/Private/IOS/IOSPlatformCrashContext.cpp
--- a/Engine/Source/Runtime/Core/Private/IOS/IOSPlatformCrashContext.cpp
+++ b/Engine/Source/Runtime/Core/Private/IOS/IOSPlatformCrashContext.cpp
@@ -227,7 +227,7 @@ void FIOSCrashContext::GenerateCrashInfo() const
     // create a crash-specific directory
     char CrashInfoFolder[PATH_MAX] = {};
     FCStringAnsi::Strncpy(CrashInfoFolder, GIOSAppInfo.CrashReportPath, PATH_MAX);
-    FCStringAnsi::Strcat(CrashInfoFolder, PATH_MAX, "/CrashReport-UE4-");
+    FCStringAnsi::Strcat(CrashInfoFolder, PATH_MAX, "/CrashReport-UE-");
     FCStringAnsi::Strcat(CrashInfoFolder, PATH_MAX, GIOSAppInfo.AppNameUTF8);
     FCStringAnsi::Strcat(CrashInfoFolder, PATH_MAX, "-pid-");
     FCStringAnsi::Strcat(CrashInfoFolder, PATH_MAX, ItoANSI(getpid(), 10));
diff --git a/Engine/Source/Runtime/Core/Private/IOS/IOSPlatformMisc.cpp b/Engine/Source/Runtime/Core/Private/IOS/IOSPlatformMisc.cpp
--- a/Engine/Source/Runtime/Core/Private/IOS/IOSPlatformMisc.cpp
+++ b/Engine/Source/Runtime/Core/Private/IOS/IOSPlatformMisc.cpp
@@ -1961,7 +1961,7 @@ FString FIOSCrashContext::CreateCrashFolder() const
 	// create a crash-specific directory
 	char CrashInfoFolder[PATH_MAX] = {};
 	FCStringAnsi::Strncpy(CrashInfoFolder, GIOSAppInfo.CrashReportPath, PATH_MAX);
-	FCStringAnsi::Strcat(CrashInfoFolder, PATH_MAX, "/CrashReport-UE4-");
+	FCStringAnsi::Strcat(CrashInfoFolder, PATH_MAX, "/CrashReport-UE-");
 	FCStringAnsi::Strcat(CrashInfoFolder, PATH_MAX, GIOSAppInfo.AppNameUTF8);
 	FCStringAnsi::Strcat(CrashInfoFolder, PATH_MAX, "-pid-");
 	FCStringAnsi::Strcat(CrashInfoFolder, PATH_MAX, ItoANSI(getpid(), 10));
diff --git a/Engine/Source/Runtime/Core/Private/Mac/MacPlatformMisc.cpp b/Engine/Source/Runtime/Core/Private/Mac/MacPlatformMisc.cpp
--- a/Engine/Source/Runtime/Core/Private/Mac/MacPlatformMisc.cpp
+++ b/Engine/Source/Runtime/Core/Private/Mac/MacPlatformMisc.cpp
@@ -566,12 +566,12 @@ void FMacPlatformMisc::PlatformPreInit()
 	}
 	if (Limit.rlim_cur < OPEN_MAX)
 	{
-		UE_LOG(LogInit, Warning, TEXT("Open files limit too small: %llu, should be at least OPEN_MAX (%llu). rlim_max is %llu, kern.maxfilesperproc is %u. UE4 may be unstable."), Limit.rlim_cur, OPEN_MAX, Limit.rlim_max, MaxFilesPerProc);
+		UE_LOG(LogInit, Warning, TEXT("Open files limit too small: %llu, should be at least OPEN_MAX (%llu). rlim_max is %llu, kern.maxfilesperproc is %u. UE may be unstable."), Limit.rlim_cur, OPEN_MAX, Limit.rlim_max, MaxFilesPerProc);
 	}
 	Result = setrlimit(RLIMIT_NOFILE, &Limit);
 	if (Result != 0)
 	{
-		UE_LOG(LogInit, Warning, TEXT("Failed to change open file limit, UE4 may be unstable."));
+		UE_LOG(LogInit, Warning, TEXT("Failed to change open file limit, UE may be unstable."));
 	}
 
 	FApplePlatformSymbolication::EnableCoreSymbolication(!FPlatformProcess::IsSandboxedApplication() && IS_PROGRAM);
@@ -2191,7 +2191,7 @@ void FMacCrashContext::GenerateInfoInFolder(char const* const InfoFolder) const
 		int ReportFile = open(FilePath, O_CREAT|O_WRONLY, 0766);
 		if (ReportFile != -1)
 		{
-			WriteUTF16String(ReportFile, TEXT("GameName UE4-"));
+			WriteUTF16String(ReportFile, TEXT("GameName UE-"));
 			WriteLine(ReportFile, *GMacAppInfo.AppName);
 			
 			WriteUTF16String(ReportFile, TEXT("BuildVersion 1.0."));
@@ -2364,7 +2364,7 @@ void FMacCrashContext::GenerateCrashInfoAndLaunchReporter() const
 	if(bCanRunCrashReportClient)
 	{
 		// create a crash-specific directory
-		FString CrashInfoFolder = FString::Printf(TEXT("%s/CrashReport-UE4-%s-pid-%d-%s"), UTF8_TO_TCHAR(GMacAppInfo.CrashReportPath), UTF8_TO_TCHAR(GMacAppInfo.AppNameUTF8), (int32)getpid(), *GMacAppInfo.RunUUID.ToString());
+		FString CrashInfoFolder = FString::Printf(TEXT("%s/CrashReport-UE-%s-pid-%d-%s"), UTF8_TO_TCHAR(GMacAppInfo.CrashReportPath), UTF8_TO_TCHAR(GMacAppInfo.AppNameUTF8), (int32)getpid(), *GMacAppInfo.RunUUID.ToString());
 
 		// Do not inline this! The lifetime of this object needs to extend over the usage of Argv in posix_spawn() call below.
 		auto CrashInfoFolderUTF8 = TStringConversion<FTCHARToUTF8_Convert>(*CrashInfoFolder);
