Mac: Add some sanity checking on CRC to make sure it actually exists and is executable.

diff --git a/Engine/Source/Runtime/Core/Private/Mac/MacPlatformMisc.cpp b/Engine/Source/Runtime/Core/Private/Mac/MacPlatformMisc.cpp
--- a/Engine/Source/Runtime/Core/Private/Mac/MacPlatformMisc.cpp
+++ b/Engine/Source/Runtime/Core/Private/Mac/MacPlatformMisc.cpp
@@ -2398,28 +2398,44 @@ void FMacCrashContext::GenerateCrashInfoAndLaunchReporter() const
 			}
 		}
 
-		posix_spawn_file_actions_t FileActions;
-		posix_spawn_file_actions_init(&FileActions);
-
-		posix_spawnattr_t SpawnAttr;
-		posix_spawnattr_init(&SpawnAttr);
-
+		// We've generated the crash report above since we were able to,
+		// now we need to make sure CRC actually exists and is executable before
+		// trying to run it.
+		struct stat Stat;
+		stat(GMacAppInfo.CrashReportClient, &Stat);
+		if (!S_ISREG(Stat.st_mode))
+		{
+			UE_LOG(LogMac, Error, TEXT("Unable to locate CrashReporterClient: %s"), UTF8_TO_TCHAR(GMacAppInfo.CrashReportClient));
+		}
+		else if ((Stat.st_mode & S_IXUSR) == 0)
 		{
-			uint16 SpawnFlags = POSIX_SPAWN_SETPGROUP;
-			posix_spawnattr_setflags(&SpawnAttr, SpawnFlags);
+			UE_LOG(LogMac, Error, TEXT("Unable to execute CrashReporterClient, please run: chmod u+x %s"), UTF8_TO_TCHAR(GMacAppInfo.CrashReportClient));
 		}
+		else
+		{
+			posix_spawn_file_actions_t FileActions;
+			posix_spawn_file_actions_init(&FileActions);
 
-		extern char **environ; // provided by libc
+			posix_spawnattr_t SpawnAttr;
+			posix_spawnattr_init(&SpawnAttr);
 
-		// Use posix_spawn() as it is async-signal safe, CreateProc can fail in Cocoa.
-		Status = posix_spawn(&CrcPID, GMacAppInfo.CrashReportClient, &FileActions, &SpawnAttr, (char *const *)Argv, environ);
+			{
+				uint16 SpawnFlags = POSIX_SPAWN_SETPGROUP;
+				posix_spawnattr_setflags(&SpawnAttr, SpawnFlags);
+			}
 
-		posix_spawn_file_actions_destroy(&FileActions);
-		posix_spawnattr_destroy(&SpawnAttr);
+			extern char **environ; // provided by libc
 
-		if (Status != 0)
-		{
-			UE_LOG(LogHAL, Fatal, TEXT("FMacPlatformMisc::GenerateCrashInfoAndLaunchReporter: posix_spawn() failed (%d, %s)"), Status, UTF8_TO_TCHAR(strerror(Status)));
+			// Use posix_spawn() as it is async-signal safe, CreateProc can fail in Cocoa.
+			Status = posix_spawn(&CrcPID, GMacAppInfo.CrashReportClient, &FileActions, &SpawnAttr, (char *const *)Argv, environ);
+
+			posix_spawn_file_actions_destroy(&FileActions);
+			posix_spawnattr_destroy(&SpawnAttr);
+
+			if (Status != 0)
+			{
+				UE_LOG(LogMac, Error, TEXT("FMacPlatformMisc::GenerateCrashInfoAndLaunchReporter: posix_spawn() failed (%d, %s)"), Status, UTF8_TO_TCHAR(strerror(Status)));
+			}
 		}
 	}
 
