Fix natvis and alignment of numbered 4B FNames

Also removed unneeded IterateNumberedNames
Only impacts UE_FNAME_OUTLINE_NUMBER

#rnx

diff --git a/Engine/Source/Runtime/Core/Private/UObject/UnrealNames.cpp b/Engine/Source/Runtime/Core/Private/UObject/UnrealNames.cpp
--- a/Engine/Source/Runtime/Core/Private/UObject/UnrealNames.cpp
+++ b/Engine/Source/Runtime/Core/Private/UObject/UnrealNames.cpp
@@ -475,7 +475,7 @@ public:
 	 * @param   Size  Size in bytes to allocate, 
 	 * @return  Allocation of passed in size cast to a FNameEntry pointer.
 	 */
-	template <class ScopeLock, bool bAlignNumberedName = false>
+	template <class ScopeLock, bool bNumbered = false>
 	FNameEntryHandle Allocate(uint32 Bytes)
 	{
 		Bytes = Align(Bytes, alignof(FNameEntry));
@@ -483,25 +483,21 @@ public:
 
 		ScopeLock _(Lock);
 
-		// Pad up FNameEntry::NumberedName to be 4B-aligned if FNameEntry is only 2B-aligned
-#if UE_FNAME_OUTLINE_NUMBER
-		if constexpr (bAlignNumberedName)
-		{
-			CurrentByteCursor = Align(CurrentByteCursor + offsetof(FNameEntry, NumberedName), alignof(FNumberedEntry)) - offsetof(FNameEntry, NumberedName);
-		}
-#endif
+		AlignCurrentByteCursor<bNumbered>();
 
 		// Allocate a new pool if current one is exhausted. We don't worry about a little bit
 		// of waste at the end given the relative size of pool to average and max allocation.
 		if (BlockSizeBytes - CurrentByteCursor < Bytes)
 		{
 			AllocateNewBlock();
+			AlignCurrentByteCursor<bNumbered>();
 		}
 
 		// Use current cursor position for this allocation and increment cursor for next allocation
 		uint32 ByteOffset = CurrentByteCursor;
 		CurrentByteCursor += Bytes;
 		
+		check(!(bNumbered && IsUnalignedNumberedEntry(ByteOffset)));
 		check(ByteOffset % Stride == 0 && ByteOffset / Stride < FNameBlockOffsets);
 
 		return FNameEntryHandle(CurrentBlock, ByteOffset / Stride);
@@ -536,13 +532,15 @@ public:
 	}
 
 #if UE_FNAME_OUTLINE_NUMBER
+	static_assert(sizeof(FNameEntry::FNumberedData::Id)	== sizeof(FNumberedEntry::Id) && sizeof(FNameEntry::FNumberedData::Number) == sizeof(FNumberedEntry::Number));
+	static_assert(offsetof(FNameEntry::FNumberedData, Number) - offsetof(FNameEntry::FNumberedData,	Id) == offsetof(FNumberedEntry,	Number) - offsetof(FNumberedEntry, Id));
+	static constexpr uint32 NumberedEntrySize = offsetof(FNameEntry, NumberedName) + sizeof(FNameEntry::NumberedName);
+
 	template<class ScopeLock>
 	FNameEntryHandle CreateWithNumber(const FNameEntryId& StringPart, uint32 NumberPart, TOptional<FNameEntryId> ComparisonId, FNameEntryHeader Header)
 	{
-		static_assert(sizeof(FNameEntry::NumberedName) == sizeof(FNumberedEntry));
-		static constexpr bool bAlignNumberedName = alignof(FNumberedEntry) < alignof(FNameEntry);
 		FPlatformMisc::Prefetch(Blocks[CurrentBlock]);
-		FNameEntryHandle Handle = Allocate<ScopeLock, bAlignNumberedName>(FNameEntry::GetDataOffset() + sizeof(FNameEntry::NumberedName));
+		FNameEntryHandle Handle = Allocate<ScopeLock, true>(NumberedEntrySize);
 		FNameEntry& Entry = Resolve(Handle);
 
 #if WITH_CASE_PRESERVING_NAME
@@ -550,7 +548,8 @@ public:
 #endif
 
 		Entry.Header = Header;
-		*reinterpret_cast<FNumberedEntry*>(Entry.NumberedName) = {StringPart, NumberPart};
+		checkf((uint64)&Entry.GetNumberedName() % alignof(FNumberedEntry) == 0, TEXT("Failed to align numbered FName data"));
+		reinterpret_cast<FNumberedEntry&>(Entry.NumberedName.Id[0]) = {StringPart, NumberPart};
 
 		return Handle;
 	}
@@ -580,60 +579,6 @@ public:
 	
 	uint8** GetBlocksForDebugVisualizer() { return Blocks; }
 
-#if UE_FNAME_OUTLINE_NUMBER
-	void IterateNumberedNames(TFunctionRef<void(FName)> InFunc) const
-	{
-		bool bStop = false;
-		TArray<FNameEntryId> Tmp;
-		Tmp.Reserve(BlockSizeBytes / FNameEntry::GetNumberedEntrySize());
-		for (uint32 BlockIdx = 0; !bStop; ++BlockIdx)
-		{
-			Tmp.Reset();
-			{
-				FRWScopeLock _(Lock, FRWScopeLockType::SLT_ReadOnly);
-
-				uint32 Size = BlockSizeBytes;
-				if(BlockIdx == CurrentBlock)
-				{
-					bStop = true;
-					Size = CurrentByteCursor;
-				}
-
-				const uint8* It = Blocks[BlockIdx];
-				const uint8* End = It + Size - FNameEntry::GetDataOffset();
-				while (It < End)
-				{
-					const FNameEntry* Entry = (const FNameEntry*)It;
-					if (uint32 Len = Entry->Header.Len)
-					{
-						// Skip, this is not a numbered entry 
-						It += FNameEntry::GetSize(Len, !Entry->IsWide());
-					}
-					else // Null-terminator or numbered-entry
-					{
-						if ((End - It) < FNameEntry::GetNumberedEntrySize() || Entry->GetNumberedName().IsEmpty())
-						{
-							break;
-						}
-						else
-						{
-							FNameEntryId Id = FNameEntryHandle(BlockIdx, (It - Blocks[BlockIdx]) / Stride);
-							Tmp.Add(Id);
-							It += FNameEntry::GetNumberedEntrySize();
-						}
-					}
-				}
-			}
-
-			// Dispatch callbacks for each block without holding lock
-			for (FNameEntryId Id : Tmp)
-			{
-				InFunc(FName(Id, Id, NAME_NO_NUMBER_INTERNAL));
-			}
-		}
-	}
-#endif
-
 	void DebugDump(TArray<const FNameEntry*>& Out) const
 	{
 		FRWScopeLock _(Lock, FRWScopeLockType::SLT_ReadOnly);
@@ -647,8 +592,31 @@ public:
 	}
 
 private:
-	static void DebugDumpBlock(const uint8* It, uint32 BlockSize, TArray<const FNameEntry*>& Out)
+	static constexpr bool IsUnalignedNumberedEntry(uint32 BlockOffset)
 	{
+#if UE_FNAME_OUTLINE_NUMBER
+		return alignof(FNameEntry) < alignof(FNumberedEntry) && (BlockOffset + offsetof(FNameEntry, NumberedName) + offsetof(FNameEntry::FNumberedData, Id)) % alignof(FNumberedEntry) != 0;
+#else
+		return false;
+#endif
+	}
+	
+	static constexpr FNameEntryHeader NumberPadHeader = {};
+
+	// Pad up FNameEntry::NumberedName to be 4B-aligned if FNameEntry is only 2B-aligned
+	template <bool bNumbered>
+	void AlignCurrentByteCursor()
+	{
+		if (bNumbered && IsUnalignedNumberedEntry(CurrentByteCursor) && (BlockSizeBytes - CurrentByteCursor) >= sizeof(FNameEntryHeader))
+		{
+			(FNameEntryHeader&)Blocks[CurrentBlock][CurrentByteCursor] = NumberPadHeader;
+			CurrentByteCursor += sizeof(NumberPadHeader);
+		}	
+	}
+
+	static void DebugDumpBlock(const uint8* Block, uint32 BlockSize, TArray<const FNameEntry*>& Out)
+	{
+		const uint8* It = Block;
 		const uint8* End = It + BlockSize - FNameEntry::GetDataOffset();
 		while (It < End)
 		{
@@ -658,20 +626,28 @@ private:
 				Out.Add(Entry);
 				It += FNameEntry::GetSize(Len, !Entry->IsWide());
 			}
-			else // Null-terminator or numbered-entry
+#if !UE_FNAME_OUTLINE_NUMBER
+			else
 			{
-#if UE_FNAME_OUTLINE_NUMBER
-				if ((End - It) < FNameEntry::GetNumberedEntrySize() || Entry->GetNumberedName().IsEmpty())
-				{
-					break;
-				}
-				else
-				{
-					Out.Add(Entry);
-					It += FNameEntry::GetNumberedEntrySize();
-				}
-#endif
+				break; // Normal terminator with Header.Len == 0
+			}
+#else
+			else if (End - It < NumberedEntrySize || Entry->GetNumberedName().IsEmpty())
+			{
+				break; // Numbered terminator or no room left for one
+			}
+			else if (IsUnalignedNumberedEntry(It - Block))
+			{
+				check(!WITH_CASE_PRESERVING_NAME);
+				check(Entry->Header == NumberPadHeader);
+				It += sizeof(NumberPadHeader);
 			}
+			else
+			{
+				Out.Add(Entry);
+				It += NumberedEntrySize;
+			}
+#endif
 		}
 	}
 
@@ -685,11 +661,11 @@ private:
 		LLM_SCOPE(ELLMTag::FName);
 		// Null-terminate final entry to allow DebugDump() entry iteration
 #if UE_FNAME_OUTLINE_NUMBER
-		if (CurrentByteCursor + FNameEntry::GetNumberedEntrySize() <= BlockSizeBytes)
+		if (CurrentByteCursor + NumberedEntrySize <= BlockSizeBytes)
 		{
 			FNameEntry* Terminator = (FNameEntry*)(Blocks[CurrentBlock] + CurrentByteCursor);
 			Terminator->Header.Len = 0;
-			*reinterpret_cast<FNumberedEntry*>(Terminator->NumberedName) = {};
+			Terminator->NumberedName = {}; // Might be unaligned but only DebugDumpBlock will read it
 		}
 #else
 		if (CurrentByteCursor + FNameEntry::GetDataOffset() <= BlockSizeBytes)
@@ -1579,9 +1555,6 @@ public:
 	void			LogHashStats(FOutputDevice& Ar) const;
 	uint8**			GetBlocksForDebugVisualizer() { return Entries.GetBlocksForDebugVisualizer(); }
 	TArray<const FNameEntry*> DebugDump() const;
-#if UE_FNAME_OUTLINE_NUMBER
-	void IterateNumberedNames(TFunctionRef<void(FName)> Func) const;
-#endif
 
 #if UE_TRACE_ENABLED
 	UE::Trace::FEventRef32 Trace(const FNameEntryId& EntryId);
@@ -1981,14 +1954,6 @@ TArray<const FNameEntry*> FNamePool::DebugDump() const
 	return Out;
 }
 
-#if UE_FNAME_OUTLINE_NUMBER
-void FNamePool::IterateNumberedNames(TFunctionRef<void(FName)> Func) const
-{
-	Entries.IterateNumberedNames(Func);
-}
-#endif
-
-
 bool FNamePool::IsValid(FNameEntryHandle Handle) const
 {
 	return Handle.Block < Entries.NumBlocks();
@@ -2357,12 +2322,6 @@ int32 FNameEntry::GetSize(int32 Length, bool bIsPureAnsi)
 }
 
 #if UE_FNAME_OUTLINE_NUMBER
-int32 FNameEntry::GetNumberedEntrySize()
-{
-	int32 Bytes = GetDataOffset() + sizeof(FNameEntry::NumberedName);
-	return Align(Bytes, alignof(FNameEntry));
-}
-
 uint32 FNameEntry::GetNumber() const
 {
 	checkName(Header.Len == 0);
@@ -2375,7 +2334,7 @@ int32 FNameEntry::GetSizeInBytes() const
 #if UE_FNAME_OUTLINE_NUMBER
 	if (Header.Len == 0)
 	{
-		return GetNumberedEntrySize();
+		return FNameEntryAllocator::NumberedEntrySize;
 	}
 	else
 #endif
diff --git a/Engine/Source/Runtime/Core/Public/UObject/NameTypes.h b/Engine/Source/Runtime/Core/Public/UObject/NameTypes.h
--- a/Engine/Source/Runtime/Core/Public/UObject/NameTypes.h
+++ b/Engine/Source/Runtime/Core/Public/UObject/NameTypes.h
@@ -209,13 +209,27 @@ private:
 	FNameEntryId ComparisonId;
 #endif
 	FNameEntryHeader Header;
+
+	// Unaligned to reduce alignment waste for non-numbered entries
+	struct FNumberedData
+	{
+#if UE_FNAME_OUTLINE_NUMBER	
+#if WITH_CASE_PRESERVING_NAME // ComparisonId is 4B-aligned, 4B-align Id/Number by 2B pad after 2B Header
+		uint8 Pad[sizeof(Header) % alignof(decltype(ComparisonId))]; 
+#endif						
+		uint8 Id[sizeof(FNameEntryId)];
+		uint8 Number[sizeof(uint32)];
+#endif // UE_FNAME_OUTLINE_NUMBER	
+	};
+
 	union
 	{
-		ANSICHAR	AnsiName[NAME_SIZE];
-		WIDECHAR	WideName[NAME_SIZE];
-		uint8		NumberedName[8 * UE_FNAME_OUTLINE_NUMBER];
+		ANSICHAR			AnsiName[NAME_SIZE];
+		WIDECHAR			WideName[NAME_SIZE];
+		FNumberedData		NumberedName;
 	};
 
+
 	FNameEntry(struct FClangKeepDebugInfo);
 	FNameEntry(const FNameEntry&) = delete;
 	FNameEntry(FNameEntry&&) = delete;
@@ -227,8 +241,6 @@ public:
 	FORCEINLINE bool IsWide() const								{ return Header.bIsWide; }
 	FORCEINLINE int32 GetNameLength() const 					{ return Header.Len; }
 	FORCEINLINE bool IsNumbered() const 						{ return !!UE_FNAME_OUTLINE_NUMBER && Header.Len == 0; }
-	// @pre IsNumbered()
-	FORCEINLINE const FNumberedEntry& GetNumberedName() const	{ return reinterpret_cast<const FNumberedEntry&>(NumberedName[0]); }
 
 	/**
 	 * Copy unterminated name to TCHAR buffer without allocating.
@@ -300,7 +312,8 @@ private:
 	const WIDECHAR* GetUnterminatedName(WIDECHAR(&OptionalDecodeBuffer)[NAME_SIZE]) const;
 
 #if UE_FNAME_OUTLINE_NUMBER
-	static int32 GetNumberedEntrySize();
+	// @pre IsNumbered()
+	FORCEINLINE const FNumberedEntry& GetNumberedName() const	{ return reinterpret_cast<const FNumberedEntry&>(NumberedName.Id[0]); }
 	uint32 GetNumber() const;
 #endif // UE_FNAME_OUTLINE_NUMBER
 };
