Improve performance of RecordWrite by 1.33x in AutoRTFM by avoiding some work and avoiding a costly memset(0) that wasn't needed.

diff --git a/Engine/Source/Runtime/Core/Private/AutoRTFM/API.cpp b/Engine/Source/Runtime/Core/Private/AutoRTFM/API.cpp
--- a/Engine/Source/Runtime/Core/Private/AutoRTFM/API.cpp
+++ b/Engine/Source/Runtime/Core/Private/AutoRTFM/API.cpp
@@ -115,7 +115,8 @@ extern "C" autortfm_status autortfm_close(void (*Work)(void* Arg), void* Arg)
 
 extern "C" void autortfm_record_open_write(void* Ptr, size_t Size)
 {
-	FContext::Get()->RecordWrite(Ptr, Size, false);
+    FContext::Get()->CheckOpenRecordWrite(Ptr);
+	FContext::Get()->RecordWrite(Ptr, Size);
 }
 
 extern "C" void autortfm_register_open_function(void* OriginalFunction, void* NewFunction)
diff --git a/Engine/Source/Runtime/Core/Private/AutoRTFM/CompilerInterface.cpp b/Engine/Source/Runtime/Core/Private/AutoRTFM/CompilerInterface.cpp
--- a/Engine/Source/Runtime/Core/Private/AutoRTFM/CompilerInterface.cpp
+++ b/Engine/Source/Runtime/Core/Private/AutoRTFM/CompilerInterface.cpp
@@ -31,11 +31,13 @@ void CheckAlignment(FContext* Context, void* Ptr, size_t AlignmentMask)
 extern "C" void autortfm_record_write(FContext* Context, void* Ptr, size_t Size)
 {
 	// check for writes to null here so we end up crashing in the user
-	// code rather than in the autortfm runtime
-	if (Ptr != nullptr)
+	// code rather than in the autortfm runtime.
+	if (UNLIKELY(nullptr == Ptr))
 	{
-		Context->RecordWrite(Ptr, Size, true);
+		return;
 	}
+
+	Context->RecordWrite(Ptr, Size);
 }
 
 extern "C" void* autortfm_lookup_function(FContext* Context, void* OriginalFunction, const char* Where)
diff --git a/Engine/Source/Runtime/Core/Private/AutoRTFM/Context.h b/Engine/Source/Runtime/Core/Private/AutoRTFM/Context.h
--- a/Engine/Source/Runtime/Core/Private/AutoRTFM/Context.h
+++ b/Engine/Source/Runtime/Core/Private/AutoRTFM/Context.h
@@ -35,8 +35,10 @@ public:
 	void ClearTransactionStatus();
 	bool IsAborting() const;
 
+	void CheckOpenRecordWrite(void* LogicalAddress);
+
     // Record that a write is about to occur at the given LogicalAddress of Size bytes.
-    void RecordWrite(void* LogicalAddress, size_t Size, bool bIsClosed);
+    void RecordWrite(void* LogicalAddress, size_t Size);
 
     void DidAllocate(void* LogicalAddress, size_t Size);
 
diff --git a/Engine/Source/Runtime/Core/Private/AutoRTFM/ContextInlines.h b/Engine/Source/Runtime/Core/Private/AutoRTFM/ContextInlines.h
--- a/Engine/Source/Runtime/Core/Private/AutoRTFM/ContextInlines.h
+++ b/Engine/Source/Runtime/Core/Private/AutoRTFM/ContextInlines.h
@@ -8,12 +8,23 @@
 namespace AutoRTFM
 {
 
-inline void FContext::RecordWrite(void* LogicalAddress, size_t Size, bool bIsClosed)
+UE_AUTORTFM_FORCEINLINE void FContext::CheckOpenRecordWrite(void* LogicalAddress)
 {
-    CurrentTransaction->RecordWrite(LogicalAddress, Size, bIsClosed);
+    // We don't record any writes to unscoped transactions when the write-address
+    // is within the nearest closed C++ function call nest. This may seem weird, but we
+    // have to prohibit this because we can't track all the comings and goings of
+    // local variables and all child C++ scopes. It's easy to get into a situation
+    // where a rollback will undo the recorded writes to stack-locals that actually
+    // corrupts the undo process itself.
+    ASSERT(CurrentTransaction->IsScopedTransaction() || !IsInnerTransactionStack(LogicalAddress));
 }
 
-inline void FContext::DidAllocate(void* LogicalAddress, size_t Size)
+UE_AUTORTFM_FORCEINLINE void FContext::RecordWrite(void* LogicalAddress, size_t Size)
+{
+    CurrentTransaction->RecordWrite(LogicalAddress, Size);
+}
+
+UE_AUTORTFM_FORCEINLINE void FContext::DidAllocate(void* LogicalAddress, size_t Size)
 {
     CurrentTransaction->DidAllocate(LogicalAddress, Size);
 }
diff --git a/Engine/Source/Runtime/Core/Private/AutoRTFM/Memcpy.cpp b/Engine/Source/Runtime/Core/Private/AutoRTFM/Memcpy.cpp
--- a/Engine/Source/Runtime/Core/Private/AutoRTFM/Memcpy.cpp
+++ b/Engine/Source/Runtime/Core/Private/AutoRTFM/Memcpy.cpp
@@ -16,21 +16,21 @@ void* MemcpyToNew(void* InDst, const void* InSrc, size_t Size, FContext* Context
 void* Memcpy(void* InDst, const void* InSrc, size_t Size, FContext* Context)
 {
 	UE_LOG(LogAutoRTFM, Verbose, TEXT("Memcpy(%p, %p, %zu)"), InDst, InSrc, Size);
-    Context->RecordWrite(InDst, Size, true);
+    Context->RecordWrite(InDst, Size);
     return memcpy(InDst, InSrc, Size);
 }
 
 void* Memmove(void* InDst, const void* InSrc, size_t Size, FContext* Context)
 {
 	UE_LOG(LogAutoRTFM, Verbose, TEXT("Memmove(%p, %p, %zu)"), InDst, InSrc, Size);
-	Context->RecordWrite(InDst, Size, true);
+	Context->RecordWrite(InDst, Size);
 	return memmove(InDst, InSrc, Size);
 }
 
 void* Memset(void* InDst, int Value, size_t Size, FContext* Context)
 {
 	UE_LOG(LogAutoRTFM, Verbose, TEXT("Memset(%p, %d, %zu)"), InDst, Value, Size);
-	Context->RecordWrite(InDst, Size, true);
+	Context->RecordWrite(InDst, Size);
 	return memset(InDst, Value, Size);
 }
 
diff --git a/Engine/Source/Runtime/Core/Private/AutoRTFM/StdLibSupport.cpp b/Engine/Source/Runtime/Core/Private/AutoRTFM/StdLibSupport.cpp
--- a/Engine/Source/Runtime/Core/Private/AutoRTFM/StdLibSupport.cpp
+++ b/Engine/Source/Runtime/Core/Private/AutoRTFM/StdLibSupport.cpp
@@ -76,14 +76,14 @@ char* RTFM_strcpy(char* const Dst, const char* const Src, FContext* const Contex
 {
     const size_t SrcLen = strlen(Src);
 
-    Context->RecordWrite(Dst, SrcLen, true);
+    Context->RecordWrite(Dst, SrcLen);
     return strcpy(Dst, Src);
 }
 UE_AUTORTFM_REGISTER_OPEN_FUNCTION(strcpy);
 
 char* RTFM_strncpy(char* const Dst, const char* const Src, const size_t Num, FContext* const Context)
 {
-    Context->RecordWrite(Dst, Num, true);
+    Context->RecordWrite(Dst, Num);
     return strncpy(Dst, Src, Num);
 }
 UE_AUTORTFM_REGISTER_OPEN_FUNCTION(strncpy);
@@ -93,7 +93,7 @@ char* RTFM_strcat(char* const Dst, const char* const Src, FContext* const Contex
     const size_t DstLen = strlen(Dst);
     const size_t SrcLen = strlen(Src);
 
-    Context->RecordWrite(Dst + DstLen, SrcLen + 1, true);
+    Context->RecordWrite(Dst + DstLen, SrcLen + 1);
     return strcat(Dst, Src);
 }
 UE_AUTORTFM_REGISTER_OPEN_FUNCTION(strcat);
@@ -102,7 +102,7 @@ char* RTFM_strncat(char* const Dst, const char* const Src, const size_t Num, FCo
 {
     const size_t DstLen = strlen(Dst);
 
-    Context->RecordWrite(Dst + DstLen, Num + 1, true);
+    Context->RecordWrite(Dst + DstLen, Num + 1);
     return strncat(Dst, Src, Num);
 }
 UE_AUTORTFM_REGISTER_OPEN_FUNCTION(strncat);
diff --git a/Engine/Source/Runtime/Core/Private/AutoRTFM/Transaction.h b/Engine/Source/Runtime/Core/Private/AutoRTFM/Transaction.h
--- a/Engine/Source/Runtime/Core/Private/AutoRTFM/Transaction.h
+++ b/Engine/Source/Runtime/Core/Private/AutoRTFM/Transaction.h
@@ -67,7 +67,7 @@ public:
     bool AttemptToCommit();
 
     // Record that a write is about to occur at the given LogicalAddress of Size bytes.
-    void RecordWrite(void* LogicalAddress, size_t Size, bool bIsClosed);
+    void RecordWrite(void* LogicalAddress, size_t Size);
     void RecordWriteMaxPageSized(void* LogicalAddress, size_t Size);
 
     void DidAllocate(void* LogicalAddress, size_t Size);
diff --git a/Engine/Source/Runtime/Core/Private/AutoRTFM/TransactionInlines.h b/Engine/Source/Runtime/Core/Private/AutoRTFM/TransactionInlines.h
--- a/Engine/Source/Runtime/Core/Private/AutoRTFM/TransactionInlines.h
+++ b/Engine/Source/Runtime/Core/Private/AutoRTFM/TransactionInlines.h
@@ -16,22 +16,8 @@ UE_AUTORTFM_FORCEINLINE void FTransaction::RecordWriteMaxPageSized(void* Logical
     WriteLog.Push(FWriteLogEntry(LogicalAddress, Size, CopyAddress));
 }
 
-inline void FTransaction::RecordWrite(void* LogicalAddress, size_t Size, bool bIsClosed)
+inline void FTransaction::RecordWrite(void* LogicalAddress, size_t Size)
 {
-	// We don't record any writes to unscoped transactions when the write-address
-	// is within the nearest closed C++ function call nest. This may seem weird, but we
-	// have to prohibit this because we can't track all the comings and goings of
-	// local variables and all child C++ scopes. It's easy to get into a situation
-	// where a rollback will undo the recorded writes to stack-locals that actually
-	// corrupts the undo process itself.
-	if (!IsScopedTransaction()
-		&& !bIsClosed
-		&& Context->IsInnerTransactionStack(LogicalAddress))
-	{
-		// It is an error to write to stack variables within an open nest.
-		UE_LOG(LogAutoRTFM, Fatal, TEXT("Writing to local stack memory from an unscoped transaction is not allowed."));
-	}
-
     // If we are recording a stack address that is relative to our current
     // transactions stack location, we do not need to record the data in the
     // write log because if that transaction aborted, that memory will cease to
diff --git a/Engine/Source/Runtime/Core/Private/AutoRTFM/Utils.h b/Engine/Source/Runtime/Core/Private/AutoRTFM/Utils.h
--- a/Engine/Source/Runtime/Core/Private/AutoRTFM/Utils.h
+++ b/Engine/Source/Runtime/Core/Private/AutoRTFM/Utils.h
@@ -13,7 +13,7 @@
 
 DECLARE_LOG_CATEGORY_EXTERN(LogAutoRTFM, Display, All)
 
-#define ASSERT(exp) UE_CLOG(!(exp), LogAutoRTFM, Fatal, TEXT("%s:%d:%s: assertion %s failed."), __FILE__, __LINE__, __PRETTY_FUNCTION__, #exp)
+#define ASSERT(exp) UE_CLOG(UNLIKELY(!(exp)), LogAutoRTFM, Fatal, TEXT("%s:%d:%s: assertion %s failed."), __FILE__, __LINE__, __PRETTY_FUNCTION__, #exp)
 
 namespace AutoRTFM
 {
diff --git a/Engine/Source/Runtime/Core/Private/AutoRTFM/WriteLog.h b/Engine/Source/Runtime/Core/Private/AutoRTFM/WriteLog.h
--- a/Engine/Source/Runtime/Core/Private/AutoRTFM/WriteLog.h
+++ b/Engine/Source/Runtime/Core/Private/AutoRTFM/WriteLog.h
@@ -30,9 +30,10 @@ namespace AutoRTFM
     {
         struct FWriteLogEntryBucket final
         {
-            static constexpr uint32_t BucketSize = 1024;
+            static constexpr uint32_t BucketSize = 128;
 
-            FWriteLogEntryBucket() = default;
+			// Specify a constructor so that `Entries` below will be left uninitialized.
+			explicit FWriteLogEntryBucket() {}
 
             FWriteLogEntry Entries[BucketSize];
             size_t Size = 0;
