Silence MSVC static analysis warnings

#rnx

diff --git a/Engine/Source/Runtime/Core/Private/UObject/UnrealNames.cpp b/Engine/Source/Runtime/Core/Private/UObject/UnrealNames.cpp
--- a/Engine/Source/Runtime/Core/Private/UObject/UnrealNames.cpp
+++ b/Engine/Source/Runtime/Core/Private/UObject/UnrealNames.cpp
@@ -595,7 +595,8 @@ private:
 	static constexpr bool IsUnalignedNumberedEntry(uint32 BlockOffset)
 	{
 #if UE_FNAME_OUTLINE_NUMBER
-		return alignof(FNameEntry) < alignof(FNumberedEntry) && (BlockOffset + offsetof(FNameEntry, NumberedName) + offsetof(FNameEntry::FNumberedData, Id)) % alignof(FNumberedEntry) != 0;
+		constexpr bool bRuntimeAlign = alignof(FNameEntry) < alignof(FNumberedEntry);
+		return bRuntimeAlign && (BlockOffset + offsetof(FNameEntry, NumberedName) + offsetof(FNameEntry::FNumberedData, Id)) % alignof(FNumberedEntry) != 0;
 #else
 		return false;
 #endif
diff --git a/Engine/Source/Runtime/Core/Public/UObject/NameTypes.h b/Engine/Source/Runtime/Core/Public/UObject/NameTypes.h
--- a/Engine/Source/Runtime/Core/Public/UObject/NameTypes.h
+++ b/Engine/Source/Runtime/Core/Public/UObject/NameTypes.h
@@ -240,7 +240,11 @@ public:
 	/** Returns whether this name entry is represented via WIDECHAR or ANSICHAR. */
 	FORCEINLINE bool IsWide() const								{ return Header.bIsWide; }
 	FORCEINLINE int32 GetNameLength() const 					{ return Header.Len; }
-	FORCEINLINE bool IsNumbered() const 						{ return !!UE_FNAME_OUTLINE_NUMBER && Header.Len == 0; }
+#if UE_FNAME_OUTLINE_NUMBER
+	FORCEINLINE bool IsNumbered() const 						{ return Header.Len == 0; }
+#else
+	FORCEINLINE bool IsNumbered() const 						{ return false; }
+#endif
 
 	/**
 	 * Copy unterminated name to TCHAR buffer without allocating.
