xxHash: Removed TTask from the API

Management of tasks is more suitable for higher-level code.

Includes rearrangement of some code that was updated inconsistently during compile time optimizations.

#rnx

diff --git a/Engine/Source/Developer/TextureCompressor/Private/TextureCompressorModule.cpp b/Engine/Source/Developer/TextureCompressor/Private/TextureCompressorModule.cpp
--- a/Engine/Source/Developer/TextureCompressor/Private/TextureCompressorModule.cpp
+++ b/Engine/Source/Developer/TextureCompressor/Private/TextureCompressorModule.cpp
@@ -3568,7 +3568,7 @@ public:
 
 			check( Mip.RawData.Num() != 0 );
 
-			MipHashTasks.Add(FXxHash64::HashBufferChunkedAsync(Mip.RawData.GetData(), Mip.RawData.Num(), 256 << 10));
+			MipHashTasks.Add(UE::Tasks::Launch(TEXT("ComputeMipChainHash"), [&Mip] { return FXxHash64::HashBufferChunked(MakeMemoryView(Mip.RawData), 256 << 10); }));
 
 			//raw_bytes_hashed += Mip.RawData.Num();
 
diff --git a/Engine/Source/Runtime/Core/Private/Hash/xxhash.cpp b/Engine/Source/Runtime/Core/Private/Hash/xxhash.cpp
--- a/Engine/Source/Runtime/Core/Private/Hash/xxhash.cpp
+++ b/Engine/Source/Runtime/Core/Private/Hash/xxhash.cpp
@@ -4,17 +4,42 @@
 
 #include "Async/ParallelFor.h"
 #include "Containers/ContainersFwd.h"
+#include "Containers/StaticArray.h"
 #include "Memory/CompositeBuffer.h"
 #include "Memory/MemoryView.h"
 #include "Memory/SharedBuffer.h"
-#include "Serialization/CompactBinaryWriter.h"
-#include "Tasks/Task.h"
 
 THIRD_PARTY_INCLUDES_START
 #define XXH_INLINE_ALL
 #include "ThirdParty/xxhash/xxhash.h"
 THIRD_PARTY_INCLUDES_END
 
+template <typename HashType>
+static HashType PrivateHashBufferChunked(const FMemoryView View, const uint64 ChunkSize)
+{
+	if (View.GetSize() <= ChunkSize || ChunkSize == 0)
+	{
+		return HashType::HashBuffer(View);
+	}
+
+	// Due to a limitation of ParallelFor, input size is limited to MAX_int32 * ChunkSize.
+	const int32 ChunkCount = IntCastChecked<int32>((View.GetSize() + ChunkSize - 1) / ChunkSize);
+
+	TArray<TStaticArray<uint8, sizeof(HashType)>, TInlineAllocator<64>> ChunkHashes;
+	ChunkHashes.AddZeroed(ChunkCount + 1);
+
+	ParallelFor(TEXT("XxHash.PF"), ChunkCount, 1, [View, ChunkSize, &ChunkHashes](const int32 Index)
+	{
+		FMemoryView ChunkView = View.Mid(Index * ChunkSize, ChunkSize);
+		HashType::HashBuffer(ChunkView).ToByteArray((uint8(&)[sizeof(HashType)])ChunkHashes[Index]);
+	});
+
+	static_assert(sizeof(uint64) <= sizeof(HashType));
+	FPlatformMemory::WriteUnaligned(&ChunkHashes.Last(), INTEL_ORDER64(View.GetSize()));
+
+	return HashType::HashBuffer(ChunkHashes.GetData(), ChunkCount * sizeof(HashType) + sizeof(uint64));
+}
+
 FXxHash64 FXxHash64::HashBuffer(FMemoryView View)
 {
 	return {XXH3_64bits(View.GetData(), View.GetSize())};
@@ -32,53 +57,16 @@ FXxHash64 FXxHash64::HashBuffer(const FCompositeBuffer& Buffer)
 	return Builder.Finalize();
 }
 
-
-/* static */ UE::Tasks::TTask<FXxHash64> FXxHash64::HashBufferChunkedAsync(FMemoryView View, uint64 ChunkSize)
+FXxHash64 FXxHash64::HashBufferChunked(const FMemoryView View, const uint64 ChunkSize)
 {
-	if (View.GetSize() <= ChunkSize || ChunkSize == 0)
-	{
-		return UE::Tasks::Launch(TEXT("XxHash64.HashBufferChunkedAsync.Single"), [View]()
-		{
-			return FXxHash64::HashBuffer(View);
-		});
-	}
-
-	return UE::Tasks::Launch(TEXT("XxHash64.HashBufferChunkedAsync.Multi"), [View, ChunkSize]()
-	{
-		uint64 ChunkCount = (View.GetSize() + ChunkSize - 1) / ChunkSize;
-
-		// Due to limitations of ParallelFor, we can't hash anything larger than int32_max * chunksize;
-		int32 ChunkCount32 = IntCastChecked<int32>(ChunkCount);
-
-		TArray<FXxHash64, TInlineAllocator<64>> ChunkHashes;
-		ChunkHashes.AddDefaulted(ChunkCount32);
-
-		ParallelFor(TEXT("XxHash64.PF"), ChunkCount32, 1, [View, ChunkSize, &ChunkHashes](int32 Index)
-		{
-			FMemoryView ChunkView = View.Mid(Index * ChunkSize, ChunkSize);
-			ChunkHashes[Index] = FXxHash64::HashBuffer(ChunkView);
-		});
-
-		FXxHash64Builder Accumulator;
-		for (FXxHash64& ChunkHash : ChunkHashes)
-		{
-			uint8 HashBytes[8];
-			ChunkHash.ToByteArray(HashBytes);
-			Accumulator.Update(HashBytes, sizeof(HashBytes));
-		}
-		uint64 ViewSize = View.GetSize();
-		Accumulator.Update(&ViewSize, sizeof(ViewSize));
-
-		return Accumulator.Finalize();
-	});
+	return PrivateHashBufferChunked<FXxHash64>(View, ChunkSize);
 }
 
-/* static */ UE::Tasks::TTask<FXxHash64> FXxHash64::HashBufferChunkedAsync(const void* Data, uint64 Size, uint64 ChunkSize)
+FXxHash64 FXxHash64::HashBufferChunked(const void* Data, uint64 Size, uint64 ChunkSize)
 {
-	return HashBufferChunkedAsync(FMemoryView(Data, Size), ChunkSize);
+	return PrivateHashBufferChunked<FXxHash64>(FMemoryView(Data, Size), ChunkSize);
 }
 
-
 FXxHash128 FXxHash128::HashBuffer(FMemoryView View)
 {
 	const XXH128_hash_t Value = XXH3_128bits(View.GetData(), View.GetSize());
@@ -91,53 +79,16 @@ FXxHash128 FXxHash128::HashBuffer(const void* Data, uint64 Size)
 	return {Value.low64, Value.high64};
 }
 
-/* static */ UE::Tasks::TTask<FXxHash128> FXxHash128::HashBufferChunkedAsync(FMemoryView View, uint64 ChunkSize)
+FXxHash128 FXxHash128::HashBufferChunked(const FMemoryView View, const uint64 ChunkSize)
 {
-	if (View.GetSize() <= ChunkSize || ChunkSize == 0)
-	{
-		return UE::Tasks::Launch(TEXT("XxHash128.HashBufferChunkedAsync.Single"), [View]()
-		{
-			return FXxHash128::HashBuffer(View);
-		});
-	}
-
-	return UE::Tasks::Launch(TEXT("XxHash128.HashBufferChunkedAsync.Multi"), [View, ChunkSize]()
-	{
-		uint64 ChunkCount = (View.GetSize() + ChunkSize - 1) / ChunkSize;
-		
-		// Due to limitations of ParallelFor, we can't hash anything larger than int32_max * chunksize;
-		int32 ChunkCount32 = IntCastChecked<int32>(ChunkCount);
-
-		TArray<FXxHash128, TInlineAllocator<64>> ChunkHashes;
-		ChunkHashes.AddDefaulted(ChunkCount32);
-
-		ParallelFor(TEXT("XxHash128.PF"), ChunkCount32, 1, [View, ChunkSize, &ChunkHashes](int32 Index)
-		{
-			FMemoryView ChunkView = View.Mid(Index * ChunkSize, ChunkSize);
-			ChunkHashes[Index] = FXxHash128::HashBuffer(ChunkView);
-		});
-
-		FXxHash128Builder Accumulator;
-		for (FXxHash128& ChunkHash : ChunkHashes)
-		{
-			uint8 HashBytes[16];
-			ChunkHash.ToByteArray(HashBytes);
-			Accumulator.Update(HashBytes, sizeof(HashBytes));
-		}
-		uint64 ViewSize = View.GetSize();
-		Accumulator.Update(&ViewSize, sizeof(ViewSize));
-		
-		return Accumulator.Finalize();
-	});
+	return PrivateHashBufferChunked<FXxHash128>(View, ChunkSize);
 }
 
-/* static */ UE::Tasks::TTask<FXxHash128> FXxHash128::HashBufferChunkedAsync(const void* Data, uint64 Size, uint64 ChunkSize)
+FXxHash128 FXxHash128::HashBufferChunked(const void* Data, uint64 Size, uint64 ChunkSize)
 {
-	return HashBufferChunkedAsync(FMemoryView(Data, Size), ChunkSize);
+	return PrivateHashBufferChunked<FXxHash128>(FMemoryView(Data, Size), ChunkSize);
 }
 
-
-
 FXxHash128 FXxHash128::HashBuffer(const FCompositeBuffer& Buffer)
 {
 	FXxHash128Builder Builder;
diff --git a/Engine/Source/Runtime/Core/Public/Hash/xxhash.h b/Engine/Source/Runtime/Core/Public/Hash/xxhash.h
--- a/Engine/Source/Runtime/Core/Public/Hash/xxhash.h
+++ b/Engine/Source/Runtime/Core/Public/Hash/xxhash.h
@@ -13,26 +13,32 @@
 
 class FCompositeBuffer;
 template <typename CharType> class TStringBuilderBase;
-namespace UE::Tasks 
-{
-	template<typename ResultType>
-	class TTask;
-}
 
 /** A 64-bit hash from XXH3. */
 struct FXxHash64
 {
+	/**
+	 * The hash in its native representation.
+	 *
+	 * Use the canonical representation from ToByteArray to serialize or display the hash.
+	 */
+	uint64 Hash{};
+
 	[[nodiscard]] CORE_API static FXxHash64 HashBuffer(FMemoryView View);
 	[[nodiscard]] CORE_API static FXxHash64 HashBuffer(const void* Data, uint64 Size);
 	[[nodiscard]] CORE_API static FXxHash64 HashBuffer(const FCompositeBuffer& Buffer);
 
-	// Kick off a task to hash the given memory. Buffers are hashed in ChunkSize independent pieces, with the results
-	// hashed together in to a final result. ChunkSize should be fairly large to cover the cost of dispatching tasks. (e.g. 256kb+)
-	[[nodiscard]] CORE_API static UE::Tasks::TTask<FXxHash64> HashBufferChunkedAsync(FMemoryView View, uint64 ChunkSize);
-	[[nodiscard]] CORE_API static UE::Tasks::TTask<FXxHash64> HashBufferChunkedAsync(const void* Data, uint64 Size, uint64 ChunkSize);
+	/**
+	 * Hash the buffer in parallel in independent chunks, and hash those hashes.
+	 *
+	 * Use a ChunkSize large enough to cover task overhead, e.g., 256+ KiB.
+	 * Hashing the same buffer with different chunk sizes produces a different output hash.
+	 */
+	[[nodiscard]] CORE_API static FXxHash64 HashBufferChunked(FMemoryView View, uint64 ChunkSize);
+	[[nodiscard]] CORE_API static FXxHash64 HashBufferChunked(const void* Data, uint64 Size, uint64 ChunkSize);
 
 	/** Load the hash from its canonical (big-endian) representation. */
-	static inline FXxHash64 FromByteArray(uint8 (&Bytes)[sizeof(uint64)])
+	static inline FXxHash64 FromByteArray(const uint8 (&Bytes)[sizeof(uint64)])
 	{
 		uint64 HashBigEndian;
 		FMemory::Memcpy(&HashBigEndian, Bytes, sizeof(uint64));
@@ -84,31 +90,39 @@ struct FXxHash64
 		UE::String::BytesToHexLower(MakeArrayView(Bytes), Builder);
 		return Builder;
 	}
+};
 
-
+/** A 128-bit hash from XXH128. */
+struct FXxHash128
+{
 	/**
-	 * The hash in its native representation.
+	 * The low 64 bits of the hash in its native representation.
 	 *
 	 * Use the canonical representation from ToByteArray to serialize or display the hash.
 	 */
-	uint64 Hash{};
-};
+	uint64 HashLow{};
+	/**
+	 * The high 64 bits of the hash in its native representation.
+	 *
+	 * Use the canonical representation from ToByteArray to serialize or display the hash.
+	 */
+	uint64 HashHigh{};
 
-/** A 128-bit hash from XXH3. */
-struct FXxHash128
-{
-public:
 	[[nodiscard]] CORE_API static FXxHash128 HashBuffer(FMemoryView View);
 	[[nodiscard]] CORE_API static FXxHash128 HashBuffer(const void* Data, uint64 Size);
 	[[nodiscard]] CORE_API static FXxHash128 HashBuffer(const FCompositeBuffer& Buffer);
 
-	// Kick off a task to hash the given memory. Buffers are hashed in ChunkSize independent pieces, with the results
-	// hashed together in to a final result. ChunkSize should be fairly large to cover the cost of dispatching tasks. (e.g. 256kb+)
-	[[nodiscard]] CORE_API static UE::Tasks::TTask<FXxHash128> HashBufferChunkedAsync(FMemoryView View, uint64 ChunkSize);
-	[[nodiscard]] CORE_API static UE::Tasks::TTask<FXxHash128> HashBufferChunkedAsync(const void* Data, uint64 Size, uint64 ChunkSize);
+	/**
+	 * Hash the buffer in parallel in independent chunks, and hash those hashes.
+	 *
+	 * Use a ChunkSize large enough to cover task overhead, e.g., 256+ KiB.
+	 * Hashing the same buffer with different chunk sizes produces a different output hash.
+	 */
+	[[nodiscard]] CORE_API static FXxHash128 HashBufferChunked(FMemoryView View, uint64 ChunkSize);
+	[[nodiscard]] CORE_API static FXxHash128 HashBufferChunked(const void* Data, uint64 Size, uint64 ChunkSize);
 
 	/** Load the hash from its canonical (big-endian) representation. */
-	static inline FXxHash128 FromByteArray(uint8 (&Bytes)[sizeof(uint64[2])])
+	static inline FXxHash128 FromByteArray(const uint8 (&Bytes)[sizeof(uint64[2])])
 	{
 		uint64 HashBigEndian[2];
 		FMemory::Memcpy(&HashBigEndian, Bytes, sizeof(uint64[2]));
@@ -122,19 +136,6 @@ public:
 		FMemory::Memcpy(Bytes, &HashBigEndian, sizeof(uint64[2]));
 	}
 
-	/**
-	 * The low 64 bits of the hash in its native representation.
-	 *
-	 * Use the canonical representation from ToByteArray to serialize or display the hash.
-	 */
-	uint64 HashLow{};
-	/**
-	 * The high 64 bits of the hash in its native representation.
-	 *
-	 * Use the canonical representation from ToByteArray to serialize or display the hash.
-	 */
-	uint64 HashHigh{};
-
 	inline bool operator==(const FXxHash128& B) const
 	{
 		return HashLow == B.HashLow && HashHigh == B.HashHigh;
@@ -196,7 +197,7 @@ private:
 	alignas(64) char StateBytes[576];
 };
 
-/** Calculates a 128-bit hash with XXH3. */
+/** Calculates a 128-bit hash with XXH128. */
 class FXxHash128Builder
 {
 public:
