Possible fix for missing hitchhunter server stacks

diff --git a/Engine/Source/Runtime/Core/Private/Unix/UnixSignalHeartBeat.cpp b/Engine/Source/Runtime/Core/Private/Unix/UnixSignalHeartBeat.cpp
--- a/Engine/Source/Runtime/Core/Private/Unix/UnixSignalHeartBeat.cpp
+++ b/Engine/Source/Runtime/Core/Private/Unix/UnixSignalHeartBeat.cpp
@@ -83,10 +83,11 @@ FUnixSignalGameHitchHeartBeat::FUnixSignalGameHitchHeartBeat()
 
 FUnixSignalGameHitchHeartBeat::~FUnixSignalGameHitchHeartBeat()
 {
-	if (TimerId)
+	if (bTimerCreated)
 	{
 		timer_delete(TimerId);
-		TimerId = nullptr;
+		bTimerCreated = false;
+		TimerId = 0;
 	}
 }
 
@@ -106,12 +107,11 @@ void FUnixSignalGameHitchHeartBeat::Init()
 	SignalEvent.sigev_notify = SIGEV_SIGNAL;
 	SignalEvent.sigev_signo = HEART_BEAT_SIGNAL;
 
-	if (timer_create(CLOCK_REALTIME, &SignalEvent, &TimerId) == -1)
+	bTimerCreated = timer_create(CLOCK_REALTIME, &SignalEvent, &TimerId) == 0;
+	if (!bTimerCreated)
 	{
 		int Errno = errno;
 		UE_LOG(LogUnixHeartBeat, Warning, TEXT("Failed to timer_create() errno=%d (%s)"), Errno, UTF8_TO_TCHAR(strerror(Errno)));
-
-		TimerId = nullptr;
 	}
 
 	float CmdLine_HitchDurationS = 0.0;
@@ -174,11 +174,14 @@ void FUnixSignalGameHitchHeartBeat::FrameStart(bool bSkipThisFrame)
 #if USE_HITCH_DETECTION
 	check(IsInGameThread());
 
-	UE_LOG(LogUnixHeartBeat, VeryVerbose, TEXT("bDisabled:%s SuspendCount:%d TimerId:0x%x timer:%d"),
-		bDisabled ? TEXT("true") : TEXT("false"), SuspendCount, TimerId, 
-		TimerId ? TimerGetTime(TimerId).tv_nsec : 0);
+	UE_LOG(LogUnixHeartBeat, VeryVerbose, TEXT("bDisabled:%s SuspendCount:%d bTimerCreated:%s TimerId:%d timer:%d"),
+		bDisabled ? TEXT("true") : TEXT("false"), 
+		SuspendCount, 
+		bTimerCreated ? TEXT("true") : TEXT("false"),
+		TimerId, 
+		bTimerCreated ? TimerGetTime(TimerId).tv_nsec : 0);
 
-	if (!bDisabled && SuspendCount == 0 && TimerId)
+	if (!bDisabled && SuspendCount == 0 && bTimerCreated)
 	{
 		UE_LOG(LogUnixHeartBeat, VeryVerbose, TEXT("HitchThresholdS:%f MinimalHitchThreashold:%f"),
 			HitchThresholdS, MinimalHitchThreashold);
@@ -234,7 +237,7 @@ void FUnixSignalGameHitchHeartBeat::SuspendHeartBeat()
 	SuspendCount++;
 	UE_LOG(LogUnixHeartBeat, Verbose, TEXT("SuspendCount:%d"), SuspendCount);
 
-	if (TimerId)
+	if (bTimerCreated)
 	{
 		struct itimerspec DisarmTime;
 		FMemory::Memzero(DisarmTime);
@@ -272,10 +275,11 @@ void FUnixSignalGameHitchHeartBeat::Restart()
 	bDisabled = false;
 
 	// If we still have a valid handle on the timer_t clean it up
-	if (TimerId)
+	if (bTimerCreated)
 	{
 		timer_delete(TimerId);
-		TimerId = nullptr;
+		bTimerCreated = false;
+		TimerId = 0;
 	}
 
 	Init();
@@ -287,3 +291,11 @@ void FUnixSignalGameHitchHeartBeat::Stop()
 	SuspendHeartBeat();
 	bDisabled = true;
 }
+
+void FUnixSignalGameHitchHeartBeat::PostFork()
+{
+	UE_LOG(LogUnixHeartBeat, Verbose, TEXT("PostFork"));
+	// timers aren't inherited by child processes
+	bTimerCreated = false;
+	TimerId = 0;
+}
\ No newline at end of file
diff --git a/Engine/Source/Runtime/Core/Public/HAL/ThreadHeartBeat.h b/Engine/Source/Runtime/Core/Public/HAL/ThreadHeartBeat.h
--- a/Engine/Source/Runtime/Core/Public/HAL/ThreadHeartBeat.h
+++ b/Engine/Source/Runtime/Core/Public/HAL/ThreadHeartBeat.h
@@ -399,6 +399,7 @@ public:
 
 	// No-op, used in FUnixSignalGameHitchHeartBeat
 	void Restart() {}
+	void PostFork() {}
 
 	//~ Begin FRunnable Interface.
 	virtual bool Init();
diff --git a/Engine/Source/Runtime/Core/Public/Unix/UnixSignalHeartBeat.h b/Engine/Source/Runtime/Core/Public/Unix/UnixSignalHeartBeat.h
--- a/Engine/Source/Runtime/Core/Public/Unix/UnixSignalHeartBeat.h
+++ b/Engine/Source/Runtime/Core/Public/Unix/UnixSignalHeartBeat.h
@@ -36,6 +36,8 @@ public:
 	void Restart();
 	void Stop();
 
+	void PostFork();
+
 private:
 	FUnixSignalGameHitchHeartBeat();
 	~FUnixSignalGameHitchHeartBeat();
@@ -48,5 +50,6 @@ private:
 	bool bHasCmdLine = false;
 	bool bDisabled = false;
 	int32 SuspendCount = 0;
-	timer_t TimerId = nullptr;
+	bool bTimerCreated = false;
+	timer_t TimerId = 0;
 };
