I/O - removed on demand TOC from core since this has moved to new module

#rnx

diff --git a/Engine/Source/Runtime/Core/Internal/IO/IoStore.h b/Engine/Source/Runtime/Core/Internal/IO/IoStore.h
--- a/Engine/Source/Runtime/Core/Internal/IO/IoStore.h
+++ b/Engine/Source/Runtime/Core/Internal/IO/IoStore.h
@@ -16,9 +16,6 @@
 
 CORE_API DECLARE_LOG_CATEGORY_EXTERN(LogIoStore, Log, All);
 
-class FCbWriter;
-class FCbFieldView;
-
 /**
  * I/O store container format version
  */
@@ -206,65 +203,3 @@ struct FIoStoreTocResource
 
 	CORE_API static uint64 HashChunkIdWithSeed(int32 Seed, const FIoChunkId& ChunkId);
 };
-
-enum class EIoOnDemandTocVersion : uint32
-{
-	Invalid			= 0,
-	Initial			= 1,
-
-	LatestPlusOne,
-	Latest			= (LatestPlusOne - 1)
-};
-
-struct FIoStoreOnDemandTocHeader
-{
-	static constexpr uint64 ExpectedMagic = 0x6f6e64656d616e64; // ondemand
-
-	uint64 Magic = ExpectedMagic;
-	uint32 Version = uint32(EIoOnDemandTocVersion::Latest);
-	uint32 BlockSize = 0;
-	FString CompressionFormat;
-	FString ChunksDirectory;
-	
-	CORE_API friend FCbWriter& operator<<(FCbWriter& Writer, const FIoStoreOnDemandTocHeader& Header);
-};
-
-CORE_API bool LoadFromCompactBinary(FCbFieldView Field, FIoStoreOnDemandTocHeader& OutTocHeader);
-
-struct FIoStoreOnDemandTocEntry
-{
-	FIoHash Hash = FIoHash::Zero;
-	FIoChunkId ChunkId = FIoChunkId::InvalidChunkId;
-	uint64 RawSize = 0;
-	uint64 EncodedSize = 0;
-	uint32 BlockOffset = ~uint32(0);
-	uint32 BlockCount = 0; 
-	
-	CORE_API friend FCbWriter& operator<<(FCbWriter& Writer, const FIoStoreOnDemandTocEntry& Entry);
-};
-
-CORE_API bool LoadFromCompactBinary(FCbFieldView Field, FIoStoreOnDemandTocEntry& OutTocEntry);
-
-struct FIoStoreOnDemandContainerEntry
-{
-	FString ContainerName;
-	FString EncryptionKeyGuid;
-	TArray<FIoStoreOnDemandTocEntry> Entries;
-	TArray<uint32> BlockSizes;
-
-	CORE_API friend FCbWriter& operator<<(FCbWriter& Writer, const FIoStoreOnDemandContainerEntry& ContainerEntry);
-};
-
-CORE_API bool LoadFromCompactBinary(FCbFieldView Field, FIoStoreOnDemandContainerEntry& OutContainer);
-
-struct FIoStoreOndemandTocResource
-{
-	FIoStoreOnDemandTocHeader Header;
-	TArray<FIoStoreOnDemandContainerEntry> Containers;
-	
-	CORE_API friend FCbWriter& operator<<(FCbWriter& Writer, const FIoStoreOndemandTocResource& TocResource);
-
-	UE_NODISCARD CORE_API static TIoStatusOr<FString> Save(const TCHAR* Directory, const FIoStoreOndemandTocResource& TocResource);
-};
-
-CORE_API bool LoadFromCompactBinary(FCbFieldView Field, FIoStoreOndemandTocResource& OutToc);
diff --git a/Engine/Source/Runtime/Core/Private/IO/IoStore.cpp b/Engine/Source/Runtime/Core/Private/IO/IoStore.cpp
--- a/Engine/Source/Runtime/Core/Private/IO/IoStore.cpp
+++ b/Engine/Source/Runtime/Core/Private/IO/IoStore.cpp
@@ -2,7 +2,6 @@
 
 #include "IO/IoStore.h"
 #include "Containers/Map.h"
-#include "Containers/StringView.h"
 #include "HAL/FileManager.h"
 #include "Templates/UniquePtr.h"
 #include "Misc/Paths.h"
@@ -22,10 +21,7 @@
 #include "Misc/ConfigCacheIni.h"
 #include "DerivedDataPluginInterface.h"
 #include "DerivedDataCacheInterface.h"
-#include "Serialization/CompactBinaryWriter.h"
-#include "Serialization/CompactBinarySerialization.h"
 #include "Serialization/LargeMemoryReader.h"
-#include "Serialization/LargeMemoryWriter.h"
 #include "Serialization/MemoryWriter.h"
 #include "ProfilingDebugging/CountersTrace.h"
 #include "Tasks/Task.h"
@@ -3310,203 +3306,3 @@ void FIoStoreReader::GetFilenamesByBlockIndex(const TArray<int32>& InBlockIndexL
 			return true;
 		});
 }
-
-////////////////////////////////////////////////////////////////////////////////
-
-FCbWriter& operator<<(FCbWriter& Writer, const FIoStoreOnDemandTocHeader& Header)
-{
-	Writer.BeginObject();
-	Writer.AddInteger(UTF8TEXTVIEW("Magic"), Header.Magic);
-	Writer.AddInteger(UTF8TEXTVIEW("Version"), Header.Version);
-	Writer.AddInteger(UTF8TEXTVIEW("BlockSize"), Header.BlockSize);
-	Writer.AddString(UTF8TEXTVIEW("CompressionFormat"), Header.CompressionFormat);
-	Writer.AddString(UTF8TEXTVIEW("ChunksDirectory"), Header.ChunksDirectory);
-	Writer.EndObject();
-
-	return Writer;
-}
-
-bool LoadFromCompactBinary(FCbFieldView Field, FIoStoreOnDemandTocHeader& OutTocHeader)
-{
-	if (FCbObjectView Obj = Field.AsObjectView())
-	{
-		OutTocHeader.Magic = Obj["Magic"].AsUInt64();
-		OutTocHeader.Version = Obj["Version"].AsUInt32();
-		OutTocHeader.BlockSize = Obj["BlockSize"].AsUInt32();
-		OutTocHeader.CompressionFormat = FString(Obj["CompressionFormat"].AsString());
-		OutTocHeader.ChunksDirectory = FString(Obj["ChunksDirectory"].AsString());
-
-		return OutTocHeader.Magic == FIoStoreOnDemandTocHeader::ExpectedMagic &&
-			static_cast<EIoOnDemandTocVersion>(OutTocHeader.Version) != EIoOnDemandTocVersion::Invalid;
-	}
-
-	return false;
-}
-
-FCbWriter& operator<<(FCbWriter& Writer, const FIoStoreOnDemandTocEntry& Entry)
-{
-	Writer.BeginObject();
-	Writer.AddHash(UTF8TEXTVIEW("Hash"), Entry.Hash);
-	Writer << UTF8TEXTVIEW("ChunkId") << Entry.ChunkId;
-	Writer.AddInteger(UTF8TEXTVIEW("RawSize"), Entry.RawSize);
-	Writer.AddInteger(UTF8TEXTVIEW("EncodedSize"), Entry.EncodedSize);
-	Writer.AddInteger(UTF8TEXTVIEW("BlockOffset"), Entry.BlockOffset);
-	Writer.AddInteger(UTF8TEXTVIEW("BlockCount"), Entry.BlockCount);
-	Writer.EndObject();
-
-	return Writer;
-}
-
-bool LoadFromCompactBinary(FCbFieldView Field, FIoStoreOnDemandTocEntry& OutTocEntry)
-{
-	if (FCbObjectView Obj = Field.AsObjectView())
-	{
-		if (!LoadFromCompactBinary(Obj["ChunkId"], OutTocEntry.ChunkId))
-		{
-			return false;
-		}
-
-		OutTocEntry.Hash = Obj["Hash"].AsHash();
-		OutTocEntry.RawSize = Obj["RawSize"].AsUInt64(~uint64(0));
-		OutTocEntry.EncodedSize = Obj["EncodedSize"].AsUInt64(~uint64(0));
-		OutTocEntry.BlockOffset = Obj["BlockOffset"].AsUInt32(~uint32(0));
-		OutTocEntry.BlockCount = Obj["BlockCount"].AsUInt32();
-
-		return OutTocEntry.Hash != FIoHash::Zero &&
-			OutTocEntry.RawSize != ~uint64(0) &&
-			OutTocEntry.EncodedSize != ~uint64(0) &&
-			OutTocEntry.BlockOffset != ~uint32(0);
-	}
-
-	return false;
-}
-
-FCbWriter& operator<<(FCbWriter& Writer, const FIoStoreOnDemandContainerEntry& ContainerEntry)
-{
-	Writer.BeginObject();
-	Writer.AddString(UTF8TEXTVIEW("Name"), ContainerEntry.ContainerName);
-	Writer.AddString(UTF8TEXTVIEW("EncryptionKeyGuid"), ContainerEntry.EncryptionKeyGuid);
-
-	Writer.BeginArray(UTF8TEXTVIEW("Entries"));
-	for (const FIoStoreOnDemandTocEntry& Entry : ContainerEntry.Entries)
-	{
-		Writer << Entry;
-	}
-	Writer.EndArray();
-	
-	Writer.BeginArray(UTF8TEXTVIEW("BlockSizes"));
-	for (uint32 BlockSize : ContainerEntry.BlockSizes)
-	{
-		Writer << BlockSize;
-	}
-	Writer.EndArray();
-
-	Writer.EndObject();
-
-	return Writer;
-}
-
-bool LoadFromCompactBinary(FCbFieldView Field, FIoStoreOnDemandContainerEntry& OutContainer)
-{
-	if (FCbObjectView Obj = Field.AsObjectView())
-	{
-		OutContainer.ContainerName = FString(Obj["Name"].AsString());
-		OutContainer.EncryptionKeyGuid = FString(Obj["EncryptionKeyGuid"].AsString());
-
-		FCbArrayView Entries = Obj["Entries"].AsArrayView();
-		OutContainer.Entries.Reserve(int32(Entries.Num()));
-		for (FCbFieldView ArrayField : Entries)
-		{
-			if (!LoadFromCompactBinary(ArrayField, OutContainer.Entries.AddDefaulted_GetRef()))
-			{
-				return false;
-			}
-		}
-
-		FCbArrayView BlockSizes = Obj["BlockSizes"].AsArrayView();
-		OutContainer.BlockSizes.Reserve(int32(BlockSizes.Num()));
-		for (FCbFieldView ArrayField : BlockSizes)
-		{
-			OutContainer.BlockSizes.Add(ArrayField.AsUInt32());
-		}
-
-		return true;
-	}
-
-	return false;
-}
-
-FCbWriter& operator<<(FCbWriter& Writer, const FIoStoreOndemandTocResource& TocResource)
-{
-	Writer.BeginObject();
-	Writer << UTF8TEXTVIEW("Header") << TocResource.Header;
-
-	Writer.BeginArray(UTF8TEXTVIEW("Containers"));
-	for (const FIoStoreOnDemandContainerEntry& Container : TocResource.Containers)
-	{
-		Writer << Container;
-	}
-	Writer.EndArray();
-	Writer.EndObject();
-	
-	return Writer;
-}
-
-TIoStatusOr<FString> FIoStoreOndemandTocResource::Save(const TCHAR* Directory, const FIoStoreOndemandTocResource& TocResource)
-{
-	if (TocResource.Header.Magic != FIoStoreOnDemandTocHeader::ExpectedMagic)
-	{
-		return FIoStatus(EIoErrorCode::CorruptToc);
-	}
-
-	if (TocResource.Header.CompressionFormat.IsEmpty())
-	{
-		return FIoStatus(EIoErrorCode::CorruptToc);
-	}
-
-	FCbWriter Writer;
-	Writer << TocResource;
-
-	FLargeMemoryWriter Ar;
-	SaveCompactBinary(Ar, Writer.Save());
-
-	const FIoHash TocHash = FIoHash::HashBuffer(Ar.GetView());
-	const FString FilePath = FString(Directory) / LexToString(TocHash) + TEXT(".iochunktoc");
-
-	IFileManager& FileMgr = IFileManager::Get();
-	if (TUniquePtr<FArchive> FileAr(FileMgr.CreateFileWriter(*FilePath)); FileAr.IsValid())
-	{
-		FileAr->Serialize(Ar.GetData(), Ar.TotalSize());
-		FileAr->Flush();
-		FileAr->Close();
-
-		return FilePath;
-	}
-
-	return FIoStatus(EIoErrorCode::WriteError);
-}
-
-bool LoadFromCompactBinary(FCbFieldView Field, FIoStoreOndemandTocResource& OutToc)
-{
-	if (FCbObjectView Obj = Field.AsObjectView())
-	{
-		if (!LoadFromCompactBinary(Obj["Header"], OutToc.Header))
-		{
-			return false;
-		}
-
-		FCbArrayView Containers = Obj["Containers"].AsArrayView();
-		OutToc.Containers.Reserve(int32(Containers.Num()));
-		for (FCbFieldView ArrayField : Containers)
-		{
-			if (!LoadFromCompactBinary(ArrayField, OutToc.Containers.AddDefaulted_GetRef()))
-			{
-				return false;
-			}
-		}
-
-		return true;
-	}
-
-	return false;
-}
