MTAccessDetector: changed thread_local members to thread_local vars in dedicated getters to workaround clang complaining about duplicated initialisation (despite it was marked as `inline`)

diff --git a/Engine/Source/Runtime/Core/Public/Misc/MTAccessDetector.h b/Engine/Source/Runtime/Core/Public/Misc/MTAccessDetector.h
--- a/Engine/Source/Runtime/Core/Public/Misc/MTAccessDetector.h
+++ b/Engine/Source/Runtime/Core/Public/Misc/MTAccessDetector.h
@@ -369,7 +369,12 @@ public:
 private:
 	// a TLS stack of destruction sentinels shared by all access detector instances on the callstack. the same access detector instance can have
 	// multiple destruction sentinels on the stack. the next "release access" with destruction sentinel always matches the top of the stack
-	inline static thread_local TArray<FDestructionSentinel*, TInlineAllocator<4>> DestructionSentinelStackTls;
+	using FDestructionSentinelStackTls = TArray<FDestructionSentinel*, TInlineAllocator<4>>;
+	static FDestructionSentinelStackTls& GetDestructionSentinelStackTls()
+	{
+		thread_local FDestructionSentinelStackTls DestructionSentinelStackTls;
+		return DestructionSentinelStackTls;
+	}
 //////////////////////////////////////////////////////////////////////
 
 private:
@@ -383,18 +388,23 @@ private:
 		uint32 Num;
 	};
 
-	inline static thread_local TArray<FReaderNum, TInlineAllocator<4>> ReadersTls;
+	using FReadersTls = TArray<FReaderNum, TInlineAllocator<4>>;
+	static FReadersTls& GetReadersTls()
+	{
+		thread_local FReadersTls ReadersTls;
+		return ReadersTls;
+	}
 
 	void RemoveReaderFromTls() const
 	{
-		uint32 ReaderIndex = ReadersTls.IndexOfByPredicate([this](FReaderNum ReaderNum) { return ReaderNum.Reader == this; });
+		uint32 ReaderIndex = GetReadersTls().IndexOfByPredicate([this](FReaderNum ReaderNum) { return ReaderNum.Reader == this; });
 		checkfSlow(ReaderIndex != INDEX_NONE,
 			TEXT("Invalid usage of the race detector! No matching AcquireReadAccess(): %u readers, %u writers on thread %d"),
 			LoadState().ReaderNum, LoadState().WriterNum, LoadState().WriterThreadId - 1);
-		uint32 ReaderNum = --ReadersTls[ReaderIndex].Num;
+		uint32 ReaderNum = --GetReadersTls()[ReaderIndex].Num;
 		if (ReaderNum == 0)
 		{
-			ReadersTls.RemoveAtSwap(ReaderIndex, 1, /* bAllowShrinking = */ false);
+			GetReadersTls().RemoveAtSwap(ReaderIndex, 1, /* bAllowShrinking = */ false);
 		}
 	}
 	///////////////////////////////////////////////
@@ -422,9 +432,9 @@ public:
 		// search for all destruction sentinels for this instance and remove them from the stack, while building an expected correct state of
 		// the access detector.
 		FState ExpectedState;
-		for (int DestructionSentinelIndex = 0; DestructionSentinelIndex != DestructionSentinelStackTls.Num(); ++DestructionSentinelIndex)
+		for (int DestructionSentinelIndex = 0; DestructionSentinelIndex != GetDestructionSentinelStackTls().Num(); ++DestructionSentinelIndex)
 		{
-			FDestructionSentinel& DestructionSentinel = *DestructionSentinelStackTls[DestructionSentinelIndex];
+			FDestructionSentinel& DestructionSentinel = *GetDestructionSentinelStackTls()[DestructionSentinelIndex];
 			if (DestructionSentinel.Accessor == this)
 			{
 				DestructionSentinel.bDestroyed = true;
@@ -440,7 +450,7 @@ public:
 					ExpectedState.WriterThreadId = FPlatformTLS::GetCurrentThreadId() + 1; // to shift from 0 TID that is considered "invalid"
 				}
 
-				DestructionSentinelStackTls.RemoveAtSwap(DestructionSentinelIndex);
+				GetDestructionSentinelStackTls().RemoveAtSwap(DestructionSentinelIndex);
 				--DestructionSentinelIndex;
 			}
 		}
@@ -461,14 +471,14 @@ public:
 		// register the reader in TLS
 		uint32 ReaderIndex = PrevState.ReaderNum == 0 ? 
 			INDEX_NONE : 
-			ReadersTls.IndexOfByPredicate([this](FReaderNum ReaderNum) { return ReaderNum.Reader == this; });
+			GetReadersTls().IndexOfByPredicate([this](FReaderNum ReaderNum) { return ReaderNum.Reader == this; });
 		if (ReaderIndex == INDEX_NONE)
 		{
-			ReadersTls.Add(FReaderNum{ this, 1 });
+			GetReadersTls().Add(FReaderNum{this, 1});
 		}
 		else
 		{
-			++ReadersTls[ReaderIndex].Num;
+			++GetReadersTls()[ReaderIndex].Num;
 		}
 	}
 
@@ -477,7 +487,7 @@ public:
 	FORCEINLINE void AcquireReadAccess(FDestructionSentinel& DestructionSentinel) const
 	{
 		DestructionSentinel.Accessor = this;
-		DestructionSentinelStackTls.Add(&DestructionSentinel);
+		GetDestructionSentinelStackTls().Add(&DestructionSentinel);
 
 		AcquireReadAccess();
 	}
@@ -496,10 +506,10 @@ public:
 		ReleaseReadAccess();
 
 		checkSlow(DestructionSentinel.Accessor == this);
-		checkfSlow(DestructionSentinelStackTls.Num() != 0, TEXT("An attempt to remove a not registered destruction sentinel"));
-		checkfSlow(DestructionSentinelStackTls.Last() == &DestructionSentinel, TEXT("Mismatched destruction sentinel: %p != %p"), DestructionSentinelStackTls.Last(), &DestructionSentinel);
+		checkfSlow(GetDestructionSentinelStackTls().Num() != 0, TEXT("An attempt to remove a not registered destruction sentinel"));
+		checkfSlow(GetDestructionSentinelStackTls().Last() == &DestructionSentinel, TEXT("Mismatched destruction sentinel: %p != %p"), GetDestructionSentinelStackTls().Last(), &DestructionSentinel);
 
-		DestructionSentinelStackTls.RemoveAtSwap(DestructionSentinelStackTls.Num() - 1);
+		GetDestructionSentinelStackTls().RemoveAtSwap(GetDestructionSentinelStackTls().Num() - 1);
 	}
 
 	FORCEINLINE void AcquireWriteAccess()
@@ -508,10 +518,10 @@ public:
 
 		if (LocalState.ReaderNum >= 1)
 		{	// check that all readers are on the current thread
-			uint32 ReaderIndex = ReadersTls.IndexOfByPredicate([this](FReaderNum ReaderNum) { return ReaderNum.Reader == this; });
-			checkf(ReaderIndex != INDEX_NONE && ReadersTls[ReaderIndex].Num == LocalState.ReaderNum, 
+			uint32 ReaderIndex = GetReadersTls().IndexOfByPredicate([this](FReaderNum ReaderNum) { return ReaderNum.Reader == this; });
+			checkf(ReaderIndex != INDEX_NONE && GetReadersTls()[ReaderIndex].Num == LocalState.ReaderNum, 
 				TEXT("Data race detected: %d reader(s) on another thread(s) while acquiring write access"), 
-				LocalState.ReaderNum - ReadersTls[ReaderIndex].Num);
+				LocalState.ReaderNum - GetReadersTls()[ReaderIndex].Num);
 		}
 
 		uint32 CurrentThreadId = FPlatformTLS::GetCurrentThreadId() + 1;
@@ -538,7 +548,7 @@ public:
 	FORCEINLINE void AcquireWriteAccess(FDestructionSentinel& DestructionSentinel)
 	{
 		DestructionSentinel.Accessor = this;
-		DestructionSentinelStackTls.Add(&DestructionSentinel);
+		GetDestructionSentinelStackTls().Add(&DestructionSentinel);
 
 		AcquireWriteAccess();
 	}
@@ -560,15 +570,16 @@ public:
 
 	// an overload that handles access detector destruction from inside a write access, must be used along with the corresponding
 	// overload of `AcquireWriteAcess`
-	FORCEINLINE void ReleaseWriteAccess(FDestructionSentinel& DestructionSentinel)
+	FORCEINLINE_DEBUGGABLE void ReleaseWriteAccess(FDestructionSentinel& DestructionSentinel)
 	{
 		ReleaseWriteAccess();
 
+		FDestructionSentinelStackTls& DestructionSentinelStackTls = GetDestructionSentinelStackTls();
 		checkSlow(DestructionSentinel.Accessor == this);
-		checkfSlow(DestructionSentinelStackTls.Num() != 0, TEXT("An attempt to remove a not registered destruction sentinel"));
-		checkfSlow(DestructionSentinelStackTls.Last() == &DestructionSentinel, TEXT("Mismatched destruction sentinel: %p != %p"), DestructionSentinelStackTls.Last(), &DestructionSentinel);
+		checkfSlow(GetDestructionSentinelStackTls().Num() != 0, TEXT("An attempt to remove a not registered destruction sentinel"));
+		checkfSlow(GetDestructionSentinelStackTls().Last() == &DestructionSentinel, TEXT("Mismatched destruction sentinel: %p != %p"), GetDestructionSentinelStackTls().Last(), &DestructionSentinel);
 
-		DestructionSentinelStackTls.RemoveAtSwap(DestructionSentinelStackTls.Num() - 1);
+		GetDestructionSentinelStackTls().RemoveAtSwap(GetDestructionSentinelStackTls().Num() - 1);
 	}
 };
 
