delegates refactoring:
* added a thread-safe (locked) delegates variant
* refactored and enabled by default race detection for the original not thread-safe version
* added "not thread-safe not checked" variant for rare cases where race detection doesn't work or it's a known case of a race

diff --git a/Engine/Source/Runtime/Core/Private/Misc/CoreDelegates.cpp b/Engine/Source/Runtime/Core/Private/Misc/CoreDelegates.cpp
--- a/Engine/Source/Runtime/Core/Private/Misc/CoreDelegates.cpp
+++ b/Engine/Source/Runtime/Core/Private/Misc/CoreDelegates.cpp
@@ -46,7 +46,7 @@ TMulticastDelegate<void(const FString&)> FCoreDelegates::NewFileAddedDelegate;
 TMulticastDelegate<void()> FCoreDelegates::NoPakFilesMountedDelegate;
 
 PRAGMA_DISABLE_DEPRECATION_WARNINGS
-TMulticastDelegate<void(const TCHAR*, const TCHAR*)> FCoreDelegates::OnFileOpenedForReadFromPakFile;
+TMulticastDelegate<void(const TCHAR*, const TCHAR*), FNotThreadSafeNotCheckedDelegateUserPolicy> FCoreDelegates::OnFileOpenedForReadFromPakFile;
 PRAGMA_ENABLE_DEPRECATION_WARNINGS
 
 TTSMulticastDelegate<void(const TCHAR*, const TCHAR*)>& FCoreDelegates::GetOnFileOpenedForReadFromPakFile()
diff --git a/Engine/Source/Runtime/Core/Public/Delegates/Delegate.h b/Engine/Source/Runtime/Core/Public/Delegates/Delegate.h
--- a/Engine/Source/Runtime/Core/Public/Delegates/Delegate.h
+++ b/Engine/Source/Runtime/Core/Public/Delegates/Delegate.h
@@ -8,7 +8,6 @@
 #include "Templates/SharedPointer.h"
 #include "UObject/WeakObjectPtrTemplates.h"
 #include "Delegates/MulticastDelegateBase.h" // IWYU pragma: export
-#include "Delegates/TSMulticastDelegateBase.h"
 #include "Delegates/IntegerSequence.h" // IWYU pragma: export
 
 /**
@@ -232,7 +231,7 @@
 
 /** Declare user's dynamic delegate, with wrapper proxy method for executing the delegate */
 #define FUNC_DECLARE_DYNAMIC_DELEGATE( DynamicDelegateClassName, ExecFunction, FuncParamList, FuncParamPassThru, ... ) \
-	class DynamicDelegateClassName : public TBaseDynamicDelegate<UE::Core::Private::TScriptDelegateDefault, __VA_ARGS__> \
+	class DynamicDelegateClassName : public TBaseDynamicDelegate<FNotThreadSafeDelegateMode, __VA_ARGS__> \
 	{ \
 	public: \
 		/** Default constructor */ \
@@ -242,7 +241,7 @@
 		\
 		/** Construction from an FScriptDelegate must be explicit.  This is really only used by UObject system internals. */ \
 		explicit DynamicDelegateClassName( const TScriptDelegate<>& InScriptDelegate ) \
-			: TBaseDynamicDelegate<UE::Core::Private::TScriptDelegateDefault, __VA_ARGS__>( InScriptDelegate ) \
+			: TBaseDynamicDelegate<FNotThreadSafeDelegateMode, __VA_ARGS__>( InScriptDelegate ) \
 		{ \
 		} \
 		\
@@ -267,7 +266,7 @@
 
 /** Declare user's dynamic delegate with return value, with wrapper proxy method for executing the delegate */
 #define FUNC_DECLARE_DYNAMIC_DELEGATE_RETVAL(DynamicDelegateRetValClassName, ExecFunction, RetValType, FuncParamList, FuncParamPassThru, ...) \
-	class DynamicDelegateRetValClassName : public TBaseDynamicDelegate<UE::Core::Private::TScriptDelegateDefault, __VA_ARGS__> \
+	class DynamicDelegateRetValClassName : public TBaseDynamicDelegate<FNotThreadSafeDelegateMode, __VA_ARGS__> \
 	{ \
 	public: \
 		/** Default constructor */ \
@@ -277,7 +276,7 @@
 		\
 		/** Construction from an FScriptDelegate must be explicit.  This is really only used by UObject system internals. */ \
 		explicit DynamicDelegateRetValClassName( const TScriptDelegate<>& InScriptDelegate ) \
-			: TBaseDynamicDelegate<UE::Core::Private::TScriptDelegateDefault, __VA_ARGS__>( InScriptDelegate ) \
+			: TBaseDynamicDelegate<FNotThreadSafeDelegateMode, __VA_ARGS__>( InScriptDelegate ) \
 		{ \
 		} \
 		\
@@ -293,7 +292,7 @@
 
 /** Declare user's dynamic multi-cast delegate, with wrapper proxy method for executing the delegate */
 #define FUNC_DECLARE_DYNAMIC_MULTICAST_DELEGATE(DynamicMulticastDelegateClassName, ExecFunction, FuncParamList, FuncParamPassThru, ...) \
-class DynamicMulticastDelegateClassName : public TBaseDynamicMulticastDelegate<UE::Core::Private::TScriptDelegateDefault, __VA_ARGS__> \
+class DynamicMulticastDelegateClassName : public TBaseDynamicMulticastDelegate<FNotThreadSafeDelegateMode, __VA_ARGS__> \
 	{ \
 	public: \
 		/** Default constructor */ \
@@ -303,7 +302,7 @@ class DynamicMulticastDelegateClassName : public TBaseDynamicMulticastDelegate<U
 		\
 		/** Construction from an FMulticastScriptDelegate must be explicit.  This is really only used by UObject system internals. */ \
 		explicit DynamicMulticastDelegateClassName( const TMulticastScriptDelegate<>& InMulticastScriptDelegate ) \
-			: TBaseDynamicMulticastDelegate<UE::Core::Private::TScriptDelegateDefault, __VA_ARGS__>( InMulticastScriptDelegate ) \
+			: TBaseDynamicMulticastDelegate<FNotThreadSafeDelegateMode, __VA_ARGS__>( InMulticastScriptDelegate ) \
 		{ \
 		} \
 		\
diff --git a/Engine/Source/Runtime/Core/Public/Delegates/DelegateAccessHandler.h b/Engine/Source/Runtime/Core/Public/Delegates/DelegateAccessHandler.h
--- a/Engine/Source/Runtime/Core/Public/Delegates/DelegateAccessHandler.h
+++ b/Engine/Source/Runtime/Core/Public/Delegates/DelegateAccessHandler.h
@@ -7,17 +7,30 @@
 #include "Misc/MTAccessDetector.h"
 #include "Misc/ScopeLock.h"
 
-//#define UE_DETECT_DELEGATES_RACE_CONDITIONS 1
-
 #if !defined(UE_DETECT_DELEGATES_RACE_CONDITIONS)
-	#define UE_DETECT_DELEGATES_RACE_CONDITIONS 0
+#define UE_DETECT_DELEGATES_RACE_CONDITIONS ENABLE_MT_DETECTOR
 #endif
 
+/**
+ * types to specify thread safety mode, an enum would be less readable in debugger
+ */
+struct FNotThreadSafeNotCheckedDelegateMode;
+struct FThreadSafeDelegateMode;
+struct FNotThreadSafeDelegateMode;
+
+/**
+ * a template for delegates core thread-safety. is supposed to be used as a base class for "empty base optimisation".
+ * any access to internal data must happen inside an "read access scope" or "write access scope"
+ */
+template<typename ThreadSafetyMode>
+class TDelegateAccessHandlerBase;
+
 /**
  * non thread-safe version that does not do any race detection. supposed to be used in a controlled environment that provides own
- * detection or synchronization.
+ * detection or synchronisation.
  */
-class FDelegateAccessHandlerBaseUnchecked
+template<>
+class TDelegateAccessHandlerBase<FNotThreadSafeNotCheckedDelegateMode>
 {
 protected:
 	struct FReadAccessScope {};
@@ -34,33 +47,80 @@ protected:
 	}
 };
 
+template<>
+struct TIsZeroConstructType<TDelegateAccessHandlerBase<FNotThreadSafeNotCheckedDelegateMode>>
+{
+	static constexpr bool Value = true;
+};
+
 /**
- * non thread-safe version that detects not thread-safe delegates used concurrently (dev builds only)
+ * thread-safe version that locks access to the delegate internals.
+ * we don't have a recursive RW mutex yet, so "read scope" is actually a write scope. This means that concurrent reads of
+ * a thread-safe delegate are mutually exclusive
  */
-class FDelegateAccessHandlerBaseChecked
-#if !UE_DETECT_DELEGATES_RACE_CONDITIONS
-	: public FDelegateAccessHandlerBaseUnchecked
-#endif // !UE_DETECT_DELEGATES_RACE_CONDITIONS
+template<>
+class TDelegateAccessHandlerBase<FThreadSafeDelegateMode>
+{
+protected:
+	struct FReadAccessScope { FScopeLock Lock; };
+	struct FWriteAccessScope { FScopeLock Lock; };
+
+	[[nodiscard]] FReadAccessScope GetReadAccessScope() const
+	{
+		return { &Mutex };
+	}
+
+	[[nodiscard]] FWriteAccessScope GetWriteAccessScope()
+	{
+		return { &Mutex };
+	}
+
+private:
+	mutable FCriticalSection Mutex;
+};
+
+template<>
+struct TIsZeroConstructType<TDelegateAccessHandlerBase<FThreadSafeDelegateMode>>
 {
+	static constexpr bool Value = false;
+};
+
 #if UE_DETECT_DELEGATES_RACE_CONDITIONS
+
+/**
+ * non thread-safe version that detects not thread-safe delegates used concurrently (dev builds only)
+ */
+template<>
+class TDelegateAccessHandlerBase<FNotThreadSafeDelegateMode>
+{
 protected:
 	class FReadAccessScope
 	{
 	public:
 		UE_NONCOPYABLE(FReadAccessScope);
 
-		explicit FReadAccessScope(FRWFullyRecursiveAccessDetector& InAccessDetector)
-			: Accessor(MakeScopedReaderAccessDetector(InAccessDetector))
+		explicit FReadAccessScope(const FMRSWRecursiveAccessDetector& InAccessDetector)
+			: DestructionSentinel(FMRSWRecursiveAccessDetector::EAccessType::Reader)
+		{
+			DestructionSentinel.Accessor = &InAccessDetector;
+			DestructionSentinel.Accessor->AcquireReadAccess(DestructionSentinel);
+		}
+
+		~FReadAccessScope()
 		{
+			if (!DestructionSentinel.bDestroyed) // only if `AccessDetector` wasn't destroyed while being accessed
+			{
+				DestructionSentinel.Accessor->ReleaseReadAccess(DestructionSentinel);
+			}
 		}
 
 	private:
-		TScopedReaderAccessDetector<FRWFullyRecursiveAccessDetector> Accessor;
+		FMRSWRecursiveAccessDetector::FDestructionSentinel DestructionSentinel;
 	};
 
 	[[nodiscard]] FReadAccessScope GetReadAccessScope() const
 	{
-		return FReadAccessScope(Accessor);
+		return FReadAccessScope(AccessDetector);
 	}
 
 	class FWriteAccessScope
@@ -68,21 +128,47 @@ protected:
 	public:
 		UE_NONCOPYABLE(FWriteAccessScope);
 
-		explicit FWriteAccessScope(FRWFullyRecursiveAccessDetector& InAccessDetector)
-			: Accessor(MakeScopedWriterAccessDetector(InAccessDetector))
+		explicit FWriteAccessScope(FMRSWRecursiveAccessDetector& InAccessDetector)
+			: DestructionSentinel(FMRSWRecursiveAccessDetector::EAccessType::Writer)
 		{
+			DestructionSentinel.Accessor = &InAccessDetector;
+			InAccessDetector.AcquireWriteAccess(DestructionSentinel);
+		}
+
+		~FWriteAccessScope()
+		{
+			if (!DestructionSentinel.bDestroyed) // only if `AccessDetector` wasn't destroyed while being accessed
+			{
+				// the `const_cast` doesn't violate constness here as we got a mutable accessor in `FWriteAccessScope` constuctor,
+				// otherwise we'd need to have a redundant ptr to a mutable accessor as a member
+				const_cast<FMRSWRecursiveAccessDetector*>(DestructionSentinel.Accessor)->ReleaseWriteAccess(DestructionSentinel);
+			}
 		}
 
 	private:
-		TScopedWriterDetector<FRWFullyRecursiveAccessDetector> Accessor;
+		FMRSWRecursiveAccessDetector::FDestructionSentinel DestructionSentinel;
 	};
 
 	[[nodiscard]] FWriteAccessScope GetWriteAccessScope()
 	{
-		return FWriteAccessScope(Accessor);
+		return FWriteAccessScope(AccessDetector);
 	}
 
 private:
-	mutable FRWFullyRecursiveAccessDetector Accessor;
+	FMRSWRecursiveAccessDetector AccessDetector;
+};
+
+#else // UE_DETECT_DELEGATES_RACE_CONDITIONS
+
+template<>
+class TDelegateAccessHandlerBase<FNotThreadSafeDelegateMode> : public TDelegateAccessHandlerBase<FNotThreadSafeNotCheckedDelegateMode>
+{
+};
+
 #endif // UE_DETECT_DELEGATES_RACE_CONDITIONS
+
+template<>
+struct TIsZeroConstructType<TDelegateAccessHandlerBase<FNotThreadSafeDelegateMode>>
+{
+	static constexpr bool Value = true;
 };
diff --git a/Engine/Source/Runtime/Core/Public/Delegates/DelegateBase.h b/Engine/Source/Runtime/Core/Public/Delegates/DelegateBase.h
--- a/Engine/Source/Runtime/Core/Public/Delegates/DelegateBase.h
+++ b/Engine/Source/Runtime/Core/Public/Delegates/DelegateBase.h
@@ -30,14 +30,12 @@
 template <typename UserPolicy>
 class TMulticastDelegateBase;
 
-template <typename UserPolicy>
-class TTSMulticastDelegateBase;
-
 template <typename UserPolicy>
 class TDelegateBase;
 
 ALIAS_TEMPLATE_TYPE_LAYOUT(template<typename ElementType>, FDelegateAllocatorType::ForElementType<ElementType>, void*);
 
+// not thread-safe version, with automatic race detection in dev builds
 struct FDefaultDelegateUserPolicy
 {
 	// To extend delegates, you should implement a policy struct like this and pass it as the second template
@@ -50,7 +48,7 @@ struct FDefaultDelegateUserPolicy
 	//   - This binding is not available through the public API of the delegate, but is accessible to FDelegateExtras.
 	//
 	// FDelegateExtras:
-	//   - Must publicly inherit TDelegateBase<UserPolicy>.
+	//   - Must publicly inherit TDelegateBase<FThreadSafetyMode>.
 	//   - Should contain any extra data and functions injected into a delegate (the object which holds an
 	//     FDelegateInstance-derived object, above).
 	//   - Public data members and member functions are accessible directly through the TDelegate object.
@@ -63,31 +61,51 @@ struct FDefaultDelegateUserPolicy
 	//     holds an array of FDelegateExtras-derived objects which is the invocation list).
 	//   - Public data members and member functions are accessible directly through the TMulticastDelegate object.
 
-	using FDelegateInstanceExtras  = IDelegateInstance;
-	using FDelegateExtras          = TDelegateBase<FDefaultDelegateUserPolicy>;
+	using FDelegateInstanceExtras = IDelegateInstance;
+	using FThreadSafetyMode =
+#if UE_DETECT_DELEGATES_RACE_CONDITIONS
+		FNotThreadSafeDelegateMode;
+#else
+		FNotThreadSafeNotCheckedDelegateMode;
+#endif
+	using FDelegateExtras = TDelegateBase<FThreadSafetyMode>;
 	using FMulticastDelegateExtras = TMulticastDelegateBase<FDefaultDelegateUserPolicy>;
 };
 
+// thread-safe version
 struct FDefaultTSDelegateUserPolicy
 {
 	// see `FDefaultDelegateUserPolicy` for documentation
-	using FDelegateInstanceExtras  = IDelegateInstance;
-	using FDelegateExtras          = TDelegateBase<FDefaultTSDelegateUserPolicy>;
-	using FMulticastDelegateExtras = TTSMulticastDelegateBase<FDefaultTSDelegateUserPolicy>;
+	using FDelegateInstanceExtras = IDelegateInstance;
+	using FThreadSafetyMode = FThreadSafeDelegateMode;
+	using FDelegateExtras = TDelegateBase<FThreadSafetyMode>;
+	using FMulticastDelegateExtras = TMulticastDelegateBase<FDefaultTSDelegateUserPolicy>;
+};
+
+// not thread-safe version, no race detection. used primarily for deprecated unsafe delegates that must be kept running for backward compatibility
+struct FNotThreadSafeNotCheckedDelegateUserPolicy
+{
+	using FDelegateInstanceExtras = IDelegateInstance;
+	using FThreadSafetyMode = FNotThreadSafeNotCheckedDelegateMode;
+	using FDelegateExtras = TDelegateBase<FThreadSafetyMode>;
+	using FMulticastDelegateExtras = TMulticastDelegateBase<FNotThreadSafeNotCheckedDelegateUserPolicy>;
 };
 
 /**
  * Base class for unicast delegates.
  */
-template <typename UserPolicy>
-class TDelegateBase : public FDelegateAccessHandlerBaseChecked
+template<typename ThreadSafetyMode>
+class TDelegateBase : public TDelegateAccessHandlerBase<ThreadSafetyMode>
 {
 private:
-	using Super = FDelegateAccessHandlerBaseChecked;
+	using Super = TDelegateAccessHandlerBase<ThreadSafetyMode>;
 
 	template <typename>
 	friend class TMulticastDelegateBase;
 
+	template <typename>
+	friend class TDelegateBase;
+
 	template <class, typename, typename, typename...>
 	friend class TBaseUFunctionDelegateInstance;
 
@@ -134,6 +152,13 @@ public:
 		return *this;
 	}
 
+	// support for moving from delegates with different thread-safety mode
+	template<typename OtherThreadSafetyMode>
+	explicit TDelegateBase(TDelegateBase<OtherThreadSafetyMode>&& Other)
+	{
+		MoveConstruct(MoveTemp(Other));
+	}
+
 	/**
 	 * Unbinds this delegate
 	 */
@@ -321,22 +346,23 @@ private:
 		return DelegateAllocator.GetAllocation();
 	}
 
-	void MoveConstruct(TDelegateBase&& Other)
+	template<typename OtherThreadSafetyMode>
+	void MoveConstruct(TDelegateBase<OtherThreadSafetyMode>&& Other)
 	{
-		FWriteAccessScope OtherWriteScope = Other.GetWriteAccessScope();
+		typename TDelegateBase<OtherThreadSafetyMode>::FWriteAccessScope OtherWriteScope = Other.GetWriteAccessScope();
 
 		DelegateAllocator.MoveToEmpty(Other.DelegateAllocator);
 		DelegateSize = Other.DelegateSize;
 		Other.DelegateSize = 0;
 	}
 
-	void MoveAssign(TDelegateBase&& Other)
+	template<typename OtherThreadSafetyMode>
+	void MoveAssign(TDelegateBase<OtherThreadSafetyMode>&& Other)
 	{
 		FDelegateAllocatorType::ForElementType<FAlignedInlineDelegateType> OtherDelegateAllocator;
 		int32 OtherDelegateSize;
 		{
-			FWriteAccessScope OtherWriteScope = Other.GetWriteAccessScope();
-
+			typename TDelegateBase<OtherThreadSafetyMode>::FWriteAccessScope OtherWriteScope = Other.GetWriteAccessScope();
 			OtherDelegateAllocator.MoveToEmpty(Other.DelegateAllocator);
 			OtherDelegateSize = Other.DelegateSize;
 			Other.DelegateSize = 0;
diff --git a/Engine/Source/Runtime/Core/Public/Delegates/DelegateInstanceInterface.h b/Engine/Source/Runtime/Core/Public/Delegates/DelegateInstanceInterface.h
--- a/Engine/Source/Runtime/Core/Public/Delegates/DelegateInstanceInterface.h
+++ b/Engine/Source/Runtime/Core/Public/Delegates/DelegateInstanceInterface.h
@@ -11,13 +11,22 @@ class IDelegateInstance;
 template <typename FuncType, typename UserPolicy>
 struct IBaseDelegateInstance;
 
+template <typename>
+class TDelegateBase;
+
+struct FNotThreadSafeNotCheckedDelegateMode;
+struct FThreadSafeDelegateMode;
+struct FNotThreadSafeDelegateMode;
+
 template <typename RetType, typename... ArgTypes, typename UserPolicy>
 struct IBaseDelegateInstance<RetType(ArgTypes...), UserPolicy> : public UserPolicy::FDelegateInstanceExtras
 {
 	/**
 	 * Emplaces a copy of the delegate instance into the given base delegate.
 	 */
-	virtual void CreateCopy(typename UserPolicy::FDelegateExtras& Base) const = 0;
+	virtual void CreateCopy(TDelegateBase<FThreadSafeDelegateMode>& Base) const = 0;
+	virtual void CreateCopy(TDelegateBase<FNotThreadSafeDelegateMode>& Base) const = 0;
+	virtual void CreateCopy(TDelegateBase<FNotThreadSafeNotCheckedDelegateMode>& Base) const = 0;
 
 	/**
 	 * Execute the delegate.  If the function pointer is not valid, an error will occur.
diff --git a/Engine/Source/Runtime/Core/Public/Delegates/DelegateInstancesImpl.h b/Engine/Source/Runtime/Core/Public/Delegates/DelegateInstancesImpl.h
--- a/Engine/Source/Runtime/Core/Public/Delegates/DelegateInstancesImpl.h
+++ b/Engine/Source/Runtime/Core/Public/Delegates/DelegateInstancesImpl.h
@@ -14,6 +14,7 @@
 #include "Delegates/DelegateInstanceInterface.h"
 #include "Delegates/DelegateInstancesImplFwd.h"
 #include "Delegates/IDelegateInstance.h"
+#include "Delegates/DelegateBase.h"
 #include "Misc/AssertionMacros.h"
 #include "Templates/RemoveReference.h"
 #include "Templates/SharedPointer.h"
@@ -22,7 +23,6 @@
 #include "UObject/NameTypes.h"
 #include "UObject/WeakObjectPtrTemplates.h"
 
-class FDelegateHandle;
 
 namespace UE::Delegates::Private
 {
@@ -127,7 +127,17 @@ public:
 
 	// IBaseDelegateInstance interface
 
-	void CreateCopy(DelegateBaseType& Base) const final
+	void CreateCopy(TDelegateBase<FThreadSafeDelegateMode>& Base) const final
+	{
+		Base.template CreateDelegateInstance<TBaseUFunctionDelegateInstance>(*this);
+	}
+
+	void CreateCopy(TDelegateBase<FNotThreadSafeDelegateMode>& Base) const final
+	{
+		Base.template CreateDelegateInstance<TBaseUFunctionDelegateInstance>(*this);
+	}
+
+	void CreateCopy(TDelegateBase<FNotThreadSafeNotCheckedDelegateMode>& Base) const final
 	{
 		Base.template CreateDelegateInstance<TBaseUFunctionDelegateInstance>(*this);
 	}
@@ -238,7 +248,17 @@ public:
 
 	// IBaseDelegateInstance interface
 
-	void CreateCopy(DelegateBaseType& Base) const final
+	void CreateCopy(TDelegateBase<FThreadSafeDelegateMode>& Base) const final
+	{
+		Base.template CreateDelegateInstance<TBaseSPMethodDelegateInstance>(*this);
+	}
+
+	void CreateCopy(TDelegateBase<FNotThreadSafeDelegateMode>& Base) const final
+	{
+		Base.template CreateDelegateInstance<TBaseSPMethodDelegateInstance>(*this);
+	}
+
+	void CreateCopy(TDelegateBase<FNotThreadSafeNotCheckedDelegateMode>& Base) const final
 	{
 		Base.template CreateDelegateInstance<TBaseSPMethodDelegateInstance>(*this);
 	}
@@ -368,7 +388,17 @@ public:
 
 	// IBaseDelegateInstance interface
 
-	void CreateCopy(DelegateBaseType& Base) const final
+	void CreateCopy(TDelegateBase<FThreadSafeDelegateMode>& Base) const final
+	{
+		Base.template CreateDelegateInstance<TBaseRawMethodDelegateInstance>(*this);
+	}
+
+	void CreateCopy(TDelegateBase<FNotThreadSafeDelegateMode>& Base) const final
+	{
+		Base.template CreateDelegateInstance<TBaseRawMethodDelegateInstance>(*this);
+	}
+
+	void CreateCopy(TDelegateBase<FNotThreadSafeNotCheckedDelegateMode>& Base) const final
 	{
 		Base.template CreateDelegateInstance<TBaseRawMethodDelegateInstance>(*this);
 	}
@@ -487,7 +517,17 @@ public:
 
 	// IBaseDelegateInstance interface
 
-	void CreateCopy(DelegateBaseType& Base) const final
+	void CreateCopy(TDelegateBase<FThreadSafeDelegateMode>& Base) const final
+	{
+		Base.template CreateDelegateInstance<TBaseUObjectMethodDelegateInstance>(*this);
+	}
+
+	void CreateCopy(TDelegateBase<FNotThreadSafeDelegateMode>& Base) const final
+	{
+		Base.template CreateDelegateInstance<TBaseUObjectMethodDelegateInstance>(*this);
+	}
+
+	void CreateCopy(TDelegateBase<FNotThreadSafeNotCheckedDelegateMode>& Base) const final
 	{
 		Base.template CreateDelegateInstance<TBaseUObjectMethodDelegateInstance>(*this);
 	}
@@ -604,7 +644,17 @@ public:
 
 	// IBaseDelegateInstance interface
 
-	void CreateCopy(DelegateBaseType& Base) const final
+	void CreateCopy(TDelegateBase<FThreadSafeDelegateMode>& Base) const final
+	{
+		Base.template CreateDelegateInstance<TBaseStaticDelegateInstance>(*this);
+	}
+
+	void CreateCopy(TDelegateBase<FNotThreadSafeDelegateMode>& Base) const final
+	{
+		Base.template CreateDelegateInstance<TBaseStaticDelegateInstance>(*this);
+	}
+
+	void CreateCopy(TDelegateBase<FNotThreadSafeNotCheckedDelegateMode>& Base) const final
 	{
 		Base.template CreateDelegateInstance<TBaseStaticDelegateInstance>(*this);
 	}
@@ -692,7 +742,17 @@ public:
 
 public:
 	// IBaseDelegateInstance interface
-	void CreateCopy(DelegateBaseType& Base) const final
+	void CreateCopy(TDelegateBase<FThreadSafeDelegateMode>& Base) const final
+	{
+		Base.template CreateDelegateInstance<TBaseFunctorDelegateInstance>(*this);
+	}
+
+	void CreateCopy(TDelegateBase<FNotThreadSafeDelegateMode>& Base) const final
+	{
+		Base.template CreateDelegateInstance<TBaseFunctorDelegateInstance>(*this);
+	}
+
+	void CreateCopy(TDelegateBase<FNotThreadSafeNotCheckedDelegateMode>& Base) const final
 	{
 		Base.template CreateDelegateInstance<TBaseFunctorDelegateInstance>(*this);
 	}
@@ -781,7 +841,17 @@ public:
 
 public:
 	// IBaseDelegateInstance interface
-	void CreateCopy(DelegateBaseType& Base) const final
+	void CreateCopy(TDelegateBase<FThreadSafeDelegateMode>& Base) const final
+	{
+		Base.template CreateDelegateInstance<TWeakBaseFunctorDelegateInstance>(*this);
+	}
+
+	void CreateCopy(TDelegateBase<FNotThreadSafeDelegateMode>& Base) const final
+	{
+		Base.template CreateDelegateInstance<TWeakBaseFunctorDelegateInstance>(*this);
+	}
+
+	void CreateCopy(TDelegateBase<FNotThreadSafeNotCheckedDelegateMode>& Base) const final
 	{
 		Base.template CreateDelegateInstance<TWeakBaseFunctorDelegateInstance>(*this);
 	}
diff --git a/Engine/Source/Runtime/Core/Public/Delegates/DelegateSignatureImpl.inl b/Engine/Source/Runtime/Core/Public/Delegates/DelegateSignatureImpl.inl
--- a/Engine/Source/Runtime/Core/Public/Delegates/DelegateSignatureImpl.inl
+++ b/Engine/Source/Runtime/Core/Public/Delegates/DelegateSignatureImpl.inl
@@ -34,7 +34,6 @@ template <typename To, typename From>
 To* Cast(From* Src);
 
 template<typename UserPolicy> class TMulticastDelegateBase;
-template<typename UserPolicy> class TTSMulticastDelegateBase;
 
 /**
  * Unicast delegate template class.
@@ -46,25 +45,23 @@ class TDelegate
 };
 
 template <typename InRetValType, typename... ParamTypes, typename UserPolicy>
-class TDelegate<InRetValType(ParamTypes...), UserPolicy> : public TDelegateBase<UserPolicy>
+class TDelegate<InRetValType(ParamTypes...), UserPolicy> : public UserPolicy::FDelegateExtras
 {
-	using Super                         = TDelegateBase<UserPolicy>;
+	using Super                         = typename UserPolicy::FDelegateExtras;
 	using FuncType                      = InRetValType (ParamTypes...);
 	using DelegateInstanceInterfaceType = IBaseDelegateInstance<FuncType, UserPolicy>;
 
 	static_assert(std::is_convertible_v<typename UserPolicy::FDelegateInstanceExtras*, IDelegateInstance*>, "UserPolicy::FDelegateInstanceExtras should publicly inherit IDelegateInstance");
-	static_assert(std::is_convertible_v<typename UserPolicy::FMulticastDelegateExtras*, TMulticastDelegateBase<UserPolicy>*> || std::is_convertible_v<typename UserPolicy::FMulticastDelegateExtras*, TTSMulticastDelegateBase<UserPolicy>*>, "UserPolicy::FMulticastDelegateExtras should publicly inherit TMulticastDelegateBase<UserPolicy> or TTSMulticastDelegateBase<UserPolicy>");
+	static_assert(std::is_convertible_v<typename UserPolicy::FMulticastDelegateExtras*, TMulticastDelegateBase<UserPolicy>*>, "UserPolicy::FMulticastDelegateExtras should publicly inherit TMulticastDelegateBase<UserPolicy>");
+
+	template <typename, typename>
+	friend class TDelegate;
 
 	template <typename>
 	friend class TMulticastDelegateBase;
 
-	template <typename>
-	friend class TTSMulticastDelegateBase;
-
-public:
-	// Make sure FDelegateBase's public functions are publicly exposed through the TDelegate API
-	using Super::Unbind;
-	using Super::GetAllocatedSize;
+	template <typename, typename>
+	friend class TMulticastDelegate;
 
 private:
 	// Make sure FDelegateBase's protected functions are not accidentally exposed through the TDelegate API
@@ -72,7 +69,6 @@ private:
 	using Super::GetReadAccessScope;
 	using typename Super::FWriteAccessScope;
 	using Super::GetWriteAccessScope;
-	using Super::GetDelegateInstanceProtected;
 
 public:
 	/** Type definition for return value type. */
@@ -557,7 +553,8 @@ protected:
 	}
 
 private:
-	void CopyFrom(const TDelegate& Other)
+	template<typename OtherUserPolicy>
+	void CopyFrom(const TDelegate<FuncType, OtherUserPolicy>& Other)
 	{
 		if ((void*)&Other == (void*)this)
 		{
@@ -568,21 +565,36 @@ private:
 		TDelegate LocalCopy;
 
 		{
-			FReadAccessScope OtherReadScope = Other.GetReadAccessScope();
+			typename TDelegate<FuncType, OtherUserPolicy>::FReadAccessScope OtherReadScope = Other.GetReadAccessScope();
 
 			// this down-cast is OK! allows for managing invocation list in the base class without requiring virtual functions
-			const DelegateInstanceInterfaceType* OtherInstance = Other.GetDelegateInstanceProtected();
+			using OtherDelegateInstanceInterfaceType = IBaseDelegateInstance<FuncType, OtherUserPolicy>;
+			const OtherDelegateInstanceInterfaceType* OtherDelegateInstance = Other.GetDelegateInstanceProtected();
 
-			if (OtherInstance != nullptr)
+			if (OtherDelegateInstance != nullptr)
 			{
-				OtherInstance->CreateCopy(LocalCopy);
+				OtherDelegateInstance->CreateCopy(LocalCopy);
 			}
 		}
 
 		*this = MoveTemp(LocalCopy);
 	}
+
+	// copying from delegates with different user policy
+	template<typename OtherUserPolicy>
+	explicit TDelegate(const TDelegate<FuncType, OtherUserPolicy>& Other)
+	{
+		CopyFrom(Other);
+	}
 };
 
+/////////////////////////////////////////////////////////////////////////////////////////////////////////////////////////////
+
+template <typename DelegateSignature>
+using TTSDelegate = TDelegate<DelegateSignature, FDefaultTSDelegateUserPolicy>;
+
+/////////////////////////////////////////////////////////////////////////////////////////////////////////////////////////////
+
 /**
  * Multicast delegate template base class, used for both normal and event multicast delegates.
  *
@@ -641,7 +653,7 @@ public:
 	{
 		if (&Other != this)
 		{
-			Super::template CopyFrom<DelegateInstanceInterfaceType, FDelegate>(Other);
+			Super::template CopyFrom<DelegateInstanceInterfaceType>(Other);
 		}
 		return *this;
 	}
@@ -894,7 +906,7 @@ public:
 	 */
 	void Broadcast(ParamTypes... Params) const
 	{
-		Super::template Broadcast<DelegateInstanceInterfaceType, FDelegate, ParamTypes...>(Params...);
+		Super::template Broadcast<DelegateInstanceInterfaceType, ParamTypes...>(Params...);
 	}
 };
 
@@ -907,8 +919,8 @@ using TTSMulticastDelegate = TMulticastDelegate<DelegateSignature, FDefaultTSDel
  * macros to create the actual delegate type, templated to the function signature the delegate is compatible with.
  * Then, you can create an instance of that class when you want to assign functions to the delegate.
  */
-template <typename ScriptDelegateDummyType, typename RetValType, typename... ParamTypes>
-class TBaseDynamicDelegate : public TScriptDelegate<ScriptDelegateDummyType>
+template <typename ThreadSafetyMode, typename RetValType, typename... ParamTypes>
+class TBaseDynamicDelegate : public TScriptDelegate<ThreadSafetyMode>
 {
 public:
 	/**
@@ -921,8 +933,8 @@ public:
 	 *
 	 * @param	InScriptDelegate	The delegate to construct from by copying
 	 */
-	explicit TBaseDynamicDelegate( const TScriptDelegate<ScriptDelegateDummyType>& InScriptDelegate )
-		: TScriptDelegate<ScriptDelegateDummyType>( InScriptDelegate )
+	explicit TBaseDynamicDelegate( const TScriptDelegate<ThreadSafetyMode>& InScriptDelegate )
+		: TScriptDelegate<ThreadSafetyMode>( InScriptDelegate )
 	{ }
 
 	/**
@@ -984,12 +996,12 @@ public:
  * signature the delegate is compatible with.   Then, you can create an instance of that class when you
  * want to assign functions to the delegate.
  */
-template <typename ScriptDelegateDummyType, typename RetValType, typename... ParamTypes>
-class TBaseDynamicMulticastDelegate : public TMulticastScriptDelegate<ScriptDelegateDummyType>
+template <typename ThreadSafetyMode, typename RetValType, typename... ParamTypes>
+class TBaseDynamicMulticastDelegate : public TMulticastScriptDelegate<ThreadSafetyMode>
 {
 public:
 	/** The actual single-cast delegate class for this multi-cast delegate */
-	typedef TBaseDynamicDelegate<ScriptDelegateDummyType, RetValType, ParamTypes...> FDelegate;
+	typedef TBaseDynamicDelegate<ThreadSafetyMode, RetValType, ParamTypes...> FDelegate;
 
 	/**
 	 * Default constructor
@@ -1001,8 +1013,8 @@ public:
 	 *
 	 * @param	InScriptDelegate	The delegate to construct from by copying
 	 */
-	explicit TBaseDynamicMulticastDelegate( const TMulticastScriptDelegate<ScriptDelegateDummyType>& InMulticastScriptDelegate )
-		: TMulticastScriptDelegate<ScriptDelegateDummyType>( InMulticastScriptDelegate )
+	explicit TBaseDynamicMulticastDelegate( const TMulticastScriptDelegate<ThreadSafetyMode>& InMulticastScriptDelegate )
+		: TMulticastScriptDelegate<ThreadSafetyMode>( InMulticastScriptDelegate )
 	{ }
 
 	/**
diff --git a/Engine/Source/Runtime/Core/Public/Delegates/MulticastDelegateBase.h b/Engine/Source/Runtime/Core/Public/Delegates/MulticastDelegateBase.h
--- a/Engine/Source/Runtime/Core/Public/Delegates/MulticastDelegateBase.h
+++ b/Engine/Source/Runtime/Core/Public/Delegates/MulticastDelegateBase.h
@@ -23,16 +23,17 @@
  * Abstract base class for multicast delegates.
  */
 template<typename UserPolicy>
-class TMulticastDelegateBase : public FDelegateAccessHandlerBaseChecked
+class TMulticastDelegateBase : public TDelegateAccessHandlerBase<typename UserPolicy::FThreadSafetyMode>
 {
 protected:
-	using Super = FDelegateAccessHandlerBaseChecked;
+	using Super = TDelegateAccessHandlerBase<typename UserPolicy::FThreadSafetyMode>;
 	using typename Super::FReadAccessScope;
 	using Super::GetReadAccessScope;
 	using typename Super::FWriteAccessScope;
 	using Super::GetWriteAccessScope;
 
-	using UnicastDelegateType = TDelegateBase<UserPolicy>;
+	// individual bindings are not checked for races as it's done for the parent delegate
+	using UnicastDelegateType = TDelegateBase<FNotThreadSafeNotCheckedDelegateMode>;
 
 	using InvocationListType = TArray<UnicastDelegateType, FMulticastInvocationListAllocatorType>;
 
@@ -204,7 +205,7 @@ protected:
 	{ }
 
 protected:
-	template<typename DelegateInstanceInterfaceType, typename DelegateType>
+	template<typename DelegateInstanceInterfaceType>
 	void CopyFrom(const TMulticastDelegateBase& Other)
 	{
 		InvocationListType TempInvocationList;
@@ -216,7 +217,7 @@ protected:
 		    {
 				if (const IDelegateInstance* OtherInstance = OtherDelegateRef.GetDelegateInstanceProtected())
 			    {
-				    DelegateType TempDelegate;
+				    UnicastDelegateType TempDelegate;
 				    static_cast<const DelegateInstanceInterfaceType*>(OtherInstance)->CreateCopy(TempDelegate);
 				    TempInvocationList.Add(MoveTemp(TempDelegate));
 			    }
@@ -231,7 +232,7 @@ protected:
 		}
 	}
 
-	template<typename DelegateInstanceInterfaceType, typename DelegateBaseType, typename... ParamTypes>
+	template<typename DelegateInstanceInterfaceType, typename... ParamTypes>
 	void Broadcast(ParamTypes... Params) const
 	{
 		// the `const` on the method is a lie
@@ -247,7 +248,7 @@ protected:
 			for (int32 InvocationListIndex = LocalInvocationList.Num() - 1; InvocationListIndex >= 0; --InvocationListIndex)
 			{
 				// this down-cast is OK! allows for managing invocation list in the base class without requiring virtual functions
-				const DelegateBaseType& DelegateBase = static_cast<const DelegateBaseType&>(LocalInvocationList[InvocationListIndex]);
+				const UnicastDelegateType& DelegateBase = LocalInvocationList[InvocationListIndex];
 
 				const IDelegateInstance* DelegateInstanceInterface = DelegateBase.GetDelegateInstanceProtected();
 				if (DelegateInstanceInterface == nullptr || !static_cast<const DelegateInstanceInterfaceType*>(DelegateInstanceInterface)->ExecuteIfSafe(Params...))
diff --git a/Engine/Source/Runtime/Core/Public/Delegates/TSMulticastDelegateBase.h b/Engine/Source/Runtime/Core/Public/Delegates/TSMulticastDelegateBase.h
deleted file mode 100644
diff --git a/Engine/Source/Runtime/Core/Public/Misc/CoreDelegates.h b/Engine/Source/Runtime/Core/Public/Misc/CoreDelegates.h
--- a/Engine/Source/Runtime/Core/Public/Misc/CoreDelegates.h
+++ b/Engine/Source/Runtime/Core/Public/Misc/CoreDelegates.h
@@ -126,7 +126,7 @@ public:
 
 	// When a file is opened for read from a pak file.
 	UE_DEPRECATED(5.3, "This delegate is not thread-safe, please use Use FCoreDelegates::GetOnFileOpenedForReadFromPakFile()")
-	static TMulticastDelegate<void(const TCHAR* PakFile, const TCHAR* FileName)> OnFileOpenedForReadFromPakFile;
+	static TMulticastDelegate<void(const TCHAR* PakFile, const TCHAR* FileName), FNotThreadSafeNotCheckedDelegateUserPolicy> OnFileOpenedForReadFromPakFile;
 	static TTSMulticastDelegate<void(const TCHAR* PakFile, const TCHAR* FileName)>& GetOnFileOpenedForReadFromPakFile();
 
 	typedef TSharedPtr<class IMovieStreamer, ESPMode::ThreadSafe> FMovieStreamerPtr;
diff --git a/Engine/Source/Runtime/Core/Public/UObject/ScriptDelegates.h b/Engine/Source/Runtime/Core/Public/UObject/ScriptDelegates.h
--- a/Engine/Source/Runtime/Core/Public/UObject/ScriptDelegates.h
+++ b/Engine/Source/Runtime/Core/Public/UObject/ScriptDelegates.h
@@ -14,14 +14,15 @@
 
 namespace UE::Core::Private
 {
-	struct TScriptDelegateDefault;
-
-	template <typename ScriptDelegateParameterType>
+	template <typename InThreadSafetyMode>
 	struct TScriptDelegateTraits
 	{
 		// Although templated, WeakPtrType is not intended to be anything other than FWeakObjectPtr,
 		// and is only a template for module organization reasons.
 		using WeakPtrType = FWeakObjectPtr;
+
+		using ThreadSafetyMode = InThreadSafetyMode;
+		using UnicastThreadSafetyModeForMulticasts = FNotThreadSafeNotCheckedDelegateMode;
 	};
 
 	template <>
@@ -32,6 +33,8 @@ namespace UE::Core::Private
 		// and also the TScriptDelegate(const TScriptDelegate<FWeakObjectPtr>&) constructor.
 
 		using WeakPtrType = FWeakObjectPtr;
+		using ThreadSafetyMode = FNotThreadSafeDelegateMode;
+		using UnicastThreadSafetyModeForMulticasts = FNotThreadSafeNotCheckedDelegateMode;
 	};
 
 	// This function only exists to allow compatibility between multicast and unicast delegate types which use an explicit FWeakObjectPtr template parameter
@@ -39,13 +42,13 @@ namespace UE::Core::Private
 	/* UE_DEPRECATED(5.3, "Deprecated - remove after TScriptDelegateTraits<FWeakObjectPtr> is removed") */
 	inline constexpr bool BackwardCompatibilityCheck()
 	{
-		if constexpr (std::is_same_v<From, TScriptDelegateDefault>)
+		if constexpr (std::is_same_v<From, FNotThreadSafeDelegateMode>)
 		{
 			return std::is_same_v<To, FWeakObjectPtr>;
 		}
 		else if constexpr (std::is_same_v<From, FWeakObjectPtr>)
 		{
-			return std::is_same_v<To, TScriptDelegateDefault>;
+			return std::is_same_v<To, FNotThreadSafeDelegateMode>;
 		}
 		else
 		{
@@ -57,9 +60,10 @@ namespace UE::Core::Private
 /**
  * Script delegate base class.
  */
-template <typename Dummy>
-class TScriptDelegate : public FDelegateAccessHandlerBaseChecked
+template <typename InThreadSafetyMode>
+class TScriptDelegate : public TDelegateAccessHandlerBase<typename UE::Core::Private::TScriptDelegateTraits<InThreadSafetyMode>::ThreadSafetyMode>
 {
+	using ThreadSafetyMode = typename UE::Core::Private::TScriptDelegateTraits<InThreadSafetyMode>::ThreadSafetyMode;
 
 	template <typename>
 	friend class TScriptDelegate;
@@ -67,14 +71,14 @@ class TScriptDelegate : public FDelegateAccessHandlerBaseChecked
 	template<typename>
 	friend class TMulticastScriptDelegate;
 
-	using Super = FDelegateAccessHandlerBaseChecked;
+	using Super = TDelegateAccessHandlerBase<ThreadSafetyMode>;
 	using typename Super::FReadAccessScope;
 	using Super::GetReadAccessScope;
 	using typename Super::FWriteAccessScope;
 	using Super::GetWriteAccessScope;
 
 public:
-	using WeakPtrType = typename UE::Core::Private::TScriptDelegateTraits<Dummy>::WeakPtrType;
+	using WeakPtrType = typename UE::Core::Private::TScriptDelegateTraits<InThreadSafetyMode>::WeakPtrType;
 
 	/** Default constructor. */
 	TScriptDelegate() 
@@ -91,13 +95,13 @@ public:
 	}
 
 	template <
-		typename OtherDummy
-		UE_REQUIRES(UE::Core::Private::BackwardCompatibilityCheck<Dummy, OtherDummy>())
+		typename OtherThreadSafetyMode
+		UE_REQUIRES(UE::Core::Private::BackwardCompatibilityCheck<InThreadSafetyMode, OtherThreadSafetyMode>())
 	>
 	/* UE_DEPRECATED(5.3, "Deprecated - remove after TScriptDelegateTraits<FWeakObjectPtr> is removed") */
-	TScriptDelegate(const TScriptDelegate<OtherDummy>& Other)
+	TScriptDelegate(const TScriptDelegate<OtherThreadSafetyMode>& Other)
 	{
-		typename TScriptDelegate<OtherDummy>::FReadAccessScope OtherReadScope = Other.GetReadAccessScope();
+		typename TScriptDelegate<OtherThreadSafetyMode>::FReadAccessScope OtherReadScope = Other.GetReadAccessScope();
 
 		Object = Other.Object;
 		FunctionName = Other.FunctionName;
@@ -123,17 +127,17 @@ public:
 		return *this;
 	}
 	template <
-		typename OtherDummy
-		UE_REQUIRES(UE::Core::Private::BackwardCompatibilityCheck<Dummy, OtherDummy>())
+		typename OtherThreadSafetyMode
+		UE_REQUIRES(UE::Core::Private::BackwardCompatibilityCheck<InThreadSafetyMode, OtherThreadSafetyMode>())
 	>
 	/* UE_DEPRECATED(5.3, "Deprecated - remove after TScriptDelegateTraits<FWeakObjectPtr> is removed") */
-	TScriptDelegate& operator=(const TScriptDelegate<OtherDummy>& Other)
+	TScriptDelegate& operator=(const TScriptDelegate<OtherThreadSafetyMode>& Other)
 	{
 		WeakPtrType OtherObject;
 		FName OtherFunctionName;
 
 		{
-			typename TScriptDelegate<OtherDummy>::FReadAccessScope OtherReadScope = Other.GetReadAccessScope();
+			typename TScriptDelegate<OtherThreadSafetyMode>::FReadAccessScope OtherReadScope = Other.GetReadAccessScope();
 			OtherObject = Other.Object;
 			OtherFunctionName = Other.FunctionName;
 		}
@@ -298,17 +302,17 @@ public:
 		return bResult;
 	}
 	template <
-		typename OtherDummy
-		UE_REQUIRES(UE::Core::Private::BackwardCompatibilityCheck<Dummy, OtherDummy>())
+		typename OtherThreadSafetyMode
+		UE_REQUIRES(UE::Core::Private::BackwardCompatibilityCheck<InThreadSafetyMode, OtherThreadSafetyMode>())
 	>
 	/* UE_DEPRECATED(5.3, "Deprecated - remove after TScriptDelegateTraits<FWeakObjectPtr> is removed") */
-	FORCEINLINE bool operator==(const TScriptDelegate<OtherDummy>& Other) const
+	FORCEINLINE bool operator==(const TScriptDelegate<OtherThreadSafetyMode>& Other) const
 	{
 		WeakPtrType OtherObject;
 		FName OtherFunctionName;
 
 		{
-			typename TScriptDelegate<OtherDummy>::FReadAccessScope OtherReadScope = Other.GetReadAccessScope();
+			typename TScriptDelegate<OtherThreadSafetyMode>::FReadAccessScope OtherReadScope = Other.GetReadAccessScope();
 			OtherObject = Other.Object;
 			OtherFunctionName = Other.FunctionName;
 		}
@@ -328,11 +332,11 @@ public:
 		return !operator==(Other);
 	}
 	template <
-		typename OtherDummy
-		UE_REQUIRES(UE::Core::Private::BackwardCompatibilityCheck<Dummy, OtherDummy>())
+		typename OtherThreadSafetyMode
+		UE_REQUIRES(UE::Core::Private::BackwardCompatibilityCheck<InThreadSafetyMode, OtherThreadSafetyMode>())
 	>
 	/* UE_DEPRECATED(5.3, "Deprecated - remove after TScriptDelegateTraits<FWeakObjectPtr> is removed") */
-	FORCEINLINE bool operator!=(const TScriptDelegate<OtherDummy>& Other) const
+	FORCEINLINE bool operator!=(const TScriptDelegate<OtherThreadSafetyMode>& Other) const
 	{
 		return !operator==(Other);
 	}
@@ -449,10 +453,12 @@ public:
 		return HashCombine(GetTypeHash(Delegate.Object), GetTypeHash(Delegate.GetFunctionName()));
 	}
 
-	template<typename OtherDummy>
-	static TScriptDelegate CopyFrom(const TScriptDelegate<OtherDummy>& Other)
+	template<typename OtherThreadSafetyMode>
+	static TScriptDelegate CopyFrom(const TScriptDelegate<OtherThreadSafetyMode>& Other)
 	{
-		typename TScriptDelegate<OtherDummy>::FReadAccessScope OtherReadScope = Other.GetReadAccessScope();
+		static_assert(std::is_same_v<ThreadSafetyMode, typename UE::Core::Private::TScriptDelegateTraits<ThreadSafetyMode>::UnicastThreadSafetyModeForMulticasts>);
+
+		typename TScriptDelegate<OtherThreadSafetyMode>::FReadAccessScope OtherReadScope = Other.GetReadAccessScope();
 
 		TScriptDelegate Copy;
 		Copy.Object = Other.Object;
@@ -475,27 +481,28 @@ protected:
 	friend struct TIsZeroConstructType<TScriptDelegate>;
 };
 
-
-template <typename Dummy>
-struct TIsZeroConstructType<TScriptDelegate<Dummy>>
+template<typename ThreadSafetyMode>
+struct TIsZeroConstructType<TScriptDelegate<ThreadSafetyMode>>
 {
-	static constexpr bool Value = TIsZeroConstructType<typename TScriptDelegate<Dummy>::WeakPtrType>::Value;
+	static constexpr bool Value = 
+		TIsZeroConstructType<typename UE::Core::Private::TScriptDelegateTraits<ThreadSafetyMode>::WeakPtrType>::Value &&
+		TIsZeroConstructType<typename TScriptDelegate<ThreadSafetyMode>::Super>::Value;
 };
 
-
 /**
  * Script multi-cast delegate base class
  */
-template <typename Dummy>
-class TMulticastScriptDelegate : public FDelegateAccessHandlerBaseChecked
+template <typename InThreadSafetyMode>
+class TMulticastScriptDelegate : public TDelegateAccessHandlerBase<typename UE::Core::Private::TScriptDelegateTraits<InThreadSafetyMode>::ThreadSafetyMode>
 {
-	using Super = FDelegateAccessHandlerBaseChecked;
+	using ThreadSafetyMode = typename UE::Core::Private::TScriptDelegateTraits<InThreadSafetyMode>::ThreadSafetyMode;
+	using Super = TDelegateAccessHandlerBase<ThreadSafetyMode>;
 	using typename Super::FReadAccessScope;
 	using Super::GetReadAccessScope;
 	using typename Super::FWriteAccessScope;
 	using Super::GetWriteAccessScope;
 
-	using UnicastDelegateType = TScriptDelegate<Dummy>;
+	using UnicastDelegateType = TScriptDelegate<typename UE::Core::Private::TScriptDelegateTraits<InThreadSafetyMode>::UnicastThreadSafetyModeForMulticasts>;
 
 public:
 
@@ -580,7 +587,7 @@ public:
 	 * @param	InDelegate	Delegate to check
 	 * @return	True if the delegate is already in the list.
 	 */
-	bool Contains( const UnicastDelegateType& InDelegate ) const
+	bool Contains( const TScriptDelegate<ThreadSafetyMode>& InDelegate ) const
 	{
 		const UObject* Object;
 		FName FunctionName;
@@ -594,17 +601,17 @@ public:
 		return Contains(Object, FunctionName);
 	}
 	template <
-		typename OtherDummy
-		UE_REQUIRES(UE::Core::Private::BackwardCompatibilityCheck<Dummy, OtherDummy>())
+		typename OtherThreadSafetyMode
+		UE_REQUIRES(UE::Core::Private::BackwardCompatibilityCheck<InThreadSafetyMode, OtherThreadSafetyMode>())
 	>
 	/* UE_DEPRECATED(5.3, "Deprecated - remove after TScriptDelegateTraits<FWeakObjectPtr> is removed") */
-	bool Contains(const TScriptDelegate<OtherDummy>& InDelegate) const
+	bool Contains(const TScriptDelegate<OtherThreadSafetyMode>& InDelegate) const
 	{
 		const UObject* Object;
 		FName FunctionName;
 
 		{
-			typename TScriptDelegate<OtherDummy>::FReadAccessScope OtherReadScope = InDelegate.GetReadAccessScope();
+			typename TScriptDelegate<OtherThreadSafetyMode>::FReadAccessScope OtherReadScope = InDelegate.GetReadAccessScope();
 			Object = InDelegate.Object.Get();
 			FunctionName = InDelegate.FunctionName;
 		}
@@ -633,7 +640,7 @@ public:
 	 *
 	 * @param	InDelegate	Delegate to add
 	 */
-	void Add( const UnicastDelegateType& InDelegate )
+	void Add( const TScriptDelegate<ThreadSafetyMode>& InDelegate )
 	{
 		UnicastDelegateType LocalCopy = UnicastDelegateType::CopyFrom(InDelegate);
 
@@ -648,11 +655,11 @@ public:
 		}
 	}
 	template <
-		typename OtherDummy
-		UE_REQUIRES(UE::Core::Private::BackwardCompatibilityCheck<Dummy, OtherDummy>())
+		typename OtherThreadSafetyMode
+		UE_REQUIRES(UE::Core::Private::BackwardCompatibilityCheck<InThreadSafetyMode, OtherThreadSafetyMode>())
 	>
 	/* UE_DEPRECATED(5.3, "Deprecated - remove after TScriptDelegateTraits<FWeakObjectPtr> is removed") */
-	void Add(const TScriptDelegate<OtherDummy>& InDelegate)
+	void Add(const TScriptDelegate<OtherThreadSafetyMode>& InDelegate)
 	{
 		UnicastDelegateType LocalCopy = UnicastDelegateType::CopyFrom(InDelegate);
 
@@ -673,7 +680,7 @@ public:
 	 *
 	 * @param	InDelegate	Delegate to add
 	 */
-	void AddUnique( const UnicastDelegateType& InDelegate )
+	void AddUnique( const TScriptDelegate<ThreadSafetyMode>& InDelegate )
 	{
 		UnicastDelegateType LocalCopy = UnicastDelegateType::CopyFrom(InDelegate);
 
@@ -688,11 +695,11 @@ public:
 		}
 	}
 	template <
-		typename OtherDummy
-		UE_REQUIRES(UE::Core::Private::BackwardCompatibilityCheck<Dummy, OtherDummy>())
+		typename OtherThreadSafetyMode
+		UE_REQUIRES(UE::Core::Private::BackwardCompatibilityCheck<InThreadSafetyMode, OtherThreadSafetyMode>())
 	>
 	/* UE_DEPRECATED(5.3, "Deprecated - remove after TScriptDelegateTraits<FWeakObjectPtr> is removed") */
-	void AddUnique(const TScriptDelegate<OtherDummy>& InDelegate)
+	void AddUnique(const TScriptDelegate<OtherThreadSafetyMode>& InDelegate)
 	{
 		UnicastDelegateType LocalCopy = UnicastDelegateType::CopyFrom(InDelegate);
 
@@ -713,7 +720,7 @@ public:
 	 *
 	 * @param	InDelegate	Delegate to remove
 	 */
-	void Remove( const UnicastDelegateType& InDelegate )
+	void Remove( const TScriptDelegate<ThreadSafetyMode>& InDelegate )
 	{
 		UnicastDelegateType LocalCopy = UnicastDelegateType::CopyFrom(InDelegate);
 
@@ -728,11 +735,11 @@ public:
 		}
 	}
 	template <
-		typename OtherDummy
-		UE_REQUIRES(UE::Core::Private::BackwardCompatibilityCheck<Dummy, OtherDummy>())
+		typename OtherThreadSafetyMode
+		UE_REQUIRES(UE::Core::Private::BackwardCompatibilityCheck<InThreadSafetyMode, OtherThreadSafetyMode>())
 	>
 	/* UE_DEPRECATED(5.3, "Deprecated - remove after TScriptDelegateTraits<FWeakObjectPtr> is removed") */
-	void Remove(const TScriptDelegate<OtherDummy>& InDelegate)
+	void Remove(const TScriptDelegate<OtherThreadSafetyMode>& InDelegate)
 	{
 		UnicastDelegateType LocalCopy = UnicastDelegateType::CopyFrom(InDelegate);
 
@@ -1054,8 +1061,8 @@ protected:
 };
 
 
-template <typename TWeakPtr>
-struct TIsZeroConstructType<TMulticastScriptDelegate<TWeakPtr>>
-{
-	static constexpr bool Value = TIsZeroConstructType<typename TMulticastScriptDelegate<TWeakPtr>::InvocationListType>::Value;
+template<typename ThreadSafetyMode> 
+struct TIsZeroConstructType<TMulticastScriptDelegate<ThreadSafetyMode> >
+{ 
+	static constexpr bool Value = TIsZeroConstructType<typename TMulticastScriptDelegate<ThreadSafetyMode>::InvocationListType>::Value;
 };
diff --git a/Engine/Source/Runtime/Core/Tests/Delegates/TSDelegatesTest.cpp b/Engine/Source/Runtime/Core/Tests/Delegates/TSDelegatesTest.cpp
--- a/Engine/Source/Runtime/Core/Tests/Delegates/TSDelegatesTest.cpp
+++ b/Engine/Source/Runtime/Core/Tests/Delegates/TSDelegatesTest.cpp
@@ -5,6 +5,9 @@
 #include "Delegates/Delegate.h"
 #include "Tasks/Task.h"
 #include "Tests/TestHarnessAdapter.h"
+#include "Async/ParallelFor.h"
+
+#include "HAL/PlatformTLS.h"
 
 #if WITH_TESTS
 
@@ -12,126 +15,457 @@ namespace UE { namespace TSDelegatesTest
 {
 	using namespace Tasks;
 
-	void Func() {}
-
-	void OneParamFunc(bool* bExecuted)
+	TEST_CASE_NAMED(FUnicastDelegateBasicTest, "System::Core::Delegates::Unicast::Basic", "[ApplicationContextMask][EngineFilter]")
 	{
-		*bExecuted = true;
-	}
+		{	// single thread, basic functionality, check compilation/execution and that there's no false race detection
+			bool bExecuted = false;
+			TDelegate<void()> Delegate;
+			Delegate.BindLambda([&bExecuted] { bExecuted = true; });
+			Delegate.Execute();
+			Delegate.Unbind();
+			check(bExecuted);
+		}
 
-	void AnotherOneParamFunc(bool* bExecuted)
-	{
-		*bExecuted = true;
-	}
+		{	// binding on "any thread"
+			Launch(UE_SOURCE_LOCATION, [] 
+			{ 
+				TDelegate<void()>().BindLambda([]{});
+			}).Wait();
+		}
 
-	TEST_CASE_NAMED(FTSMulticastDelegateBasicTest, "System::Core::Delegates::MulticastBasic", "[ApplicationContextMask][EngineFilter]")
-	{
-		DECLARE_TS_MULTICAST_DELEGATE(FTestDelegate);
+		{	// executing on "any thread"
+			TDelegate<void()> Delegate;
+			Delegate.BindLambda([]{});
+			Launch(UE_SOURCE_LOCATION, [&Delegate] { Delegate.Execute(); }).Wait();
+		}
 
-		{	// single thread
-			FTestDelegate TestDelegate;
+		{	// unbinding on "any thread"
+			TDelegate<void()> Delegate;
+			Delegate.BindLambda([]{});
+			Launch(UE_SOURCE_LOCATION, [&Delegate] { Delegate.Unbind(); }).Wait();
+		}
+
+		{	// compilation test for returning a non-copyable type
+			struct FNoncopyableExample
+			{
+				UE_NONCOPYABLE(FNoncopyableExample);
+				FNoncopyableExample() = default;
+			};
+			TDelegate<const FNoncopyableExample& ()> Delegate;
+			FNoncopyableExample RetValue;
+			Delegate.BindLambda([&RetValue] () -> const FNoncopyableExample& { return RetValue; });
+			Delegate.Execute();
+		}
+
+		{	// move
+			TDelegate<void()> Source;
 			bool bExecuted = false;
-			FDelegateHandle DelegateHandle = TestDelegate.AddStatic(&OneParamFunc, &bExecuted);
-			TestDelegate.AddStatic(&AnotherOneParamFunc, &bExecuted);
-			TestDelegate.Broadcast();
+			Source.BindLambda([&bExecuted] { bExecuted = true; });
+			
+			// move construction
+			TDelegate<void()> Dest = MoveTemp(Source);
+			check(!Source.IsBound());
+			Dest.Execute();
+			check(bExecuted);
+
+			// move assignment
+			Source.BindLambda([&bExecuted] { bExecuted = true; });
+			Dest = MoveTemp(Source);
+			check(!Source.IsBound());
+			bExecuted = false;
+			Dest.Execute();
 			check(bExecuted);
-			TestDelegate.Remove(DelegateHandle);
 		}
 
-		{	// concurrent Add/Remove and Execute
-			FTestDelegate TestDelegate;
+		{	// copy
+			TDelegate<void()> Source;
 			bool bExecuted = false;
-			FTask AddStaticTask = Launch(UE_SOURCE_LOCATION, [&TestDelegate, &bExecuted] { TestDelegate.AddStatic(&AnotherOneParamFunc, &bExecuted); });
-			FTask BroadcastTask = Launch(UE_SOURCE_LOCATION, [&TestDelegate] { TestDelegate.Broadcast(); });
-			FDelegateHandle DelegateHandle = TestDelegate.AddStatic(&OneParamFunc, &bExecuted);
-			TestDelegate.Remove(DelegateHandle);
-			Wait(TArray{ AddStaticTask, BroadcastTask });
+			Source.BindLambda([&bExecuted] { bExecuted = true; });
+			// copy construction
+			TDelegate<void()> Dest = Source;
+			Dest.Execute();
+			check(bExecuted);
+			// copy assignment
+			Dest = Source;
+			bExecuted = false;
+			Dest.Execute();
+			check(bExecuted);
 		}
 
-		{ // same as above, but Add/Remove from different threads (not concurrently)
-			FTestDelegate TestDelegate;
-			bool bExecuted = false;
-			TTask<FDelegateHandle> AddStaticTask = Launch(UE_SOURCE_LOCATION, [&TestDelegate, &bExecuted] { return TestDelegate.AddStatic(&OneParamFunc, &bExecuted); });
-			FTask RemoveTask = Launch(UE_SOURCE_LOCATION, [&TestDelegate, &AddStaticTask] { TestDelegate.Remove(AddStaticTask.GetResult()); }, AddStaticTask);
-			TestDelegate.Broadcast();
-			RemoveTask.Wait();
-			TestDelegate.Remove(AddStaticTask.GetResult());
+		///////////////////////////////////////////////////////////////////
+		// unicast delegates can destroy themselves from inside their execution
+
+		{
+			TDelegate<void()>* Delegate = new TDelegate<void()>;
+			Delegate->BindLambda([Delegate] { delete Delegate; });
+			Delegate->Execute();
 		}
 
-		{	// recursive locking: removing a delegate handle from inside broadcasting
-			FTestDelegate TestDelegate;
-			FDelegateHandle DelegateHandle;
-			auto DelegateLambda = [&TestDelegate, &DelegateHandle] { TestDelegate.Remove(DelegateHandle); };
-			DelegateHandle = TestDelegate.AddLambda(MoveTemp(DelegateLambda));
-			check(TestDelegate.IsBound());
-			TestDelegate.Broadcast();
-			check(!TestDelegate.IsBound());
+		{	// nested destruction of different delegates
+			TDelegate<void()>* Delegate = new TDelegate<void()>;
+			Delegate->BindLambda(
+				[Delegate] 
+				{
+					TDelegate<void()>* NestedDelegate = new TDelegate<void()>;
+					NestedDelegate->BindLambda([NestedDelegate] { delete NestedDelegate; });
+					NestedDelegate->Execute();
+					
+					delete Delegate;
+				}
+			);
+			Delegate->Execute();
 		}
 
-		// test that all macros compile
-		{
-			DECLARE_TS_MULTICAST_DELEGATE_OneParam(FDelegate, int);
-			FDelegate Delegate;
-			Delegate.AddLambda([](int) {});
-			Delegate.Broadcast(0);
+		{	// delegate destruction from inside nested access
+			bool bDestroy = false;
+			TDelegate<void()>* Delegate = new TDelegate<void()>;
+			Delegate->BindLambda(
+				[Delegate, &bDestroy] 
+				{ 
+					if (!bDestroy)
+					{
+						bDestroy = true;
+						Delegate->Execute();
+					}
+					else
+					{
+						delete Delegate;
+					}
+				}
+			);
+			Delegate->Execute();
+			check(bDestroy);
 		}
-		{
-			DECLARE_TS_MULTICAST_DELEGATE_TwoParams(FDelegate, int, int);
-			FDelegate Delegate;
-			Delegate.AddLambda([](int, int) {});
-			Delegate.Broadcast(0, 0);
+
+		{	// delegate destruction from inside nested access with intermediate access to another delegate
+			TDelegate<void()>* Delegate = new TDelegate<void()>;
+			Delegate->BindLambda(
+				[Delegate]
+				{
+					TDelegate<void()> AnotherDelegate;
+					AnotherDelegate.BindLambda(
+						[Delegate] 
+						{
+							Delegate->Unbind();
+							Delegate->BindLambda([Delegate] { delete Delegate; });
+						}
+					);
+					AnotherDelegate.Execute();
+				}
+			);
+			Delegate->Execute();
 		}
-		{
-			DECLARE_TS_MULTICAST_DELEGATE_ThreeParams(FDelegate, int, int, int);
-			FDelegate Delegate;
-			Delegate.AddLambda([](int, int, int) {});
-			Delegate.Broadcast(0, 0, 0);
+
+		{	// multicast delegate destruction
+			TMulticastDelegate<void()>* Delegate = new TMulticastDelegate<void()>;
+			Delegate->AddLambda([Delegate] { delete Delegate; });
+			Delegate->Broadcast();
 		}
+
 		{
-			DECLARE_TS_MULTICAST_DELEGATE_FourParams(FDelegate, int, int, int, int);
-			FDelegate Delegate;
-			Delegate.AddLambda([](int, int, int, int) {});
-			Delegate.Broadcast(0, 0, 0, 0);
+			struct FDelegateHolder
+			{
+				TDelegate<void()> Delegate;
+
+				void SelfDestruct()
+				{
+					delete this;
+				}
+			};
+
+			FDelegateHolder* DelegateHolder = new FDelegateHolder;
+			DelegateHolder->Delegate.BindRaw(DelegateHolder, &FDelegateHolder::SelfDestruct);
+			DelegateHolder->Delegate.Execute();
 		}
-		{
-			DECLARE_TS_MULTICAST_DELEGATE_FiveParams(FDelegate, int, int, int, int, int);
-			FDelegate Delegate;
-			Delegate.AddLambda([](int, int, int, int, int) {});
-			Delegate.Broadcast(0, 0, 0, 0, 0);
+
+		//////////////////////////////////////////////////////////////////////
+
+		{	// Executing (which is a read op) a delegate concurrently, which can be safe if the binding is safe
+			// the same case with a TS delegate would deadlock as we don't use a reader-writer mutex.
+			// Broadcasting a multicast delegate is a write op and so would be caught by its race detector.
+			// The executions on different threads are deliberately partially overlapped to check for an edge case.
+
+			FTaskEvent FirstExecuting{ UE_SOURCE_LOCATION };
+			FTaskEvent FirstFinish{ UE_SOURCE_LOCATION };
+			FTaskEvent SecondExecuting{ UE_SOURCE_LOCATION };
+			FTaskEvent SecondFinish{ UE_SOURCE_LOCATION };
+
+			TDelegate<void(FTaskEvent&, FTaskEvent&)> Delegate;
+			Delegate.BindLambda([](FTaskEvent& Executing, FTaskEvent& Finish)
+			{
+				Executing.Trigger();
+				Finish.Wait();
+			});
+
+			FTask FirstTask = Launch(UE_SOURCE_LOCATION, [&Delegate, &FirstExecuting, &FirstFinish] { Delegate.Execute(FirstExecuting, FirstFinish); });
+			FTask SecondTask = Launch(UE_SOURCE_LOCATION, [&Delegate, &SecondExecuting, &SecondFinish] { Delegate.Execute(SecondExecuting, SecondFinish); });
+			Wait(TArray{ FirstExecuting, SecondExecuting });
+			// at this moment there're two concurrent executions of the same delegate
+			FirstFinish.Trigger();
+			SecondFinish.Trigger();
+			Wait(TArray{ FirstTask, SecondTask });
 		}
-		{
-			DECLARE_TS_MULTICAST_DELEGATE_SixParams(FDelegate, int, int, int, int, int, int);
-			FDelegate Delegate;
-			Delegate.AddLambda([](int, int, int, int, int, int) {});
-			Delegate.Broadcast(0, 0, 0, 0, 0, 0);
+
+		{	// executing a delegate stored in an array with binding that modifies the array causing realloc.
+			// this is incompatible with race detection, so a "not checked" delegate is used
+			using FNotCheckedDelegate = TDelegate<void(), FNotThreadSafeNotCheckedDelegateUserPolicy>;
+			TArray<FNotCheckedDelegate> Ds;
+			FNotCheckedDelegate D;
+			D.BindLambda([&Ds] { Ds.Add(FNotCheckedDelegate{}); });
+			Ds.Add(D);
+			Ds.Shrink();
+			Ds[0].Execute();
 		}
-		{
-			DECLARE_TS_MULTICAST_DELEGATE_SevenParams(FDelegate, int, int, int, int, int, int, int);
-			FDelegate Delegate;
-			Delegate.AddLambda([](int, int, int, int, int, int, int) {});
-			Delegate.Broadcast(0, 0, 0, 0, 0, 0, 0);
+
+		{ // destroying a container with delegates
+			TArray<TDelegate<void()>> Array;
+			Array.AddDefaulted();
+			TDelegate<void()> Delegate;
+			Delegate.BindLambda([] {});
+			Array.Add(Delegate);
+			Array.Add(MoveTemp(Delegate));
 		}
-		{
-			DECLARE_TS_MULTICAST_DELEGATE_EightParams(FDelegate, int, int, int, int, int, int, int, int);
-			FDelegate Delegate;
-			Delegate.AddLambda([](int, int, int, int, int, int, int, int) {});
-			Delegate.Broadcast(0, 0, 0, 0, 0, 0, 0, 0);
+	}
+
+	TEST_CASE_NAMED(FUnicastTSDelegateBasicTest, "System::Core::Delegates::TS::Unicast::Basic", "[ApplicationContextMask][EngineFilter]")
+	{
+		{	// single thread
+			TTSDelegate<void()> Delegate;
+			bool bExecuted = false;
+			Delegate.BindLambda([&bExecuted] { bExecuted = true; });
+			Delegate.Execute();
+			check(bExecuted);
 		}
-		{
-			DECLARE_TS_MULTICAST_DELEGATE_NineParams(FDelegate, int, int, int, int, int, int, int, int, int);
-			FDelegate Delegate;
-			Delegate.AddLambda([](int, int, int, int, int, int, int, int, int) {});
-			Delegate.Broadcast(0, 0, 0, 0, 0, 0, 0, 0, 0);
+
+		{	// unbind from inside execution
+			TTSDelegate<void()> Delegate;
+			Delegate.BindLambda([&Delegate] { Delegate.Unbind(); });
+			Delegate.Execute();
+			check(!Delegate.IsBound());
+		}
+
+		{	// binging/unbinding concurrently with execution
+			TTSDelegate<void()> Delegate;
+			FTask Task = Launch(UE_SOURCE_LOCATION, [&Delegate] 
+			{ 
+				Delegate.BindLambda([] {}); 
+				Delegate.Unbind();
+			});
+			Delegate.ExecuteIfBound();
+			Task.Wait();
+		}
+	}
+
+	TEST_CASE_NAMED(FUnicastTSDelegateStressTest, "System::Core::Delegates::TS::Unicast::Stress", "[ApplicationContextMask][EngineFilter]")
+	{
+		{	// one thread repeatedly binds/unbinds the delegate, another executes
+			TTSDelegate<void()> Delegate;
+			std::atomic<bool> bQuit{ false };
+			std::atomic<int32> NumExecuted = 0;
+			TArray<FTask> Tasks;
+
+			Tasks.Add(Launch(UE_SOURCE_LOCATION, [&bQuit, &Delegate]
+			{
+				while (!bQuit)
+				{
+					Delegate.BindLambda([]{});
+					FPlatformProcess::YieldCycles(100);
+					Delegate.Unbind();
+					FPlatformProcess::YieldCycles(100);
+				}
+			}));
+
+			for (int i = 0; i != FPlatformMisc::NumberOfCores() / 2; ++i)
+			{
+				Tasks.Add(Launch(UE_SOURCE_LOCATION,
+					[&bQuit, &Delegate, &NumExecuted]
+					{
+						while (!bQuit)
+						{
+							if (Delegate.ExecuteIfBound())
+							{
+								++NumExecuted;
+							}
+						}
+					}
+				));
+			}
+
+			FPlatformProcess::Sleep(1.0);
+			bQuit = true;
+			Wait(Tasks);
+		}
+	}
+
+	TEST_CASE_NAMED(FMulticastDelegateBasicTest, "System::Core::Delegates::Multicast::Basic", "[ApplicationContextMask][EngineFilter]")
+	{
+		{ // add 2 bindings, broadcast, remove one binding, broadcast
+			TMulticastDelegate<void()> Delegate;
+			bool bExecuted1 = false;
+			bool bExecuted2 = false;
+			FDelegateHandle DelegateHandle1 = Delegate.AddLambda([&bExecuted1] { bExecuted1 = true; });
+			FDelegateHandle DelegateHandle2 = Delegate.AddLambda([&bExecuted2] { bExecuted2 = true; });
+			Delegate.Broadcast();
+			check(bExecuted1 && bExecuted2);
+			Delegate.Remove(DelegateHandle1);
+			bExecuted1 = bExecuted2 = false;
+			Delegate.Broadcast();
+			check(!bExecuted1 && bExecuted2);
+		}
+
+		{	// check removing "raw method" binding by object ptr
+			struct FDummy
+			{
+				void Method(int InValue)
+				{
+					Value = InValue;
+				}
+
+				int Value = 0;
+			};
+
+			TMulticastDelegate<void(int)> Delegate;
+			FDummy Dummy1;
+			Delegate.AddRaw(&Dummy1, &FDummy::Method);
+			{
+				FDummy Dummy2;
+				Delegate.AddRaw(&Dummy2, &FDummy::Method);
+				Delegate.Broadcast(1);
+				check(Dummy1.Value == 1 && Dummy2.Value == 1);
+				verify(Delegate.RemoveAll(&Dummy2) == 1);
+			}
+			Delegate.Broadcast(2);
+			check(Dummy1.Value == 2);
+			verify(Delegate.RemoveAll(&Dummy1));
+		}
+
+		{	// using unicast API instead of `Add*` shortcuts
+			TMulticastDelegate<void()> Delegate;
+			TDelegate<void()> UniDelegate;
+			bool bExecuted = false;
+			UniDelegate.BindLambda([&bExecuted] { bExecuted = true; });
+			Delegate.Add(UniDelegate);
+			Delegate.Broadcast();
+			check(bExecuted);
+		}
+
+		{	// move
+			TMulticastDelegate<void()> Source;
+			int32 NumExecuted = 0;
+			Source.AddLambda([&NumExecuted] { ++NumExecuted; });
+			Source.AddLambda([&NumExecuted] { ++NumExecuted; });
+			// move construction
+			TMulticastDelegate<void()> Dest = MoveTemp(Source);
+			check(!Source.IsBound());
+			Dest.Broadcast();
+			check(NumExecuted == 2);
+			// move assignment
+			Source.AddLambda([&NumExecuted] { ++NumExecuted; });
+			Source.AddLambda([&NumExecuted] { ++NumExecuted; });
+			Dest = MoveTemp(Source);
+			check(!Source.IsBound());
+			NumExecuted = 0;
+			Dest.Broadcast();
+			check(NumExecuted == 2);
+		}
+
+		{	// copy
+			TMulticastDelegate<void()> Source;
+			int32 NumExecuted = 0;
+			Source.AddLambda([&NumExecuted] { ++NumExecuted; });
+			Source.AddLambda([&NumExecuted] { ++NumExecuted; });
+			// copy construction
+			TMulticastDelegate<void()> Dest = Source;
+			Dest.Broadcast();
+			check(NumExecuted == 2);
+			// copy assignment
+			Dest = Source;
+			NumExecuted = 0;
+			Dest.Broadcast();
+			check(NumExecuted == 2);
 		}
 
 	}
 
 	TEST_CASE_NAMED(FTSMulticastDelegateStressTest, "System::Core::Delegates::MulticastStress", "[ApplicationContextMask][EngineFilter]")
 	{
-		DECLARE_TS_MULTICAST_DELEGATE(FTestDelegate);
+		{	// single thread
+			TTSMulticastDelegate<void()> Delegate;
+			int32 ExecRes = 0;
+			FDelegateHandle DelegateHandle = Delegate.AddLambda([&ExecRes] { ++ExecRes; });
+			Delegate.AddLambda([&ExecRes] { ExecRes += 2; });
+			Delegate.Broadcast();
+			check(ExecRes == 3);
+			Delegate.Remove(DelegateHandle);
+			ExecRes = 0;
+			Delegate.Broadcast();
+			check(ExecRes == 2);
+		}
+
+		{	// concurrent Add/Remove and Execute
+			TTSMulticastDelegate<void()> Delegate;
+			FTask AddBindingTask = Launch(UE_SOURCE_LOCATION, [&Delegate] 
+			{ 
+				Delegate.AddLambda([]{});
+			});
+			FTask BroadcastTask = Launch(UE_SOURCE_LOCATION, [&Delegate] { Delegate.Broadcast(); });
+			FDelegateHandle DelegateHandle = Delegate.AddLambda([]{});
+			Delegate.Remove(DelegateHandle);
+			Wait(TArray{ AddBindingTask, BroadcastTask });
+		}
+
+		{ // same as above, but Add/Remove from different threads (not concurrently)
+			TTSMulticastDelegate<void()> Delegate;
+			TTask<FDelegateHandle> AddBindingTask = Launch(UE_SOURCE_LOCATION, [&Delegate] 
+			{ 
+				return Delegate.AddLambda([]{});
+			});
+			FTask RemoveBindingTask = Launch(UE_SOURCE_LOCATION, [&Delegate, &AddBindingTask] { Delegate.Remove(AddBindingTask.GetResult()); }, Prerequisites(AddBindingTask));
+			Delegate.Broadcast();
+			RemoveBindingTask.Wait();
+		}
+
+		{	// recursive locking: removing a delegate handle from inside broadcasting
+			TTSMulticastDelegate<void()> Delegate;
+			FDelegateHandle DelegateHandle;
+			auto DelegateLambda = [&Delegate, &DelegateHandle] { Delegate.Remove(DelegateHandle); };
+			DelegateHandle = Delegate.AddLambda(MoveTemp(DelegateLambda));
+			check(Delegate.IsBound());
+			Delegate.Broadcast();
+			check(!Delegate.IsBound());
+		}
 
+#if UE_DETECT_DELEGATES_RACE_CONDITIONS
+		{	// thread-unsafe delegate with disabled race detection
+			// unbinding would assert if `FNotThreadSafeNotCheckedDelegateUserPolicy` wasn't specified.
+			// It's still a bug as the bound lambda gets destroyed on the delegate removal, while it's in the middle of execution and can try accessing
+			// its capture that is destroyed
+
+			//TMulticastDelegate<void(), FNotThreadSafeNotCheckedDelegateUserPolicy> Delegate;
+
+			//// model concurrent execution and unbinding
+			//FTaskEvent BroadcastingStarted{ UE_SOURCE_LOCATION };
+			//FTaskEvent ResumeBroadcasting{ UE_SOURCE_LOCATION };
+			//// get blocked inside broadcasting
+			//FDelegateHandle DelegateHandle = Delegate.AddLambda([&BroadcastingStarted, &ResumeBroadcasting]
+			//{
+			//	BroadcastingStarted.Trigger();
+			//	ResumeBroadcasting.Wait(); 
+			//});
+
+			//FTask Task = Launch(UE_SOURCE_LOCATION, [&Delegate] { Delegate.Broadcast(); });
+			//BroadcastingStarted.Wait();
+			//Delegate.Remove(DelegateHandle); // happens concurrently with broadcasting on a thread-unsafe delegate, so it's a race
+			//ResumeBroadcasting.Trigger();
+			//Task.Wait();
+		}
+#endif
+	}
+
+	TEST_CASE_NAMED(FMulticastTSDelegateStressTest, "System::Core::Delegates::TS::Multicast::Stress", "[ApplicationContextMask][EngineFilter]")
+	{
 		{	// one thread repeatedly adds/removes delegates, another broadcasts
-			FTestDelegate Delegate;
+			TTSMulticastDelegate<void()> Delegate;
 			std::atomic<bool> bQuit{ false };
 
 			FTask Binding = Launch(UE_SOURCE_LOCATION,
@@ -139,7 +473,7 @@ namespace UE { namespace TSDelegatesTest
 				{
 					while (!bQuit)
 					{
-						FDelegateHandle Handle = Delegate.AddStatic(&Func);
+						FDelegateHandle Handle = Delegate.AddLambda([]{});
 						Delegate.Remove(Handle);
 					}
 				}
@@ -157,7 +491,7 @@ namespace UE { namespace TSDelegatesTest
 
 			FPlatformProcess::Sleep(1.0);
 			bQuit = true;
-			Wait(TArray<FTask>{ Binding, Executing });
+			Wait(TArray{ Binding, Executing });
 		}
 	}
 }}
diff --git a/Engine/Source/Runtime/CoreUObject/Public/UObject/ScriptDelegateFwd.h b/Engine/Source/Runtime/CoreUObject/Public/UObject/ScriptDelegateFwd.h
--- a/Engine/Source/Runtime/CoreUObject/Public/UObject/ScriptDelegateFwd.h
+++ b/Engine/Source/Runtime/CoreUObject/Public/UObject/ScriptDelegateFwd.h
@@ -8,14 +8,15 @@
 
 #include "UObject/WeakObjectPtrFwd.h"
 
-namespace UE::Core::Private
-{
-	struct TScriptDelegateDefault;
-}
+struct FNotThreadSafeDelegateMode;
+struct FThreadSafeDelegateMode;
 
-template <typename Dummy = UE::Core::Private::TScriptDelegateDefault> class TScriptDelegate;
-template <typename Dummy = UE::Core::Private::TScriptDelegateDefault> class TMulticastScriptDelegate;
+template <typename ThreadSafetyMode = FNotThreadSafeDelegateMode> class TScriptDelegate;
+template <typename ThreadSafetyMode = FNotThreadSafeDelegateMode> class TMulticastScriptDelegate;
 
 // Typedef script delegates for convenience.
 typedef TScriptDelegate<> FScriptDelegate;
 typedef TMulticastScriptDelegate<> FMulticastScriptDelegate;
+
+typedef TScriptDelegate<FThreadSafeDelegateMode> FTSScriptDelegate;
+typedef TMulticastScriptDelegate<FThreadSafeDelegateMode> FTSMulticastScriptDelegate;
diff --git a/Engine/Source/Runtime/Engine/Public/TimerManager.h b/Engine/Source/Runtime/Engine/Public/TimerManager.h
--- a/Engine/Source/Runtime/Engine/Public/TimerManager.h
+++ b/Engine/Source/Runtime/Engine/Public/TimerManager.h
@@ -19,7 +19,9 @@ class UGameInstance;
 enum class ELevelCollectionType : uint8;
 struct FTimerSourceList;
 
-DECLARE_DELEGATE(FTimerDelegate);
+// using "not checked" user policy (means race detection is disabled) because this delegate is stored in a TSparseArray and causes its reallocation
+// from inside delegate's execution. This is incompatible with race detection that needs to access the delegate instance after its execution
+using FTimerDelegate = TDelegate<void(), FNotThreadSafeNotCheckedDelegateUserPolicy>;
 
 #ifndef UE_ENABLE_TRACKING_TIMER_SOURCES
 #define UE_ENABLE_TRACKING_TIMER_SOURCES !UE_BUILD_SHIPPING
diff --git a/Engine/Source/Runtime/MovieScene/Public/MovieSceneLatentActionManager.h b/Engine/Source/Runtime/MovieScene/Public/MovieSceneLatentActionManager.h
--- a/Engine/Source/Runtime/MovieScene/Public/MovieSceneLatentActionManager.h
+++ b/Engine/Source/Runtime/MovieScene/Public/MovieSceneLatentActionManager.h
@@ -9,7 +9,9 @@
 
 class UObject;
 
-DECLARE_DELEGATE(FMovieSceneSequenceLatentActionDelegate);
+// using "not checked" user policy (means race detection is disabled) because this delegate is stored in a TArray and causes its reallocation
+// from inside delegate's execution. This is incompatible with race detection that needs to access the delegate instance after its execution
+using FMovieSceneSequenceLatentActionDelegate = TDelegate<void(), FNotThreadSafeNotCheckedDelegateUserPolicy>;
 
 /**
  * Utility class for running latent actions created from sequence players.
