Add GetSize() to TBound, consistent with other bounding shapes in the math module.

diff --git a/Engine/Source/Runtime/Core/Public/Math/Bounds.h b/Engine/Source/Runtime/Core/Public/Math/Bounds.h
--- a/Engine/Source/Runtime/Core/Public/Math/Bounds.h
+++ b/Engine/Source/Runtime/Core/Public/Math/Bounds.h
@@ -78,6 +78,11 @@ struct TBounds
 		return (Max - Min) * 0.5f;
 	}
 
+	FORCEINLINE TVector<T> GetSize() const
+	{
+		return (Max - Min);
+	}
+
 	FORCEINLINE T GetSurfaceArea() const
 	{
 		TVector<T> Size = Max - Min;
