Added missing header to Parse.h

diff --git a/Engine/Source/Runtime/Core/Public/Misc/Parse.h b/Engine/Source/Runtime/Core/Public/Misc/Parse.h
--- a/Engine/Source/Runtime/Core/Public/Misc/Parse.h
+++ b/Engine/Source/Runtime/Core/Public/Misc/Parse.h
@@ -7,6 +7,7 @@
 #include "CoreTypes.h"
 #include "HAL/PlatformCrt.h"
 #include "Misc/Build.h"
+#include "Templates/Function.h"
 
 class FExec;
 class FName;
