Added an option on FArchive to early out while setting a custom version if that version was already set in the local custom versions container.

In projects with several custom versions, this results in a considerable improvement because the local container is usually smaller than the global one, and therefore the look-up time is faster, also we avoid going through the SetVersion path.

But this means if the version data in the version registry changed to the specific key you were trying to see, this means the local custom versions container of the archive will not be up-to-date. That is the reason this new option is disabled by default

diff --git a/Engine/Source/Runtime/Core/Private/Serialization/Archive.cpp b/Engine/Source/Runtime/Core/Private/Serialization/Archive.cpp
--- a/Engine/Source/Runtime/Core/Private/Serialization/Archive.cpp
+++ b/Engine/Source/Runtime/Core/Private/Serialization/Archive.cpp
@@ -191,6 +191,7 @@ void FArchiveState::Reset()
 	ArIsNetArchive						= false;
 	ArCustomPropertyList				= nullptr;
 	ArUseCustomPropertyList				= false;
+	ArShouldSkipUpdateCustomVersion		= false;
 	CookData							= nullptr;
 	SerializedProperty					= nullptr;
 
@@ -249,6 +250,7 @@ void FArchiveState::CopyTrivialFArchiveStatusMembers(const FArchiveState& Archiv
 	ArIsNetArchive                       = ArchiveToCopy.ArIsNetArchive;
 	ArCustomPropertyList                 = ArchiveToCopy.ArCustomPropertyList;
 	ArUseCustomPropertyList              = ArchiveToCopy.ArUseCustomPropertyList;
+	ArShouldSkipUpdateCustomVersion		 = ArchiveToCopy.ArShouldSkipUpdateCustomVersion;
 	CookData							 = ArchiveToCopy.CookData;
 	SerializedProperty					 = ArchiveToCopy.SerializedProperty;
 #if USE_STABLE_LOCALIZATION_KEYS
@@ -569,8 +571,8 @@ void FArchive::UsingCustomVersion(const FGuid& Key)
 		return;
 	}
 
-	FCustomVersion RegisteredVersion = FCurrentCustomVersions::Get(Key).GetValue();
-	const_cast<FCustomVersionContainer&>(GetCustomVersions()).SetVersion(Key, RegisteredVersion.Version, RegisteredVersion.GetFriendlyName());
+	ESetCustomVersionFlags SetVersionFlags = ArShouldSkipUpdateCustomVersion ? ESetCustomVersionFlags::SkipUpdateExistingVersion : ESetCustomVersionFlags::None;
+	const_cast<FCustomVersionContainer&>(GetCustomVersions()).SetVersionUsingRegistry(Key, SetVersionFlags);
 }
 
 int32 FArchiveState::CustomVer(const FGuid& Key) const
@@ -584,6 +586,11 @@ int32 FArchiveState::CustomVer(const FGuid& Key) const
 	return CustomVersion ? CustomVersion->Version : -1;
 }
 
+void FArchiveState::SetShouldSkipUpdateCustomVersion(bool bShouldSkip)
+{
+	ForEachState([bShouldSkip](FArchiveState& State) { State.ArShouldSkipUpdateCustomVersion = bShouldSkip; });
+}
+
 void FArchiveState::SetCustomVersion(const FGuid& Key, int32 Version, FName FriendlyName)
 {
 	const_cast<FCustomVersionContainer&>(GetCustomVersions()).SetVersion(Key, Version, FriendlyName);
diff --git a/Engine/Source/Runtime/Core/Private/Serialization/CustomVersion.cpp b/Engine/Source/Runtime/Core/Private/Serialization/CustomVersion.cpp
--- a/Engine/Source/Runtime/Core/Private/Serialization/CustomVersion.cpp
+++ b/Engine/Source/Runtime/Core/Private/Serialization/CustomVersion.cpp
@@ -400,3 +400,28 @@ void FCustomVersionContainer::SetVersion(FGuid CustomKey, int32 Version, FName F
 		Versions.Add(FCustomVersion(CustomKey, Version, FriendlyName));
 	}
 }
+
+void FCustomVersionContainer::SetVersionUsingRegistry(FGuid CustomKey, ESetCustomVersionFlags Options)
+{
+	if (CustomKey == UnusedCustomVersionKey)
+	{
+		return;
+	}
+
+	if (FCustomVersion* Found = Versions.FindByKey(CustomKey))
+	{
+		if (EnumHasAnyFlags(Options, ESetCustomVersionFlags::SkipUpdateExistingVersion))
+		{
+			return;
+		}
+
+		const FCustomVersion& RegisteredVersion = FCurrentCustomVersions::Get(CustomKey).GetValue();
+		Found->Version      = RegisteredVersion.Version;
+		Found->FriendlyName = RegisteredVersion.FriendlyName;
+	}
+	else
+	{
+		const FCustomVersion& RegisteredVersion = FCurrentCustomVersions::Get(CustomKey).GetValue();
+		Versions.Emplace(FCustomVersion(CustomKey, RegisteredVersion.Version, RegisteredVersion.FriendlyName));
+	}
+}
diff --git a/Engine/Source/Runtime/Core/Public/Serialization/Archive.h b/Engine/Source/Runtime/Core/Public/Serialization/Archive.h
--- a/Engine/Source/Runtime/Core/Public/Serialization/Archive.h
+++ b/Engine/Source/Runtime/Core/Public/Serialization/Archive.h
@@ -333,6 +333,19 @@ public:
 		return ArShouldSkipCompilingAssets;
 	}
 
+	/** If set true, UsingCurrentVersion will not query the version registry to get the latest version data for the key if the version key is already set in the local container.
+	 * Querying the version registry is not cheap, and the look up time increases as more custom versions are registered, therefore if you don't change the registered custom version data once is set in the archive
+	 * Setting this to true should improve serialization performance by skipping the step entirely.
+	 * @param bShouldSkip Whether we should skip the checking the version registry for new version data if the version key is already set
+	 */
+	void SetShouldSkipUpdateCustomVersion(bool bShouldSkip);
+
+	/** Returns true if this Archive is does not update the custom version data for a given key if has already been set previously in the local custom versions container*/
+	FORCEINLINE bool ShouldSkipUpdateCustomVersion() const
+	{
+		return ArShouldSkipUpdateCustomVersion;
+	}
+
 	/** Returns true if this archive contains errors, which means that further serialization is generally not safe. */
 	FORCEINLINE bool IsError() const
 	{
@@ -785,6 +798,11 @@ private:
 	/** Whether or not it is allowed to skip serialization on assets still being compiled to avoid waiting unless strictly necessary. */
 	uint8 ArShouldSkipCompilingAssets : 1;
 
+	/** If true, UsingCurrentVersion will not query the version registry to get the latest version data for the key if the version key is already set in the local container.
+	 * Querying the version registry is not cheap, and the look up time increases as more custom versions are registered, therefore if you don't change the registered custom version data once is set in the archive
+	 * Setting this to true should improve serialization performance by skipping the step entirely */
+	uint8 ArShouldSkipUpdateCustomVersion : 1;
+
 public:
 	/** Quickly tell if an archive contains script code. */
 	uint8 ArContainsCode : 1;
@@ -2078,6 +2096,8 @@ public:
 	using FArchiveState::SetEngineNetVer;
 	using FArchiveState::SetGameNetVer;
 	using FArchiveState::ShouldSkipCompilingAssets;
+	using FArchiveState::SetShouldSkipUpdateCustomVersion;
+	using FArchiveState::ShouldSkipUpdateCustomVersion;
 	
 private:
 	using FArchiveState::ArUEVer;
diff --git a/Engine/Source/Runtime/Core/Public/Serialization/CustomVersion.h b/Engine/Source/Runtime/Core/Public/Serialization/CustomVersion.h
--- a/Engine/Source/Runtime/Core/Public/Serialization/CustomVersion.h
+++ b/Engine/Source/Runtime/Core/Public/Serialization/CustomVersion.h
@@ -97,6 +97,16 @@ private:
 
 class CORE_API FCustomVersionRegistration;
 
+/** Flags defining how the Custom Version should be applied */
+enum class ESetCustomVersionFlags : uint8
+{
+	None,
+
+	// If Set, it will not query the versions registry to update the version data if the key is already in the container 
+	SkipUpdateExistingVersion = 0x01,
+};
+ENUM_CLASS_FLAGS(ESetCustomVersionFlags)
+
 /**
  * Container for all available/serialized custom versions.
  */
@@ -136,6 +146,13 @@ public:
 	 */
 	void SetVersion(FGuid CustomKey, int32 Version, FName FriendlyName);
 
+	/**
+	 * Sets a specific custom version in the container. It queries the versions registry to get the version data for the provided key
+	 * @param CustomKey Custom key for which to retrieve the version.
+	 * @param Options Optional flags used to alter the behavior of this method
+	 */
+	void SetVersionUsingRegistry(FGuid CustomKey, ESetCustomVersionFlags Options = ESetCustomVersionFlags::None);
+
 	/** Serialization. */
 	void Serialize(FArchive& Ar, ECustomVersionSerializationFormat::Type Format = ECustomVersionSerializationFormat::Latest);
 	void Serialize(FStructuredArchive::FSlot Slot, ECustomVersionSerializationFormat::Type Format = ECustomVersionSerializationFormat::Latest);
