Avoid reading globals from other threads by reordering conditions of SlowTask MakeDialog

#rnx

diff --git a/Engine/Source/Runtime/Core/Private/Misc/SlowTask.cpp b/Engine/Source/Runtime/Core/Private/Misc/SlowTask.cpp
--- a/Engine/Source/Runtime/Core/Private/Misc/SlowTask.cpp
+++ b/Engine/Source/Runtime/Core/Private/Misc/SlowTask.cpp
@@ -198,9 +198,9 @@ const FText& FSlowTask::GetCurrentMessage() const
 
 void FSlowTask::MakeDialog(bool bShowCancelButton, bool bAllowInPIE)
 {
-	const bool bIsDisabledByPIE = Context.IsPlayingInEditor() && !bAllowInPIE;
-	const bool bIsDialogAllowed = bEnabled && !GIsSilent && !bIsDisabledByPIE && !IsRunningCommandlet() && IsInGameThread() && Visibility != ESlowTaskVisibility::Invisible;
-	if (!GIsSlowTask && bIsDialogAllowed)
+	const auto IsDisabledByPIE = [this, bAllowInPIE]() { return Context.IsPlayingInEditor() && !bAllowInPIE; };
+	const bool bIsDialogAllowed = bEnabled && IsInGameThread() && !GIsSilent && !IsDisabledByPIE() && !IsRunningCommandlet() && Visibility != ESlowTaskVisibility::Invisible;
+	if (bIsDialogAllowed && !GIsSlowTask)
 	{
 		Context.StartSlowTask(GetCurrentMessage(), bShowCancelButton);
 		if (GIsSlowTask)
