Add Range overload for Algo::Partition

diff --git a/Engine/Source/Runtime/Core/Public/Algo/Partition.h b/Engine/Source/Runtime/Core/Public/Algo/Partition.h
--- a/Engine/Source/Runtime/Core/Public/Algo/Partition.h
+++ b/Engine/Source/Runtime/Core/Public/Algo/Partition.h
@@ -12,8 +12,8 @@ namespace Algo
 	 * @param	Predicate	unary predicate class
 	 * @return	index of the first element in the second group
 	 */
-	template<class T, typename IndexType, class UnaryPredicate>
-	IndexType Partition(T* Elements, const IndexType Num, const UnaryPredicate& Predicate)
+	template<class T, typename IndexType, typename UnaryPredicate>
+	IndexType Partition(T* Elements, const IndexType Num, UnaryPredicate Predicate)
 	{
 		T* First = Elements;
 		T* Last = Elements + Num;
@@ -44,4 +44,17 @@ namespace Algo
 	
 		return (IndexType)(First - Elements);
 	}
+
+	/**
+	 * Rearranges the elements so that all the elements for which Predicate returns true precede all those for which it returns false.  (not stable)
+	 *
+	 * @param	Range		the range to sort
+	 * @param	Predicate	a unary predicate object
+	 * @return	index of the first element in the second group
+	 */
+	template <typename RangeType, typename UnaryPredicateType>
+	FORCEINLINE auto Partition(RangeType&& Range, UnaryPredicateType Predicate) -> decltype(GetNum(Range))
+	{
+		return Partition(GetData(Range), GetNum(Range), MoveTemp(Predicate));
+	}
 } //namespace Algo
\ No newline at end of file
