MTAccessDetector: moved thread_local vars to .cpp as this was causing issues in editor on some platforms

diff --git a/Engine/Source/Runtime/Core/Private/Misc/MTAccessDetector.cpp b/Engine/Source/Runtime/Core/Private/Misc/MTAccessDetector.cpp
--- a/Engine/Source/Runtime/Core/Private/Misc/MTAccessDetector.cpp
+++ b/Engine/Source/Runtime/Core/Private/Misc/MTAccessDetector.cpp
@@ -2,6 +2,22 @@
 
 #include "Misc/MTAccessDetector.h"
 
+#if ENABLE_MT_DETECTOR
+
+FMRSWRecursiveAccessDetector::FReadersTls& FMRSWRecursiveAccessDetector::GetReadersTls()
+{
+	thread_local FReadersTls ReadersTls;
+	return ReadersTls;
+}
+
+FMRSWRecursiveAccessDetector::FDestructionSentinelStackTls& FMRSWRecursiveAccessDetector::GetDestructionSentinelStackTls()
+{
+	thread_local FDestructionSentinelStackTls DestructionSentinelStackTls;
+	return DestructionSentinelStackTls;
+}
+
+#endif
+
 #if ENABLE_MT_DETECTOR && WITH_DEV_AUTOMATION_TESTS
 
 #include "Misc/AutomationTest.h"
diff --git a/Engine/Source/Runtime/Core/Public/Misc/MTAccessDetector.h b/Engine/Source/Runtime/Core/Public/Misc/MTAccessDetector.h
--- a/Engine/Source/Runtime/Core/Public/Misc/MTAccessDetector.h
+++ b/Engine/Source/Runtime/Core/Public/Misc/MTAccessDetector.h
@@ -370,11 +370,8 @@ private:
 	// a TLS stack of destruction sentinels shared by all access detector instances on the callstack. the same access detector instance can have
 	// multiple destruction sentinels on the stack. the next "release access" with destruction sentinel always matches the top of the stack
 	using FDestructionSentinelStackTls = TArray<FDestructionSentinel*, TInlineAllocator<4>>;
-	static FDestructionSentinelStackTls& GetDestructionSentinelStackTls()
-	{
-		thread_local FDestructionSentinelStackTls DestructionSentinelStackTls;
-		return DestructionSentinelStackTls;
-	}
+	
+	static CORE_API FDestructionSentinelStackTls& GetDestructionSentinelStackTls();
 //////////////////////////////////////////////////////////////////////
 
 private:
@@ -389,11 +386,8 @@ private:
 	};
 
 	using FReadersTls = TArray<FReaderNum, TInlineAllocator<4>>;
-	static FReadersTls& GetReadersTls()
-	{
-		thread_local FReadersTls ReadersTls;
-		return ReadersTls;
-	}
+	
+	static CORE_API FReadersTls& GetReadersTls();
 
 	void RemoveReaderFromTls() const
 	{
