Fix for AddressSanitizer: stack-use-after-scope in FCustomVersionContainer::SetVersionUsingRegistry

diff --git a/Engine/Source/Runtime/Core/Private/Serialization/CustomVersion.cpp b/Engine/Source/Runtime/Core/Private/Serialization/CustomVersion.cpp
--- a/Engine/Source/Runtime/Core/Private/Serialization/CustomVersion.cpp
+++ b/Engine/Source/Runtime/Core/Private/Serialization/CustomVersion.cpp
@@ -415,13 +415,17 @@ void FCustomVersionContainer::SetVersionUsingRegistry(FGuid CustomKey, ESetCusto
 			return;
 		}
 
-		const FCustomVersion& RegisteredVersion = FCurrentCustomVersions::Get(CustomKey).GetValue();
-		Found->Version      = RegisteredVersion.Version;
-		Found->FriendlyName = RegisteredVersion.FriendlyName;
+		TOptional<FCustomVersion> RegisteredVersion = FCurrentCustomVersions::Get(CustomKey);
+		checkf(RegisteredVersion, TEXT("Attempted to set a version that is not registered"));
+
+		Found->Version      = RegisteredVersion->Version;
+		Found->FriendlyName = RegisteredVersion->FriendlyName;
 	}
 	else
 	{
-		const FCustomVersion& RegisteredVersion = FCurrentCustomVersions::Get(CustomKey).GetValue();
-		Versions.Emplace(FCustomVersion(CustomKey, RegisteredVersion.Version, RegisteredVersion.FriendlyName));
+		TOptional<FCustomVersion> RegisteredVersion = FCurrentCustomVersions::Get(CustomKey);
+		checkf(RegisteredVersion, TEXT("Attempted to set a version that is not registered"));
+
+		Versions.Emplace(FCustomVersion(CustomKey, RegisteredVersion->Version, RegisteredVersion->FriendlyName));
 	}
 }
