remove log spam

#rnx

diff --git a/Engine/Source/Runtime/Core/Private/Misc/ConfigContext.cpp b/Engine/Source/Runtime/Core/Private/Misc/ConfigContext.cpp
--- a/Engine/Source/Runtime/Core/Private/Misc/ConfigContext.cpp
+++ b/Engine/Source/Runtime/Core/Private/Misc/ConfigContext.cpp
@@ -120,7 +120,6 @@ const FConfigContext::FPerPlatformDirs& FConfigContext::GetPerPlatformDirs(const
 			// look if there's a plugin extension for this platform, it will have the platform name in the path
 			for (const FString& ChildDir : ChildPluginBaseDirs)
 			{
-				UE_LOG(LogConfig, Display, TEXT("HAVE CHILD PLUGIN DIR: %s"), *ChildDir);
 				if (ChildDir.Contains(*FString::Printf(TEXT("/%s/"), *PlatformName)))
 				{
 					PluginExtDir = ChildDir;
