Fixed TUseBitwiseSwap which was broken by the move to std:: traits in CL# 23895694.
Removed deprecated AsConst overloads.

diff --git a/Engine/Source/Runtime/Core/Public/Templates/UnrealTemplate.h b/Engine/Source/Runtime/Core/Public/Templates/UnrealTemplate.h
--- a/Engine/Source/Runtime/Core/Public/Templates/UnrealTemplate.h
+++ b/Engine/Source/Runtime/Core/Public/Templates/UnrealTemplate.h
@@ -113,26 +113,6 @@ constexpr int32 GetNum(std::initializer_list<T> List)
 	return static_cast<int32>(List.size());
 }
 
-/**
- * Returns a non-const pointer type as const.
- */
-template <typename T>
-UE_DEPRECATED(4.26, "Call with a reference instead of a pointer.")
-constexpr FORCEINLINE const T* AsConst(T* const& Ptr)
-{
-	return Ptr;
-}
-
-/**
- * Returns a non-const pointer type as const.
- */
-template <typename T>
-UE_DEPRECATED(4.26, "Call with a reference instead of a pointer.")
-constexpr FORCEINLINE const T* AsConst(T* const&& Ptr)
-{
-	return Ptr;
-}
-
 /**
  * Returns a non-const reference type as const.
  */
@@ -602,7 +582,7 @@ template <typename T>
 struct TUseBitwiseSwap
 {
 	// We don't use bitwise swapping for 'register' types because this will force them into memory and be slower.
-	enum { Value = !std::is_enum_v<T> || std::is_pointer_v<T> || std::is_arithmetic_v<T> };
+	enum { Value = !(std::is_enum_v<T> || std::is_pointer_v<T> || std::is_arithmetic_v<T>) };
 };
 
 
