improved error reporting for checked delegates that are copied trivially

diff --git a/Engine/Source/Runtime/Core/Public/Misc/MTAccessDetector.h b/Engine/Source/Runtime/Core/Public/Misc/MTAccessDetector.h
--- a/Engine/Source/Runtime/Core/Public/Misc/MTAccessDetector.h
+++ b/Engine/Source/Runtime/Core/Public/Misc/MTAccessDetector.h
@@ -512,8 +512,9 @@ public:
 
 		if (LocalState.ReaderNum >= 1)
 		{	// check that all readers are on the current thread
-			uint32 ReaderIndex = GetReadersTls().IndexOfByPredicate([this](FReaderNum ReaderNum) { return ReaderNum.Reader == this; });
-			checkf(ReaderIndex != INDEX_NONE && GetReadersTls()[ReaderIndex].Num == LocalState.ReaderNum, 
+			int32 ReaderIndex = GetReadersTls().IndexOfByPredicate([this](FReaderNum ReaderNum) { return ReaderNum.Reader == this; });
+			checkf(ReaderIndex != INDEX_NONE, TEXT("Race detector is not trivially copyable while this delegate is copied trivially. Consider changing this delegate to use `FNotThreadSafeNotCheckedDelegateUserPolicy`"));
+			checkf(GetReadersTls()[ReaderIndex].Num == LocalState.ReaderNum, 
 				TEXT("Data race detected: %d reader(s) on another thread(s) while acquiring write access"), 
 				LocalState.ReaderNum - GetReadersTls()[ReaderIndex].Num);
 		}
