Removed self-swap test in Swap to avoid a branch.

diff --git a/Engine/Source/Runtime/Core/Public/Templates/UnrealTemplate.h b/Engine/Source/Runtime/Core/Public/Templates/UnrealTemplate.h
--- a/Engine/Source/Runtime/Core/Public/Templates/UnrealTemplate.h
+++ b/Engine/Source/Runtime/Core/Public/Templates/UnrealTemplate.h
@@ -594,13 +594,10 @@ inline void Swap(T& A, T& B)
 {
 	if constexpr (TUseBitwiseSwap<T>::Value)
 	{
-		if (LIKELY(&A != &B))
-		{
-			TTypeCompatibleBytes<T> Temp;
-			FMemory::Memcpy(&Temp, &A, sizeof(T));
-			FMemory::Memcpy(&A, &B, sizeof(T));
-			FMemory::Memcpy(&B, &Temp, sizeof(T));
-		}
+		TTypeCompatibleBytes<T> Temp;
+		*(TTypeCompatibleBytes<T>*)&Temp = *(TTypeCompatibleBytes<T>*)&A;
+		*(TTypeCompatibleBytes<T>*)&A    = *(TTypeCompatibleBytes<T>*)&B;
+		*(TTypeCompatibleBytes<T>*)&B    = *(TTypeCompatibleBytes<T>*)&Temp;
 	}
 	else
 	{
