fixed race detection reporting thread id (was off by one)

diff --git a/Engine/Source/Runtime/Core/Public/Misc/MTAccessDetector.h b/Engine/Source/Runtime/Core/Public/Misc/MTAccessDetector.h
--- a/Engine/Source/Runtime/Core/Public/Misc/MTAccessDetector.h
+++ b/Engine/Source/Runtime/Core/Public/Misc/MTAccessDetector.h
@@ -340,8 +340,8 @@ private:
 	{
 		if (InState.WriterNum != 0)
 		{
-			uint32 CurrentThreadId = FPlatformTLS::GetCurrentThreadId() + 1; // to shift from 0 TID that is considered "invalid"
-			checkf(InState.WriterThreadId == CurrentThreadId, TEXT("Data race detected! Writer on thread %u while acquiring read access on thread %u"), InState.WriterThreadId, CurrentThreadId);
+			uint32 CurrentThreadId = FPlatformTLS::GetCurrentThreadId();
+			checkf(InState.WriterThreadId - 1 == CurrentThreadId, TEXT("Data race detected! Writer on thread %u while acquiring read access on thread %u"), InState.WriterThreadId, CurrentThreadId);
 		}
 	}
 
