UE_LOG: Compile-time format string UB sanitizer/checker.

diff --git a/Engine/Source/Runtime/Core/Public/Logging/LogMacros.h b/Engine/Source/Runtime/Core/Public/Logging/LogMacros.h
--- a/Engine/Source/Runtime/Core/Public/Logging/LogMacros.h
+++ b/Engine/Source/Runtime/Core/Public/Logging/LogMacros.h
@@ -12,6 +12,7 @@
 #include "Misc/AssertionMacros.h"
 #include "Misc/Build.h"
 #include "Misc/VarArgs.h"
+#include "String/FormatStringSan.h"
 #include "Templates/AndOrNot.h"
 #include "Templates/EnableIf.h"
 #include "Templates/IsArrayOrRefOfTypeByPredicate.h"
@@ -229,8 +230,8 @@ CORE_API void BasicFatalLog(const FLogCategoryBase& Category, const FStaticBasic
 	#define UE_SET_LOG_VERBOSITY(CategoryName, Verbosity) \
 		CategoryName.SetVerbosity(ELogVerbosity::Verbosity);
 
-	#if UE_VALIDATE_FORMAT_STRINGS && defined(_MSC_VER)
-		#define UE_VALIDATE_FORMAT_STRING(Format, ...) do { if (false) { wprintf(Format, ##__VA_ARGS__); } } while(false)
+	#if UE_VALIDATE_FORMAT_STRINGS
+		#define UE_VALIDATE_FORMAT_STRING UE_CHECK_FORMAT_STRING
 	#else
 		#define UE_VALIDATE_FORMAT_STRING(Format, ...)
 	#endif
diff --git a/Engine/Source/Runtime/Core/Public/String/FormatStringSan.h b/Engine/Source/Runtime/Core/Public/String/FormatStringSan.h
new file mode 100644
--- /dev/null
+++ b/Engine/Source/Runtime/Core/Public/String/FormatStringSan.h
@@ -0,0 +1,470 @@
+// Copyright Epic Games, Inc. All Rights Reserved.
+#pragma once
+
+#include <type_traits>
+#include "CoreTypes.h"
+
+#define UE_CHECK_FORMAT_STRING(Fmt, ...) do {														\
+		typedef UE::Core::Private::FormatStringSan::TCheckFormatString<decltype(UE::Core::Private::FormatStringSan::GetFmtArgTypes(__VA_ARGS__))> UE_FMT_STR_Checker; \
+		constexpr UE::Core::Private::FormatStringSan::FResult UE_FMT_STR_Result = UE_FMT_STR_Checker::Check(false, 0, Fmt); \
+		typedef UE::Core::Private::FormatStringSan::TPresentErr<UE_FMT_STR_Result.Status> UE_FMT_STR_ErrorMsg; \
+		typedef typename UE_FMT_STR_ErrorMsg::template TValue<UE_FMT_STR_Result.ArgPos> UE_FMT_STR_Status; \
+		(void) UE_FMT_STR_Status{};																					\
+	} while(false)
+
+#define UE_CHECK_FORMAT_STRING_ERR(Err, Fmt, ...)												\
+	(UE::Core::Private::FormatStringSan::TCheckFormatString<decltype(UE::Core::Private::FormatStringSan::GetFmtArgTypes(__VA_ARGS__))>::Check(false, 0, Fmt).Status == Err)
+
+/// implementation
+namespace UE::Core::Private
+{
+	namespace FormatStringSan
+	{
+		enum EFormatStringSanStatus
+		{
+			StatusOk,
+			StatusSNeedsTCHARPtrArg,
+			StatusHSNeedsCharPtrArg,
+			StatusIncompleteFormatSpecifierOrUnescapedPercent,
+			StatusCNeedsTCHARArg,
+			StatusPNeedsPointerArg,
+			StatusFNeedsFloatOrDoubleArg,
+			StatusCapitalSNeedsCharPtr,
+			StatusCapitalSNeedsTCHARPtr,
+			StatusInvalidFormatSpec,
+			StatusNotEnoughArguments,
+			StatusDNeedsIntegerArg,
+			StatusZNeedsIntegerSpec,
+			StatusZNeedsIntegerArg,
+			StatusDynamicLengthSpecNeedsIntegerArg,
+			StatusLNeedsIntegerArg,
+			StatusHNeedsIntegerArg,
+			StatusHHNeedsIntegerSpec,
+			StatusHHNeedsIntegerArg,
+			StatusLLNeedsIntegerSpec,
+			StatusLLNeedsIntegerArg,
+			StatusI64BadSpec,
+			StatusI64NeedsIntegerArg
+		};
+
+		template <int N>
+		struct TAtArgPos {};
+
+		template <typename>
+		inline constexpr bool bFmtStrErr = false;
+
+		struct FResult
+		{
+			int Status;
+			int ArgPos;
+		};
+
+		template <int N>
+		struct TPresentErr;
+		template <>
+		struct TPresentErr<StatusOk>
+		{
+			template <int> struct TValue {};
+		};
+
+#define FMT_STR_ERR(N, Description)									\
+		template <>																			\
+		struct TPresentErr<N>														\
+		{																								\
+			template <typename T>													\
+			struct TErrMsg																\
+			{																							\
+				static_assert(bFmtStrErr<T>, Description);	\
+			};																						\
+			template <int ArgN>														\
+			using TValue = TErrMsg<TAtArgPos<ArgN>>;				\
+		}
+
+		FMT_STR_ERR(StatusSNeedsTCHARPtrArg, "'%s' expects `TCHAR*`; use '%hs' for `const char*`.");
+		FMT_STR_ERR(StatusHSNeedsCharPtrArg, "'%hs' expects `char*`; use '%s' for `const TCHAR*`, `TEXT(...)`, etc.");
+		FMT_STR_ERR(StatusIncompleteFormatSpecifierOrUnescapedPercent, "incomplete '%' escape sequence. (to print a percent sign, write '%%'.)");
+		FMT_STR_ERR(StatusCNeedsTCHARArg, "'%c' expects `TCHAR`. (construct a `TCHAR` from a `char` like this: 'TCHAR(my_char)'.)");
+		FMT_STR_ERR(StatusPNeedsPointerArg, "'%p' expects a pointer.");
+		FMT_STR_ERR(StatusFNeedsFloatOrDoubleArg, "'%f' expects `float` or `double`.");
+		FMT_STR_ERR(StatusCapitalSNeedsCharPtr, "'%S' expects 'char*' (or a `UTF8CHAR*` if enabled).");
+		FMT_STR_ERR(StatusCapitalSNeedsTCHARPtr, "'%S' expects 'TCHAR*'.");
+		FMT_STR_ERR(StatusInvalidFormatSpec, "unsupported '%' format specifier. (to print a percent sign, write '%%'.)");
+		FMT_STR_ERR(StatusNotEnoughArguments, "not enough arguments provided to format string.");
+		FMT_STR_ERR(StatusDNeedsIntegerArg, "'%d' expects integral arg (eg. `char`, `int`, `long`, etc.)");
+		FMT_STR_ERR(StatusZNeedsIntegerArg, "'%z' expects integral arg (eg. `char`, `int`, `long`, etc.)");
+		FMT_STR_ERR(StatusZNeedsIntegerSpec, "'%z' must have a suffix (eg. '%zd', '%zu', etc.)");
+		FMT_STR_ERR(StatusLNeedsIntegerArg, "'%l[ ]' expects integral arg (eg. `char`, `int`, `long`, etc.)");
+		FMT_STR_ERR(StatusHNeedsIntegerArg, "'%h[ ]' expects integral arg (eg. `char`, `int`, `long`, etc.)");
+		FMT_STR_ERR(StatusHHNeedsIntegerSpec, "'%hh[ ]' expects integral suffix (eg. '%hhd', '%hhu', etc.)");
+		FMT_STR_ERR(StatusHHNeedsIntegerArg, "'%hh[ ]' expects integral arg (eg. `char`, `int`, `long`, etc.)");
+		FMT_STR_ERR(StatusLLNeedsIntegerSpec, "'%ll[ ]' expects integral suffix (eg. '%hhd', '%hhu', etc.)");
+		FMT_STR_ERR(StatusLLNeedsIntegerArg, "'%ll[ ]' expects integral arg (eg. `char`, `int`, `long`, etc.)");
+		FMT_STR_ERR(StatusI64BadSpec, "'%I' must appear as '%I64' with an integral suffix (eg. '%I64d', '%I64u', etc.)");
+		FMT_STR_ERR(StatusI64NeedsIntegerArg, "'%I64[ ]' expects integral arg (eg. `char`, `int`, `long`, etc.)");
+		FMT_STR_ERR(StatusDynamicLengthSpecNeedsIntegerArg, "dynamic field width specifier '*' expects integral arg (eg. `char`, `int`, `long`, etc.)");
+
+#undef FMT_STR_ERR
+
+		template <typename T, typename V>
+		inline constexpr bool bIsPointerTo =
+			std::is_pointer_v<T> && std::is_same_v<std::remove_cv_t<std::remove_pointer_t<T>>, V>;
+		template <typename T>
+		inline constexpr bool bIsCharPtr = bIsPointerTo<T, char> || bIsPointerTo<T, UTF8CHAR>;
+		template <typename T>
+		inline constexpr bool bIsTCharPtr = bIsPointerTo<T, TCHAR>;
+		template <typename T>
+		inline constexpr bool bIsFloatOrDouble = std::is_same_v<float, T> || std::is_same_v<double, T>;
+		template <typename T>
+		inline constexpr bool bIsIntegralEnum = std::is_enum_v<T> && std::is_convertible_v<T,int>;
+
+		template <typename... Ts>
+		struct TFmtArgTypes {};
+
+		template <typename... Ts>
+		constexpr TFmtArgTypes<Ts...> GetFmtArgTypes(Ts...)
+		{
+			return {};
+		}
+
+		template <typename>
+		struct TCheckFormatString
+		{
+			static constexpr FResult Check(bool, int, const TCHAR*);
+		};
+
+		template <typename Arg, typename... Args>
+		struct TCheckFormatString<TFmtArgTypes<Arg, Args...>>
+		{
+			static constexpr FResult Check(bool bInsideFormatSpec, int CurArgPos, const TCHAR* Fmt)
+			{
+				auto CharIsDigit = [](const TCHAR C)
+				{
+					return C >= TEXT('0') && C <= TEXT('9');
+				};
+
+				auto SkipInteger = [&CharIsDigit](const TCHAR* P)
+				{
+					while (CharIsDigit(*P))
+					{
+						++P;
+					}
+					return P;
+				};
+
+				auto CharIsIntegerFormatSpecifier = [](const TCHAR C)
+				{
+					switch (C)
+					{
+					case TEXT('i'): case TEXT('d'): case TEXT('u'): case TEXT('X'): case TEXT('x'):
+						return true;
+					default:
+						return false;
+					}
+				};
+
+				auto HandleDynamicLengthSpecifier = [CurArgPos](const TCHAR* P) constexpr -> FResult
+				{
+					if constexpr (!(std::is_integral_v<Arg> || bIsIntegralEnum<Arg>))
+					{
+						return {StatusDynamicLengthSpecNeedsIntegerArg, CurArgPos};
+					}
+					else
+					{
+						return TCheckFormatString<TFmtArgTypes<Args...>>::Check(true, CurArgPos + 1, P + 1);
+					}
+				};
+
+				const TCHAR* P = Fmt;
+				if (!P[0] && !bInsideFormatSpec)
+				{
+					return {StatusOk, 0};
+				}
+
+				if (!bInsideFormatSpec)
+				{
+					while (*P && *P != TEXT('%'))
+					{
+						++P;
+					}
+					if (*P == TEXT('%'))
+					{
+						++P;
+					}
+				}
+
+				if (!*P)
+				{
+					if (P[-1] == TEXT('%') || bInsideFormatSpec) // nb: P[-1] is safe here because zero-length fmt strings are addressed above
+					{
+						return {StatusIncompleteFormatSpecifierOrUnescapedPercent, CurArgPos};
+					}
+					else
+					{
+						return {StatusOk, 0};
+					}
+				}
+
+				while (*P == TEXT('+') || *P == TEXT('#') || *P == TEXT(' ') || *P == TEXT('0'))
+				{
+					++P;
+				}
+
+				if (*P == TEXT('-'))
+				{
+					++P;
+				}
+				if (*P == TEXT('*'))
+				{
+					return HandleDynamicLengthSpecifier(P);
+				}
+				else if (CharIsDigit(*P))
+				{
+					P = SkipInteger(P);
+				}
+
+				if (*P == TEXT('.'))
+				{
+					++P;
+				}
+				if (*P == TEXT('*'))
+				{
+					return HandleDynamicLengthSpecifier(P);
+				}
+				else if (CharIsDigit(*P))
+				{
+					P = SkipInteger(P);
+				}
+
+				if (P[0] == TEXT('l') && P[1] == TEXT('s'))
+				{
+					++P;
+				}
+				else if (P[0] == TEXT('h') && P[1] == TEXT('s'))
+				{
+					if constexpr (!bIsCharPtr<Arg>)
+					{
+						return {StatusHSNeedsCharPtrArg, CurArgPos};
+					}
+					else
+					{
+						return TCheckFormatString<TFmtArgTypes<Args...>>::Check(false, CurArgPos + 1, P + 2);
+					}
+				}
+
+				switch (P[0])
+				{
+				case TEXT('%'):
+					return TCheckFormatString<TFmtArgTypes<Arg, Args...>>::Check(false, CurArgPos, P + 1);
+				case TEXT('s'):
+					if constexpr (!bIsTCharPtr<Arg>)
+					{
+						return {StatusSNeedsTCHARPtrArg, CurArgPos};
+					}
+					else
+					{
+						return TCheckFormatString<TFmtArgTypes<Args...>>::Check(false, CurArgPos + 1, P + 1);
+					}
+				case TEXT('c'):
+					if constexpr (!std::is_same_v<Arg, TCHAR>)
+					{
+						return {StatusCNeedsTCHARArg, CurArgPos};
+					}
+					else
+					{
+						return TCheckFormatString<TFmtArgTypes<Args...>>::Check(false, CurArgPos + 1, P + 1);
+					}
+				case TEXT('d'):
+				case TEXT('i'):
+				case TEXT('X'):
+				case TEXT('x'):
+				case TEXT('u'):
+					if constexpr (!(bIsIntegralEnum<Arg> || std::is_integral_v<Arg> || std::is_pointer_v<Arg>))
+					{
+						return {StatusDNeedsIntegerArg, CurArgPos};
+					}
+					else
+					{
+						return TCheckFormatString<TFmtArgTypes<Args...>>::Check(false, CurArgPos + 1, P + 1);
+					}
+				case TEXT('z'):
+					if (!CharIsIntegerFormatSpecifier(P[1]))
+					{
+						return {StatusZNeedsIntegerSpec, CurArgPos};
+					}
+					else if constexpr (!(bIsIntegralEnum<Arg> || std::is_integral_v<Arg>))
+					{
+						return {StatusZNeedsIntegerArg, CurArgPos};
+					}
+					else
+					{
+						return TCheckFormatString<TFmtArgTypes<Args...>>::Check(false, CurArgPos + 1, P + 1);
+					}
+				case TEXT('p'):
+					if constexpr (!std::is_pointer_v<Arg>)
+					{
+						return {StatusPNeedsPointerArg, CurArgPos};
+					}
+					else
+					{
+						return TCheckFormatString<TFmtArgTypes<Args...>>::Check(false, CurArgPos + 1, P + 1);
+					}
+				case TEXT('I'):
+					if (!(P[1] == TEXT('6') && P[2] == TEXT('4')))
+					{
+						return {StatusI64BadSpec, CurArgPos};
+					}
+					else if (!CharIsIntegerFormatSpecifier(P[3]))
+					{
+						return {StatusI64BadSpec, CurArgPos};
+					}
+					else if constexpr (!(bIsIntegralEnum<Arg> || std::is_integral_v<Arg>))
+					{
+						return {StatusI64NeedsIntegerArg, CurArgPos};
+					}
+					else
+					{
+						return TCheckFormatString<TFmtArgTypes<Args...>>::Check(false, CurArgPos + 1, P + 4);
+					}
+
+				case TEXT('l'):
+					if (CharIsIntegerFormatSpecifier(P[1]))
+					{
+						if constexpr (!(bIsIntegralEnum<Arg> || std::is_integral_v<Arg>))
+						{
+							return {StatusLNeedsIntegerArg, CurArgPos};
+						}
+						else
+						{
+							return TCheckFormatString<TFmtArgTypes<Args...>>::Check(false, CurArgPos + 1, P + 2);
+						}
+					}
+					else if (P[1] == TEXT('f'))
+					{
+						if constexpr (!bIsFloatOrDouble<Arg>)
+						{
+							return {StatusFNeedsFloatOrDoubleArg, CurArgPos};
+						}
+						else
+						{
+							return TCheckFormatString<TFmtArgTypes<Args...>>::Check(false, CurArgPos + 1, P + 2);
+						}
+					}
+					else if (P[1] != TEXT('l'))
+					{
+						return {StatusInvalidFormatSpec, CurArgPos};
+					}
+					else if (!CharIsIntegerFormatSpecifier(P[2]))
+					{
+						return {StatusLLNeedsIntegerSpec, CurArgPos};
+					}
+					else if constexpr (!(bIsIntegralEnum<Arg> || std::is_integral_v<Arg> || std::is_pointer_v<Arg>))
+					{
+						return {StatusLLNeedsIntegerArg, CurArgPos};
+					}
+					else
+					{
+						return TCheckFormatString<TFmtArgTypes<Args...>>::Check(false, CurArgPos + 1, P + 3);
+					}
+
+				case TEXT('h'):
+					if (CharIsIntegerFormatSpecifier(P[1]))
+					{
+						if (!(bIsIntegralEnum<Arg> || std::is_integral_v<Arg>))
+						{
+							return {StatusHNeedsIntegerArg, CurArgPos};
+						}
+						return TCheckFormatString<TFmtArgTypes<Args...>>::Check(false, CurArgPos + 1, P + 2);
+					}
+					else if (P[1] != TEXT('h'))
+					{
+						return {StatusInvalidFormatSpec, CurArgPos};
+					}
+					else if (!CharIsIntegerFormatSpecifier(P[2]))
+					{
+						return {StatusHHNeedsIntegerSpec, CurArgPos};
+					}
+					else if constexpr (!(bIsIntegralEnum<Arg> || std::is_integral_v<Arg>))
+					{
+						return {StatusHHNeedsIntegerArg, CurArgPos};
+					}
+					else
+					{
+						return TCheckFormatString<TFmtArgTypes<Args...>>::Check(false, CurArgPos + 1, P + 3);
+					}
+
+				case TEXT('f'):
+				case TEXT('e'):
+				case TEXT('g'):
+					if constexpr (!bIsFloatOrDouble<Arg>)
+					{
+						return {StatusFNeedsFloatOrDoubleArg, CurArgPos};
+					}
+					else
+					{
+						return TCheckFormatString<TFmtArgTypes<Args...>>::Check(false, CurArgPos + 1, P + 1);
+					}
+
+				case TEXT('S'):
+					using OtherCharType = std::conditional_t<std::is_same_v<TCHAR, ANSICHAR>, WIDECHAR, ANSICHAR>;
+					if constexpr (!bIsPointerTo<Arg, OtherCharType>)
+					{
+						if constexpr(std::is_same_v<OtherCharType, ANSICHAR>)
+						{
+							return {StatusCapitalSNeedsCharPtr, CurArgPos};
+						}
+						else if constexpr(std::is_same_v<OtherCharType, WIDECHAR>)
+						{
+							return {StatusCapitalSNeedsTCHARPtr, CurArgPos};
+						}
+					}
+					else
+					{
+						return TCheckFormatString<TFmtArgTypes<Args...>>::Check(false, CurArgPos + 1, P + 1);
+					}
+
+				case TEXT(' '):
+					return {StatusIncompleteFormatSpecifierOrUnescapedPercent, CurArgPos};
+
+				default:
+					return {StatusInvalidFormatSpec, CurArgPos};
+				}
+			}
+		};
+
+		template <>
+		struct TCheckFormatString<TFmtArgTypes<>>
+		{
+			static constexpr FResult Check(bool bInsideFormatSpec, int CurArgPos, const TCHAR* P)
+			{
+				if (bInsideFormatSpec)
+				{
+					return {StatusIncompleteFormatSpecifierOrUnescapedPercent, CurArgPos};
+				}
+
+				while (*P)
+				{
+					if (P[0] != TEXT('%'))
+					{
+						++P;
+						continue;
+					}
+
+					if (P[1] == TEXT('%'))
+					{
+						P += 2;
+						continue;
+					}
+
+					if (P[1] == 0)
+					{
+						return {StatusIncompleteFormatSpecifierOrUnescapedPercent, CurArgPos};
+					}
+
+					return {StatusNotEnoughArguments, CurArgPos};
+				}
+
+				return {StatusOk, 0};
+			}
+		};
+	}
+}
diff --git a/Engine/Source/Runtime/Core/Tests/String/FormatStringSanTest.cpp b/Engine/Source/Runtime/Core/Tests/String/FormatStringSanTest.cpp
new file mode 100644
--- /dev/null
+++ b/Engine/Source/Runtime/Core/Tests/String/FormatStringSanTest.cpp
@@ -0,0 +1,101 @@
+// Copyright Epic Games, Inc. All Rights Reserved.
+
+#if WITH_TESTS
+
+#include "String/FormatStringSan.h"
+
+#include "Tests/TestHarnessAdapter.h"
+
+namespace UE
+{
+TEST_CASE_NAMED(FFormatStringValidatorTest, "Core::String::FormatStringSan", "[Core][String][FormatStringSan]")
+{
+	using namespace UE::Core::Private;
+	SECTION("Error Handling")
+	{
+		STATIC_CHECK(UE_CHECK_FORMAT_STRING_ERR(FormatStringSan::StatusSNeedsTCHARPtrArg, TEXT("Test %s"), "wrong"));
+		STATIC_CHECK(UE_CHECK_FORMAT_STRING_ERR(FormatStringSan::StatusHSNeedsCharPtrArg, TEXT("Test %hs"), TEXT("wrong")));
+		STATIC_CHECK(UE_CHECK_FORMAT_STRING_ERR(FormatStringSan::StatusIncompleteFormatSpecifierOrUnescapedPercent, TEXT("Hello %")));
+		STATIC_CHECK(UE_CHECK_FORMAT_STRING_ERR(FormatStringSan::StatusCNeedsTCHARArg, TEXT("Hello %c"), char(42)));
+		STATIC_CHECK(UE_CHECK_FORMAT_STRING_ERR(FormatStringSan::StatusPNeedsPointerArg, TEXT("Hello %p"), 42));
+		STATIC_CHECK(UE_CHECK_FORMAT_STRING_ERR(FormatStringSan::StatusFNeedsFloatOrDoubleArg, TEXT("Hello %f"), 42));
+		STATIC_CHECK(std::is_same_v<TCHAR, WIDECHAR> == UE_CHECK_FORMAT_STRING_ERR(FormatStringSan::StatusCapitalSNeedsCharPtr, TEXT("Hello %S"), TEXT("wrong")));
+		STATIC_CHECK(std::is_same_v<TCHAR, ANSICHAR> == UE_CHECK_FORMAT_STRING_ERR(FormatStringSan::StatusCapitalSNeedsTCHARPtr, TEXT("Hello %S"), "wrong"));
+		STATIC_CHECK(UE_CHECK_FORMAT_STRING_ERR(FormatStringSan::StatusInvalidFormatSpec, TEXT("Hello %k"), 42));
+		STATIC_CHECK(UE_CHECK_FORMAT_STRING_ERR(FormatStringSan::StatusNotEnoughArguments, TEXT("Hello %s")));
+		STATIC_CHECK(UE_CHECK_FORMAT_STRING_ERR(FormatStringSan::StatusDNeedsIntegerArg, TEXT("Hello %d"), 42.0));
+		STATIC_CHECK(UE_CHECK_FORMAT_STRING_ERR(FormatStringSan::StatusZNeedsIntegerSpec, TEXT("Hello %z test"), 42));
+		STATIC_CHECK(UE_CHECK_FORMAT_STRING_ERR(FormatStringSan::StatusZNeedsIntegerArg, TEXT("Hello %zu"), "hi"));
+		STATIC_CHECK(UE_CHECK_FORMAT_STRING_ERR(FormatStringSan::StatusDynamicLengthSpecNeedsIntegerArg, TEXT("Hey %*.*d"), "hi", "hi"));
+		STATIC_CHECK(UE_CHECK_FORMAT_STRING_ERR(FormatStringSan::StatusLNeedsIntegerArg, TEXT("Hello %ld"), 43.0f));
+		STATIC_CHECK(UE_CHECK_FORMAT_STRING_ERR(FormatStringSan::StatusHNeedsIntegerArg, TEXT("Hello %hd"), 43.0f));
+		STATIC_CHECK(UE_CHECK_FORMAT_STRING_ERR(FormatStringSan::StatusHHNeedsIntegerSpec, TEXT("Hello %hh "), 42));
+		STATIC_CHECK(UE_CHECK_FORMAT_STRING_ERR(FormatStringSan::StatusHHNeedsIntegerArg, TEXT("Hello %hhd"), 43.0f));
+		STATIC_CHECK(UE_CHECK_FORMAT_STRING_ERR(FormatStringSan::StatusLLNeedsIntegerSpec, TEXT("Hello %ll "), 42));
+		STATIC_CHECK(UE_CHECK_FORMAT_STRING_ERR(FormatStringSan::StatusLLNeedsIntegerArg, TEXT("Hello %lld"), 43.0f));
+		STATIC_CHECK(UE_CHECK_FORMAT_STRING_ERR(FormatStringSan::StatusI64BadSpec, TEXT("Hello %I32d"), 44));
+		STATIC_CHECK(UE_CHECK_FORMAT_STRING_ERR(FormatStringSan::StatusI64BadSpec, TEXT("Hello %I64p"), 44));
+		STATIC_CHECK(UE_CHECK_FORMAT_STRING_ERR(FormatStringSan::StatusI64NeedsIntegerArg	, TEXT("Hello %I64u"), 44.0f));
+		STATIC_CHECK(UE_CHECK_FORMAT_STRING_ERR(FormatStringSan::StatusInvalidFormatSpec, TEXT("%l^"), 42));
+		STATIC_CHECK(UE_CHECK_FORMAT_STRING_ERR(FormatStringSan::StatusInvalidFormatSpec, TEXT("%h^"), 42));
+		STATIC_CHECK(UE_CHECK_FORMAT_STRING_ERR(FormatStringSan::StatusIncompleteFormatSpecifierOrUnescapedPercent, TEXT("%-*"), 42));		
+		
+		{
+			enum class MyEnum { Value };
+			STATIC_CHECK(UE_CHECK_FORMAT_STRING_ERR(FormatStringSan::StatusDNeedsIntegerArg, TEXT("enum class %d value"), MyEnum::Value));
+		}
+	}
+
+	SECTION("Accepted Formatting")
+	{
+		STATIC_CHECK(UE_CHECK_FORMAT_STRING_ERR(FormatStringSan::StatusOk, TEXT("Test %d %% %% %d"), 32, 44));
+		STATIC_CHECK(UE_CHECK_FORMAT_STRING_ERR(FormatStringSan::StatusOk, TEXT("Test")));
+
+		STATIC_CHECK(UE_CHECK_FORMAT_STRING_ERR(FormatStringSan::StatusOk, TEXT("%% Test %d %f %s "), 32, 44.4, TEXT("hey")));
+		STATIC_CHECK(UE_CHECK_FORMAT_STRING_ERR(FormatStringSan::StatusOk, TEXT("Test %.3f %d"), 4.4, 2));
+		STATIC_CHECK(UE_CHECK_FORMAT_STRING_ERR(FormatStringSan::StatusOk, TEXT("Test %2.3f"), 4.4));
+		STATIC_CHECK(UE_CHECK_FORMAT_STRING_ERR(FormatStringSan::StatusOk, TEXT("Test %2.f"), 4.4));
+		STATIC_CHECK(UE_CHECK_FORMAT_STRING_ERR(FormatStringSan::StatusOk, TEXT("Test %2f"), 4.4));
+		STATIC_CHECK(UE_CHECK_FORMAT_STRING_ERR(FormatStringSan::StatusOk, TEXT("Test %d"), long(32)));
+		STATIC_CHECK(UE_CHECK_FORMAT_STRING_ERR(FormatStringSan::StatusOk, TEXT("Test %s"), TEXT("hello")));
+		STATIC_CHECK(UE_CHECK_FORMAT_STRING_ERR(FormatStringSan::StatusOk, TEXT("Test percent %% more")));
+		STATIC_CHECK(UE_CHECK_FORMAT_STRING_ERR(FormatStringSan::StatusOk, TEXT("Test %hs"), "hi"));
+
+		STATIC_CHECK(UE_CHECK_FORMAT_STRING_ERR(FormatStringSan::StatusOk, TEXT("%-8d %f"), 42, 4.4));
+		STATIC_CHECK(UE_CHECK_FORMAT_STRING_ERR(FormatStringSan::StatusOk, TEXT("%8d %f"), 42, 4.4));
+
+		STATIC_CHECK(UE_CHECK_FORMAT_STRING_ERR(FormatStringSan::StatusOk, TEXT("%-8.8d %f"), 42, 4.4));
+		STATIC_CHECK(UE_CHECK_FORMAT_STRING_ERR(FormatStringSan::StatusOk, TEXT("%hhd %d"), int(42), 42));
+		STATIC_CHECK(UE_CHECK_FORMAT_STRING_ERR(FormatStringSan::StatusOk, TEXT("%lld %d"), 42LL, 42));
+		STATIC_CHECK(UE_CHECK_FORMAT_STRING_ERR(FormatStringSan::StatusOk, TEXT("%I64d %d "), 42LL, 42));
+		STATIC_CHECK(UE_CHECK_FORMAT_STRING_ERR(FormatStringSan::StatusOk, TEXT("%d"), size_t(44)));
+
+		STATIC_CHECK(UE_CHECK_FORMAT_STRING_ERR(FormatStringSan::StatusOk, TEXT("%f"), 42.f));
+
+		STATIC_CHECK(UE_CHECK_FORMAT_STRING_ERR(FormatStringSan::StatusOk, TEXT("%-*.*d %f"), 4, 8, 42, 4.4));
+		STATIC_CHECK(UE_CHECK_FORMAT_STRING_ERR(FormatStringSan::StatusOk, TEXT("%*.*d %f"), 4, 8, 42, 4.4));
+		STATIC_CHECK(UE_CHECK_FORMAT_STRING_ERR(FormatStringSan::StatusOk, TEXT("%-*.9d $d"), 4, 42, 44));
+		STATIC_CHECK(UE_CHECK_FORMAT_STRING_ERR(FormatStringSan::StatusOk, TEXT("%*.9d $d"), 4, 42, 44));
+		STATIC_CHECK(UE_CHECK_FORMAT_STRING_ERR(FormatStringSan::StatusOk, TEXT("%-*d %s"), 4, 42, TEXT("hi")));
+		STATIC_CHECK(UE_CHECK_FORMAT_STRING_ERR(FormatStringSan::StatusOk, TEXT("%*d %s"), 4, 42, TEXT("hi")));
+		STATIC_CHECK(UE_CHECK_FORMAT_STRING_ERR(FormatStringSan::StatusOk, TEXT("%.*f %s"), 4, 42.4, TEXT("hi")));
+
+		STATIC_CHECK(UE_CHECK_FORMAT_STRING_ERR(FormatStringSan::StatusOk, TEXT("Test extra arg '%s'."), TEXT("ok"), TEXT("hi")));
+
+		STATIC_CHECK(UE_CHECK_FORMAT_STRING_ERR(FormatStringSan::StatusOk, TEXT("a")));
+		STATIC_CHECK(UE_CHECK_FORMAT_STRING_ERR(FormatStringSan::StatusOk, TEXT("")));
+
+		{
+			const TCHAR* Foo;
+			STATIC_CHECK(UE_CHECK_FORMAT_STRING_ERR(FormatStringSan::StatusOk, TEXT("hello %s"), Foo));
+		}
+		{
+			enum MyIntegralEnum { MyIntegralEnumA };
+			STATIC_CHECK(UE_CHECK_FORMAT_STRING_ERR(FormatStringSan::StatusOk, TEXT("hello %d is an enum actually"), MyIntegralEnumA));
+		}
+	}
+}
+
+} // namespace UE
+
+#endif // WITH_TESTS
