disabled delegates race detection while investigating reports

diff --git a/Engine/Source/Runtime/Core/Public/Delegates/DelegateAccessHandler.h b/Engine/Source/Runtime/Core/Public/Delegates/DelegateAccessHandler.h
--- a/Engine/Source/Runtime/Core/Public/Delegates/DelegateAccessHandler.h
+++ b/Engine/Source/Runtime/Core/Public/Delegates/DelegateAccessHandler.h
@@ -7,7 +7,7 @@
 #include "Misc/MTAccessDetector.h"
 #include "Misc/ScopeLock.h"
 
-//#define UE_DETECT_DELEGATES_RACE_CONDITIONS 0
+#define UE_DETECT_DELEGATES_RACE_CONDITIONS 0
 
 #if !defined(UE_DETECT_DELEGATES_RACE_CONDITIONS)
 #define UE_DETECT_DELEGATES_RACE_CONDITIONS ENABLE_MT_DETECTOR
