Remove aggressive early-out on rbp==0 in stack tracing.
Fall back to dbghelp-based stack tracing if manual unwinding fails.

#localization none
#tests Stepped through thread tracing in debugger. Confirmed local crash context xml has callstacks for all threads.

diff --git a/Engine/Source/Runtime/Core/Private/Microsoft/MicrosoftPlatformStackWalk.cpp b/Engine/Source/Runtime/Core/Private/Microsoft/MicrosoftPlatformStackWalk.cpp
--- a/Engine/Source/Runtime/Core/Private/Microsoft/MicrosoftPlatformStackWalk.cpp
+++ b/Engine/Source/Runtime/Core/Private/Microsoft/MicrosoftPlatformStackWalk.cpp
@@ -85,7 +85,7 @@ int32 FMicrosoftPlatformStackWalk::CaptureStackTraceInternal(uint64* OutBacktrac
 
 	*OutDepth = 0;
 
-	if (!ContextCopy.Rbp || !ContextCopy.Rip || !ContextCopy.Rsp)
+	if (!ContextCopy.Rip || !ContextCopy.Rsp)
 	{
 		return EXCEPTION_EXECUTE_HANDLER;
 	}
diff --git a/Engine/Source/Runtime/Core/Private/Windows/WindowsPlatformStackWalk.cpp b/Engine/Source/Runtime/Core/Private/Windows/WindowsPlatformStackWalk.cpp
--- a/Engine/Source/Runtime/Core/Private/Windows/WindowsPlatformStackWalk.cpp
+++ b/Engine/Source/Runtime/Core/Private/Windows/WindowsPlatformStackWalk.cpp
@@ -273,14 +273,19 @@ static int32 CaptureStackTraceExternalProcess(uint64* BackTrace, uint32 MaxDepth
 
 void FWindowsPlatformStackWalk::CaptureStackTraceByProcess(uint64* OutBacktrace, uint32 MaxDepth, void* InContext, void* InThreadHandle, uint32* OutDepth, bool bExternalProcess)
 {
-	if (bExternalProcess)
-	{
-		CaptureStackTraceExternalProcess(OutBacktrace, MaxDepth, reinterpret_cast<PCONTEXT>(InContext), reinterpret_cast<HANDLE>(InThreadHandle), OutDepth);	
-	}	
-	else
+	if (!bExternalProcess)
 	{
 		FMicrosoftPlatformStackWalk::CaptureStackTraceInternal(OutBacktrace, MaxDepth, InContext, InThreadHandle, OutDepth);
+		// If we fail to capture a stack trace with this method, fall back to using dbghelp
+		if (*OutDepth != 0)
+		{
+			return;
+		}
 	}
+
+	// Init if not already initialized for another process
+	InitStackWalking();
+	CaptureStackTraceExternalProcess(OutBacktrace, MaxDepth, reinterpret_cast<PCONTEXT>(InContext), reinterpret_cast<HANDLE>(InThreadHandle), OutDepth);	
 }
 
 uint32 FWindowsPlatformStackWalk::CaptureThreadStackBackTrace(uint64 ThreadId, uint64* BackTrace, uint32 MaxDepth, void* Context)
@@ -289,14 +294,18 @@ uint32 FWindowsPlatformStackWalk::CaptureThreadStackBackTrace(uint64 ThreadId, u
 	{
 		return 0;
 	}
-	
 
 	// Do not call InitStackWalking as we're just capturing addresses without using the dbghelp library etc if this is the same process.
 	// If we're walking another process, they will already have initialized stackwalking to set GProcessHandle.
 	// Don't suspend the calling thread, capture it's context directly and trace
 	if (ThreadId == FPlatformTLS::GetCurrentThreadId())
 	{
-		return CaptureStackBackTrace(BackTrace, MaxDepth);
+		// If we fail to capture a stack trace with this method, fall back to opening the thread handle
+		uint32 Depth = CaptureStackBackTrace(BackTrace, MaxDepth);
+		if (Depth != 0)
+		{
+			return Depth;
+		}
 	}
 
 	HANDLE ThreadHandle = OpenThread(THREAD_GET_CONTEXT | THREAD_SET_CONTEXT | THREAD_TERMINATE | THREAD_SUSPEND_RESUME, false, (DWORD)ThreadId);
