Change StackOffset to 0 in functions that call AddError from AutomationTest.h and force inline those functions

#rnx

diff --git a/Engine/Source/Runtime/Core/Public/Misc/AutomationTest.h b/Engine/Source/Runtime/Core/Public/Misc/AutomationTest.h
--- a/Engine/Source/Runtime/Core/Public/Misc/AutomationTest.h
+++ b/Engine/Source/Runtime/Core/Public/Misc/AutomationTest.h
@@ -1730,11 +1730,11 @@ public:
 	 * @see TestNotEqual
 	 */
 	template<typename ValueType> 
-	bool TestEqual(const TCHAR* What, const ValueType& Actual, const ValueType& Expected)
+	FORCEINLINE bool TestEqual(const TCHAR* What, const ValueType& Actual, const ValueType& Expected)
 	{
 		if (Actual != Expected)
 		{
-			AddError(FString::Printf(TEXT("%s: The two values are not equal."), What), 1);
+			AddError(FString::Printf(TEXT("%s: The two values are not equal."), What));
 			return false;
 		}
 		return true;
@@ -1801,11 +1801,12 @@ public:
 	 *
 	 * @see TestValid
 	 */
-	template<typename ValueType> bool TestInvalid(const TCHAR* Description, const TSharedPtr<ValueType>& SharedPointer)
+	template<typename ValueType>
+	FORCEINLINE bool TestInvalid(const TCHAR* Description, const TSharedPtr<ValueType>& SharedPointer)
 	{
 		if (SharedPointer.IsValid())
 		{
-			AddError(FString::Printf(TEXT("%s: The shared pointer is valid."), Description), 1);
+			AddError(FString::Printf(TEXT("%s: The shared pointer is valid."), Description));
 			return false;
 		}
 		return true;
@@ -1825,11 +1826,12 @@ public:
 	 *
 	 * @see TestEqual
 	 */
-	template<typename ValueType> bool TestNotEqual(const TCHAR* Description, const ValueType& Actual, const ValueType& Expected)
+	template<typename ValueType>
+	FORCEINLINE bool TestNotEqual(const TCHAR* Description, const ValueType& Actual, const ValueType& Expected)
 	{
 		if (Actual == Expected)
 		{
-			AddError(FString::Printf(TEXT("%s: The two values are equal."), Description), 1);
+			AddError(FString::Printf(TEXT("%s: The two values are equal."), Description));
 			return false;
 		}
 		return true;
@@ -1848,11 +1850,12 @@ public:
 	 *
 	 * @see TestNull
 	 */
-	template<typename ValueType> bool TestNotNull(const TCHAR* What, const ValueType* Pointer)
+	template<typename ValueType>
+	FORCEINLINE bool TestNotNull(const TCHAR* What, const ValueType* Pointer)
 	{
 		if (Pointer == nullptr)
 		{
-			AddError(FString::Printf(TEXT("Expected '%s' to be not null."), What), 1);
+			AddError(FString::Printf(TEXT("Expected '%s' to be not null."), What));
 			return false;
 		}
 		return true;
@@ -1872,11 +1875,12 @@ public:
 	 *
 	 * @see TestSame
 	 */
-	template<typename ValueType> bool TestNotSame(const TCHAR* Description, const ValueType& Actual, const ValueType& Expected)
+	template<typename ValueType>
+	FORCEINLINE bool TestNotSame(const TCHAR* Description, const ValueType& Actual, const ValueType& Expected)
 	{
 		if (&Actual == &Expected)
 		{
-			AddError(FString::Printf(TEXT("%s: The two values are the same."), Description), 1);
+			AddError(FString::Printf(TEXT("%s: The two values are the same."), Description));
 			return false;
 		}
 		return true;
@@ -1911,11 +1915,12 @@ public:
 	 *
 	 * @see TestNotSame
 	 */
-	template<typename ValueType> bool TestSame(const TCHAR* Description, const ValueType& Actual, const ValueType& Expected)
+	template<typename ValueType>
+	FORCEINLINE bool TestSame(const TCHAR* Description, const ValueType& Actual, const ValueType& Expected)
 	{
 		if (&Actual != &Expected)
 		{
-			AddError(FString::Printf(TEXT("%s: The two values are not the same."), Description), 1);
+			AddError(FString::Printf(TEXT("%s: The two values are not the same."), Description));
 			return false;
 		}
 		return true;
@@ -1952,11 +1957,12 @@ public:
 	 *
 	 * @see TestInvalid
 	 */
-	template<typename ValueType> bool TestValid(const TCHAR* Description, const TSharedPtr<ValueType>& SharedPointer)
+	template<typename ValueType>
+	FORCEINLINE bool TestValid(const TCHAR* Description, const TSharedPtr<ValueType>& SharedPointer)
 	{
 		if (!SharedPointer.IsValid())
 		{
-			AddError(FString::Printf(TEXT("%s: The shared pointer is not valid."), Description), 1);
+			AddError(FString::Printf(TEXT("%s: The shared pointer is not valid."), Description));
 			return false;
 		}
 		return true;
