Adding test for Parse::GrammarBasedCLIParse

diff --git a/Engine/Source/Runtime/Core/Tests/Misc/ParseTests.cpp b/Engine/Source/Runtime/Core/Tests/Misc/ParseTests.cpp
--- a/Engine/Source/Runtime/Core/Tests/Misc/ParseTests.cpp
+++ b/Engine/Source/Runtime/Core/Tests/Misc/ParseTests.cpp
@@ -8,6 +8,8 @@
 
 #include "Tests/TestHarnessAdapter.h"
 
+#include <catch2/generators/catch_generators.hpp>
+
 TEST_CASE("Parse::Value::ToBuffer", "[Parse][Smoke]")
 {
 	TCHAR Buffer[256];
@@ -117,5 +119,109 @@ TEST_CASE("Parse::Value::ToBuffer", "[Parse][Smoke]")
 	}
 }
 
+
+TEST_CASE("Parse::GrammaredCLIParse::Callback", "[Smoke]")
+{
+	struct StringKeyValue {
+		const TCHAR* Key;
+		const TCHAR* Value;
+	};
+
+	SECTION("ExpectedPass")
+	{
+		auto [Input, ExpectedResults] = GENERATE_COPY(table<const TCHAR*, std::vector<StringKeyValue>>(
+			{
+				{ TEXT("basic_ident"),	{ {TEXT("basic_ident"), nullptr} } },
+				{ TEXT("-one_dash"),	{ {TEXT("-one_dash"), nullptr} } },
+				{ TEXT("--two_dash"),	{ {TEXT("--two_dash"), nullptr} } },
+				{ TEXT("/slash"),		{ {TEXT("/slash"), nullptr} } },
+				{ TEXT("key=value"),	{ {TEXT("key"), TEXT("value")} } },
+				{ TEXT("key.with.dots=value"),	{ {TEXT("key.with.dots"), TEXT("value")} } },
+				{ TEXT("-key=value"),	{ {TEXT("-key"), TEXT("value")} } },
+				{ TEXT("-key=\"value\""), { {TEXT("-key"), TEXT("\"value\"")} } },
+				{ TEXT("-key=111"),		{ {TEXT("-key"), TEXT("111")} } },
+				{ TEXT("-key=111."),	{ {TEXT("-key"), TEXT("111.")} } },
+				{ TEXT("-key=111.222"),	{ {TEXT("-key"), TEXT("111.222")} } },
+				{ TEXT("-key=-111"),	{ {TEXT("-key"), TEXT("-111")} } },
+				{ TEXT("-key=-111.22"),	{ {TEXT("-key"), TEXT("-111.22")} } },
+				{ TEXT("-key=../../some+dir\\text-file.txt"),	{ {TEXT("-key"), TEXT("../../some+dir\\text-file.txt")} } },
+				{ TEXT("-key=c:\\log.txt"),	{ {TEXT("-key"), TEXT("c:\\log.txt")} } },
+				{ TEXT("a -b --c d=e"),	{ {TEXT("a"), nullptr},
+										  {TEXT("-b"), nullptr},
+										  {TEXT("--c"), nullptr},
+										  {TEXT("d"), TEXT("e")} } },
+				{ TEXT("a \"-b --c\" d=e"),	{ {TEXT("a"), nullptr},
+											  {TEXT("-b"), nullptr},
+											  {TEXT("--c"), nullptr},
+											  {TEXT("d"), TEXT("e")} } },
+				{ TEXT("\"a -b --c d=e\""),	{ {TEXT("a"), nullptr},
+											  {TEXT("-b"), nullptr},
+											  {TEXT("--c"), nullptr},
+											  {TEXT("d"), TEXT("e")} } },
+				{ TEXT("    leading_space"), { {TEXT("leading_space"), nullptr} } },
+				{ TEXT("trailing_space   "), { {TEXT("trailing_space"), nullptr} } },
+			}));
+		size_t CallbackCalledCount = 0;
+		auto CallBack = [&](FStringView Key, FStringView Value)
+		{
+			REQUIRE(CallbackCalledCount < ExpectedResults.size());
+			CHECK(Key == FStringView{ ExpectedResults[CallbackCalledCount].Key });
+			CHECK(Value == FStringView{ ExpectedResults[CallbackCalledCount].Value });
+			++CallbackCalledCount;
+		};
+
+		FParse::FGrammarBasedParseResult Result = FParse::GrammarBasedCLIParse(Input, CallBack);
+
+		CHECK(CallbackCalledCount == ExpectedResults.size());
+		CHECK(Result.ErrorCode == FParse::EGrammarBasedParseErrorCode::Succeeded);
+	}
+
+	SECTION("Quoted commadns may be dissallowed, if so gives an error code.")
+	{
+		auto [Input, ExpectedErrorCode, ExpectedErrorAt, ExpectedResults] = GENERATE_COPY(table<const TCHAR*, FParse::EGrammarBasedParseErrorCode, size_t, std::vector<StringKeyValue>>(
+			{
+				{ TEXT("a \"-b --c\" d=e"), FParse::EGrammarBasedParseErrorCode::DisallowedQuotedCommand, 2, { {TEXT("a"), nullptr} } },
+			}));
+		size_t CallbackCalledCount = 0;
+		auto CallBack = [&](FStringView Key, FStringView Value)
+		{
+			REQUIRE(CallbackCalledCount < ExpectedResults.size());
+			CHECK(Key == FStringView{ ExpectedResults[CallbackCalledCount].Key });
+			CHECK(Value == FStringView{ ExpectedResults[CallbackCalledCount].Value });
+			++CallbackCalledCount;
+		};
+
+		FParse::FGrammarBasedParseResult Result = FParse::GrammarBasedCLIParse(Input, CallBack, FParse::EGrammarBasedParseFlags::None);
+
+		CHECK(CallbackCalledCount == ExpectedResults.size());
+		CHECK(Result.ErrorCode == ExpectedErrorCode);
+		CHECK(Result.At == Input + ExpectedErrorAt);
+	}
+
+	SECTION("Expected Fail cases")
+	{
+		auto [Input, ExpectedErrorCode, ExpectedErrorAt, ExpectedResults] = GENERATE_COPY(table<const TCHAR*, FParse::EGrammarBasedParseErrorCode, size_t, std::vector<StringKeyValue>>(
+			{
+					{ TEXT("-a \"-b"), FParse::EGrammarBasedParseErrorCode::UnBalancedQuote, 3, { {TEXT("-a"), nullptr},
+																								{TEXT("-b"), nullptr}} },
+					{ TEXT("-a=\"unbalanced_quote_value"), FParse::EGrammarBasedParseErrorCode::UnBalancedQuote, 3, { } },
+			}));
+		size_t CallbackCalledCount = 0;
+		auto CallBack = [&](FStringView Key, FStringView Value)
+		{
+			REQUIRE(CallbackCalledCount < ExpectedResults.size());
+			CHECK(Key == FStringView{ ExpectedResults[CallbackCalledCount].Key });
+			CHECK(Value == FStringView{ ExpectedResults[CallbackCalledCount].Value });
+			++CallbackCalledCount;
+		};
+
+		FParse::FGrammarBasedParseResult Result = FParse::GrammarBasedCLIParse(Input, CallBack);
+
+		CHECK(CallbackCalledCount == ExpectedResults.size());
+		CHECK(Result.ErrorCode == ExpectedErrorCode);
+		CHECK(Result.At == Input + ExpectedErrorAt);
+	}
+}
+
 #endif
 	
\ No newline at end of file
