When get hotfixed, make marked cvars can be saved for next boot.

diff --git a/Engine/Source/Runtime/Core/Private/Misc/ConfigUtilities.cpp b/Engine/Source/Runtime/Core/Private/Misc/ConfigUtilities.cpp
--- a/Engine/Source/Runtime/Core/Private/Misc/ConfigUtilities.cpp
+++ b/Engine/Source/Runtime/Core/Private/Misc/ConfigUtilities.cpp
@@ -5,6 +5,9 @@
 #include "HAL/FileManager.h"
 #include "HAL/IConsoleManager.h"
 #include "Misc/CoreDelegates.h"
+#include "Misc/FileHelper.h"
+
+#define UE_HOTFIX_FOR_NEXT_BOOT_FILENAME TEXT("HotfixForNextBoot.ini")
 
 namespace UE::ConfigUtilities
 {
@@ -30,6 +33,80 @@ const TCHAR* ConvertValueFromHumanFriendlyValue( const TCHAR* Value )
 	return Value;
 }
 
+// The dedicated server could have been deployed through cloud provider and allocated dynamically. 
+// So can't really save the file for next boot, consider the runtime args for cvar instead
+#if !UE_SERVER
+void LoadCVarsFromFileForNextBoot(TMap<FString, FString>& OutCVars)
+{
+	IFileManager& FileManager = IFileManager::Get();
+
+	const FString FullPath = FPaths::ProjectPersistentDownloadDir() / UE_HOTFIX_FOR_NEXT_BOOT_FILENAME;
+
+	if (!FileManager.FileExists(*FullPath))
+	{
+		return;
+	}
+
+	FString Content;
+	FFileHelper::LoadFileToString(Content, *FullPath);
+
+	// Delete it so that we don't worry about it when write.
+	// Also if for some reason the switch don't work well when boot even before getting latest hotfix, 
+	// the next boot will likely succeed by default like before without this file
+	FileManager.Delete(*FullPath);
+
+	TArray<FString> Lines;
+	Content.ParseIntoArrayLines(Lines);
+
+	for (const FString& Line : Lines)
+	{
+		FString Key;
+		FString Value;
+		if (ensure(Line.Split(TEXT("="), &Key, &Value)))
+		{
+			OutCVars.FindOrAdd(Key) = Value;
+		}
+	}
+}
+#endif // !UE_SERVER
+
+void SaveCVarForNextBoot(const TCHAR* Key, const TCHAR* Value)
+{
+#if !UE_SERVER
+	TMap<FString, FString> CVarsToSave;
+
+	// Read from file, in case there are more than one cvar hotfix event in same run
+	LoadCVarsFromFileForNextBoot(CVarsToSave);
+
+	CVarsToSave.FindOrAdd(Key) = Value;
+
+	FString ContentToSave;
+	for (const TPair<FString, FString>& CVarPair : CVarsToSave)
+	{
+		ContentToSave.Append(FString::Format(TEXT("{0}={1}\r\n"), { CVarPair.Key, CVarPair.Value }));
+	}
+
+	const FString FullPath = FPaths::ProjectPersistentDownloadDir() / UE_HOTFIX_FOR_NEXT_BOOT_FILENAME;
+	FFileHelper::SaveStringToFile(ContentToSave, *FullPath);
+#endif // !UE_SERVER
+}
+
+void ApplyCVarsFromBootHotfix()
+{
+#if !UE_SERVER
+	TMap<FString, FString> CVarsToApply;
+	LoadCVarsFromFileForNextBoot(CVarsToApply);
+
+	for (const TPair<FString, FString>& CVarPair : CVarsToApply)
+	{
+		IConsoleVariable* CVar = IConsoleManager::Get().FindConsoleVariable(*CVarPair.Key); 
+		if (CVar)
+		{
+			CVar->Set(*CVarPair.Value, ECVF_SetByConsoleVariablesIni);
+		}
+	}
+#endif // !UE_SERVER
+}
 
 void OnSetCVarFromIniEntry(const TCHAR *IniFile, const TCHAR *Key, const TCHAR* Value, uint32 SetBy, bool bAllowCheating, bool bNoLogging)
 {
@@ -110,6 +187,11 @@ void OnSetCVarFromIniEntry(const TCHAR *IniFile, const TCHAR *Key, const TCHAR*
 			}
 #endif // !DISABLE_CHEAT_CVARS
 		}
+
+		if (CVar->TestFlags(ECVF_SaveForNextBoot) && (SetBy & ECVF_SetByConsoleVariablesIni))
+		{
+			SaveCVarForNextBoot(Key, Value);
+		}
 	}
 	else
 	{
diff --git a/Engine/Source/Runtime/Core/Public/HAL/IConsoleManager.h b/Engine/Source/Runtime/Core/Public/HAL/IConsoleManager.h
--- a/Engine/Source/Runtime/Core/Public/HAL/IConsoleManager.h
+++ b/Engine/Source/Runtime/Core/Public/HAL/IConsoleManager.h
@@ -107,6 +107,17 @@ enum EConsoleVariableFlags
 	/* CVars with this flag will be excluded from the device profile previews. */
 	ECVF_ExcludeFromPreview = 0x1000,
 
+	/**
+	 * CVars with this flag will be saved into local temp file for next boot after hotfixed. Normally use for feature switch which is too early, even before launching the hotfix http request
+	 * Use it carefully with these constrains in mind: 
+	 *  - WARNING: This exposes the config to end user who could change it through the local temp file in shipping client, it's your responsibility to make sure client can't "legally" cheat through this type of CVar;
+	 *  - For the first time get hotfixed, the value will change from default value during boot to the hotfixed value after the hotfix http request, so make sure caching it when boot if the logic doesn't support to change at runtime;
+	 *  - When there is a new released game version, the saved local temp file of last version could still be there. So it will still got read once. Keep the logic work well with the old value, or remove the CVar if the old value no longer supported;
+	 *  - Only set the default value in code, don't set default value in [ConsoleVariables] section in the Engine.ini, otherwise reading the Engine.ini shipped within the game will overwrite the hotfixed value. This constrain also means it won't support undo at runtime;
+	 *  - It will only work for game client, not for dedicated servers which will be dynamically allocated on different server instances.
+	 */
+	ECVF_SaveForNextBoot = 0x2000,
+
 	// ------------------------------------------------
 
 	/* Set flags */
diff --git a/Engine/Source/Runtime/Core/Public/Misc/ConfigUtilities.h b/Engine/Source/Runtime/Core/Public/Misc/ConfigUtilities.h
--- a/Engine/Source/Runtime/Core/Public/Misc/ConfigUtilities.h
+++ b/Engine/Source/Runtime/Core/Public/Misc/ConfigUtilities.h
@@ -22,6 +22,11 @@ namespace UE::ConfigUtilities
 	 * Single function to set a cvar from ini (handing friendly names, cheats for shipping and message about cheats in non shipping)
 	 */
 	CORE_API void OnSetCVarFromIniEntry(const TCHAR* IniFile, const TCHAR* Key, const TCHAR* Value, uint32 SetBy, bool bAllowCheating=false, bool bNoLogging=false);
+
+	/**
+	 * When boot the game, use this function to apply cvars from last saved file from hotfix
+	 */
+	CORE_API void ApplyCVarsFromBootHotfix();
 	
 	/**
 	 * Helper function to read the contents of an ini file and a specified group of cvar parameters, where sections in the ini file are marked [InName]
diff --git a/Engine/Source/Runtime/Launch/Private/LaunchEngineLoop.cpp b/Engine/Source/Runtime/Launch/Private/LaunchEngineLoop.cpp
--- a/Engine/Source/Runtime/Launch/Private/LaunchEngineLoop.cpp
+++ b/Engine/Source/Runtime/Launch/Private/LaunchEngineLoop.cpp
@@ -2769,6 +2769,8 @@ int32 FEngineLoop::PreInitPreStartupScreen(const TCHAR* CmdLine)
 		UE::ConfigUtilities::RecordApplyCVarSettingsFromIni();
 	}
 
+	UE::ConfigUtilities::ApplyCVarsFromBootHotfix();
+
 #if WITH_ENGINE
 	extern ENGINE_API void InitializeRenderingCVarsCaching();
 	InitializeRenderingCVarsCaching();
