UE_LOG sanitizer: use global UE namespace explicitly

diff --git a/Engine/Source/Runtime/Core/Public/String/FormatStringSan.h b/Engine/Source/Runtime/Core/Public/String/FormatStringSan.h
--- a/Engine/Source/Runtime/Core/Public/String/FormatStringSan.h
+++ b/Engine/Source/Runtime/Core/Public/String/FormatStringSan.h
@@ -5,15 +5,15 @@
 #include "CoreTypes.h"
 
 #define UE_CHECK_FORMAT_STRING(Fmt, ...) do {														\
-		typedef UE::Core::Private::FormatStringSan::TCheckFormatString<decltype(UE::Core::Private::FormatStringSan::GetFmtArgTypes(__VA_ARGS__))> UE_FMT_STR_Checker; \
-		constexpr UE::Core::Private::FormatStringSan::FResult UE_FMT_STR_Result = UE_FMT_STR_Checker::Check(false, 0, Fmt); \
-		typedef UE::Core::Private::FormatStringSan::TPresentErr<UE_FMT_STR_Result.Status> UE_FMT_STR_ErrorMsg; \
+		typedef ::UE::Core::Private::FormatStringSan::TCheckFormatString<decltype(::UE::Core::Private::FormatStringSan::GetFmtArgTypes(__VA_ARGS__))> UE_FMT_STR_Checker; \
+		constexpr ::UE::Core::Private::FormatStringSan::FResult UE_FMT_STR_Result = UE_FMT_STR_Checker::Check(false, 0, Fmt); \
+		typedef ::UE::Core::Private::FormatStringSan::TPresentErr<UE_FMT_STR_Result.Status> UE_FMT_STR_ErrorMsg; \
 		typedef typename UE_FMT_STR_ErrorMsg::template TValue<UE_FMT_STR_Result.ArgPos> UE_FMT_STR_Status; \
 		(void) UE_FMT_STR_Status{};																					\
 	} while(false)
 
 #define UE_CHECK_FORMAT_STRING_ERR(Err, Fmt, ...)												\
-	(UE::Core::Private::FormatStringSan::TCheckFormatString<decltype(UE::Core::Private::FormatStringSan::GetFmtArgTypes(__VA_ARGS__))>::Check(false, 0, Fmt).Status == Err)
+	(::UE::Core::Private::FormatStringSan::TCheckFormatString<decltype(::UE::Core::Private::FormatStringSan::GetFmtArgTypes(__VA_ARGS__))>::Check(false, 0, Fmt).Status == Err)
 
 /// implementation
 namespace UE::Core::Private
