Remove constraint on FPlayInEditorLoadingScope being set from workers as it doesn't matter anyway

#rnx

diff --git a/Engine/Source/Runtime/Core/Private/Misc/CoreGlobals.cpp b/Engine/Source/Runtime/Core/Private/Misc/CoreGlobals.cpp
--- a/Engine/Source/Runtime/Core/Private/Misc/CoreGlobals.cpp
+++ b/Engine/Source/Runtime/Core/Private/Misc/CoreGlobals.cpp
@@ -786,7 +786,7 @@ int32 PRIVATE_GetGPlayInEditorID()
 	}
 	else
 	{
-		//ensureMsgf(false, TEXT("GPlayInEditorID can only be read by the game-thread or loading-thread, returning -1."));
+		// GPlayInEditorID doesn't have a value on worker threads. If it's needed, it should be captured in the task context.
 		return -1;
 	}
 }
@@ -799,7 +799,7 @@ void PRIVATE_SetGPlayInEditorID(int32 InValue)
 	}
 	else
 	{
-		checkf(false, TEXT("GPlayInEditorID can only be set by the game-thread or loading-thread."));
+		// There is no value on worker thread... so just do nothing here. -1 is always returned anyway.
 	}
 }
 
