Fix self registering math functions for Linux.
Since they have ambigous entry points, float and long double need to be explicit for each one

diff --git a/Engine/Source/Runtime/Core/Private/AutoRTFM/StdLibSupport.cpp b/Engine/Source/Runtime/Core/Private/AutoRTFM/StdLibSupport.cpp
--- a/Engine/Source/Runtime/Core/Private/AutoRTFM/StdLibSupport.cpp
+++ b/Engine/Source/Runtime/Core/Private/AutoRTFM/StdLibSupport.cpp
@@ -117,38 +117,78 @@ UE_AUTORTFM_REGISTER_SELF_FUNCTION(strlen);
 
 UE_AUTORTFM_REGISTER_SELF_FUNCTION(wcscmp);
 
-// Self register Math functions
-UE_AUTORTFM_REGISTER_SELF_FUNCTION(sqrt);
-UE_AUTORTFM_REGISTER_SELF_FUNCTION(sqrtf);
+#if PLATFORM_WINDOWS
 UE_AUTORTFM_REGISTER_SELF_FUNCTION(sin);
-UE_AUTORTFM_REGISTER_SELF_FUNCTION(sinf);
 UE_AUTORTFM_REGISTER_SELF_FUNCTION(cos);
-UE_AUTORTFM_REGISTER_SELF_FUNCTION(cosf);
 UE_AUTORTFM_REGISTER_SELF_FUNCTION(tan);
-UE_AUTORTFM_REGISTER_SELF_FUNCTION(tanf);
 UE_AUTORTFM_REGISTER_SELF_FUNCTION(asin);
-UE_AUTORTFM_REGISTER_SELF_FUNCTION(asinf);
 UE_AUTORTFM_REGISTER_SELF_FUNCTION(acos);
-UE_AUTORTFM_REGISTER_SELF_FUNCTION(acosf);
 UE_AUTORTFM_REGISTER_SELF_FUNCTION(atan);
-UE_AUTORTFM_REGISTER_SELF_FUNCTION(atanf);
 UE_AUTORTFM_REGISTER_SELF_FUNCTION(atan2);
 UE_AUTORTFM_REGISTER_SELF_FUNCTION(sinh);
-UE_AUTORTFM_REGISTER_SELF_FUNCTION(sinhf);
 UE_AUTORTFM_REGISTER_SELF_FUNCTION(cosh);
-UE_AUTORTFM_REGISTER_SELF_FUNCTION(coshf);
 UE_AUTORTFM_REGISTER_SELF_FUNCTION(tanh);
-UE_AUTORTFM_REGISTER_SELF_FUNCTION(tanhf);
 UE_AUTORTFM_REGISTER_SELF_FUNCTION(asinh);
 UE_AUTORTFM_REGISTER_SELF_FUNCTION(acosh);
 UE_AUTORTFM_REGISTER_SELF_FUNCTION(atanh);
 UE_AUTORTFM_REGISTER_SELF_FUNCTION(exp);
-UE_AUTORTFM_REGISTER_SELF_FUNCTION(expf);
 UE_AUTORTFM_REGISTER_SELF_FUNCTION(log);
-UE_AUTORTFM_REGISTER_SELF_FUNCTION(logf);
 UE_AUTORTFM_REGISTER_SELF_FUNCTION(pow);
-UE_AUTORTFM_REGISTER_SELF_FUNCTION(powf);
 UE_AUTORTFM_REGISTER_SELF_FUNCTION(llrint);
+// Linux (likely Mac) have ambiguous overrides to these math functions
+#else
+UE_AUTORTFM_REGISTER_SELF_FUNCTION(static_cast<float(*)(float)>(&sqrt));
+UE_AUTORTFM_REGISTER_SELF_FUNCTION(static_cast<long double(*)(long double)>(&sqrt));
+UE_AUTORTFM_REGISTER_SELF_FUNCTION(static_cast<float(*)(float)>(&sin));
+UE_AUTORTFM_REGISTER_SELF_FUNCTION(static_cast<long double(*)(long double)>(&sin));
+UE_AUTORTFM_REGISTER_SELF_FUNCTION(static_cast<float(*)(float)>(&cos));
+UE_AUTORTFM_REGISTER_SELF_FUNCTION(static_cast<long double(*)(long double)>(&cos));
+UE_AUTORTFM_REGISTER_SELF_FUNCTION(static_cast<float(*)(float)>(&tan));
+UE_AUTORTFM_REGISTER_SELF_FUNCTION(static_cast<long double(*)(long double)>(&tan));
+UE_AUTORTFM_REGISTER_SELF_FUNCTION(static_cast<float(*)(float)>(&asin));
+UE_AUTORTFM_REGISTER_SELF_FUNCTION(static_cast<long double(*)(long double)>(&asin));
+UE_AUTORTFM_REGISTER_SELF_FUNCTION(static_cast<float(*)(float)>(&acos));
+UE_AUTORTFM_REGISTER_SELF_FUNCTION(static_cast<long double(*)(long double)>(&acos));
+UE_AUTORTFM_REGISTER_SELF_FUNCTION(static_cast<float(*)(float)>(&atan));
+UE_AUTORTFM_REGISTER_SELF_FUNCTION(static_cast<long double(*)(long double)>(&atan));
+UE_AUTORTFM_REGISTER_SELF_FUNCTION(static_cast<float(*)(float, float)>(&atan2));
+UE_AUTORTFM_REGISTER_SELF_FUNCTION(static_cast<long double(*)(long double, long double)>(&atan2));
+UE_AUTORTFM_REGISTER_SELF_FUNCTION(static_cast<float(*)(float)>(&sinh));
+UE_AUTORTFM_REGISTER_SELF_FUNCTION(static_cast<long double(*)(long double)>(&sinh));
+UE_AUTORTFM_REGISTER_SELF_FUNCTION(static_cast<float(*)(float)>(&cosh));
+UE_AUTORTFM_REGISTER_SELF_FUNCTION(static_cast<long double(*)(long double)>(&cosh));
+UE_AUTORTFM_REGISTER_SELF_FUNCTION(static_cast<float(*)(float)>(&tanh));
+UE_AUTORTFM_REGISTER_SELF_FUNCTION(static_cast<long double(*)(long double)>(&tanh));
+UE_AUTORTFM_REGISTER_SELF_FUNCTION(static_cast<float(*)(float)>(&asinh));
+UE_AUTORTFM_REGISTER_SELF_FUNCTION(static_cast<long double(*)(long double)>(&asinh));
+UE_AUTORTFM_REGISTER_SELF_FUNCTION(static_cast<float(*)(float)>(&acosh));
+UE_AUTORTFM_REGISTER_SELF_FUNCTION(static_cast<long double(*)(long double)>(&acosh));
+UE_AUTORTFM_REGISTER_SELF_FUNCTION(static_cast<float(*)(float)>(&atanh));
+UE_AUTORTFM_REGISTER_SELF_FUNCTION(static_cast<long double(*)(long double)>(&atanh));
+UE_AUTORTFM_REGISTER_SELF_FUNCTION(static_cast<float(*)(float)>(&exp));
+UE_AUTORTFM_REGISTER_SELF_FUNCTION(static_cast<long double(*)(long double)>(&exp));
+UE_AUTORTFM_REGISTER_SELF_FUNCTION(static_cast<float(*)(float)>(&log));
+UE_AUTORTFM_REGISTER_SELF_FUNCTION(static_cast<long double(*)(long double)>(&log));
+UE_AUTORTFM_REGISTER_SELF_FUNCTION(static_cast<float(*)(float, float)>(&pow));
+UE_AUTORTFM_REGISTER_SELF_FUNCTION(static_cast<long double(*)(long double, long double)>(&pow));
+UE_AUTORTFM_REGISTER_SELF_FUNCTION(static_cast<long long(*)(float)>(&llrint));
+UE_AUTORTFM_REGISTER_SELF_FUNCTION(static_cast<long long(*)(long double)>(&llrint));
+#endif
+
+// Self register Math functions
+UE_AUTORTFM_REGISTER_SELF_FUNCTION(sqrtf);
+UE_AUTORTFM_REGISTER_SELF_FUNCTION(sinf);
+UE_AUTORTFM_REGISTER_SELF_FUNCTION(cosf);
+UE_AUTORTFM_REGISTER_SELF_FUNCTION(tanf);
+UE_AUTORTFM_REGISTER_SELF_FUNCTION(asinf);
+UE_AUTORTFM_REGISTER_SELF_FUNCTION(acosf);
+UE_AUTORTFM_REGISTER_SELF_FUNCTION(atanf);
+UE_AUTORTFM_REGISTER_SELF_FUNCTION(sinhf);
+UE_AUTORTFM_REGISTER_SELF_FUNCTION(coshf);
+UE_AUTORTFM_REGISTER_SELF_FUNCTION(tanhf);
+UE_AUTORTFM_REGISTER_SELF_FUNCTION(expf);
+UE_AUTORTFM_REGISTER_SELF_FUNCTION(logf);
+UE_AUTORTFM_REGISTER_SELF_FUNCTION(powf);
 
 // FIXME: This is only correct when:
 // - Str is newly allocated
