Fix missing math function self registering

diff --git a/Engine/Source/Runtime/Core/Private/AutoRTFM/StdLibSupport.cpp b/Engine/Source/Runtime/Core/Private/AutoRTFM/StdLibSupport.cpp
--- a/Engine/Source/Runtime/Core/Private/AutoRTFM/StdLibSupport.cpp
+++ b/Engine/Source/Runtime/Core/Private/AutoRTFM/StdLibSupport.cpp
@@ -118,6 +118,7 @@ UE_AUTORTFM_REGISTER_SELF_FUNCTION(strlen);
 UE_AUTORTFM_REGISTER_SELF_FUNCTION(wcscmp);
 
 #if PLATFORM_WINDOWS
+UE_AUTORTFM_REGISTER_SELF_FUNCTION(sqrt);
 UE_AUTORTFM_REGISTER_SELF_FUNCTION(sin);
 UE_AUTORTFM_REGISTER_SELF_FUNCTION(cos);
 UE_AUTORTFM_REGISTER_SELF_FUNCTION(tan);
