Restore old allowlist filtering over the gammar based parser filtering

diff --git a/Engine/Source/Runtime/Core/Private/Misc/CommandLine.cpp b/Engine/Source/Runtime/Core/Private/Misc/CommandLine.cpp
--- a/Engine/Source/Runtime/Core/Private/Misc/CommandLine.cpp
+++ b/Engine/Source/Runtime/Core/Private/Misc/CommandLine.cpp
@@ -225,13 +225,66 @@ void FCommandLine::ApplyCommandLineAllowList()
 		FCommandLine::Parse(FilterForLoggingList, FilterArgsForLogging, Ignored);
 	}
 	// Process the original command line
-	FCommandLine::FilterMove(OriginalCmdLine, UE_ARRAY_COUNT(OriginalCmdLine), OriginalCmdLine, ApprovedArgs);
+	TArray<FString> OriginalList = FilterCommandLine(OriginalCmdLine);
+	BuildCommandLineAllowList(OriginalCmdLine, UE_ARRAY_COUNT(OriginalCmdLine), OriginalList);
 	// Process the current command line
-	FCommandLine::FilterMove(CmdLine, UE_ARRAY_COUNT(CmdLine), CmdLine, ApprovedArgs);
+	TArray<FString> CmdList = FilterCommandLine(CmdLine);
+	BuildCommandLineAllowList(CmdLine, UE_ARRAY_COUNT(CmdLine), CmdList);
 	// Process the command line for logging purposes
-	FCommandLine::FilterMove(LoggingCmdLine, UE_ARRAY_COUNT(LoggingCmdLine), LoggingCmdLine, FilterArgsForLogging);
+	TArray<FString> LoggingCmdList = FilterCommandLineForLogging(LoggingCmdLine);
+	BuildCommandLineAllowList(LoggingCmdLine, UE_ARRAY_COUNT(LoggingCmdLine), LoggingCmdList);
 	// Process the original command line for logging purposes
-	FCommandLine::FilterMove(LoggingOriginalCmdLine, UE_ARRAY_COUNT(LoggingOriginalCmdLine), LoggingOriginalCmdLine, FilterArgsForLogging);
+	TArray<FString> LoggingOriginalCmdList = FilterCommandLineForLogging(LoggingOriginalCmdLine);
+	BuildCommandLineAllowList(LoggingOriginalCmdLine, UE_ARRAY_COUNT(LoggingOriginalCmdLine), LoggingOriginalCmdList);
+}
+
+TArray<FString> FCommandLine::FilterCommandLine(TCHAR* CommandLine)
+{
+	TArray<FString> Ignored;
+	TArray<FString> ParsedList;
+	// Parse the command line list
+	FCommandLine::Parse(CommandLine, ParsedList, Ignored);
+	// Remove any that are not in our approved list
+	for (int32 Index = 0; Index < ParsedList.Num(); Index++)
+	{
+		bool bFound = false;
+		for (auto ApprovedArg : ApprovedArgs)
+		{
+			if (ParsedList[Index].StartsWith(ApprovedArg))
+			{
+				bFound = true;
+				break;
+			}
+		}
+		if (!bFound)
+		{
+			ParsedList.RemoveAt(Index);
+			Index--;
+		}
+	}
+	return ParsedList;
+}
+
+TArray<FString> FCommandLine::FilterCommandLineForLogging(TCHAR* CommandLine)
+{
+	TArray<FString> Ignored;
+	TArray<FString> ParsedList;
+	// Parse the command line list
+	FCommandLine::Parse(CommandLine, ParsedList, Ignored);
+	// Remove any that are not in our approved list
+	for (int32 Index = 0; Index < ParsedList.Num(); Index++)
+	{
+		for (auto Filter : FilterArgsForLogging)
+		{
+			if (ParsedList[Index].StartsWith(Filter))
+			{
+				ParsedList.RemoveAt(Index);
+				Index--;
+				break;
+			}
+		}
+	}
+	return ParsedList;
 }
 
 void FCommandLine::BuildCommandLineAllowList(TCHAR* CommandLine, uint32 ArrayCount, const TArray<FString>& FilteredArgs)
diff --git a/Engine/Source/Runtime/Core/Public/Misc/CommandLine.h b/Engine/Source/Runtime/Core/Public/Misc/CommandLine.h
--- a/Engine/Source/Runtime/Core/Public/Misc/CommandLine.h
+++ b/Engine/Source/Runtime/Core/Public/Misc/CommandLine.h
@@ -126,6 +126,10 @@ private:
 #if UE_COMMAND_LINE_USES_ALLOW_LIST
 	/** Filters both the original and current command line list for approved only args */
 	static void ApplyCommandLineAllowList();
+	/** Filters any command line args that aren't on the approved list */
+	static TArray<FString> FilterCommandLine(TCHAR* CommandLine);
+	/** Filters any command line args that are on the to-strip list */
+	static TArray<FString> FilterCommandLineForLogging(TCHAR* CommandLine);
 	/** Rebuilds the command line using the filtered args */
 	static void BuildCommandLineAllowList(TCHAR* CommandLine, uint32 Length, const TArray<FString>& FilteredArgs);
 	static TArray<FString> ApprovedArgs;
