Preparation for deprecating the hashing arrays and strings by pointer - now you have to be explicit and call the right function for hashing a pointer, hashing an array contents or hashing a TCHAR* pointer as a null-terminated string.
Added GetArrayHash() for hashing arrays by their contents.
Added a new pair of keyfuncs for TSets and TMaps which allow TCHAR* keys (or TCHAR* values passed to *ByHash functions) to be hashed by FCrc::Strihash_DEPRECATED().

diff --git a/Engine/Plugins/Developer/PerforceSourceControl/Source/PerforceSourceControl/Private/PerforceConnection.h b/Engine/Plugins/Developer/PerforceSourceControl/Source/PerforceSourceControl/Private/PerforceConnection.h
--- a/Engine/Plugins/Developer/PerforceSourceControl/Source/PerforceSourceControl/Private/PerforceConnection.h
+++ b/Engine/Plugins/Developer/PerforceSourceControl/Source/PerforceSourceControl/Private/PerforceConnection.h
@@ -26,7 +26,7 @@ public:
 
 	const FString& operator()(const TCHAR* Key) const
 	{
-		if (const FString* Found = FindByHash(GetTypeHash(Key), Key))
+		if (const FString* Found = FindByHash(FCrc::Strihash_DEPRECATED(Key), Key))
 		{
 			return *Found;
 		}
diff --git a/Engine/Plugins/Editor/AssetManagerEditor/Source/AssetManagerEditor/Private/Insights/Table/ViewModels/TreeNodeGrouping.cpp b/Engine/Plugins/Editor/AssetManagerEditor/Source/AssetManagerEditor/Private/Insights/Table/ViewModels/TreeNodeGrouping.cpp
--- a/Engine/Plugins/Editor/AssetManagerEditor/Source/AssetManagerEditor/Private/Insights/Table/ViewModels/TreeNodeGrouping.cpp
+++ b/Engine/Plugins/Editor/AssetManagerEditor/Source/AssetManagerEditor/Private/Insights/Table/ViewModels/TreeNodeGrouping.cpp
@@ -168,7 +168,7 @@ void FTreeNodeGroupingByUniqueValueCString::GroupNodes(const TArray<FTableTreeNo
 {
 	FTableTreeNodePtr UnsetGroupPtr = nullptr; // for unset FTableCellValue
 	FTableTreeNodePtr EmptyGroupPtr = nullptr; // for nullptr and empty strings
-	TMap<const TCHAR*, FTableTreeNodePtr> GroupMap; // for valid strings
+	TMap<const TCHAR*, FTableTreeNodePtr, FDefaultSetAllocator, TStringPointerMapKeyFuncs_DEPRECATED<const TCHAR*, FTableTreeNodePtr>> GroupMap; // for valid strings
 	TMap<FName, FTableTreeNodePtr> GroupNameMap; // for valid strings
 
 	ParentGroup.ClearChildren();
diff --git a/Engine/Plugins/Editor/AssetManagerEditor/Source/AssetManagerEditor/Private/TreeView/AssetTable.cpp b/Engine/Plugins/Editor/AssetManagerEditor/Source/AssetManagerEditor/Private/TreeView/AssetTable.cpp
--- a/Engine/Plugins/Editor/AssetManagerEditor/Source/AssetManagerEditor/Private/TreeView/AssetTable.cpp
+++ b/Engine/Plugins/Editor/AssetManagerEditor/Source/AssetManagerEditor/Private/TreeView/AssetTable.cpp
@@ -1288,7 +1288,7 @@ void FAssetTable::AddDefaultColumns()
 	//////////////////////////////////////////////////
 }
 
-/*static*/TSet<int32> FAssetTableRow::GatherAllReachableNodes(const TArray<int32>& StartingNodes, const FAssetTable& OwningTable, const TSet<int32>& AdditionalNodesToStopAt, const TSet<const TCHAR*>& RestrictToGFPs)
+/*static*/TSet<int32> FAssetTableRow::GatherAllReachableNodes(const TArray<int32>& StartingNodes, const FAssetTable& OwningTable, const TSet<int32>& AdditionalNodesToStopAt, const TSet<const TCHAR*, TStringPointerSetKeyFuncs_DEPRECATED<const TCHAR*>>& RestrictToGFPs)
 {
 	// "visit" ThisIndex to seed the exploration
 	TSet<int32> VisitedIndices;
@@ -1331,7 +1331,7 @@ void FAssetTable::AddDefaultColumns()
 
 ////////////////////////////////////////////////////////////////////////////////////////////////////
 
-/*static*/ void FAssetTableRow::RefineDependencies(const TSet<int32>& PreviouslyVisitedIndices, const FAssetTable& OwningTable, const TSet<int32> RootIndices, const TSet<const TCHAR*>& RestrictToGFPs, TSet<int32>& OutUniqueIndices, TSet<int32>& OutSharedIndices)
+/*static*/ void FAssetTableRow::RefineDependencies(const TSet<int32>& PreviouslyVisitedIndices, const FAssetTable& OwningTable, const TSet<int32> RootIndices, const TSet<const TCHAR*, TStringPointerSetKeyFuncs_DEPRECATED<const TCHAR*>>& RestrictToGFPs, TSet<int32>& OutUniqueIndices, TSet<int32>& OutSharedIndices)
 {
 	// "visit" ThisIndex to seed the exploration
 	TSet<int32> IndicesToVisit;
@@ -1401,7 +1401,7 @@ void FAssetTable::AddDefaultColumns()
 
 	FAssetTableDependencySizes Result;
 
-	TSet<const TCHAR*> RestrictToGFPs;
+	TSet<const TCHAR*, TStringPointerSetKeyFuncs_DEPRECATED<const TCHAR*>> RestrictToGFPs;
 	for (int32 RootIndex : RootIndices)
 	{
 		RestrictToGFPs.Add(OwningTable.GetAssetChecked(RootIndex).GetGameFeaturePlugin());
@@ -1481,8 +1481,8 @@ void FAssetTable::AddDefaultColumns()
 	UE::Insights::FStopwatch Stopwatch;
 	Stopwatch.Start();
 
-	TSet<int32> AllReachableNodes = GatherAllReachableNodes(StartingNodes.Array(), OwningTable, TSet<int32>{}, TSet<const TCHAR*>{});
-	TSet<const TCHAR*> SourceGFPs;
+	TSet<int32> AllReachableNodes = GatherAllReachableNodes(StartingNodes.Array(), OwningTable, TSet<int32>{}, TSet<const TCHAR*, TStringPointerSetKeyFuncs_DEPRECATED<const TCHAR*>>{});
+	TSet<const TCHAR*, TStringPointerSetKeyFuncs_DEPRECATED<const TCHAR*>> SourceGFPs;
 	for (int32 RootIndex : StartingNodes)
 	{
 		SourceGFPs.Add(OwningTable.GetAsset(RootIndex)->GetGameFeaturePlugin());
diff --git a/Engine/Plugins/Editor/AssetManagerEditor/Source/AssetManagerEditor/Private/TreeView/AssetTable.h b/Engine/Plugins/Editor/AssetManagerEditor/Source/AssetManagerEditor/Private/TreeView/AssetTable.h
--- a/Engine/Plugins/Editor/AssetManagerEditor/Source/AssetManagerEditor/Private/TreeView/AssetTable.h
+++ b/Engine/Plugins/Editor/AssetManagerEditor/Source/AssetManagerEditor/Private/TreeView/AssetTable.h
@@ -171,11 +171,11 @@ public:
 	static int64 ComputeTotalSizeExternalDependencies(const FAssetTable& OwningTable, const TSet<int32>& StartingNodes, TSet<int32>* OutExternalDependencies = nullptr);
 
 private:
-	static void RefineDependencies(const TSet<int32>& PreviouslyVisitedIndices, const FAssetTable& OwningTable, const TSet<int32> RootIndices, const TSet<const TCHAR*>& RestrictToGFPs, TSet<int32>& OutUniqueIndices, TSet<int32>& OutSharedIndices);
-	static inline bool ShouldSkipDueToGFP(const TSet<const TCHAR*>& RestrictToGFPs, const TCHAR* GFP) { return (RestrictToGFPs.Num() > 0) && !RestrictToGFPs.Contains(GFP); }
+	static void RefineDependencies(const TSet<int32>& PreviouslyVisitedIndices, const FAssetTable& OwningTable, const TSet<int32> RootIndices, const TSet<const TCHAR*, TStringPointerSetKeyFuncs_DEPRECATED<const TCHAR*>>& RestrictToGFPs, TSet<int32>& OutUniqueIndices, TSet<int32>& OutSharedIndices);
+	static inline bool ShouldSkipDueToGFP(const TSet<const TCHAR*, TStringPointerSetKeyFuncs_DEPRECATED<const TCHAR*>>& RestrictToGFPs, const TCHAR* GFP) { return (RestrictToGFPs.Num() > 0) && !RestrictToGFPs.Contains(GFP); }
 
 	// Finds all nodes reachable from StartingNodes INCLUDING those in StartingNodes
-	static TSet<int32> GatherAllReachableNodes(const TArray<int32>& StartingNodes, const FAssetTable& OwningTable, const TSet<int32>& AdditionalNodesToStopAt, const TSet<const TCHAR*>& RestrictToGFPs);
+	static TSet<int32> GatherAllReachableNodes(const TArray<int32>& StartingNodes, const FAssetTable& OwningTable, const TSet<int32>& AdditionalNodesToStopAt, const TSet<const TCHAR*, TStringPointerSetKeyFuncs_DEPRECATED<const TCHAR*>>& RestrictToGFPs);
 
 private:
 	const TCHAR* Type = nullptr;
diff --git a/Engine/Plugins/Enterprise/DatasmithCADImporter/Source/CADTools/Private/CADData.cpp b/Engine/Plugins/Enterprise/DatasmithCADImporter/Source/CADTools/Private/CADData.cpp
--- a/Engine/Plugins/Enterprise/DatasmithCADImporter/Source/CADTools/Private/CADData.cpp
+++ b/Engine/Plugins/Enterprise/DatasmithCADImporter/Source/CADTools/Private/CADData.cpp
@@ -34,7 +34,7 @@ FMaterialUId BuildMaterialUId(const FCADMaterial& Material)
 	uint32 MaterialUId = 0;
 	if (!Material.MaterialName.IsEmpty())
 	{
-		MaterialUId = GetTypeHash(*Material.MaterialName); // we add material name because it could be used by the end user so two material with same parameters but different name are different.
+		MaterialUId = FCrc::Strihash_DEPRECATED(*Material.MaterialName); // we add material name because it could be used by the end user so two material with same parameters but different name are different.
 	}
 
 	MaterialUId = HashCombine(MaterialUId, GetTypeHash(Material.Diffuse));
@@ -46,7 +46,7 @@ FMaterialUId BuildMaterialUId(const FCADMaterial& Material)
 
 	if (!Material.TextureName.IsEmpty())
 	{
-		MaterialUId = HashCombine(MaterialUId, GetTypeHash(*Material.TextureName));
+		MaterialUId = HashCombine(MaterialUId, FCrc::Strihash_DEPRECATED(*Material.TextureName));
 	}
 	return FMath::Abs((int32)MaterialUId);
 }
@@ -194,8 +194,8 @@ uint32 GetTypeHash(const FFileDescriptor& FileDescriptor)
 	using ::GetTypeHash;
 	FFileStatData FileStatData = IFileManager::Get().GetStatData(*FileDescriptor.SourceFilePath);
 
-	uint32 DescriptorHash = GetTypeHash(*FileDescriptor.Name);
-	DescriptorHash = HashCombine(DescriptorHash, GetTypeHash(*FileDescriptor.Configuration));
+	uint32 DescriptorHash = FCrc::Strihash_DEPRECATED(*FileDescriptor.Name);
+	DescriptorHash = HashCombine(DescriptorHash, FCrc::Strihash_DEPRECATED(*FileDescriptor.Configuration));
 	DescriptorHash = HashCombine(DescriptorHash, GetTypeHash(FileStatData.FileSize));
 	DescriptorHash = HashCombine(DescriptorHash, GetTypeHash(FileStatData.ModificationTime));
 
diff --git a/Engine/Plugins/Enterprise/DatasmithCADImporter/Source/DatasmithCADTranslator/Private/DatasmithSceneGraphBuilder.cpp b/Engine/Plugins/Enterprise/DatasmithCADImporter/Source/DatasmithCADTranslator/Private/DatasmithSceneGraphBuilder.cpp
--- a/Engine/Plugins/Enterprise/DatasmithCADImporter/Source/DatasmithCADTranslator/Private/DatasmithSceneGraphBuilder.cpp
+++ b/Engine/Plugins/Enterprise/DatasmithCADImporter/Source/DatasmithCADTranslator/Private/DatasmithSceneGraphBuilder.cpp
@@ -102,7 +102,7 @@ void GetNodeUuidAndLabel(CADLibrary::FArchiveInstance& Instance, CADLibrary::FAr
 		OutLabel = TEXT("NoName");
 	}
 
-	FCadUuid UeUuid = HashCombine(GetTypeHash(InParentUEUUID), GetTypeHash(Instance.Id));
+	FCadUuid UeUuid = HashCombine(FCrc::Strihash_DEPRECATED(InParentUEUUID), GetTypeHash(Instance.Id));
 	UeUuid = HashCombine(UeUuid, GetTypeHash(Reference.Id));
 
 	OutUEUUID = FString::Printf(TEXT("0x%08x"), UeUuid);
diff --git a/Engine/Plugins/Enterprise/DatasmithCADImporter/Source/DatasmithWireTranslator/Private/DatasmithWireTranslator.cpp b/Engine/Plugins/Enterprise/DatasmithCADImporter/Source/DatasmithWireTranslator/Private/DatasmithWireTranslator.cpp
--- a/Engine/Plugins/Enterprise/DatasmithCADImporter/Source/DatasmithWireTranslator/Private/DatasmithWireTranslator.cpp
+++ b/Engine/Plugins/Enterprise/DatasmithCADImporter/Source/DatasmithWireTranslator/Private/DatasmithWireTranslator.cpp
@@ -160,7 +160,7 @@ const FString AliasSdkVersion   = TEXT("2023.1");
 
 FColor CreateShaderColorFromShaderName(const FString& ShaderName)
 {
-	const uint32 ShaderHash = GetTypeHash(*ShaderName);
+	const uint32 ShaderHash = FCrc::Strihash_DEPRECATED(*ShaderName);
 	const uint32 Red = (ShaderHash & 0xff000000) >> 24;
 	const uint32 Green = (ShaderHash & 0x00ff0000) >> 16;
 	const uint32 Blue = (ShaderHash & 0x0000ff00) >> 8;
diff --git a/Engine/Plugins/Enterprise/DatasmithImporter/Source/DatasmithImporter/Private/DatasmithMaterialImporter.cpp b/Engine/Plugins/Enterprise/DatasmithImporter/Source/DatasmithImporter/Private/DatasmithMaterialImporter.cpp
--- a/Engine/Plugins/Enterprise/DatasmithImporter/Source/DatasmithImporter/Private/DatasmithMaterialImporter.cpp
+++ b/Engine/Plugins/Enterprise/DatasmithImporter/Source/DatasmithImporter/Private/DatasmithMaterialImporter.cpp
@@ -99,14 +99,14 @@ namespace DatasmithMaterialImporterUtils
 
 				if ( KeyValue )
 				{
-					Hash = HashCombine( Hash, GetTypeHash( KeyValue->GetName() ) );
+					Hash = HashCombine( Hash, FCrc::Strihash_DEPRECATED( KeyValue->GetName() ) );
 					Hash = HashCombine( Hash, GetTypeHash( KeyValue->GetPropertyType() ) );
 
 					// Only hash values if it's not the parameter
 					// Currently, if we're setting values on multiple properties, we're not sure which one is the parameter so we hash them all
 					if ( MaterialCDO && ( !MaterialCDO->HasAParameterName() || GenericExpression->GetPropertiesCount() > 1 ) )
 					{
-						Hash = HashCombine( Hash, GetTypeHash( KeyValue->GetValue() ) );
+						Hash = HashCombine( Hash, FCrc::Strihash_DEPRECATED( KeyValue->GetValue() ) );
 					}
 				}
 			}
@@ -116,7 +116,7 @@ namespace DatasmithMaterialImporterUtils
 			// Hash the path to the function as calling different functions should result in different hash values
 			IDatasmithMaterialExpressionFunctionCall* FunctionCallExpression = static_cast< IDatasmithMaterialExpressionFunctionCall* >( MaterialExpression );
 
-			Hash = HashCombine( Hash, GetTypeHash( FunctionCallExpression->GetFunctionPathName() ) );
+			Hash = HashCombine( Hash, FCrc::Strihash_DEPRECATED( FunctionCallExpression->GetFunctionPathName() ) );
 		}
 
 		for ( int32 InputIndex = 0; InputIndex < MaterialExpression->GetInputCount(); ++InputIndex )
diff --git a/Engine/Plugins/Experimental/AVCodecs/AMFCodecs/Source/AMFCodecs/Public/Video/Decoders/Configs/VideoDecoderConfigAMF.h b/Engine/Plugins/Experimental/AVCodecs/AMFCodecs/Source/AMFCodecs/Public/Video/Decoders/Configs/VideoDecoderConfigAMF.h
--- a/Engine/Plugins/Experimental/AVCodecs/AMFCodecs/Source/AMFCodecs/Public/Video/Decoders/Configs/VideoDecoderConfigAMF.h
+++ b/Engine/Plugins/Experimental/AVCodecs/AMFCodecs/Source/AMFCodecs/Public/Video/Decoders/Configs/VideoDecoderConfigAMF.h
@@ -10,7 +10,7 @@
 struct AMFCODECS_API FVideoDecoderConfigAMF : public FAVConfig
 {
 private:
-	TMap<wchar_t const*, amf::AMFVariantStruct> Variants;
+	TMap<wchar_t const*, amf::AMFVariantStruct, FDefaultSetAllocator, TStringPointerMapKeyFuncs_DEPRECATED<wchar_t const*, amf::AMFVariantStruct>> Variants;
 
 public:
 	struct FParsedPicture
diff --git a/Engine/Plugins/Experimental/AVCodecs/AMFCodecs/Source/AMFCodecs/Public/Video/Encoders/Configs/VideoEncoderConfigAMF.h b/Engine/Plugins/Experimental/AVCodecs/AMFCodecs/Source/AMFCodecs/Public/Video/Encoders/Configs/VideoEncoderConfigAMF.h
--- a/Engine/Plugins/Experimental/AVCodecs/AMFCodecs/Source/AMFCodecs/Public/Video/Encoders/Configs/VideoEncoderConfigAMF.h
+++ b/Engine/Plugins/Experimental/AVCodecs/AMFCodecs/Source/AMFCodecs/Public/Video/Encoders/Configs/VideoEncoderConfigAMF.h
@@ -10,7 +10,7 @@
 struct AMFCODECS_API FVideoEncoderConfigAMF : public FAVConfig
 {
 private:
-	TMap<wchar_t const*, amf::AMFVariantStruct> Variants;
+	TMap<wchar_t const*, amf::AMFVariantStruct, FDefaultSetAllocator, TStringPointerMapKeyFuncs_DEPRECATED<wchar_t const*, amf::AMFVariantStruct>> Variants;
 
 public:
 	static TAVResult<AMF_VIDEO_ENCODER_RATE_CONTROL_METHOD_ENUM> ConvertRateControlMode(ERateControlMode Mode);
diff --git a/Engine/Plugins/FX/Niagara/Source/Niagara/Classes/NiagaraDataInterfaceSkeletalMesh.h b/Engine/Plugins/FX/Niagara/Source/Niagara/Classes/NiagaraDataInterfaceSkeletalMesh.h
--- a/Engine/Plugins/FX/Niagara/Source/Niagara/Classes/NiagaraDataInterfaceSkeletalMesh.h
+++ b/Engine/Plugins/FX/Niagara/Source/Niagara/Classes/NiagaraDataInterfaceSkeletalMesh.h
@@ -300,7 +300,7 @@ public:
 
 	static TypeHash GetTypeHash()
 	{
-		static const TypeHash Hash = ::GetTypeHash(TEXT("FNDI_SkeletalMesh_GeneratedData"));
+		static const TypeHash Hash = FCrc::Strihash_DEPRECATED(TEXT("FNDI_SkeletalMesh_GeneratedData"));
 		return Hash;
 	}
 };
diff --git a/Engine/Plugins/FX/Niagara/Source/Niagara/Internal/DataInterface/NiagaraDataInterfaceStaticMesh.h b/Engine/Plugins/FX/Niagara/Source/Niagara/Internal/DataInterface/NiagaraDataInterfaceStaticMesh.h
--- a/Engine/Plugins/FX/Niagara/Source/Niagara/Internal/DataInterface/NiagaraDataInterfaceStaticMesh.h
+++ b/Engine/Plugins/FX/Niagara/Source/Niagara/Internal/DataInterface/NiagaraDataInterfaceStaticMesh.h
@@ -321,7 +321,7 @@ public:
 
 	static TypeHash GetTypeHash()
 	{
-		static const TypeHash Hash = ::GetTypeHash(TEXT("FNDI_StaticMesh_GeneratedData"));
+		static const TypeHash Hash = FCrc::Strihash_DEPRECATED(TEXT("FNDI_StaticMesh_GeneratedData"));
 		return Hash;
 	}
 };
diff --git a/Engine/Plugins/FX/Niagara/Source/Niagara/Private/NiagaraDataInterfaceLandscape.cpp b/Engine/Plugins/FX/Niagara/Source/Niagara/Private/NiagaraDataInterfaceLandscape.cpp
--- a/Engine/Plugins/FX/Niagara/Source/Niagara/Private/NiagaraDataInterfaceLandscape.cpp
+++ b/Engine/Plugins/FX/Niagara/Source/Niagara/Private/NiagaraDataInterfaceLandscape.cpp
@@ -950,7 +950,7 @@ void FNDI_Landscape_GeneratedData::Tick(ETickingGroup TickGroup, float DeltaSeco
 
 FNDI_GeneratedData::TypeHash FNDI_Landscape_GeneratedData::GetTypeHash()
 {
-	static const TypeHash Hash = ::GetTypeHash(TEXT("FNDI_Landscape_GeneratedData"));
+	static const TypeHash Hash = FCrc::Strihash_DEPRECATED(TEXT("FNDI_Landscape_GeneratedData"));
 	return Hash;
 }
 
diff --git a/Engine/Plugins/FX/Niagara/Source/Niagara/Private/NiagaraPerfBaseline.cpp b/Engine/Plugins/FX/Niagara/Source/Niagara/Private/NiagaraPerfBaseline.cpp
--- a/Engine/Plugins/FX/Niagara/Source/Niagara/Private/NiagaraPerfBaseline.cpp
+++ b/Engine/Plugins/FX/Niagara/Source/Niagara/Private/NiagaraPerfBaseline.cpp
@@ -378,7 +378,7 @@ void FNiagaraPerfBaselineStatsListener::TickRT()
 }
 
 //////////////////////////////////////////////////////////////////////////
-const int32 FParticlePerfStatsListener_NiagaraPerformanceReporter::TestDebugMessageID = GetTypeHash(TEXT("NiagaraPerfReporterMessageID"));
+const int32 FParticlePerfStatsListener_NiagaraPerformanceReporter::TestDebugMessageID = FCrc::Strihash_DEPRECATED(TEXT("NiagaraPerfReporterMessageID"));
 
 void FParticlePerfStatsListener_NiagaraPerformanceReporter::ReportToLog()
 {
@@ -1209,7 +1209,7 @@ void FNiagaraPerfBaselineHandler::Tick(UWorld* World, float DeltaSeconds)
 			{
 				if (BaselineGenerationState == EBaselineGenState::WaitingToGenerate)
 				{
-					static const int32 GenBaselineMessageID = GetTypeHash(TEXT("NiagaraGenBaselineMessageID"));
+					static const int32 GenBaselineMessageID = FCrc::Strihash_DEPRECATED(TEXT("NiagaraGenBaselineMessageID"));
 					if (World->GetTimeSeconds() >= WorldTimeToGenerate)
 					{
 						BaselineGenerationState = EBaselineGenState::Generating;
diff --git a/Engine/Plugins/FX/Niagara/Source/NiagaraShader/Private/NiagaraShader.cpp b/Engine/Plugins/FX/Niagara/Source/NiagaraShader/Private/NiagaraShader.cpp
--- a/Engine/Plugins/FX/Niagara/Source/NiagaraShader/Private/NiagaraShader.cpp
+++ b/Engine/Plugins/FX/Niagara/Source/NiagaraShader/Private/NiagaraShader.cpp
@@ -389,7 +389,7 @@ void FNiagaraShaderType::BeginCompileShader(
 	NewJobs.Add(FShaderCommonCompileJobPtr(NewJob));
 }
 
-void FNiagaraShaderType::CacheUniformBufferIncludes(TMap<const TCHAR*, FCachedUniformBufferDeclaration>& Cache, EShaderPlatform Platform) const
+void FNiagaraShaderType::CacheUniformBufferIncludes(TMap<const TCHAR*, FCachedUniformBufferDeclaration, FDefaultSetAllocator, TStringPointerMapKeyFuncs_DEPRECATED<const TCHAR*, FCachedUniformBufferDeclaration>>& Cache, EShaderPlatform Platform) const
 {
 PRAGMA_DISABLE_DEPRECATION_WARNINGS
 	::CacheUniformBufferIncludes(Cache, Platform);
diff --git a/Engine/Plugins/FX/Niagara/Source/NiagaraShader/Public/NiagaraShaderType.h b/Engine/Plugins/FX/Niagara/Source/NiagaraShader/Public/NiagaraShaderType.h
--- a/Engine/Plugins/FX/Niagara/Source/NiagaraShader/Public/NiagaraShaderType.h
+++ b/Engine/Plugins/FX/Niagara/Source/NiagaraShader/Public/NiagaraShaderType.h
@@ -165,7 +165,7 @@ public:
 	}
 
 	UE_DEPRECATED(5.2, "CacheUniformBufferIncludes should no longer be used.")
-	inline void CacheUniformBufferIncludes(TMap<const TCHAR*, FCachedUniformBufferDeclaration>& Cache, EShaderPlatform Platform) const;
+	inline void CacheUniformBufferIncludes(TMap<const TCHAR*, FCachedUniformBufferDeclaration, FDefaultSetAllocator, TStringPointerMapKeyFuncs_DEPRECATED<const TCHAR*, FCachedUniformBufferDeclaration>>& Cache, EShaderPlatform Platform) const;
 
 protected:
 
diff --git a/Engine/Plugins/Interchange/Runtime/Source/Import/Private/Gltf/InterchangeGltfTranslator.cpp b/Engine/Plugins/Interchange/Runtime/Source/Import/Private/Gltf/InterchangeGltfTranslator.cpp
--- a/Engine/Plugins/Interchange/Runtime/Source/Import/Private/Gltf/InterchangeGltfTranslator.cpp
+++ b/Engine/Plugins/Interchange/Runtime/Source/Import/Private/Gltf/InterchangeGltfTranslator.cpp
@@ -1868,7 +1868,7 @@ bool UInterchangeGLTFTranslator::GetVariantSetPayloadData(UE::Interchange::FVari
 
 	PayloadData.Variants.SetNum(GltfAsset.Variants.Num());
 
-	TMap<const TCHAR*, Interchange::FVariant*> VariantMap;
+	TMap<const TCHAR*, Interchange::FVariant*, FDefaultSetAllocator, TStringPointerMapKeyFuncs_DEPRECATED<const TCHAR*, Interchange::FVariant*>> VariantMap;
 	VariantMap.Reserve(GltfAsset.Variants.Num());
 
 	for (int32 VariantIndex = 0; VariantIndex < GltfAsset.Variants.Num(); ++VariantIndex)
diff --git a/Engine/Plugins/Online/OnlineServices/Source/OnlineServicesCommon/Public/Online/OnlineTypeInfo.h b/Engine/Plugins/Online/OnlineServices/Source/OnlineServicesCommon/Public/Online/OnlineTypeInfo.h
--- a/Engine/Plugins/Online/OnlineServices/Source/OnlineServicesCommon/Public/Online/OnlineTypeInfo.h
+++ b/Engine/Plugins/Online/OnlineServices/Source/OnlineServicesCommon/Public/Online/OnlineTypeInfo.h
@@ -28,7 +28,7 @@ struct FOnlineTypeName
 
 inline uint32 GetTypeHash(const FOnlineTypeName& TypeName)
 {
-	return ::GetTypeHash(TypeName.Name);
+	return FCrc::Strihash_DEPRECATED(TypeName.Name);
 }
 
 template <typename T>
diff --git a/Engine/Plugins/Runtime/ReplicationSystemTestPlugin/Source/Private/NetworkAutomationTest.cpp b/Engine/Plugins/Runtime/ReplicationSystemTestPlugin/Source/Private/NetworkAutomationTest.cpp
--- a/Engine/Plugins/Runtime/ReplicationSystemTestPlugin/Source/Private/NetworkAutomationTest.cpp
+++ b/Engine/Plugins/Runtime/ReplicationSystemTestPlugin/Source/Private/NetworkAutomationTest.cpp
@@ -20,8 +20,8 @@ public:
 private:
 	friend void PrintNetworkAutomationTestSummary();
 
-	TMap<const TCHAR*, FNetworkAutomationTestStats> TestsWithFailures;
-	TMap<const TCHAR*, FNetworkAutomationTestStats> TestsWithWarnings;
+	TMap<const TCHAR*, FNetworkAutomationTestStats, FDefaultSetAllocator, TStringPointerMapKeyFuncs_DEPRECATED<const TCHAR*, FNetworkAutomationTestStats>> TestsWithFailures;
+	TMap<const TCHAR*, FNetworkAutomationTestStats, FDefaultSetAllocator, TStringPointerMapKeyFuncs_DEPRECATED<const TCHAR*, FNetworkAutomationTestStats>> TestsWithWarnings;
 	SIZE_T TestCount;
 	SIZE_T SuccessCount;
 };
diff --git a/Engine/Source/Developer/TraceInsights/Private/Insights/ViewModels/ThreadTimingTrack.h b/Engine/Source/Developer/TraceInsights/Private/Insights/ViewModels/ThreadTimingTrack.h
--- a/Engine/Source/Developer/TraceInsights/Private/Insights/ViewModels/ThreadTimingTrack.h
+++ b/Engine/Source/Developer/TraceInsights/Private/Insights/ViewModels/ThreadTimingTrack.h
@@ -112,7 +112,7 @@ private:
 	TMap<uint32, TSharedPtr<FCpuTimingTrack>> CpuTracks;
 
 	/** Maps thread group name to thread group info. */
-	TMap<const TCHAR*, FThreadGroup> ThreadGroups;
+	TMap<const TCHAR*, FThreadGroup, FDefaultSetAllocator, TStringPointerMapKeyFuncs_DEPRECATED<const TCHAR*, FThreadGroup>> ThreadGroups;
 
 	uint64 TimingProfilerTimelineCount;
 	uint64 LoadTimeProfilerTimelineCount;
diff --git a/Engine/Source/Developer/TraceServices/Private/Analyzers/CpuProfilerTraceAnalysis.h b/Engine/Source/Developer/TraceServices/Private/Analyzers/CpuProfilerTraceAnalysis.h
--- a/Engine/Source/Developer/TraceServices/Private/Analyzers/CpuProfilerTraceAnalysis.h
+++ b/Engine/Source/Developer/TraceServices/Private/Analyzers/CpuProfilerTraceAnalysis.h
@@ -76,7 +76,7 @@ private:
 	IEditableThreadProvider& EditableThreadProvider;
 	TMap<uint32, FThreadState*> ThreadStatesMap;
 	TMap<uint32, uint32> SpecIdToTimerIdMap;
-	TMap<const TCHAR*, uint32> ScopeNameToTimerIdMap;
+	TMap<const TCHAR*, uint32, FDefaultSetAllocator, TStringPointerMapKeyFuncs_DEPRECATED<const TCHAR*, uint32>> ScopeNameToTimerIdMap;
 	uint32 CoroutineTimerId = ~0;
 	uint32 CoroutineUnknownTimerId = ~0;
 	uint64 TotalEventSize = 0;
diff --git a/Engine/Source/Developer/TraceServices/Private/Model/CsvProfilerProviderPrivate.h b/Engine/Source/Developer/TraceServices/Private/Model/CsvProfilerProviderPrivate.h
--- a/Engine/Source/Developer/TraceServices/Private/Model/CsvProfilerProviderPrivate.h
+++ b/Engine/Source/Developer/TraceServices/Private/Model/CsvProfilerProviderPrivate.h
@@ -144,7 +144,7 @@ private:
 		uint32 StartFrame = 0;
 		uint32 EndFrame = 0;
 		TArray<FStatSeries*> StatSeries;
-		TMap<const TCHAR*, const TCHAR*> Metadata;
+		TMap<const TCHAR*, const TCHAR*, FDefaultSetAllocator, TStringPointerMapKeyFuncs_DEPRECATED<const TCHAR*, const TCHAR*>> Metadata;
 	};
 
 	FStatSeriesValue& GetValueRef(uint64 SeriesHandle, uint32 FrameNumber);
@@ -152,7 +152,7 @@ private:
 	IAnalysisSession& Session;
 	TArray<FStatSeries*> StatSeries;
 	TMap<uint32, FEvents*> Events;
-	TMap<const TCHAR*, const TCHAR*> Metadata;
+	TMap<const TCHAR*, const TCHAR*, FDefaultSetAllocator, TStringPointerMapKeyFuncs_DEPRECATED<const TCHAR*, const TCHAR*>> Metadata;
 	TArray<FCapture*> Captures;
 	FCapture* CurrentCapture = nullptr;
 };
diff --git a/Engine/Source/Developer/TraceServices/Private/Model/Threads.cpp b/Engine/Source/Developer/TraceServices/Private/Model/Threads.cpp
--- a/Engine/Source/Developer/TraceServices/Private/Model/Threads.cpp
+++ b/Engine/Source/Developer/TraceServices/Private/Model/Threads.cpp
@@ -167,7 +167,7 @@ uint32 FThreadProvider::GetGroupSortOrder(const TCHAR* GroupName)
 	}
 	else
 	{
-		return GetTypeHash(GroupName);
+		return FCrc::Strihash_DEPRECATED(GroupName);
 	}
 }
 
diff --git a/Engine/Source/Developer/TraceServices/Public/TraceServices/Model/CsvProfilerProvider.h b/Engine/Source/Developer/TraceServices/Public/TraceServices/Model/CsvProfilerProvider.h
--- a/Engine/Source/Developer/TraceServices/Public/TraceServices/Model/CsvProfilerProvider.h
+++ b/Engine/Source/Developer/TraceServices/Public/TraceServices/Model/CsvProfilerProvider.h
@@ -15,7 +15,7 @@ class IUntypedTable;
 
 struct FCaptureInfo
 {
-	const TMap<const TCHAR*, const TCHAR*>& Metadata;
+	const TMap<const TCHAR*, const TCHAR*, FDefaultSetAllocator, TStringPointerMapKeyFuncs_DEPRECATED<const TCHAR*, const TCHAR*>>& Metadata;
 	const TCHAR* Filename = nullptr;
 	uint32 Id = uint32(-1);
 	uint32 FrameCount = 0;
diff --git a/Engine/Source/Runtime/Core/Private/Misc/RemoteConfigIni.cpp b/Engine/Source/Runtime/Core/Private/Misc/RemoteConfigIni.cpp
--- a/Engine/Source/Runtime/Core/Private/Misc/RemoteConfigIni.cpp
+++ b/Engine/Source/Runtime/Core/Private/Misc/RemoteConfigIni.cpp
@@ -302,7 +302,7 @@ bool FRemoteConfig::ShouldReadRemoteFile(const TCHAR* Filename)
 
 FRemoteConfigAsyncIOInfo* FRemoteConfig::FindConfig(const TCHAR* Filename)
 {
-	return ConfigBuffers.FindByHash(GetTypeHash(Filename), Filename);
+	return ConfigBuffers.FindByHash(FCrc::Strihash_DEPRECATED(Filename), Filename);
 }
 
 
diff --git a/Engine/Source/Runtime/Core/Private/Misc/TimeGuard.cpp b/Engine/Source/Runtime/Core/Private/Misc/TimeGuard.cpp
--- a/Engine/Source/Runtime/Core/Private/Misc/TimeGuard.cpp
+++ b/Engine/Source/Runtime/Core/Private/Misc/TimeGuard.cpp
@@ -7,11 +7,11 @@ DEFINE_LOG_CATEGORY_STATIC(LogTimeGuard, Log, All);
 
 #if DO_TIMEGUARD
 
-TMap<const TCHAR*, FTimeGuard::FGuardInfo>  FTimeGuard::HitchData;
+TMap<const TCHAR*, FTimeGuard::FGuardInfo, FDefaultSetAllocator, TStringPointerMapKeyFuncs_DEPRECATED<const TCHAR*, FGuardInfo>>  FTimeGuard::HitchData;
 bool FTimeGuard::bEnabled;
 float FTimeGuard::FrameTimeThresholdMS = 1000.0 / 30.0;
 FCriticalSection FTimeGuard::ReportMutex;
-TSet<const TCHAR *> FTimeGuard::VolatileNames; // any names which come in volatile we allocate them to a static string and put them in this array
+TSet<const TCHAR *, TStringPointerSetKeyFuncs_DEPRECATED<const TCHAR*>> FTimeGuard::VolatileNames; // any names which come in volatile we allocate them to a static string and put them in this array
 
 
 void FTimeGuard::SetEnabled(bool InEnable)
@@ -30,7 +30,7 @@ void FTimeGuard::ClearData()
 	HitchData.Empty();
 }
 
-void FTimeGuard::GetData(TMap<const TCHAR*, FGuardInfo>& Dest)
+void FTimeGuard::GetData(TMap<const TCHAR*, FGuardInfo, FDefaultSetAllocator, TStringPointerMapKeyFuncs_DEPRECATED<const TCHAR*, FGuardInfo>>& Dest)
 {
 	FScopeLock Lock(&ReportMutex);
 	Dest = HitchData;
diff --git a/Engine/Source/Runtime/Core/Private/ProfilingDebugging/Windows/VTuneExternalProfiler.cpp b/Engine/Source/Runtime/Core/Private/ProfilingDebugging/Windows/VTuneExternalProfiler.cpp
--- a/Engine/Source/Runtime/Core/Private/ProfilingDebugging/Windows/VTuneExternalProfiler.cpp
+++ b/Engine/Source/Runtime/Core/Private/ProfilingDebugging/Windows/VTuneExternalProfiler.cpp
@@ -78,7 +78,7 @@ public:
 
 	void StartScopedEvent(const struct FColor& Color, const TCHAR* Text) override
 	{
-		uint32 Hash = GetTypeHash(Text);
+		uint32 Hash = FCrc::Strihash_DEPRECATED(Text);
 
 		TMap<uint32, __itt_string_handle *>& VTuneHandleMap = FVTunePerThreadHandleMap::Get().VTuneHandleMap;
 
diff --git a/Engine/Source/Runtime/Core/Public/Containers/Set.h b/Engine/Source/Runtime/Core/Public/Containers/Set.h
--- a/Engine/Source/Runtime/Core/Public/Containers/Set.h
+++ b/Engine/Source/Runtime/Core/Public/Containers/Set.h
@@ -940,7 +940,10 @@ public:
 	template<typename ComparableKey>
 	FSetElementId FindIdByHash(uint32 KeyHash, const ComparableKey& Key) const
 	{
+		PRAGMA_DISABLE_DEPRECATION_WARNINGS
+		// Disable deprecations warnings to stop warnings being thrown by our check macro.
 		checkSlow(KeyHash == KeyFuncs::GetKeyHash(Key));
+		PRAGMA_ENABLE_DEPRECATION_WARNINGS
 
 		return FSetElementId(FindIndexByHash(KeyHash, Key));
 	}
@@ -1057,7 +1060,10 @@ public:
 	template<typename ComparableKey>
 	int32 RemoveByHash(uint32 KeyHash, const ComparableKey& Key)
 	{
+		PRAGMA_DISABLE_DEPRECATION_WARNINGS
+		// Disable deprecations warnings to stop warnings being thrown by our check macro.
 		checkSlow(KeyHash == KeyFuncs::GetKeyHash(Key));
+		PRAGMA_ENABLE_DEPRECATION_WARNINGS
 
 		if (Elements.Num())
 		{
@@ -1085,7 +1091,10 @@ public:
 	template<typename ComparableKey>
 	FORCEINLINE bool ContainsByHash(uint32 KeyHash, const ComparableKey& Key) const
 	{
+		PRAGMA_DISABLE_DEPRECATION_WARNINGS
+		// Disable deprecations warnings to stop warnings being thrown by our check macro.
 		checkSlow(KeyHash == KeyFuncs::GetKeyHash(Key));
+		PRAGMA_ENABLE_DEPRECATION_WARNINGS
 
 		return FindIndexByHash(KeyHash, Key) != INDEX_NONE;
 	}
@@ -2068,7 +2077,10 @@ public:
 
 			// We don't update the hash because we don't need to - the new element
 			// should have the same hash, but let's just check.
+			PRAGMA_DISABLE_DEPRECATION_WARNINGS
+			// Disable deprecations warnings to stop warnings being thrown by our check macro.
 			checkSlow(KeyHash == GetKeyHash(ElementPtr));
+			PRAGMA_ENABLE_DEPRECATION_WARNINGS
 		}
 		else
 		{
diff --git a/Engine/Source/Runtime/Core/Public/Containers/UnrealString.h b/Engine/Source/Runtime/Core/Public/Containers/UnrealString.h
--- a/Engine/Source/Runtime/Core/Public/Containers/UnrealString.h
+++ b/Engine/Source/Runtime/Core/Public/Containers/UnrealString.h
@@ -25,6 +25,7 @@
 #include "Templates/IsArrayOrRefOfTypeByPredicate.h"
 #include "Templates/TypeHash.h"
 #include "Templates/IsFloatingPoint.h"
+#include "Templates/UnrealTypeTraits.h"
 #include "Traits/IsCharType.h"
 #include "Traits/IsCharEncodingCompatibleWith.h"
 #include "Traits/IsCharEncodingSimplyConvertibleTo.h"
@@ -2365,10 +2366,6 @@ struct TTypeFromString
  *     Logf(TEXT("Value: %s"), ToCStr(LexToString(Val)));
  * }
  */
-FORCEINLINE const TCHAR* ToCStr(const TCHAR* Ptr)
-{
-	return Ptr;
-}
 FORCEINLINE const TCHAR* ToCStr(const FString& Str)
 {
 	return *Str;
diff --git a/Engine/Source/Runtime/Core/Public/Misc/Crc.h b/Engine/Source/Runtime/Core/Public/Misc/Crc.h
--- a/Engine/Source/Runtime/Core/Public/Misc/Crc.h
+++ b/Engine/Source/Runtime/Core/Public/Misc/Crc.h
@@ -8,6 +8,11 @@
 #include "Misc/CString.h"
 #include "Misc/Char.h"
 #include "Templates/EnableIf.h"
+#include "Templates/UnrealTypeTraits.h"
+#include "Traits/IsCharType.h"
+
+template <typename... Types>
+struct TTuple;
 
 /** 
  * CRC hash generation for different types of input data
@@ -241,3 +246,83 @@ inline uint32 FCrc::Strihash_DEPRECATED(const int32 DataLen, const UTF8CHAR* Dat
 
 	return Result;
 }
+
+/**
+ * Gets a non-owning TCHAR pointer from a string type.
+ *
+ * Can be used generically to get a const TCHAR*, when it is not known if the argument is a TCHAR* or an FString:
+ *
+ * template <typename T>
+ * void LogValue(const T& Val)
+ * {
+ *     Logf(TEXT("Value: %s"), ToCStr(LexToString(Val)));
+ * }
+ */
+FORCEINLINE const TCHAR* ToCStr(const TCHAR* Ptr)
+{
+	return Ptr;
+}
+
+/**
+ * An implementation of KeyFuncs for sets which hashes string pointers by FCrc::Strihash_DEPRECATED.
+ */
+template <typename InKeyType, bool bInAllowDuplicateKeys = false>
+struct TStringPointerSetKeyFuncs_DEPRECATED
+{
+	static_assert(TIsCharType<std::remove_pointer_t<decltype(ToCStr(std::declval<InKeyType>()))>>::Value, "TStringPointerSetKeyFuncs_DEPRECATED should only be used with keys which character types");
+
+	using KeyType         = InKeyType;
+	using KeyInitType     = typename TTypeTraits<InKeyType>::ConstPointerType;
+	using ElementInitType = typename TCallTraits<InKeyType>::ParamType;
+
+	static constexpr bool bAllowDuplicateKeys = bInAllowDuplicateKeys;
+
+	static FORCEINLINE KeyInitType GetSetKey(ElementInitType Element)
+	{
+		return Element;
+	}
+
+	template <typename ComparableKey>
+	static FORCEINLINE bool Matches(KeyInitType A, const ComparableKey& B)
+	{
+		return A == B;
+	}
+
+	template <typename ComparableKey = KeyInitType>
+	static FORCEINLINE uint32 GetKeyHash(ComparableKey Key)
+	{
+		return FCrc::Strihash_DEPRECATED(ToCStr(Key));
+	}
+};
+
+/**
+ * An implementation of KeyFuncs for maps which hashes string pointers by FCrc::Strihash_DEPRECATED.
+ */
+template <typename InKeyType, typename InValueType, bool bInAllowDuplicateKeys = false>
+struct TStringPointerMapKeyFuncs_DEPRECATED
+{
+	static_assert(TIsCharType<std::remove_pointer_t<decltype(ToCStr(std::declval<InKeyType>()))>>::Value, "TStringPointerMapKeyFuncs_DEPRECATED should only be used with keys which character types");
+
+	using KeyType         = InKeyType;
+	using KeyInitType     = typename TTypeTraits<InKeyType>::ConstPointerType;
+	using ElementInitType = const TTuple<typename TTypeTraits<InKeyType>::ConstInitType, typename TTypeTraits<InValueType>::ConstInitType>&;
+
+	static constexpr bool bAllowDuplicateKeys = bInAllowDuplicateKeys;
+
+	static FORCEINLINE KeyInitType GetSetKey(ElementInitType Element)
+	{
+		return Element.Key;
+	}
+
+	template <typename ComparableKey>
+	static FORCEINLINE bool Matches(KeyInitType A, const ComparableKey& B)
+	{
+		return A == B;
+	}
+
+	template <typename ComparableKey = KeyInitType>
+	static FORCEINLINE uint32 GetKeyHash(ComparableKey Key)
+	{
+		return FCrc::Strihash_DEPRECATED(ToCStr(Key));
+	}
+};
diff --git a/Engine/Source/Runtime/Core/Public/Misc/TimeGuard.h b/Engine/Source/Runtime/Core/Public/Misc/TimeGuard.h
--- a/Engine/Source/Runtime/Core/Public/Misc/TimeGuard.h
+++ b/Engine/Source/Runtime/Core/Public/Misc/TimeGuard.h
@@ -111,17 +111,17 @@ public:
 	static CORE_API void						ClearData();
 	static CORE_API void						SetEnabled(bool InEnable);
 	static CORE_API void						SetFrameTimeThresholdMS(float InTimeMS);
-	static CORE_API void						GetData(TMap<const TCHAR*, FGuardInfo>& Dest);
+	static CORE_API void						GetData(TMap<const TCHAR*, FGuardInfo, FDefaultSetAllocator, TStringPointerMapKeyFuncs_DEPRECATED<const TCHAR*, FGuardInfo>>& Dest);
 
 protected:
 
 	// Used for reporting
-	static CORE_API void						ReportHitch(const TCHAR* InName, float TimeMS, bool VolatileName );
-	static TMap<const TCHAR*, FGuardInfo>		HitchData;
-	static TSet<const TCHAR *>					VolatileNames; // any names which come in volatile we allocate them to a static string and put them in this array
-	static FCriticalSection						ReportMutex;
-	static CORE_API bool						bEnabled;
-	static CORE_API float						FrameTimeThresholdMS;
+	static CORE_API void																											ReportHitch(const TCHAR* InName, float TimeMS, bool VolatileName );
+	static TMap<const TCHAR*, FGuardInfo, FDefaultSetAllocator, TStringPointerMapKeyFuncs_DEPRECATED<const TCHAR*, FGuardInfo>>		HitchData;
+	static TSet<const TCHAR *, TStringPointerSetKeyFuncs_DEPRECATED<const TCHAR*>>													VolatileNames; // any names which come in volatile we allocate them to a static string and put them in this array
+	static FCriticalSection																											ReportMutex;
+	static CORE_API bool																											bEnabled;
+	static CORE_API float																											FrameTimeThresholdMS;
 
 protected:
 
diff --git a/Engine/Source/Runtime/Core/Public/Templates/TypeHash.h b/Engine/Source/Runtime/Core/Public/Templates/TypeHash.h
--- a/Engine/Source/Runtime/Core/Public/Templates/TypeHash.h
+++ b/Engine/Source/Runtime/Core/Public/Templates/TypeHash.h
@@ -89,7 +89,7 @@ inline uint32 PointerHash(const void* Key, uint32 C)
 
 template <
 	typename ScalarType,
-	std::enable_if_t<std::is_scalar_v<ScalarType>>* = nullptr
+	std::enable_if_t<std::is_scalar_v<ScalarType> && !std::is_same_v<ScalarType, TCHAR*> && !std::is_same_v<ScalarType, const TCHAR*>>* = nullptr
 >
 inline uint32 GetTypeHash(ScalarType Value)
 {
@@ -135,14 +135,10 @@ inline uint32 GetTypeHash(ScalarType Value)
 	}
 	else if constexpr (std::is_pointer_v<ScalarType>)
 	{
-		if constexpr (std::is_same_v<std::remove_cv_t<std::remove_pointer_t<ScalarType>>, TCHAR>)
-		{
-			return FCrc::Strihash_DEPRECATED(Value);
-		}
-		else
-		{
-			return PointerHash(Value);
-		}
+		// Once the TCHAR* deprecations below are removed, we want to prevent accidental string hashing, so this static_assert should be commented back in
+		//static_assert(!TIsCharType<std::remove_pointer_t<ScalarType>>::Value, "Pointers to string types should use a PointerHash() or FCrc::Stricmp_DEPRECATED() call depending on requirements");
+
+		return PointerHash(Value);
 	}
 	else
 	{
@@ -150,6 +146,43 @@ inline uint32 GetTypeHash(ScalarType Value)
 	}
 }
 
+template <
+	typename T,
+	uint32 N,
+	std::enable_if_t<!std::is_same_v<const T, const TCHAR>>* = nullptr
+>
+//UE_DEPRECATED(all, "Hashing arrays is deprecated - use PointerHash() instead to force a conversion to a pointer or GetArrayHash to hash the array contents")
+inline uint32 GetTypeHash(T (&Array)[N])
+{
+	return PointerHash(Array);
+}
+
+template <
+	typename T,
+	std::enable_if_t<std::is_same_v<const T, const TCHAR>>* = nullptr
+>
+//UE_DEPRECATED(5.3, "Hashing TCHAR arrays is deprecated - use PointerHash() to force a conversion to a pointer or FCrc::Strihash_DEPRECATED to do a string hash, or use TStringPointerSetKeyFuncs_DEPRECATED or TStringPointerMapKeyFuncs_DEPRECATED as keyfuncs for a TSet or TMap respectively")
+inline uint32 GetTypeHash(T* Value)
+{
+	// Hashing a TCHAR* array differently from a void* is dangerous and is deprecated.
+	// When removing these overloads post-deprecation, comment in the related static_assert in the std::is_pointer_v block of the GetTypeHash overload above.
+	return FCrc::Strihash_DEPRECATED(Value);
+}
+
+template <typename T>
+inline uint32 GetArrayHash(const T* Ptr, uint64 Size, uint32 PreviousHash = 0)
+{
+	uint32 Result = PreviousHash;
+	while (Size)
+	{
+		Result = HashCombineFast(Result, GetTypeHash(*Ptr));
+		++Ptr;
+		--Size;
+	}
+
+	return Result;
+}
+
 // Use this when inside type that has GetTypeHash() (no in-parameters) implemented. It makes GetTypeHash dispatch in global namespace
 template <typename T>
 FORCEINLINE uint32 GetTypeHashHelper(const T& V) { return GetTypeHash(V); }
diff --git a/Engine/Source/Runtime/Engine/Private/VT/VirtualTextureScalability.cpp b/Engine/Source/Runtime/Engine/Private/VT/VirtualTextureScalability.cpp
--- a/Engine/Source/Runtime/Engine/Private/VT/VirtualTextureScalability.cpp
+++ b/Engine/Source/Runtime/Engine/Private/VT/VirtualTextureScalability.cpp
@@ -295,11 +295,7 @@ namespace VirtualTextureScalability
 	
 	uint32 GetPhysicalPoolSettingsHash()
 	{
-		uint32 Hash = 0;
-		for (float PoolSizeScale : GPoolSizeScales)
-		{
-			Hash = HashCombine(Hash, GetTypeHash(PoolSizeScale));
-		}
+		uint32 Hash = GetArrayHash(GPoolSizeScales, UE_ARRAY_COUNT(GPoolSizeScales));
 		Hash = HashCombine(Hash, GetTypeHash(GetSplitPhysicalPoolSize()));
 		return Hash;
 	}
diff --git a/Engine/Source/Runtime/RenderCore/Private/ShaderCore.cpp b/Engine/Source/Runtime/RenderCore/Private/ShaderCore.cpp
--- a/Engine/Source/Runtime/RenderCore/Private/ShaderCore.cpp
+++ b/Engine/Source/Runtime/RenderCore/Private/ShaderCore.cpp
@@ -2170,7 +2170,7 @@ void GenerateReferencedUniformBufferNames(
 	const TCHAR* SourceFilename,
 	const TCHAR* ShaderTypeName,
 	const TMap<FString, TArray<const TCHAR*> >& ShaderFileToUniformBufferVariables,
-	TSet<const TCHAR*>& UniformBufferNames)
+	TSet<const TCHAR*, TStringPointerSetKeyFuncs_DEPRECATED<const TCHAR*>>& UniformBufferNames)
 {
 	TArray<FString> FilesToSearch;
 	GetShaderIncludes(SourceFilename, SourceFilename, FilesToSearch, GMaxRHIShaderPlatform);
@@ -2356,7 +2356,7 @@ void AppendKeyStringShaderDependencies(
 	FString& OutKeyString,
 	bool bIncludeSourceHashes)
 {
-	TSet<const TCHAR*> ReferencedUniformBufferNames;
+	TSet<const TCHAR*, TStringPointerSetKeyFuncs_DEPRECATED<const TCHAR*>> ReferencedUniformBufferNames;
 
 	for (const FShaderTypeDependency& ShaderTypeDependency : ShaderTypeDependencies)
 	{
diff --git a/Engine/Source/Runtime/RenderCore/Private/ShaderParameterMetadata.cpp b/Engine/Source/Runtime/RenderCore/Private/ShaderParameterMetadata.cpp
--- a/Engine/Source/Runtime/RenderCore/Private/ShaderParameterMetadata.cpp
+++ b/Engine/Source/Runtime/RenderCore/Private/ShaderParameterMetadata.cpp
@@ -782,7 +782,7 @@ void FShaderParametersMetadata::InitializeLayout(FRHIUniformBufferLayoutInitiali
 			MemberHash = HashCombine(MemberHash, GetTypeHash(int32(CurrentMember.GetOffset())));
 			MemberHash = HashCombine(MemberHash, GetTypeHash(uint8(BaseType)));
 			static_assert(EUniformBufferBaseType_NumBits <= 8, "Invalid EUniformBufferBaseType_NumBits");
-			MemberHash = HashCombine(MemberHash, GetTypeHash(CurrentMember.GetName()));
+			MemberHash = HashCombine(MemberHash, FCrc::Strihash_DEPRECATED(CurrentMember.GetName()));
 			MemberHash = HashCombine(MemberHash, GetTypeHash(int32(CurrentMember.GetNumElements())));
 
 			const bool bIsRHIResource = (
@@ -809,7 +809,7 @@ void FShaderParametersMetadata::InitializeLayout(FRHIUniformBufferLayoutInitiali
 			}
 			else if (bIsRHIResource || bIsRDGResource)
 			{
-				MemberHash = HashCombine(MemberHash, GetTypeHash(CurrentMember.GetShaderType()));
+				MemberHash = HashCombine(MemberHash, FCrc::Strihash_DEPRECATED(CurrentMember.GetShaderType()));
 			}
 
 			RootStructureHash = HashCombine(RootStructureHash, MemberHash);
diff --git a/Engine/Source/Runtime/RenderCore/Private/ShaderParameters.cpp b/Engine/Source/Runtime/RenderCore/Private/ShaderParameters.cpp
--- a/Engine/Source/Runtime/RenderCore/Private/ShaderParameters.cpp
+++ b/Engine/Source/Runtime/RenderCore/Private/ShaderParameters.cpp
@@ -373,9 +373,9 @@ static FShaderParametersMetadata* FindShaderParametersMetadataWithVariableName(c
 }
 
 /* deprecated */
-void CacheUniformBufferIncludes(TMap<const TCHAR*, FCachedUniformBufferDeclaration>& Cache, EShaderPlatform Platform)
+void CacheUniformBufferIncludes(TMap<const TCHAR*, FCachedUniformBufferDeclaration, FDefaultSetAllocator, TStringPointerMapKeyFuncs_DEPRECATED<const TCHAR*, FCachedUniformBufferDeclaration>>& Cache, EShaderPlatform Platform)
 {
-	for (TMap<const TCHAR*, FCachedUniformBufferDeclaration>::TIterator It(Cache); It; ++It)
+	for (TMap<const TCHAR*, FCachedUniformBufferDeclaration, FDefaultSetAllocator, TStringPointerMapKeyFuncs_DEPRECATED<const TCHAR*, FCachedUniformBufferDeclaration>>::TIterator It(Cache); It; ++It)
 	{
 		const TCHAR* UniformBufferName = It.Key();
 		FCachedUniformBufferDeclaration& BufferDeclaration = It.Value();
@@ -390,7 +390,7 @@ void CacheUniformBufferIncludes(TMap<const TCHAR*, FCachedUniformBufferDeclarati
 	}
 }
 
-void UE::ShaderParameters::AddUniformBufferIncludesToEnvironment(FShaderCompilerEnvironment& OutEnvironment, const TSet<const TCHAR*>& InUniformBufferNames)
+void UE::ShaderParameters::AddUniformBufferIncludesToEnvironment(FShaderCompilerEnvironment& OutEnvironment, const TSet<const TCHAR*, TStringPointerSetKeyFuncs_DEPRECATED<const TCHAR*>>& InUniformBufferNames)
 {
 	TRACE_CPUPROFILER_EVENT_SCOPE(UE::ShaderParameters::AddUniformBufferIncludesToEnvironment);
 
diff --git a/Engine/Source/Runtime/RenderCore/Public/Shader.h b/Engine/Source/Runtime/RenderCore/Public/Shader.h
--- a/Engine/Source/Runtime/RenderCore/Public/Shader.h
+++ b/Engine/Source/Runtime/RenderCore/Public/Shader.h
@@ -1431,7 +1431,10 @@ public:
 	}
 
 #if WITH_EDITOR
-	inline const TSet<const TCHAR*>& GetReferencedUniformBufferNames() const { return ReferencedUniformBufferNames; };
+	inline const TSet<const TCHAR*, TStringPointerSetKeyFuncs_DEPRECATED<const TCHAR*>>& GetReferencedUniformBufferNames() const
+	{
+		return ReferencedUniformBufferNames;
+	}
 
 	/** Adds include statements for uniform buffers that this shader type references. */
 	void AddUniformBufferIncludesToEnvironment(FShaderCompilerEnvironment& OutEnvironment, EShaderPlatform Platform) const;
@@ -1489,7 +1492,7 @@ protected:
 	* These are derived from source files so they need to be flushed when editing and recompiling shaders on the fly.
 	* FShaderType::Initialize will add the referenced uniform buffers, but this set may be updated by FlushShaderFileCache.
 	*/
-	TSet<const TCHAR*> ReferencedUniformBufferNames;
+	TSet<const TCHAR*, TStringPointerSetKeyFuncs_DEPRECATED<const TCHAR*>> ReferencedUniformBufferNames;
 #endif // WITH_EDITOR
 };
 
diff --git a/Engine/Source/Runtime/RenderCore/Public/ShaderCore.h b/Engine/Source/Runtime/RenderCore/Public/ShaderCore.h
--- a/Engine/Source/Runtime/RenderCore/Public/ShaderCore.h
+++ b/Engine/Source/Runtime/RenderCore/Public/ShaderCore.h
@@ -1240,7 +1240,7 @@ extern void GenerateReferencedUniformBufferNames(
 	const TCHAR* SourceFilename,
 	const TCHAR* ShaderTypeName,
 	const TMap<FString, TArray<const TCHAR*> >& ShaderFileToUniformBufferVariables,
-	TSet<const TCHAR*>& UniformBufferNames);
+	TSet<const TCHAR*, TStringPointerSetKeyFuncs_DEPRECATED<const TCHAR*>>& UniformBufferNames);
 
 struct FUniformBufferNameSortOrder
 {
diff --git a/Engine/Source/Runtime/RenderCore/Public/ShaderParameters.h b/Engine/Source/Runtime/RenderCore/Public/ShaderParameters.h
--- a/Engine/Source/Runtime/RenderCore/Public/ShaderParameters.h
+++ b/Engine/Source/Runtime/RenderCore/Public/ShaderParameters.h
@@ -37,7 +37,7 @@ namespace UE::ShaderParameters
 	/** Creates a shader code declaration of this struct for the given shader platform. */
 	RENDERCORE_API FString CreateUniformBufferShaderDeclaration(const TCHAR* Name, const FShaderParametersMetadata& UniformBufferStruct);
 
-	RENDERCORE_API void AddUniformBufferIncludesToEnvironment(FShaderCompilerEnvironment& OutEnvironment, const TSet<const TCHAR*>& InUniformBufferNames);
+	RENDERCORE_API void AddUniformBufferIncludesToEnvironment(FShaderCompilerEnvironment& OutEnvironment, const TSet<const TCHAR*, TStringPointerSetKeyFuncs_DEPRECATED<const TCHAR*>>& InUniformBufferNames);
 }
 
 UE_DEPRECATED(5.2, "CreateUniformBufferShaderDeclaration has moved to UE::ShaderParameters::CreateUniformBufferShaderDeclaration, does not take a EShaderPlatform argument and now returns the Declaration.")
@@ -47,7 +47,7 @@ inline void CreateUniformBufferShaderDeclaration(const TCHAR* Name, const FShade
 }
 
 UE_DEPRECATED(5.2, "CacheUniformBufferIncludes should no longer be used.")
-RENDERCORE_API void CacheUniformBufferIncludes(TMap<const TCHAR*, FCachedUniformBufferDeclaration>& Cache, EShaderPlatform Platform);
+RENDERCORE_API void CacheUniformBufferIncludes(TMap<const TCHAR*, FCachedUniformBufferDeclaration, FDefaultSetAllocator, TStringPointerMapKeyFuncs_DEPRECATED<const TCHAR*, FCachedUniformBufferDeclaration>>& Cache, EShaderPlatform Platform);
 #endif
 
 enum EShaderParameterFlags
diff --git a/Engine/Source/Runtime/RenderCore/Public/VertexFactory.h b/Engine/Source/Runtime/RenderCore/Public/VertexFactory.h
--- a/Engine/Source/Runtime/RenderCore/Public/VertexFactory.h
+++ b/Engine/Source/Runtime/RenderCore/Public/VertexFactory.h
@@ -441,7 +441,7 @@ public:
 	UE_DEPRECATED(5.2, "FlushShaderFileCache is deprecated. UpdateReferencedUniformBufferNames should be used to flush any uniform buffer changes")
 	RENDERCORE_API void FlushShaderFileCache(const TMap<FString, TArray<const TCHAR*> >& ShaderFileToUniformBufferVariables);
 
-	inline const TSet<const TCHAR*>& GetReferencedUniformBufferNames() const { return ReferencedUniformBufferNames; };
+	inline const TSet<const TCHAR*, TStringPointerSetKeyFuncs_DEPRECATED<const TCHAR*>>& GetReferencedUniformBufferNames() const { return ReferencedUniformBufferNames; };
 #endif // WITH_EDITOR
 
 private:
@@ -473,7 +473,7 @@ private:
 	 * These are derived from source files so they need to be flushed when editing and recompiling shaders on the fly. 
 	 * FShaderType::Initialize will add the referenced uniform buffers, but this set may be updated by FlushShaderFileCache
 	 */
-	TSet<const TCHAR*> ReferencedUniformBufferNames;
+	TSet<const TCHAR*, TStringPointerSetKeyFuncs_DEPRECATED<const TCHAR*>> ReferencedUniformBufferNames;
 #endif // WITH_EDITOR
 };
 
diff --git a/Engine/Source/Runtime/Sockets/Private/BSDSockets/IPAddressBSD.cpp b/Engine/Source/Runtime/Sockets/Private/BSDSockets/IPAddressBSD.cpp
--- a/Engine/Source/Runtime/Sockets/Private/BSDSockets/IPAddressBSD.cpp
+++ b/Engine/Source/Runtime/Sockets/Private/BSDSockets/IPAddressBSD.cpp
@@ -598,7 +598,7 @@ uint32 FInternetAddrBSD::GetTypeHash() const
 	}
 	else if (CurrentFamily == FNetworkProtocolTypes::IPv6)
 	{
-		return ::GetTypeHash(*ToString(true));
+		return FCrc::Strihash_DEPRECATED(*ToString(true));
 	}
 
 	return 0;
