Fixed alignment and thread safety issues in FMallocStomp and TConcurrentLinearAllocator

- FMallocStomp now respects requests for alignment larger than its page size.
- FMallocStomp now uses a hash as its sentinel to detect stomps of the allocation header. Previously these stomps could cause the wrong memory to be decommitted on free.
- TConcurrentLinearAllocator now uses GMalloc instead of BaseAllocator because BaseAllocator is not always thread-safe.
- TConcurrentLinearAllocator now passes its required alignment to its allocator instead of assuming it will receive 64 KiB aligned blocks.

diff --git a/Engine/Source/Runtime/Core/Private/Experimental/ConcurrentLinearAllocator.cpp b/Engine/Source/Runtime/Core/Private/Experimental/ConcurrentLinearAllocator.cpp
--- a/Engine/Source/Runtime/Core/Private/Experimental/ConcurrentLinearAllocator.cpp
+++ b/Engine/Source/Runtime/Core/Private/Experimental/ConcurrentLinearAllocator.cpp
@@ -2,6 +2,27 @@
 
 #include "Experimental/ConcurrentLinearAllocator.h"
 
+void* FOsAllocator::Malloc(SIZE_T Size, uint32 Alignment)
+{
+	if (UNLIKELY(!GMalloc))
+	{
+		// There is no public function to create GMalloc and this will do it for us.
+		FMemory::Free(FMemory::Malloc(0));
+		check(GMalloc);
+	}
+	void* Pointer = GMalloc->Malloc(Size, Alignment);
+	MemoryTrace_Alloc(uint64(Pointer), Size, Alignment, EMemoryTraceRootHeap::SystemMemory);
+	MemoryTrace_MarkAllocAsHeap(uint64(Pointer), EMemoryTraceRootHeap::SystemMemory);
+	return Pointer;
+}
+
+void FOsAllocator::Free(void* Pointer, SIZE_T Size)
+{
+	MemoryTrace_UnmarkAllocAsHeap(uint64(Pointer), EMemoryTraceRootHeap::SystemMemory);
+	MemoryTrace_Free(uint64(Pointer), EMemoryTraceRootHeap::SystemMemory);
+	GMalloc->Free(Pointer);
+}
+
 FORCENOINLINE void UE::Core::Private::OnInvalidConcurrentLinearArrayAllocatorNum(int32 NewNum, SIZE_T NumBytesPerElement)
 {
 	UE_LOG(LogCore, Fatal, TEXT("Trying to resize TConcurrentLinearArrayAllocator to an invalid size of %d with element size %" SIZE_T_FMT), NewNum, NumBytesPerElement);
diff --git a/Engine/Source/Runtime/Core/Private/HAL/MallocStomp.cpp b/Engine/Source/Runtime/Core/Private/HAL/MallocStomp.cpp
--- a/Engine/Source/Runtime/Core/Private/HAL/MallocStomp.cpp
+++ b/Engine/Source/Runtime/Core/Private/HAL/MallocStomp.cpp
@@ -1,10 +1,12 @@
 // Copyright Epic Games, Inc. All Rights Reserved.
 
 #include "HAL/MallocStomp.h"
-#include "Math/UnrealMathUtility.h"
-#include "HAL/UnrealMemory.h"
-#include "HAL/IConsoleManager.h"
+
 #include "GenericPlatform/GenericPlatformMemory.h"
+#include "HAL/IConsoleManager.h"
+#include "HAL/UnrealMemory.h"
+#include "Hash/xxhash.h"
+#include "Math/UnrealMathUtility.h"
 
 #if PLATFORM_UNIX
 	#include <sys/mman.h>
@@ -23,7 +25,7 @@
 #endif
 
 #if PLATFORM_64BITS
-	// 64-bit ABIs on x86_64 expect a 16-byte alignment
+// 64-bit ABIs on x86_64 expect a 16-byte alignment
 #define STOMPALIGNMENT 16U
 #else
 #define STOMPALIGNMENT 0U
@@ -46,6 +48,23 @@ FAutoConsoleCommand MallocStompTestCommand
 	FConsoleCommandDelegate::CreateStatic( &MallocStompOverrunTest )
 );
 
+struct FMallocStomp::FAllocationData
+{
+	/** Pointer to the full allocation. Needed so the OS knows what to free. */
+	void*	FullAllocationPointer;
+	/** Full size of the allocation including the extra page. */
+	SIZE_T	FullSize;
+	/** Size of the allocation requested. */
+	SIZE_T	Size;
+	/** Sentinel used to check for underrun. */
+	SIZE_T	Sentinel;
+
+	/** Calculate the expected sentinel value for this allocation data. */
+	SIZE_T CalculateSentinel() const
+	{
+		return (SIZE_T)FXxHash64::HashBuffer(this, offsetof(FAllocationData, Sentinel)).Hash;
+	}
+};
 
 FMallocStomp::FMallocStomp(const bool InUseUnderrunMode) 
 	: PageSize(FPlatformMemory::GetConstants().PageSize)
@@ -77,12 +96,15 @@ void* FMallocStomp::TryMalloc(SIZE_T Size, uint32 Alignment)
 	Alignment = FMath::Max<uint32>(Alignment, STOMPALIGNMENT);
 #endif
 
-	const SIZE_T AlignedSize = (Alignment > 0U) ? ((Size + Alignment - 1U) & -static_cast<int32>(Alignment)) : Size;
-	const SIZE_T AllocFullPageSize = AlignedSize + sizeof(FAllocationData) + (PageSize - 1) & ~(PageSize - 1U);
+	constexpr static SIZE_T AllocationDataSize = sizeof(FAllocationData);
+
+	const SIZE_T AlignedSize = Alignment ? ((Size + Alignment - 1) & -(int32)Alignment) : Size;
+	const SIZE_T AlignmentSize = Alignment > PageSize ? Alignment - PageSize : 0;
+	const SIZE_T AllocFullPageSize = (AlignedSize + AlignmentSize + AllocationDataSize + PageSize - 1) & -(SSIZE_T)PageSize;
 	const SIZE_T TotalAllocationSize = AllocFullPageSize + PageSize;
 
 #if PLATFORM_UNIX || PLATFORM_MAC
-	void *FullAllocationPointer = mmap(nullptr, TotalAllocationSize, PROT_READ | PROT_WRITE, MAP_PRIVATE | MAP_ANON, -1, 0);
+	void* FullAllocationPointer = mmap(nullptr, TotalAllocationSize, PROT_READ | PROT_WRITE, MAP_PRIVATE | MAP_ANON, -1, 0);
 #elif PLATFORM_WINDOWS && MALLOC_STOMP_KEEP_VIRTUAL_MEMORY
 	// Allocate virtual address space from current block using linear allocation strategy.
 	// If there is not enough space, try to allocate new block from OS. Report OOM if block allocation fails.
@@ -107,7 +129,7 @@ void* FMallocStomp::TryMalloc(SIZE_T Size, uint32 Alignment)
 	VirtualAddressCursor += TotalAllocationSize;
 
 #else
-	void *FullAllocationPointer = FPlatformMemory::BinnedAllocFromOS(TotalAllocationSize);
+	void* FullAllocationPointer = FPlatformMemory::BinnedAllocFromOS(TotalAllocationSize);
 #endif // PLATFORM_UNIX || PLATFORM_MAC
 
 	if (!FullAllocationPointer)
@@ -116,51 +138,62 @@ void* FMallocStomp::TryMalloc(SIZE_T Size, uint32 Alignment)
 	}
 
 	void* ReturnedPointer = nullptr;
-	static const SIZE_T AllocationDataSize = sizeof(FAllocationData);
 
-	const FAllocationData AllocData = { FullAllocationPointer, TotalAllocationSize, AlignedSize, SentinelExpectedValue };
+	checkSlow(IsAligned(FullAllocationPointer, PageSize));
 
-	if(bUseUnderrunMode)
+	if (bUseUnderrunMode)
 	{
-		const SIZE_T AlignedAllocationData = (Alignment > 0U) ? ((AllocationDataSize + Alignment - 1U) & -static_cast<int32>(Alignment)) : AllocationDataSize;
-		ReturnedPointer = reinterpret_cast<void*>(reinterpret_cast<uint8*>(FullAllocationPointer) + PageSize + AlignedAllocationData);
-		void* AllocDataPointerStart = reinterpret_cast<FAllocationData*>(reinterpret_cast<uint8*>(FullAllocationPointer) + PageSize);
+		ReturnedPointer = Align((uint8*)FullAllocationPointer + PageSize + AllocationDataSize, Alignment);
+		void* AllocDataPointerStart = static_cast<FAllocationData*>(ReturnedPointer) - 1;
+		checkSlow(AllocDataPointerStart >= FullAllocationPointer);
 
 #if PLATFORM_WINDOWS && MALLOC_STOMP_KEEP_VIRTUAL_MEMORY
 		// Commit physical pages to the used range, leaving the first page unmapped.
-		void* CommittedMemory = VirtualAlloc(AllocDataPointerStart, AllocFullPageSize, MEM_COMMIT, PAGE_READWRITE);
+		void* CommittedMemory = VirtualAlloc(AllocDataPointerStart, AllocationDataSize + AlignedSize, MEM_COMMIT, PAGE_READWRITE);
 		if (!CommittedMemory)
 		{
-			// Failed to allocate and commit physical memory pages. 
+			// Failed to allocate and commit physical memory pages.
 			return nullptr;
 		}
-		check(CommittedMemory == AllocDataPointerStart);
+		check(CommittedMemory == AlignDown(AllocDataPointerStart, PageSize));
 #else
-		// Page protect the first page, this will cause the exception in case the is an underrun.
-		FPlatformMemory::PageProtect(FullAllocationPointer, PageSize, false, false);
+		// Page protect the first page, this will cause the exception in case there is an underrun.
+		FPlatformMemory::PageProtect((uint8*)AlignDown(AllocDataPointerStart, PageSize) - PageSize, PageSize, false, false);
 #endif
 	} //-V773
 	else
 	{
-		ReturnedPointer = reinterpret_cast<void*>(reinterpret_cast<uint8*>(FullAllocationPointer) + AllocFullPageSize - AlignedSize);
+		ReturnedPointer = AlignDown((uint8*)FullAllocationPointer + AllocFullPageSize - AlignedSize, Alignment);
+		void* ReturnedPointerEnd = (uint8*)ReturnedPointer + AlignedSize;
+		checkSlow(IsAligned(ReturnedPointerEnd, PageSize));
+
+		void* AllocDataPointerStart = static_cast<FAllocationData*>(ReturnedPointer) - 1;
+		checkSlow(AllocDataPointerStart >= FullAllocationPointer);
 
 #if PLATFORM_WINDOWS && MALLOC_STOMP_KEEP_VIRTUAL_MEMORY
 		// Commit physical pages to the used range, leaving the last page unmapped.
-		void* CommittedMemory = VirtualAlloc(FullAllocationPointer, AllocFullPageSize, MEM_COMMIT, PAGE_READWRITE);
+		void* CommitPointerStart = AlignDown(AllocDataPointerStart, PageSize);
+		void* CommittedMemory = VirtualAlloc(CommitPointerStart, SIZE_T((uint8*)ReturnedPointerEnd - (uint8*)CommitPointerStart), MEM_COMMIT, PAGE_READWRITE);
 		if (!CommittedMemory)
 		{
-			// Failed to allocate and commit physical memory pages
+			// Failed to allocate and commit physical memory pages.
 			return nullptr;
 		}
-		check(CommittedMemory == FullAllocationPointer);
+		check(CommittedMemory == CommitPointerStart);
 #else
-		// Page protect the last page, this will cause the exception in case the is an overrun.
-		FPlatformMemory::PageProtect(reinterpret_cast<void*>(reinterpret_cast<uint8*>(FullAllocationPointer) + AllocFullPageSize), PageSize, false, false);
+		// Page protect the last page, this will cause the exception in case there is an overrun.
+		FPlatformMemory::PageProtect(ReturnedPointerEnd, PageSize, false, false);
 #endif
 	} //-V773
 
-	FAllocationData* AllocDataPointer = reinterpret_cast<FAllocationData*>(reinterpret_cast<uint8*>(ReturnedPointer) - AllocationDataSize);
-	*AllocDataPointer = AllocData;
+	checkSlow(IsAligned(FullAllocationPointer, PageSize));
+	checkSlow(IsAligned(TotalAllocationSize, PageSize));
+	checkSlow(IsAligned(ReturnedPointer, Alignment));
+	checkSlow((uint8*)ReturnedPointer + AlignedSize <= (uint8*)FullAllocationPointer + TotalAllocationSize);
+
+	FAllocationData& AllocationData = static_cast<FAllocationData*>(ReturnedPointer)[-1];
+	AllocationData = { FullAllocationPointer, TotalAllocationSize, AlignedSize, 0 };
+	AllocationData.Sentinel = AllocationData.CalculateSentinel();
 
 	return ReturnedPointer;
 }
@@ -216,8 +249,8 @@ void FMallocStomp::Free(void* InPtr)
 	FAllocationData *AllocDataPtr = reinterpret_cast<FAllocationData*>(InPtr);
 	AllocDataPtr--;
 
-	// Check that our sentinel is intact.
-	if(AllocDataPtr->Sentinel != SentinelExpectedValue)
+	// Check the sentinel to verify that the allocation data is intact.
+	if (AllocDataPtr->Sentinel != AllocDataPtr->CalculateSentinel())
 	{
 		// There was a memory underrun related to this allocation.
 		UE_DEBUG_BREAK();
diff --git a/Engine/Source/Runtime/Core/Public/Experimental/ConcurrentLinearAllocator.h b/Engine/Source/Runtime/Core/Public/Experimental/ConcurrentLinearAllocator.h
--- a/Engine/Source/Runtime/Core/Public/Experimental/ConcurrentLinearAllocator.h
+++ b/Engine/Source/Runtime/Core/Public/Experimental/ConcurrentLinearAllocator.h
@@ -50,20 +50,8 @@ struct FOsAllocator
 	// should set this variable to 'false'.
 	static constexpr bool UsesFMalloc = true;
 
-	FORCEINLINE static void* Malloc(SIZE_T Size, uint32 Alignment)
-	{
-		void* Pointer = FPlatformMemory::BaseAllocator()->Malloc(Size, DEFAULT_ALIGNMENT);
-		MemoryTrace_Alloc(uint64(Pointer), Size, DEFAULT_ALIGNMENT, EMemoryTraceRootHeap::SystemMemory);
-		MemoryTrace_MarkAllocAsHeap(uint64(Pointer), EMemoryTraceRootHeap::SystemMemory);
-		return Pointer;
-	}
-
-	FORCEINLINE static void Free(void* Pointer, SIZE_T Size)
-	{
-		MemoryTrace_UnmarkAllocAsHeap(uint64(Pointer), EMemoryTraceRootHeap::SystemMemory);
-		MemoryTrace_Free(uint64(Pointer), EMemoryTraceRootHeap::SystemMemory);
-		FPlatformMemory::BaseAllocator()->Free(Pointer);
-	}
+	CORE_API static void* Malloc(SIZE_T Size, uint32 Alignment);
+	CORE_API static void Free(void* Pointer, SIZE_T Size);
 };
 
 #if PLATFORM_USES_ANSI_MALLOC
@@ -299,7 +287,7 @@ class TConcurrentLinearAllocator
 		{
 			static_assert(BlockAllocationTag::BlockSize >= sizeof(FBlockHeader) + sizeof(FAllocationHeader));
 			Header = new (BlockAllocationTag::Allocator::Malloc(BlockAllocationTag::BlockSize, BlockAllocationTag::BlockSize)) FBlockHeader;
-			checkSlow(IsAligned(Header, alignof(FBlockHeader)));
+			checkSlow(IsAligned(Header, BlockAllocationTag::BlockSize));
 			if constexpr (!SupportsFastPath)
 			{
 				ASAN_POISON_MEMORY_REGION( Header + 1, BlockAllocationTag::BlockSize - sizeof(FBlockHeader) );
@@ -338,7 +326,7 @@ public:
 			{
 				static_assert(BlockAllocationTag::BlockSize >= sizeof(FBlockHeader) + sizeof(FAllocationHeader));
 				Header = new (BlockAllocationTag::Allocator::Malloc(BlockAllocationTag::BlockSize, BlockAllocationTag::BlockSize)) FBlockHeader;
-				checkSlow(IsAligned(Header, alignof(FBlockHeader)));
+				checkSlow(IsAligned(Header, BlockAllocationTag::BlockSize));
 				if constexpr (!SupportsFastPath)
 				{
 					ASAN_POISON_MEMORY_REGION( Header + 1, BlockAllocationTag::BlockSize - sizeof(FBlockHeader) );
diff --git a/Engine/Source/Runtime/Core/Public/HAL/MallocStomp.h b/Engine/Source/Runtime/Core/Public/HAL/MallocStomp.h
--- a/Engine/Source/Runtime/Core/Public/HAL/MallocStomp.h
+++ b/Engine/Source/Runtime/Core/Public/HAL/MallocStomp.h
@@ -27,29 +27,10 @@ class UWorld;
  */
 class FMallocStomp final : public FMalloc
 {
-private:
-#if PLATFORM_64BITS
-	/** Expected value to be found in the sentinel. */
-	static constexpr SIZE_T SentinelExpectedValue = 0xdeadbeefdeadbeef;
-#else
-	/** Expected value to be found in the sentinel. */
-	static constexpr SIZE_T SentinelExpectedValue = 0xdeadbeef;
-#endif
+	struct FAllocationData;
 
 	const SIZE_T PageSize;
 
-	struct FAllocationData
-	{
-		/** Pointer to the full allocation. Needed so the OS knows what to free. */
-		void	*FullAllocationPointer;
-		/** Full size of the allocation including the extra page. */
-		SIZE_T	FullSize;
-		/** Size of the allocation requested. */
-		SIZE_T	Size;
-		/** Sentinel used to check for underrun. */
-		SIZE_T	Sentinel;
-	};
-
 	/** If it is set to true, instead of focusing on overruns the allocator will focus on underruns. */
 	const bool bUseUnderrunMode;
 
@@ -59,7 +40,7 @@ private:
 
 public:
 	// FMalloc interface.
-	FMallocStomp(const bool InUseUnderrunMode = false);
+	explicit FMallocStomp(const bool InUseUnderrunMode = false);
 
 	/**
 	 * Allocates a block of a given number of bytes of memory with the required alignment.
