Added local ref to avoid compile issue

diff --git a/Engine/Source/Runtime/Core/Tests/Misc/ParseTests.cpp b/Engine/Source/Runtime/Core/Tests/Misc/ParseTests.cpp
--- a/Engine/Source/Runtime/Core/Tests/Misc/ParseTests.cpp
+++ b/Engine/Source/Runtime/Core/Tests/Misc/ParseTests.cpp
@@ -162,11 +162,15 @@ TEST_CASE("Parse::GrammaredCLIParse::Callback", "[Smoke]")
 				{ TEXT("trailing_space   "), { {TEXT("trailing_space"), nullptr} } },
 			}));
 		size_t CallbackCalledCount = 0;
-		auto CallBack = [&](FStringView Key, FStringView Value)
+		// NOTE: I'm making a Ref to a Structured binding var because CLANG
+		// has issues when trying to use them in lambdas.
+		// [REDACTED_PROFILE_URL]/llvm-project/issues/48582
+		const std::vector<StringKeyValue>& ExpectedResultsRef = ExpectedResults;
+		auto CallBack = [&CallbackCalledCount, &ExpectedResultsRef](FStringView Key, FStringView Value)
 		{
-			REQUIRE(CallbackCalledCount < ExpectedResults.size());
-			CHECK(Key == FStringView{ ExpectedResults[CallbackCalledCount].Key });
-			CHECK(Value == FStringView{ ExpectedResults[CallbackCalledCount].Value });
+			REQUIRE(CallbackCalledCount < ExpectedResultsRef.size());
+			CHECK(Key == FStringView{ ExpectedResultsRef[CallbackCalledCount].Key });
+			CHECK(Value == FStringView{ ExpectedResultsRef[CallbackCalledCount].Value });
 			++CallbackCalledCount;
 		};
 
@@ -183,11 +187,12 @@ TEST_CASE("Parse::GrammaredCLIParse::Callback", "[Smoke]")
 				{ TEXT("a \"-b --c\" d=e"), FParse::EGrammarBasedParseErrorCode::DisallowedQuotedCommand, 2, { {TEXT("a"), nullptr} } },
 			}));
 		size_t CallbackCalledCount = 0;
-		auto CallBack = [&](FStringView Key, FStringView Value)
+		const std::vector<StringKeyValue>& ExpectedResultsRef = ExpectedResults;
+		auto CallBack = [&CallbackCalledCount, &ExpectedResultsRef](FStringView Key, FStringView Value)
 		{
-			REQUIRE(CallbackCalledCount < ExpectedResults.size());
-			CHECK(Key == FStringView{ ExpectedResults[CallbackCalledCount].Key });
-			CHECK(Value == FStringView{ ExpectedResults[CallbackCalledCount].Value });
+			REQUIRE(CallbackCalledCount < ExpectedResultsRef.size());
+			CHECK(Key == FStringView{ ExpectedResultsRef[CallbackCalledCount].Key });
+			CHECK(Value == FStringView{ ExpectedResultsRef[CallbackCalledCount].Value });
 			++CallbackCalledCount;
 		};
 
@@ -207,11 +212,12 @@ TEST_CASE("Parse::GrammaredCLIParse::Callback", "[Smoke]")
 					{ TEXT("-a=\"unbalanced_quote_value"), FParse::EGrammarBasedParseErrorCode::UnBalancedQuote, 3, { } },
 			}));
 		size_t CallbackCalledCount = 0;
-		auto CallBack = [&](FStringView Key, FStringView Value)
+		const std::vector<StringKeyValue>& ExpectedResultsRef = ExpectedResults;
+		auto CallBack = [&CallbackCalledCount, &ExpectedResultsRef](FStringView Key, FStringView Value)
 		{
-			REQUIRE(CallbackCalledCount < ExpectedResults.size());
-			CHECK(Key == FStringView{ ExpectedResults[CallbackCalledCount].Key });
-			CHECK(Value == FStringView{ ExpectedResults[CallbackCalledCount].Value });
+			REQUIRE(CallbackCalledCount < ExpectedResultsRef.size());
+			CHECK(Key == FStringView{ ExpectedResultsRef[CallbackCalledCount].Key });
+			CHECK(Value == FStringView{ ExpectedResultsRef[CallbackCalledCount].Value });
 			++CallbackCalledCount;
 		};
 
@@ -224,4 +230,4 @@ TEST_CASE("Parse::GrammaredCLIParse::Callback", "[Smoke]")
 }
 
 #endif
-	
\ No newline at end of file
+	
