Thread idle scope tracking improvements

- option to check() on idle scope creation for a nested idle scope
  - enabled via DETECT_NESTED_THREAD_IDLE_STAT_SCOPES which is defined to DO_CHECK by default

- trace events added for thread idle scopes so they appear in captures
  - labeled as FThreadIdleStats::FScopeIdle
  - in new trace channel ThreadIdleScope

- testing
  - ran Lyra on Win64 and a NDA'd console platform
  - ran Lyra editor
  - ran replay on internal project on NDA'd console platform

diff --git a/Engine/Source/Runtime/Core/Private/Stats/Stats2.cpp b/Engine/Source/Runtime/Core/Private/Stats/Stats2.cpp
--- a/Engine/Source/Runtime/Core/Private/Stats/Stats2.cpp
+++ b/Engine/Source/Runtime/Core/Private/Stats/Stats2.cpp
@@ -64,6 +64,28 @@ DEFINE_STAT(STAT_FrameTime);
 DEFINE_STAT(STAT_NamedMarker);
 DEFINE_STAT(STAT_SecondsPerCycle);
 
+#if CPUPROFILERTRACE_ENABLED
+	UE_TRACE_CHANNEL_DEFINE(ThreadIdleScopeChannel);
+	TRACE_CPUPROFILER_EVENT_DECLARE(TraceEventScopeId);
+#endif
+
+CORE_API FThreadIdleStats::FScopeIdle::FScopeIdle(bool bInIgnore /* = false */)
+	: Start(FPlatformTime::Cycles())
+	, bIgnore(bInIgnore)
+#if CPUPROFILERTRACE_ENABLED
+	, TraceEventScope(TraceEventScopeId, TEXT("FThreadIdleStats::FScopeIdle"), ThreadIdleScopeChannel, !bInIgnore, __FILE__, __LINE__)
+#endif
+{
+#if DETECT_NESTED_THREAD_IDLE_STAT_SCOPES
+		if (!bInIgnore)
+		{
+			FThreadIdleStats& IdleStats = FThreadIdleStats::Get();
+			checkf(!IdleStats.bInIdleScope, TEXT("Nested FThreadIdleStats::FScopeIdle encountered"));
+			IdleStats.bInIdleScope = true;
+		}
+#endif
+}
+
 /*-----------------------------------------------------------------------------
 	DebugLeakTest, for the stats based memory profiler
 -----------------------------------------------------------------------------*/
diff --git a/Engine/Source/Runtime/Core/Public/Stats/Stats2.h b/Engine/Source/Runtime/Core/Public/Stats/Stats2.h
--- a/Engine/Source/Runtime/Core/Public/Stats/Stats2.h
+++ b/Engine/Source/Runtime/Core/Public/Stats/Stats2.h
@@ -41,6 +41,10 @@
 #include "UObject/NameTypes.h"
 #include "UObject/UnrealNames.h"
 
+#ifndef DETECT_NESTED_THREAD_IDLE_STAT_SCOPES
+	#define DETECT_NESTED_THREAD_IDLE_STAT_SCOPES	DO_CHECK		// This will cause check()s to fire if thread idle scopes are nested and hence double counting idle time.
+#endif // DETECT_NESTED_THREAD_IDLE_STAT_SCOPES
+
 class FOutputDevice;
 class FScopeCycleCounter;
 class FThreadStats;
@@ -58,6 +62,9 @@ class CORE_API FThreadIdleStats : public TThreadSingleton<FThreadIdleStats>
 		: Waits(0)
 		, WaitsCriticalPath(0)
 		, IsCriticalPathCounter(1)
+#if DETECT_NESTED_THREAD_IDLE_STAT_SCOPES
+		, bInIdleScope(false)
+#endif
 	{}
 
 public:
@@ -70,6 +77,10 @@ public:
 
 	int IsCriticalPathCounter;
 
+#if DETECT_NESTED_THREAD_IDLE_STAT_SCOPES
+	bool bInIdleScope;
+#endif
+
 	static void BeginCriticalPath()
 	{
 		FThreadIdleStats::Get().IsCriticalPathCounter++;
@@ -121,11 +132,7 @@ public:
 		{
 		}
 #else
-		FScopeIdle( bool bInIgnore = false )
-			: Start(FPlatformTime::Cycles())
-			, bIgnore( bInIgnore )
-		{
-		}
+		CORE_API FScopeIdle(bool bInIgnore = false);
 
 		~FScopeIdle()
 		{
@@ -139,8 +146,16 @@ public:
 				{
 					IdleStats.WaitsCriticalPath += CyclesElapsed;
 				}
+
+#if DETECT_NESTED_THREAD_IDLE_STAT_SCOPES
+				IdleStats.bInIdleScope = false;
+#endif
 			}
 		}
+
+#if CPUPROFILERTRACE_ENABLED
+		FCpuProfilerTrace::FEventScope TraceEventScope;
+#endif
 #endif
 	};
 };
