Add basic arithmetic operators to TIntVector2

diff --git a/Engine/Source/Runtime/Core/Public/Math/IntVector.h b/Engine/Source/Runtime/Core/Public/Math/IntVector.h
--- a/Engine/Source/Runtime/Core/Public/Math/IntVector.h
+++ b/Engine/Source/Runtime/Core/Public/Math/IntVector.h
@@ -643,6 +643,156 @@ PRAGMA_ENABLE_DEPRECATION_WARNINGS
 		return X!=Other.X || Y!=Other.Y;
 	}
 
+	/**
+	 * Multiplies this vector with another vector, using component-wise multiplication.
+	 *
+	 * @param Other The point to multiply with.
+	 * @return Reference to this point after multiplication.
+	 */
+	TIntVector2& operator*=(const TIntVector2& Other)
+	{
+		X *= Other.X;
+		Y *= Other.Y;
+
+		return *this;
+	}
+
+	/**
+	 * Scales this point.
+	 *
+	 * @param Scale What to multiply the point by.
+	 * @return Reference to this point after multiplication.
+	 */
+	TIntVector2& operator*=(IntType Scale)
+	{
+		X *= Scale;
+		Y *= Scale;
+
+		return *this;
+	}
+
+	/**
+	 * Divides this point.
+	 *
+	 * @param Divisor What to divide the point by.
+	 * @return Reference to this point after division.
+	 */
+	TIntVector2& operator/=(IntType Divisor)
+	{
+		X /= Divisor;
+		Y /= Divisor;
+
+		return *this;
+	}
+
+	/**
+	 * Remainder of division of this point.
+	 *
+	 * @param Divisor What to divide the point by.
+	 * @return Reference to this point after remainder.
+	 */
+	TIntVector2& operator%=(IntType Divisor)
+	{
+		X %= Divisor;
+		Y %= Divisor;
+
+		return *this;
+	}
+
+	/**
+	 * Adds to this point.
+	 *
+	 * @param Other The point to add to this point.
+	 * @return Reference to this point after addition.
+	 */
+	TIntVector2& operator+=(const TIntVector2& Other)
+	{
+		X += Other.X;
+		Y += Other.Y;
+
+		return *this;
+	}
+
+	/**
+	 * Subtracts from this point.
+	 *
+	 * @param Other The point to subtract from this point.
+	 * @return Reference to this point after subtraction.
+	 */
+	TIntVector2& operator-=(const TIntVector2& Other)
+	{
+		X -= Other.X;
+		Y -= Other.Y;
+
+		return *this;
+	}
+
+	/**
+	 * Gets the result of component-wise multiplication of this point by another.
+	 *
+	 * @param Other The point to multiply with.
+	 * @return The result of multiplication.
+	 */
+	TIntVector2 operator*(const TIntVector2& Other) const
+	{
+		return TIntVector2(*this) *= Other;
+	}
+
+	/**
+	 * Gets the result of scaling on this point.
+	 *
+	 * @param Scale What to multiply the point by.
+	 * @return A new scaled int point.
+	 */
+	TIntVector2 operator*(IntType Scale) const
+	{
+		return TIntVector2(*this) *= Scale;
+	}
+
+	/**
+	 * Gets the result of division on this point.
+	 *
+	 * @param Divisor What to divide the point by.
+	 * @return A new divided int point.
+	 */
+	TIntVector2 operator/(IntType Divisor) const
+	{
+		return TIntVector2(*this) /= Divisor;
+	}
+
+	/**
+	 * Gets the remainder of division on this point.
+	 *
+	 * @param Divisor What to divide the point by.
+	 * @return A new remainder int point.
+	 */
+	TIntVector2 operator%(IntType Divisor) const
+	{
+		return TIntVector2(*this) %= Divisor;
+	}
+
+	/**
+	 * Gets the result of addition on this point.
+	 *
+	 * @param Other The other point to add to this.
+	 * @return A new combined int point.
+	 */
+	TIntVector2 operator+(const TIntVector2& Other) const
+	{
+		return TIntVector2(*this) += Other;
+	}
+
+	/**
+	 * Gets the result of subtraction from this point.
+	 *
+	 * @param Other The other point to subtract from this.
+	 * @return A new subtracted int point.
+	 */
+	TIntVector2 operator-(const TIntVector2& Other) const
+	{
+		return TIntVector2(*this) -= Other;
+	}
+
 	/**
 	 * Serializes the Vector2.
 	 *
