make sure that the platform has support for persistent downloads before attempting to query the directory.

#rnx

diff --git a/Engine/Source/Runtime/Core/Private/Misc/ConfigUtilities.cpp b/Engine/Source/Runtime/Core/Private/Misc/ConfigUtilities.cpp
--- a/Engine/Source/Runtime/Core/Private/Misc/ConfigUtilities.cpp
+++ b/Engine/Source/Runtime/Core/Private/Misc/ConfigUtilities.cpp
@@ -38,6 +38,11 @@ const TCHAR* ConvertValueFromHumanFriendlyValue( const TCHAR* Value )
 #if !UE_SERVER
 void LoadCVarsFromFileForNextBoot(TMap<FString, FString>& OutCVars)
 {
+	if (!FPaths::HasProjectPersistentDownloadDir())
+	{
+		return;
+	}
+
 	IFileManager& FileManager = IFileManager::Get();
 
 	const FString FullPath = FPaths::ProjectPersistentDownloadDir() / UE_HOTFIX_FOR_NEXT_BOOT_FILENAME;
@@ -73,6 +78,11 @@ void LoadCVarsFromFileForNextBoot(TMap<FString, FString>& OutCVars)
 void SaveCVarForNextBoot(const TCHAR* Key, const TCHAR* Value)
 {
 #if !UE_SERVER
+	if (!FPaths::HasProjectPersistentDownloadDir())
+	{
+		return;
+	}
+
 	TMap<FString, FString> CVarsToSave;
 
 	// Read from file, in case there are more than one cvar hotfix event in same run
