fixed "for licensees" compilation

diff --git a/Engine/Source/Runtime/Core/Public/UObject/ScriptDelegates.h b/Engine/Source/Runtime/Core/Public/UObject/ScriptDelegates.h
--- a/Engine/Source/Runtime/Core/Public/UObject/ScriptDelegates.h
+++ b/Engine/Source/Runtime/Core/Public/UObject/ScriptDelegates.h
@@ -63,8 +63,11 @@ namespace UE::Core::Private
 template <typename InThreadSafetyMode>
 class TScriptDelegate : public TDelegateAccessHandlerBase<typename UE::Core::Private::TScriptDelegateTraits<InThreadSafetyMode>::ThreadSafetyMode>
 {
+public:
 	using ThreadSafetyMode = typename UE::Core::Private::TScriptDelegateTraits<InThreadSafetyMode>::ThreadSafetyMode;
+	using WeakPtrType = typename UE::Core::Private::TScriptDelegateTraits<InThreadSafetyMode>::WeakPtrType;
 
+private:
 	template <typename>
 	friend class TScriptDelegate;
 
@@ -78,8 +81,6 @@ class TScriptDelegate : public TDelegateAccessHandlerBase<typename UE::Core::Pri
 	using Super::GetWriteAccessScope;
 
 public:
-	using WeakPtrType = typename UE::Core::Private::TScriptDelegateTraits<InThreadSafetyMode>::WeakPtrType;
-
 	/** Default constructor. */
 	TScriptDelegate() 
 		: Object( nullptr )
@@ -495,8 +496,8 @@ struct TIsZeroConstructType<TScriptDelegate<ThreadSafetyMode>>
 template <typename InThreadSafetyMode>
 class TMulticastScriptDelegate : public TDelegateAccessHandlerBase<typename UE::Core::Private::TScriptDelegateTraits<InThreadSafetyMode>::ThreadSafetyMode>
 {
-	using ThreadSafetyMode = typename UE::Core::Private::TScriptDelegateTraits<InThreadSafetyMode>::ThreadSafetyMode;
-	using Super = TDelegateAccessHandlerBase<ThreadSafetyMode>;
+private:
+	using Super = TDelegateAccessHandlerBase<InThreadSafetyMode>;
 	using typename Super::FReadAccessScope;
 	using Super::GetReadAccessScope;
 	using typename Super::FWriteAccessScope;
@@ -505,7 +506,7 @@ class TMulticastScriptDelegate : public TDelegateAccessHandlerBase<typename UE::
 	using UnicastDelegateType = TScriptDelegate<typename UE::Core::Private::TScriptDelegateTraits<InThreadSafetyMode>::UnicastThreadSafetyModeForMulticasts>;
 
 public:
-
+	using ThreadSafetyMode = typename UE::Core::Private::TScriptDelegateTraits<InThreadSafetyMode>::ThreadSafetyMode;
 	using InvocationListType = TArray<UnicastDelegateType>;
 
 	TMulticastScriptDelegate() = default;
