UE_LOG sanitizer: allow implicit "enum class" conversions

diff --git a/Engine/Source/Runtime/Core/Public/String/FormatStringSan.h b/Engine/Source/Runtime/Core/Public/String/FormatStringSan.h
--- a/Engine/Source/Runtime/Core/Public/String/FormatStringSan.h
+++ b/Engine/Source/Runtime/Core/Public/String/FormatStringSan.h
@@ -115,7 +115,7 @@ namespace UE::Core::Private
 		template <typename T>
 		inline constexpr bool bIsFloatOrDouble = std::is_same_v<float, T> || std::is_same_v<double, T>;
 		template <typename T>
-		inline constexpr bool bIsIntegralEnum = std::is_enum_v<T> && std::is_convertible_v<T,int>;
+		inline constexpr bool bIsIntegralEnum = std::is_enum_v<T>;
 
 		template <typename... Ts>
 		struct TFmtArgTypes {};
diff --git a/Engine/Source/Runtime/Core/Tests/String/FormatStringSanTest.cpp b/Engine/Source/Runtime/Core/Tests/String/FormatStringSanTest.cpp
--- a/Engine/Source/Runtime/Core/Tests/String/FormatStringSanTest.cpp
+++ b/Engine/Source/Runtime/Core/Tests/String/FormatStringSanTest.cpp
@@ -39,11 +39,6 @@ TEST_CASE_NAMED(FFormatStringValidatorTest, "Core::String::FormatStringSan", "[C
 		STATIC_CHECK(UE_CHECK_FORMAT_STRING_ERR(FormatStringSan::StatusInvalidFormatSpec, TEXT("%l^"), 42));
 		STATIC_CHECK(UE_CHECK_FORMAT_STRING_ERR(FormatStringSan::StatusInvalidFormatSpec, TEXT("%h^"), 42));
 		STATIC_CHECK(UE_CHECK_FORMAT_STRING_ERR(FormatStringSan::StatusIncompleteFormatSpecifierOrUnescapedPercent, TEXT("%-*"), 42));		
-		
-		{
-			enum class MyEnum { Value };
-			STATIC_CHECK(UE_CHECK_FORMAT_STRING_ERR(FormatStringSan::StatusDNeedsIntegerArg, TEXT("enum class %d value"), MyEnum::Value));
-		}
 	}
 
 	SECTION("Accepted Formatting")
@@ -93,6 +88,10 @@ TEST_CASE_NAMED(FFormatStringValidatorTest, "Core::String::FormatStringSan", "[C
 			enum MyIntegralEnum { MyIntegralEnumA };
 			STATIC_CHECK(UE_CHECK_FORMAT_STRING_ERR(FormatStringSan::StatusOk, TEXT("hello %d is an enum actually"), MyIntegralEnumA));
 		}
+		{
+			enum class MyEnum { Value };
+			STATIC_CHECK(UE_CHECK_FORMAT_STRING_ERR(FormatStringSan::StatusOk, TEXT("enum class %d value"), MyEnum::Value));
+		}
 	}
 }
 
