Disable internal trace store for all platforms.

Standalone servers have been submitted for all desktop platform. Disable hosted trace store everywhere and enable launching of the server.

diff --git a/Engine/Source/Developer/TraceAnalysis/Public/Trace/StoreService.h b/Engine/Source/Developer/TraceAnalysis/Public/Trace/StoreService.h
--- a/Engine/Source/Developer/TraceAnalysis/Public/Trace/StoreService.h
+++ b/Engine/Source/Developer/TraceAnalysis/Public/Trace/StoreService.h
@@ -7,11 +7,7 @@
 ////////////////////////////////////////////////////////////////////////////////
 /* Indicates if the external store is supported or if an instance should be
    self-hosted by something */
-#if PLATFORM_WINDOWS
-#	define WITH_TRACE_STORE 0
-#else
-#	define WITH_TRACE_STORE 1
-#endif
+#define WITH_TRACE_STORE 0
 
 namespace UE {
 namespace Trace {
diff --git a/Engine/Source/Runtime/Core/Private/ProfilingDebugging/TraceAuxiliary.cpp b/Engine/Source/Runtime/Core/Private/ProfilingDebugging/TraceAuxiliary.cpp
--- a/Engine/Source/Runtime/Core/Private/ProfilingDebugging/TraceAuxiliary.cpp
+++ b/Engine/Source/Runtime/Core/Private/ProfilingDebugging/TraceAuxiliary.cpp
@@ -10,6 +10,9 @@
 #	include <Windows.h>
 #	include "Windows/HideWindowsPlatformTypes.h"
 #endif
+#if PLATFORM_LINUX || PLATFORM_MAC
+# 	include <sys/wait.h>
+#endif
 
 #if !defined(WITH_UNREAL_TRACE_LAUNCH)
 #	define WITH_UNREAL_TRACE_LAUNCH (PLATFORM_DESKTOP && !UE_BUILD_SHIPPING && !IS_PROGRAM)
@@ -1176,9 +1179,6 @@ static void LaunchUnrealTraceInternal(const TCHAR* CommandLine)
 #if PLATFORM_UNIX || PLATFORM_MAC
 static void LaunchUnrealTraceInternal(const TCHAR* CommandLine)
 {
-	/* nop */
-
-#if 0
 	if (GUnrealTraceLaunched.load(std::memory_order_relaxed))
 	{
 		UE_LOG(LogCore, Log, TEXT("UnrealTraceServer: Trace store already started"));
@@ -1202,7 +1202,7 @@ static void LaunchUnrealTraceInternal(const TCHAR* CommandLine)
 	TAnsiStringBuilder<64> ForkArg;
 	ForkArg << "fork";
 
-	pid_t UtsPid = vfork();
+	pid_t UtsPid = fork();
 	if (UtsPid < 0)
 	{
 		UE_LOG(LogCore, Display, TEXT("UnrealTraceServer: Unable to fork (errno: %d)"), errno);
@@ -1238,7 +1238,6 @@ static void LaunchUnrealTraceInternal(const TCHAR* CommandLine)
 		UE_LOG(LogCore, Log, TEXT("UnrealTraceServer: Trace store launch successful"));
 		GUnrealTraceLaunched.fetch_add(1, std::memory_order_relaxed);
 	}
-#endif // 0
 }
 #endif // PLATFORM_UNIX/MAC
 #endif // WITH_UNREAL_TRACE_LAUNCH
