Disable nested thread idle scope detection

- GameThreadWaitForTask() can be called recursively

diff --git a/Engine/Source/Runtime/Core/Public/Stats/Stats2.h b/Engine/Source/Runtime/Core/Public/Stats/Stats2.h
--- a/Engine/Source/Runtime/Core/Public/Stats/Stats2.h
+++ b/Engine/Source/Runtime/Core/Public/Stats/Stats2.h
@@ -42,7 +42,7 @@
 #include "UObject/UnrealNames.h"
 
 #ifndef DETECT_NESTED_THREAD_IDLE_STAT_SCOPES
-	#define DETECT_NESTED_THREAD_IDLE_STAT_SCOPES	DO_CHECK		// This will cause check()s to fire if thread idle scopes are nested and hence double counting idle time.
+	#define DETECT_NESTED_THREAD_IDLE_STAT_SCOPES	0		// This will cause check()s to fire if thread idle scopes are nested and hence double counting idle time.
 #endif // DETECT_NESTED_THREAD_IDLE_STAT_SCOPES
 
 class FOutputDevice;
