task pipes: fixed a race on FTaskBase::TraceId reported by TSAN. Task was constructed and pushed into a pipe, then another thread was pushing another task and building a chain between them. `release` memory order was insufficient to order initialisation of TraceId and access to it for tracing the chain

diff --git a/Engine/Source/Runtime/Core/Private/Tasks/Pipe.cpp b/Engine/Source/Runtime/Core/Private/Tasks/Pipe.cpp
--- a/Engine/Source/Runtime/Core/Private/Tasks/Pipe.cpp
+++ b/Engine/Source/Runtime/Core/Private/Tasks/Pipe.cpp
@@ -10,7 +10,8 @@ namespace UE::Tasks
 	Private::FTaskBase* FPipe::PushIntoPipe(Private::FTaskBase& Task)
 	{
 		Task.AddRef(); // the pipe holds a ref to the last task, until it's replaced by the next task or cleared on completion
-		Private::FTaskBase* LastTask_Local = LastTask.exchange(&Task, std::memory_order_release);
+		Private::FTaskBase* LastTask_Local = LastTask.exchange(&Task, std::memory_order_acq_rel); // `acq_rel` to order task construction before
+		// its usage by a thread that replaces it as the last piped task
 		checkf(LastTask_Local != &Task, TEXT("Dependency cycle: adding itself as a prerequisite (or use after destruction)"));
 
 		if (LastTask_Local == nullptr)
