enabled delegates race detection. the preflight succeeded on what was the previous issue on cooking ShooterGame (https://horde.devtools.epicgames.com/job/646fc7cf8b0e69cfb2ea98a7?label=173&issue=309830)

diff --git a/Engine/Source/Runtime/Core/Public/Delegates/DelegateAccessHandler.h b/Engine/Source/Runtime/Core/Public/Delegates/DelegateAccessHandler.h
--- a/Engine/Source/Runtime/Core/Public/Delegates/DelegateAccessHandler.h
+++ b/Engine/Source/Runtime/Core/Public/Delegates/DelegateAccessHandler.h
@@ -7,7 +7,7 @@
 #include "Misc/MTAccessDetector.h"
 #include "Misc/ScopeLock.h"
 
-#define UE_DETECT_DELEGATES_RACE_CONDITIONS 0
+//#define UE_DETECT_DELEGATES_RACE_CONDITIONS 0
 
 #if !defined(UE_DETECT_DELEGATES_RACE_CONDITIONS)
 #define UE_DETECT_DELEGATES_RACE_CONDITIONS ENABLE_MT_DETECTOR
