iOS: guard usage of Apple system headers

We share PLATFORM_* macros with the Apple system headers, which will cause problems if we are checking for IOS/TVOS functionality.  As part of guarding the headers, I also re-organized the headers in this file.

Additionally, make sure we use a static class function correctly (IOSAppDelegate ConvertFromUIInterfaceOrientation).

#rnx

diff --git a/Engine/Source/Runtime/Core/Private/IOS/IOSPlatformMisc.cpp b/Engine/Source/Runtime/Core/Private/IOS/IOSPlatformMisc.cpp
--- a/Engine/Source/Runtime/Core/Private/IOS/IOSPlatformMisc.cpp
+++ b/Engine/Source/Runtime/Core/Private/IOS/IOSPlatformMisc.cpp
@@ -5,60 +5,55 @@
 =============================================================================*/
 
 #include "IOS/IOSPlatformMisc.h"
-#include "Misc/App.h"
+#include "Async/TaskGraphInterfaces.h"
+#include "Apple/ApplePlatformCrashContext.h"
 #include "HAL/ExceptionHandling.h"
-#include "Misc/SecureHash.h"
-#include "Misc/EngineVersion.h"
-#include "Templates/Function.h"
+#include "HAL/FileManager.h"
+#include "HAL/PlatformOutputDevices.h"
+#include "Internationalization/Culture.h"
+#include "Internationalization/Internationalization.h"
+#include "Internationalization/Regex.h"
+#include "IOSChunkInstaller.h"
 #include "IOS/IOSMallocZone.h"
 #include "IOS/IOSApplication.h"
 #include "IOS/IOSAppDelegate.h"
+#include "IOS/IOSPlatformCrashContext.h"
+#include "IOS/IOSPlatformPLCrashReporterIncludes.h"
 #include "IOS/IOSView.h"
-#include "IOSChunkInstaller.h"
+#include "Misc/App.h"
 #include "Misc/CommandLine.h"
 #include "Misc/ConfigCacheIni.h"
 #include "Misc/CoreDelegates.h"
-#include "Misc/ScopeExit.h"
-#include "Modules/ModuleManager.h"
-#include "Apple/ApplePlatformCrashContext.h"
-#include "IOS/IOSPlatformCrashContext.h"
-#include "IOS/IOSPlatformPLCrashReporterIncludes.h"
-#include "HAL/FileManager.h"
-#include "HAL/PlatformOutputDevices.h"
+#include "Misc/EngineVersion.h"
+#include "Misc/FeedbackContext.h"
 #include "Misc/OutputDeviceError.h"
 #include "Misc/OutputDeviceRedirector.h"
-#include "Misc/FeedbackContext.h"
-#include "Internationalization/Internationalization.h"
-#include "Internationalization/Culture.h"
-#include "Internationalization/Regex.h"
+#include "Misc/ScopeExit.h"
+#include "Misc/SecureHash.h"
+#include "Modules/ModuleManager.h"
+#include "Templates/Function.h"
 
+#include "Apple/PreAppleSystemHeaders.h"
+
 #if !PLATFORM_TVOS
 #include <AdSupport/ASIdentifierManager.h> 
-#endif
+#endif // !PLATFORM_TVOS
 
-#include "Async/TaskGraphInterfaces.h"
-#include <SystemConfiguration/SystemConfiguration.h>
-#include <netinet/in.h>
-
-#import <StoreKit/StoreKit.h>
 #import <DeviceCheck/DeviceCheck.h>
-
-#import <UserNotifications/UserNotifications.h>
-#include "Async/TaskGraphInterfaces.h"
-#include "Misc/CoreDelegates.h"
-
-
 #import <Foundation/Foundation.h>
+//#include <libproc.h>
 #import <mach-o/dyld.h>
+#include <netinet/in.h>
+#include <SystemConfiguration/SystemConfiguration.h>
+#import <StoreKit/StoreKit.h>
+#import <UserNotifications/UserNotifications.h>
 
+#include "Apple/PostAppleSystemHeaders.h"
+
 #if !defined ENABLE_ADVERTISING_IDENTIFIER
 	#define ENABLE_ADVERTISING_IDENTIFIER 0
 #endif
 
-//#include <libproc.h>
-// @pjs commented out to resolve issue with PLATFORM_TVOS being defined by mach-o loader
-//#include <mach-o/dyld.h>
-
 /** Amount of free memory in MB reported by the system at startup */
 CORE_API int32 GStartupFreeMemoryMB;
 
@@ -335,7 +330,7 @@
 		AppDelegate.InterfaceOrientation = [[[[[UIApplication sharedApplication] delegate] window] windowScene] interfaceOrientation];
 	}
 
-	return [AppDelegate ConvertFromUIInterfaceOrientation:AppDelegate.InterfaceOrientation];
+	return [IOSAppDelegate ConvertFromUIInterfaceOrientation:AppDelegate.InterfaceOrientation];
 #else
 	return EDeviceScreenOrientation::Unknown;
 #endif
