Fix UB in FClangPlatformMath::FloorLog2*

diff --git a/Engine/Source/Runtime/Core/Public/Clang/ClangPlatformMath.h b/Engine/Source/Runtime/Core/Public/Clang/ClangPlatformMath.h
--- a/Engine/Source/Runtime/Core/Public/Clang/ClangPlatformMath.h
+++ b/Engine/Source/Runtime/Core/Public/Clang/ClangPlatformMath.h
@@ -91,13 +91,11 @@ struct FClangPlatformMath : public FGenericPlatformMath
 
 	static FORCEINLINE uint32 FloorLog2(uint32 Value)
 	{
-		int32 Mask = -int32(Value != 0);
-		return (31 - __builtin_clz(Value)) & Mask;
+		return 31 - __builtin_clz(Value | 1);
 	}
 
 	static FORCEINLINE uint64 FloorLog2_64(uint64 Value)
 	{
-		int64 Mask = -int64(Value != 0);
-		return (63 - __builtin_clzll(Value)) & Mask;
+		return 63 - __builtin_clzll(Value | 1);
 	}
 };
