Use another SetBy flag for cvar which got hotfixed.

diff --git a/Engine/Source/Runtime/Core/Private/Misc/ConfigCacheIni.cpp b/Engine/Source/Runtime/Core/Private/Misc/ConfigCacheIni.cpp
--- a/Engine/Source/Runtime/Core/Private/Misc/ConfigCacheIni.cpp
+++ b/Engine/Source/Runtime/Core/Private/Misc/ConfigCacheIni.cpp
@@ -2082,7 +2082,7 @@ namespace
 	{
 		if (IniFilename == GEngineIni && SectionNames.Contains(TEXT("ConsoleVariables")))
 		{
-			FConfigCacheIni::LoadConsoleVariablesFromINI();
+			UE::ConfigUtilities::ApplyCVarSettingsFromIni(TEXT("ConsoleVariables"), *GEngineIni, ECVF_SetByHotfix);
 		}
 	}
 }
diff --git a/Engine/Source/Runtime/Core/Private/Misc/ConfigUtilities.cpp b/Engine/Source/Runtime/Core/Private/Misc/ConfigUtilities.cpp
--- a/Engine/Source/Runtime/Core/Private/Misc/ConfigUtilities.cpp
+++ b/Engine/Source/Runtime/Core/Private/Misc/ConfigUtilities.cpp
@@ -112,7 +112,7 @@ void ApplyCVarsFromBootHotfix()
 		IConsoleVariable* CVar = IConsoleManager::Get().FindConsoleVariable(*CVarPair.Key); 
 		if (CVar)
 		{
-			CVar->Set(*CVarPair.Value, ECVF_SetByConsoleVariablesIni);
+			CVar->Set(*CVarPair.Value, ECVF_SetByHotfix);
 		}
 	}
 #endif // !UE_SERVER
@@ -198,7 +198,7 @@ void OnSetCVarFromIniEntry(const TCHAR *IniFile, const TCHAR *Key, const TCHAR*
 #endif // !DISABLE_CHEAT_CVARS
 		}
 
-		if (CVar->TestFlags(ECVF_SaveForNextBoot) && (SetBy & ECVF_SetByConsoleVariablesIni))
+		if (CVar->TestFlags(ECVF_SaveForNextBoot) && (SetBy == ECVF_SetByHotfix))
 		{
 			SaveCVarForNextBoot(Key, Value);
 		}
diff --git a/Engine/Source/Runtime/Core/Public/HAL/IConsoleManager.h b/Engine/Source/Runtime/Core/Public/HAL/IConsoleManager.h
--- a/Engine/Source/Runtime/Core/Public/HAL/IConsoleManager.h
+++ b/Engine/Source/Runtime/Core/Public/HAL/IConsoleManager.h
@@ -113,7 +113,6 @@ enum EConsoleVariableFlags
 	 *  - WARNING: This exposes the config to end user who could change it through the local temp file in shipping client, it's your responsibility to make sure client can't "legally" cheat through this type of CVar;
 	 *  - For the first time get hotfixed, the value will change from default value during boot to the hotfixed value after the hotfix http request, so make sure caching it when boot if the logic doesn't support to change at runtime;
 	 *  - When there is a new released game version, the saved local temp file of last version could still be there. So it will still got read once. Keep the logic work well with the old value, or remove the CVar if the old value no longer supported;
-	 *  - Only set the default value in code, don't set default value in [ConsoleVariables] section in the Engine.ini, otherwise reading the Engine.ini shipped within the game will overwrite the hotfixed value. This constrain also means it won't support undo at runtime;
 	 *  - It will only work for game client, not for dedicated servers which will be dynamically allocated on different server instances.
 	 */
 	ECVF_SaveForNextBoot = 0x2000,
@@ -149,12 +148,14 @@ enum EConsoleVariableFlags
 	ECVF_SetByGameOverride =		0x06000000,
 	// Set by local consolevariables.ini, mostly used for testing multiple projects
 	ECVF_SetByConsoleVariablesIni = 0x07000000,
+	// Set by hotfix
+	ECVF_SetByHotfix =				0x08000000,
 	// Used by some command line parameters, others use the Console priority instead
-	ECVF_SetByCommandline =			0x08000000,
+	ECVF_SetByCommandline =			0x09000000,
 	// Used for high priority temporary debugging or operation modes 
-	ECVF_SetByCode =				0x09000000,
+	ECVF_SetByCode =				0x0A000000,
 	// Highest priority used via editor UI or or game/editor interactive console
-	ECVF_SetByConsole =				0x0A000000,
+	ECVF_SetByConsole =				0x0B000000,
 
 
 	// ------------------------------------------------
