FPathViews::SetExtension with tests, to mirror FPaths::SetExtension.

diff --git a/Engine/Source/Runtime/Core/Private/Misc/PathViews.cpp b/Engine/Source/Runtime/Core/Private/Misc/PathViews.cpp
--- a/Engine/Source/Runtime/Core/Private/Misc/PathViews.cpp
+++ b/Engine/Source/Runtime/Core/Private/Misc/PathViews.cpp
@@ -265,6 +265,35 @@ FString FPathViews::ChangeExtension(const FStringView& InPath, const FStringView
 	}
 }
 
+FString FPathViews::SetExtension(const FStringView& InPath, const FStringView& InNewExtension)
+{
+	int32 Pos = INDEX_NONE;
+	const TCHAR* PathEndPos = Algo::FindLastByPredicate(InPath, UE4PathViews_Private::IsSlashOrBackslashOrPeriod);
+	if (PathEndPos != nullptr && *PathEndPos == TEXT('.'))
+	{
+		Pos = UE_PTRDIFF_TO_INT32(PathEndPos - InPath.GetData());
+	}
+
+	const FStringView FileWithoutExtension = Pos == INDEX_NONE ? InPath : InPath.Left(Pos);
+
+	if (!InNewExtension.IsEmpty() && !InNewExtension.StartsWith(TEXT('.')))
+	{
+		// The new extension lacks a period so we need to add it ourselves.
+		FString Result(FileWithoutExtension, InNewExtension.Len() + 1);
+		Result += TEXT('.');
+		Result += InNewExtension;
+
+		return Result;
+	}
+	else
+	{
+		FString Result(FileWithoutExtension, InNewExtension.Len());
+		Result += InNewExtension;
+
+		return Result;
+	}
+}
+
 bool FPathViews::IsSeparator(TCHAR c)
 {
 	return UE4PathViews_Private::IsSlashOrBackslash(c);
diff --git a/Engine/Source/Runtime/Core/Public/Misc/PathViews.h b/Engine/Source/Runtime/Core/Public/Misc/PathViews.h
--- a/Engine/Source/Runtime/Core/Public/Misc/PathViews.h
+++ b/Engine/Source/Runtime/Core/Public/Misc/PathViews.h
@@ -187,6 +187,16 @@ public:
 	 * case a copy of InPath will be returned instead.
 	 */
 	static FString ChangeExtension(const FStringView& InPath, const FStringView& InNewExtension);
+	
+	/**
+	 * Sets the file extension of a filename.
+	 *
+	 * @param InPath A file path with or without a pre-existing extension.
+	 * @param InNewExtension The new extension to use (prefixing with a '.' is optional)
+	 *
+	 * @return The new file path complete with the new extension.
+	 */
+	static FString SetExtension(const FStringView& InPath, const FStringView& InNewExtension);
 
 	/** Return whether the given character is a path-separator character (/ or \) */
 	static bool IsSeparator(TCHAR c);
diff --git a/Engine/Source/Runtime/Core/Tests/Misc/PathViewsTest.cpp b/Engine/Source/Runtime/Core/Tests/Misc/PathViewsTest.cpp
--- a/Engine/Source/Runtime/Core/Tests/Misc/PathViewsTest.cpp
+++ b/Engine/Source/Runtime/Core/Tests/Misc/PathViewsTest.cpp
@@ -585,6 +585,36 @@ TEST_CASE_NAMED(FPathViewsChangeExtensionTest, "System::Core::Misc::PathViews::C
 	RunChangeExtensionTest(TEXT("C:/Folder/First.Last/file.tar.gz"), TEXT("gz2"), TEXT("C:/Folder/First.Last/file.tar.gz2"));
 }
 
+TEST_CASE_NAMED(FPathViewsSetExtensionTest, "System::Core::Misc::PathViews::SetExtension", "[ApplicationContextMask][SmokeFilter]")
+{
+	auto RunSetExtensionTest = [this](const TCHAR* InPath, const TCHAR* InNewExt, const TCHAR* InExpectedPath)
+	{
+		// Run test
+		const FString NewPath = FPathViews::SetExtension(InPath, InNewExt);
+		if (NewPath != InExpectedPath)
+		{
+			FAIL_CHECK(FString::Printf(TEXT("Path '%s' failed to set the extension (got '%s', expected '%s')."), InPath, *NewPath, InExpectedPath));
+		}
+	};
+
+	RunSetExtensionTest(nullptr, nullptr, TEXT(""));
+	RunSetExtensionTest(TEXT(""), TEXT(""), TEXT(""));
+	RunSetExtensionTest(TEXT(""), TEXT(".txt"), TEXT(".txt"));
+	RunSetExtensionTest(TEXT("file"), TEXT("log"), TEXT("file.log"));
+	RunSetExtensionTest(TEXT("file.txt"), TEXT("log"), TEXT("file.log"));
+	RunSetExtensionTest(TEXT("file.tar.gz"), TEXT("gz2"), TEXT("file.tar.gz2"));
+	RunSetExtensionTest(TEXT("file.txt"), TEXT(""), TEXT("file"));
+	RunSetExtensionTest(TEXT("C:/Folder/file"), TEXT("log"), TEXT("C:/Folder/file.log"));
+	RunSetExtensionTest(TEXT("C:/Folder/file.txt"), TEXT("log"), TEXT("C:/Folder/file.log"));
+	RunSetExtensionTest(TEXT("C:/Folder/file.tar.gz"), TEXT("gz2"), TEXT("C:/Folder/file.tar.gz2"));
+	RunSetExtensionTest(TEXT("C:/Folder/First.Last/file"), TEXT("log"), TEXT("C:/Folder/First.Last/file.log"));
+	RunSetExtensionTest(TEXT("C:/Folder/First.Last/file.txt"), TEXT("log"), TEXT("C:/Folder/First.Last/file.log"));
+	RunSetExtensionTest(TEXT("C:/Folder/First.Last/file.tar.gz"), TEXT("gz2"), TEXT("C:/Folder/First.Last/file.tar.gz2"));
+
+	return true;
+}
+
+
 TEST_CASE_NAMED(FPathViewsEqualsAndLessTest, "System::Core::Misc::PathViews::EqualsAndLess", "[ApplicationContextMask][SmokeFilter]")
 {
 	FPathViewsTest Instance = FPathViewsTest();
