Build fix for translated test.

diff --git a/Engine/Source/Runtime/Core/Tests/Misc/PathViewsTest.cpp b/Engine/Source/Runtime/Core/Tests/Misc/PathViewsTest.cpp
--- a/Engine/Source/Runtime/Core/Tests/Misc/PathViewsTest.cpp
+++ b/Engine/Source/Runtime/Core/Tests/Misc/PathViewsTest.cpp
@@ -587,7 +587,7 @@ TEST_CASE_NAMED(FPathViewsChangeExtensionTest, "System::Core::Misc::PathViews::C
 
 TEST_CASE_NAMED(FPathViewsSetExtensionTest, "System::Core::Misc::PathViews::SetExtension", "[ApplicationContextMask][SmokeFilter]")
 {
-	auto RunSetExtensionTest = [this](const TCHAR* InPath, const TCHAR* InNewExt, const TCHAR* InExpectedPath)
+	auto RunSetExtensionTest = [](const TCHAR* InPath, const TCHAR* InNewExt, const TCHAR* InExpectedPath)
 	{
 		// Run test
 		const FString NewPath = FPathViews::SetExtension(InPath, InNewExt);
@@ -610,8 +610,6 @@ TEST_CASE_NAMED(FPathViewsSetExtensionTest, "System::Core::Misc::PathViews::SetE
 	RunSetExtensionTest(TEXT("C:/Folder/First.Last/file"), TEXT("log"), TEXT("C:/Folder/First.Last/file.log"));
 	RunSetExtensionTest(TEXT("C:/Folder/First.Last/file.txt"), TEXT("log"), TEXT("C:/Folder/First.Last/file.log"));
 	RunSetExtensionTest(TEXT("C:/Folder/First.Last/file.tar.gz"), TEXT("gz2"), TEXT("C:/Folder/First.Last/file.tar.gz2"));
-
-	return true;
 }
 
 
