Change TraceAuxilliary internal management of tracing status and trace destination

* Removed internal state tracking connection status. This was causing many problems because connections could be severed inside TraceLog without TraceAuxiliary knowing about it causing a mismatch in state.
* Wrapped trace destination and connection type state in struct and added guards for thread safety.
* Deprecated method to access trace destination as char array.

diff --git a/Engine/Plugins/TraceUtilities/Source/EditorTraceUtilities/Private/InsightsStatusBar/SInsightsStatusBar.cpp b/Engine/Plugins/TraceUtilities/Source/EditorTraceUtilities/Private/InsightsStatusBar/SInsightsStatusBar.cpp
--- a/Engine/Plugins/TraceUtilities/Source/EditorTraceUtilities/Private/InsightsStatusBar/SInsightsStatusBar.cpp
+++ b/Engine/Plugins/TraceUtilities/Source/EditorTraceUtilities/Private/InsightsStatusBar/SInsightsStatusBar.cpp
@@ -60,7 +60,7 @@ public:
 
 	void DoTask(ENamedThreads::Type CurrentThread, const FGraphEventRef& MyCompletionGraphEvent)
 	{
-		FUnrealInsightsLauncher::Get()->TryOpenTraceFromDestination(FTraceAuxiliary::GetTraceDestination());
+		FUnrealInsightsLauncher::Get()->TryOpenTraceFromDestination(FTraceAuxiliary::GetTraceDestinationString());
 	}
 };
 
@@ -490,7 +490,7 @@ FText SInsightsStatusBarWidget::GetTitleToolTipText() const
 {
 	FTextBuilder DescBuilder;
 
-	const TCHAR* Dest = FTraceAuxiliary::GetTraceDestination();
+	const FString Dest = FTraceAuxiliary::GetTraceDestinationString();
 	
 	if (*Dest != 0)
 	{
@@ -572,7 +572,7 @@ void SInsightsStatusBarWidget::OpenLiveSession_OnClicked()
 
 void SInsightsStatusBarWidget::OpenLiveSession()
 {
-	FUnrealInsightsLauncher::Get()->TryOpenTraceFromDestination(FTraceAuxiliary::GetTraceDestination());
+	FUnrealInsightsLauncher::Get()->TryOpenTraceFromDestination(FTraceAuxiliary::GetTraceDestinationString());
 }
 
 void SInsightsStatusBarWidget::OpenProfilingDirectory_OnClicked()
@@ -738,7 +738,7 @@ void SInsightsStatusBarWidget::ToggleTracing_OnClicked()
 		bool bResult = StartTracing();
 		if (bResult)
 		{
-			FString TraceDestinationStr = FTraceAuxiliary::GetTraceDestination();
+			FString TraceDestinationStr = FTraceAuxiliary::GetTraceDestinationString();
 			if (TraceDestinationStr.IsEmpty())
 			{
 				TraceDestinationStr = TEXT("External Target");
diff --git a/Engine/Source/Runtime/Core/Private/ProfilingDebugging/TraceAuxiliary.cpp b/Engine/Source/Runtime/Core/Private/ProfilingDebugging/TraceAuxiliary.cpp
--- a/Engine/Source/Runtime/Core/Private/ProfilingDebugging/TraceAuxiliary.cpp
+++ b/Engine/Source/Runtime/Core/Private/ProfilingDebugging/TraceAuxiliary.cpp
@@ -4,6 +4,7 @@
 #include "ProfilingDebugging/StringsTrace.h"
 #include "Trace/Trace.h"
 #include "CoreGlobals.h"
+#include "Misc/ScopeRWLock.h"
 
 #if PLATFORM_WINDOWS
 #	include "Windows/AllowWindowsPlatformTypes.h"
@@ -73,18 +74,11 @@ DECLARE_MEMORY_STAT(TEXT("Important event cache used"), STAT_TraceCacheUsed, STA
 DECLARE_MEMORY_STAT(TEXT("Important event cache waste"), STAT_TraceCacheWaste, STATGROUP_Trace);
 DECLARE_MEMORY_STAT(TEXT("Sent"), STAT_TraceSent, STATGROUP_Trace);
 
-////////////////////////////////////////////////////////////////////////////////
-enum class ETraceConnectType
-{
-	Network,
-	File
-};
-
 ////////////////////////////////////////////////////////////////////////////////
 class FTraceAuxiliaryImpl
 {
 public:
-	const TCHAR* GetDest() const;
+	FString GetDest() const;
 	bool IsConnected() const;
 	bool IsConnected(FGuid& OutSessionGuid, FGuid& OutTraceGuid);
 	FTraceAuxiliary::EConnectionType GetConnectionType() const;
@@ -94,7 +88,7 @@ public:
 	bool HasCommandlineChannels() const { return !CommandlineChannels.IsEmpty(); }
 	void EnableChannels(const TCHAR* ChannelList);
 	void DisableChannels(const TCHAR* ChannelList);
-	bool Connect(ETraceConnectType Type, const TCHAR* Parameter, const FTraceAuxiliary::FLogCategoryAlias& LogCategory, uint16 SendFlags);
+	bool Connect(FTraceAuxiliary::EConnectionType Type, const TCHAR* Parameter, const FTraceAuxiliary::FLogCategoryAlias& LogCategory, uint16 SendFlags);
 	bool Stop();
 	void ResumeChannels();
 	void PauseChannels();
@@ -135,12 +129,16 @@ private:
 
 	typedef TMap<uint32, FChannelEntry, TInlineSetAllocator<128>> ChannelSet;
 	ChannelSet CommandlineChannels;
-	FString TraceDest;
-	std::atomic<FTraceAuxiliary::EConnectionType> TraceType = FTraceAuxiliary::EConnectionType::None;
-	std::atomic<EState> State = EState::Stopped;
-	bool bTruncateFile = false;
 	bool bWorkerThreadStarted = false;
+	bool bTruncateFile = false;
 	FString PausedPreset;
+	
+	struct FCurrentTraceTarget
+	{
+		FString TraceDest;
+		FTraceAuxiliary::EConnectionType TraceType = FTraceAuxiliary::EConnectionType::None;
+	} CurrentTraceTarget;
+	mutable FRWLock CurrentTargetLock;
 };
 
 static FTraceAuxiliaryImpl GTraceAuxiliary;
@@ -255,7 +253,7 @@ void FTraceAuxiliaryImpl::AddChannel(const TCHAR* Name)
 	FChannelEntry& Value = CommandlineChannels.Add(Hash, {});
 	Value.Name = Name;
 
-	if (State.load() >= EState::Tracing && !Value.bActive)
+	if (IsConnected() && !Value.bActive)
 	{
 		Value.bActive = EnableChannel(*Value.Name);
 	}
@@ -272,7 +270,7 @@ void FTraceAuxiliaryImpl::RemoveChannel(const TCHAR* Name)
 		return;
 	}
 
-	if (State.load() >= EState::Tracing && Channel.bActive)
+	if (IsConnected() && Channel.bActive)
 	{
 		DisableChannel(*Channel.Name);
 		Channel.bActive = false;
@@ -280,19 +278,18 @@ void FTraceAuxiliaryImpl::RemoveChannel(const TCHAR* Name)
 }
 
 ////////////////////////////////////////////////////////////////////////////////
-bool FTraceAuxiliaryImpl::Connect(ETraceConnectType Type, const TCHAR* Parameter, const FTraceAuxiliary::FLogCategoryAlias& LogCategory, uint16 SendFlags)
+bool FTraceAuxiliaryImpl::Connect(FTraceAuxiliary::EConnectionType Type, const TCHAR* Parameter, const FTraceAuxiliary::FLogCategoryAlias& LogCategory, uint16 SendFlags)
 {
 	// Connect/write to file, but only if we're not already sending/writing.
 	bool bConnected = UE::Trace::IsTracing();
 	if (!bConnected)
 	{
-		if (Type == ETraceConnectType::Network)
+		if (Type == FTraceAuxiliary::EConnectionType::Network)
 		{
 			bConnected = SendToHost(Parameter, LogCategory, SendFlags);
 			if (bConnected)
 			{
-				UE_LOG_REF(LogCategory, Display, TEXT("Trace started (connected to trace server %s)."), GetDest());
-				TraceType.store(FTraceAuxiliary::EConnectionType::Network);
+				UE_LOG_REF(LogCategory, Display, TEXT("Trace started (connected to trace server %s)."), *GetDest());
 			}
 			else
 			{
@@ -301,13 +298,12 @@ bool FTraceAuxiliaryImpl::Connect(ETraceConnectType Type, const TCHAR* Parameter
 
 		}
 
-		else if (Type == ETraceConnectType::File)
+		else if (Type == FTraceAuxiliary::EConnectionType::File)
 		{
 			bConnected = WriteToFile(Parameter, LogCategory, SendFlags);
 			if (bConnected)
 			{
-				UE_LOG_REF(LogCategory, Display, TEXT("Trace started (writing to file \"%s\")."), GetDest());
-				TraceType.store(FTraceAuxiliary::EConnectionType::File);
+				UE_LOG_REF(LogCategory, Display, TEXT("Trace started (writing to file \"%s\")."), *GetDest());
 			}
 			else
 			{
@@ -318,17 +314,16 @@ bool FTraceAuxiliaryImpl::Connect(ETraceConnectType Type, const TCHAR* Parameter
 
 		if (bConnected)
 		{
-			FTraceAuxiliary::OnTraceStarted.Broadcast(TraceType.load(), TraceDest);
+			FReadScopeLock _(CurrentTargetLock);
+			FTraceAuxiliary::OnTraceStarted.Broadcast(CurrentTraceTarget.TraceType, CurrentTraceTarget.TraceDest);
 		}
 	}
-
-	if (!bConnected)
+	else
 	{
-		return false;
+		UE_LOG_REF(LogCategory, Error, TEXT("Already tracing from unknown source."));
 	}
 
-	State.store(EState::Tracing);
-	return true;
+	return bConnected;
 }
 
 ////////////////////////////////////////////////////////////////////////////////
@@ -344,11 +339,13 @@ bool FTraceAuxiliaryImpl::Stop()
 		return false;
 	}
 
-	FTraceAuxiliary::OnTraceStopped.Broadcast(TraceType.load(), TraceDest);
+	{
+		FWriteScopeLock _(CurrentTargetLock);
+		FTraceAuxiliary::OnTraceStopped.Broadcast(CurrentTraceTarget.TraceType, CurrentTraceTarget.TraceDest);
 
-	State.store(EState::Stopped);
-	TraceType.store(FTraceAuxiliary::EConnectionType::None);
-	TraceDest.Reset();
+		CurrentTraceTarget.TraceType = FTraceAuxiliary::EConnectionType::None;
+		CurrentTraceTarget.TraceDest.Reset();
+	}
 	return true;
 }
 
@@ -456,7 +453,11 @@ bool FTraceAuxiliaryImpl::SendToHost(const TCHAR* Host, const FTraceAuxiliary::F
 		return false;
 	}
 
-	TraceDest = Host;
+	{
+		FWriteScopeLock _(CurrentTargetLock);
+		CurrentTraceTarget.TraceType = FTraceAuxiliary::EConnectionType::Network;
+		CurrentTraceTarget.TraceDest = MoveTemp(Host);
+	}
 	return true;
 }
 
@@ -551,7 +552,11 @@ bool FTraceAuxiliaryImpl::WriteToFile(const TCHAR* Path, const FTraceAuxiliary::
 		return false;
 	}
 
-	TraceDest = MoveTemp(NativePath);
+	{
+		FWriteScopeLock _(CurrentTargetLock);
+		CurrentTraceTarget.TraceType = FTraceAuxiliary::EConnectionType::File;
+		CurrentTraceTarget.TraceDest = MoveTemp(NativePath);
+	}
 	return true;
 }
 
@@ -612,15 +617,16 @@ bool FTraceAuxiliaryImpl::SendSnapshot(const TCHAR* InHost, uint32 InPort, const
 }
 
 ////////////////////////////////////////////////////////////////////////////////
-const TCHAR* FTraceAuxiliaryImpl::GetDest() const
+FString FTraceAuxiliaryImpl::GetDest() const
 {
-	return *TraceDest;
+	FReadScopeLock _(CurrentTargetLock);
+	return CurrentTraceTarget.TraceDest;
 }
 
 ////////////////////////////////////////////////////////////////////////////////
 bool FTraceAuxiliaryImpl::IsConnected() const
 {
-	return State.load() == EState::Tracing;
+	return UE::Trace::IsTracing();
 }
 
 ////////////////////////////////////////////////////////////////////////////////
@@ -640,7 +646,8 @@ bool FTraceAuxiliaryImpl::IsConnected(FGuid& OutSessionGuid, FGuid& OutTraceGuid
 ////////////////////////////////////////////////////////////////////////////////
 FTraceAuxiliary::EConnectionType FTraceAuxiliaryImpl::GetConnectionType() const
 {
-	return TraceType.load();
+	FReadScopeLock _(CurrentTargetLock);
+	return CurrentTraceTarget.TraceType;
 }
 
 ////////////////////////////////////////////////////////////////////////////////
@@ -744,18 +751,6 @@ static void SetupInitFromConfig(UE::Trace::FInitializeDesc& OutDesc)
 ////////////////////////////////////////////////////////////////////////////////
 static void TraceAuxiliaryConnectEpilogue()
 {
-	// It is possible that something outside of TraceAux's world view has called
-	// UE::Trace::SendTo/WriteTo(). A plugin that has created its own store for
-	// example. There's not really much that can be done about that here (tracing
-	// is singular within a process. We can at least detect the obvious case and
-	// inform the user.
-	const TCHAR* TraceDest = GTraceAuxiliary.GetDest();
-	if (TraceDest[0] == '\0')
-	{
-		UE_LOG(LogConsoleResponse, Warning, TEXT("Trace system already in use by a plugin or -trace*=... argument. Use 'Trace.Stop' first."));
-		return;
-	}
-
 	// Give the user some feedback that everything's underway.
 	TStringBuilder<128> Channels;
 	GTraceAuxiliary.GetActiveChannelsString(Channels);
@@ -856,10 +851,10 @@ static void TraceAuxiliaryStatus()
 	FGuid SessionGuid, TraceGuid;
 	if (GTraceAuxiliary.IsConnected(SessionGuid, TraceGuid))
 	{
-		const TCHAR* Dest = GTraceAuxiliary.GetDest();
-		if (Dest && FCString::Strlen(Dest) > 0)
+		const FString Dest = GTraceAuxiliary.GetDest();
+		if (!Dest.IsEmpty())
 		{
-			ConnectionStr.Appendf(TEXT("Tracing to '%s', "), Dest);
+			ConnectionStr.Appendf(TEXT("Tracing to '%s', "), *Dest);
 			ConnectionStr.Appendf(TEXT("session %s trace %s"), *SessionGuid.ToString(), *TraceGuid.ToString());
 		}
 		else
@@ -1365,7 +1360,7 @@ bool FTraceAuxiliary::Start(EConnectionType Type, const TCHAR* Target, const TCH
 
 	if (GTraceAuxiliary.IsConnected())
 	{
-		UE_LOG_REF(LogCategory, Error, TEXT("Unable to start trace, already tracing to %s"), GTraceAuxiliary.GetDest());
+		UE_LOG_REF(LogCategory, Error, TEXT("Unable to start trace, already tracing to %s"), *GTraceAuxiliary.GetDest());
 		return false;
 	}
 
@@ -1397,14 +1392,7 @@ bool FTraceAuxiliary::Start(EConnectionType Type, const TCHAR* Target, const TCH
 
 	uint16 SendFlags = (Options && Options->bExcludeTail) ? UE::Trace::FSendFlags::ExcludeTail : 0;
 
-	if (Type == EConnectionType::File)
-	{
-		return GTraceAuxiliary.Connect(ETraceConnectType::File, Target, LogCategory, SendFlags);
-	}
-	else if(Type == EConnectionType::Network)
-	{
-		return GTraceAuxiliary.Connect(ETraceConnectType::Network, Target, LogCategory, SendFlags);
-	}
+	return GTraceAuxiliary.Connect(Type, Target, LogCategory, SendFlags);
 #endif
 	return false;
 }
@@ -1697,11 +1685,20 @@ void FTraceAuxiliary::DisableChannels(const TCHAR* Channels)
 const TCHAR* FTraceAuxiliary::GetTraceDestination()
 {
 #if UE_TRACE_ENABLED
-	return GTraceAuxiliary.GetDest();
+	static FString TempUnsafe = GTraceAuxiliary.GetDest();
+	return *TempUnsafe;
 #endif
 	return nullptr;
 }
 
+FString FTraceAuxiliary::GetTraceDestinationString()
+{
+#if UE_TRACE_ENABLED
+	return GTraceAuxiliary.GetDest();
+#endif
+	return FString();
+}
+
 bool FTraceAuxiliary::IsConnected()
 {
 #if UE_TRACE_ENABLED
diff --git a/Engine/Source/Runtime/Core/Public/ProfilingDebugging/TraceAuxiliary.h b/Engine/Source/Runtime/Core/Public/ProfilingDebugging/TraceAuxiliary.h
--- a/Engine/Source/Runtime/Core/Public/ProfilingDebugging/TraceAuxiliary.h
+++ b/Engine/Source/Runtime/Core/Public/ProfilingDebugging/TraceAuxiliary.h
@@ -150,8 +150,15 @@ public:
 	 *  Returns the destination string that is currently being traced to.
 	 *  Contains either a file path or network address. Points to an empty string if tracing is disabled.
 	 */
+	UE_DEPRECATED(5.3, "Use GetTraceDestinationString instead.")
 	static const TCHAR*	GetTraceDestination();
 	
+	/**
+	 *  Returns the destination string that is currently being traced to.
+	 *  Contains either a file path or network address. Empty if tracing is disabled.
+	 */
+	static FString GetTraceDestinationString();
+	
 	/**
 	 *  Returns whether the trace system is currently connected to a trace sink (file or network)
 	 */
