Added support for LLM to track allocations in UE::Trace.

  - Currently all allocations from UE::Trace are Untracked.
  - Added a new ELLMTag::Trace.
  - Added hooks (lambda) to FMemory_GCreateMalloc_ThreadUnsafe that uses FLLMScope to track allocations.
  - Note this change will only work if nothing else uses UE::Trace::SetMemoryHooks for it's own hooks.

#tests Built most Development targets. Tested a Replay and saw Trace with 5MB in the Perfreports. Ran Horde Incremental CI - Preflight.

diff --git a/Engine/Source/Runtime/Core/Private/HAL/LowLevelMemTracker.cpp b/Engine/Source/Runtime/Core/Private/HAL/LowLevelMemTracker.cpp
--- a/Engine/Source/Runtime/Core/Private/HAL/LowLevelMemTracker.cpp
+++ b/Engine/Source/Runtime/Core/Private/HAL/LowLevelMemTracker.cpp
@@ -290,6 +290,7 @@ DECLARE_LLM_MEMORY_STAT(TEXT("ElectraPlayer"), STAT_ElectraPlayerLLM, STATGROUP_
 DECLARE_LLM_MEMORY_STAT(TEXT("WMFPlayer"), STAT_WMFPlayerLLM, STATGROUP_LLMFULL);
 DECLARE_LLM_MEMORY_STAT(TEXT("MMIO"), STAT_PlatformMMIOLLM, STATGROUP_LLMPlatform);
 DECLARE_LLM_MEMORY_STAT(TEXT("VirtualMemory"), STAT_PlatformVMLLM, STATGROUP_LLMPlatform);
+DECLARE_LLM_MEMORY_STAT(TEXT("Trace"), STAT_Trace, STATGROUP_LLMFULL);
 DECLARE_LLM_MEMORY_STAT(TEXT("CustomName"), STAT_CustomName, STATGROUP_LLMFULL);
 
 namespace UE::LLMPrivate
diff --git a/Engine/Source/Runtime/Core/Private/HAL/UnrealMemory.cpp b/Engine/Source/Runtime/Core/Private/HAL/UnrealMemory.cpp
--- a/Engine/Source/Runtime/Core/Private/HAL/UnrealMemory.cpp
+++ b/Engine/Source/Runtime/Core/Private/HAL/UnrealMemory.cpp
@@ -344,6 +344,26 @@ static int FMemory_GCreateMalloc_ThreadUnsafe()
 	// Setup malloc crash as soon as possible.
 	FPlatformMallocCrash::Get( GMalloc );
 
+#if UE_TRACE_ENABLED && ENABLE_LOW_LEVEL_MEM_TRACKER
+	// Only hookup LLM support for UE::Trace if something hasn't already set it. This means that
+	// LLM support for UE::Trace is only possible if nothing else is using these hooks.
+	if (!UE::Trace::AreMemoryHooksSet())
+	{
+		auto TraceAlloc = [] (SIZE_T Size, uint32 Alignment) -> void*
+		{
+			FLLMScope LLMScope(ELLMTag::Trace, false, ELLMTagSet::None, ELLMTracker::Default);
+			return FMemory::Malloc(Size);
+		};
+
+		auto TraceFree = [] (void* Address, SIZE_T Size)
+		{
+			FMemory::Free(Address);
+		};
+
+		UE::Trace::SetMemoryHooks(TraceAlloc, TraceFree);
+	}
+#endif // UE_TRACE_ENABLED && ENABLE_LOW_LEVEL_MEM_TRACKER
+
 #if PLATFORM_USES_FIXED_GMalloc_CLASS
 #if MALLOC_VERIFY || MALLOC_LEAKDETECTION || UE_USE_MALLOC_FILL_BYTES
 #error "Turn off PLATFORM_USES_FIXED_GMalloc_CLASS in order to use special allocator proxies"
diff --git a/Engine/Source/Runtime/Core/Public/HAL/LowLevelMemTracker.h b/Engine/Source/Runtime/Core/Public/HAL/LowLevelMemTracker.h
--- a/Engine/Source/Runtime/Core/Public/HAL/LowLevelMemTracker.h
+++ b/Engine/Source/Runtime/Core/Public/HAL/LowLevelMemTracker.h
@@ -243,6 +243,7 @@ enum class ELLMTagSet : uint8
 	macro(WMFPlayer,							"WMFPlayer",					GET_STATFNAME(STAT_WMFPlayerLLM),							GET_STATFNAME(STAT_MediaStreamingSummaryLLM),	ELLMTag::MediaStreaming)\
 	macro(PlatformMMIO,							"MMIO",							GET_STATFNAME(STAT_PlatformMMIOLLM),						NAME_None,										-1)\
 	macro(PlatformVM,							"Virtual Memory",				GET_STATFNAME(STAT_PlatformVMLLM),							NAME_None,										-1)\
+	macro(Trace,								"Trace",						GET_STATFNAME(STAT_Trace),									NAME_None,										-1)\
 	macro(CustomName,							"CustomName",					GET_STATFNAME(STAT_CustomName),								NAME_None,										-1)\
 
 /*
diff --git a/Engine/Source/Runtime/TraceLog/Private/Trace/Trace.cpp b/Engine/Source/Runtime/TraceLog/Private/Trace/Trace.cpp
--- a/Engine/Source/Runtime/TraceLog/Private/Trace/Trace.cpp
+++ b/Engine/Source/Runtime/TraceLog/Private/Trace/Trace.cpp
@@ -16,6 +16,7 @@ namespace Private
 
 ////////////////////////////////////////////////////////////////////////////////
 void	Writer_MemorySetHooks(AllocFunc, FreeFunc);
+bool	Writer_AreMemoryHooksSet();
 void	Writer_Initialize(const FInitializeDesc&);
 void	Writer_WorkerCreate();
 void	Writer_Shutdown();
@@ -58,6 +59,12 @@ void SetMemoryHooks(AllocFunc Alloc, FreeFunc Free)
 	Private::Writer_MemorySetHooks(Alloc, Free);
 }
 
+////////////////////////////////////////////////////////////////////////////////
+bool AreMemoryHooksSet()
+{
+	return Private::Writer_AreMemoryHooksSet();
+}
+
 ////////////////////////////////////////////////////////////////////////////////
 void Initialize(const FInitializeDesc& Desc)
 {
diff --git a/Engine/Source/Runtime/TraceLog/Private/Trace/Writer.cpp b/Engine/Source/Runtime/TraceLog/Private/Trace/Writer.cpp
--- a/Engine/Source/Runtime/TraceLog/Private/Trace/Writer.cpp
+++ b/Engine/Source/Runtime/TraceLog/Private/Trace/Writer.cpp
@@ -208,6 +208,12 @@ void Writer_MemorySetHooks(decltype(AllocHook) Alloc, decltype(FreeHook) Free)
 	FreeHook = Free;
 }
 
+////////////////////////////////////////////////////////////////////////////////
+bool Writer_AreMemoryHooksSet()
+{
+	return (AllocHook != nullptr) || (FreeHook != nullptr);
+}
+
 ////////////////////////////////////////////////////////////////////////////////
 void* Writer_MemoryAllocate(SIZE_T Size, uint32 Alignment)
 {
diff --git a/Engine/Source/Runtime/TraceLog/Public/Trace/Trace.h b/Engine/Source/Runtime/TraceLog/Public/Trace/Trace.h
--- a/Engine/Source/Runtime/TraceLog/Public/Trace/Trace.h
+++ b/Engine/Source/Runtime/TraceLog/Public/Trace/Trace.h
@@ -116,6 +116,7 @@ struct FSendFlags
 };
 
 UE_TRACE_API void	SetMemoryHooks(AllocFunc Alloc, FreeFunc Free) UE_TRACE_IMPL();
+UE_TRACE_API bool	AreMemoryHooksSet() UE_TRACE_IMPL(false);
 UE_TRACE_API void	Initialize(const FInitializeDesc& Desc) UE_TRACE_IMPL();
 UE_TRACE_API void	StartWorkerThread() UE_TRACE_IMPL();	
 UE_TRACE_API void	Shutdown() UE_TRACE_IMPL();
