Correctly track thread idle time through nested idle scopes

- we do nest FThreadIdleStats::FScopeIdles
  - only accumulate the idle time for the outer most scope
  - otherwise we are multiply counting idle time and our accumulated idle time per-thread is incorrect

- testing
  - ran Lyra on Win64 and console platform
  - ran Lyra editor on Win64
  - took Insights capture on console platform and verified idle scopes are present when including ThreadIdleScope trace channel

diff --git a/Engine/Source/Runtime/Core/Private/Stats/Stats2.cpp b/Engine/Source/Runtime/Core/Private/Stats/Stats2.cpp
--- a/Engine/Source/Runtime/Core/Private/Stats/Stats2.cpp
+++ b/Engine/Source/Runtime/Core/Private/Stats/Stats2.cpp
@@ -64,27 +64,26 @@ DEFINE_STAT(STAT_FrameTime);
 DEFINE_STAT(STAT_NamedMarker);
 DEFINE_STAT(STAT_SecondsPerCycle);
 
+#if !(defined(DISABLE_THREAD_IDLE_STATS) && DISABLE_THREAD_IDLE_STATS)
 #if CPUPROFILERTRACE_ENABLED
 	UE_TRACE_CHANNEL_DEFINE(ThreadIdleScopeChannel);
-	TRACE_CPUPROFILER_EVENT_DECLARE(TraceEventScopeId);
+	TRACE_CPUPROFILER_EVENT_DECLARE(ThreadIdleScopeTraceEventId);
 #endif
 
 CORE_API FThreadIdleStats::FScopeIdle::FScopeIdle(bool bInIgnore /* = false */)
 	: Start(FPlatformTime::Cycles())
-	, bIgnore(bInIgnore)
+	, bIgnore(bInIgnore || FThreadIdleStats::Get().bInIdleScope)
 #if CPUPROFILERTRACE_ENABLED
-	, TraceEventScope(TraceEventScopeId, TEXT("FThreadIdleStats::FScopeIdle"), ThreadIdleScopeChannel, !bInIgnore, __FILE__, __LINE__)
+	, TraceEventScope(ThreadIdleScopeTraceEventId, TEXT("FThreadIdleStats::FScopeIdle"), ThreadIdleScopeChannel, !bIgnore, __FILE__, __LINE__)
 #endif
 {
-#if DETECT_NESTED_THREAD_IDLE_STAT_SCOPES
-		if (!bInIgnore)
-		{
-			FThreadIdleStats& IdleStats = FThreadIdleStats::Get();
-			checkf(!IdleStats.bInIdleScope, TEXT("Nested FThreadIdleStats::FScopeIdle encountered"));
-			IdleStats.bInIdleScope = true;
-		}
-#endif
+	if (!bIgnore)
+	{
+		FThreadIdleStats& IdleStats = FThreadIdleStats::Get();
+		IdleStats.bInIdleScope = true;
+	}
 }
+#endif // #if !(defined(DISABLE_THREAD_IDLE_STATS) && DISABLE_THREAD_IDLE_STATS)
 
 /*-----------------------------------------------------------------------------
 	DebugLeakTest, for the stats based memory profiler
diff --git a/Engine/Source/Runtime/Core/Public/Stats/Stats2.h b/Engine/Source/Runtime/Core/Public/Stats/Stats2.h
--- a/Engine/Source/Runtime/Core/Public/Stats/Stats2.h
+++ b/Engine/Source/Runtime/Core/Public/Stats/Stats2.h
@@ -41,10 +41,6 @@
 #include "UObject/NameTypes.h"
 #include "UObject/UnrealNames.h"
 
-#ifndef DETECT_NESTED_THREAD_IDLE_STAT_SCOPES
-	#define DETECT_NESTED_THREAD_IDLE_STAT_SCOPES	0		// This will cause check()s to fire if thread idle scopes are nested and hence double counting idle time.
-#endif // DETECT_NESTED_THREAD_IDLE_STAT_SCOPES
-
 class FOutputDevice;
 class FScopeCycleCounter;
 class FThreadStats;
@@ -62,9 +58,7 @@ class CORE_API FThreadIdleStats : public TThreadSingleton<FThreadIdleStats>
 		: Waits(0)
 		, WaitsCriticalPath(0)
 		, IsCriticalPathCounter(1)
-#if DETECT_NESTED_THREAD_IDLE_STAT_SCOPES
 		, bInIdleScope(false)
-#endif
 	{}
 
 public:
@@ -76,10 +70,7 @@ public:
 	uint32 WaitsCriticalPath;
 
 	int IsCriticalPathCounter;
-
-#if DETECT_NESTED_THREAD_IDLE_STAT_SCOPES
 	bool bInIdleScope;
-#endif
 
 	static void BeginCriticalPath()
 	{
@@ -119,19 +110,20 @@ public:
 
 	struct FScopeIdle
 	{
+#if defined(DISABLE_THREAD_IDLE_STATS) && DISABLE_THREAD_IDLE_STATS
+		FScopeIdle( bool bInIgnore = false )
+		{}
+#else
 		/** Starting cycle counter. */
 		const uint32 Start;
 
 		/** If true, we ignore this thread idle stats. */
 		const bool bIgnore;
 
-#if defined(DISABLE_THREAD_IDLE_STATS) && DISABLE_THREAD_IDLE_STATS
-		FScopeIdle( bool bInIgnore = false )
-			: Start(0)
-			, bIgnore( bInIgnore )
-		{
-		}
-#else
+#if CPUPROFILERTRACE_ENABLED
+		FCpuProfilerTrace::FEventScope TraceEventScope;
+#endif
+
 		CORE_API FScopeIdle(bool bInIgnore = false);
 
 		~FScopeIdle()
@@ -147,15 +139,9 @@ public:
 					IdleStats.WaitsCriticalPath += CyclesElapsed;
 				}
 
-#if DETECT_NESTED_THREAD_IDLE_STAT_SCOPES
 				IdleStats.bInIdleScope = false;
-#endif
 			}
 		}
-
-#if CPUPROFILERTRACE_ENABLED
-		FCpuProfilerTrace::FEventScope TraceEventScope;
-#endif
 #endif
 	};
 };
