corrected the comment about task pipes ordering guarantees

diff --git a/Engine/Source/Runtime/Core/Public/Tasks/Pipe.h b/Engine/Source/Runtime/Core/Public/Tasks/Pipe.h
--- a/Engine/Source/Runtime/Core/Public/Tasks/Pipe.h
+++ b/Engine/Source/Runtime/Core/Public/Tasks/Pipe.h
@@ -20,7 +20,7 @@ namespace UE::Tasks
 	// non-concurrent tasks execution. FPipe is a replacement for named threads because it's lightweight and flexible -
 	// there can be a large dynamic number of pipes each controlling its own shared resource. Can be used as a replacement for
 	// dedicated threads.
-	// Execution order is not specified, only that tasks from the same pipe are not executed concurrently.
+	// Execution order is FIFO, i.e. it's the same as launching order. This means that if launching order is unspecified
 	// A pipe must be alive until its last task is completed.
 	// See `FTasksPipeTest` for tests and examples.
 	class FPipe
