cosmetic: corrected access detector test name

diff --git a/Engine/Source/Runtime/Core/Private/Misc/MTAccessDetector.cpp b/Engine/Source/Runtime/Core/Private/Misc/MTAccessDetector.cpp
--- a/Engine/Source/Runtime/Core/Private/Misc/MTAccessDetector.cpp
+++ b/Engine/Source/Runtime/Core/Private/Misc/MTAccessDetector.cpp
@@ -383,9 +383,9 @@ namespace MTAccessDetector_Private
 	};
 }
 
-IMPLEMENT_SIMPLE_AUTOMATION_TEST(FRWRecursiveAccessDetectorTest, "System.Core.RWRecursiveAccessDetector", EAutomationTestFlags::ApplicationContextMask | EAutomationTestFlags::EngineFilter);
+IMPLEMENT_SIMPLE_AUTOMATION_TEST(FMRSWRecursiveAccessDetectorTest, "System.Core.MRSWRecursiveAccessDetector", EAutomationTestFlags::ApplicationContextMask | EAutomationTestFlags::EngineFilter);
 
-bool FRWRecursiveAccessDetectorTest::RunTest(const FString& Parameters)
+bool FMRSWRecursiveAccessDetectorTest::RunTest(const FString& Parameters)
 {
 #pragma warning(push)
 #pragma warning(disable:6001) // "Using uninitialized memory" warning
