Add torque units and extend force units

diff --git a/Engine/Config/BaseEditor.ini b/Engine/Config/BaseEditor.ini
--- a/Engine/Config/BaseEditor.ini
+++ b/Engine/Config/BaseEditor.ini
@@ -517,6 +517,7 @@ AngleUnits=EUnit::Degrees
 SpeedUnits=EUnit::MetersPerSecond
 TemperatureUnits=EUnit::Celsius
 ForceUnits=EUnit::Newtons
+TorqueUnits=EUnit::NewtonMeters
 
 [BlueprintSearchSettings]
 AsyncTaskBatchSize=2048
diff --git a/Engine/Programs/UnrealHeaderTool/Config/DefaultEngine.ini b/Engine/Programs/UnrealHeaderTool/Config/DefaultEngine.ini
--- a/Engine/Programs/UnrealHeaderTool/Config/DefaultEngine.ini
+++ b/Engine/Programs/UnrealHeaderTool/Config/DefaultEngine.ini
@@ -99,6 +99,12 @@ gc.MaxObjectsInProgram=262144
 +Units="lbf"
 +Units="KilogramsForce"
 +Units="kgf"
++Units="KilogramsCentimetersPerSecondSquared"
++Units="kgcm/s\u00B2"
++Units="NewtonMeters"
++Units="Nm"
++Units="KilogramCentimetersSquaredPerSecondSquared"
++Units="kgcm\u00B2/s\u00B2"
 +Units="Hertz"
 +Units="Hz"
 +Units="Kilohertz"
diff --git a/Engine/Source/Editor/UnrealEd/Private/Settings/EditorProjectSettings.cpp b/Engine/Source/Editor/UnrealEd/Private/Settings/EditorProjectSettings.cpp
--- a/Engine/Source/Editor/UnrealEd/Private/Settings/EditorProjectSettings.cpp
+++ b/Engine/Source/Editor/UnrealEd/Private/Settings/EditorProjectSettings.cpp
@@ -103,6 +103,10 @@ void UEditorProjectAppearanceSettings::PostEditChangeProperty( struct FPropertyC
 	{
 		Settings.SetDisplayUnits(EUnitType::Force, ForceUnits);
 	}
+	else if (Name == GET_MEMBER_NAME_CHECKED(UEditorProjectAppearanceSettings, TorqueUnits))
+	{
+		Settings.SetDisplayUnits(EUnitType::Torque, TorqueUnits);
+	}
 	else if (Name == GET_MEMBER_NAME_CHECKED(UEditorProjectAppearanceSettings, bDisplayUnits))
 	{
 		Settings.SetShouldDisplayUnits(bDisplayUnits);
diff --git a/Engine/Source/Editor/UnrealEd/Public/Settings/EditorProjectSettings.h b/Engine/Source/Editor/UnrealEd/Public/Settings/EditorProjectSettings.h
--- a/Engine/Source/Editor/UnrealEd/Public/Settings/EditorProjectSettings.h
+++ b/Engine/Source/Editor/UnrealEd/Public/Settings/EditorProjectSettings.h
@@ -92,9 +92,12 @@ public:
 	UPROPERTY(EditAnywhere, config, Category=Units, AdvancedDisplay, meta=(DisplayName="Temperature", Tooltip="Choose the units in which to display temperatures.", ValidEnumValues="Celsius, Farenheit, Kelvin"))
 	EUnit TemperatureUnits;
 
-	UPROPERTY(EditAnywhere, config, Category=Units, AdvancedDisplay, meta=(DisplayName="Force", Tooltip="Choose the units in which to display forces.", ValidEnumValues="Newtons, PoundsForce, KilogramsForce"))
+	UPROPERTY(EditAnywhere, config, Category=Units, AdvancedDisplay, meta=(DisplayName="Force", Tooltip="Choose the units in which to display forces.", ValidEnumValues="Newtons, PoundsForce, KilogramsForce, KilogramCentimetersPerSecondSquared"))
 	EUnit ForceUnits;
 
+	UPROPERTY(EditAnywhere, config, Category = Units, AdvancedDisplay, meta = (DisplayName = "Torque", Tooltip = "Choose the units in which to display torques.", ValidEnumValues = "NewtonMeters, KilogramCentimetersSquaredPerSecondSquared"))
+	EUnit TorqueUnits;
+
 	// Should the Reference Viewer have 'Show Searchable Names' checked by default when opened in this project
 	UPROPERTY(EditAnywhere, config, Category=ReferenceViewer)
 	EReferenceViewerSettingMode ShowSearchableNames;
diff --git a/Engine/Source/Runtime/Core/Private/Math/UnitConversion.cpp b/Engine/Source/Runtime/Core/Private/Math/UnitConversion.cpp
--- a/Engine/Source/Runtime/Core/Private/Math/UnitConversion.cpp
+++ b/Engine/Source/Runtime/Core/Private/Math/UnitConversion.cpp
@@ -64,6 +64,10 @@ FParseCandidate ParseCandidates[] = {
 	{ TEXT("Newtons"),				EUnit::Newtons },				{ TEXT("N"),		EUnit::Newtons },
 	{ TEXT("PoundsForce"),			EUnit::PoundsForce },			{ TEXT("lbf"),		EUnit::PoundsForce },
 	{ TEXT("KilogramsForce"),		EUnit::KilogramsForce },		{ TEXT("kgf"),		EUnit::KilogramsForce },
+	{ TEXT("KilogramsCentimetersPerSecondSquared"),	EUnit::KilogramCentimetersPerSecondSquared },	{ TEXT("kgcm/s2"),	EUnit::KilogramCentimetersPerSecondSquared },	{ TEXT("kgcm/s\u00B2"),	EUnit::KilogramCentimetersPerSecondSquared },
+
+	{ TEXT("NewtonMeters"),			EUnit::NewtonMeters },			{ TEXT("Nm"),		EUnit::NewtonMeters },
+	{ TEXT("KilogramsCentimetersSquaredPerSecondSquared"),	EUnit::KilogramCentimetersSquaredPerSecondSquared },		{ TEXT("kgcm2/s2"),	EUnit::KilogramCentimetersSquaredPerSecondSquared },	{ TEXT("kgcm\u00B2/s\u00B2"),	EUnit::KilogramCentimetersSquaredPerSecondSquared },
 
 	{ TEXT("Hertz"),				EUnit::Hertz },					{ TEXT("Hz"),		EUnit::Hertz },
 	{ TEXT("Kilohertz"),			EUnit::Kilohertz },				{ TEXT("KHz"),		EUnit::Kilohertz },
@@ -120,7 +124,9 @@ const TCHAR* const DisplayStrings[] = {
 	TEXT("\u00B5g"), TEXT("mg"), TEXT("g"), TEXT("kg"), TEXT("t"),
 	TEXT("oz"), TEXT("lb"), TEXT("st"),
 
-	TEXT("N"), TEXT("lbf"), TEXT("kgf"),
+	TEXT("N"), TEXT("lbf"), TEXT("kgf"), TEXT("kgcm/s\u00B2"),
+
+	TEXT("Nm"), TEXT("kgcm\u00B2/s\u00B2"),
 
 	TEXT("Hz"), TEXT("KHz"), TEXT("MHz"), TEXT("GHz"), TEXT("rpm"),
 
@@ -153,7 +159,9 @@ const EUnitType UnitTypes[] = {
 	EUnitType::Mass,		EUnitType::Mass,		EUnitType::Mass,		EUnitType::Mass,		EUnitType::Mass,
 	EUnitType::Mass,		EUnitType::Mass,		EUnitType::Mass,
 
-	EUnitType::Force,		EUnitType::Force,		EUnitType::Force,
+	EUnitType::Force,		EUnitType::Force,		EUnitType::Force,		EUnitType::Force,
+
+	EUnitType::Torque,		EUnitType::Torque,
 
 	EUnitType::Frequency,	EUnitType::Frequency,	EUnitType::Frequency,	EUnitType::Frequency,	EUnitType::Frequency,
 
@@ -437,6 +445,7 @@ FUnitSettings::FUnitSettings()
 	DisplayUnits[(uint8)EUnitType::Temperature].Add(EUnit::Celsius);
 	DisplayUnits[(uint8)EUnitType::Mass].Add(EUnit::Kilograms);
 	DisplayUnits[(uint8)EUnitType::Force].Add(EUnit::Newtons);
+	DisplayUnits[(uint8)EUnitType::Torque].Add(EUnit::NewtonMeters);
 	DisplayUnits[(uint8)EUnitType::Frequency].Add(EUnit::Hertz);
 	DisplayUnits[(uint8)EUnitType::DataSize].Add(EUnit::Megabytes);
 	DisplayUnits[(uint8)EUnitType::LuminousFlux].Add(EUnit::Lumens);
@@ -622,9 +631,20 @@ namespace UnitConversion
 		// Convert to Newtons
 		switch (From)
 		{
-			case EUnit::PoundsForce:		return 4.44822162;
-			case EUnit::KilogramsForce:		return 9.80665;
-			default: 						return 1;
+		case EUnit::PoundsForce:							return 4.44822162;
+		case EUnit::KilogramsForce:							return 9.80665;
+		case EUnit::KilogramCentimetersPerSecondSquared:	return 0.01;
+		default: 											return 1;
+		}
+	}
+
+	double TorqueUnificationFactor(EUnit From)
+	{
+		// Convert to NewtonMeters
+		switch (From)
+		{
+		case EUnit::KilogramCentimetersSquaredPerSecondSquared:	return 0.0001;
+		default: 												return 1;
 		}
 	}
 
diff --git a/Engine/Source/Runtime/Core/Public/Math/UnitConversion.h b/Engine/Source/Runtime/Core/Public/Math/UnitConversion.h
--- a/Engine/Source/Runtime/Core/Public/Math/UnitConversion.h
+++ b/Engine/Source/Runtime/Core/Public/Math/UnitConversion.h
@@ -36,7 +36,10 @@ enum class EUnit : uint8
 	Ounces, Pounds, Stones,
 
 	/** Force units */
-	Newtons, PoundsForce, KilogramsForce,
+	Newtons, PoundsForce, KilogramsForce, KilogramCentimetersPerSecondSquared,
+
+	/** Torque Units */
+	NewtonMeters, KilogramCentimetersSquaredPerSecondSquared,
 
 	/** Frequency units */
 	Hertz, Kilohertz, Megahertz, Gigahertz, RevolutionsPerMinute,
@@ -66,7 +69,7 @@ enum class EUnit : uint8
 /** Enumeration that specifies particular classes of unit */
 enum class EUnitType
 {
-	Distance, Angle, Speed, Temperature, Mass, Force, Frequency, DataSize, LuminousFlux, LuminousIntensity, Illuminance, Luminance, Time, PixelDensity, Multipliers, ExposureValue, Stress,
+	Distance, Angle, Speed, Temperature, Mass, Force, Torque, Frequency, DataSize, LuminousFlux, LuminousIntensity, Illuminance, Luminance, Time, PixelDensity, Multipliers, ExposureValue, Stress,
 
 	// Symbolic entry - do not use directly
 	NumberOf,
diff --git a/Engine/Source/Runtime/Core/Public/Math/UnitConversion.inl b/Engine/Source/Runtime/Core/Public/Math/UnitConversion.inl
--- a/Engine/Source/Runtime/Core/Public/Math/UnitConversion.inl
+++ b/Engine/Source/Runtime/Core/Public/Math/UnitConversion.inl
@@ -30,6 +30,8 @@ namespace UnitConversion
 	CORE_API double MassUnificationFactor(EUnit From);
 	/** Find the common quantization factor for the specified force unit. Quantizes to Newtons. */
 	CORE_API double ForceUnificationFactor(EUnit From);
+	/** Find the common quantization factor for the specified torque unit. Quantizes to Newton meters. */
+	CORE_API double TorqueUnificationFactor(EUnit From);
 	/** Find the common quantization factor for the specified frequency unit. Quantizes to KHz. */
 	CORE_API double FrequencyUnificationFactor(EUnit From);
 	/** Find the common quantization factor for the specified data size unit. Quantizes to MB. */
@@ -83,6 +85,7 @@ T FUnitConversion::Convert(T InValue, EUnit From, EUnit To)
 		case EUnitType::Speed:				return InValue * SpeedUnificationFactor(From) 			* (1.0 / SpeedUnificationFactor(To));
 		case EUnitType::Mass:				return InValue * MassUnificationFactor(From) 			* (1.0 / MassUnificationFactor(To));
 		case EUnitType::Force:				return InValue * ForceUnificationFactor(From) 			* (1.0 / ForceUnificationFactor(To));
+		case EUnitType::Torque:				return InValue * TorqueUnificationFactor(From)			* (1.0 / TorqueUnificationFactor(To));
 		case EUnitType::Frequency:			return InValue * FrequencyUnificationFactor(From) 		* (1.0 / FrequencyUnificationFactor(To));
 		case EUnitType::DataSize:			return InValue * DataSizeUnificationFactor(From) 		* (1.0 / DataSizeUnificationFactor(To));
 		case EUnitType::LuminousFlux:		return InValue;
diff --git a/Engine/Source/Runtime/Core/Tests/Math/UnitConversionTests.cpp b/Engine/Source/Runtime/Core/Tests/Math/UnitConversionTests.cpp
--- a/Engine/Source/Runtime/Core/Tests/Math/UnitConversionTests.cpp
+++ b/Engine/Source/Runtime/Core/Tests/Math/UnitConversionTests.cpp
@@ -26,63 +26,67 @@ TEST_CASE_NAMED(FUnitUnitTests, "System::Core::Math::Unit Conversion", "[Applica
 		EUnit FromUnit, ToUnit;
 	};
 	static FTestStruct Tests[] = {
-		{ 0.025,	2.3651e11,	1e7, 	EUnit::Lightyears, 			EUnit::Kilometers	 		},
-		{ 0.5,		80467.2,	0.1,	EUnit::Miles, 				EUnit::Centimeters	 		},
-		{ 0.2,		182.88,		0.01,	EUnit::Yards, 				EUnit::Millimeters	 		},
-		{ 0.2,		60960,		0.1,	EUnit::Feet, 				EUnit::Micrometers	 		},
-		{ 10,		0.254,		0.001,	EUnit::Inches, 				EUnit::Meters		 		},
-		{ 0.75,		2460.6299,	1e-4,	EUnit::Kilometers, 			EUnit::Feet					},
-		{ 1,		39.37,		0.01,	EUnit::Meters, 				EUnit::Inches				},
-		{ 2750,		27.5,		1e-6,	EUnit::Centimeters, 		EUnit::Meters				},
-		{ 1000,		1.0936,		1e-4,	EUnit::Millimeters, 		EUnit::Yards		 		},
-		{ 2000,		0.0787,		1e-4,	EUnit::Micrometers, 		EUnit::Inches		 		},
-
-		{ 90,		UE_PI/2,	1e-3,	EUnit::Degrees, 			EUnit::Radians		 		},
-		{ UE_PI,	180,		1e-3,	EUnit::Radians, 			EUnit::Degrees		 		},
-
-		{ 12,		43.2,		0.1,	EUnit::MetersPerSecond,		EUnit::KilometersPerHour	},
-		{ 1,		0.6214,		1e-4,	EUnit::KilometersPerHour,	EUnit::MilesPerHour			},
-		{ 15,		6.7056,		1e-4,	EUnit::MilesPerHour,		EUnit::MetersPerSecond		},
-
-		{ 100, 		212,		0.1,	EUnit::Celsius,				EUnit::Farenheit			},
-		{ 400, 		477.594,	1e-3,	EUnit::Farenheit,			EUnit::Kelvin				},
-		{ 72, 		-201.15,	0.01,	EUnit::Kelvin,				EUnit::Celsius				},
-
-		{ 1000, 	3.5274e-5,	1e-6,	EUnit::Micrograms,			EUnit::Ounces,				},
-		{ 1000,		1,			0.1,	EUnit::Milligrams,			EUnit::Grams,				},
-		{ 200,		0.4409,		1e-4,	EUnit::Grams,				EUnit::Pounds,				},
-		{ 0.15,		150000,		0.1,	EUnit::Kilograms,			EUnit::Milligrams,			},
-		{ 1,		157.473,	1e-3,	EUnit::MetricTons,			EUnit::Stones,				},
-		{ 0.001,	28349.5,	0.1,	EUnit::Ounces,				EUnit::Micrograms,			},
-		{ 500,		226.796,	1e-3,	EUnit::Pounds,				EUnit::Kilograms,			},
-		{ 100,		0.6350,		1e-4,	EUnit::Stones,				EUnit::MetricTons,			},
-
-		{ 100,		10.1972,	1e-4,	EUnit::Newtons,				EUnit::KilogramsForce,		},
-		{ 15,		66.7233,	1e-4,	EUnit::PoundsForce,			EUnit::Newtons,				},
-		{ 2,		4.4092,		1e-4,	EUnit::KilogramsForce,		EUnit::PoundsForce,			},
+		{ 0.025,	2.3651e11,	1e7, 	EUnit::Lightyears, 									EUnit::Kilometers	 								},
+		{ 0.5,		80467.2,	0.1,	EUnit::Miles, 										EUnit::Centimeters	 								},
+		{ 0.2,		182.88,		0.01,	EUnit::Yards, 										EUnit::Millimeters	 								},
+		{ 0.2,		60960,		0.1,	EUnit::Feet, 										EUnit::Micrometers	 								},
+		{ 10,		0.254,		0.001,	EUnit::Inches, 										EUnit::Meters		 								},
+		{ 0.75,		2460.6299,	1e-4,	EUnit::Kilometers, 									EUnit::Feet											},
+		{ 1,		39.37,		0.01,	EUnit::Meters, 										EUnit::Inches										},
+		{ 2750,		27.5,		1e-6,	EUnit::Centimeters, 								EUnit::Meters										},
+		{ 1000,		1.0936,		1e-4,	EUnit::Millimeters, 								EUnit::Yards		 								},
+		{ 2000,		0.0787,		1e-4,	EUnit::Micrometers, 								EUnit::Inches		 								},
+
+		{ 90,		UE_PI/2,	1e-3,	EUnit::Degrees, 									EUnit::Radians		 								},
+		{ UE_PI,	180,		1e-3,	EUnit::Radians, 									EUnit::Degrees		 								},
+
+		{ 12,		43.2,		0.1,	EUnit::MetersPerSecond,								EUnit::KilometersPerHour							},
+		{ 1,		0.6214,		1e-4,	EUnit::KilometersPerHour,							EUnit::MilesPerHour									},
+		{ 15,		6.7056,		1e-4,	EUnit::MilesPerHour,								EUnit::MetersPerSecond								},
+
+		{ 100, 		212,		0.1,	EUnit::Celsius,										EUnit::Farenheit									},
+		{ 400, 		477.594,	1e-3,	EUnit::Farenheit,									EUnit::Kelvin										},
+		{ 72, 		-201.15,	0.01,	EUnit::Kelvin,										EUnit::Celsius										},
+
+		{ 1000, 	3.5274e-5,	1e-6,	EUnit::Micrograms,									EUnit::Ounces,										},
+		{ 1000,		1,			0.1,	EUnit::Milligrams,									EUnit::Grams,										},
+		{ 200,		0.4409,		1e-4,	EUnit::Grams,										EUnit::Pounds,										},
+		{ 0.15,		150000,		0.1,	EUnit::Kilograms,									EUnit::Milligrams,									},
+		{ 1,		157.473,	1e-3,	EUnit::MetricTons,									EUnit::Stones,										},
+		{ 0.001,	28349.5,	0.1,	EUnit::Ounces,										EUnit::Micrograms,									},
+		{ 500,		226.796,	1e-3,	EUnit::Pounds,										EUnit::Kilograms,									},
+		{ 100,		0.6350,		1e-4,	EUnit::Stones,										EUnit::MetricTons,									},
+
+		{ 100,		10.1972,	1e-4,	EUnit::Newtons,										EUnit::KilogramsForce,								},
+		{ 2,		4.4092,		1e-4,	EUnit::KilogramsForce,								EUnit::PoundsForce,									},
+		{ 15,		6672.33,	1e-2,	EUnit::PoundsForce,									EUnit::KilogramCentimetersPerSecondSquared,			},
+		{ 500,		5,			0.1,	EUnit::KilogramCentimetersPerSecondSquared,			EUnit::Newtons,										},
+
+		{ 10,		100000,		0.1,	EUnit::NewtonMeters, 								EUnit::KilogramCentimetersSquaredPerSecondSquared	},
+		{ 2000,		0.2,		0.1,	EUnit::KilogramCentimetersSquaredPerSecondSquared, 	EUnit::NewtonMeters		 							},
+
+		{ 1000,		1,			0.1,	EUnit::Hertz,										EUnit::Kilohertz,									},
+		{ 0.25,		250*60,		1e-3,	EUnit::Kilohertz,									EUnit::RevolutionsPerMinute,						},
+		{ 1000,		1,			1e-3,	EUnit::Megahertz,									EUnit::Gigahertz,									},
+		{ 0.001,	1000000,	1e-3,	EUnit::Gigahertz,									EUnit::Hertz,										},
+		{ 100,		100.0/60,	1e-3,	EUnit::RevolutionsPerMinute,						EUnit::Hertz,										},
 		
-		{ 1000,		1,			0.1,	EUnit::Hertz,				EUnit::Kilohertz,			},
-		{ 0.25,		250*60,		1e-3,	EUnit::Kilohertz,			EUnit::RevolutionsPerMinute,},
-		{ 1000,		1,			1e-3,	EUnit::Megahertz,			EUnit::Gigahertz,			},
-		{ 0.001,	1000000,	1e-3,	EUnit::Gigahertz,			EUnit::Hertz,				},
-		{ 100,		100.0/60,	1e-3,	EUnit::RevolutionsPerMinute,EUnit::Hertz,				},
+		{ 1024,		1,			1e-3,	EUnit::Bytes,										EUnit::Kilobytes,									},
+		{ 1.5,		1536,		1e-3,	EUnit::Kilobytes,									EUnit::Bytes,										},
+		{ 1000,		9.5367e-4,	1e-5,	EUnit::Megabytes,									EUnit::Terabytes,									},
+		{ 0.5,		512,		1e-3,	EUnit::Gigabytes,									EUnit::Megabytes,									},
+		{ 0.25,		256,		1e-3,	EUnit::Terabytes,									EUnit::Gigabytes,									},
 		
-		{ 1024,		1,			1e-3,	EUnit::Bytes,				EUnit::Kilobytes,			},
-		{ 1.5,		1536,		1e-3,	EUnit::Kilobytes,			EUnit::Bytes,				},
-		{ 1000,		9.5367e-4,	1e-5,	EUnit::Megabytes,			EUnit::Terabytes,			},
-		{ 0.5,		512,		1e-3,	EUnit::Gigabytes,			EUnit::Megabytes,			},
-		{ 0.25,		256,		1e-3,	EUnit::Terabytes,			EUnit::Gigabytes,			},
-		
-		{ 10000,	0.166667,	1e-6,	EUnit::Milliseconds,		EUnit::Minutes,				},
-		{ 0.5,		500,		1e-6,	EUnit::Seconds,				EUnit::Milliseconds,		},
-		{ 30,		60*30,		1e-6,	EUnit::Minutes,				EUnit::Seconds,				},
-		{ 5,		5.0/24,		1e-6,	EUnit::Hours,				EUnit::Days,				},
-		{ 0.75,		18,			1e-6,	EUnit::Days,				EUnit::Hours,				},
-		{ 3,		0.25,		1e-6,	EUnit::Months,				EUnit::Years,				},
-		{ 0.5,		6,			1e-6,	EUnit::Years,				EUnit::Months,				},
-
-		{ 22.5,		0.225,		1e-6,	EUnit::Percentage,			EUnit::Multiplier,			},
-		{ 22.5,		2250,		0.f,	EUnit::Multiplier,			EUnit::Percentage,			},
+		{ 10000,	0.166667,	1e-6,	EUnit::Milliseconds,								EUnit::Minutes,										},
+		{ 0.5,		500,		1e-6,	EUnit::Seconds,										EUnit::Milliseconds,								},
+		{ 30,		60*30,		1e-6,	EUnit::Minutes,										EUnit::Seconds,										},
+		{ 5,		5.0/24,		1e-6,	EUnit::Hours,										EUnit::Days,										},
+		{ 0.75,		18,			1e-6,	EUnit::Days,										EUnit::Hours,										},
+		{ 3,		0.25,		1e-6,	EUnit::Months,										EUnit::Years,										},
+		{ 0.5,		6,			1e-6,	EUnit::Years,										EUnit::Months,										},
+
+		{ 22.5,		0.225,		1e-6,	EUnit::Percentage,									EUnit::Multiplier,									},
+		{ 22.5,		2250,		0.f,	EUnit::Multiplier,									EUnit::Percentage,									},
 
 	};
 
diff --git a/Engine/Source/Runtime/CoreUObject/Public/UObject/NoExportTypes.h b/Engine/Source/Runtime/CoreUObject/Public/UObject/NoExportTypes.h
--- a/Engine/Source/Runtime/CoreUObject/Public/UObject/NoExportTypes.h
+++ b/Engine/Source/Runtime/CoreUObject/Public/UObject/NoExportTypes.h
@@ -375,7 +375,10 @@ enum class EUnit : uint8
 	Micrograms, Milligrams, Grams, Kilograms, MetricTons, Ounces, Pounds, Stones,
 	
 	/** Force units */
-	Newtons, PoundsForce, KilogramsForce,
+	Newtons, PoundsForce, KilogramsForce, KilogramCentimetersPerSecondSquared,
+
+	/** Torque units */
+	NewtonMeters, KilogramCentimetersSquaredPerSecondSquared,
 	
 	/** Frequency units */
 	Hertz, Kilohertz, Megahertz, Gigahertz, RevolutionsPerMinute,
