dumping another thread callstack on race detection

diff --git a/Engine/Source/Runtime/Core/Public/Misc/MTAccessDetector.h b/Engine/Source/Runtime/Core/Public/Misc/MTAccessDetector.h
--- a/Engine/Source/Runtime/Core/Public/Misc/MTAccessDetector.h
+++ b/Engine/Source/Runtime/Core/Public/Misc/MTAccessDetector.h
@@ -13,6 +13,7 @@
 #include "Containers/ContainerAllocationPolicies.h"
 #include "HAL/PlatformTLS.h"
 #include "HAL/PreprocessorHelpers.h"
+#include "HAL/PlatformStackWalk.h"
 #include "Misc/Build.h"
 #include <atomic>
 
@@ -293,7 +294,7 @@ private:
 			static_assert(sizeof(FState) == sizeof(uint64)); // `FState` is stored in `std::atomic<uint64>`
 		}
 
-		constexpr FState(uint64 InValue)
+		explicit constexpr FState(uint64 InValue)
 			: Value(InValue)
 		{
 		}
@@ -304,6 +305,11 @@ private:
 			, WriterThreadId(InWriterThreadId)
 		{
 		}
+
+		uint32 GetWriterThreadId() const
+		{
+			return WriterThreadId - 1;
+		}
 	};
 
 	// all atomic ops are relaxed to preserve the original memory order, as the detector is compiled out in non-dev builds
@@ -341,10 +347,21 @@ private:
 		if (InState.WriterNum != 0)
 		{
 			uint32 CurrentThreadId = FPlatformTLS::GetCurrentThreadId();
-			checkf(InState.WriterThreadId - 1 == CurrentThreadId, TEXT("Data race detected! Writer on thread %u while acquiring read access on thread %u"), InState.WriterThreadId, CurrentThreadId);
+			checkf(InState.GetWriterThreadId() == CurrentThreadId,
+				TEXT("Data race detected! Acquiring read access on thread %u concurrently with %d writers on thread %u:\nWriter thread %u callstack:\n%s"),
+				CurrentThreadId, InState.WriterNum, InState.GetWriterThreadId(), 
+				InState.GetWriterThreadId(), *GetThreadCallstack(InState.GetWriterThreadId()));
 		}
 	}
 
+	static FString GetThreadCallstack(uint32 ThreadId)
+	{
+		const SIZE_T StackTraceSize = 65536;
+		ANSICHAR StackTrace[StackTraceSize] = { 0 };
+		FPlatformStackWalk::ThreadStackWalkAndDump(StackTrace, StackTraceSize, 0, ThreadId);
+		return StackTrace;
+	}
+
 //////////////////////////////////////////////////////////////////////
 // DestructionSentinel
 // This access detector supports its destruction (along with its owner) while being "accessed". If this is expected, acquire/release access 
@@ -393,8 +410,8 @@ private:
 	{
 		uint32 ReaderIndex = GetReadersTls().IndexOfByPredicate([this](FReaderNum ReaderNum) { return ReaderNum.Reader == this; });
 		checkfSlow(ReaderIndex != INDEX_NONE,
-			TEXT("Invalid usage of the race detector! No matching AcquireReadAccess(): %u readers, %u writers on thread %d"),
-			LoadState().ReaderNum, LoadState().WriterNum, LoadState().WriterThreadId - 1);
+			TEXT("Invalid usage of the race detector! No matching AcquireReadAccess(): %u readers, %u writers on thread %u"),
+			LoadState().ReaderNum, LoadState().WriterNum, LoadState().GetWriterThreadId());
 		uint32 ReaderNum = --GetReadersTls()[ReaderIndex].Num;
 		if (ReaderNum == 0)
 		{
@@ -450,10 +467,10 @@ public:
 		}
 
 		checkf(LoadState().Value == ExpectedState.Value,
-			TEXT("Race detector destroyed while being accessed on another thread: %d readers, %d writers on thread %d"),
+			TEXT("Race detector destroyed while being accessed on another thread: %d readers, %d writers on thread %u:\nWriter thread %u callstack:\n%s"),
 			LoadState().ReaderNum - ExpectedState.ReaderNum, 
-			LoadState().WriterNum - ExpectedState.WriterNum, 
-			LoadState().WriterThreadId - 1);
+			LoadState().WriterNum - ExpectedState.WriterNum, LoadState().GetWriterThreadId(), 
+			LoadState().GetWriterThreadId(), *GetThreadCallstack(LoadState().GetWriterThreadId()));
 	}
 
 	FORCEINLINE void AcquireReadAccess() const
@@ -506,7 +523,7 @@ public:
 		GetDestructionSentinelStackTls().RemoveAtSwap(GetDestructionSentinelStackTls().Num() - 1);
 	}
 
-	FORCEINLINE void AcquireWriteAccess()
+	FORCEINLINE_DEBUGGABLE void AcquireWriteAccess()
 	{
 		FState LocalState = LoadState();
 
@@ -519,23 +536,25 @@ public:
 				LocalState.ReaderNum - GetReadersTls()[ReaderIndex].Num);
 		}
 
-		uint32 CurrentThreadId = FPlatformTLS::GetCurrentThreadId() + 1;
+		uint32 CurrentThreadId = FPlatformTLS::GetCurrentThreadId();
 		if (LocalState.WriterNum != 0)
 		{
-			checkf(LocalState.WriterThreadId == CurrentThreadId, 
-				TEXT("Data race detected: writer on thread %d during acquiring write access on thread %d"), 
-				LocalState.WriterThreadId - 1, CurrentThreadId - 1);
+			checkf(LocalState.GetWriterThreadId() == CurrentThreadId,
+				TEXT("Data race detected: acquiring write access on thread %u concurrently with %d writers on thread %u:\nWriter thread %u callstack:\n%s"), 
+				CurrentThreadId, LocalState.WriterNum, LocalState.GetWriterThreadId(), 
+				LocalState.GetWriterThreadId(), *GetThreadCallstack(LocalState.GetWriterThreadId()));
 		}
 
-		FState NewState{ LocalState.ReaderNum, LocalState.WriterNum + 1u, CurrentThreadId };
+		FState NewState{ LocalState.ReaderNum, LocalState.WriterNum + 1u, CurrentThreadId + 1 };
 		FState PrevState = ExchangeState(NewState);
 
 		checkf(LocalState == PrevState,
-			TEXT("Data race detected: other thread(s) activity during acquiring write access on thread %d: %u -> %u readers, %u -> %u writers on thread %d -> %d"),
-			CurrentThreadId - 1,
+			TEXT("Data race detected: other thread(s) activity during acquiring write access on thread %u: %u -> %u readers, %u -> %u writers on thread %u -> %u:\nWriter thread %u callstack:\n%s"),
+			CurrentThreadId,
 			LocalState.ReaderNum, PrevState.ReaderNum,
 			LocalState.WriterNum, PrevState.WriterNum,
-			LocalState.WriterThreadId - 1, PrevState.WriterThreadId - 1);
+			LocalState.GetWriterThreadId(), PrevState.GetWriterThreadId(), 
+			PrevState.GetWriterThreadId(), *GetThreadCallstack(PrevState.GetWriterThreadId()));
 	}
 
 	// an overload that handles access detector destruction from inside a write access, must be used along with the corresponding
@@ -553,19 +572,20 @@ public:
 		FState LocalState = LoadState();
 
 		uint32 WriterThreadId = LocalState.WriterNum != 1 ? LocalState.WriterThreadId : InvalidThreadId;
-		FState NewState{ LocalState.ReaderNum, LocalState.WriterNum - 1u,WriterThreadId };
+		FState NewState{ LocalState.ReaderNum, LocalState.WriterNum - 1u, WriterThreadId };
 		FState PrevState = ExchangeState(NewState);
 
 		checkf(LocalState == PrevState,
-			TEXT("Data race detected: other thread(s) activity during releasing write access on thread %d: %u -> %u readers, %u -> %u writers on thread %d -> %d"),
+			TEXT("Data race detected: other thread(s) activity during releasing write access on thread %d: %u -> %u readers, %u -> %u writers on thread %u -> %u\nWriter thread %u callstack:\n%s"),
 			LocalState.ReaderNum, PrevState.ReaderNum,
 			LocalState.WriterNum, PrevState.WriterNum,
-			LocalState.WriterThreadId - 1, PrevState.WriterThreadId - 1);
+			LocalState.GetWriterThreadId(), PrevState.GetWriterThreadId(), 
+			PrevState.GetWriterThreadId(), *GetThreadCallstack(PrevState.GetWriterThreadId()));
 	}
 
 	// an overload that handles access detector destruction from inside a write access, must be used along with the corresponding
 	// overload of `AcquireWriteAcess`
-	FORCEINLINE_DEBUGGABLE void ReleaseWriteAccess(FDestructionSentinel& DestructionSentinel)
+	FORCEINLINE void ReleaseWriteAccess(FDestructionSentinel& DestructionSentinel)
 	{
 		ReleaseWriteAccess();
 
