OodleDataCompression : make CompressParallel output not be machine dependent

#rnx

diff --git a/Engine/Source/Runtime/Core/Private/Compression/OodleDataCompression.cpp b/Engine/Source/Runtime/Core/Private/Compression/OodleDataCompression.cpp
--- a/Engine/Source/Runtime/Core/Private/Compression/OodleDataCompression.cpp
+++ b/Engine/Source/Runtime/Core/Private/Compression/OodleDataCompression.cpp
@@ -613,9 +613,9 @@ int64 CompressParallelSub(
 							ECompressor Compressor,	ECompressionLevel Level,
 							bool CompressIndependentChunks)
 {
-	// beware: doing the split based on NumWorkers means that compressed output depends on the machine
-
-	int32 NumWorkers = FTaskGraphInterface::Get().GetNumWorkerThreads();
+	// beware: doing the split based on NumWorkers means that compressed output depends on the machine (if ! CompressIndependentChunks)
+	//int32 NumWorkers = FTaskGraphInterface::Get().GetNumWorkerThreads();
+	int32 NumWorkers = 128; // hard coded so that output is the same on all machines
 	int64 ChunkLen = OodleLZ_MakeSeekChunkLen(UncompressedSize,NumWorkers);
 	int64 NumChunks = (UncompressedSize + ChunkLen-1)/ChunkLen;
 
@@ -651,7 +651,7 @@ int64 CompressParallelSub(
 		check( Ret >= 0 );
 
 		OutCompressedArray.SetNum(Ret,false);
-	} );
+	} , EParallelForFlags::Unbalanced );
 
 	int64 Total = 0;
 	for(int64 i=0;i<NumChunks;i++)
@@ -782,13 +782,13 @@ bool DecompressParallel(
 	// this function assumes you either have no  seek resets at all
 	//	or seek resets at every OODLELZ_BLOCK_LEN
 	// therefore we are free to choose our decode seek chunk len here freely
-	// if that was no the case
+	// if that was not the case
 	//	(eg. say if you did seek resets at OODLELZ_BLOCK_LEN*4)
 	// then we would have to change this logic to instead scan the comp data
 	//	and find the seek resets, using OodleLZ_GetCompressedStepForRawStep
 
 	int32 NumWorkers = FTaskGraphInterface::Get().GetNumWorkerThreads();
-	int32 ChunkLen = OodleLZ_MakeSeekChunkLen(UncompressedSize,NumWorkers*4);
+	int32 ChunkLen = OodleLZ_MakeSeekChunkLen(UncompressedSize,NumWorkers*3);
 	int64 NumChunks = (UncompressedSize + ChunkLen-1)/ChunkLen;
 
 	int64 SeekTableMemSize = OodleLZ_GetSeekTableMemorySizeNeeded((int32)NumChunks,OodleLZSeekTable_Flags_None);
@@ -832,7 +832,7 @@ bool DecompressParallel(
 			// no need for atomics :
 			AnyFailed = true;
 		}
-	} );
+	}, EParallelForFlags::Unbalanced );
 
 	FMemory::Free(SeekTableMem);
 
