Deprecated hashing arrays and TCHAR*.

diff --git a/Engine/Source/Developer/TraceInsights/Private/Insights/Table/ViewModels/TreeNodeGrouping.cpp b/Engine/Source/Developer/TraceInsights/Private/Insights/Table/ViewModels/TreeNodeGrouping.cpp
--- a/Engine/Source/Developer/TraceInsights/Private/Insights/Table/ViewModels/TreeNodeGrouping.cpp
+++ b/Engine/Source/Developer/TraceInsights/Private/Insights/Table/ViewModels/TreeNodeGrouping.cpp
@@ -166,7 +166,7 @@ void FTreeNodeGroupingByUniqueValueCString::GroupNodes(const TArray<FTableTreeNo
 {
 	FTableTreeNodePtr UnsetGroupPtr = nullptr; // for unset FTableCellValue
 	FTableTreeNodePtr EmptyGroupPtr = nullptr; // for nullptr and empty strings
-	TMap<const TCHAR*, FTableTreeNodePtr> GroupMap; // for valid strings
+	TMap<const TCHAR*, FTableTreeNodePtr, FDefaultSetAllocator, TStringPointerMapKeyFuncs_DEPRECATED<const TCHAR*, FTableTreeNodePtr>> GroupMap; // for valid strings
 	TMap<FName, FTableTreeNodePtr> GroupNameMap; // for valid strings
 
 	ParentGroup.ClearChildren();
diff --git a/Engine/Source/Runtime/Core/Public/Templates/TypeHash.h b/Engine/Source/Runtime/Core/Public/Templates/TypeHash.h
--- a/Engine/Source/Runtime/Core/Public/Templates/TypeHash.h
+++ b/Engine/Source/Runtime/Core/Public/Templates/TypeHash.h
@@ -151,7 +151,7 @@ template <
 	uint32 N,
 	std::enable_if_t<!std::is_same_v<const T, const TCHAR>>* = nullptr
 >
-//UE_DEPRECATED(all, "Hashing arrays is deprecated - use PointerHash() instead to force a conversion to a pointer or GetArrayHash to hash the array contents")
+UE_DEPRECATED(all, "Hashing arrays is deprecated - use PointerHash() instead to force a conversion to a pointer or GetArrayHash to hash the array contents")
 inline uint32 GetTypeHash(T (&Array)[N])
 {
 	return PointerHash(Array);
@@ -161,7 +161,7 @@ template <
 	typename T,
 	std::enable_if_t<std::is_same_v<const T, const TCHAR>>* = nullptr
 >
-//UE_DEPRECATED(5.3, "Hashing TCHAR arrays is deprecated - use PointerHash() to force a conversion to a pointer or FCrc::Strihash_DEPRECATED to do a string hash, or use TStringPointerSetKeyFuncs_DEPRECATED or TStringPointerMapKeyFuncs_DEPRECATED as keyfuncs for a TSet or TMap respectively")
+UE_DEPRECATED(5.3, "Hashing TCHAR arrays is deprecated - use PointerHash() to force a conversion to a pointer or FCrc::Strihash_DEPRECATED to do a string hash, or use TStringPointerSetKeyFuncs_DEPRECATED or TStringPointerMapKeyFuncs_DEPRECATED as keyfuncs for a TSet or TMap respectively")
 inline uint32 GetTypeHash(T* Value)
 {
 	// Hashing a TCHAR* array differently from a void* is dangerous and is deprecated.
diff --git a/Engine/Source/Runtime/RenderCore/Private/ShaderParameters.cpp b/Engine/Source/Runtime/RenderCore/Private/ShaderParameters.cpp
--- a/Engine/Source/Runtime/RenderCore/Private/ShaderParameters.cpp
+++ b/Engine/Source/Runtime/RenderCore/Private/ShaderParameters.cpp
@@ -399,7 +399,7 @@ void UE::ShaderParameters::AddUniformBufferIncludesToEnvironment(FShaderCompiler
 	for (const TCHAR* UniformBufferName : InUniformBufferNames)
 	{
 		FStringView UniformBufferNameView(UniformBufferName);
-		if (!OutEnvironment.UniformBufferMap.FindByHash(GetTypeHash(UniformBufferName), UniformBufferName))
+		if (!OutEnvironment.UniformBufferMap.FindByHash(FCrc::Strihash_DEPRECATED(UniformBufferName), UniformBufferName))
 		{
 			if (const FShaderParametersMetadata* Metadata = FindShaderParametersMetadataWithVariableName(UniformBufferName))
 			{
