Added UE_FORCE_CONSTEVAL macro to cause a constexpr expression to be evaluated at compile time when used as a subexpression of a runtime expression.
Added some usage of the macro around FMath::ConstExprCeilLogTwo.

diff --git a/Engine/Source/Runtime/Core/Public/Templates/UnrealTemplate.h b/Engine/Source/Runtime/Core/Public/Templates/UnrealTemplate.h
--- a/Engine/Source/Runtime/Core/Public/Templates/UnrealTemplate.h
+++ b/Engine/Source/Runtime/Core/Public/Templates/UnrealTemplate.h
@@ -209,6 +209,26 @@ OutType FloatCastChecked(InType In, InType Precision)
 	#define VTABLE_OFFSET( Class, MultipleInheritenceParent )	( ((PTRINT) static_cast<MultipleInheritenceParent*>((Class*)1)) - 1)
 #endif
 
+namespace UE::Core::Private
+{
+	template <typename T, T Val>
+	constexpr T TForceConstEval_V = Val;
+}
+
+// Forces an expression to be evaluated at compile-time, even if it is part of a runtime expression:
+//
+// Example:
+//   // Arg 3 is evaluated at runtime as it's used in a runtime context, despite the function being marked constexpr and having a compile-time argument.
+//   // Requires an optimizer pass to eliminate.
+//   RegisterTypeWithSizeAndLog2Alignment("MyType", sizeof(FMyType), FMath::ConstExprCeilLogTwo(alignof(FMyType)));
+//
+//   // Arg 3 is evaluated at compile-time, but is non-intuitive and requires another variable to be introduced
+//   constexpr SIZE_T AlignOfMyTypeLog2 = alignof(FMyType);
+//   RegisterTypeWithSizeAndLog2Alignment("MyType", sizeof(FMyType), AlignOfMyTypeLog2);
+//
+//   // Arg 3 is evaluated at compile time with a more direct syntax
+//   RegisterTypeWithSizeAndLog2Alignment("MyType", sizeof(FMyType), UE_FORCE_CONSTEVAL(FMath::ConstExprCeilLogTwo(alignof(FMyType))));
+#define UE_FORCE_CONSTEVAL(expr) UE::Core::Private::TForceConstEval_V<std::decay_t<decltype(expr)>, (expr)>
 
 /**
  * works just like std::min_element.
diff --git a/Engine/Source/Runtime/CoreUObject/Public/UObject/DynamicallyTypedValue.h b/Engine/Source/Runtime/CoreUObject/Public/UObject/DynamicallyTypedValue.h
--- a/Engine/Source/Runtime/CoreUObject/Public/UObject/DynamicallyTypedValue.h
+++ b/Engine/Source/Runtime/CoreUObject/Public/UObject/DynamicallyTypedValue.h
@@ -156,9 +156,8 @@ namespace UE
 		// Whether the value's data is stored in InlineData or in the memory pointed to by HeapData.
 		bool IsInline() const
 		{
-			constexpr uint8 InlineDataMinAlignmentLogTwo = FMath::ConstExprCeilLogTwo(alignof(UPTRINT));
 			return Type->GetNumBytes() <= sizeof(UPTRINT)
-				&& Type->GetMinAlignmentLogTwo() <= InlineDataMinAlignmentLogTwo;
+				&& Type->GetMinAlignmentLogTwo() <= UE_FORCE_CONSTEVAL(FMath::ConstExprCeilLogTwo(alignof(UPTRINT)));
 		}
 
 		// Allocates heap memory for the value if it uses it.
