Slight improvement to GetEntryLength

#android

diff --git a/Engine/Source/Runtime/Core/Private/Android/AndroidPlatformFile.cpp b/Engine/Source/Runtime/Core/Private/Android/AndroidPlatformFile.cpp
--- a/Engine/Source/Runtime/Core/Private/Android/AndroidPlatformFile.cpp
+++ b/Engine/Source/Runtime/Core/Private/Android/AndroidPlatformFile.cpp
@@ -898,7 +898,7 @@ public:
 
 	int64 GetEntryLength(const FString & Path)
 	{
-		TSharedPtr<FFileHandleAndroid> File = Entries[Path]->File;
+		TSharedPtr<FFileHandleAndroid>& File = Entries[Path]->File;
 		return File != nullptr ? File->Size() : 0;
 	}
 
