removed the invalid unit-test deleting a multicast delegate from inside its broadcast, as this is not supported.
caught by ASAN: https://horde.devtools.epicgames.com/log/6478ba47e25209b6cb8ee2f6?lineindex=30723

diff --git a/Engine/Source/Runtime/Core/Tests/Delegates/TSDelegatesTest.cpp b/Engine/Source/Runtime/Core/Tests/Delegates/TSDelegatesTest.cpp
--- a/Engine/Source/Runtime/Core/Tests/Delegates/TSDelegatesTest.cpp
+++ b/Engine/Source/Runtime/Core/Tests/Delegates/TSDelegatesTest.cpp
@@ -156,12 +156,6 @@ namespace UE { namespace TSDelegatesTest
 			Delegate->Execute();
 		}
 
-		{	// multicast delegate destruction
-			TMulticastDelegate<void()>* Delegate = new TMulticastDelegate<void()>;
-			Delegate->AddLambda([Delegate] { delete Delegate; });
-			Delegate->Broadcast();
-		}
-
 		{
 			struct FDelegateHolder
 			{
