Improved Clamp() codegen.

diff --git a/Engine/Source/Runtime/Core/Public/Math/UnrealMathUtility.h b/Engine/Source/Runtime/Core/Public/Math/UnrealMathUtility.h
--- a/Engine/Source/Runtime/Core/Public/Math/UnrealMathUtility.h
+++ b/Engine/Source/Runtime/Core/Public/Math/UnrealMathUtility.h
@@ -590,9 +590,9 @@ public:
 
 	/** Clamps X to be between Min and Max, inclusive */
 	template< class T >
-	UE_NODISCARD static constexpr FORCEINLINE T Clamp(const T X, const T Min, const T Max)
+	UE_NODISCARD static constexpr FORCEINLINE T Clamp(const T X, const T MinValue, const T MaxValue)
 	{
-		return (X < Min) ? Min : (X < Max) ? X : Max;
+		return Max(Min(X, MaxValue), MinValue);
 	}
 	/** Allow mixing float/double arguments, promoting to highest precision type. */
 	MIX_FLOATS_3_ARGS(Clamp);
