added ExecuteOnGameThread() as a wrapper for FTSTicker::GetCoreTicker()

diff --git a/Engine/Source/Runtime/Core/Public/Containers/Ticker.h b/Engine/Source/Runtime/Core/Public/Containers/Ticker.h
--- a/Engine/Source/Runtime/Core/Public/Containers/Ticker.h
+++ b/Engine/Source/Runtime/Core/Public/Containers/Ticker.h
@@ -159,3 +159,17 @@ private:
 	/** Delegate for callbacks to Tick */
 	FTSTicker::FDelegateHandle TickHandle;
 };
+
+// Schedules execution of the given functor on the game thread at a particular moment of a frame when interference with other systems 
+// is minimal.
+template<typename FunctorType>
+void ExecuteOnGameThread(const TCHAR* DebugName, FunctorType&& Functor)
+{
+	FTSTicker::GetCoreTicker().AddTicker(DebugName, 0.0f,
+		[Functor = MoveTemp(Functor)](float)
+		{
+			Functor();
+			return false;
+		}
+	);
+}
\ No newline at end of file
diff --git a/Engine/Source/Runtime/Core/Tests/Containers/TickerTests.cpp b/Engine/Source/Runtime/Core/Tests/Containers/TickerTests.cpp
--- a/Engine/Source/Runtime/Core/Tests/Containers/TickerTests.cpp
+++ b/Engine/Source/Runtime/Core/Tests/Containers/TickerTests.cpp
@@ -198,4 +198,17 @@ TEST_CASE_NAMED(FTSTickerTest,"System::Core::Containers::TSTicker", "[Applicatio
 
 }
 
+TEST_CASE_NAMED(ExecuteOnGameThreadTest, "System::Core::Containers::TSTicker::ExecuteOnGameThreadTest", "[ApplicationContextMask][EngineFilter]")
+{
+	uint32 NumExecuted = 0;
+	ExecuteOnGameThread(UE_SOURCE_LOCATION, [&NumExecuted] { ++NumExecuted; });
+	check(NumExecuted == 0);
+
+	FTSTicker::GetCoreTicker().Tick(0.0f); // the command above is executed by the core ticker
+	check(NumExecuted == 1);
+
+	FTSTicker::GetCoreTicker().Tick(0.0f); // tick again to check that the command was executed only once
+	check(NumExecuted == 1);
+}
+
 #endif // WITH_TESTS
\ No newline at end of file
