AssetTools.LogFolderPermissions command that logs the read and write folder permission lists

diff --git a/Engine/Source/Developer/AssetTools/Private/AssetTools.cpp b/Engine/Source/Developer/AssetTools/Private/AssetTools.cpp
--- a/Engine/Source/Developer/AssetTools/Private/AssetTools.cpp
+++ b/Engine/Source/Developer/AssetTools/Private/AssetTools.cpp
@@ -432,6 +432,16 @@ namespace UE::AssetTools::Private
 		TEXT("When set, if you import an asset at a location with a redirector, you'll instead import to the redirector's destination location")
 	);
 
+	FAutoConsoleCommand LogFolderPermissionsCommand = FAutoConsoleCommand(
+		TEXT("AssetTools.LogFolderPermissions"),
+		TEXT("Logs the read and write permissions for folders"),
+		FConsoleCommandWithArgsDelegate::CreateLambda(
+			[](const TArray<FString>& Args)
+			{
+				UE_LOG(LogAssetTools, Log, TEXT("Folder Read Permissions:\n%s"), *IAssetTools::Get().GetFolderPermissionList()->ToString());
+				UE_LOG(LogAssetTools, Log, TEXT("Folder Write Permissions:\n%s"), *IAssetTools::Get().GetWritableFolderPermissionList()->ToString());
+			}));
+
 	// use a struct as a namespace to allow easier friend declarations
 	struct FPackageMigrationImpl
 	{
diff --git a/Engine/Source/Runtime/Core/Private/Misc/PathPermissionList.cpp b/Engine/Source/Runtime/Core/Private/Misc/PathPermissionList.cpp
--- a/Engine/Source/Runtime/Core/Private/Misc/PathPermissionList.cpp
+++ b/Engine/Source/Runtime/Core/Private/Misc/PathPermissionList.cpp
@@ -477,3 +477,67 @@ void FPathPermissionList::VerifyItemMatchesListType(const FStringView Item) cons
 		ensureAlwaysMsgf(IsClassPathNameOrNone(Item), TEXT("Short class name \"%.*s\" provided for PathPermissionList representing class paths"), Item.Len(), Item.GetData());
 	}
 }
+
+FString FPathPermissionList::ToString() const
+{
+	TStringBuilder<4096> StringBuilder;
+
+	auto SortAndAppendOwners = [&StringBuilder](FPermissionListOwners& Owners)
+	{
+		Owners.Sort(FNameLexicalLess());
+
+		StringBuilder.AppendChar(TCHAR('('));
+		bool bFirst = true;
+		for (FName Owner : Owners)
+		{
+			if (bFirst)
+			{
+				bFirst = false;
+			}
+			else
+			{
+				StringBuilder.Append(TEXT(", "));
+			}
+			StringBuilder.Append(Owner.ToString());
+		}
+		StringBuilder.AppendChar(TCHAR(')'));
+	};
+
+	if (!DenyListAll.IsEmpty())
+	{
+		StringBuilder.Append(TEXT("Deny All "));
+		FPermissionListOwners SortedDenyListAll = DenyListAll;
+		SortAndAppendOwners(SortedDenyListAll);
+		StringBuilder.Append(TEXT("\n"));;
+	}
+
+	auto AppendList = [&StringBuilder, &SortAndAppendOwners](const TMap<FString, FPermissionListOwners>& List)
+	{
+		TMap<FString, FPermissionListOwners> SortedList = List;
+		SortedList.KeySort(TLess<FString>());
+		for (TPair<FString, FPermissionListOwners>& ListEntry : SortedList)
+		{
+			StringBuilder.AppendChar(TCHAR('\t'));
+			StringBuilder.AppendChar(TCHAR('"'));
+			StringBuilder.Append(ListEntry.Key);
+			StringBuilder.AppendChar(TCHAR('"'));
+			StringBuilder.AppendChar(TCHAR(' '));
+			SortAndAppendOwners(ListEntry.Value);
+			StringBuilder.AppendChar(TCHAR('\n'));
+		}
+	};
+
+	if (!DenyList.IsEmpty())
+	{
+		StringBuilder.Append(TEXT("Deny List\n"));
+		AppendList(DenyList);
+	}
+
+	if (!AllowList.IsEmpty())
+	{
+		StringBuilder.Append(TEXT("Allow List\n"));
+		AppendList(AllowList);
+	}
+
+	return StringBuilder.ToString();
+}
\ No newline at end of file
diff --git a/Engine/Source/Runtime/Core/Public/Misc/NamePermissionList.h b/Engine/Source/Runtime/Core/Public/Misc/NamePermissionList.h
--- a/Engine/Source/Runtime/Core/Public/Misc/NamePermissionList.h
+++ b/Engine/Source/Runtime/Core/Public/Misc/NamePermissionList.h
@@ -277,6 +277,9 @@ public:
 	/** Triggered when filter changes */
 	FSimpleMulticastDelegate& OnFilterChanged() { return OnFilterChangedDelegate; }
 
+	/** Dumps the path permission list details into a multi-line string */
+	FString ToString() const;
+
 protected:
 
 	/**
