Use LogConsoleResponse when running Localization.DumpLiveTable so that the result appears in the client console window

#rnx

diff --git a/Engine/Source/Runtime/Core/Private/Internationalization/TextLocalizationManager.cpp b/Engine/Source/Runtime/Core/Private/Internationalization/TextLocalizationManager.cpp
--- a/Engine/Source/Runtime/Core/Private/Internationalization/TextLocalizationManager.cpp
+++ b/Engine/Source/Runtime/Core/Private/Internationalization/TextLocalizationManager.cpp
@@ -82,7 +82,16 @@ static FAutoConsoleCommand CmdDumpLiveTable(
 			}
 		}
 
-		FTextLocalizationManager::Get().DumpLiveTable(NamespaceFilter.GetPtrOrNull(), KeyFilter.GetPtrOrNull(), DisplayStringFilter.GetPtrOrNull());
+		auto GetConsoleResponseLogCategoryPtr = []() -> const FLogCategoryBase*
+		{
+#if NO_LOGGING
+			return nullptr;
+#else
+			return &LogConsoleResponse;
+#endif
+		};
+
+		FTextLocalizationManager::Get().DumpLiveTable(NamespaceFilter.GetPtrOrNull(), KeyFilter.GetPtrOrNull(), DisplayStringFilter.GetPtrOrNull(), GetConsoleResponseLogCategoryPtr());
 	}));
 #endif
 }
@@ -567,14 +576,17 @@ void FTextLocalizationManager::CompactDataStructures()
 }
 
 #if ENABLE_LOC_TESTING
-void FTextLocalizationManager::DumpLiveTable(const FString* NamespaceFilter, const FString* KeyFilter, const FString* DisplayStringFilter) const
+void FTextLocalizationManager::DumpLiveTable(const FString* NamespaceFilter, const FString* KeyFilter, const FString* DisplayStringFilter, const FLogCategoryBase* CategoryOverride) const
 {
+#if !NO_LOGGING
+	const FLogCategoryBase& Category = CategoryOverride ? *CategoryOverride : LogLocalization;
+
 	auto PassesFilter = [](const FString& Str, const FString* Filter)
 	{
 		return !Filter || Str.MatchesWildcard(*Filter, ESearchCase::IgnoreCase); // Note: This is case insensitive since its used from a debug command
 	};
 
-	UE_LOG(LogLocalization, Display, TEXT("----------------------------------------------------------------------"));
+	UE_LOG_REF(Category, Display, TEXT("----------------------------------------------------------------------"));
 
 	FScopeLock ScopeLock(&DisplayStringLookupTableCS);
 	for (const auto& DisplayStringPair : DisplayStringLookupTable)
@@ -583,11 +595,12 @@ void FTextLocalizationManager::DumpLiveTable(const FString* NamespaceFilter, con
 			PassesFilter(DisplayStringPair.Key.GetKey().GetChars(), KeyFilter) &&
 			PassesFilter(**DisplayStringPair.Value.DisplayString, DisplayStringFilter))
 		{
-			UE_LOG(LogLocalization, Display, TEXT("LiveTableEntry: Namespace: '%s', Key: '%s', DisplayString: '%s'"), DisplayStringPair.Key.GetNamespace().GetChars(), DisplayStringPair.Key.GetKey().GetChars(), **DisplayStringPair.Value.DisplayString);
+			UE_LOG_REF(Category, Display, TEXT("LiveTableEntry: Namespace: '%s', Key: '%s', DisplayString: '%s'"), DisplayStringPair.Key.GetNamespace().GetChars(), DisplayStringPair.Key.GetKey().GetChars(), **DisplayStringPair.Value.DisplayString);
 		}
 	}
 
-	UE_LOG(LogLocalization, Display, TEXT("----------------------------------------------------------------------"));
+	UE_LOG_REF(Category, Display, TEXT("----------------------------------------------------------------------"));
+#endif // !NO_LOGGING
 }
 #endif
 
diff --git a/Engine/Source/Runtime/Core/Public/Internationalization/TextLocalizationManager.h b/Engine/Source/Runtime/Core/Public/Internationalization/TextLocalizationManager.h
--- a/Engine/Source/Runtime/Core/Public/Internationalization/TextLocalizationManager.h
+++ b/Engine/Source/Runtime/Core/Public/Internationalization/TextLocalizationManager.h
@@ -26,6 +26,7 @@
 class FTextLocalizationResource;
 class ILocalizedTextSource;
 class IPakFile;
+struct FLogCategoryBase;
 struct FPolyglotTextData;
 
 enum class ETextLocalizationManagerInitializedFlags : uint8
@@ -120,7 +121,7 @@ public:
 	 * Dumps the current live table state to the log, optionally filtering it based on the given wildcard arguments.
 	 * @note Calling this function with no filters specified will dump the entire live table.
 	 */
-	void DumpLiveTable(const FString* NamespaceFilter = nullptr, const FString* KeyFilter = nullptr, const FString* DisplayStringFilter = nullptr) const;
+	void DumpLiveTable(const FString* NamespaceFilter = nullptr, const FString* KeyFilter = nullptr, const FString* DisplayStringFilter = nullptr, const FLogCategoryBase* CategoryOverride = nullptr) const;
 #endif
 
 	/**
