Additional crash reporting information.
- Whether the binaries were created with debug info.
- Engine 'compatible' version from build info.
- CIS job URL that created the binaries.
- Currently executing CIS job URL.

diff --git a/Engine/Source/Programs/AutomationTool/AutomationUtils/UnrealBuild.cs b/Engine/Source/Programs/AutomationTool/AutomationUtils/UnrealBuild.cs
--- a/Engine/Source/Programs/AutomationTool/AutomationUtils/UnrealBuild.cs
+++ b/Engine/Source/Programs/AutomationTool/AutomationUtils/UnrealBuild.cs
@@ -290,7 +290,7 @@ namespace AutomationTool
 		/// <summary>
 		/// Updates the engine version files
 		/// </summary>
-		public List<FileReference> UpdateVersionFiles(bool ActuallyUpdateVersionFiles = true, int? ChangelistNumberOverride = null, int? CompatibleChangelistNumberOverride = null, string Build = null, bool? IsPromotedOverride = null)
+		public List<FileReference> UpdateVersionFiles(bool ActuallyUpdateVersionFiles = true, int? ChangelistNumberOverride = null, int? CompatibleChangelistNumberOverride = null, string Build = null, string BuildURL = null, bool? IsPromotedOverride = null)
 		{
 			bool bIsLicenseeVersion = ParseParam("Licensee") || !FileReference.Exists(FileReference.Combine(Unreal.EngineDirectory, "Restricted", "NotForLicensees", "Build", "EpicInternal.txt"));
 			bool bIsPromotedBuild = IsPromotedOverride.HasValue? IsPromotedOverride.Value : (ParseParamInt("Promoted", 1) != 0);
@@ -312,10 +312,10 @@ namespace AutomationTool
 				Branch = CommandUtils.P4Enabled ? CommandUtils.EscapePath(CommandUtils.P4Env.Branch) : "";
 			}
 
-			return StaticUpdateVersionFiles(ChangelistNumber, CompatibleChangelistNumber, Branch, Build, bIsLicenseeVersion, bIsPromotedBuild, bDoUpdateVersionFiles);
+			return StaticUpdateVersionFiles(ChangelistNumber, CompatibleChangelistNumber, Branch, Build, BuildURL, bIsLicenseeVersion, bIsPromotedBuild, bDoUpdateVersionFiles);
 		}
 
-		public static List<FileReference> StaticUpdateVersionFiles(int ChangelistNumber, int CompatibleChangelistNumber, string Branch, string Build, bool bIsLicenseeVersion, bool bIsPromotedBuild, bool bDoUpdateVersionFiles)
+		public static List<FileReference> StaticUpdateVersionFiles(int ChangelistNumber, int CompatibleChangelistNumber, string Branch, string Build, string BuildURL, bool bIsLicenseeVersion, bool bIsPromotedBuild, bool bDoUpdateVersionFiles)
 		{
 			FileReference BuildVersionFile = BuildVersion.GetDefaultFileName();
 
@@ -342,6 +342,7 @@ namespace AutomationTool
 					Logger.LogDebug("  IsPromotedBuild={Arg0}", bIsPromotedBuild? 1 : 0);
 					Logger.LogDebug("  BranchName={Branch}", Branch);
 					Logger.LogDebug("  BuildVersion={Build}", Build);
+					Logger.LogDebug("  BuildURL={BuildURL}", BuildURL);
 
 					Version.Changelist = ChangelistNumber;
 					if(CompatibleChangelistNumber > 0)
@@ -356,6 +357,7 @@ namespace AutomationTool
 					Version.IsPromotedBuild = bIsPromotedBuild;
 					Version.BranchName = Branch;
 					Version.BuildVersionString = Build;
+					Version.BuildURL = BuildURL;
 
 					if (File.Exists(BuildVersionFile.FullName))
 					{
diff --git a/Engine/Source/Programs/AutomationTool/BuildGraph/Tasks/SetVersionTask.cs b/Engine/Source/Programs/AutomationTool/BuildGraph/Tasks/SetVersionTask.cs
--- a/Engine/Source/Programs/AutomationTool/BuildGraph/Tasks/SetVersionTask.cs
+++ b/Engine/Source/Programs/AutomationTool/BuildGraph/Tasks/SetVersionTask.cs
@@ -42,6 +42,12 @@ namespace AutomationTool.Tasks
 		[TaskParameter(Optional = true)]
 		public string Build;
 
+		/// <summary>
+		/// The URL for a running continuous integration job.
+		/// </summary>
+		[TaskParameter(Optional = true)]
+		public string BuildURL;
+
 		/// <summary>
 		/// Whether to set the IS_LICENSEE_VERSION flag to true.
 		/// </summary>
@@ -96,7 +102,7 @@ namespace AutomationTool.Tasks
 		public override Task ExecuteAsync(JobContext Job, HashSet<FileReference> BuildProducts, Dictionary<string, HashSet<FileReference>> TagNameToFileSet)
 		{
 			// Update the version files
-			List<FileReference> VersionFiles = UnrealBuild.StaticUpdateVersionFiles(Parameters.Change, Parameters.CompatibleChange, Parameters.Branch, Parameters.Build, Parameters.Licensee, Parameters.Promoted, !Parameters.SkipWrite);
+			List<FileReference> VersionFiles = UnrealBuild.StaticUpdateVersionFiles(Parameters.Change, Parameters.CompatibleChange, Parameters.Branch, Parameters.Build, Parameters.BuildURL, Parameters.Licensee, Parameters.Promoted, !Parameters.SkipWrite);
 
 			// Apply the optional tag to them
 			foreach(string TagName in FindTagNamesFromList(Parameters.Tag))
@@ -156,13 +162,14 @@ namespace AutomationTool.Tasks
 		/// <summary>
 		/// Updates the current engine version
 		/// </summary>
-		public static async Task<FileSet> SetVersionAsync(int Change, string Branch, int? CompatibleChange = null, string Build = null, bool? Licensee = null, bool? Promoted = null, bool? SkipWrite = null)
+		public static async Task<FileSet> SetVersionAsync(int Change, string Branch, int? CompatibleChange = null, string Build = null, string BuildURL = null, bool? Licensee = null, bool? Promoted = null, bool? SkipWrite = null)
 		{
 			SetVersionTaskParameters Parameters = new SetVersionTaskParameters();
 			Parameters.Change = Change;
 			Parameters.CompatibleChange = CompatibleChange ?? Parameters.CompatibleChange;
 			Parameters.Branch = Branch ?? Parameters.Branch;
-			Parameters.Branch = Build ?? Parameters.Build;
+			Parameters.Build = Build ?? Parameters.Build;
+			Parameters.BuildURL = BuildURL ?? Parameters.BuildURL;
 			Parameters.Licensee = Licensee ?? Parameters.Licensee;
 			Parameters.Promoted = Promoted ?? Parameters.Promoted;
 			Parameters.SkipWrite = SkipWrite ?? Parameters.SkipWrite;
diff --git a/Engine/Source/Programs/UnrealBuildTool/Configuration/UEBuildModuleCPP.cs b/Engine/Source/Programs/UnrealBuildTool/Configuration/UEBuildModuleCPP.cs
--- a/Engine/Source/Programs/UnrealBuildTool/Configuration/UEBuildModuleCPP.cs
+++ b/Engine/Source/Programs/UnrealBuildTool/Configuration/UEBuildModuleCPP.cs
@@ -1851,6 +1851,11 @@ namespace UnrealBuildTool
 			AddDefinitions(Result.Definitions, PublicDefinitions);
 
 			Result.Definitions.AddRange(Rules.PrivateDefinitions);
+			
+			if (Rules.Name == "BuildSettings")
+			{
+				Result.Definitions.Add(String.Format("UE_WITH_DEBUG_INFO={0}", BaseCompileEnvironment.bCreateDebugInfo ? "1" : "0"));
+			}
 
 			// Add the project definitions
 			if (!Rules.bTreatAsEngineModule)
diff --git a/Engine/Source/Programs/UnrealBuildTool/System/BuildVersion.cs b/Engine/Source/Programs/UnrealBuildTool/System/BuildVersion.cs
--- a/Engine/Source/Programs/UnrealBuildTool/System/BuildVersion.cs
+++ b/Engine/Source/Programs/UnrealBuildTool/System/BuildVersion.cs
@@ -65,6 +65,11 @@ namespace UnrealBuildTool
 		/// </summary>
 		public string? BuildVersionString;
 
+		/// <summary>
+		/// Optional URL for a continuous integration job associated with this build version. (e.g. the job that build a set of binaries)
+		/// </summary>
+		public string? BuildURL;
+
 		/// <summary>
 		/// Returns the value which can be used as the compatible changelist. Requires that the regular changelist is also set, and defaults to the 
 		/// regular changelist if a specific compatible changelist is not set.
@@ -155,6 +160,7 @@ namespace UnrealBuildTool
 			Object.TryGetStringField("BranchName", out NewVersion.BranchName);
 			Object.TryGetStringField("BuildId", out NewVersion.BuildId);
 			Object.TryGetStringField("BuildVersion", out NewVersion.BuildVersionString);
+			Object.TryGetStringField("BuildURL", out NewVersion.BuildURL);
 
 			Version = NewVersion;
 			return true;
@@ -221,6 +227,10 @@ namespace UnrealBuildTool
 			{
 				Writer.WriteValue("BuildVersion", BuildVersionString);
 			}
+			if (!String.IsNullOrEmpty(BuildURL))
+			{
+				Writer.WriteValue("BuildURL", BuildURL);
+			}
 		}
 	}
 
@@ -300,6 +310,7 @@ namespace UnrealBuildTool
 		public string? BranchName => Inner.BranchName;
 
 		public string? BuildVersionString => Inner.BuildVersionString;
+		public string? BuildURL => Inner.BuildURL;
 
 #pragma warning restore C1591
 		#endregion
diff --git a/Engine/Source/Runtime/BuildSettings/BuildSettings.Build.cs b/Engine/Source/Runtime/BuildSettings/BuildSettings.Build.cs
--- a/Engine/Source/Runtime/BuildSettings/BuildSettings.Build.cs
+++ b/Engine/Source/Runtime/BuildSettings/BuildSettings.Build.cs
@@ -37,7 +37,11 @@ public class BuildSettings : ModuleRules
 		PrivateDefinitions.Add(string.Format("BRANCH_NAME=\"{0}\"", Target.Version.BranchName));
 
 		PrivateDefinitions.Add(string.Format("BUILD_VERSION=\"{0}\"", Target.BuildVersion));
-
+		if (!String.IsNullOrEmpty(Target.Version.BuildURL))
+		{
+			PrivateDefinitions.Add(string.Format("BUILD_SOURCE_URL=\"{0}\"", Target.Version.BuildURL));
+		}
+				
 		UnsafeTypeCastWarningLevel = WarningLevel.Error;
 	}
 }
diff --git a/Engine/Source/Runtime/BuildSettings/Private/BuildSettings.cpp b/Engine/Source/Runtime/BuildSettings/Private/BuildSettings.cpp
--- a/Engine/Source/Runtime/BuildSettings/Private/BuildSettings.cpp
+++ b/Engine/Source/Runtime/BuildSettings/Private/BuildSettings.cpp
@@ -43,6 +43,20 @@ namespace BuildSettings
 	{
 		return ENGINE_IS_PROMOTED_BUILD;
 	}
+	
+	bool IsWithDebugInfo()
+	{
+		return UE_WITH_DEBUG_INFO;
+	}
+	
+	const TCHAR* GetBuildURL()
+	{
+#ifdef BUILD_SOURCE_URL
+		return TEXT(BUILD_SOURCE_URL);
+#else
+		return TEXT("");
+#endif	
+	}
 }
 
 
diff --git a/Engine/Source/Runtime/BuildSettings/Public/BuildSettings.h b/Engine/Source/Runtime/BuildSettings/Public/BuildSettings.h
--- a/Engine/Source/Runtime/BuildSettings/Public/BuildSettings.h
+++ b/Engine/Source/Runtime/BuildSettings/Public/BuildSettings.h
@@ -69,4 +69,14 @@ namespace BuildSettings
 	 * @return True if this is a promoted build of the engine.
 	 */
 	BUILDSETTINGS_API bool IsPromotedBuild();
+
+	/**
+	 * Identifies whether this build was compiled with or without debug info. (e.g. pdb files on Microsoft platforms)
+	 */
+	BUILDSETTINGS_API bool IsWithDebugInfo();
+
+	/**
+	 * Returns a URL where the job which created these binaries on an automation system (e.g. Horde) can be found, or an empty string.
+	 */
+	BUILDSETTINGS_API const TCHAR* GetBuildURL();
 }
diff --git a/Engine/Source/Runtime/Core/Private/GenericPlatform/GenericPlatformCrashContext.cpp b/Engine/Source/Runtime/Core/Private/GenericPlatform/GenericPlatformCrashContext.cpp
--- a/Engine/Source/Runtime/Core/Private/GenericPlatform/GenericPlatformCrashContext.cpp
+++ b/Engine/Source/Runtime/Core/Private/GenericPlatform/GenericPlatformCrashContext.cpp
@@ -159,9 +159,13 @@ void FGenericCrashContext::Initialize()
 #if !NOINITCRASHREPORTER
 	NCached::Session.bIsInternalBuild = FEngineBuildSettings::IsInternalBuild();
 	NCached::Session.bIsPerforceBuild = FEngineBuildSettings::IsPerforceBuild();
+	NCached::Session.bWithDebugInfo = FApp::GetIsWithDebugInfo();
 	NCached::Session.bIsSourceDistribution = FEngineBuildSettings::IsSourceDistribution();
 	NCached::Session.ProcessId = FPlatformProcess::GetCurrentProcessId();
 
+	NCached::Set(NCached::Session.EngineVersion, *FEngineVersion::Current().ToString());
+	NCached::Set(NCached::Session.EngineCompatibleVersion, *FEngineVersion::Current().ToString());
+	NCached::Set(NCached::Session.BuildVersion, FApp::GetBuildVersion());
 	NCached::Set(NCached::Session.GameName, *FString::Printf(TEXT("UE-%s"), FApp::GetProjectName()));
 	NCached::Set(NCached::Session.GameSessionID, TEXT("")); // Updated by callback
 	NCached::Set(NCached::Session.GameStateName, TEXT("")); // Updated by callback
@@ -334,6 +338,16 @@ void FGenericCrashContext::Initialize()
 		}
 	});
 
+	// Store some additional info in the generic maps if it is available
+	if (const TCHAR* BuildURL = FApp::GetBuildURL())
+	{
+		SetEngineData(TEXT("BuildURL"), BuildURL);
+	}
+	if (const TCHAR* ExecutingJobURL = FApp::GetExecutingJobURL())
+	{
+		SetEngineData(TEXT("ExecutingJobURL"), ExecutingJobURL);
+	}
+
 	SerializeTempCrashContextToFile();
 
 	CleanupPlatformSpecificFiles();
@@ -413,7 +427,8 @@ void FGenericCrashContext::CopySharedCrashContext(FSharedCrashContext& Dst)
 	TCHAR* DynamicDataStart = &Dst.DynamicData[0];
 	TCHAR* DynamicDataPtr = DynamicDataStart;
 
-	#define CR_DYNAMIC_BUFFER_REMAIN uint32((CR_MAX_DYNAMIC_BUFFER_CHARS) - (DynamicDataPtr-DynamicDataStart))
+	// -1 to allow space for null terminator
+	#define CR_DYNAMIC_BUFFER_REMAIN uint32((CR_MAX_DYNAMIC_BUFFER_CHARS) - (DynamicDataPtr-DynamicDataStart) - 1)
 
 	Dst.EnabledPluginsOffset = (uint32)(DynamicDataPtr - DynamicDataStart);
 	Dst.EnabledPluginsNum = NCached::EnabledPluginsList.Num();
@@ -627,6 +642,8 @@ void FGenericCrashContext::SerializeTempCrashContextToFile()
 	FFileHelper::SaveStringToFile(SessionBuffer, *SessionFilePath);
 }
 
+// This function may be called in the crashing executable or in an external crash reporter program. Take care with accessing global variables vs member variables or 
+// fields in NCached!
 void FGenericCrashContext::SerializeSessionContext(FString& Buffer)
 {
 	AddCrashPropertyInternal(Buffer, TEXT("ProcessId"), NCached::Session.ProcessId);
@@ -634,6 +651,7 @@ void FGenericCrashContext::SerializeSessionContext(FString& Buffer)
 
 	AddCrashPropertyInternal(Buffer, TEXT("IsInternalBuild"), NCached::Session.bIsInternalBuild);
 	AddCrashPropertyInternal(Buffer, TEXT("IsPerforceBuild"), NCached::Session.bIsPerforceBuild);
+	AddCrashPropertyInternal(Buffer, TEXT("IsWithDebugInfo"), NCached::Session.bWithDebugInfo);
 	AddCrashPropertyInternal(Buffer, TEXT("IsSourceDistribution"), NCached::Session.bIsSourceDistribution);
 
 	if (FCString::Strlen(NCached::Session.GameName) > 0)
@@ -664,11 +682,12 @@ void FGenericCrashContext::SerializeSessionContext(FString& Buffer)
 
 	AddCrashPropertyInternal(Buffer, TEXT("DeploymentName"), NCached::Session.DeploymentName);
 
-	AddCrashPropertyInternal(Buffer, TEXT("EngineVersion"), *FEngineVersion::Current().ToString());
+	AddCrashPropertyInternal(Buffer, TEXT("EngineVersion"), NCached::Session.EngineVersion); 
+	AddCrashPropertyInternal(Buffer, TEXT("EngineCompatibleVersion"), NCached::Session.EngineCompatibleVersion); 
 	AddCrashPropertyInternal(Buffer, TEXT("CommandLine"), NCached::Session.CommandLine);
 	AddCrashPropertyInternal(Buffer, TEXT("LanguageLCID"), NCached::Session.LanguageLCID);
 	AddCrashPropertyInternal(Buffer, TEXT("AppDefaultLocale"), NCached::Session.DefaultLocale);
-	AddCrashPropertyInternal(Buffer, TEXT("BuildVersion"), FApp::GetBuildVersion());
+	AddCrashPropertyInternal(Buffer, TEXT("BuildVersion"), NCached::Session.BuildVersion);
 	AddCrashPropertyInternal(Buffer, TEXT("Symbols"), NCached::Session.SymbolsLabel);
 	AddCrashPropertyInternal(Buffer, TEXT("IsUERelease"), NCached::Session.bIsUERelease);
 
@@ -734,6 +753,8 @@ void FGenericCrashContext::SerializeUserSettings(FString& Buffer)
 	AddCrashPropertyInternal(Buffer, TEXT("LogFilePath"), FPlatformOutputDevices::GetAbsoluteLogFilename()); // Don't use the value cached, it may be out of date.
 }
 
+// This function may be called in the crashing executable or in an external crash reporter program. Take care with accessing global variables vs member variables or 
+// fields in NCached!
 void FGenericCrashContext::SerializeContentToBuffer() const
 {
 	// Clear the buffer in case the content is serialized more than once, keeping the most up to date values and preventing to store several XML documents in the same buffer.
diff --git a/Engine/Source/Runtime/Core/Private/Misc/App.cpp b/Engine/Source/Runtime/Core/Private/Misc/App.cpp
--- a/Engine/Source/Runtime/Core/Private/Misc/App.cpp
+++ b/Engine/Source/Runtime/Core/Private/Misc/App.cpp
@@ -4,6 +4,7 @@
 #include "HAL/FileManager.h"
 #include "Misc/Paths.h"
 #include "Misc/ConfigCacheIni.h"
+#include "Misc/CoreDelegates.h"
 #include "Runtime/Launch/Resources/Version.h"
 #include "HAL/LowLevelMemTracker.h"
 #include "BuildSettings.h"
@@ -58,11 +59,55 @@ const TCHAR* FApp::GetBuildVersion()
 	return BuildSettings::GetBuildVersion();
 }
 
+const TCHAR* FApp::GetBuildURL()
+{
+	if (FCoreDelegates::OnGetBuildURL.IsBound()) 
+	{
+		return FCoreDelegates::OnGetBuildURL.Execute();
+	}	
+	return BuildSettings::GetBuildURL();
+}	
+
 int32 FApp::GetEngineIsPromotedBuild()
 {
 	return BuildSettings::IsPromotedBuild()? 1 : 0;
 }
 
+bool FApp::GetIsWithDebugInfo()
+{
+	return BuildSettings::IsWithDebugInfo();
+}
+
+const TCHAR* FApp::GetExecutingJobURL()
+{
+	if (FCoreDelegates::OnGetExecutingJobURL.IsBound())
+	{
+		return FCoreDelegates::OnGetExecutingJobURL.Execute();
+	}
+
+	static const FString URL = []() -> FString {
+		FString HordeUrl = FPlatformMisc::GetEnvironmentVariable(TEXT("UE_HORDE_URL"));
+		if (HordeUrl.IsEmpty())
+		{
+			return FString();
+		}
+		FString HordeJobId = FPlatformMisc::GetEnvironmentVariable(TEXT("UE_HORDE_JOBID"));
+		if (HordeJobId.IsEmpty())
+		{
+			return FString();
+		}
+		FString HordeStepId = FPlatformMisc::GetEnvironmentVariable(TEXT("UE_HORDE_STEPID"));
+		if (HordeStepId.IsEmpty())
+		{
+			return FString::Printf(TEXT("%s/job/%s"), *HordeUrl, *HordeJobId);
+		}
+		else 
+		{
+			return FString::Printf(TEXT("%s/job/%s?step=%s"), *HordeUrl, *HordeJobId, *HordeStepId);
+		}
+	}();
+	return *URL;
+}
 
 FString FApp::GetEpicProductIdentifier()
 {
diff --git a/Engine/Source/Runtime/Core/Private/Misc/CoreDelegates.cpp b/Engine/Source/Runtime/Core/Private/Misc/CoreDelegates.cpp
--- a/Engine/Source/Runtime/Core/Private/Misc/CoreDelegates.cpp
+++ b/Engine/Source/Runtime/Core/Private/Misc/CoreDelegates.cpp
@@ -334,6 +334,9 @@ CORE_API void RegisterEncryptionKeyCallback(TEncryptionKeyFunc InCallback)
 	});
 }
 
+TDelegate<const TCHAR*(void)> FCoreDelegates::OnGetBuildURL;
+TDelegate<const TCHAR*(void)> FCoreDelegates::OnGetExecutingJobURL;
+
 FSimpleMulticastDelegate FCoreDelegates::OnParentBeginFork;
 FSimpleMulticastDelegate FCoreDelegates::OnParentPreFork;
 TMulticastDelegate<void(EForkProcessRole /* ProcessRole */)> FCoreDelegates::OnPostFork;
diff --git a/Engine/Source/Runtime/Core/Public/GenericPlatform/GenericPlatformCrashContext.h b/Engine/Source/Runtime/Core/Public/GenericPlatform/GenericPlatformCrashContext.h
--- a/Engine/Source/Runtime/Core/Public/GenericPlatform/GenericPlatformCrashContext.h
+++ b/Engine/Source/Runtime/Core/Public/GenericPlatform/GenericPlatformCrashContext.h
@@ -156,6 +156,7 @@ struct FSessionContext
 {
 	bool 					bIsInternalBuild;
 	bool 					bIsPerforceBuild;
+	bool 					bWithDebugInfo;
 	bool 					bIsSourceDistribution;
 	bool 					bIsUERelease;
 	bool					bIsOOM;
@@ -169,6 +170,9 @@ struct FSessionContext
 	int32					CrashTrigger;
 	int32					OOMAllocationAlignment;
 	uint64					OOMAllocationSize;
+	TCHAR 					EngineVersion[CR_MAX_GENERIC_FIELD_CHARS];
+	TCHAR 					EngineCompatibleVersion[CR_MAX_GENERIC_FIELD_CHARS];
+	TCHAR 					BuildVersion[CR_MAX_GENERIC_FIELD_CHARS];
 	TCHAR 					GameName[CR_MAX_GENERIC_FIELD_CHARS];
 	TCHAR					EngineMode[CR_MAX_GENERIC_FIELD_CHARS];
 	TCHAR					EngineModeEx[CR_MAX_GENERIC_FIELD_CHARS];
diff --git a/Engine/Source/Runtime/Core/Public/Misc/App.h b/Engine/Source/Runtime/Core/Public/Misc/App.h
--- a/Engine/Source/Runtime/Core/Public/Misc/App.h
+++ b/Engine/Source/Runtime/Core/Public/Misc/App.h
@@ -66,6 +66,22 @@ public:
 	* @return The build version
 	*/
 	static const TCHAR* GetBuildVersion();
+	
+	/* 
+	 * Gets the URL for a job which created these binaries. 
+	 * This may not be accessible if you are not part of the organization that created this build. 
+	 * May be overridden with FCoreDelegates::OnGetBuildURL
+	 */
+	static const TCHAR* GetBuildURL();
+
+	/* 
+	 * Gets the URL for the currently running Horde job/step if any. 
+	 * May be overridden with FCoreDelegates::OnGetExecutingJobURL
+	 */ 
+	static const TCHAR* GetExecutingJobURL();
+
+	/* Returns whether the binaries were built with debug info */
+	static bool GetIsWithDebugInfo();
 
 	/**
 	 * Gets the date at which this application was built.
diff --git a/Engine/Source/Runtime/Core/Public/Misc/CoreDelegates.h b/Engine/Source/Runtime/Core/Public/Misc/CoreDelegates.h
--- a/Engine/Source/Runtime/Core/Public/Misc/CoreDelegates.h
+++ b/Engine/Source/Runtime/Core/Public/Misc/CoreDelegates.h
@@ -553,6 +553,12 @@ public:
 	DECLARE_MULTICAST_DELEGATE_OneParam(FOnDisallowedExecCommandCalled, const TCHAR* /*Cmd*/);
 	static FOnDisallowedExecCommandCalled OnDisallowedExecCommandCalled;
 
+	// Extension point for projects to report the URL for a continuous integration job which produced these binaries.
+	static TDelegate<const TCHAR*()> OnGetBuildURL;
+	
+	// Extension point for projects to report the URL for a continuous integration job which is currently executing this process. 
+	static TDelegate<const TCHAR*()> OnGetExecutingJobURL;
+	
 private:
 
 	// Callbacks for hotfixes
