delegates race detection: added logging of the current thread callstack as automatic assert callstack logging is unreliable in tests

diff --git a/Engine/Source/Runtime/Core/Public/Misc/MTAccessDetector.h b/Engine/Source/Runtime/Core/Public/Misc/MTAccessDetector.h
--- a/Engine/Source/Runtime/Core/Public/Misc/MTAccessDetector.h
+++ b/Engine/Source/Runtime/Core/Public/Misc/MTAccessDetector.h
@@ -348,8 +348,9 @@ private:
 		{
 			uint32 CurrentThreadId = FPlatformTLS::GetCurrentThreadId();
 			checkf(InState.GetWriterThreadId() == CurrentThreadId,
-				TEXT("Data race detected! Acquiring read access on thread %u concurrently with %d writers on thread %u:\nWriter thread %u callstack:\n%s"),
+				TEXT("Data race detected! Acquiring read access on thread %u concurrently with %d writers on thread %u:\nCurrent thread %u callstack:\n%s\nWriter thread %u callstack:\n%s"),
 				CurrentThreadId, InState.WriterNum, InState.GetWriterThreadId(), 
+				CurrentThreadId, *GetThreadCallstack(CurrentThreadId),
 				InState.GetWriterThreadId(), *GetThreadCallstack(InState.GetWriterThreadId()));
 		}
 	}
@@ -467,9 +468,10 @@ public:
 		}
 
 		checkf(LoadState().Value == ExpectedState.Value,
-			TEXT("Race detector destroyed while being accessed on another thread: %d readers, %d writers on thread %u:\nWriter thread %u callstack:\n%s"),
+			TEXT("Race detector destroyed while being accessed on another thread: %d readers, %d writers on thread %u:\nCurrent thread %u callstack:\n%s\nWriter thread %u callstack:\n%s"),
 			LoadState().ReaderNum - ExpectedState.ReaderNum, 
 			LoadState().WriterNum - ExpectedState.WriterNum, LoadState().GetWriterThreadId(), 
+			FPlatformTLS::GetCurrentThreadId(), *GetThreadCallstack(FPlatformTLS::GetCurrentThreadId()),
 			LoadState().GetWriterThreadId(), *GetThreadCallstack(LoadState().GetWriterThreadId()));
 	}
 
@@ -531,17 +533,19 @@ public:
 		{	// check that all readers are on the current thread
 			int32 ReaderIndex = GetReadersTls().IndexOfByPredicate([this](FReaderNum ReaderNum) { return ReaderNum.Reader == this; });
 			checkf(ReaderIndex != INDEX_NONE, TEXT("Race detector is not trivially copyable while this delegate is copied trivially. Consider changing this delegate to use `FNotThreadSafeNotCheckedDelegateUserPolicy`"));
-			checkf(GetReadersTls()[ReaderIndex].Num == LocalState.ReaderNum, 
-				TEXT("Data race detected: %d reader(s) on another thread(s) while acquiring write access"), 
-				LocalState.ReaderNum - GetReadersTls()[ReaderIndex].Num);
+			checkf(GetReadersTls()[ReaderIndex].Num == LocalState.ReaderNum,
+				TEXT("Data race detected: %d reader(s) on another thread(s) while acquiring write access.\nCurrent thread %u callstack:\n%s"),
+				LocalState.ReaderNum - GetReadersTls()[ReaderIndex].Num,
+				FPlatformTLS::GetCurrentThreadId(), *GetThreadCallstack(FPlatformTLS::GetCurrentThreadId()));
 		}
 
 		uint32 CurrentThreadId = FPlatformTLS::GetCurrentThreadId();
 		if (LocalState.WriterNum != 0)
 		{
 			checkf(LocalState.GetWriterThreadId() == CurrentThreadId,
-				TEXT("Data race detected: acquiring write access on thread %u concurrently with %d writers on thread %u:\nWriter thread %u callstack:\n%s"), 
+				TEXT("Data race detected: acquiring write access on thread %u concurrently with %d writers on thread %u:\nCurrent thread %u callstack:\n%s\nWriter thread %u callstack:\n%s"), 
 				CurrentThreadId, LocalState.WriterNum, LocalState.GetWriterThreadId(), 
+				CurrentThreadId, *GetThreadCallstack(CurrentThreadId),
 				LocalState.GetWriterThreadId(), *GetThreadCallstack(LocalState.GetWriterThreadId()));
 		}
 
@@ -549,11 +553,12 @@ public:
 		FState PrevState = ExchangeState(NewState);
 
 		checkf(LocalState == PrevState,
-			TEXT("Data race detected: other thread(s) activity during acquiring write access on thread %u: %u -> %u readers, %u -> %u writers on thread %u -> %u:\nWriter thread %u callstack:\n%s"),
+			TEXT("Data race detected: other thread(s) activity during acquiring write access on thread %u: %u -> %u readers, %u -> %u writers on thread %u -> %u:\nCurrent thread %u callstack:\n%s\nWriter thread %u callstack:\n%s"),
 			CurrentThreadId,
 			LocalState.ReaderNum, PrevState.ReaderNum,
 			LocalState.WriterNum, PrevState.WriterNum,
 			LocalState.GetWriterThreadId(), PrevState.GetWriterThreadId(), 
+			CurrentThreadId, *GetThreadCallstack(CurrentThreadId),
 			PrevState.GetWriterThreadId(), *GetThreadCallstack(PrevState.GetWriterThreadId()));
 	}
 
@@ -576,10 +581,11 @@ public:
 		FState PrevState = ExchangeState(NewState);
 
 		checkf(LocalState == PrevState,
-			TEXT("Data race detected: other thread(s) activity during releasing write access on thread %d: %u -> %u readers, %u -> %u writers on thread %u -> %u\nWriter thread %u callstack:\n%s"),
+			TEXT("Data race detected: other thread(s) activity during releasing write access on thread %d: %u -> %u readers, %u -> %u writers on thread %u -> %u\nCurrent thread %u callstack:\n%s\nWriter thread %u callstack:\n%s"),
 			LocalState.ReaderNum, PrevState.ReaderNum,
 			LocalState.WriterNum, PrevState.WriterNum,
 			LocalState.GetWriterThreadId(), PrevState.GetWriterThreadId(), 
+			FPlatformTLS::GetCurrentThreadId(), *GetThreadCallstack(FPlatformTLS::GetCurrentThreadId()),
 			PrevState.GetWriterThreadId(), *GetThreadCallstack(PrevState.GetWriterThreadId()));
 	}
 
