Fix potential thread safety issue in virtualization manager notification callbacks.

- Changed the notification callbacks (used for statistics gathering) from TMulticastDelegate to TTSMulticastDelegate

diff --git a/Engine/Source/Runtime/Core/Public/Virtualization/VirtualizationSystem.h b/Engine/Source/Runtime/Core/Public/Virtualization/VirtualizationSystem.h
--- a/Engine/Source/Runtime/Core/Public/Virtualization/VirtualizationSystem.h
+++ b/Engine/Source/Runtime/Core/Public/Virtualization/VirtualizationSystem.h
@@ -804,7 +804,7 @@ public:
 	};
 
 	/** Declare delegate for notifications*/
-	DECLARE_MULTICAST_DELEGATE_TwoParams(FOnNotification, ENotification, const FIoHash&);
+	using FOnNotification = TTSMulticastDelegate<void(ENotification, const FIoHash&)>;
 
 	virtual FOnNotification& GetNotificationEvent() = 0;
 };
