Fix deprecation warning

diff --git a/Engine/Source/Runtime/Core/Private/AutoRTFM/GlobalData.cpp b/Engine/Source/Runtime/Core/Private/AutoRTFM/GlobalData.cpp
--- a/Engine/Source/Runtime/Core/Private/AutoRTFM/GlobalData.cpp
+++ b/Engine/Source/Runtime/Core/Private/AutoRTFM/GlobalData.cpp
@@ -37,7 +37,7 @@ void InitializeGlobalDataIfNecessary()
             FContext::InitializeGlobalData();
             char Buffer[128];
             snprintf(Buffer, sizeof(Buffer), "%s=%p", EnvName, GlobalData);
-            putenv(Buffer);
+            _putenv(Buffer);
         }
     });
 }
