Ensure zero terminated string when launching UTS on Mac/Linux

GetData does not guarantee that the string builder has zero terminated the string. Check with ToString to get valid c-string.

diff --git a/Engine/Source/Runtime/Core/Private/ProfilingDebugging/TraceAuxiliary.cpp b/Engine/Source/Runtime/Core/Private/ProfilingDebugging/TraceAuxiliary.cpp
--- a/Engine/Source/Runtime/Core/Private/ProfilingDebugging/TraceAuxiliary.cpp
+++ b/Engine/Source/Runtime/Core/Private/ProfilingDebugging/TraceAuxiliary.cpp
@@ -1187,6 +1187,7 @@ static void LaunchUnrealTraceInternal(const TCHAR* CommandLine)
 #elif PLATFORM_MAC
 	BinPath << "Binaries/Mac/UnrealTraceServer";
 #endif
+	BinPath.ToString(); //Ensure zero termination
 
 	if (access(*BinPath, F_OK) < 0)
 	{
@@ -1196,6 +1197,7 @@ static void LaunchUnrealTraceInternal(const TCHAR* CommandLine)
 
 	TAnsiStringBuilder<64> ForkArg;
 	ForkArg << "fork";
+	ForkArg.ToString(); //Ensure zero termination
 
 	pid_t UtsPid = fork();
 	if (UtsPid < 0)
