Min Fix for Linux GrammaredParse

diff --git a/Engine/Source/Runtime/Core/Private/Misc/Parse.cpp b/Engine/Source/Runtime/Core/Private/Misc/Parse.cpp
--- a/Engine/Source/Runtime/Core/Private/Misc/Parse.cpp
+++ b/Engine/Source/Runtime/Core/Private/Misc/Parse.cpp
@@ -1539,11 +1539,10 @@ struct FGrammarBasedParser
 		// String literal
 		if (MatchChar(TCHAR('"')))
 		{
-			ZeroOrMore([this]() 
+			while (*Cursor && (TCHAR('"') != *Cursor))
 			{
 				++Cursor;
-				return !IsAt(TCHAR('"'));
-			});
+			}
 
 			if (!MatchChar(TCHAR('"')))
 			{
diff --git a/Engine/Source/Runtime/Core/Tests/Misc/ParseTests.cpp b/Engine/Source/Runtime/Core/Tests/Misc/ParseTests.cpp
--- a/Engine/Source/Runtime/Core/Tests/Misc/ParseTests.cpp
+++ b/Engine/Source/Runtime/Core/Tests/Misc/ParseTests.cpp
@@ -120,7 +120,7 @@ TEST_CASE("Parse::Value::ToBuffer", "[Parse][Smoke]")
 }
 
 
-TEST_CASE("Parse::GrammaredCLIParse::Callback", "[Smoke][.]")
+TEST_CASE("Parse::GrammaredCLIParse::Callback", "[Smoke]")
 {
 	struct StringKeyValue {
 		const TCHAR* Key;
