Re-add push/pop disabling std function warning

#jijra none

diff --git a/Engine/Source/Runtime/Core/Private/AutoRTFM/GlobalData.cpp b/Engine/Source/Runtime/Core/Private/AutoRTFM/GlobalData.cpp
--- a/Engine/Source/Runtime/Core/Private/AutoRTFM/GlobalData.cpp
+++ b/Engine/Source/Runtime/Core/Private/AutoRTFM/GlobalData.cpp
@@ -11,6 +11,15 @@ FGlobalData* GlobalData;
 
 void InitializeGlobalDataIfNecessary()
 {
+#ifdef _MSC_VER
+/*
+   Disable warning about deprecated STD C functions.
+*/
+#pragma warning(disable : 4996)
+
+#pragma warning(push)
+#endif
+
     UE_CALL_ONCE([]
     {
         // AutoRTFM is intended to be used primarily from monolithic UE binaries, where having
@@ -23,7 +32,7 @@ void InitializeGlobalDataIfNecessary()
         //
         // - All AutoRTFM instances coordinate together on things like the function table,
         //   lock table, and TLS key.
-        
+
         constexpr const char* EnvName = "AutoRTFMGlobalData";
 
         if (char* EnvString = getenv(EnvName))
@@ -40,6 +49,9 @@ void InitializeGlobalDataIfNecessary()
             putenv(Buffer);
         }
     });
+#ifdef _MSC_VER
+#pragma warning(pop)
+#endif
 }
 
 struct FInitializeGlobalData
