Implements GetConsoleVariableSetByName()

diff --git a/Engine/Source/Runtime/Core/Private/HAL/ConsoleManager.cpp b/Engine/Source/Runtime/Core/Private/HAL/ConsoleManager.cpp
--- a/Engine/Source/Runtime/Core/Private/HAL/ConsoleManager.cpp
+++ b/Engine/Source/Runtime/Core/Private/HAL/ConsoleManager.cpp
@@ -49,11 +49,9 @@ bool IsGoodHelpString(const TCHAR* In)
 	return bGoodEndChar;
 }
 
-// Get human readable string
-// @return never 0
-static const TCHAR* GetSetByTCHAR(EConsoleVariableFlags InSetBy)
+const TCHAR* GetConsoleVariableSetByName(EConsoleVariableFlags ConsoleVariableFlags)
 {
-	EConsoleVariableFlags SetBy = (EConsoleVariableFlags)((uint32)InSetBy & ECVF_SetByMask);
+	EConsoleVariableFlags SetBy = (EConsoleVariableFlags)((uint32)ConsoleVariableFlags & ECVF_SetByMask);
 
 	switch(SetBy)
 	{
@@ -63,9 +61,10 @@ static const TCHAR* GetSetByTCHAR(EConsoleVariableFlags InSetBy)
 		CASE(Scalability)
 		CASE(GameSetting)
 		CASE(ProjectSetting)
-		CASE(DeviceProfile)
 		CASE(SystemSettingsIni)
+		CASE(DeviceProfile)
 		CASE(ConsoleVariablesIni)
+		CASE(Hotfix)
 		CASE(Commandline)
 		CASE(Code)
 		CASE(Console)
@@ -178,8 +177,8 @@ public:
 
 			const FString Message = FString::Printf(TEXT("Setting the console variable '%s' with 'SetBy%s' was ignored as it is lower priority than the previous 'SetBy%s'. Value remains '%s'"),
 				CVarName.IsEmpty() ? TEXT("unknown?") : *CVarName,
-				GetSetByTCHAR((EConsoleVariableFlags)NewPri),
-				GetSetByTCHAR((EConsoleVariableFlags)OldPri),
+				GetConsoleVariableSetByName((EConsoleVariableFlags)NewPri),
+				GetConsoleVariableSetByName((EConsoleVariableFlags)OldPri),
 				*GetString()
 				);
 				
@@ -1917,11 +1916,11 @@ static void DumpObjects(const TMap<FString, IConsoleObject*>& ConsoleObjects, co
 			{
 				if (bWriteToCSV)
 				{
-					MultiLogf(Log, CSV, TEXT("%s,%s,%s%s"), *Key, *CVar->GetString(), GetSetByTCHAR(CVar->GetFlags()), *Help);
+					MultiLogf(Log, CSV, TEXT("%s,%s,%s%s"), *Key, *CVar->GetString(), GetConsoleVariableSetByName(CVar->GetFlags()), *Help);
 				}
 				else
 				{
-					MultiLogf(Log, CSV, TEXT("%s = \"%s\"      LastSetBy: %s%s"), *Key, *CVar->GetString(), GetSetByTCHAR(CVar->GetFlags()), *Help);
+					MultiLogf(Log, CSV, TEXT("%s = \"%s\"      LastSetBy: %s%s"), *Key, *CVar->GetString(), GetConsoleVariableSetByName(CVar->GetFlags()), *Help);
 				}
 			}
 		}
@@ -2107,7 +2106,7 @@ bool FConsoleManager::ProcessUserConsoleInput(const TCHAR* InInput, FOutputDevic
 
 		if(bShowCurrentState)
 		{
-			Ar.Logf(TEXT("%s = \"%s\"      LastSetBy: %s"), *Param1, *CVar->GetString(), GetSetByTCHAR(CVar->GetFlags()));
+			Ar.Logf(TEXT("%s = \"%s\"      LastSetBy: %s"), *Param1, *CVar->GetString(), GetConsoleVariableSetByName(CVar->GetFlags()));
 		}
 	}
 
diff --git a/Engine/Source/Runtime/Core/Public/HAL/IConsoleManager.h b/Engine/Source/Runtime/Core/Public/HAL/IConsoleManager.h
--- a/Engine/Source/Runtime/Core/Public/HAL/IConsoleManager.h
+++ b/Engine/Source/Runtime/Core/Public/HAL/IConsoleManager.h
@@ -161,6 +161,10 @@ enum EConsoleVariableFlags
 	// ------------------------------------------------
 };
 
+/** Returns human readable ECVF_SetByMask bits of the console variable flags. */
+extern CORE_API const TCHAR* GetConsoleVariableSetByName(EConsoleVariableFlags ConsoleVariableFlags);
+
+
 class IConsoleVariable;
 
 #if !NO_CVARS
diff --git a/Engine/Source/Runtime/RenderCore/Private/DumpGPU.cpp b/Engine/Source/Runtime/RenderCore/Private/DumpGPU.cpp
--- a/Engine/Source/Runtime/RenderCore/Private/DumpGPU.cpp
+++ b/Engine/Source/Runtime/RenderCore/Private/DumpGPU.cpp
@@ -577,67 +577,7 @@ public:
 				return;
 			}
 
-			const TCHAR* SetByName = TEXT("Unknown");
-			switch (CVarFlags & ECVF_SetByMask)
-			{
-			case ECVF_SetByConstructor:
-			{
-				SetByName = TEXT("Constructor");
-				break;
-			}
-			case ECVF_SetByScalability:
-			{
-				SetByName = TEXT("Scalability");
-				break;
-			}
-			case ECVF_SetByGameSetting:
-			{
-				SetByName = TEXT("GameSetting");
-				break;
-			}
-			case ECVF_SetByProjectSetting:
-			{
-				SetByName = TEXT("ProjectSetting");
-				break;
-			}
-			case ECVF_SetBySystemSettingsIni:
-			{
-				SetByName = TEXT("SystemSettingsIni");
-				break;
-			}
-			case ECVF_SetByDeviceProfile:
-			{
-				SetByName = TEXT("DeviceProfile");
-				break;
-			}
-			case ECVF_SetByGameOverride:
-			{
-				SetByName = TEXT("GameOverride");
-				break;
-			}
-			case ECVF_SetByConsoleVariablesIni:
-			{
-				SetByName = TEXT("ConsoleVariablesIni");
-				break;
-			}
-			case ECVF_SetByCommandline:
-			{
-				SetByName = TEXT("Commandline");
-				break;
-			}
-			case ECVF_SetByCode:
-			{
-				SetByName = TEXT("Code");
-				break;
-			}
-			case ECVF_SetByConsole:
-			{
-				SetByName = TEXT("Console");
-				break;
-			}
-			default:
-				ensure(false);
-			}
+			const TCHAR* SetByName = GetConsoleVariableSetByName(CVarFlags);
 
 			FString Value = CVar->GetString();
 
