Move the LTO/PGO/PGI status to the startup messages so they can be printed consistently.

- This allows them to be found in logs of all platforms and forked servers.

diff --git a/Engine/Source/Runtime/Core/Private/Android/AndroidPlatformMisc.cpp b/Engine/Source/Runtime/Core/Private/Android/AndroidPlatformMisc.cpp
--- a/Engine/Source/Runtime/Core/Private/Android/AndroidPlatformMisc.cpp
+++ b/Engine/Source/Runtime/Core/Private/Android/AndroidPlatformMisc.cpp
@@ -588,11 +588,6 @@ void FAndroidMisc::PlatformInit()
 
 
 	InitCpuThermalSensor();
-
-	UE_LOG(LogInit, Log, TEXT(" - This binary is optimized with LTO: %s, PGO: %s, instrumented for PGO data collection: %s"),
-		PLATFORM_COMPILER_OPTIMIZATION_LTCG ? TEXT("yes") : TEXT("no"),
-		FPlatformMisc::IsPGOEnabled() ? TEXT("yes") : TEXT("no"),
-		PLATFORM_COMPILER_OPTIMIZATION_PG_PROFILING ? TEXT("yes") : TEXT("no"));
 }
 
 extern void AndroidThunkCpp_DismissSplashScreen();
diff --git a/Engine/Source/Runtime/Core/Private/Misc/App.cpp b/Engine/Source/Runtime/Core/Private/Misc/App.cpp
--- a/Engine/Source/Runtime/Core/Private/Misc/App.cpp
+++ b/Engine/Source/Runtime/Core/Private/Misc/App.cpp
@@ -425,6 +425,11 @@ void FApp::PrintStartupLogMessages()
 	//UE_LOG(LogInit, Log, TEXT("Character set: %s"), sizeof(TCHAR)==1 ? TEXT("ANSI") : TEXT("Unicode") );
 	UE_LOG(LogInit, Log, TEXT("Allocator: %s"), GMalloc->GetDescriptiveName());
 	UE_LOG(LogInit, Log, TEXT("Installed Engine Build: %d"), FApp::IsEngineInstalled() ? 1 : 0);
+	UE_LOG(LogInit, Log, TEXT("This binary is optimized with LTO: %s, PGO: %s, instrumented for PGO data collection: %s"),
+		PLATFORM_COMPILER_OPTIMIZATION_LTCG ? TEXT("yes") : TEXT("no"),
+		FPlatformMisc::IsPGOEnabled() ? TEXT("yes") : TEXT("no"),
+		PLATFORM_COMPILER_OPTIMIZATION_PG_PROFILING ? TEXT("yes") : TEXT("no")
+	);
 
 	FDevVersionRegistration::DumpVersionsToLog();
 }
diff --git a/Engine/Source/Runtime/Core/Private/Unix/UnixPlatformMisc.cpp b/Engine/Source/Runtime/Core/Private/Unix/UnixPlatformMisc.cpp
--- a/Engine/Source/Runtime/Core/Private/Unix/UnixPlatformMisc.cpp
+++ b/Engine/Source/Runtime/Core/Private/Unix/UnixPlatformMisc.cpp
@@ -210,11 +210,6 @@ void FUnixPlatformMisc::PlatformInit()
 	}
 
 	UE_LOG(LogInit, Log, TEXT(" - Memory allocator used: %s"), GMalloc->GetDescriptiveName());
-	UE_LOG(LogInit, Log, TEXT(" - This binary is optimized with LTO: %s, PGO: %s, instrumented for PGO data collection: %s"),
-		PLATFORM_COMPILER_OPTIMIZATION_LTCG ? TEXT("yes") : TEXT("no"),
-		FPlatformMisc::IsPGOEnabled() ? TEXT("yes") : TEXT("no"),
-		PLATFORM_COMPILER_OPTIMIZATION_PG_PROFILING ? TEXT("yes") : TEXT("no")
-		);
 	UE_LOG(LogInit, Log, TEXT(" - This is %s build."), BuildSettings::IsLicenseeVersion() ? TEXT("a licensee") : TEXT("an internal"));
 
 	FPlatformTime::PrintCalibrationLog();
