Add a missing intermediate copy to ASTC compression memory usage estimation

diff --git a/Engine/Source/Runtime/Core/Public/PixelFormat.h b/Engine/Source/Runtime/Core/Public/PixelFormat.h
--- a/Engine/Source/Runtime/Core/Public/PixelFormat.h
+++ b/Engine/Source/Runtime/Core/Public/PixelFormat.h
@@ -260,6 +260,31 @@ ENUM_CLASS_FLAGS(EPixelFormatCapabilities);
 // type conversion warnings
 #define UE_PIXELFORMAT_TO_UINT8(argument) static_cast<uint8>(argument)
 
+FORCEINLINE bool IsASTCBlockCompressedTextureFormat(EPixelFormat PixelFormat)
+{
+	switch (PixelFormat)
+	{
+	case PF_ASTC_4x4:
+	case PF_ASTC_6x6:
+	case PF_ASTC_8x8:
+	case PF_ASTC_10x10:
+	case PF_ASTC_12x12:
+	case PF_ASTC_4x4_HDR:
+	case PF_ASTC_6x6_HDR:
+	case PF_ASTC_8x8_HDR:
+	case PF_ASTC_10x10_HDR:
+	case PF_ASTC_12x12_HDR:
+	case PF_ASTC_4x4_NORM_RG:
+	case PF_ASTC_6x6_NORM_RG:
+	case PF_ASTC_8x8_NORM_RG:
+	case PF_ASTC_10x10_NORM_RG:
+	case PF_ASTC_12x12_NORM_RG:
+		return true;
+	default:
+		return false;
+	}
+}
+
 static inline bool IsBlockCompressedFormat(EPixelFormat Format)
 {
 	switch (Format)
@@ -278,7 +303,8 @@ static inline bool IsBlockCompressedFormat(EPixelFormat Format)
 
 FORCEINLINE bool IsHDR(EPixelFormat PixelFormat)
 {
-	return PixelFormat == PF_FloatRGBA || PixelFormat == PF_BC6H || PixelFormat == PF_R16F || PixelFormat == PF_R32_FLOAT || PixelFormat == PF_A32B32G32R32F;
+	return PixelFormat == PF_FloatRGBA || PixelFormat == PF_BC6H || PixelFormat == PF_R16F || PixelFormat == PF_R32_FLOAT || PixelFormat == PF_A32B32G32R32F
+		|| PixelFormat == PF_ASTC_4x4_HDR || PixelFormat == PF_ASTC_6x6_HDR || PixelFormat == PF_ASTC_8x8_HDR || PixelFormat == PF_ASTC_10x10_HDR || PixelFormat == PF_ASTC_12x12_HDR;
 }
 
 FORCEINLINE bool IsInteger(EPixelFormat PixelFormat)
diff --git a/Engine/Source/Runtime/Engine/Private/TextureDerivedDataTask.cpp b/Engine/Source/Runtime/Engine/Private/TextureDerivedDataTask.cpp
--- a/Engine/Source/Runtime/Engine/Private/TextureDerivedDataTask.cpp
+++ b/Engine/Source/Runtime/Engine/Private/TextureDerivedDataTask.cpp
@@ -1178,6 +1178,11 @@ static int64 GetBuildRequiredMemoryEstimate(UTexture* InTexture,
 				MemoryEstimate += FMath::Max(RDOPhase1MemUse,RDOPhase2MemUse);
 			}
 		}
+		else if (IsASTCBlockCompressedTextureFormat(PixelFormat))
+		{
+			// ASTCenc does an entermediate copy to RGBA16F for HDR formats and RGBA8 for LDR
+			MemoryEstimate += (IsHDR(PixelFormat) ? 8 : 4) * TotalNumPixels;
+		}
 		else
 		{
 			// note: memory ues of non-Oodle encoders is not estimated
