Fix the multithreaded delegate access handler for AutoRTFM.

diff --git a/Engine/Source/Runtime/Core/Public/Delegates/DelegateAccessHandler.h b/Engine/Source/Runtime/Core/Public/Delegates/DelegateAccessHandler.h
--- a/Engine/Source/Runtime/Core/Public/Delegates/DelegateAccessHandler.h
+++ b/Engine/Source/Runtime/Core/Public/Delegates/DelegateAccessHandler.h
@@ -133,22 +133,44 @@ protected:
 		explicit FWriteAccessScope(FMRSWRecursiveAccessDetector& InAccessDetector)
 			: DestructionSentinel(FMRSWRecursiveAccessDetector::EAccessType::Writer)
 		{
-			DestructionSentinel.Accessor = &InAccessDetector;
-			InAccessDetector.AcquireWriteAccess(DestructionSentinel);
+			// We need to acquire the write access in the open, as this is being used
+			// to ensure that the thread with transactions enabled isn't wrongly
+			// sharing a delegate with other concurrent threads.
+			UE_AUTORTFM_OPEN(
+			{
+				DestructionSentinel.Accessor = &InAccessDetector;
+				InAccessDetector.AcquireWriteAccess(DestructionSentinel);
+			});
+
+			// And if an abort occurs, we need to clean up that we aren't in the
+			// delegate anymore (because we aborted!).
+			UE_AUTORTFM_OPENABORT(
+			{
+				Release();
+			});
 		}
 
 		~FWriteAccessScope()
 		{
-			if (!DestructionSentinel.bDestroyed) // only if `AccessDetector` wasn't destroyed while being accessed
-			{
-				// the `const_cast` doesn't violate constness here as we got a mutable accessor in `FWriteAccessScope` constuctor,
-				// otherwise we'd need to have a redundant ptr to a mutable accessor as a member
-				const_cast<FMRSWRecursiveAccessDetector*>(DestructionSentinel.Accessor)->ReleaseWriteAccess(DestructionSentinel);
-			}
+			Release();
 		}
 
 	private:
 		FMRSWRecursiveAccessDetector::FDestructionSentinel DestructionSentinel;
+
+		void Release()
+		{
+			// We need to release in the open, as we acquired in the open!
+			UE_AUTORTFM_OPEN(
+			{
+				if (!DestructionSentinel.bDestroyed) // only if `AccessDetector` wasn't destroyed while being accessed
+				{
+					// the `const_cast` doesn't violate constness here as we got a mutable accessor in `FWriteAccessScope` constuctor,
+					// otherwise we'd need to have a redundant ptr to a mutable accessor as a member
+					const_cast<FMRSWRecursiveAccessDetector*>(DestructionSentinel.Accessor)->ReleaseWriteAccess(DestructionSentinel);
+				}
+			});
+		}
 	};
 
 	[[nodiscard]] FWriteAccessScope GetWriteAccessScope()
