Make -stresstestunity work with the AutoRTFM compiler.

I've had to fix a few things:
- We had some checks that needed to check for clang and not windows (cpuid.h related).
- We don't need another SharedPtr definition in AutoRTFM.
- Had to remove the `using namespace AutoRTFM` to avoid unity build funtimes.

diff --git a/Engine/Source/Runtime/Core/Private/AutoRTFM/SharedPtr.h b/Engine/Source/Runtime/Core/Private/AutoRTFM/SharedPtr.h
deleted file mode 100644
diff --git a/Engine/Source/Runtime/Core/Private/AutoRTFM/Stats.cpp b/Engine/Source/Runtime/Core/Private/AutoRTFM/Stats.cpp
--- a/Engine/Source/Runtime/Core/Private/AutoRTFM/Stats.cpp
+++ b/Engine/Source/Runtime/Core/Private/AutoRTFM/Stats.cpp
@@ -3,89 +3,87 @@
 #if (defined(__AUTORTFM) && __AUTORTFM)
 #include "Stats.h"
 
-using namespace AutoRTFM;
-
-template<> void AutoRTFM::FStats::Report<EStatsKind::Transaction>(const uint64_t Data) const
+template<> void AutoRTFM::FStats::Report<AutoRTFM::EStatsKind::Transaction>(const uint64_t Data) const
 {
 	UE_LOG(LogAutoRTFM, Display, TEXT("  Total transactions:        %11u"), Data);
 }
 
-template<> void AutoRTFM::FStats::Report<EStatsKind::Commit>(const uint64_t Data) const
+template<> void AutoRTFM::FStats::Report<AutoRTFM::EStatsKind::Commit>(const uint64_t Data) const
 {
 	UE_LOG(LogAutoRTFM, Display, TEXT("  Total commits:             %11u"), Data);
 }
 
-template<> void AutoRTFM::FStats::Report<EStatsKind::Abort>(const uint64_t Data) const
+template<> void AutoRTFM::FStats::Report<AutoRTFM::EStatsKind::Abort>(const uint64_t Data) const
 {
 	UE_LOG(LogAutoRTFM, Display, TEXT("  Total aborts:              %11u"), Data);
 }
 
-template<> void AutoRTFM::FStats::Report<EStatsKind::AverageTransactionDepth>(const uint64_t Data) const
+template<> void AutoRTFM::FStats::Report<AutoRTFM::EStatsKind::AverageTransactionDepth>(const uint64_t Data) const
 {
-	const uint64_t TotalTransactions = Datas[static_cast<size_t>(EStatsKind::Transaction)];
+	const uint64_t TotalTransactions = Datas[static_cast<size_t>(AutoRTFM::EStatsKind::Transaction)];
 	UE_LOG(LogAutoRTFM, Display, TEXT("  Average transaction depth: %15.3f"), (static_cast<double>(Data) / static_cast<double>(TotalTransactions)));
 }
 
-template<> void AutoRTFM::FStats::Report<EStatsKind::MaximumTransactionDepth>(const uint64_t Data) const
+template<> void AutoRTFM::FStats::Report<AutoRTFM::EStatsKind::MaximumTransactionDepth>(const uint64_t Data) const
 {
     UE_LOG(LogAutoRTFM, Display, TEXT("  Maximum transaction depth: %11u"), Data);
 }
 
-template<> void AutoRTFM::FStats::Report<EStatsKind::AverageWriteLogEntries>(const uint64_t Data) const
+template<> void AutoRTFM::FStats::Report<AutoRTFM::EStatsKind::AverageWriteLogEntries>(const uint64_t Data) const
 {
-    const uint64_t TotalTransactions = Datas[static_cast<size_t>(EStatsKind::Transaction)];
+    const uint64_t TotalTransactions = Datas[static_cast<size_t>(AutoRTFM::EStatsKind::Transaction)];
     UE_LOG(LogAutoRTFM, Display, TEXT("  Average write log entries: %15.3f"), (static_cast<double>(Data) / static_cast<double>(TotalTransactions)));
 }
 
-template<> void AutoRTFM::FStats::Report<EStatsKind::MaximumWriteLogEntries>(const uint64_t Data) const
+template<> void AutoRTFM::FStats::Report<AutoRTFM::EStatsKind::MaximumWriteLogEntries>(const uint64_t Data) const
 {
     UE_LOG(LogAutoRTFM, Display, TEXT("  Maximum write log entries: %11u"), Data);
 }
 
-template<> void AutoRTFM::FStats::Report<EStatsKind::AverageWriteLogBytes>(const uint64_t Data) const
+template<> void AutoRTFM::FStats::Report<AutoRTFM::EStatsKind::AverageWriteLogBytes>(const uint64_t Data) const
 {
-    const uint64_t TotalTransactions = Datas[static_cast<size_t>(EStatsKind::Transaction)];
+    const uint64_t TotalTransactions = Datas[static_cast<size_t>(AutoRTFM::EStatsKind::Transaction)];
     UE_LOG(LogAutoRTFM, Display, TEXT("  Average write log bytes:   %15.3f"), (static_cast<double>(Data) / static_cast<double>(TotalTransactions)));
 }
 
-template<> void AutoRTFM::FStats::Report<EStatsKind::MaximumWriteLogBytes>(const uint64_t Data) const
+template<> void AutoRTFM::FStats::Report<AutoRTFM::EStatsKind::MaximumWriteLogBytes>(const uint64_t Data) const
 {
     UE_LOG(LogAutoRTFM, Display, TEXT("  Maximum write log bytes:   %11u"), Data);
 }
 
-template<> void AutoRTFM::FStats::Report<EStatsKind::HitSetHit>(const uint64_t Data) const
+template<> void AutoRTFM::FStats::Report<AutoRTFM::EStatsKind::HitSetHit>(const uint64_t Data) const
 {
     UE_LOG(LogAutoRTFM, Display, TEXT("  HitSet hits:               %11u"), Data);
 }
 
-template<> void AutoRTFM::FStats::Report<EStatsKind::HitSetMiss>(const uint64_t Data) const
+template<> void AutoRTFM::FStats::Report<AutoRTFM::EStatsKind::HitSetMiss>(const uint64_t Data) const
 {
     UE_LOG(LogAutoRTFM, Display, TEXT("  HitSet misses:             %11u"), Data);
 }
 
-template<> void AutoRTFM::FStats::Report<EStatsKind::HitSetSkippedBecauseOfStackLocalMemory>(const uint64_t Data) const
+template<> void AutoRTFM::FStats::Report<AutoRTFM::EStatsKind::HitSetSkippedBecauseOfStackLocalMemory>(const uint64_t Data) const
 {
     UE_LOG(LogAutoRTFM, Display, TEXT("  HitSet skip (stack local): %11u"), Data);
 }
 
-template<> void AutoRTFM::FStats::Report<EStatsKind::AverageCommitTasks>(const uint64_t Data) const
+template<> void AutoRTFM::FStats::Report<AutoRTFM::EStatsKind::AverageCommitTasks>(const uint64_t Data) const
 {
-    const uint64_t TotalTransactions = Datas[static_cast<size_t>(EStatsKind::Transaction)];
+    const uint64_t TotalTransactions = Datas[static_cast<size_t>(AutoRTFM::EStatsKind::Transaction)];
     UE_LOG(LogAutoRTFM, Display, TEXT("  Average commit tasks:      %15.3f"), (static_cast<double>(Data) / static_cast<double>(TotalTransactions)));
 }
 
-template<> void AutoRTFM::FStats::Report<EStatsKind::MaximumCommitTasks>(const uint64_t Data) const
+template<> void AutoRTFM::FStats::Report<AutoRTFM::EStatsKind::MaximumCommitTasks>(const uint64_t Data) const
 {
     UE_LOG(LogAutoRTFM, Display, TEXT("  Maximum commit tasks:      %11u"), Data);
 }
 
-template<> void AutoRTFM::FStats::Report<EStatsKind::AverageAbortTasks>(const uint64_t Data) const
+template<> void AutoRTFM::FStats::Report<AutoRTFM::EStatsKind::AverageAbortTasks>(const uint64_t Data) const
 {
-    const uint64_t TotalTransactions = Datas[static_cast<size_t>(EStatsKind::Transaction)];
+    const uint64_t TotalTransactions = Datas[static_cast<size_t>(AutoRTFM::EStatsKind::Transaction)];
     UE_LOG(LogAutoRTFM, Display, TEXT("  Average abort tasks:       %15.3f"), (static_cast<double>(Data) / static_cast<double>(TotalTransactions)));
 }
 
-template<> void AutoRTFM::FStats::Report<EStatsKind::MaximumAbortTasks>(const uint64_t Data) const
+template<> void AutoRTFM::FStats::Report<AutoRTFM::EStatsKind::MaximumAbortTasks>(const uint64_t Data) const
 {
     UE_LOG(LogAutoRTFM, Display, TEXT("  Maximum abort tasks:       %11u"), Data);
 }
@@ -96,28 +94,28 @@ void AutoRTFM::FStats::Report() const
 	{
 		UE_LOG(LogAutoRTFM, Display, TEXT("AutoRTFM Statistics:"));
 
-		for (size_t I = 0; I < static_cast<size_t>(EStatsKind::Total); I++)
+		for (size_t I = 0; I < static_cast<size_t>(AutoRTFM::EStatsKind::Total); I++)
 		{
 			switch (static_cast<EStatsKind>(I))
 			{
 			default: AutoRTFM::Unreachable(); break;
 #define REPORT_CASE(x) case (x): Report<x>(Datas[static_cast<size_t>(x)]); break
-			REPORT_CASE(EStatsKind::Transaction);
-			REPORT_CASE(EStatsKind::Commit);
-			REPORT_CASE(EStatsKind::Abort);
-			REPORT_CASE(EStatsKind::AverageTransactionDepth);
-            REPORT_CASE(EStatsKind::MaximumTransactionDepth);
-            REPORT_CASE(EStatsKind::AverageWriteLogEntries);
-            REPORT_CASE(EStatsKind::MaximumWriteLogEntries);
-			REPORT_CASE(EStatsKind::AverageWriteLogBytes);
-			REPORT_CASE(EStatsKind::MaximumWriteLogBytes);
-            REPORT_CASE(EStatsKind::HitSetHit);
-            REPORT_CASE(EStatsKind::HitSetMiss);
-            REPORT_CASE(EStatsKind::HitSetSkippedBecauseOfStackLocalMemory);
-            REPORT_CASE(EStatsKind::AverageCommitTasks);
-            REPORT_CASE(EStatsKind::MaximumCommitTasks);
-            REPORT_CASE(EStatsKind::AverageAbortTasks);
-            REPORT_CASE(EStatsKind::MaximumAbortTasks);
+			REPORT_CASE(AutoRTFM::EStatsKind::Transaction);
+			REPORT_CASE(AutoRTFM::EStatsKind::Commit);
+			REPORT_CASE(AutoRTFM::EStatsKind::Abort);
+			REPORT_CASE(AutoRTFM::EStatsKind::AverageTransactionDepth);
+            REPORT_CASE(AutoRTFM::EStatsKind::MaximumTransactionDepth);
+            REPORT_CASE(AutoRTFM::EStatsKind::AverageWriteLogEntries);
+            REPORT_CASE(AutoRTFM::EStatsKind::MaximumWriteLogEntries);
+			REPORT_CASE(AutoRTFM::EStatsKind::AverageWriteLogBytes);
+			REPORT_CASE(AutoRTFM::EStatsKind::MaximumWriteLogBytes);
+            REPORT_CASE(AutoRTFM::EStatsKind::HitSetHit);
+            REPORT_CASE(AutoRTFM::EStatsKind::HitSetMiss);
+            REPORT_CASE(AutoRTFM::EStatsKind::HitSetSkippedBecauseOfStackLocalMemory);
+            REPORT_CASE(AutoRTFM::EStatsKind::AverageCommitTasks);
+            REPORT_CASE(AutoRTFM::EStatsKind::MaximumCommitTasks);
+            REPORT_CASE(AutoRTFM::EStatsKind::AverageAbortTasks);
+            REPORT_CASE(AutoRTFM::EStatsKind::MaximumAbortTasks);
 #undef REPORT_CASE
 			}
 		}
diff --git a/Engine/Source/Runtime/Core/Private/AutoRTFM/StdLibSupport.cpp b/Engine/Source/Runtime/Core/Private/AutoRTFM/StdLibSupport.cpp
--- a/Engine/Source/Runtime/Core/Private/AutoRTFM/StdLibSupport.cpp
+++ b/Engine/Source/Runtime/Core/Private/AutoRTFM/StdLibSupport.cpp
@@ -17,8 +17,9 @@
 #endif
 
 #if PLATFORM_WINDOWS
-#define NOMINMAX
-#include <windows.h>
+#include "Windows/PreWindowsApi.h"
+#include "Windows/MinWindows.h"
+#include "Windows/PostWindowsApi.h"
 #endif
 
 namespace AutoRTFM
diff --git a/Engine/Source/Runtime/Core/Private/AutoRTFM/TaskArray.h b/Engine/Source/Runtime/Core/Private/AutoRTFM/TaskArray.h
--- a/Engine/Source/Runtime/Core/Private/AutoRTFM/TaskArray.h
+++ b/Engine/Source/Runtime/Core/Private/AutoRTFM/TaskArray.h
@@ -2,9 +2,10 @@
 
 #pragma once
 
-#include "SharedPtr.h"
 #include "Containers/Array.h"
 
+#include "Utils.h"
+
 namespace AutoRTFM
 {
 
@@ -29,22 +30,23 @@ public:
     void AddAll(TTaskArray<T>&& Other)
     {
         Canonicalize();
-        for (TSharedPtr<TArray<T>>& StashedVectorBox : Other.Stash)
+
+        for (TArray<T>& StashedVectorBox : Other.Stash)
         {
             Stash.Push(MoveTemp(StashedVectorBox));
         }
+
         Other.Stash.Empty();
-        if (!Other.Latest.IsEmpty())
-        {
-            Stash.Push(TSharedPtr<TArray<T>>::New(MoveTemp(Other.Latest)));
-        }
+
+		ASSERT(Latest.IsEmpty());
+		Latest = MoveTemp(Other.Latest);
     }
 
     void AddAll(const TTaskArray<T>& Other)
     {
         Canonicalize();
         Other.Canonicalize();
-        for (const TSharedPtr<TArray<T>>& StashedVectorBox : Other.Stash)
+        for (const TArray<T>& StashedVectorBox : Other.Stash)
         {
             Stash.Push(StashedVectorBox);
         }
@@ -53,9 +55,9 @@ public:
     template<typename TFunc>
     bool ForEachForward(const TFunc& Func)
     {
-        for (TSharedPtr<TArray<T>>& StashedVectorBox : Stash)
+        for (const TArray<T>& StashedVectorBox : Stash)
         {
-            for (const T& Entry : *StashedVectorBox)
+            for (const T& Entry : StashedVectorBox)
             {
                 if (!Func(Entry))
                 {
@@ -63,6 +65,7 @@ public:
                 }
             }
         }
+
         for (const T& Entry : Latest)
         {
             if (!Func(Entry))
@@ -70,6 +73,7 @@ public:
                 return false;
             }
         }
+
         return true;
     }
 
@@ -86,7 +90,7 @@ public:
 
         for (size_t IndexInStash = Stash.Num(); IndexInStash--;)
         {
-            const TArray<T>& StashedVector = *Stash[IndexInStash];
+            const TArray<T>& StashedVector = Stash[IndexInStash];
             for (size_t Index = StashedVector.Num(); Index--;)
             {
                 if (!Func(StashedVector[Index]))
@@ -110,7 +114,7 @@ public:
 
         for (size_t Index = 0; Index < Stash.Num(); Index++)
         {
-            const TArray<T>& StashedVector = *Stash[Index];
+            const TArray<T>& StashedVector = Stash[Index];
             Result += StashedVector.Num();
         }
 
@@ -125,13 +129,13 @@ private:
     {
         if (!Latest.IsEmpty())
         {
-            Stash.Push(TSharedPtr<TArray<T>>::New(MoveTemp(Latest)));
+            Stash.Push(MoveTemp(Latest));
         }
     }
     
     mutable TArray<T> Latest;
     
-    mutable TArray<TSharedPtr<TArray<T>>> Stash;
+    mutable TArray<TArray<T>> Stash;
 };
 
 } // namespace AutoRTFM
diff --git a/Engine/Source/Runtime/Core/Private/AutoRTFM/Utils.cpp b/Engine/Source/Runtime/Core/Private/AutoRTFM/Utils.cpp
--- a/Engine/Source/Runtime/Core/Private/AutoRTFM/Utils.cpp
+++ b/Engine/Source/Runtime/Core/Private/AutoRTFM/Utils.cpp
@@ -7,7 +7,9 @@
 #include <string.h>
 
 #if PLATFORM_WINDOWS
-#include <windows.h>
+#include "Windows/PreWindowsApi.h"
+#include "Windows/MinWindows.h"
+#include "Windows/PostWindowsApi.h"
 #include <dbghelp.h>
 #else
 #include <execinfo.h>
diff --git a/Engine/Source/Runtime/Core/Private/Misc/AES.cpp b/Engine/Source/Runtime/Core/Private/Misc/AES.cpp
--- a/Engine/Source/Runtime/Core/Private/Misc/AES.cpp
+++ b/Engine/Source/Runtime/Core/Private/Misc/AES.cpp
@@ -9,7 +9,7 @@
 #if defined(_M_AMD64) || defined(__x86_64__)
 
 #include <wmmintrin.h>
-#if PLATFORM_COMPILER_CLANG
+#if PLATFORM_COMPILER_CLANG && !PLATFORM_WINDOWS
 #include <cpuid.h>
 #endif
 
@@ -591,7 +591,7 @@ static inline void AesDecrypt(const uint8* Key, uint8* Contents, uint64 NumBytes
 static inline bool DetectAesInstructions()
 {
 	int info[4];
-#if PLATFORM_COMPILER_CLANG
+#if PLATFORM_COMPILER_CLANG && !PLATFORM_WINDOWS
 	__cpuid(1, info[0], info[1], info[2], info[3]);
 #else
 	__cpuid(info, 1);
diff --git a/Engine/Source/Runtime/Core/Private/Misc/SecureHash.cpp b/Engine/Source/Runtime/Core/Private/Misc/SecureHash.cpp
--- a/Engine/Source/Runtime/Core/Private/Misc/SecureHash.cpp
+++ b/Engine/Source/Runtime/Core/Private/Misc/SecureHash.cpp
@@ -12,7 +12,7 @@
 #if (defined(_M_AMD64) || defined(__x86_64__)) && !defined(_M_ARM64EC)
 
 #include <immintrin.h>
-#if PLATFORM_COMPILER_CLANG
+#if PLATFORM_COMPILER_CLANG && !PLATFORM_WINDOWS
 #include <cpuid.h>
 #include <shaintrin.h>
 #endif
@@ -726,7 +726,7 @@ static inline bool DetectShaInstructions()
 {
 	int info1[4];
 	int info7[4];
-#if PLATFORM_COMPILER_CLANG
+#if PLATFORM_COMPILER_CLANG && !PLATFORM_WINDOWS
 	__cpuid(1, info1[0], info1[1], info1[2], info1[3]);
 	__cpuid_count(7, 0, info7[0], info7[1], info7[2], info7[3]);
 #else
