Avoid locking TraceAux destination and type while issuing delegates.

We were holding a lock while calling the start/stop trace delegates. This was causing issues in locations that tried to read destination at the same time. It is safer to copy the values before issuing the callback.

diff --git a/Engine/Source/Runtime/Core/Private/ProfilingDebugging/TraceAuxiliary.cpp b/Engine/Source/Runtime/Core/Private/ProfilingDebugging/TraceAuxiliary.cpp
--- a/Engine/Source/Runtime/Core/Private/ProfilingDebugging/TraceAuxiliary.cpp
+++ b/Engine/Source/Runtime/Core/Private/ProfilingDebugging/TraceAuxiliary.cpp
@@ -314,8 +314,16 @@ bool FTraceAuxiliaryImpl::Connect(FTraceAuxiliary::EConnectionType Type, const T
 
 		if (bConnected)
 		{
-			FReadScopeLock _(CurrentTargetLock);
-			FTraceAuxiliary::OnTraceStarted.Broadcast(CurrentTraceTarget.TraceType, CurrentTraceTarget.TraceDest);
+			FString StartedDest;
+			FTraceAuxiliary::EConnectionType StartedType = FTraceAuxiliary::EConnectionType::None;
+			
+			{
+				FReadScopeLock _(CurrentTargetLock);
+				StartedDest = CurrentTraceTarget.TraceDest;
+				StartedType = CurrentTraceTarget.TraceType;
+			}
+			
+			FTraceAuxiliary::OnTraceStarted.Broadcast(StartedType, StartedDest);
 		}
 	}
 	else
@@ -339,13 +347,19 @@ bool FTraceAuxiliaryImpl::Stop()
 		return false;
 	}
 
+	FString StopedDest;
+	FTraceAuxiliary::EConnectionType StopedType = FTraceAuxiliary::EConnectionType::None;
+	
 	{
 		FWriteScopeLock _(CurrentTargetLock);
-		FTraceAuxiliary::OnTraceStopped.Broadcast(CurrentTraceTarget.TraceType, CurrentTraceTarget.TraceDest);
-
+		StopedDest = CurrentTraceTarget.TraceDest;
+		StopedType = CurrentTraceTarget.TraceType;
 		CurrentTraceTarget.TraceType = FTraceAuxiliary::EConnectionType::None;
 		CurrentTraceTarget.TraceDest.Reset();
 	}
+
+	FTraceAuxiliary::OnTraceStopped.Broadcast(StopedType, StopedDest);
+	
 	return true;
 }
 
