Fixed some constness and renamed some TOptional variables to standard forms.

diff --git a/Engine/Source/Runtime/Core/Public/Misc/Optional.h b/Engine/Source/Runtime/Core/Public/Misc/Optional.h
--- a/Engine/Source/Runtime/Core/Public/Misc/Optional.h
+++ b/Engine/Source/Runtime/Core/Public/Misc/Optional.h
@@ -57,46 +57,46 @@ public:
 	}
 
 	/** Copy/Move construction */
-	TOptional(const TOptional& InValue)
+	TOptional(const TOptional& Other)
 		: bIsSet(false)
 	{
-		if (InValue.bIsSet)
+		if (Other.bIsSet)
 		{
-			new(&Value) OptionalType(*(const OptionalType*)&InValue.Value);
+			new(&Value) OptionalType(*(const OptionalType*)&Other.Value);
 			bIsSet = true;
 		}
 	}
-	TOptional(TOptional&& InValue)
+	TOptional(TOptional&& Other)
 		: bIsSet(false)
 	{
-		if (InValue.bIsSet)
+		if (Other.bIsSet)
 		{
-			new(&Value) OptionalType(MoveTempIfPossible(*(OptionalType*)&InValue.Value));
+			new(&Value) OptionalType(MoveTempIfPossible(*(OptionalType*)&Other.Value));
 			bIsSet = true;
 		}
 	}
 
-	TOptional& operator=(const TOptional& InValue)
+	TOptional& operator=(const TOptional& Other)
 	{
-		if (&InValue != this)
+		if (&Other != this)
 		{
 			Reset();
-			if (InValue.bIsSet)
+			if (Other.bIsSet)
 			{
-				new(&Value) OptionalType(*(const OptionalType*)&InValue.Value);
+				new(&Value) OptionalType(*(const OptionalType*)&Other.Value);
 				bIsSet = true;
 			}
 		}
 		return *this;
 	}
-	TOptional& operator=(TOptional&& InValue)
+	TOptional& operator=(TOptional&& Other)
 	{
-		if (&InValue != this)
+		if (&Other != this)
 		{
 			Reset();
-			if (InValue.bIsSet)
+			if (Other.bIsSet)
 			{
-				new(&Value) OptionalType(MoveTempIfPossible(*(OptionalType*)&InValue.Value));
+				new(&Value) OptionalType(MoveTempIfPossible(*(OptionalType*)&Other.Value));
 				bIsSet = true;
 			}
 		}
@@ -105,7 +105,7 @@ public:
 
 	TOptional& operator=(const OptionalType& InValue)
 	{
-		if (&InValue != (OptionalType*)&Value)
+		if (&InValue != (const OptionalType*)&Value)
 		{
 			Reset();
 			new(&Value) OptionalType(InValue);
@@ -115,7 +115,7 @@ public:
 	}
 	TOptional& operator=(OptionalType&& InValue)
 	{
-		if (&InValue != (OptionalType*)&Value)
+		if (&InValue != (const OptionalType*)&Value)
 		{
 			Reset();
 			new(&Value) OptionalType(MoveTempIfPossible(InValue));
@@ -145,22 +145,22 @@ public:
 		return *Result;
 	}
 
-	friend bool operator==(const TOptional& lhs, const TOptional& rhs)
+	friend bool operator==(const TOptional& Lhs, const TOptional& Rhs)
 	{
-		if (lhs.bIsSet != rhs.bIsSet)
+		if (Lhs.bIsSet != Rhs.bIsSet)
 		{
 			return false;
 		}
-		if (!lhs.bIsSet) // both unset
+		if (!Lhs.bIsSet) // both unset
 		{
 			return true;
 		}
-		return (*(OptionalType*)&lhs.Value) == (*(OptionalType*)&rhs.Value);
+		return (*(const OptionalType*)&Lhs.Value) == (*(const OptionalType*)&Rhs.Value);
 	}
 
-	friend bool operator!=(const TOptional& lhs, const TOptional& rhs)
+	friend bool operator!=(const TOptional& Lhs, const TOptional& Rhs)
 	{
-		return !(lhs == rhs);
+		return !(Lhs == Rhs);
 	}
 
 	void Serialize(FArchive& Ar)
