Added comments to TOptional's emplacement functions to guide users to the preferred solution when used with non-public constructors, similar to comments in MakeShared or TArray::Emplace.

diff --git a/Engine/Source/Runtime/Core/Public/Misc/Optional.h b/Engine/Source/Runtime/Core/Public/Misc/Optional.h
--- a/Engine/Source/Runtime/Core/Public/Misc/Optional.h
+++ b/Engine/Source/Runtime/Core/Public/Misc/Optional.h
@@ -35,6 +35,26 @@ public:
 	template <typename... ArgTypes>
 	explicit TOptional(EInPlace, ArgTypes&&... Args)
 	{
+		// If this fails to compile when trying to call TOptional(EInPlace, ...) with a non-public constructor,
+		// do not make TOptional a friend.
+		//
+		// Instead, prefer this pattern:
+		//
+		//     class FMyType
+		//     {
+		//     private:
+		//         struct FPrivateToken { explicit FPrivateToken() = default; };
+		//
+		//     public:
+		//         // This has an equivalent access level to a private constructor,
+		//         // as only friends of FMyType will have access to FPrivateToken,
+		//         // but the TOptional constructor can legally call it since it's public.
+		//         explicit FMyType(FPrivateToken, int32 Int, float Real, const TCHAR* String);
+		//     };
+		//
+		//     // Won't compile if the caller doesn't have access to FMyType::FPrivateToken
+		//     TOptional<FMyType> Opt(InPlace, FMyType::FPrivateToken{}, 5, 3.14f, TEXT("Banana"));
+		//
 		new(&Value) OptionalType(Forward<ArgTypes>(Args)...);
 		bIsSet = true;
 	}
@@ -140,6 +160,29 @@ public:
 	OptionalType& Emplace(ArgsType&&... Args)
 	{
 		Reset();
+
+		// If this fails to compile when trying to call Emplace with a non-public constructor,
+		// do not make TOptional a friend.
+		//
+		// Instead, prefer this pattern:
+		//
+		//     class FMyType
+		//     {
+		//     private:
+		//         struct FPrivateToken { explicit FPrivateToken() = default; };
+		//
+		//     public:
+		//         // This has an equivalent access level to a private constructor,
+		//         // as only friends of FMyType will have access to FPrivateToken,
+		//         // but Emplace can legally call it since it's public.
+		//         explicit FMyType(FPrivateToken, int32 Int, float Real, const TCHAR* String);
+		//     };
+		//
+		//     TOptional<FMyType> Opt:
+		//
+		//     // Won't compile if the caller doesn't have access to FMyType::FPrivateToken
+		//     Opt.Emplace(FMyType::FPrivateToken{}, 5, 3.14f, TEXT("Banana"));
+		//
 		OptionalType* Result = new(&Value) OptionalType(Forward<ArgsType>(Args)...);
 		bIsSet = true;
 		return *Result;
