Fix missing include to TFunction

diff --git a/Engine/Source/Runtime/Core/Private/AutoRTFM/TaskArray.h b/Engine/Source/Runtime/Core/Private/AutoRTFM/TaskArray.h
--- a/Engine/Source/Runtime/Core/Private/AutoRTFM/TaskArray.h
+++ b/Engine/Source/Runtime/Core/Private/AutoRTFM/TaskArray.h
@@ -3,7 +3,7 @@
 #pragma once
 
 #include "Containers/Array.h"
-
+#include "Templates/Function.h"
 #include "Utils.h"
 
 namespace AutoRTFM
