Unified related TOptional functions and expanded them out onto their own lines for standards compliance.

diff --git a/Engine/Source/Runtime/Core/Public/Misc/Optional.h b/Engine/Source/Runtime/Core/Public/Misc/Optional.h
--- a/Engine/Source/Runtime/Core/Public/Misc/Optional.h
+++ b/Engine/Source/Runtime/Core/Public/Misc/Optional.h
@@ -235,25 +235,59 @@ public:
 	}
 
 	/** @return true when the value is meaningful; false if calling GetValue() is undefined. */
-	bool IsSet() const { return bIsSet; }
-	FORCEINLINE explicit operator bool() const { return bIsSet; }
+	bool IsSet() const
+	{
+		return bIsSet;
+	}
+	FORCEINLINE explicit operator bool() const
+	{
+		return IsSet();
+	}
 
 	/** @return The optional value; undefined when IsSet() returns false. */
-	const OptionalType& GetValue() const { checkf(IsSet(), TEXT("It is an error to call GetValue() on an unset TOptional. Please either check IsSet() or use Get(DefaultValue) instead.")); return *(OptionalType*)&Value; }
-		  OptionalType& GetValue()		 { checkf(IsSet(), TEXT("It is an error to call GetValue() on an unset TOptional. Please either check IsSet() or use Get(DefaultValue) instead.")); return *(OptionalType*)&Value; }
+	OptionalType& GetValue()
+	{
+		checkf(IsSet(), TEXT("It is an error to call GetValue() on an unset TOptional. Please either check IsSet() or use Get(DefaultValue) instead."));
+		return *(OptionalType*)&Value;
+	}
+	FORCEINLINE const OptionalType& GetValue() const
+	{
+		return const_cast<TOptional*>(this)->GetValue();
+	}
 
-	const OptionalType* operator->() const { return &GetValue(); }
-		  OptionalType* operator->()	   { return &GetValue(); }
+	OptionalType* operator->()
+	{
+		return &GetValue();
+	}
+	FORCEINLINE const OptionalType* operator->() const
+	{
+		return const_cast<TOptional*>(this)->operator->();
+	}
 
-	const OptionalType& operator*() const { return GetValue(); }
-		  OptionalType& operator*()		  { return GetValue(); }
+	OptionalType& operator*()
+	{
+		return GetValue();
+	}
+	FORCEINLINE const OptionalType& operator*() const
+	{
+		return const_cast<TOptional*>(this)->operator*();
+	}
 
 	/** @return The optional value when set; DefaultValue otherwise. */
-	const OptionalType& Get(const OptionalType& DefaultValue UE_LIFETIMEBOUND) const UE_LIFETIMEBOUND { return IsSet() ? *(OptionalType*)&Value : DefaultValue; }
+	const OptionalType& Get(const OptionalType& DefaultValue UE_LIFETIMEBOUND) const UE_LIFETIMEBOUND
+	{
+		return IsSet() ? *(const OptionalType*)&Value : DefaultValue;
+	}
 
 	/** @return A pointer to the optional value when set, nullptr otherwise. */
-	OptionalType* GetPtrOrNull() { return IsSet() ? (OptionalType*)&Value : nullptr; }
-	const OptionalType* GetPtrOrNull() const { return IsSet() ? (const OptionalType*)&Value : nullptr; }
+	OptionalType* GetPtrOrNull()
+	{
+		return IsSet() ? (OptionalType*)&Value : nullptr;
+	}
+	FORCEINLINE const OptionalType* GetPtrOrNull() const
+	{
+		return const_cast<TOptional*>(this)->GetPtrOrNull();
+	}
 
 private:
 	TTypeCompatibleBytes<OptionalType> Value;
