Pause StallDetectors on thread when SlowTask with dialog is started
* SlowTask with dialog is the only case slowtask will broadcast the delegates
* Specific handling for case where StallDetector is added to scope while slowtask is present

diff --git a/Engine/Source/Runtime/Core/Private/ProfilingDebugging/StallDetector.cpp b/Engine/Source/Runtime/Core/Private/ProfilingDebugging/StallDetector.cpp
--- a/Engine/Source/Runtime/Core/Private/ProfilingDebugging/StallDetector.cpp
+++ b/Engine/Source/Runtime/Core/Private/ProfilingDebugging/StallDetector.cpp
@@ -300,6 +300,91 @@ void UE::FStallDetectorStats::TabulateStats(TArray<TabulatedResult>& TabulatedRe
 	}
 }
 
+// globals
+
+/**
+ * Helper to initialize TLS variable
+ */
+class ThreadLocalSlotAcquire
+{
+public:
+	ThreadLocalSlotAcquire()
+	{
+		Slot = FPlatformTLS::AllocTlsSlot();
+		FPlatformTLS::SetTlsValue(Slot, nullptr);
+	}
+
+	~ThreadLocalSlotAcquire()
+	{
+		FPlatformTLS::FreeTlsSlot(Slot);
+		Slot = ~0;
+	}
+
+	uint32 GetSlot()
+	{
+		return Slot;
+	}
+
+private:
+	uint32 Slot = ~0;
+};
+
+static ThreadLocalSlotAcquire ThreadLocalSlot;
+
+TAtomic<uint64> UE::FStallDetector::UIDGenerator = 0;
+
+/**
+ * Pause on slow task
+ */
+namespace UE::StallDetector::Private
+{
+	FDelegateHandle GSlowTaskStartDelegate;
+	FDelegateHandle GSlowTaskFinalizedDelegate;
+	
+	uint32 GSlowTaskReportingThread = ~0;
+	uint32 GSlowTaskCount = 0;
+
+	void RegisterSlowTaskHandlers()
+	{
+		GSlowTaskStartDelegate = GWarn->OnStartSlowTask().AddLambda([](const FText& TaskName)
+		{
+			if (GSlowTaskReportingThread == ~0)
+			{
+				GSlowTaskReportingThread = FPlatformTLS::GetCurrentThreadId();
+			}
+			checkf(GSlowTaskReportingThread == FPlatformTLS::GetCurrentThreadId(), TEXT("Expected SlowTasks with Dialog to only be called from one thread"));
+			++GSlowTaskCount;
+			
+			const double Seconds = FStallDetector::Seconds();
+			if (Seconds != InvalidSeconds)
+			{
+				for (FStallDetector* Current = static_cast<FStallDetector*>(FPlatformTLS::GetTlsValue(ThreadLocalSlot.GetSlot())); Current; Current = Current->GetParent())
+				{
+					Current->Pause(Seconds);
+				}
+			}
+		});
+
+		GSlowTaskFinalizedDelegate = GWarn->OnFinalizeSlowTask().AddLambda([](const FText& TaskName, double DurationInSeconds)
+		{
+			--GSlowTaskCount;
+			
+			const double Seconds = FStallDetector::Seconds();
+			if (Seconds != InvalidSeconds)
+			{
+				for (FStallDetector* Current = static_cast<FStallDetector*>(FPlatformTLS::GetTlsValue(ThreadLocalSlot.GetSlot())); Current; Current = Current->GetParent())
+				{
+					Current->Resume(Seconds);
+				}
+			}
+		});
+	}
+	void UnregisterSlowTaskHandlers()
+	{
+		GWarn->OnFinalizeSlowTask().Remove(GSlowTaskFinalizedDelegate);
+		GWarn->OnStartSlowTask().Remove(GSlowTaskStartDelegate);
+	}
+}
 
 /**
 * Stall Timer
@@ -310,6 +395,13 @@ UE::FStallTimer::FStallTimer()
 , LastCheckSeconds(InvalidSeconds)
 , RemainingSeconds(InvalidSeconds)
 {
+	using namespace UE::StallDetector::Private;
+	// Handle the case where a stall detector scope is entered while a slowtask is already in progress
+	// Update the PauseCount to the number of resumes that will be called
+	if (GSlowTaskReportingThread == FPlatformTLS::GetCurrentThreadId())
+	{
+		PauseCount = StallDetector::Private::GSlowTaskCount;
+	}	
 }
 
 void UE::FStallTimer::Reset(const double InSeconds, const double InRemainingSeconds)
@@ -378,35 +470,6 @@ void UE::FStallTimer::Resume(const double InSeconds)
 /**
 * Stall Detector
 **/
-
-// globals
-class ThreadLocalSlotAcquire
-{
-public:
-	ThreadLocalSlotAcquire()
-	{
-		Slot = FPlatformTLS::AllocTlsSlot();
-		FPlatformTLS::SetTlsValue(Slot, nullptr);
-	}
-
-	~ThreadLocalSlotAcquire()
-	{
-		FPlatformTLS::FreeTlsSlot(Slot);
-		Slot = ~0;
-	}
-
-	uint32 GetSlot()
-	{
-		return Slot;
-	}
-
-private:
-	uint32 Slot = ~0;
-};
-static ThreadLocalSlotAcquire ThreadLocalSlot;
-
-TAtomic<uint64> UE::FStallDetector::UIDGenerator = 0;
-
 UE::FStallDetector::FStallDetector(FStallDetectorStats& InStats)
 	: Stats(InStats)
 	, Parent(nullptr)
@@ -602,6 +665,8 @@ void UE::FStallDetector::Startup()
 		{
 			if (StallDetectorThread == nullptr)
 			{
+				StallDetector::Private::RegisterSlowTaskHandlers();
+				
 				StallDetectorThread = FRunnableThread::Create(StallDetectorRunnable, TEXT("StallDetectorThread"));
 				check(StallDetectorThread);
 
@@ -622,13 +687,15 @@ void UE::FStallDetector::Shutdown()
 	if (--InitCount == 0)
 	{
 		UE_LOG(LogStall, Log, TEXT("Shutdown..."));
+		
+		StallDetector::Private::UnregisterSlowTaskHandlers();
 
 		delete StallDetectorThread;
 		StallDetectorThread = nullptr;
 
 		delete StallDetectorRunnable;
 		StallDetectorRunnable = nullptr;
-
+		
 		UE_LOG(LogStall, Log, TEXT("Shutdown complete."));
 	}
 	check(InitCount >= 0);
