Fix AppleHD A8 CRC crash due to invalid hardware cpu check.

#rnx

diff --git a/Engine/Source/Runtime/Core/Private/IOS/IOSPlatformMisc.cpp b/Engine/Source/Runtime/Core/Private/IOS/IOSPlatformMisc.cpp
--- a/Engine/Source/Runtime/Core/Private/IOS/IOSPlatformMisc.cpp
+++ b/Engine/Source/Runtime/Core/Private/IOS/IOSPlatformMisc.cpp
@@ -1881,8 +1881,10 @@ bool FIOSPlatformMisc::CPUHasHwCrcSupport()
 {
 	// HW CRC instructions support is available on Apple A10+
 	static int HwCrcSupported = -1;
-	if (HwCrcSupported != -1)
+	if (HwCrcSupported == -1)
 	{
+		HwCrcSupported = 0;
+		
 		const FString DeviceIDString = GetIOSDeviceIDString();
 		if (DeviceIDString.StartsWith(TEXT("iPod")))
 		{
diff --git a/Engine/Source/Runtime/Core/Private/Misc/Crc.cpp b/Engine/Source/Runtime/Core/Private/Misc/Crc.cpp
--- a/Engine/Source/Runtime/Core/Private/Misc/Crc.cpp
+++ b/Engine/Source/Runtime/Core/Private/Misc/Crc.cpp
@@ -5,7 +5,11 @@
 #include "Templates/UnrealTemplate.h"
 #include "Misc/ByteSwap.h"
 
-#if defined(__clang__) && defined(__ARM_FEATURE_CRYPTO)
+// As Apple sets __ARM_FEATURE_CRYPTO by default for all of it's arm chips even when not fully supported,
+// it can not be relied upon to detect the lack of hardware CRC in an A8 in an AppleHDTV, so manually force the CRC check
+#if PLATFORM_TVOS
+#	define DETECT_HW_CRC32_SUPPORT_IN_RUNTIME 1
+#elif defined(__clang__) && defined(__ARM_FEATURE_CRYPTO)
 #	define DETECT_HW_CRC32_SUPPORT_IN_RUNTIME 0
 #elif PLATFORM_ANDROID_ARM64 || PLATFORM_LINUXARM64 || PLATFORM_IOS
 #	define DETECT_HW_CRC32_SUPPORT_IN_RUNTIME 1
@@ -370,6 +374,7 @@ struct CRCRuntimeHWSupportDetector
 {
 	CRCRuntimeHWSupportDetector()
 	{
+		FCrc::MemCrc32Func = &MemCrc32SW;
 #if PLATFORM_ANDROID_ARM64
 		if (android_getCpuFeatures() & ANDROID_CPU_ARM64_FEATURE_CRC32)
 		{
