Add missing include for StallDetector.cpp which fails to compile when not in unity build

diff --git a/Engine/Source/Runtime/Core/Private/ProfilingDebugging/StallDetector.cpp b/Engine/Source/Runtime/Core/Private/ProfilingDebugging/StallDetector.cpp
--- a/Engine/Source/Runtime/Core/Private/ProfilingDebugging/StallDetector.cpp
+++ b/Engine/Source/Runtime/Core/Private/ProfilingDebugging/StallDetector.cpp
@@ -10,6 +10,7 @@
 #include "HAL/PlatformTime.h"
 #include "HAL/Runnable.h"
 #include "HAL/RunnableThread.h"
+#include "Misc/FeedbackContext.h"
 #include "Tasks/Task.h"
 #include "Templates/Greater.h"
 
