Limit android setthreadpriority log spam.
 - some OS versions may prevent thread scheduling or priority changes, resulting in constant error spam that saturates the logs.

diff --git a/Engine/Source/Runtime/Core/Private/Android/AndroidPlatformProcess.cpp b/Engine/Source/Runtime/Core/Private/Android/AndroidPlatformProcess.cpp
--- a/Engine/Source/Runtime/Core/Private/Android/AndroidPlatformProcess.cpp
+++ b/Engine/Source/Runtime/Core/Private/Android/AndroidPlatformProcess.cpp
@@ -305,16 +305,28 @@ void FRunnableThreadAndroid::SetThreadPriority(pthread_t InThread, EThreadPriori
 	struct sched_param Sched = { };
 	pthread_getschedparam(InThread, &InitialPolicy, &Sched);
 
-	if (InitialPolicy != NewPolicy && sched_setscheduler(InThread, NewPolicy, &Sched)!=0)
+	// The task system can call this function with high frequency,
+	// in cases where the OS consistently prevents the requested pri changes we must limit logging.
+	const bool bCanLog = ErrorLogLimit >= 0;
+	bool bErrorEncountered = false;
+
+	if (InitialPolicy != NewPolicy && (sched_setscheduler(InThread, NewPolicy, &Sched)!=0) && bCanLog)
 	{
+		bErrorEncountered = true;
 		UE_LOG(LogHAL, Error, TEXT("Failed to set %s thread scheduler, tid %d from %d to %d (errno %d)"), *GetThreadName(), ThreadID, InitialPolicy, NewPolicy, errno);
 	}
 
 	int InitialNice = getpriority(PRIO_PROCESS, ThreadID);
 	int NewNice = AndroidThreadPriorityToNice(NewPriority);
-	if (InitialNice != NewNice && setpriority(PRIO_PROCESS, ThreadID, NewNice) != 0)
+	if (InitialNice != NewNice && (setpriority(PRIO_PROCESS, ThreadID, NewNice) != 0) && bCanLog)
 	{
+		bErrorEncountered = true;
 		UE_LOG(LogHAL, Error, TEXT("Failed to set %s thread priority, tid %d from %d to %d (errno %d)"), *GetThreadName(), ThreadID, InitialNice, NewNice, errno);
 	}
+
+	if (bCanLog && bErrorEncountered)
+	{
+		ErrorLogLimit--;
+	}
 }
 #endif
diff --git a/Engine/Source/Runtime/Core/Public/Android/AndroidPlatformRunnableThread.h b/Engine/Source/Runtime/Core/Public/Android/AndroidPlatformRunnableThread.h
--- a/Engine/Source/Runtime/Core/Public/Android/AndroidPlatformRunnableThread.h
+++ b/Engine/Source/Runtime/Core/Public/Android/AndroidPlatformRunnableThread.h
@@ -84,6 +84,7 @@ private:
 	}
 
 #if ANDROID_USE_NICE_VALUE_THREADPRIORITY
+	int32 ErrorLogLimit = 5;
 	virtual void SetThreadPriority(pthread_t InThread, EThreadPriority NewPriority) override;
 #endif
 };
