Unix MemoryTrace implementation

#tests DS

diff --git a/Engine/Source/Runtime/Core/Core.Build.cs b/Engine/Source/Runtime/Core/Core.Build.cs
--- a/Engine/Source/Runtime/Core/Core.Build.cs
+++ b/Engine/Source/Runtime/Core/Core.Build.cs
@@ -152,6 +152,13 @@ public class Core : ModuleRules
 
 			PublicSystemIncludePaths.Add(Path.Combine(Target.UEThirdPartySourceDirectory, "mimalloc/include"));
 			PrivateDefinitions.Add("PLATFORM_BUILDS_MIMALLOC=1");
+
+			if (Target.Configuration != UnrealTargetConfiguration.Shipping && Target.Type != TargetType.Program)
+			{
+				PublicDefinitions.Add("UE_MEMORY_TRACE_AVAILABLE=1");
+				PublicDefinitions.Add("UE_MEMORY_TAGS_TRACE_ENABLED=1");
+				PublicDefinitions.Add("UE_CALLSTACK_TRACE_ENABLED=1");
+			}
 		}
 
 		if (Target.bCompileICU == true)
diff --git a/Engine/Source/Runtime/Core/Private/ProfilingDebugging/Unix/UnixCallstackTrace.cpp b/Engine/Source/Runtime/Core/Private/ProfilingDebugging/Unix/UnixCallstackTrace.cpp
new file mode 100644
--- /dev/null
+++ b/Engine/Source/Runtime/Core/Private/ProfilingDebugging/Unix/UnixCallstackTrace.cpp
@@ -0,0 +1,135 @@
+// Copyright Epic Games, Inc. All Rights Reserved.
+
+#include "ProfilingDebugging/CallstackTrace.h"
+
+#if UE_CALLSTACK_TRACE_ENABLED
+
+#include "CoreTypes.h"
+#include "HAL/Platform.h"
+#include "HAL/PlatformStackWalk.h"
+#include "ProfilingDebugging/CallstackTracePrivate.h"
+
+////////////////////////////////////////////////////////////////////////////////
+class FBacktracer
+{
+public:
+	FBacktracer(FMalloc* InMalloc);
+	~FBacktracer();
+	static FBacktracer*	Get();
+	uint32 GetBacktraceId(void* ReturnAddress);
+
+private:
+	static FBacktracer* Instance;
+	FMalloc* Malloc;
+	FCallstackTracer CallstackTracer;
+};
+
+////////////////////////////////////////////////////////////////////////////////
+FBacktracer* FBacktracer::Instance = nullptr;
+
+////////////////////////////////////////////////////////////////////////////////
+FBacktracer::FBacktracer(FMalloc* InMalloc)
+	: Malloc(InMalloc)
+	, CallstackTracer(InMalloc)
+{
+	Instance = this;
+}
+
+////////////////////////////////////////////////////////////////////////////////
+FBacktracer::~FBacktracer()
+{
+}
+
+////////////////////////////////////////////////////////////////////////////////
+FBacktracer* FBacktracer::Get()
+{
+	return Instance;
+}
+
+////////////////////////////////////////////////////////////////////////////////
+uint32 FBacktracer::GetBacktraceId(void* ReturnAddress)
+{
+#if !UE_BUILD_SHIPPING
+	uint64 StackFrames[256];
+	int32 NumStackFrames = FPlatformStackWalk::CaptureStackBackTrace(StackFrames, UE_ARRAY_COUNT(StackFrames));
+	if (NumStackFrames > 0)
+	{
+		FCallstackTracer::FBacktraceEntry BacktraceEntry;
+		uint64 BacktraceId = 0;
+		uint32 FrameIdx = 0;
+		bool bUseAddress = false;
+		for (int32 Index = 0; Index < NumStackFrames; Index++)
+		{
+			if (!bUseAddress)
+			{
+				// start using backtrace only after ReturnAddress
+				if (StackFrames[Index] == (uint64)ReturnAddress)
+				{
+					bUseAddress = true;
+				}
+			}
+			if (bUseAddress || NumStackFrames == 1)
+			{
+				uint64 RetAddr = StackFrames[Index];
+				StackFrames[FrameIdx++] = RetAddr;
+
+				// This is a simple order-dependent LCG. Should be sufficient enough
+				BacktraceId += RetAddr;
+				BacktraceId *= 0x30be8efa499c249dull;
+			}
+		}
+
+		// Save the collected id
+		BacktraceEntry.Hash = BacktraceId;
+		BacktraceEntry.FrameCount = FrameIdx;
+		BacktraceEntry.Frames = StackFrames;
+
+		// Add to queue to be processed. This might block until there is room in the
+		// queue (i.e. the processing thread has caught up processing).
+		return CallstackTracer.AddCallstack(BacktraceEntry);
+	}
+#endif
+
+	return 0;
+}
+
+////////////////////////////////////////////////////////////////////////////////
+void Modules_Initialize();
+
+////////////////////////////////////////////////////////////////////////////////
+void CallstackTrace_CreateInternal(FMalloc* Malloc)
+{
+	if (FBacktracer::Get() != nullptr)
+	{
+		return;
+	}
+
+	// Allocate, construct and intentionally leak backtracer
+	void* Alloc = Malloc->Malloc(sizeof(FBacktracer), alignof(FBacktracer));
+	new (Alloc) FBacktracer(Malloc);
+}
+
+////////////////////////////////////////////////////////////////////////////////
+void CallstackTrace_InitializeInternal()
+{
+	if (FBacktracer* Instance = FBacktracer::Get())
+	{
+#if !UE_BUILD_SHIPPING
+		Modules_Initialize();
+#endif
+	}
+}
+
+////////////////////////////////////////////////////////////////////////////////
+uint32 CallstackTrace_GetCurrentId()
+{
+	void* ReturnAddress = PLATFORM_RETURN_ADDRESS_FOR_CALLSTACKTRACING();
+	if (FBacktracer* Instance = FBacktracer::Get())
+	{
+		return Instance->GetBacktraceId(ReturnAddress);
+	}
+
+	return 0;
+}
+
+#endif
\ No newline at end of file
diff --git a/Engine/Source/Runtime/Core/Private/ProfilingDebugging/Unix/UnixMemoryTrace.cpp b/Engine/Source/Runtime/Core/Private/ProfilingDebugging/Unix/UnixMemoryTrace.cpp
new file mode 100644
--- /dev/null
+++ b/Engine/Source/Runtime/Core/Private/ProfilingDebugging/Unix/UnixMemoryTrace.cpp
@@ -0,0 +1,43 @@
+// Copyright Epic Games, Inc. All Rights Reserved.
+
+#include "ProfilingDebugging/MemoryTrace.h"
+
+#if UE_MEMORY_TRACE_ENABLED
+
+#include "CoreTypes.h"
+#include "HAL/Platform.h"
+
+////////////////////////////////////////////////////////////////////////////////
+FMalloc* MemoryTrace_CreateInternal(FMalloc*, int, const ANSICHAR* const*);
+
+////////////////////////////////////////////////////////////////////////////////
+FMalloc* MemoryTrace_Create(FMalloc* InMalloc)
+{
+	if (FILE* CmdLineFile = fopen("/proc/self/cmdline", "r"))
+	{
+		const char* ArgV[255] = {};
+		int32 ArgC = 0;
+
+		char * Arg = nullptr;
+		size_t Size = 0;
+		while (getdelim(&Arg, &Size, 0, CmdLineFile) != -1)
+		{
+			ArgV[ArgC++] = Arg;
+			Arg = nullptr; // getdelim will allocate buffer for next Arg
+		}
+		fclose(CmdLineFile);
+
+		FMalloc* OutMalloc = MemoryTrace_CreateInternal(InMalloc, ArgC, ArgV);
+
+		// cleanup after getdelim
+		while (ArgC > 0)
+		{
+			free((void*)ArgV[--ArgC]);
+		}
+
+		return (OutMalloc != nullptr) ? OutMalloc : InMalloc;
+	}
+	return InMalloc;
+}
+
+#endif // UE_MEMORY_TRACE_ENABLED
diff --git a/Engine/Source/Runtime/Core/Private/ProfilingDebugging/Unix/UnixModuleDiagnostics.cpp b/Engine/Source/Runtime/Core/Private/ProfilingDebugging/Unix/UnixModuleDiagnostics.cpp
new file mode 100644
--- /dev/null
+++ b/Engine/Source/Runtime/Core/Private/ProfilingDebugging/Unix/UnixModuleDiagnostics.cpp
@@ -0,0 +1,72 @@
+// Copyright Epic Games, Inc. All Rights Reserved.
+
+#include "CoreTypes.h"
+#include "ProfilingDebugging/ModuleDiagnostics.h"
+#include <link.h>
+
+////////////////////////////////////////////////////////////////////////////////
+void Modules_Initialize()
+{
+#if !UE_BUILD_SHIPPING
+	using namespace UE::Trace;
+
+	constexpr uint32 SizeOfSymbolFormatString = 5;
+	UE_TRACE_LOG(Diagnostics, ModuleInit, ModuleChannel, sizeof(TCHAR) * SizeOfSymbolFormatString)
+		<< ModuleInit.SymbolFormat("dwarf", SizeOfSymbolFormatString)
+		<< ModuleInit.ModuleBaseShift(uint8(0));
+
+	auto IterateCallback = [](struct dl_phdr_info *Info, size_t /*size*/, void *Data)
+	{
+		int32 TotalMemSize = 0;
+		uint64 RealBase = Info->dlpi_addr;
+		bool bRealBaseSet = false;
+
+		constexpr uint32 GNUSectionNameSize = 4;
+		constexpr uint32 BuildIdSize = 20;
+		uint8 BuildId[BuildIdSize] = {0};
+		FMemory::Memzero(BuildId, BuildIdSize);
+		bool bBuildIdSet = false;
+		for (int SectionIdx = 0; SectionIdx < Info->dlpi_phnum; SectionIdx++)
+		{
+			TotalMemSize += Info->dlpi_phdr[SectionIdx].p_memsz;
+
+			uint64 Offset = Info->dlpi_addr + Info->dlpi_phdr[SectionIdx].p_vaddr;
+			uint32 Type = Info->dlpi_phdr[SectionIdx].p_type;
+
+			if (!bRealBaseSet && Type == PT_LOAD)
+			{
+				RealBase = Offset;
+				bRealBaseSet = true;
+			}
+			if (!bBuildIdSet && Type == PT_NOTE)
+			{
+				ElfW(Nhdr)* Note = (ElfW(Nhdr)*)Offset;
+				char* NoteName = (char*)Note + sizeof(ElfW(Nhdr));
+				uint8* NoteDesc = (uint8*)Note + sizeof(ElfW(Nhdr)) + Note->n_namesz;
+				if (Note->n_namesz == GNUSectionNameSize && Note->n_type == NT_GNU_BUILD_ID && strcmp(NoteName, ELF_NOTE_GNU) == 0)
+				{
+					FMemory::Memcpy(BuildId, NoteDesc, FMath::Min(BuildIdSize, Note->n_descsz));
+					bBuildIdSet = true;
+				}
+			}
+		}
+
+		FString ImageName = FString(Info->dlpi_name);
+		if (ImageName.IsEmpty())
+		{
+			constexpr bool bRemoveExtension = false;
+			ImageName = FPlatformProcess::ExecutableName(bRemoveExtension);
+		}
+		ImageName = FPaths::GetCleanFilename(ImageName);
+
+		UE_TRACE_LOG(Diagnostics, ModuleLoad, ModuleChannel, sizeof(TCHAR) * ImageName.Len() + BuildIdSize)
+			<< ModuleLoad.Name(*ImageName, ImageName.Len())
+			<< ModuleLoad.Base(RealBase)
+			<< ModuleLoad.Size(TotalMemSize)
+			<< ModuleLoad.ImageId(BuildId, BuildIdSize);
+
+		return 0;
+	};
+	dl_iterate_phdr(IterateCallback, nullptr);
+#endif
+}
diff --git a/Engine/Source/Runtime/Core/Public/Unix/UnixPlatform.h b/Engine/Source/Runtime/Core/Public/Unix/UnixPlatform.h
--- a/Engine/Source/Runtime/Core/Public/Unix/UnixPlatform.h
+++ b/Engine/Source/Runtime/Core/Public/Unix/UnixPlatform.h
@@ -70,6 +70,8 @@ typedef FUnixPlatformTypes FPlatformTypes;
 #define PLATFORM_SUPPORTS_VIRTUAL_TEXTURE_STREAMING		1
 #define PLATFORM_USE_ANSI_POSIX_MALLOC					1
 
+#define PLATFORM_RETURN_ADDRESS_FOR_CALLSTACKTRACING	PLATFORM_RETURN_ADDRESS
+
 #if PLATFORM_CPU_X86_FAMILY
 	#define PLATFORM_BREAK()							__asm__ volatile("int $0x03")
 #else
