#ue5 - Change TQuat::FindBetween_Helper() when dealing with anti-parallel (opposite direction) vectors. This restores the old behavior to prefer rotation around the Z-axis when dealing with vectors with X or Y dominant components, which causes rotations in gameplay to avoid adding roll when the quat is used for interpolation. Added more validation tests to compare old and new methods.

diff --git a/Engine/Source/Runtime/Core/Private/Math/UnrealMath.cpp b/Engine/Source/Runtime/Core/Private/Math/UnrealMath.cpp
--- a/Engine/Source/Runtime/Core/Private/Math/UnrealMath.cpp
+++ b/Engine/Source/Runtime/Core/Private/Math/UnrealMath.cpp
@@ -1428,8 +1428,8 @@ namespace Math
 			const T Y = FMath::Abs(A.Y);
 			const T Z = FMath::Abs(A.Z);
 
-			// Find most orthogonal basis 
-			const TVector<T> Basis = X < Y ? (X < Z ? TVector<T>::XAxisVector : TVector<T>::ZAxisVector) : (Y < Z ? TVector<T>::YAxisVector : TVector<T>::ZAxisVector);
+			// Find orthogonal basis. 
+			const TVector<T> Basis = (X > Y && X > Z) ? TVector<T>::YAxisVector : -TVector<T>::XAxisVector;
 
 			//Result = FVector::CrossProduct(A, Basis);
 			Result = TQuat<T>(
diff --git a/Engine/Source/Runtime/Core/Tests/Math/UnrealMathTest.cpp b/Engine/Source/Runtime/Core/Tests/Math/UnrealMathTest.cpp
--- a/Engine/Source/Runtime/Core/Tests/Math/UnrealMathTest.cpp
+++ b/Engine/Source/Runtime/Core/Tests/Math/UnrealMathTest.cpp
@@ -761,6 +761,38 @@ FORCENOINLINE FQuat4f FindBetween_Old(const FVector3f& vec1, const FVector3f& ve
 		cosHalfAng);
 }
 
+FQuat4f FindBetween_Helper_5_2(const FVector3f& A, const FVector3f& B, float NormAB)
+{
+	float W = NormAB + FVector3f::DotProduct(A, B);
+	FQuat4f Result;
+
+	if (W >= 1e-6f * NormAB)
+	{
+		//Axis = FVector::CrossProduct(A, B);
+		Result = FQuat4f(
+			A.Y * B.Z - A.Z * B.Y,
+			A.Z * B.X - A.X * B.Z,
+			A.X * B.Y - A.Y * B.X,
+			W);
+	}
+	else
+	{
+		// A and B point in opposite directions
+		W = 0.f;
+		Result = FMath::Abs(A.X) > FMath::Abs(A.Y)
+			? FQuat4f(-A.Z, 0.f, A.X, W)
+			: FQuat4f(0.f, -A.Z, A.Y, W);
+	}
+
+	Result.Normalize();
+	return Result;
+}
+
+FQuat4f FindBetweenNormals_5_2(const FVector3f& A, const FVector3f& B)
+{
+	const float NormAB = 1.f;
+	return FindBetween_Helper_5_2(A, B, NormAB);
+}
 
 // ROTATOR TESTS
 
@@ -3617,21 +3649,27 @@ TEST_CASE_NAMED(FVectorRegisterAbstractionTest, "System::Core::Math::Vector Regi
 						const FVector3f BNorm = (B * SignVector).GetSafeNormal();
 
 						const FQuat4f Old = FindBetween_Old(ANorm, BNorm);
+						const FQuat4f Old_5_2 = FindBetweenNormals_5_2(ANorm, BNorm);
 						const FQuat4f NewNormal = FQuat4f::FindBetweenNormals(ANorm, BNorm);
 						const FQuat4f NewVector = FQuat4f::FindBetweenVectors(A, B * SignVector);
 
 						const FVector3f RotAOld = Old.RotateVector(ANorm);
+						const FVector3f RotAOld_5_2 = Old_5_2.RotateVector(ANorm);
 						const FVector3f RotANewNormal = NewNormal.RotateVector(ANorm);
 						const FVector3f RotANewVector = NewVector.RotateVector(ANorm);
 
 						if (A.IsZero() || B.IsZero())
 						{
 							LogTest<float>(TEXT("FindBetween: Old == New (normal)"), TestQuatsEqual(Old, NewNormal, 1e-6f));
+							LogTest<float>(TEXT("FindBetween: Old_5_2 == New (normal)"), TestQuatsEqual(Old_5_2, NewNormal, 1e-6f));
 							LogTest<float>(TEXT("FindBetween: Old == New (vector)"), TestQuatsEqual(Old, NewVector, 1e-6f));
 						}
 						else
 						{
+							LogTest<float>(TEXT("FindBetween: Old_5_2 == New (normal)"), TestQuatsEqual(Old_5_2, NewNormal, 1e-6f));
+
 							LogTest<float>(TEXT("FindBetween: Old A->B"), TestFVector3Equal(RotAOld, BNorm, UE_KINDA_SMALL_NUMBER));
+							LogTest<float>(TEXT("FindBetween: Old_5_2 A->B"), TestFVector3Equal(RotAOld_5_2, BNorm, UE_KINDA_SMALL_NUMBER));
 							LogTest<float>(TEXT("FindBetween: New A->B (normal)"), TestFVector3Equal(RotANewNormal, BNorm, UE_KINDA_SMALL_NUMBER));
 							LogTest<float>(TEXT("FindBetween: New A->B (vector)"), TestFVector3Equal(RotANewVector, BNorm, UE_KINDA_SMALL_NUMBER));
 						}
