Fixed FPerThreadFreeBlockLists leak in MB3

diff --git a/Engine/Source/Runtime/Core/Private/HAL/MallocBinned3.cpp b/Engine/Source/Runtime/Core/Private/HAL/MallocBinned3.cpp
--- a/Engine/Source/Runtime/Core/Private/HAL/MallocBinned3.cpp
+++ b/Engine/Source/Runtime/Core/Private/HAL/MallocBinned3.cpp
@@ -402,16 +402,10 @@ struct FMallocBinned3::Private
 		{
 			uint64 PoolArraySize = NumPools * sizeof(FPoolInfoLarge);
 
-			void* Result;
-			{
-
-				LLM_PLATFORM_SCOPE(ELLMTag::FMalloc);
-				Result = Allocator.AllocateMetaDataMemory(PoolArraySize);
+			void* Result = Allocator.AllocateMetaDataMemory(PoolArraySize);
 #if BINNED3_ALLOCATOR_STATS
-				Binned3PoolInfoMemory += PoolArraySize;
+			Binned3PoolInfoMemory += PoolArraySize;
 #endif
-			}
-
 			if (!Result)
 			{
 				OutOfMemory(PoolArraySize);
@@ -935,6 +929,11 @@ void FMallocBinned3::Decommit(uint32 InPoolIndex, void *Ptr, SIZE_T Size)
 }
 
 void* FMallocBinned3::AllocateMetaDataMemory(SIZE_T Size)
+{
+	return AllocateMemoryBlock(Size).GetVirtualPointer();
+}
+
+FPlatformMemory::FPlatformVirtualMemoryBlock FMallocBinned3::AllocateMemoryBlock(SIZE_T Size)
 {
 	LLM_PLATFORM_SCOPE(ELLMTag::FMalloc);
 	size_t VirtualAlignedSize = Align(Size, FPlatformMemory::FPlatformVirtualMemoryBlock::GetVirtualSizeAlignment());
@@ -942,10 +941,14 @@ void* FMallocBinned3::AllocateMetaDataMemory(SIZE_T Size)
 	size_t CommitAlignedSize = Align(Size, FPlatformMemory::FPlatformVirtualMemoryBlock::GetCommitAlignment());
 	Block.Commit(0, CommitAlignedSize);
 	LLM(FLowLevelMemTracker::Get().OnLowLevelAlloc(ELLMTracker::Platform, Block.GetVirtualPointer(), CommitAlignedSize));
-	return Block.GetVirtualPointer();
+	return Block;
 }
 
-
+void FMallocBinned3::DeallocateMemoryBlock(FPlatformMemory::FPlatformVirtualMemoryBlock& Block)
+{
+	LLM(FLowLevelMemTracker::Get().OnLowLevelFree(ELLMTracker::Platform, Block.GetVirtualPointer()));
+	Block.FreeVirtual();
+}
 
 bool FMallocBinned3::IsInternallyThreadSafe() const
 { 
@@ -1410,11 +1413,11 @@ void FMallocBinned3::SetupTLSCachesOnCurrentThread()
 
 void FMallocBinned3::ClearAndDisableTLSCachesOnCurrentThread()
 {
-	FlushCurrentThreadCache();
 	if (!BINNED3_ALLOW_RUNTIME_TWEAKING && !GMallocBinned3PerThreadCaches)
 	{
 		return;
 	}
+	FlushCurrentThreadCache();
 	FPerThreadFreeBlockLists::ClearTLS();
 }
 
@@ -1487,8 +1490,13 @@ void FMallocBinned3::FPerThreadFreeBlockLists::SetTLS()
 	FPerThreadFreeBlockLists* ThreadSingleton = (FPerThreadFreeBlockLists*)FPlatformTLS::GetTlsValue(FMallocBinned3::Binned3TlsSlot);
 	if (!ThreadSingleton)
 	{
-		int64 TLSSize = Align(sizeof(FPerThreadFreeBlockLists), FMallocBinned3::OsAllocationGranularity);
-		ThreadSingleton = new (FMallocBinned3::AllocateMetaDataMemory(TLSSize)) FPerThreadFreeBlockLists();
+		const int VMBlockSize = Align(sizeof(FPlatformMemory::FPlatformVirtualMemoryBlock), alignof(FPerThreadFreeBlockLists));
+		const int64 TLSSize = VMBlockSize + sizeof(FPerThreadFreeBlockLists);
+		FPlatformMemory::FPlatformVirtualMemoryBlock Block = FMallocBinned3::AllocateMemoryBlock(TLSSize);
+		char* Ptr = (char*)Block.GetVirtualPointer();
+		FMemory::Memcpy(Ptr, &Block, sizeof(Block));
+		Ptr += VMBlockSize;
+		ThreadSingleton = new (Ptr) FPerThreadFreeBlockLists();
 #if BINNED3_ALLOCATOR_STATS
 		Binned3TLSMemory += TLSSize;
 #endif
@@ -1502,9 +1510,13 @@ void FMallocBinned3::FPerThreadFreeBlockLists::ClearTLS()
 {
 	check(FMallocBinned3::Binned3TlsSlot);
 	FPerThreadFreeBlockLists* ThreadSingleton = (FPerThreadFreeBlockLists*)FPlatformTLS::GetTlsValue(FMallocBinned3::Binned3TlsSlot);
-	if ( ThreadSingleton )
+	if (ThreadSingleton)
 	{
 		FMallocBinned3::Private::UnregisterThreadFreeBlockLists(ThreadSingleton);
+
+		ThreadSingleton->~FPerThreadFreeBlockLists();
+		FPlatformMemory::FPlatformVirtualMemoryBlock* Block = (FPlatformMemory::FPlatformVirtualMemoryBlock*)AlignDown(ThreadSingleton, FPlatformMemory::FPlatformVirtualMemoryBlock::GetCommitAlignment());
+		DeallocateMemoryBlock(*Block);
 	}
 	FPlatformTLS::SetTlsValue(FMallocBinned3::Binned3TlsSlot, nullptr);
 }
diff --git a/Engine/Source/Runtime/Core/Public/HAL/MallocBinned3.h b/Engine/Source/Runtime/Core/Public/HAL/MallocBinned3.h
--- a/Engine/Source/Runtime/Core/Public/HAL/MallocBinned3.h
+++ b/Engine/Source/Runtime/Core/Public/HAL/MallocBinned3.h
@@ -731,7 +731,10 @@ public:
 	void Commit(uint32 InPoolIndex, void *Ptr, SIZE_T Size);
 	void Decommit(uint32 InPoolIndex, void *Ptr, SIZE_T Size);
 
+private:
 	static void* AllocateMetaDataMemory(SIZE_T Size);
+	static FPlatformMemory::FPlatformVirtualMemoryBlock AllocateMemoryBlock(SIZE_T Size);
+	static void DeallocateMemoryBlock(FPlatformMemory::FPlatformVirtualMemoryBlock& Block);
 };
 
 PRAGMA_RESTORE_UNSAFE_TYPECAST_WARNINGS
