More function reuse in TOptional.

diff --git a/Engine/Source/Runtime/Core/Public/Misc/Optional.h b/Engine/Source/Runtime/Core/Public/Misc/Optional.h
--- a/Engine/Source/Runtime/Core/Public/Misc/Optional.h
+++ b/Engine/Source/Runtime/Core/Public/Misc/Optional.h
@@ -23,14 +23,12 @@ public:
 	
 	/** Construct an OptionalType with a valid value. */
 	TOptional(const OptionalType& InValue)
+		: TOptional(InPlace, InValue)
 	{
-		new(&Value) OptionalType(InValue);
-		bIsSet = true;
 	}
 	TOptional(OptionalType&& InValue)
+		: TOptional(InPlace, MoveTempIfPossible(InValue))
 	{
-		new(&Value) OptionalType(MoveTempIfPossible(InValue));
-		bIsSet = true;
 	}
 	template <typename... ArgTypes>
 	explicit TOptional(EInPlace, ArgTypes&&... Args)
@@ -61,7 +59,7 @@ public:
 	
 	/** Construct an OptionalType with an invalid value. */
 	TOptional(FNullOpt)
-		: bIsSet(false)
+		: TOptional()
 	{
 	}
 
@@ -127,9 +125,7 @@ public:
 	{
 		if (&InValue != (const OptionalType*)&Value)
 		{
-			Reset();
-			new(&Value) OptionalType(InValue);
-			bIsSet = true;
+			Emplace(InValue);
 		}
 		return *this;
 	}
@@ -137,9 +133,7 @@ public:
 	{
 		if (&InValue != (const OptionalType*)&Value)
 		{
-			Reset();
-			new(&Value) OptionalType(MoveTempIfPossible(InValue));
-			bIsSet = true;
+			Emplace(MoveTempIfPossible(InValue));
 		}
 		return *this;
 	}
