Moved TOptional state into its own type.

diff --git a/Engine/Source/Runtime/Core/Public/Misc/Optional.h b/Engine/Source/Runtime/Core/Public/Misc/Optional.h
--- a/Engine/Source/Runtime/Core/Public/Misc/Optional.h
+++ b/Engine/Source/Runtime/Core/Public/Misc/Optional.h
@@ -5,12 +5,21 @@
 #include "CoreTypes.h"
 #include "Misc/AssertionMacros.h"
 #include "Misc/OptionalFwd.h"
-#include "Templates/TypeCompatibleBytes.h"
 #include "Templates/UnrealTemplate.h"
 #include "Serialization/Archive.h"
 
 inline constexpr FNullOpt NullOpt{0};
 
+namespace UE::Core::Private
+{
+	template <size_t N>
+	struct TOptionalStorage
+	{
+		uint8 Storage[N];
+		bool bIsSet;
+	};
+}
+
 /**
  * When we have an optional value IsSet() returns true, and GetValue() is meaningful.
  * Otherwise GetValue() is not meaningful.
@@ -54,7 +63,7 @@ public:
 		//     TOptional<FMyType> Opt(InPlace, FMyType::FPrivateToken{}, 5, 3.14f, TEXT("Banana"));
 		//
 		new(&Value) OptionalType(Forward<ArgTypes>(Args)...);
-		bIsSet = true;
+		Value.bIsSet = true;
 	}
 	
 	/** Construct an OptionalType with an invalid value. */
@@ -65,8 +74,8 @@ public:
 
 	/** Construct an OptionalType with no value; i.e. unset */
 	TOptional()
-		: bIsSet(false)
 	{
+		Value.bIsSet = false;
 	}
 
 	~TOptional()
@@ -76,22 +85,26 @@ public:
 
 	/** Copy/Move construction */
 	TOptional(const TOptional& Other)
-		: bIsSet(false)
 	{
-		if (Other.bIsSet)
+		bool bLocalIsSet = Other.Value.bIsSet;
+		Value.bIsSet = bLocalIsSet;
+		if (!bLocalIsSet)
 		{
-			new(&Value) OptionalType(*(const OptionalType*)&Other.Value);
-			bIsSet = true;
+			return;
 		}
+
+		new(&Value) OptionalType(*(const OptionalType*)&Other.Value);
 	}
 	TOptional(TOptional&& Other)
-		: bIsSet(false)
 	{
-		if (Other.bIsSet)
+		bool bLocalIsSet = Other.Value.bIsSet;
+		Value.bIsSet = bLocalIsSet;
+		if (!bLocalIsSet)
 		{
-			new(&Value) OptionalType(MoveTempIfPossible(*(OptionalType*)&Other.Value));
-			bIsSet = true;
+			return;
 		}
+
+		new(&Value) OptionalType(MoveTempIfPossible(*(OptionalType*)&Other.Value));
 	}
 
 	TOptional& operator=(const TOptional& Other)
@@ -99,10 +112,10 @@ public:
 		if (&Other != this)
 		{
 			Reset();
-			if (Other.bIsSet)
+			if (Other.Value.bIsSet)
 			{
 				new(&Value) OptionalType(*(const OptionalType*)&Other.Value);
-				bIsSet = true;
+				Value.bIsSet = true;
 			}
 		}
 		return *this;
@@ -112,10 +125,10 @@ public:
 		if (&Other != this)
 		{
 			Reset();
-			if (Other.bIsSet)
+			if (Other.Value.bIsSet)
 			{
 				new(&Value) OptionalType(MoveTempIfPossible(*(OptionalType*)&Other.Value));
-				bIsSet = true;
+				Value.bIsSet = true;
 			}
 		}
 		return *this;
@@ -140,9 +153,9 @@ public:
 
 	void Reset()
 	{
-		if (bIsSet)
+		if (Value.bIsSet)
 		{
-			bIsSet = false;
+			Value.bIsSet = false;
 
 			// We need a typedef here because VC won't compile the destructor call below if OptionalType itself has a member called OptionalType
 			typedef OptionalType OptionalDestructOptionalType;
@@ -178,17 +191,19 @@ public:
 		//     Opt.Emplace(FMyType::FPrivateToken{}, 5, 3.14f, TEXT("Banana"));
 		//
 		OptionalType* Result = new(&Value) OptionalType(Forward<ArgsType>(Args)...);
-		bIsSet = true;
+		Value.bIsSet = true;
 		return *Result;
 	}
 
 	friend bool operator==(const TOptional& Lhs, const TOptional& Rhs)
 	{
-		if (Lhs.bIsSet != Rhs.bIsSet)
+		bool bIsLhsSet = Lhs.Value.bIsSet;
+		bool bIsRhsSet = Rhs.Value.bIsSet;
+		if (bIsLhsSet != bIsRhsSet)
 		{
 			return false;
 		}
-		if (!Lhs.bIsSet) // both unset
+		if (!bIsLhsSet) // both unset
 		{
 			return true;
 		}
@@ -202,13 +217,14 @@ public:
 
 	void Serialize(FArchive& Ar)
 	{
-		bool bOptionalIsSet = bIsSet;
-		Ar << bOptionalIsSet;
+		bool bOptionalIsSet = IsSet();
 		if (Ar.IsLoading())
 		{
-			if (bOptionalIsSet)
+			bool bOptionalWasSaved = false;
+			Ar << bOptionalWasSaved;
+			if (bOptionalWasSaved)
 			{
-				if (!bIsSet)
+				if (!bOptionalIsSet)
 				{
 					Emplace();
 				}
@@ -221,6 +237,7 @@ public:
 		}
 		else
 		{
+			Ar << bOptionalIsSet;
 			if (bOptionalIsSet)
 			{
 				Ar << GetValue();
@@ -231,7 +248,7 @@ public:
 	/** @return true when the value is meaningful; false if calling GetValue() is undefined. */
 	bool IsSet() const
 	{
-		return bIsSet;
+		return Value.bIsSet;
 	}
 	FORCEINLINE explicit operator bool() const
 	{
@@ -284,8 +301,8 @@ public:
 	}
 
 private:
-	TTypeCompatibleBytes<OptionalType> Value;
-	bool bIsSet;
+	using ValueStorageType = UE::Core::Private::TOptionalStorage<sizeof(OptionalType)>;
+	alignas(OptionalType) ValueStorageType Value;
 };
 
 template<typename OptionalType>
