Fix for AutoRTFM cloning functions it shouldn't when using unity builds.

diff --git a/Engine/Source/Programs/AutoRTFMTests/Private/Open.cpp b/Engine/Source/Programs/AutoRTFMTests/Private/Open.cpp
--- a/Engine/Source/Programs/AutoRTFMTests/Private/Open.cpp
+++ b/Engine/Source/Programs/AutoRTFMTests/Private/Open.cpp
@@ -4,6 +4,7 @@
 
 #include "Catch2Includes.h"
 #include <AutoRTFM/AutoRTFM.h>
+#include <atomic>
 #include <map>
 #include <vector>
 
@@ -90,3 +91,16 @@ TEST_CASE("Open.Large")
     REQUIRE(m[12][0] == 13);
     REQUIRE(m[12][1] == 14);
 }
+
+TEST_CASE("Open.Atomics")
+{
+    std::atomic<bool> bDidRun = false;
+    REQUIRE(
+        AutoRTFM::ETransactionResult::Committed ==
+        AutoRTFM::Transact([&] ()
+        {
+            AutoRTFM::Open([&] () { bDidRun = true; });
+            AutoRTFM::Open([&] () { REQUIRE(bDidRun); });
+        }));
+    REQUIRE(bDidRun);
+}
diff --git a/Engine/Source/Runtime/Core/Private/AutoRTFM/API.cpp b/Engine/Source/Runtime/Core/Private/AutoRTFM/API.cpp
--- a/Engine/Source/Runtime/Core/Private/AutoRTFM/API.cpp
+++ b/Engine/Source/Runtime/Core/Private/AutoRTFM/API.cpp
@@ -20,84 +20,84 @@
 namespace AutoRTFM
 {
 
-extern "C" bool autortfm_is_transactional()
+FORCENOINLINE extern "C" bool autortfm_is_transactional()
 {
     return FContext::Get()->IsTransactional();
 }
 
-extern "C" bool autortfm_is_closed()
+FORCENOINLINE extern "C" bool autortfm_is_closed()
 {
     return false;
 }
 
 // First Part - the API exposed outside transactions.
-extern "C" autortfm_result autortfm_transact(void (*Work)(void* Arg), void* Arg)
+FORCENOINLINE extern "C" autortfm_result autortfm_transact(void (*Work)(void* Arg), void* Arg)
 {
     return static_cast<autortfm_result>(FContext::Get()->Transact(Work, Arg));
 }
 
-extern "C" void autortfm_commit(void (*Work)(void* Arg), void* Arg)
+FORCENOINLINE extern "C" void autortfm_commit(void (*Work)(void* Arg), void* Arg)
 {
     autortfm_result Result = autortfm_transact(Work, Arg);
 	UE_CLOG(Result != autortfm_committed, LogAutoRTFM, Fatal, TEXT("Unexpected transaction result: %u."), Result);
 }
 
-extern "C" void autortfm_abort()
+FORCENOINLINE extern "C" void autortfm_abort()
 {
 	UE_CLOG(!FContext::IsTransactional(), LogAutoRTFM, Fatal, TEXT("The function `autortfm_abort` was called from outside a transaction."));
 	FContext::Get()->AbortByRequestAndThrow();
 }
 
-extern "C" bool autortfm_start_transaction()
+FORCENOINLINE extern "C" bool autortfm_start_transaction()
 {
 	UE_CLOG(!FContext::IsTransactional(), LogAutoRTFM, Fatal, TEXT("The function `autortfm_start_transaction` was called from outside a transact."));
 	return FContext::Get()->StartTransaction();
 }
 
-extern "C" autortfm_result autortfm_commit_transaction()
+FORCENOINLINE extern "C" autortfm_result autortfm_commit_transaction()
 {
 	UE_CLOG(!FContext::IsTransactional(), LogAutoRTFM, Fatal, TEXT("The function `autortfm_commit_transaction` was called from outside a transact."));
 	return static_cast<autortfm_result>(FContext::Get()->CommitTransaction());
 }
 
-extern "C" autortfm_result autortfm_abort_transaction()
+FORCENOINLINE extern "C" autortfm_result autortfm_abort_transaction()
 {
 	UE_CLOG(!FContext::IsTransactional(), LogAutoRTFM, Fatal, TEXT("The function `autortfm_abort_transaction` was called from outside a transact."));
 	return static_cast<autortfm_result>(FContext::Get()->AbortTransaction(false));
 }
 
-extern "C" void autortfm_clear_transaction_status()
+FORCENOINLINE extern "C" void autortfm_clear_transaction_status()
 {
 	ASSERT(FContext::Get()->IsAborting());
 	FContext::Get()->ClearTransactionStatus();
 }
 
-extern "C" bool autortfm_is_aborting()
+FORCENOINLINE extern "C" bool autortfm_is_aborting()
 {
 	return FContext::Get()->IsAborting();
 }
 
-extern "C" bool autortfm_current_nest_throw()
+FORCENOINLINE extern "C" bool autortfm_current_nest_throw()
 {
 	FContext::Get()->Throw();
 	return true;
 }
 
-extern "C" void autortfm_abort_if_transactional()
+FORCENOINLINE extern "C" void autortfm_abort_if_transactional()
 {
 	UE_CLOG(FContext::IsTransactional(), LogAutoRTFM, Fatal, TEXT("The function `autortfm_abort_if_transactional` was called from an open inside a transaction."));
 }
 
-extern "C" void autortfm_abort_if_closed()
+FORCENOINLINE extern "C" void autortfm_abort_if_closed()
 {
 }
 
-extern "C" void autortfm_open(void (*Work)(void* Arg), void* Arg)
+FORCENOINLINE extern "C" void autortfm_open(void (*Work)(void* Arg), void* Arg)
 {
 	Work(Arg);
 }
 
-extern "C" autortfm_status autortfm_close(void (*Work)(void* Arg), void* Arg)
+FORCENOINLINE extern "C" autortfm_status autortfm_close(void (*Work)(void* Arg), void* Arg)
 {
 	UE_CLOG(!FContext::IsTransactional(), LogAutoRTFM, Fatal, TEXT("Close called from an outside a transaction."));
 
@@ -113,13 +113,13 @@ extern "C" autortfm_status autortfm_close(void (*Work)(void* Arg), void* Arg)
 	return Result;
 }
 
-extern "C" void autortfm_record_open_write(void* Ptr, size_t Size)
+FORCENOINLINE extern "C" void autortfm_record_open_write(void* Ptr, size_t Size)
 {
     FContext::Get()->CheckOpenRecordWrite(Ptr);
 	FContext::Get()->RecordWrite(Ptr, Size);
 }
 
-extern "C" void autortfm_register_open_function(void* OriginalFunction, void* NewFunction)
+FORCENOINLINE extern "C" void autortfm_register_open_function(void* OriginalFunction, void* NewFunction)
 {
 	UE_LOG(LogAutoRTFM, Verbose, TEXT("Registering open %p->%p"), OriginalFunction, NewFunction);
     FunctionMapAdd(OriginalFunction, NewFunction);
@@ -134,21 +134,21 @@ void OpenAbort(TFunction<void()>&& Work)
 {
 }
 
-extern "C" void autortfm_open_commit(void (*Work)(void* Arg), void* Arg)
+FORCENOINLINE extern "C" void autortfm_open_commit(void (*Work)(void* Arg), void* Arg)
 {
     Work(Arg);
 }
 
-extern "C" void autortfm_open_abort(void (*Work)(void* arg), void* Arg)
+FORCENOINLINE extern "C" void autortfm_open_abort(void (*Work)(void* arg), void* Arg)
 {
 }
 
-extern "C" void* autortfm_did_allocate(void* Ptr, size_t Size)
+FORCENOINLINE extern "C" void* autortfm_did_allocate(void* Ptr, size_t Size)
 {
     return Ptr;
 }
 
-extern "C" void autortfm_check_consistency_assuming_no_races()
+FORCENOINLINE extern "C" void autortfm_check_consistency_assuming_no_races()
 {
     if (FContext::IsTransactional())
     {
@@ -156,7 +156,7 @@ extern "C" void autortfm_check_consistency_assuming_no_races()
     }
 }
 
-extern "C" void autortfm_check_abi(void* const Ptr, const size_t Size)
+FORCENOINLINE extern "C" void autortfm_check_abi(void* const Ptr, const size_t Size)
 {
     struct FConstants final
     {
