changed delegates race detection to `ensure` instead of `check` to avoid interrupting unrelated devs until we weed out existing races

diff --git a/Engine/Source/Runtime/Core/Public/Misc/MTAccessDetector.h b/Engine/Source/Runtime/Core/Public/Misc/MTAccessDetector.h
--- a/Engine/Source/Runtime/Core/Public/Misc/MTAccessDetector.h
+++ b/Engine/Source/Runtime/Core/Public/Misc/MTAccessDetector.h
@@ -347,7 +347,7 @@ private:
 		if (InState.WriterNum != 0)
 		{
 			uint32 CurrentThreadId = FPlatformTLS::GetCurrentThreadId();
-			checkf(InState.GetWriterThreadId() == CurrentThreadId,
+			ensureMsgf(InState.GetWriterThreadId() == CurrentThreadId,
 				TEXT("Data race detected! Acquiring read access on thread %u concurrently with %d writers on thread %u:\nCurrent thread %u callstack:\n%s\nWriter thread %u callstack:\n%s"),
 				CurrentThreadId, InState.WriterNum, InState.GetWriterThreadId(), 
 				CurrentThreadId, *GetThreadCallstack(CurrentThreadId),
@@ -467,7 +467,7 @@ public:
 			}
 		}
 
-		checkf(LoadState().Value == ExpectedState.Value,
+		ensureMsgf(LoadState().Value == ExpectedState.Value,
 			TEXT("Race detector destroyed while being accessed on another thread: %d readers, %d writers on thread %u:\nCurrent thread %u callstack:\n%s\nWriter thread %u callstack:\n%s"),
 			LoadState().ReaderNum - ExpectedState.ReaderNum, 
 			LoadState().WriterNum - ExpectedState.WriterNum, LoadState().GetWriterThreadId(), 
@@ -532,8 +532,12 @@ public:
 		if (LocalState.ReaderNum >= 1)
 		{	// check that all readers are on the current thread
 			int32 ReaderIndex = GetReadersTls().IndexOfByPredicate([this](FReaderNum ReaderNum) { return ReaderNum.Reader == this; });
-			checkf(ReaderIndex != INDEX_NONE, TEXT("Race detector is not trivially copyable while this delegate is copied trivially. Consider changing this delegate to use `FNotThreadSafeNotCheckedDelegateUserPolicy`"));
-			checkf(GetReadersTls()[ReaderIndex].Num == LocalState.ReaderNum,
+			ensureMsgf(ReaderIndex != INDEX_NONE, TEXT("Race detector is not trivially copyable while this delegate is copied trivially. Consider changing this delegate to use `FNotThreadSafeNotCheckedDelegateUserPolicy`"));
+			if (ReaderIndex == INDEX_NONE)
+			{
+				return; // to avoid asserting on the next line inside TArray
+			}
+			ensureMsgf(GetReadersTls()[ReaderIndex].Num == LocalState.ReaderNum,
 				TEXT("Data race detected: %d reader(s) on another thread(s) while acquiring write access.\nCurrent thread %u callstack:\n%s"),
 				LocalState.ReaderNum - GetReadersTls()[ReaderIndex].Num,
 				FPlatformTLS::GetCurrentThreadId(), *GetThreadCallstack(FPlatformTLS::GetCurrentThreadId()));
@@ -542,7 +546,7 @@ public:
 		uint32 CurrentThreadId = FPlatformTLS::GetCurrentThreadId();
 		if (LocalState.WriterNum != 0)
 		{
-			checkf(LocalState.GetWriterThreadId() == CurrentThreadId,
+			ensureMsgf(LocalState.GetWriterThreadId() == CurrentThreadId,
 				TEXT("Data race detected: acquiring write access on thread %u concurrently with %d writers on thread %u:\nCurrent thread %u callstack:\n%s\nWriter thread %u callstack:\n%s"), 
 				CurrentThreadId, LocalState.WriterNum, LocalState.GetWriterThreadId(), 
 				CurrentThreadId, *GetThreadCallstack(CurrentThreadId),
@@ -552,7 +556,7 @@ public:
 		FState NewState{ LocalState.ReaderNum, LocalState.WriterNum + 1u, CurrentThreadId + 1 };
 		FState PrevState = ExchangeState(NewState);
 
-		checkf(LocalState == PrevState,
+		ensureMsgf(LocalState == PrevState,
 			TEXT("Data race detected: other thread(s) activity during acquiring write access on thread %u: %u -> %u readers, %u -> %u writers on thread %u -> %u:\nCurrent thread %u callstack:\n%s\nWriter thread %u callstack:\n%s"),
 			CurrentThreadId,
 			LocalState.ReaderNum, PrevState.ReaderNum,
@@ -580,7 +584,7 @@ public:
 		FState NewState{ LocalState.ReaderNum, LocalState.WriterNum - 1u, WriterThreadId };
 		FState PrevState = ExchangeState(NewState);
 
-		checkf(LocalState == PrevState,
+		ensureMsgf(LocalState == PrevState,
 			TEXT("Data race detected: other thread(s) activity during releasing write access on thread %d: %u -> %u readers, %u -> %u writers on thread %u -> %u\nCurrent thread %u callstack:\n%s\nWriter thread %u callstack:\n%s"),
 			LocalState.ReaderNum, PrevState.ReaderNum,
 			LocalState.WriterNum, PrevState.WriterNum,
