Trace counter Get() on atomic ints would return reference to local due to implicit cast.

When fixing the Get() I also discovered a couple of counters where the storage template type was incorrectly the same as the value template type. Fixed to use the atomic counter type given in counters header.

#rnx

diff --git a/Engine/Source/Runtime/Core/Private/ProfilingDebugging/StallDetector.cpp b/Engine/Source/Runtime/Core/Private/ProfilingDebugging/StallDetector.cpp
--- a/Engine/Source/Runtime/Core/Private/ProfilingDebugging/StallDetector.cpp
+++ b/Engine/Source/Runtime/Core/Private/ProfilingDebugging/StallDetector.cpp
@@ -229,8 +229,8 @@ uint32 UE::FStallDetectorRunnable::Run()
 
 static FCriticalSection StallDetectorStatsCritical;
 static TSet<UE::FStallDetectorStats*> StallDetectorStats;
-FCountersTrace::TCounter<std::atomic<int64>, TraceCounterType_Int> UE::FStallDetectorStats::TotalTriggeredCount (TEXT("StallDetector/TotalTriggeredCount"), TraceCounterDisplayHint_None);
-FCountersTrace::TCounter<std::atomic<int64>, TraceCounterType_Int> UE::FStallDetectorStats::TotalReportedCount (TEXT("StallDetector/TotalReportedCount"), TraceCounterDisplayHint_None);
+FCountersTrace::FCounterAtomicInt UE::FStallDetectorStats::TotalTriggeredCount (TEXT("StallDetector/TotalTriggeredCount"), TraceCounterDisplayHint_None);
+FCountersTrace::FCounterAtomicInt UE::FStallDetectorStats::TotalReportedCount (TEXT("StallDetector/TotalReportedCount"), TraceCounterDisplayHint_None);
 
 UE::FStallDetectorStats::FStallDetectorStats(const TCHAR* InName, const double InBudgetSeconds, const EStallDetectorReportingMode InReportingMode)
 	: Name(InName)
diff --git a/Engine/Source/Runtime/Core/Public/ProfilingDebugging/CountersTrace.h b/Engine/Source/Runtime/Core/Public/ProfilingDebugging/CountersTrace.h
--- a/Engine/Source/Runtime/Core/Public/ProfilingDebugging/CountersTrace.h
+++ b/Engine/Source/Runtime/Core/Public/ProfilingDebugging/CountersTrace.h
@@ -102,7 +102,7 @@ struct FCountersTrace
 			}
 		}
 
-		const ValueType& Get() const
+		ValueType Get() const
 		{
 			return Value;
 		}
diff --git a/Engine/Source/Runtime/Core/Public/ProfilingDebugging/StallDetector.h b/Engine/Source/Runtime/Core/Public/ProfilingDebugging/StallDetector.h
--- a/Engine/Source/Runtime/Core/Public/ProfilingDebugging/StallDetector.h
+++ b/Engine/Source/Runtime/Core/Public/ProfilingDebugging/StallDetector.h
@@ -180,10 +180,10 @@ namespace UE
 		bool bReported;
 
 		// The total number of times all callsites have been triggered
-		static FCountersTrace::TCounter<std::atomic<int64>, TraceCounterType_Int> TotalTriggeredCount;
+		static FCountersTrace::FCounterAtomicInt TotalTriggeredCount;
 
 		// The total number of reports that have been sent
-		static FCountersTrace::TCounter<std::atomic<int64>, TraceCounterType_Int> TotalReportedCount;
+		static FCountersTrace::FCounterAtomicInt TotalReportedCount;
 
 	private:
 		// The number of times this callsite has been triggered
