Reenabling CLI filtering using grammar based parse.
To be submitted after CL 25765744

diff --git a/Engine/Source/Runtime/Core/Private/Misc/CommandLine.cpp b/Engine/Source/Runtime/Core/Private/Misc/CommandLine.cpp
--- a/Engine/Source/Runtime/Core/Private/Misc/CommandLine.cpp
+++ b/Engine/Source/Runtime/Core/Private/Misc/CommandLine.cpp
@@ -108,7 +108,7 @@ bool FCommandLine::IsCommandLineLoggingFiltered()
 #endif
 }
 
-bool FCommandLine::FilterMove(TCHAR* OutLine, int32 MaxLen, const TCHAR* InLine, const TArrayView<FString>& AllowedList)
+bool FCommandLine::FilterCLIUsingGrammarBasedParser(TCHAR* OutLine, int32 MaxLen, const TCHAR* InLine, const TArrayView<FString>& AllowedList)
 {
 	if (MaxLen == 0)
 	{
@@ -225,17 +225,13 @@ void FCommandLine::ApplyCommandLineAllowList()
 		FCommandLine::Parse(FilterForLoggingList, FilterArgsForLogging, Ignored);
 	}
 	// Process the original command line
-	TArray<FString> OriginalList = FilterCommandLine(OriginalCmdLine);
-	BuildCommandLineAllowList(OriginalCmdLine, UE_ARRAY_COUNT(OriginalCmdLine), OriginalList);
+	FilterCLIInplace(OriginalCmdLine, UE_ARRAY_COUNT(OriginalCmdLine), OriginalCmdLine, OriginalList);
 	// Process the current command line
-	TArray<FString> CmdList = FilterCommandLine(CmdLine);
-	BuildCommandLineAllowList(CmdLine, UE_ARRAY_COUNT(CmdLine), CmdList);
+	FilterCLIInplace(CmdLine, UE_ARRAY_COUNT(CmdLine), CmdLine, CmdList);
 	// Process the command line for logging purposes
-	TArray<FString> LoggingCmdList = FilterCommandLineForLogging(LoggingCmdLine);
-	BuildCommandLineAllowList(LoggingCmdLine, UE_ARRAY_COUNT(LoggingCmdLine), LoggingCmdList);
+	FilterCLIInplace(LoggingCmdLine, UE_ARRAY_COUNT(LoggingCmdLine), LoggingCmdLine, LoggingCmdList);
 	// Process the original command line for logging purposes
-	TArray<FString> LoggingOriginalCmdList = FilterCommandLineForLogging(LoggingOriginalCmdLine);
-	BuildCommandLineAllowList(LoggingOriginalCmdLine, UE_ARRAY_COUNT(LoggingOriginalCmdLine), LoggingOriginalCmdList);
+	FilterCLIInplace(LoggingOriginalCmdLine, UE_ARRAY_COUNT(LoggingOriginalCmdLine), LoggingOriginalCmdLine, LoggingOriginalCmdList);
 }
 
 TArray<FString> FCommandLine::FilterCommandLine(TCHAR* CommandLine)
diff --git a/Engine/Source/Runtime/Core/Public/Misc/CommandLine.h b/Engine/Source/Runtime/Core/Public/Misc/CommandLine.h
--- a/Engine/Source/Runtime/Core/Public/Misc/CommandLine.h
+++ b/Engine/Source/Runtime/Core/Public/Misc/CommandLine.h
@@ -110,7 +110,8 @@ struct FCommandLine
 	CORE_API static FString BuildFromArgV(const ANSICHAR* Prefix, int32 ArgC, ANSICHAR* ArgV[], const ANSICHAR* Suffix);
 
 	/**
-	* Filter Move parses CLI style arguments in the InLine, filters for commands or keys specified in the AllowedList
+	* FilterCLIUsingGrammarBasedParser parses CLI style arguments.
+	* filters for commands or keys specified in the AllowedList
 	* and writes the to the OutLine.
 	* OutLine and InLine may be point to the same buffer.
 	* 
@@ -121,7 +122,7 @@ struct FCommandLine
 	* 
 	* Returns true if Outline was large enough to hold the filtered string, Otherwise returns false.
 	*/
-	static bool FilterMove(TCHAR* OutLine, int32 MaxLen, const TCHAR* InLine, const TArrayView<FString>& AllowedList);
+	static bool FilterCLIUsingGrammarBasedParser(TCHAR* OutLine, int32 MaxLen, const TCHAR* InLine, const TArrayView<FString>& AllowedList);
 private:
 #if UE_COMMAND_LINE_USES_ALLOW_LIST
 	/** Filters both the original and current command line list for approved only args */
diff --git a/Engine/Source/Runtime/Core/Tests/Misc/CommandLineTests.cpp b/Engine/Source/Runtime/Core/Tests/Misc/CommandLineTests.cpp
--- a/Engine/Source/Runtime/Core/Tests/Misc/CommandLineTests.cpp
+++ b/Engine/Source/Runtime/Core/Tests/Misc/CommandLineTests.cpp
@@ -10,7 +10,7 @@
 
 #include <catch2/generators/catch_generators.hpp>
 
-TEST_CASE("CommandLine::Filtering::FilterMove", "[Smoke]")
+TEST_CASE("CommandLine::Filtering::FilterCLIInplace", "[Smoke]")
 {
 	TArray<FString> AllowedList{ { "cmd_a", "-cmd_b" } };
 
@@ -31,7 +31,7 @@ TEST_CASE("CommandLine::Filtering::FilterMove", "[Smoke]")
 		}));
 		TCHAR Result[256]{};
 
-		CHECK(FCommandLine::FilterMove(Result, UE_ARRAY_COUNT(Result), Input, AllowedList));
+		CHECK(FCommandLine::FilterCLIUsingGrammarBasedParser(Result, UE_ARRAY_COUNT(Result), Input, AllowedList));
 		CHECK(FStringView{ Result } == FStringView{ Expected });
 	}
 
@@ -45,7 +45,7 @@ TEST_CASE("CommandLine::Filtering::FilterMove", "[Smoke]")
 			}));
 		TCHAR Result[256]{};
 
-		CHECK(FCommandLine::FilterMove(Result, UE_ARRAY_COUNT(Result), Input, AllowedList));
+		CHECK(FCommandLine::FilterCLIUsingGrammarBasedParser(Result, UE_ARRAY_COUNT(Result), Input, AllowedList));
 		CHECK(FStringView{ Result } == FStringView{ Expected });
 	}
 
@@ -71,7 +71,7 @@ TEST_CASE("CommandLine::Filtering::FilterMove", "[Smoke]")
 		TCHAR SourceAndResult[256]{};
 		
 		FCString::Strcpy(SourceAndResult, Input);
-		CHECK(FCommandLine::FilterMove(SourceAndResult, UE_ARRAY_COUNT(SourceAndResult), SourceAndResult, AllowedList));
+		CHECK(FCommandLine::FilterCLIUsingGrammarBasedParser(SourceAndResult, UE_ARRAY_COUNT(SourceAndResult), SourceAndResult, AllowedList));
 		CHECK(FStringView{ SourceAndResult } == FStringView{ Expected });
 	}
 
@@ -80,7 +80,7 @@ TEST_CASE("CommandLine::Filtering::FilterMove", "[Smoke]")
 		const TCHAR* Input = TEXT("-cmd_a --cmd_b");
 		TCHAR Result[256] = TEXT("Not Empty");
 
-		CHECK(FCommandLine::FilterMove(Result, UE_ARRAY_COUNT(Result), Input, {}));
+		CHECK(FCommandLine::FilterCLIUsingGrammarBasedParser(Result, UE_ARRAY_COUNT(Result), Input, {}));
 		CHECK(Result[0] == TCHAR(0));
 	}
 
@@ -89,7 +89,7 @@ TEST_CASE("CommandLine::Filtering::FilterMove", "[Smoke]")
 		const TCHAR* Input = TEXT("-cmd_a --cmd_b");
 		TCHAR Result[5]{};
 
-		CHECK(false == FCommandLine::FilterMove(Result, UE_ARRAY_COUNT(Result), Input, AllowedList));
+		CHECK(false == FCommandLine::FilterCLIUsingGrammarBasedParser(Result, UE_ARRAY_COUNT(Result), Input, AllowedList));
 		CHECK(Result[0] == TCHAR(0));
 	}
 
@@ -126,7 +126,7 @@ TEST_CASE("CommandLine::Filtering::FilterMove", "[Smoke]")
 		TCHAR SourceAndResult[256]{};
 
 		FCString::Strcpy(SourceAndResult, Input);
-		CHECK(FCommandLine::FilterMove(SourceAndResult, UE_ARRAY_COUNT(SourceAndResult), SourceAndResult, ApprovedArgs));
+		CHECK(FCommandLine::FilterCLIUsingGrammarBasedParser(SourceAndResult, UE_ARRAY_COUNT(SourceAndResult), SourceAndResult, ApprovedArgs));
 		CHECK(FStringView{ SourceAndResult } == FStringView{ Expected });
 	}
 
