Added support for intrusive TOptional null states.
Added intrusive TOptional null states to TFunctionRef and TSharedRef as examples.

diff --git a/Engine/Source/Runtime/Core/Public/Misc/IntrusiveUnsetOptionalState.h b/Engine/Source/Runtime/Core/Public/Misc/IntrusiveUnsetOptionalState.h
new file mode 100644
--- /dev/null
+++ b/Engine/Source/Runtime/Core/Public/Misc/IntrusiveUnsetOptionalState.h
@@ -0,0 +1,90 @@
+// Copyright Epic Games, Inc. All Rights Reserved.
+
+#pragma once
+
+/**
+ * Facilities to allow types to have an intrusive invalid state which can act as
+ * TOptional's 'unset' state, saving space.  A class in such a state will only ever be
+ * compared against FIntrusiveUnsetOptionalState or destructed.
+ *
+ * A class should implement a constructor taking FIntrusiveUnsetOptionalState, an
+ * assignment operator from FIntrusiveUnsetOptionalState, and an equality comparison
+ * operator against FIntrusiveUnsetOptionalState, which will put a class instance into
+ * the 'unset' state (in the case of the constructor and assignment) and allow testing
+ * of its unset state.
+ *
+ * A public constexpr static data member of type bool called bHasIntrusiveUnsetOptionalState
+ * should be defined and set to true.  There must also be a public typedef/alias called
+ * IntrusiveUnsetOptionalStateType and set to the type itself.  This enables the optimization
+ * of this type within TOptional.
+ *
+ * These functions should be public - regular user code will not be able to call them as
+ * they will not be able to construct an FIntrusiveUnsetOptionalState object to pass.
+ *
+ * Example:
+ *
+ * struct FMyType
+ * {
+ *     // This static member should be constexpr, public and equal to true.
+ *     static constexpr bool bHasIntrusiveUnsetOptionalState = true;
+ *
+ *     explicit FMyType(int32 InIndex)
+ *     {
+ *         // Validate class invariant.
+ *         check(InIndex >= 0);
+ *
+ *         Index = InIndex;
+ *     }
+ *
+ *     // This constructor will only ever be called by TOptional<FMyType> to enter its
+ *     // 'unset' state.
+ *     explicit FMyType(FIntrusiveUnsetOptionalState)
+ *     {
+ *         // Since negative indices are illegal as per the class invariant, we can use -1
+ *         // here as TOptional's 'unset' state, which no legal class instance will have.
+ *         Index = -1;
+ *     }
+ *
+ *     // Similarly with this assignment operator.  It need not return a reference to itself
+ *     // as most assignment operators do.
+ *     void operator=(FIntrusiveUnsetOptionalState)
+ *     {
+ *         Index = -1;
+ *     }
+ *
+ *     // This comparison function will only ever be called by TOptional to check if the
+ *     // object is in the 'unset' state.  It does not need to be commutative like most
+ *     // comparison operators, nor is an operator!= necessary.
+ *     bool operator==(FIntrusiveUnsetOptionalState) const
+ *     {
+ *         return Index == -1;
+ *     }
+ *
+ * private:
+ *     // Non-negative indices are part of the class invariant.
+ *     int32 Index;
+ * };
+ */
+
+#include "Misc/OptionalFwd.h"
+
+namespace UE::Core::Private
+{
+	template <typename OptionalType, bool HasIntrusiveUnsetOptionalState>
+	struct TOptionalStorage;
+}
+
+struct FIntrusiveUnsetOptionalState
+{
+	// Defined in the global namespace for convenience and consistency
+	// with TOptional which is also global.
+
+	template <typename>
+	friend struct TOptional;
+
+	template <typename, bool>
+	friend struct UE::Core::Private::TOptionalStorage;
+
+private:
+	explicit FIntrusiveUnsetOptionalState() = default;
+};
diff --git a/Engine/Source/Runtime/Core/Public/Misc/Optional.h b/Engine/Source/Runtime/Core/Public/Misc/Optional.h
--- a/Engine/Source/Runtime/Core/Public/Misc/Optional.h
+++ b/Engine/Source/Runtime/Core/Public/Misc/Optional.h
@@ -4,6 +4,7 @@
 
 #include "CoreTypes.h"
 #include "Misc/AssertionMacros.h"
+#include "Misc/IntrusiveUnsetOptionalState.h"
 #include "Misc/OptionalFwd.h"
 #include "Templates/UnrealTemplate.h"
 #include "Serialization/Archive.h"
@@ -12,14 +13,40 @@ inline constexpr FNullOpt NullOpt{0};
 
 namespace UE::Core::Private
 {
+	struct CIntrusiveUnsettable
+	{
+		template <typename T>
+		auto Requires(bool& b) -> decltype(
+			b = std::decay_t<typename T::IntrusiveUnsetOptionalStateType>::bHasIntrusiveUnsetOptionalState
+		);
+	};
+
 	template <size_t N>
-	struct TOptionalStorage
+	struct TNonIntrusiveOptionalStorage
 	{
 		uint8 Storage[N];
 		bool bIsSet;
 	};
 }
 
+template <typename T>
+constexpr bool HasIntrusiveUnsetOptionalState()
+{
+	if constexpr (!TModels<UE::Core::Private::CIntrusiveUnsettable, T>::Value)
+	{
+		return false;
+	}
+	// Derived types are not guaranteed to have an intrusive state, so ensure IntrusiveUnsetOptionalStateType matches the type in the optional
+	else if constexpr (!std::is_same_v<const typename T::IntrusiveUnsetOptionalStateType, const T>)
+	{
+		return false;
+	}
+	else
+	{
+		return T::bHasIntrusiveUnsetOptionalState;
+	}
+}
+
 /**
  * When we have an optional value IsSet() returns true, and GetValue() is meaningful.
  * Otherwise GetValue() is not meaningful.
@@ -27,6 +54,9 @@ namespace UE::Core::Private
 template<typename OptionalType>
 struct TOptional
 {
+private:
+	static constexpr bool bUsingIntrusiveUnsetState = HasIntrusiveUnsetOptionalState<OptionalType>();
+
 public:
 	using ElementType = OptionalType;
 	
@@ -63,7 +93,11 @@ public:
 		//     TOptional<FMyType> Opt(InPlace, FMyType::FPrivateToken{}, 5, 3.14f, TEXT("Banana"));
 		//
 		new(&Value) OptionalType(Forward<ArgTypes>(Args)...);
-		Value.bIsSet = true;
+
+		if constexpr (!bUsingIntrusiveUnsetState)
+		{
+			Value.bIsSet = true;
+		}
 	}
 	
 	/** Construct an OptionalType with an invalid value. */
@@ -75,7 +109,14 @@ public:
 	/** Construct an OptionalType with no value; i.e. unset */
 	TOptional()
 	{
-		Value.bIsSet = false;
+		if constexpr (bUsingIntrusiveUnsetState)
+		{
+			new (&Value) OptionalType(FIntrusiveUnsetOptionalState{});
+		}
+		else
+		{
+			Value.bIsSet = false;
+		}
 	}
 
 	~TOptional()
@@ -86,22 +127,28 @@ public:
 	/** Copy/Move construction */
 	TOptional(const TOptional& Other)
 	{
-		bool bLocalIsSet = Other.Value.bIsSet;
-		Value.bIsSet = bLocalIsSet;
-		if (!bLocalIsSet)
+		if constexpr (!bUsingIntrusiveUnsetState)
 		{
-			return;
+			bool bLocalIsSet = Other.Value.bIsSet;
+			Value.bIsSet = bLocalIsSet;
+			if (!bLocalIsSet)
+			{
+				return;
+			}
 		}
 
 		new(&Value) OptionalType(*(const OptionalType*)&Other.Value);
 	}
 	TOptional(TOptional&& Other)
 	{
-		bool bLocalIsSet = Other.Value.bIsSet;
-		Value.bIsSet = bLocalIsSet;
-		if (!bLocalIsSet)
+		if constexpr (!bUsingIntrusiveUnsetState)
 		{
-			return;
+			bool bLocalIsSet = Other.Value.bIsSet;
+			Value.bIsSet = bLocalIsSet;
+			if (!bLocalIsSet)
+			{
+				return;
+			}
 		}
 
 		new(&Value) OptionalType(MoveTempIfPossible(*(OptionalType*)&Other.Value));
@@ -111,11 +158,18 @@ public:
 	{
 		if (&Other != this)
 		{
-			Reset();
-			if (Other.Value.bIsSet)
+			if constexpr (bUsingIntrusiveUnsetState)
 			{
-				new(&Value) OptionalType(*(const OptionalType*)&Other.Value);
-				Value.bIsSet = true;
+				*(OptionalType*)&Value = *(const OptionalType*)&Other.Value;
+			}
+			else
+			{
+				Reset();
+				if (Other.Value.bIsSet)
+				{
+					new(&Value) OptionalType(*(const OptionalType*)&Other.Value);
+					Value.bIsSet = true;
+				}
 			}
 		}
 		return *this;
@@ -124,11 +178,18 @@ public:
 	{
 		if (&Other != this)
 		{
-			Reset();
-			if (Other.Value.bIsSet)
+			if constexpr (bUsingIntrusiveUnsetState)
 			{
-				new(&Value) OptionalType(MoveTempIfPossible(*(OptionalType*)&Other.Value));
-				Value.bIsSet = true;
+				*(OptionalType*)&Value = MoveTempIfPossible(*(OptionalType*)&Other.Value);
+			}
+			else
+			{
+				Reset();
+				if (Other.Value.bIsSet)
+				{
+					new(&Value) OptionalType(MoveTempIfPossible(*(OptionalType*)&Other.Value));
+					Value.bIsSet = true;
+				}
 			}
 		}
 		return *this;
@@ -153,20 +214,38 @@ public:
 
 	void Reset()
 	{
-		if (Value.bIsSet)
+		if constexpr (bUsingIntrusiveUnsetState)
 		{
-			Value.bIsSet = false;
+			*(OptionalType*)&Value = FIntrusiveUnsetOptionalState{};
+		}
+		else
+		{
+			if (Value.bIsSet)
+			{
+				Value.bIsSet = false;
 
-			// We need a typedef here because VC won't compile the destructor call below if OptionalType itself has a member called OptionalType
-			typedef OptionalType OptionalDestructOptionalType;
-			((OptionalType*)&Value)->OptionalDestructOptionalType::~OptionalDestructOptionalType();
+				// We need a typedef here because VC won't compile the destructor call below if OptionalType itself has a member called OptionalType
+				typedef OptionalType OptionalDestructOptionalType;
+				((OptionalType*)&Value)->OptionalDestructOptionalType::~OptionalDestructOptionalType();
+			}
 		}
 	}
 
 	template <typename... ArgsType>
 	OptionalType& Emplace(ArgsType&&... Args)
 	{
-		Reset();
+		if constexpr (bUsingIntrusiveUnsetState)
+		{
+			// Destroy the member in-place before replacing it - a bit nasty, but it'll work since we don't support exceptions
+
+			// We need a typedef here because VC won't compile the destructor call below if OptionalType itself has a member called OptionalType
+			typedef OptionalType OptionalDestructOptionalType;
+			((OptionalType*)&Value)->OptionalDestructOptionalType::~OptionalDestructOptionalType();
+		}
+		else
+		{
+			Reset();
+		}
 
 		// If this fails to compile when trying to call Emplace with a non-public constructor,
 		// do not make TOptional a friend.
@@ -191,22 +270,31 @@ public:
 		//     Opt.Emplace(FMyType::FPrivateToken{}, 5, 3.14f, TEXT("Banana"));
 		//
 		OptionalType* Result = new(&Value) OptionalType(Forward<ArgsType>(Args)...);
-		Value.bIsSet = true;
+
+		if constexpr (!bUsingIntrusiveUnsetState)
+		{
+			Value.bIsSet = true;
+		}
+
 		return *Result;
 	}
 
 	friend bool operator==(const TOptional& Lhs, const TOptional& Rhs)
 	{
-		bool bIsLhsSet = Lhs.Value.bIsSet;
-		bool bIsRhsSet = Rhs.Value.bIsSet;
-		if (bIsLhsSet != bIsRhsSet)
+		if constexpr (!bUsingIntrusiveUnsetState)
 		{
-			return false;
-		}
-		if (!bIsLhsSet) // both unset
-		{
-			return true;
+			bool bIsLhsSet = Lhs.Value.bIsSet;
+			bool bIsRhsSet = Rhs.Value.bIsSet;
+			if (bIsLhsSet != bIsRhsSet)
+			{
+				return false;
+			}
+			if (!bIsLhsSet) // both unset
+			{
+				return true;
+			}
 		}
+
 		return (*(const OptionalType*)&Lhs.Value) == (*(const OptionalType*)&Rhs.Value);
 	}
 
@@ -248,7 +336,14 @@ public:
 	/** @return true when the value is meaningful; false if calling GetValue() is undefined. */
 	bool IsSet() const
 	{
-		return Value.bIsSet;
+		if constexpr (bUsingIntrusiveUnsetState)
+		{
+			return !(*(const OptionalType*)&Value == FIntrusiveUnsetOptionalState{});
+		}
+		else
+		{
+			return Value.bIsSet;
+		}
 	}
 	FORCEINLINE explicit operator bool() const
 	{
@@ -301,7 +396,7 @@ public:
 	}
 
 private:
-	using ValueStorageType = UE::Core::Private::TOptionalStorage<sizeof(OptionalType)>;
+	using ValueStorageType = std::conditional_t<bUsingIntrusiveUnsetState, uint8[sizeof(OptionalType)], UE::Core::Private::TNonIntrusiveOptionalStorage<sizeof(OptionalType)>>;
 	alignas(OptionalType) ValueStorageType Value;
 };
 
diff --git a/Engine/Source/Runtime/Core/Public/Templates/Function.h b/Engine/Source/Runtime/Core/Public/Templates/Function.h
--- a/Engine/Source/Runtime/Core/Public/Templates/Function.h
+++ b/Engine/Source/Runtime/Core/Public/Templates/Function.h
@@ -4,6 +4,7 @@
 
 #include "CoreTypes.h"
 #include "Misc/AssertionMacros.h"
+#include "Misc/IntrusiveUnsetOptionalState.h"
 #include "HAL/UnrealMemory.h"
 #include "Templates/AndOrNot.h"
 #include "Templates/ChooseClass.h"
@@ -640,10 +641,20 @@ namespace UE::Core::Private::Function
 			Storage.Unbind();
 		}
 
+		void Reset()
+		{
+			if (Callable)
+			{
+				Storage.Unbind();
+				Callable = nullptr;
+			}
+		}
+
 	protected:
 		bool IsSet() const
 		{
-			static_assert(StorageType::bCanBeNull, "Nullable storage cannot be unset");
+			// Normally we'd assert that bCanBeNull here because it should always be true, but we reuse this
+			// function to test that a `TOptional<TFunctionRef>` with an intrusive state is unset.
 
 			return !!Callable;
 		}
@@ -802,7 +813,7 @@ namespace UE::Core::Private::Function
  * }
  */
 template <typename FuncType>
-class TFunctionRef : public UE::Core::Private::Function::TFunctionRefBase<UE::Core::Private::Function::FFunctionRefStoragePolicy, FuncType>
+class TFunctionRef final : public UE::Core::Private::Function::TFunctionRefBase<UE::Core::Private::Function::FFunctionRefStoragePolicy, FuncType>
 {
 	using Super = UE::Core::Private::Function::TFunctionRefBase<UE::Core::Private::Function::FFunctionRefStoragePolicy, FuncType>;
 
@@ -825,6 +836,27 @@ public:
 		// This constructor is disabled for TFunctionRef types so it isn't incorrectly selected as copy/move constructors.
 	}
 
+	/////////////////////////////////////////////////////
+	// Start - intrusive TOptional<TFunctionRef> state //
+	/////////////////////////////////////////////////////
+	constexpr static bool bHasIntrusiveUnsetOptionalState = true;
+	using IntrusiveUnsetOptionalStateType = TFunctionRef;
+
+	explicit TFunctionRef(FIntrusiveUnsetOptionalState)
+	{
+	}
+	void operator=(FIntrusiveUnsetOptionalState)
+	{
+		Super::Reset();
+	}
+	bool operator==(FIntrusiveUnsetOptionalState) const
+	{
+		return !Super::IsSet();
+	}
+	///////////////////////////////////////////////////
+	// End - intrusive TOptional<TFunctionRef> state //
+	///////////////////////////////////////////////////
+
 	TFunctionRef(const TFunctionRef&) = default;
 
 	// We delete the assignment operators because we don't want it to be confused with being related to
diff --git a/Engine/Source/Runtime/Core/Public/Templates/SharedPointer.h b/Engine/Source/Runtime/Core/Public/Templates/SharedPointer.h
--- a/Engine/Source/Runtime/Core/Public/Templates/SharedPointer.h
+++ b/Engine/Source/Runtime/Core/Public/Templates/SharedPointer.h
@@ -3,6 +3,7 @@
 #pragma once
 
 #include "CoreTypes.h"
+#include "Misc/IntrusiveUnsetOptionalState.h"
 #include "Templates/PointerIsConvertibleFromTo.h"
 #include "Misc/AssertionMacros.h"
 #include "HAL/UnrealMemory.h"
@@ -547,6 +548,29 @@ public:
 		return SharedReferenceCount.IsUnique();
 	}
 
+	/////////////////////////////////////////////////////
+	// Start - intrusive TOptional<TSharedRef> state //
+	/////////////////////////////////////////////////////
+	constexpr static bool bHasIntrusiveUnsetOptionalState = true;
+	using IntrusiveUnsetOptionalStateType = TSharedRef;
+
+	explicit TSharedRef(FIntrusiveUnsetOptionalState)
+		: Object(nullptr)
+	{
+	}
+	void operator=(FIntrusiveUnsetOptionalState)
+	{
+		Object = nullptr;
+		SharedReferenceCount = {};
+	}
+	bool operator==(FIntrusiveUnsetOptionalState) const
+	{
+		return !IsValid();
+	}
+	/////////////////////////////////////////////////
+	// End - intrusive TOptional<TSharedRef> state //
+	/////////////////////////////////////////////////
+
 private:
 	template<class OtherType>
 	void Init(OtherType* InObject)
