added UE::Tasks::MakeCompletedTask() that produces an already completed task with given result value

diff --git a/Engine/Source/Runtime/Core/Public/Tasks/Task.h b/Engine/Source/Runtime/Core/Public/Tasks/Task.h
--- a/Engine/Source/Runtime/Core/Public/Tasks/Task.h
+++ b/Engine/Source/Runtime/Core/Public/Tasks/Task.h
@@ -449,4 +449,15 @@ namespace UE::Tasks
 		ETaskPriority Priority;
 		EExtendedTaskPriority ExtendedPriority;
 	};
+
+	// creates and returns an already completed task holding a result value constructed from given args
+	template<typename ResultType, typename... ArgTypes>
+	TTask<ResultType> MakeCompletedTask(ArgTypes&&... Args)
+	{
+		return Launch(
+			UE_SOURCE_LOCATION,
+			[&] { return ResultType(Forward<ArgTypes>(Args)...); },
+			ETaskPriority::Default, // doesn't matter
+			EExtendedTaskPriority::Inline);
+	}
 }
diff --git a/Engine/Source/Runtime/Core/Tests/Tasks/TasksTest.cpp b/Engine/Source/Runtime/Core/Tests/Tasks/TasksTest.cpp
--- a/Engine/Source/Runtime/Core/Tests/Tasks/TasksTest.cpp
+++ b/Engine/Source/Runtime/Core/Tests/Tasks/TasksTest.cpp
@@ -7,6 +7,7 @@
 #include "HAL/Thread.h"
 #include "Async/ParallelFor.h"
 #include "Tests/TestHarnessAdapter.h"
+#include "Containers/UnrealString.h"
 
 #include <atomic>
 
@@ -1377,6 +1378,56 @@ namespace UE { namespace TasksTests
 		check(FinalTask.Wait(FTimespan::FromMilliseconds(100)));
 
 	}
+
+	TEST_CASE_NAMED(FTasksMakeCompletedTask, "System::Core::Tasks::MakeCompletedTask", "[.][ApplicationContextMask][EngineFilter]")
+	{
+		{	// basic
+			TTask<int> Task = MakeCompletedTask<int>(42);
+			check(Task.GetResult() == 42);
+		}
+
+		{	// move-only result
+			TTask<TUniquePtr<int>> Task = MakeCompletedTask<TUniquePtr<int>>(MakeUnique<int>(42));
+			check(*Task.GetResult() == 42);
+		}
+
+		{	// copy-only result
+			struct FCopyOnly
+			{
+				FCopyOnly() = default;
+				FCopyOnly(const FCopyOnly&) = default;
+				FCopyOnly& operator=(const FCopyOnly&) = default;
+			};
+
+			TTask<FCopyOnly> Task = MakeCompletedTask<FCopyOnly>(FCopyOnly{});
+		}
+
+		{ // non-movable result
+			struct FNonMovable
+			{
+				UE_NONCOPYABLE(FNonMovable);
+				FNonMovable() = default;
+			};
+
+			TTask<FNonMovable> Task = MakeCompletedTask<FNonMovable>();
+		}
+
+		{	// result construction from multiple args
+			struct FDummy
+			{
+				int Int;
+				FString Str;
+
+				FDummy(int InInt, const FString& InStr)
+					: Int(InInt)
+					, Str(InStr)
+				{}
+			};
+
+			TTask<FDummy> Task = MakeCompletedTask<FDummy>(42, TEXT("Test"));
+			check(Task.GetResult().Int == 42 && Task.GetResult().Str == TEXT("Test"));
+		}
+	}
 }}
 
 #endif // WITH_TESTS
