Cap LLM Total stat to be no less than Total Tracked on Android

diff --git a/Engine/Source/Runtime/Core/Private/HAL/LowLevelMemTracker.cpp b/Engine/Source/Runtime/Core/Private/HAL/LowLevelMemTracker.cpp
--- a/Engine/Source/Runtime/Core/Private/HAL/LowLevelMemTracker.cpp
+++ b/Engine/Source/Runtime/Core/Private/HAL/LowLevelMemTracker.cpp
@@ -953,8 +953,8 @@ void FLowLevelMemTracker::TickInternal()
 	FLLMTracker& PlatformTracker = *Trackers[static_cast<int32>(ELLMTracker::Platform)];
 
 	// calculate FMalloc unused stat and set it in the Default tracker
-	int64 FMallocAmount = DefaultTracker.GetAllocTypeAmount(ELLMAllocType::FMalloc);
-	int64 FMallocPlatformAmount = PlatformTracker.GetTagAmount(FindOrAddTagData(ELLMTag::FMalloc));
+	const int64 FMallocAmount = DefaultTracker.GetAllocTypeAmount(ELLMAllocType::FMalloc);
+	const int64 FMallocPlatformAmount = PlatformTracker.GetTagAmount(FindOrAddTagData(ELLMTag::FMalloc));
 	int64 FMallocUnused = FMallocPlatformAmount - FMallocAmount;
 	if (FMallocPlatformAmount == 0)
 	{
@@ -964,22 +964,27 @@ void FLowLevelMemTracker::TickInternal()
 	}
 	DefaultTracker.SetTagAmountInUpdate(FindOrAddTagData(ELLMTag::FMallocUnused), FMallocUnused, true);
 
-	int64 StaticOverhead = sizeof(FLowLevelMemTracker);
+	const int64 StaticOverhead = sizeof(FLowLevelMemTracker);
 	MemoryUsageCurrentOverhead = StaticOverhead + Allocator.GetTotal();
 	PlatformTracker.SetTagAmountInUpdate(FindOrAddTagData(ELLMTag::PlatformOverhead), MemoryUsageCurrentOverhead, !FPlatformMemory::TracksLLMAllocations());
 
+	const int64 TrackedTotal = DefaultTracker.GetTrackedTotal();
 	// Compare memory the platform thinks we have allocated to what we have tracked, including the program memory
-	FPlatformMemoryStats PlatformStats = FPlatformMemory::GetStats();
+	const FPlatformMemoryStats PlatformStats = FPlatformMemory::GetStats();
 #if PLATFORM_DESKTOP
 	// virtual is working set + paged out memory.
-	int64 PlatformProcessMemory = static_cast<int64>(PlatformStats.UsedVirtual);
-#elif PLATFORM_ANDROID || PLATFORM_IOS || UE_SERVER
-	int64 PlatformProcessMemory = static_cast<int64>(PlatformStats.UsedPhysical);
+	const int64 PlatformProcessMemory = static_cast<int64>(PlatformStats.UsedVirtual);
+#elif PLATFORM_ANDROID
+	// On some Android devices total used mem (VmRss + VmSwap) does not include GPU memory allocations
+	// and so it's going to be lower than TrackedTotal
+	const int64 PlatformProcessMemory = FMath::Max((int64)PlatformStats.UsedPhysical, TrackedTotal);
+#elif  PLATFORM_IOS || UE_SERVER
+	const int64 PlatformProcessMemory = static_cast<int64>(PlatformStats.UsedPhysical);
 #else
-	int64 PlatformProcessMemory = static_cast<int64>(PlatformStats.TotalPhysical) -
+	const int64 PlatformProcessMemory = static_cast<int64>(PlatformStats.TotalPhysical) -
 		static_cast<int64>(PlatformStats.AvailablePhysical);
 #endif
-	int64 PlatformTrackedTotal = PlatformTracker.GetTrackedTotal();
+	const int64 PlatformTrackedTotal = PlatformTracker.GetTrackedTotal();
 	MemoryUsagePlatformTotalUntracked = FMath::Max<int64>(0, PlatformProcessMemory - PlatformTrackedTotal);
 
 	PlatformTracker.SetTagAmountInUpdate(FindOrAddTagData(ELLMTag::PlatformTotal), PlatformProcessMemory, false);
@@ -988,11 +993,10 @@ void FLowLevelMemTracker::TickInternal()
 	PlatformTracker.SetTagAmountInUpdate(FindOrAddTagData(ELLMTag::PlatformOSAvailable),
 		PlatformStats.AvailablePhysical, false);
 
-	int64 TrackedTotal = DefaultTracker.GetTrackedTotal();
 	// remove the MemoryUsageCurrentOverhead from the "Total" for the default LLM as it's not something anyone needs
 	// to investigate when finding what to reduce the platform LLM will have the info 
-	int64 DefaultProcessMemory = PlatformProcessMemory - (!FPlatformMemory::TracksLLMAllocations() ? MemoryUsageCurrentOverhead : 0);
-	int64 DefaultUntracked = FMath::Max<int64>(0, DefaultProcessMemory - TrackedTotal);
+	const int64 DefaultProcessMemory = PlatformProcessMemory - (!FPlatformMemory::TracksLLMAllocations() ? MemoryUsageCurrentOverhead : 0);
+	const int64 DefaultUntracked = FMath::Max<int64>(0, DefaultProcessMemory - TrackedTotal);
 	DefaultTracker.SetTagAmountInUpdate(FindOrAddTagData(ELLMTag::Total), DefaultProcessMemory, false);
 	DefaultTracker.SetTagAmountInUpdate(FindOrAddTagData(ELLMTag::Untracked), DefaultUntracked, false);
 
