[Backout] - CL25923486 (backed out due to compile errors)
Original CL Desc
-----------------------------------------------------------------
Reenabling CLI filtering using grammar based parse.
To be submitted after CL 25765744

diff --git a/Engine/Source/Runtime/Core/Private/Misc/CommandLine.cpp b/Engine/Source/Runtime/Core/Private/Misc/CommandLine.cpp
--- a/Engine/Source/Runtime/Core/Private/Misc/CommandLine.cpp
+++ b/Engine/Source/Runtime/Core/Private/Misc/CommandLine.cpp
@@ -108,7 +108,7 @@ bool FCommandLine::IsCommandLineLoggingFiltered()
 #endif
 }
 
-bool FCommandLine::FilterCLIUsingGrammarBasedParser(TCHAR* OutLine, int32 MaxLen, const TCHAR* InLine, const TArrayView<FString>& AllowedList)
+bool FCommandLine::FilterMove(TCHAR* OutLine, int32 MaxLen, const TCHAR* InLine, const TArrayView<FString>& AllowedList)
 {
 	if (MaxLen == 0)
 	{
@@ -225,13 +225,17 @@ void FCommandLine::ApplyCommandLineAllowList()
 		FCommandLine::Parse(FilterForLoggingList, FilterArgsForLogging, Ignored);
 	}
 	// Process the original command line
-	FilterCLIInplace(OriginalCmdLine, UE_ARRAY_COUNT(OriginalCmdLine), OriginalCmdLine, OriginalList);
+	TArray<FString> OriginalList = FilterCommandLine(OriginalCmdLine);
+	BuildCommandLineAllowList(OriginalCmdLine, UE_ARRAY_COUNT(OriginalCmdLine), OriginalList);
 	// Process the current command line
-	FilterCLIInplace(CmdLine, UE_ARRAY_COUNT(CmdLine), CmdLine, CmdList);
+	TArray<FString> CmdList = FilterCommandLine(CmdLine);
+	BuildCommandLineAllowList(CmdLine, UE_ARRAY_COUNT(CmdLine), CmdList);
 	// Process the command line for logging purposes
-	FilterCLIInplace(LoggingCmdLine, UE_ARRAY_COUNT(LoggingCmdLine), LoggingCmdLine, LoggingCmdList);
+	TArray<FString> LoggingCmdList = FilterCommandLineForLogging(LoggingCmdLine);
+	BuildCommandLineAllowList(LoggingCmdLine, UE_ARRAY_COUNT(LoggingCmdLine), LoggingCmdList);
 	// Process the original command line for logging purposes
-	FilterCLIInplace(LoggingOriginalCmdLine, UE_ARRAY_COUNT(LoggingOriginalCmdLine), LoggingOriginalCmdLine, LoggingOriginalCmdList);
+	TArray<FString> LoggingOriginalCmdList = FilterCommandLineForLogging(LoggingOriginalCmdLine);
+	BuildCommandLineAllowList(LoggingOriginalCmdLine, UE_ARRAY_COUNT(LoggingOriginalCmdLine), LoggingOriginalCmdList);
 }
 
 TArray<FString> FCommandLine::FilterCommandLine(TCHAR* CommandLine)
diff --git a/Engine/Source/Runtime/Core/Public/Misc/CommandLine.h b/Engine/Source/Runtime/Core/Public/Misc/CommandLine.h
--- a/Engine/Source/Runtime/Core/Public/Misc/CommandLine.h
+++ b/Engine/Source/Runtime/Core/Public/Misc/CommandLine.h
@@ -110,8 +110,7 @@ struct FCommandLine
 	CORE_API static FString BuildFromArgV(const ANSICHAR* Prefix, int32 ArgC, ANSICHAR* ArgV[], const ANSICHAR* Suffix);
 
 	/**
-	* FilterCLIUsingGrammarBasedParser parses CLI style arguments.
-	* filters for commands or keys specified in the AllowedList
+	* Filter Move parses CLI style arguments in the InLine, filters for commands or keys specified in the AllowedList
 	* and writes the to the OutLine.
 	* OutLine and InLine may be point to the same buffer.
 	* 
@@ -122,7 +121,7 @@ struct FCommandLine
 	* 
 	* Returns true if Outline was large enough to hold the filtered string, Otherwise returns false.
 	*/
-	static bool FilterCLIUsingGrammarBasedParser(TCHAR* OutLine, int32 MaxLen, const TCHAR* InLine, const TArrayView<FString>& AllowedList);
+	static bool FilterMove(TCHAR* OutLine, int32 MaxLen, const TCHAR* InLine, const TArrayView<FString>& AllowedList);
 private:
 #if UE_COMMAND_LINE_USES_ALLOW_LIST
 	/** Filters both the original and current command line list for approved only args */
diff --git a/Engine/Source/Runtime/Core/Tests/Misc/CommandLineTests.cpp b/Engine/Source/Runtime/Core/Tests/Misc/CommandLineTests.cpp
--- a/Engine/Source/Runtime/Core/Tests/Misc/CommandLineTests.cpp
+++ b/Engine/Source/Runtime/Core/Tests/Misc/CommandLineTests.cpp
@@ -10,7 +10,7 @@
 
 #include <catch2/generators/catch_generators.hpp>
 
-TEST_CASE("CommandLine::Filtering::FilterCLIInplace", "[Smoke]")
+TEST_CASE("CommandLine::Filtering::FilterMove", "[Smoke]")
 {
 	TArray<FString> AllowedList{ { "cmd_a", "-cmd_b" } };
 
@@ -31,7 +31,7 @@ TEST_CASE("CommandLine::Filtering::FilterCLIInplace", "[Smoke]")
 		}));
 		TCHAR Result[256]{};
 
-		CHECK(FCommandLine::FilterCLIUsingGrammarBasedParser(Result, UE_ARRAY_COUNT(Result), Input, AllowedList));
+		CHECK(FCommandLine::FilterMove(Result, UE_ARRAY_COUNT(Result), Input, AllowedList));
 		CHECK(FStringView{ Result } == FStringView{ Expected });
 	}
 
@@ -45,7 +45,7 @@ TEST_CASE("CommandLine::Filtering::FilterCLIInplace", "[Smoke]")
 			}));
 		TCHAR Result[256]{};
 
-		CHECK(FCommandLine::FilterCLIUsingGrammarBasedParser(Result, UE_ARRAY_COUNT(Result), Input, AllowedList));
+		CHECK(FCommandLine::FilterMove(Result, UE_ARRAY_COUNT(Result), Input, AllowedList));
 		CHECK(FStringView{ Result } == FStringView{ Expected });
 	}
 
@@ -71,7 +71,7 @@ TEST_CASE("CommandLine::Filtering::FilterCLIInplace", "[Smoke]")
 		TCHAR SourceAndResult[256]{};
 		
 		FCString::Strcpy(SourceAndResult, Input);
-		CHECK(FCommandLine::FilterCLIUsingGrammarBasedParser(SourceAndResult, UE_ARRAY_COUNT(SourceAndResult), SourceAndResult, AllowedList));
+		CHECK(FCommandLine::FilterMove(SourceAndResult, UE_ARRAY_COUNT(SourceAndResult), SourceAndResult, AllowedList));
 		CHECK(FStringView{ SourceAndResult } == FStringView{ Expected });
 	}
 
@@ -80,7 +80,7 @@ TEST_CASE("CommandLine::Filtering::FilterCLIInplace", "[Smoke]")
 		const TCHAR* Input = TEXT("-cmd_a --cmd_b");
 		TCHAR Result[256] = TEXT("Not Empty");
 
-		CHECK(FCommandLine::FilterCLIUsingGrammarBasedParser(Result, UE_ARRAY_COUNT(Result), Input, {}));
+		CHECK(FCommandLine::FilterMove(Result, UE_ARRAY_COUNT(Result), Input, {}));
 		CHECK(Result[0] == TCHAR(0));
 	}
 
@@ -89,7 +89,7 @@ TEST_CASE("CommandLine::Filtering::FilterCLIInplace", "[Smoke]")
 		const TCHAR* Input = TEXT("-cmd_a --cmd_b");
 		TCHAR Result[5]{};
 
-		CHECK(false == FCommandLine::FilterCLIUsingGrammarBasedParser(Result, UE_ARRAY_COUNT(Result), Input, AllowedList));
+		CHECK(false == FCommandLine::FilterMove(Result, UE_ARRAY_COUNT(Result), Input, AllowedList));
 		CHECK(Result[0] == TCHAR(0));
 	}
 
@@ -126,7 +126,7 @@ TEST_CASE("CommandLine::Filtering::FilterCLIInplace", "[Smoke]")
 		TCHAR SourceAndResult[256]{};
 
 		FCString::Strcpy(SourceAndResult, Input);
-		CHECK(FCommandLine::FilterCLIUsingGrammarBasedParser(SourceAndResult, UE_ARRAY_COUNT(SourceAndResult), SourceAndResult, ApprovedArgs));
+		CHECK(FCommandLine::FilterMove(SourceAndResult, UE_ARRAY_COUNT(SourceAndResult), SourceAndResult, ApprovedArgs));
 		CHECK(FStringView{ SourceAndResult } == FStringView{ Expected });
 	}
 
