Add method to set LLM category from FName.

Needed to explicitly set categories that are not tracked by normal means (e.g. TraceLog memory usage) using a FName identifier

#rnx

diff --git a/Engine/Source/Runtime/Core/Private/HAL/LowLevelMemTracker.cpp b/Engine/Source/Runtime/Core/Private/HAL/LowLevelMemTracker.cpp
--- a/Engine/Source/Runtime/Core/Private/HAL/LowLevelMemTracker.cpp
+++ b/Engine/Source/Runtime/Core/Private/HAL/LowLevelMemTracker.cpp
@@ -2190,6 +2190,21 @@ void FLowLevelMemTracker::SetTagAmountForTracker(ELLMTracker Tracker, ELLMTag Ta
 	GetTracker(Tracker)->SetTagAmountExternal(TagData, Amount, bAddToTotal);
 }
 
+void FLowLevelMemTracker::SetTagAmountForTracker(ELLMTracker Tracker, FName Tag, ELLMTagSet TagSet, int64 Amount, bool bAddToTotal)
+{
+	using namespace UE::LLMPrivate;
+
+	if (bIsDisabled)
+	{
+		return;
+	}
+	BootstrapInitialise();
+	const FTagData* TagData = FindOrAddTagData(Tag, TagSet);
+
+	FScopeLock UpdateScopeLock(&UpdateLock); // uses of TagSizes are guarded by the UpdateLock
+	GetTracker(Tracker)->SetTagAmountExternal(TagData, Amount, bAddToTotal);
+}
+
 int64 FLowLevelMemTracker::GetActiveTag(ELLMTracker Tracker)
 {
 	using namespace UE::LLMPrivate;
diff --git a/Engine/Source/Runtime/Core/Public/HAL/LowLevelMemTracker.h b/Engine/Source/Runtime/Core/Public/HAL/LowLevelMemTracker.h
--- a/Engine/Source/Runtime/Core/Public/HAL/LowLevelMemTracker.h
+++ b/Engine/Source/Runtime/Core/Public/HAL/LowLevelMemTracker.h
@@ -644,6 +644,9 @@ public:
 	/** Set the amount of memory for an ELLMTag for a given tracker and optionally update the total tracked memory. */
 	void SetTagAmountForTracker(ELLMTracker Tracker, ELLMTag Tag, int64 Amount, bool bAddToTotal);
 
+	/** Set the amount of memory for an tag name for a given tracker and optionally update the total tracked memory. */
+	void SetTagAmountForTracker(ELLMTracker Tracker, FName Tag, ELLMTagSet TagSet, int64 Amount, bool bAddToTotal);
+	
 	/** Dump the display name of the current TagData for the given tracker to the output. */
 	uint64 DumpTag(ELLMTracker Tracker, const char* FileName, int LineNumber);
 
