Add LLM tracing for TraceLog memory usage

diff --git a/Engine/Source/Runtime/Core/Private/ProfilingDebugging/TraceAuxiliary.cpp b/Engine/Source/Runtime/Core/Private/ProfilingDebugging/TraceAuxiliary.cpp
--- a/Engine/Source/Runtime/Core/Private/ProfilingDebugging/TraceAuxiliary.cpp
+++ b/Engine/Source/Runtime/Core/Private/ProfilingDebugging/TraceAuxiliary.cpp
@@ -708,6 +708,9 @@ void FTraceAuxiliaryImpl::StartWorkerThread()
 	}
 }
 
+LLM_DEFINE_TAG(Trace_TraceLog);
+LLM_DEFINE_TAG(Trace_TraceLog_Cache);
+
 ////////////////////////////////////////////////////////////////////////////////
 void FTraceAuxiliaryImpl::StartEndFramePump()
 {
@@ -728,6 +731,12 @@ void FTraceAuxiliaryImpl::StartEndFramePump()
 			SET_MEMORY_STAT(STAT_TraceCacheUsed, Stats.CacheUsed);
 			SET_MEMORY_STAT(STAT_TraceCacheWaste, Stats.CacheWaste);
 			SET_MEMORY_STAT(STAT_TraceSent, Stats.BytesSent);
+
+#if ENABLE_LOW_LEVEL_MEM_TRACKER
+			const int64 NonCacheMemory = Stats.MemoryUsed - Stats.CacheAllocated;
+			FLowLevelMemTracker::Get().SetTagAmountForTracker(ELLMTracker::Default, LLM_TAG_NAME(Trace_TraceLog), ELLMTagSet::None, NonCacheMemory, true);
+			FLowLevelMemTracker::Get().SetTagAmountForTracker(ELLMTracker::Default, LLM_TAG_NAME(Trace_TraceLog_Cache), ELLMTagSet::None, Stats.CacheAllocated, true);
+#endif
 		});
 	}
 }
