Return allocator result in the case of zero sized allocations.

Remove checks for zero sized allocations in Trace's allocation wrappers leaving it to the allocator implementation to decide. Analysis already supports zero sized allocations (as long as they have unique addresses).

diff --git a/Engine/Source/Runtime/Core/Private/ProfilingDebugging/MemoryTrace.cpp b/Engine/Source/Runtime/Core/Private/ProfilingDebugging/MemoryTrace.cpp
--- a/Engine/Source/Runtime/Core/Private/ProfilingDebugging/MemoryTrace.cpp
+++ b/Engine/Source/Runtime/Core/Private/ProfilingDebugging/MemoryTrace.cpp
@@ -207,11 +207,6 @@ uint32 FMallocWrapper::GetActualAlignment(SIZE_T Size, uint32 Alignment)
 ////////////////////////////////////////////////////////////////////////////////
 void* FMallocWrapper::Malloc(SIZE_T Size, uint32 Alignment)
 {
-	if (Size == 0)
-	{
-		return nullptr;
-	}
-
 	uint32 ActualAlignment = GetActualAlignment(Size, Alignment);
 	void* Address = InnerMalloc->Malloc(Size, Alignment);
 
@@ -229,12 +224,6 @@ void* FMallocWrapper::Realloc(void* PrevAddress, SIZE_T NewSize, uint32 Alignmen
 		return Malloc(NewSize, Alignment);
 	}
 
-	if (NewSize == 0)
-	{
-		Free(PrevAddress);
-		return nullptr;
-	}
-
 	MemoryTrace_ReallocFree((uint64)PrevAddress);
 
 	void* RetAddress = InnerMalloc->Realloc(PrevAddress, NewSize, Alignment);
