Tasks integration: destroying task body as soon as it's executed. Some systems rely on this (probably a bug), but this also improves mem usage. Same approach is used by [REDACTED_PERSON] and TaskGraph. Tasks integration is disabled so this change is compiled out for now

diff --git a/Engine/Source/Runtime/Core/Public/Async/TaskGraphInterfaces.h b/Engine/Source/Runtime/Core/Public/Async/TaskGraphInterfaces.h
--- a/Engine/Source/Runtime/Core/Public/Async/TaskGraphInterfaces.h
+++ b/Engine/Source/Runtime/Core/Public/Async/TaskGraphInterfaces.h
@@ -634,11 +634,6 @@ public:
 		return FConstructor(Prerequisites);
 	}
 
-	FORCENOINLINE virtual ~TGraphTask() override
-	{
-		DestructItem(TaskStorage.GetTypedPtr());
-	}
-
 	static void* operator new(size_t Size);
 	static void operator delete(void* Ptr, size_t Size);
 
@@ -664,6 +659,7 @@ private:
 				ENamedThreads::Type ThreadIndex = ENamedThreads::GetThreadIndex(TaskObject->GetDesiredThread());
 
 				TaskObject->DoTask(ThreadIndex, GraphEventRef);
+				DestructItem(TaskObject);
 			}
 		);
 	}
