Improve command line handling to better handle quotes on Unix.  PR #5583

diff --git a/Engine/Source/Runtime/Core/Private/Unix/UnixPlatformProcess.cpp b/Engine/Source/Runtime/Core/Private/Unix/UnixPlatformProcess.cpp
--- a/Engine/Source/Runtime/Core/Private/Unix/UnixPlatformProcess.cpp
+++ b/Engine/Source/Runtime/Core/Private/Unix/UnixPlatformProcess.cpp
@@ -787,6 +787,62 @@ FProcHandle FUnixPlatformProcess::CreateProc(const TCHAR* URL, const TCHAR* Parm
 	return CreateProc(URL, Parms, bLaunchDetached, bLaunchHidden, bLaunchReallyHidden, OutProcessID, PriorityModifier, OptionalWorkingDirectory, PipeWriteChild, PipeReadChild, PipeWriteChild);
 }
 
+static char* MallocedUtf8FromString(const FString& Str)
+{
+	FTCHARToUTF8 AnsiBuffer(*Str);
+	const char* Ansi = AnsiBuffer.Get();
+	size_t AnsiSize = FCStringAnsi::Strlen(Ansi) + 1;	// will work correctly with UTF-8
+	check(AnsiSize);
+
+	char* Ret = reinterpret_cast<char*>(FMemory::Malloc(AnsiSize));
+	check(Ret);
+
+	FCStringAnsi::Strncpy(Ret, Ansi, AnsiSize);	// will work correctly with UTF-8
+	return Ret;
+}
+
+static bool AddCmdLineArgumentTo(char** Argv, int& Argc, const FString& CurArg)
+{
+	if (Argc == PlatformProcessLimits::MaxArgvParameters)
+	{
+		UE_LOG(LogHAL, Warning, TEXT("FUnixPlatformProcess::CreateProc: too many (%d) commandline arguments passed, will only pass %d"),
+			Argc, PlatformProcessLimits::MaxArgvParameters);
+		return false;
+	}
+	Argv[Argc] = MallocedUtf8FromString(CurArg);
+	UE_LOG(LogHAL, Verbose, TEXT("FUnixPlatformProcess::CreateProc: Argv[%d] = '%s'"), Argc, *CurArg);
+	Argc++;
+	return true;
+}
+
+// returns true if the token completes the current argument
+static bool ParseCmdLineToken(TCHAR token, bool& OutIsInString, bool& OutHasArg, FString& OutCurArg, bool& OutEOL)
+{
+	if (token == TEXT('\0'))
+	{
+		OutEOL = true;
+		return OutHasArg;
+	}
+
+	if (token == TEXT('"'))
+	{
+		OutIsInString = !OutIsInString;
+		OutHasArg = true;
+		return false;
+	}
+
+	if (token == TEXT(' ') && !OutIsInString)
+	{
+		return OutHasArg;
+	}
+
+	// if we've made it this far, the token should be added to the argument
+	OutCurArg += token;
+	OutHasArg = true;
+
+	return false;
+}
+
 FProcHandle FUnixPlatformProcess::CreateProc(const TCHAR* URL, const TCHAR* Parms, bool bLaunchDetached, bool bLaunchHidden, bool bLaunchReallyHidden, uint32* OutProcessID, int32 PriorityModifier, const TCHAR* OptionalWorkingDirectory, void* PipeWriteChild, void* PipeReadChild, void* PipeStdErrChild)
 {
 	// @TODO bLaunchHidden bLaunchReallyHidden are not handled
@@ -832,17 +888,17 @@ FProcHandle FUnixPlatformProcess::CreateProc(const TCHAR* URL, const TCHAR* Parm
 	const FString Commandline = FString::Printf(TEXT("\"%s\" %s"), *ProcessPath, Parms);
 	UE_LOG(LogHAL, Verbose, TEXT("FUnixPlatformProcess::CreateProc: '%s'"), *Commandline);
 
-	TArray<FString> ArgvArray;
-	int Argc = Commandline.ParseIntoArray(ArgvArray, TEXT(" "), true);
+	int Argc = 1;
 	char* Argv[PlatformProcessLimits::MaxArgvParameters + 1] = { NULL };	// last argument is NULL, hence +1
+	Argv[0] = MallocedUtf8FromString(ProcessPath);
 	struct CleanupArgvOnExit
 	{
 		int Argc;
 		char** Argv;	// relying on it being long enough to hold Argc elements
 
-		CleanupArgvOnExit( int InArgc, char *InArgv[] )
-			:	Argc(InArgc)
-			,	Argv(InArgv)
+		CleanupArgvOnExit(int InArgc, char* InArgv[])
+			: Argc(InArgc)
+			, Argv(InArgv)
 		{}
 
 		~CleanupArgvOnExit()
@@ -854,97 +910,39 @@ FProcHandle FUnixPlatformProcess::CreateProc(const TCHAR* URL, const TCHAR* Parm
 		}
 	} CleanupGuard(Argc, Argv);
 
-	// make sure we do not lose arguments with spaces in them due to Commandline.ParseIntoArray breaking them apart above
-	// @todo this code might need to be optimized somehow and integrated with main argument parser below it
-	TArray<FString> NewArgvArray;
-	if (Argc > 0)
-	{
-		if (Argc > PlatformProcessLimits::MaxArgvParameters)
-		{
-			UE_LOG(LogHAL, Warning, TEXT("FUnixPlatformProcess::CreateProc: too many (%d) commandline arguments passed, will only pass %d"),
-				Argc, PlatformProcessLimits::MaxArgvParameters);
-			Argc = PlatformProcessLimits::MaxArgvParameters;
-		}
+	UE_LOG(LogHAL, Verbose, TEXT("FUnixPlatformProcess::CreateProc: ProcessPath = '%s' Parms = '%s'"), *ProcessPath, Parms);
+	FString CurArg; // current argument, during parsing new chars will be appended, will be reused, once the argument is complete
+	const TCHAR* CurChar = Parms; // pointer to the current char
+	bool IsInString = false; // are we in a string?  if yes, spaces are treated as normal chars
+	bool HasArg = false; // do we have a partial argument? CurArg might be empty if Parms contains "". Parms might contain two or more spaces in a row, so not every space indicates the end of an argument
+	bool EOL = false;
 
-		FString MultiPartArg;
-		for (int32 Index = 0; Index < Argc; Index++)
+	// parse Parms and fill Argv
+	while (!EOL)
+	{
+		if (ParseCmdLineToken(*CurChar, IsInString, HasArg, CurArg, EOL))
 		{
-			if (MultiPartArg.IsEmpty())
-			{
-				if ((ArgvArray[Index].StartsWith(TEXT("\"")) && !ArgvArray[Index].EndsWith(TEXT("\""))) // check for a starting quote but no ending quote, excludes quoted single arguments
-					|| (ArgvArray[Index].Contains(TEXT("=\"")) && !ArgvArray[Index].EndsWith(TEXT("\""))) // check for quote after =, but no ending quote, this gets arguments of the type -blah="string string string"
-					|| ArgvArray[Index].EndsWith(TEXT("=\""))) // check for ending quote after =, this gets arguments of the type -blah=" string string string "
-				{
-					MultiPartArg = ArgvArray[Index];
-				}
-				else
-				{
-					if (ArgvArray[Index].Contains(TEXT("=\"")))
-					{
-						FString SingleArg = ArgvArray[Index];
-						SingleArg = SingleArg.Replace(TEXT("=\""), TEXT("="));
-						NewArgvArray.Add(SingleArg.TrimQuotes(NULL));
-					}
-					else
-					{
-						NewArgvArray.Add(ArgvArray[Index].TrimQuotes(NULL));
-					}
-				}
-			}
-			else
+			if (!AddCmdLineArgumentTo(Argv, Argc, CurArg))
 			{
-				MultiPartArg += TEXT(" ");
-				MultiPartArg += ArgvArray[Index];
-				if (ArgvArray[Index].EndsWith(TEXT("\"")))
-				{
-					if (MultiPartArg.StartsWith(TEXT("\"")))
-					{
-						NewArgvArray.Add(MultiPartArg.TrimQuotes(NULL));
-					}
-					else if (MultiPartArg.Contains(TEXT("=\"")))
-					{
-						FString SingleArg = MultiPartArg.Replace(TEXT("=\""), TEXT("="));
-						NewArgvArray.Add(SingleArg.TrimQuotes(nullptr));
-					}
-					else
-					{
-						NewArgvArray.Add(MultiPartArg);
-					}
-					MultiPartArg.Empty();
-				}
+				break;	// we could not add any more arguments
 			}
-		}
-	}
-	// update Argc with the new argument count
-	Argc = NewArgvArray.Num();
 
-	if (Argc > 0)	// almost always, unless there's no program name
-	{
-		if (Argc > PlatformProcessLimits::MaxArgvParameters)
-		{
-			UE_LOG(LogHAL, Warning, TEXT("FUnixPlatformProcess::CreateProc: too many (%d) commandline arguments passed, will only pass %d"), 
-				Argc, PlatformProcessLimits::MaxArgvParameters);
-			Argc = PlatformProcessLimits::MaxArgvParameters;
+			HasArg = false;
+			CurArg.Reset(0);
 		}
 
-		for (int Idx = 0; Idx < Argc; ++Idx)
-		{
-			FTCHARToUTF8 AnsiBuffer(*NewArgvArray[Idx]);
-			const char* Ansi = AnsiBuffer.Get();
-			size_t AnsiSize = FCStringAnsi::Strlen(Ansi) + 1;	// will work correctly with UTF-8
-			check(AnsiSize);
-
-			Argv[Idx] = reinterpret_cast< char* >( FMemory::Malloc(AnsiSize) );
-			check(Argv[Idx]);
-
-			FCStringAnsi::Strncpy(Argv[Idx], Ansi, AnsiSize);	// will work correctly with UTF-8
-		}
+		CurChar++;
+	}
 
-		// last Argv should be NULL
-		check(Argc <= PlatformProcessLimits::MaxArgvParameters + 1);
-		Argv[Argc] = NULL;
+	if (IsInString)
+	{
+		UE_LOG(LogHAL, Warning, TEXT("FUnixPlatformProcess::CreateProc: mismatched quotes in command line (%s %s)"), *ProcessPath, Parms);
 	}
 
+	// we assume PlatformProcessLimits::MaxArgvParameters is >= 1. Since Argc starts at 1 and can never grow larger than PlatformProcessLimits::MaxArgvParameters,
+	// we are within the Array
+	Argv[Argc] = NULL;
+
 	extern char ** environ;	// provided by libc
 	pid_t ChildPid = -1;
 
