[autortfm] add a CVar to control enabling the AutoRTFM runtime components

diff --git a/Engine/Source/Runtime/Core/Private/AutoRTFM/API.cpp b/Engine/Source/Runtime/Core/Private/AutoRTFM/API.cpp
--- a/Engine/Source/Runtime/Core/Private/AutoRTFM/API.cpp
+++ b/Engine/Source/Runtime/Core/Private/AutoRTFM/API.cpp
@@ -1,8 +1,21 @@
 // Copyright Epic Games, Inc. All Rights Reserved.
 
-#if (defined(__AUTORTFM) && __AUTORTFM)
 
 #include "AutoRTFM/AutoRTFM.h"
+
+#if UE_AUTORTFM
+bool GAutoRTFMRuntimeEnabled = true;
+static FAutoConsoleVariableRef CVarAutoRTFMRuntimeEnabled(
+	TEXT("AutoRTFMRuntimeEnabled"),
+	GAutoRTFMRuntimeEnabled,
+	TEXT("Enables the AutoRTFM runtime"),
+	ECVF_Default
+);
+#else
+static constexpr bool GAutoRTFMRuntimeEnabled = false;
+#endif
+
+#if (defined(__AUTORTFM) && __AUTORTFM)
 #include "AutoRTFM/AutoRTFMConstants.h"
 #include "CallNest.h"
 #include "Context.h"
@@ -22,7 +35,12 @@ namespace AutoRTFM
 
 FORCENOINLINE extern "C" bool autortfm_is_transactional()
 {
-    return FContext::Get()->IsTransactional();
+	if(GAutoRTFMRuntimeEnabled)
+	{
+		return FContext::Get()->IsTransactional();
+	}
+
+	return false;
 }
 
 FORCENOINLINE extern "C" bool autortfm_is_closed()
@@ -33,7 +51,13 @@ FORCENOINLINE extern "C" bool autortfm_is_closed()
 // First Part - the API exposed outside transactions.
 FORCENOINLINE extern "C" autortfm_result autortfm_transact(void (*Work)(void* Arg), void* Arg)
 {
-    return static_cast<autortfm_result>(FContext::Get()->Transact(Work, Arg));
+	if(GAutoRTFMRuntimeEnabled)
+	{
+	    return static_cast<autortfm_result>(FContext::Get()->Transact(Work, Arg));
+	}
+
+	Work(Arg);
+	return autortfm_committed;
 }
 
 FORCENOINLINE extern "C" void autortfm_commit(void (*Work)(void* Arg), void* Arg)
@@ -74,7 +98,12 @@ FORCENOINLINE extern "C" void autortfm_clear_transaction_status()
 
 FORCENOINLINE extern "C" bool autortfm_is_aborting()
 {
-	return FContext::Get()->IsAborting();
+	if(GAutoRTFMRuntimeEnabled)
+	{
+		return FContext::Get()->IsAborting();
+	}
+
+	return false;
 }
 
 FORCENOINLINE extern "C" bool autortfm_current_nest_throw()
@@ -99,16 +128,23 @@ FORCENOINLINE extern "C" void autortfm_open(void (*Work)(void* Arg), void* Arg)
 
 FORCENOINLINE extern "C" autortfm_status autortfm_close(void (*Work)(void* Arg), void* Arg)
 {
-	UE_CLOG(!FContext::IsTransactional(), LogAutoRTFM, Fatal, TEXT("Close called from an outside a transaction."));
-
 	autortfm_status Result = autortfm_status_ontrack;
 
-    FContext* Context = FContext::Get();
-    void (*WorkClone)(void* Arg, FContext* Context) = FunctionMapLookup(Work, Context, "autortfm_close");
-    if (WorkClone)
-    {
-		Result = static_cast<autortfm_status>(Context->CallClosedNest(WorkClone, Arg));
-    }
+	if(GAutoRTFMRuntimeEnabled)
+	{
+		UE_CLOG(!FContext::IsTransactional(), LogAutoRTFM, Fatal, TEXT("Close called from an outside a transaction."));
+
+		FContext* Context = FContext::Get();
+		void (*WorkClone)(void* Arg, FContext* Context) = FunctionMapLookup(Work, Context, "autortfm_close");
+		if (WorkClone)
+		{
+			Result = static_cast<autortfm_status>(Context->CallClosedNest(WorkClone, Arg));
+		}
+	}
+	else
+	{
+		Work(Arg);
+	}
 
 	return Result;
 }
