Tasks integration: minor profiler-guided optimisations

diff --git a/Engine/Source/Runtime/Core/Public/Async/TaskGraphInterfaces.h b/Engine/Source/Runtime/Core/Public/Async/TaskGraphInterfaces.h
--- a/Engine/Source/Runtime/Core/Public/Async/TaskGraphInterfaces.h
+++ b/Engine/Source/Runtime/Core/Public/Async/TaskGraphInterfaces.h
@@ -487,7 +487,7 @@ public:
 		TryLaunch(0);
 	}
 
-	void Execute(TArray<FBaseGraphTask*>& NewTasks, ENamedThreads::Type CurrentThread, bool bDeleteOnCompletion)
+	FORCEINLINE void Execute(TArray<FBaseGraphTask*>& NewTasks, ENamedThreads::Type CurrentThread, bool bDeleteOnCompletion)
 	{	// only called for named thread tasks, normal tasks are executed using `FTaskBase` API directly (see `TGraphTask`)
 		checkSlow(NewTasks.Num() == 0);
 		checkSlow(bDeleteOnCompletion);
@@ -648,7 +648,7 @@ private:
 		FBaseGraphTask::Init(TEXT("GraphTask"), InPriority, InExtendedPriority);
 	}
 
-	virtual bool TryExecuteTaskVirtual() override
+	virtual bool TryExecuteTaskVirtual() final
 	{
 		return TryExecute(
 			[](UE::Tasks::Private::FTaskBase& Task)
diff --git a/Engine/Source/Runtime/Core/Public/Tasks/TaskPrivate.h b/Engine/Source/Runtime/Core/Public/Tasks/TaskPrivate.h
--- a/Engine/Source/Runtime/Core/Public/Tasks/TaskPrivate.h
+++ b/Engine/Source/Runtime/Core/Public/Tasks/TaskPrivate.h
@@ -414,7 +414,7 @@ namespace UE::Tasks
 			// tries to get execution permission and if successful, executes given task body and completes the task if there're no pending nested tasks. 
 			// does all required accounting before/after task execution. the task can be deleted as a result of this call.
 			// @returns true if the task was executed by the current thread
-			FORCENOINLINE bool TryExecute(FTaskBodyType TaskBody)
+			bool TryExecute(FTaskBodyType TaskBody)
 			{
 				if (!TrySetExecutionFlag())
 				{
@@ -487,7 +487,7 @@ namespace UE::Tasks
 
 			CORE_API void ClearPipe();
 
-			bool TryExecuteTask()
+			FORCEINLINE bool TryExecuteTask()
 			{
 				UE::FInheritedContextScope InheritedContextScope = RestoreInheritedContext();
 				return TryExecuteTaskVirtual();
@@ -678,7 +678,7 @@ namespace UE::Tasks
 			UE_NONCOPYABLE(TExecutableTaskBase);
 
 		public:
-			virtual bool TryExecuteTaskVirtual() override
+			virtual bool TryExecuteTaskVirtual() final
 			{
 				return FTaskBase::TryExecute(
 					[](FTaskBase& Task)
@@ -713,7 +713,7 @@ namespace UE::Tasks
 			UE_NONCOPYABLE(TExecutableTaskBase);
 
 		public:
-			virtual bool TryExecuteTaskVirtual() override
+			virtual bool TryExecuteTaskVirtual() final
 			{
 				return TryExecute(
 					[](FTaskBase& Task)
