Skip unreal trace server start when running under TSAN to prevent immediate crash on fork

#rnx

diff --git a/Engine/Source/Runtime/Core/Private/ProfilingDebugging/TraceAuxiliary.cpp b/Engine/Source/Runtime/Core/Private/ProfilingDebugging/TraceAuxiliary.cpp
--- a/Engine/Source/Runtime/Core/Private/ProfilingDebugging/TraceAuxiliary.cpp
+++ b/Engine/Source/Runtime/Core/Private/ProfilingDebugging/TraceAuxiliary.cpp
@@ -1197,6 +1197,8 @@ static void LaunchUnrealTraceInternal(const TCHAR* CommandLine)
 #if PLATFORM_UNIX || PLATFORM_MAC
 static void LaunchUnrealTraceInternal(const TCHAR* CommandLine)
 {
+// TSAN doesn't like fork(), so disable this for now.
+#if !USING_THREAD_SANITISER
 	if (GUnrealTraceLaunched.load(std::memory_order_relaxed))
 	{
 		UE_LOG(LogCore, Log, TEXT("UnrealTraceServer: Trace store already started"));
@@ -1258,6 +1260,7 @@ static void LaunchUnrealTraceInternal(const TCHAR* CommandLine)
 		UE_LOG(LogCore, Log, TEXT("UnrealTraceServer: Trace store launch successful"));
 		GUnrealTraceLaunched.fetch_add(1, std::memory_order_relaxed);
 	}
+#endif // #if !USING_THREAD_SANITISER
 }
 #endif // PLATFORM_UNIX/MAC
 #endif // WITH_UNREAL_TRACE_LAUNCH
