Add missing header for CVars

diff --git a/Engine/Source/Runtime/Core/Private/AutoRTFM/API.cpp b/Engine/Source/Runtime/Core/Private/AutoRTFM/API.cpp
--- a/Engine/Source/Runtime/Core/Private/AutoRTFM/API.cpp
+++ b/Engine/Source/Runtime/Core/Private/AutoRTFM/API.cpp
@@ -2,6 +2,7 @@
 
 
 #include "AutoRTFM/AutoRTFM.h"
+#include "HAL/IConsoleManager.h"
 
 #if UE_AUTORTFM
 bool GAutoRTFMRuntimeEnabled = true;
@@ -35,7 +36,7 @@ namespace AutoRTFM
 
 FORCENOINLINE extern "C" bool autortfm_is_transactional()
 {
-	if(GAutoRTFMRuntimeEnabled)
+	if (GAutoRTFMRuntimeEnabled)
 	{
 		return FContext::Get()->IsTransactional();
 	}
@@ -51,7 +52,7 @@ FORCENOINLINE extern "C" bool autortfm_is_closed()
 // First Part - the API exposed outside transactions.
 FORCENOINLINE extern "C" autortfm_result autortfm_transact(void (*Work)(void* Arg), void* Arg)
 {
-	if(GAutoRTFMRuntimeEnabled)
+	if (GAutoRTFMRuntimeEnabled)
 	{
 	    return static_cast<autortfm_result>(FContext::Get()->Transact(Work, Arg));
 	}
@@ -98,7 +99,7 @@ FORCENOINLINE extern "C" void autortfm_clear_transaction_status()
 
 FORCENOINLINE extern "C" bool autortfm_is_aborting()
 {
-	if(GAutoRTFMRuntimeEnabled)
+	if (GAutoRTFMRuntimeEnabled)
 	{
 		return FContext::Get()->IsAborting();
 	}
@@ -130,7 +131,7 @@ FORCENOINLINE extern "C" autortfm_status autortfm_close(void (*Work)(void* Arg),
 {
 	autortfm_status Result = autortfm_status_ontrack;
 
-	if(GAutoRTFMRuntimeEnabled)
+	if (GAutoRTFMRuntimeEnabled)
 	{
 		UE_CLOG(!FContext::IsTransactional(), LogAutoRTFM, Fatal, TEXT("Close called from an outside a transaction."));
 
