Invalid material shader map due to image layout changed in CL 25946381

diff --git a/Engine/Source/Runtime/Core/Private/UObject/FortniteShaderworkObjectVersion.cpp b/Engine/Source/Runtime/Core/Private/UObject/FortniteShaderworkObjectVersion.cpp
--- a/Engine/Source/Runtime/Core/Private/UObject/FortniteShaderworkObjectVersion.cpp
+++ b/Engine/Source/Runtime/Core/Private/UObject/FortniteShaderworkObjectVersion.cpp
@@ -8,5 +8,6 @@ TMap<FGuid, FGuid> FFortniteShaderworkObjectVersion::GetSystemGuids()
 
 	SystemGuids.Add(DevGuids.NANITE_DERIVEDDATA_VER, FGuid("64D2BA4FF333408E90CDD22CA38B1E33"));
 	SystemGuids.Add(DevGuids.STATICMESH_DERIVEDDATA_VER, FGuid("EB4D2761094040DD889CC1BE3D24E1B3"));
+	SystemGuids.Add(DevGuids.MATERIALSHADERMAP_DERIVEDDATA_VER, FGuid("0820372747734716A3BAA8C0F6A40ABA"));
 	return SystemGuids;
 }
