Unreal maths: Add missing constexpr
#rnx
#p4v-preflight-copy 25963304

diff --git a/Engine/Source/Runtime/Core/Public/GenericPlatform/GenericPlatformMath.h b/Engine/Source/Runtime/Core/Public/GenericPlatform/GenericPlatformMath.h
--- a/Engine/Source/Runtime/Core/Public/GenericPlatform/GenericPlatformMath.h
+++ b/Engine/Source/Runtime/Core/Public/GenericPlatform/GenericPlatformMath.h
@@ -28,15 +28,15 @@ struct FGenericPlatformMath
 {
 	// load half (F16) to float
 	//https://gist.github.com/rygorous/2156668
-	static FORCEINLINE float LoadHalf(const uint16* Ptr)
+	static constexpr FORCEINLINE float LoadHalf(const uint16* Ptr)
 	{
 		uint16 FP16 = *Ptr;
-		uint32 shifted_exp = 0x7c00 << 13;			// exponent mask after shift
+		constexpr uint32 shifted_exp = 0x7c00 << 13;			// exponent mask after shift
 		union FP32T
 		{
 			uint32 u;
 			float f;		
-		} FP32, magic = { 113 << 23 };
+		} FP32 = {}, magic = { 113 << 23 };
 
 		FP32.u = (FP16 & 0x7fff) << 13;				// exponent/mantissa bits
 		uint32 exp = shifted_exp & FP32.u;			// just the exponent
@@ -62,7 +62,7 @@ struct FGenericPlatformMath
 	// values too large for F16 are stored as +-Inf
 	// https://gist.github.com/rygorous/2156668
 	// float_to_half_fast3_rtne
-	static FORCEINLINE void StoreHalf(uint16* Ptr, float Value)
+	static constexpr FORCEINLINE void StoreHalf(uint16* Ptr, float Value)
 	{
 		union FP32T
 		{
@@ -73,10 +73,10 @@ struct FGenericPlatformMath
 
 		FP32.f = Value;
 
-		FP32T f32infty = { uint32(255 << 23) };
-		FP32T f16max = { uint32(127 + 16) << 23 };
-		FP32T denorm_magic = { (uint32(127 - 15) + uint32(23 - 10) + 1) << 23 };
-		uint32 sign_mask = 0x80000000u;
+		constexpr FP32T f32infty = { uint32(255 << 23) };
+		constexpr FP32T f16max = { uint32(127 + 16) << 23 };
+		constexpr FP32T denorm_magic = { (uint32(127 - 15) + uint32(23 - 10) + 1) << 23 };
+		constexpr uint32 sign_mask = 0x80000000u;
 
 		uint32 sign = FP32.u & sign_mask;
 		FP32.u ^= sign;
@@ -119,7 +119,7 @@ struct FGenericPlatformMath
 		*Ptr = FP16;
 	}
 
-	static FORCEINLINE void VectorLoadHalf(float* RESTRICT Dst, const uint16* RESTRICT Src)
+	static constexpr FORCEINLINE void VectorLoadHalf(float* RESTRICT Dst, const uint16* RESTRICT Src)
 	{
 		Dst[0] = LoadHalf(&Src[0]);
 		Dst[1] = LoadHalf(&Src[1]);
@@ -127,7 +127,7 @@ struct FGenericPlatformMath
 		Dst[3] = LoadHalf(&Src[3]);
 	}
 
-	static FORCEINLINE void VectorStoreHalf(uint16* RESTRICT Dst, const float* RESTRICT Src)
+	static constexpr FORCEINLINE void VectorStoreHalf(uint16* RESTRICT Dst, const float* RESTRICT Src)
 	{
 		StoreHalf(&Dst[0], Src[0]);
 		StoreHalf(&Dst[1], Src[1]);
@@ -135,13 +135,13 @@ struct FGenericPlatformMath
 		StoreHalf(&Dst[3], Src[3]);
 	}
 
-	static FORCEINLINE void WideVectorLoadHalf(float* RESTRICT Dst, const uint16* RESTRICT Src)
+	static constexpr FORCEINLINE void WideVectorLoadHalf(float* RESTRICT Dst, const uint16* RESTRICT Src)
 	{
 		VectorLoadHalf(Dst, Src);
 		VectorLoadHalf(Dst + 4, Src + 4);
 	}
 
-	static FORCEINLINE void WideVectorStoreHalf(uint16* RESTRICT Dst, const float* RESTRICT Src)
+	static constexpr FORCEINLINE void WideVectorStoreHalf(uint16* RESTRICT Dst, const float* RESTRICT Src)
 	{
 		VectorStoreHalf(Dst, Src);
 		VectorStoreHalf(Dst + 4, Src + 4);
@@ -205,21 +205,21 @@ struct FGenericPlatformMath
 	 * @param F		Floating point value to convert
 	 * @return		Truncated integer.
 	 */
-	static CONSTEXPR FORCEINLINE int32 TruncToInt32(float F)
+	static constexpr FORCEINLINE int32 TruncToInt32(float F)
 	{
 		return (int32)F;
 	}
-	static CONSTEXPR FORCEINLINE int32 TruncToInt32(double F)
+	static constexpr FORCEINLINE int32 TruncToInt32(double F)
 	{
 		return (int32)F;
 	}
-	static CONSTEXPR FORCEINLINE int64 TruncToInt64(double F)
+	static constexpr FORCEINLINE int64 TruncToInt64(double F)
 	{
 		return (int64)F;
 	}
 
-	static CONSTEXPR FORCEINLINE int32 TruncToInt(float F) { return TruncToInt32(F); }
-	static CONSTEXPR FORCEINLINE int64 TruncToInt(double F) { return TruncToInt64(F); }
+	static constexpr FORCEINLINE int32 TruncToInt(float F) { return TruncToInt32(F); }
+	static constexpr FORCEINLINE int64 TruncToInt(double F) { return TruncToInt64(F); }
 
 	/**
 	 * Converts a float to an integer value with truncation towards zero.
@@ -640,7 +640,7 @@ struct FGenericPlatformMath
 	 * @param Value		The value to compute the log of
 	 * @return			Log2 of Value. 0 if Value is 0.
 	 */	
-	static FORCEINLINE uint32 FloorLog2(uint32 Value) 
+	static constexpr FORCEINLINE uint32 FloorLog2(uint32 Value) 
 	{
 		uint32 pos = 0;
 		if (Value >= 1<<16) { Value >>= 16; pos += 16; }
@@ -658,7 +658,7 @@ struct FGenericPlatformMath
 	 * @param Value		The value to compute the log of
 	 * @return			Log2 of Value. 0 if Value is 0.
 	 */	
-	static FORCEINLINE uint64 FloorLog2_64(uint64 Value) 
+	static constexpr FORCEINLINE uint64 FloorLog2_64(uint64 Value) 
 	{
 		uint64 pos = 0;
 		if (Value >= 1ull<<32) { Value >>= 32; pos += 32; }
@@ -677,7 +677,7 @@ struct FGenericPlatformMath
 	 *
 	 * @return the number of zeros before the first "on" bit
 	 */
-	static FORCEINLINE uint8 CountLeadingZeros8(uint8 Value)
+	static constexpr FORCEINLINE uint8 CountLeadingZeros8(uint8 Value)
 	{
 		if (Value == 0) return 8;
 		return uint8(7 - FloorLog2(uint32(Value)));
@@ -690,7 +690,7 @@ struct FGenericPlatformMath
 	 *
 	 * @return the number of zeros before the first "on" bit
 	 */
-	static FORCEINLINE uint32 CountLeadingZeros(uint32 Value)
+	static constexpr FORCEINLINE uint32 CountLeadingZeros(uint32 Value)
 	{
 		if (Value == 0) return 32;
 		return 31 - FloorLog2(Value);
@@ -703,7 +703,7 @@ struct FGenericPlatformMath
 	 *
 	 * @return the number of zeros before the first "on" bit
 	 */
-	static FORCEINLINE uint64 CountLeadingZeros64(uint64 Value)
+	static constexpr FORCEINLINE uint64 CountLeadingZeros64(uint64 Value)
 	{
 		if (Value == 0) return 64;
 		return 63 - FloorLog2_64(Value);
@@ -716,7 +716,7 @@ struct FGenericPlatformMath
 	 *
 	 * @return the number of zeros after the last "on" bit
 	 */
-	static FORCEINLINE uint32 CountTrailingZeros(uint32 Value)
+	static constexpr FORCEINLINE uint32 CountTrailingZeros(uint32 Value)
 	{
 		if (Value == 0)
 		{
@@ -738,7 +738,7 @@ struct FGenericPlatformMath
 	 *
 	 * @return the number of zeros after the last "on" bit
 	 */
-	static FORCEINLINE uint64 CountTrailingZeros64(uint64 Value)
+	static constexpr FORCEINLINE uint64 CountTrailingZeros64(uint64 Value)
 	{
 		if (Value == 0)
 		{
@@ -757,14 +757,14 @@ struct FGenericPlatformMath
 	 * Returns smallest N such that (1<<N)>=Arg.
 	 * Note: CeilLogTwo(0)=0 
 	 */
-	static FORCEINLINE uint32 CeilLogTwo( uint32 Arg )
+	static constexpr FORCEINLINE uint32 CeilLogTwo( uint32 Arg )
 	{
 		// if Arg is 0, change it to 1 so that we return 0
 		Arg = Arg ? Arg : 1;
 		return 32 - CountLeadingZeros(Arg - 1);
 	}
 
-	static FORCEINLINE uint64 CeilLogTwo64( uint64 Arg )
+	static constexpr FORCEINLINE uint64 CeilLogTwo64( uint64 Arg )
 	{
 		// if Arg is 0, change it to 1 so that we return 0
 		Arg = Arg ? Arg : 1;
@@ -775,24 +775,24 @@ struct FGenericPlatformMath
 	 * Returns the smallest N such that (1<<N)>=Arg. This is a less efficient version of CeilLogTwo, but written in a
 	 * way that can be evaluated at compile-time.
 	 */
-	static FORCEINLINE constexpr uint8 ConstExprCeilLogTwo(SIZE_T Arg)
+	static constexpr FORCEINLINE uint8 ConstExprCeilLogTwo(SIZE_T Arg)
 	{
 		return Arg <= 1 ? 0 : (1 + ConstExprCeilLogTwo(Arg / 2));
 	}
 
 	/** @return Rounds the given number up to the next highest power of two. */
-	static FORCEINLINE uint32 RoundUpToPowerOfTwo(uint32 Arg)
+	static constexpr FORCEINLINE uint32 RoundUpToPowerOfTwo(uint32 Arg)
 	{
 		return 1 << CeilLogTwo(Arg);
 	}
 
-	static FORCEINLINE uint64 RoundUpToPowerOfTwo64(uint64 V)
+	static constexpr FORCEINLINE uint64 RoundUpToPowerOfTwo64(uint64 V)
 	{
 		return uint64(1) << CeilLogTwo64(V);
 	}
 
 	/** Spreads bits to every other. */
-	static FORCEINLINE uint32 MortonCode2( uint32 x )
+	static constexpr FORCEINLINE uint32 MortonCode2( uint32 x )
 	{
 		x &= 0x0000ffff;
 		x = (x ^ (x << 8)) & 0x00ff00ff;
@@ -802,7 +802,7 @@ struct FGenericPlatformMath
 		return x;
 	}
 
-	static FORCEINLINE uint64 MortonCode2_64( uint64 x )
+	static constexpr FORCEINLINE uint64 MortonCode2_64( uint64 x )
 	{
 		x &= 0x00000000ffffffff;
 		x = (x ^ (x << 16)) & 0x0000ffff0000ffff;
@@ -814,7 +814,7 @@ struct FGenericPlatformMath
 	}
 
 	/** Reverses MortonCode2. Compacts every other bit to the right. */
-	static FORCEINLINE uint32 ReverseMortonCode2( uint32 x )
+	static constexpr FORCEINLINE uint32 ReverseMortonCode2( uint32 x )
 	{
 		x &= 0x55555555;
 		x = (x ^ (x >> 1)) & 0x33333333;
@@ -824,7 +824,7 @@ struct FGenericPlatformMath
 		return x;
 	}
 
-	static FORCEINLINE uint64 ReverseMortonCode2_64( uint64 x )
+	static constexpr FORCEINLINE uint64 ReverseMortonCode2_64( uint64 x )
 	{
 		x &= 0x5555555555555555;
 		x = (x ^ (x >> 1)) & 0x3333333333333333;
@@ -836,7 +836,7 @@ struct FGenericPlatformMath
 	}
 
 	/** Spreads bits to every 3rd. */
-	static FORCEINLINE uint32 MortonCode3( uint32 x )
+	static constexpr FORCEINLINE uint32 MortonCode3( uint32 x )
 	{
 		x &= 0x000003ff;
 		x = (x ^ (x << 16)) & 0xff0000ff;
@@ -847,7 +847,7 @@ struct FGenericPlatformMath
 	}
 
 	/** Reverses MortonCode3. Compacts every 3rd bit to the right. */
-	static FORCEINLINE uint32 ReverseMortonCode3( uint32 x )
+	static constexpr FORCEINLINE uint32 ReverseMortonCode3( uint32 x )
 	{
 		x &= 0x09249249;
 		x = (x ^ (x >>  2)) & 0x030c30c3;
@@ -871,7 +871,7 @@ struct FGenericPlatformMath
 	 *
 	 * @return	ValueGEZero if Comparand >= 0, ValueLTZero otherwise
 	 */
-	static CONSTEXPR FORCEINLINE float FloatSelect( float Comparand, float ValueGEZero, float ValueLTZero )
+	static constexpr FORCEINLINE float FloatSelect( float Comparand, float ValueGEZero, float ValueLTZero )
 	{
 		return Comparand >= 0.f ? ValueGEZero : ValueLTZero;
 	}
@@ -890,35 +890,35 @@ struct FGenericPlatformMath
 	 *
 	 * @return	ValueGEZero if Comparand >= 0, ValueLTZero otherwise
 	 */
-	static CONSTEXPR FORCEINLINE double FloatSelect( double Comparand, double ValueGEZero, double ValueLTZero )
+	static constexpr FORCEINLINE double FloatSelect( double Comparand, double ValueGEZero, double ValueLTZero )
 	{
 		return Comparand >= 0.0 ? ValueGEZero : ValueLTZero;
 	}
 
 	/** Computes absolute value in a generic way */
 	template< class T > 
-	static CONSTEXPR FORCEINLINE T Abs( const T A )
+	static constexpr FORCEINLINE T Abs( const T A )
 	{
 		return (A < (T)0) ? -A : A;
 	}
 
 	/** Returns 1, 0, or -1 depending on relation of T to 0 */
 	template< class T > 
-	static CONSTEXPR FORCEINLINE T Sign( const T A )
+	static constexpr FORCEINLINE T Sign( const T A )
 	{
 		return (A > (T)0) ? (T)1 : ((A < (T)0) ? (T)-1 : (T)0);
 	}
 
 	/** Returns higher value in a generic way */
 	template< class T > 
-	static CONSTEXPR FORCEINLINE T Max( const T A, const T B )
+	static constexpr FORCEINLINE T Max( const T A, const T B )
 	{
 		return (B < A) ? A : B;
 	}
 
 	/** Returns lower value in a generic way */
 	template< class T > 
-	static CONSTEXPR FORCEINLINE T Min( const T A, const T B )
+	static constexpr FORCEINLINE T Min( const T A, const T B )
 	{
 		return (A < B) ? A : B;
 	}
@@ -1005,7 +1005,7 @@ struct FGenericPlatformMath
 		return CurMax;
 	}
 
-	static FORCEINLINE int32 CountBits(uint64 Bits)
+	static constexpr FORCEINLINE int32 CountBits(uint64 Bits)
 	{
 		// https://en.wikipedia.org/wiki/Hamming_weight
 		Bits -= (Bits >> 1) & 0x5555555555555555ull;
diff --git a/Engine/Source/Runtime/Core/Public/Math/UnrealMathUtility.h b/Engine/Source/Runtime/Core/Public/Math/UnrealMathUtility.h
--- a/Engine/Source/Runtime/Core/Public/Math/UnrealMathUtility.h
+++ b/Engine/Source/Runtime/Core/Public/Math/UnrealMathUtility.h
@@ -776,14 +776,14 @@ public:
 	UE_NODISCARD static FORCEINLINE float FastAsin(float Value)
 	{
 		// Clamp input to [-1,1].
-		bool nonnegative = (Value >= 0.0f);
-		float x = FMath::Abs(Value);
+		const bool nonnegative = (Value >= 0.0f);
+		const float x = FMath::Abs(Value);
 		float omx = 1.0f - x;
 		if (omx < 0.0f)
 		{
 			omx = 0.0f;
 		}
-		float root = FMath::Sqrt(omx);
+		const float root = FMath::Sqrt(omx);
 		// 7-degree minimax approximation
 		float result = ((((((-0.0012624911f * x + 0.0066700901f) * x - 0.0170881256f) * x + 0.0308918810f) * x - 0.0501743046f) * x + 0.0889789874f) * x - 0.2145988016f) * x + FASTASIN_HALF_PI;
 		result *= root;  // acos(|x|)
@@ -811,8 +811,8 @@ public:
 		return RadVal * (180.f / UE_PI);
 	}
 
-	static FORCEINLINE float RadiansToDegrees(float const& RadVal) { return RadVal * (180.f / UE_PI); }
-	static FORCEINLINE double RadiansToDegrees(double const& RadVal) { return RadVal * (180.0 / UE_DOUBLE_PI); }
+	static constexpr FORCEINLINE float RadiansToDegrees(float const& RadVal) { return RadVal * (180.f / UE_PI); }
+	static constexpr FORCEINLINE double RadiansToDegrees(double const& RadVal) { return RadVal * (180.0 / UE_DOUBLE_PI); }
 
 	/** 
 	 * Converts degrees to radians.
@@ -825,8 +825,8 @@ public:
 		return DegVal * (UE_PI / 180.f);
 	}
 
-	static FORCEINLINE float DegreesToRadians(float const& DegVal) { return DegVal * (UE_PI / 180.f); }
-	static FORCEINLINE double DegreesToRadians(double const& DegVal) { return DegVal * (UE_DOUBLE_PI / 180.0); }
+	static constexpr FORCEINLINE float DegreesToRadians(float const& DegVal) { return DegVal * (UE_PI / 180.f); }
+	static constexpr FORCEINLINE double DegreesToRadians(double const& DegVal) { return DegVal * (UE_DOUBLE_PI / 180.0); }
 
 	/** 
 	 * Clamps an arbitrary angle to be between the given angles.  Will clamp to nearest boundary.
