Create an AutoRTFM API to run a Transact callback in the open

This allows the VerseVM module to be compiled without autortfm instrumentation.

The API is exactly like Transact([&] { Open([&] { Functor(); }); }) but it allows the closed code to live in Core so that the module where Functor is implemented doesn't need to be compiled with autortfm instrumentation.

diff --git a/Engine/Source/Runtime/Core/Private/AutoRTFM/API.cpp b/Engine/Source/Runtime/Core/Private/AutoRTFM/API.cpp
--- a/Engine/Source/Runtime/Core/Private/AutoRTFM/API.cpp
+++ b/Engine/Source/Runtime/Core/Private/AutoRTFM/API.cpp
@@ -34,6 +34,18 @@ static constexpr bool GAutoRTFMRuntimeEnabled = false;
 namespace AutoRTFM
 {
 
+UE_AUTORTFM_FORCEINLINE autortfm_result TransactThenOpenImpl(void (*Work)(void* Arg), void* Arg)
+{
+	return static_cast<autortfm_result>(
+		AutoRTFM::Transact([&]
+		{
+			AutoRTFM::Open([&]
+			{
+				Work(Arg);
+			});
+		}));
+}
+
 FORCENOINLINE extern "C" bool autortfm_is_transactional()
 {
 	if (GAutoRTFMRuntimeEnabled)
@@ -61,6 +73,11 @@ FORCENOINLINE extern "C" autortfm_result autortfm_transact(void (*Work)(void* Ar
 	return autortfm_committed;
 }
 
+FORCENOINLINE extern "C" autortfm_result autortfm_transact_then_open(void (*Work)(void* Arg), void* Arg)
+{
+    return TransactThenOpenImpl(Work, Arg);
+}
+
 FORCENOINLINE extern "C" void autortfm_commit(void (*Work)(void* Arg), void* Arg)
 {
     autortfm_result Result = autortfm_transact(Work, Arg);
@@ -260,6 +277,12 @@ autortfm_result RTFM_autortfm_transact(void (*Work)(void* Arg), void* Arg, FCont
 }
 UE_AUTORTFM_REGISTER_OPEN_FUNCTION(autortfm_transact);
 
+autortfm_result RTFM_autortfm_transact_then_open(void (*Work)(void* Arg), void* Arg, FContext*)
+{
+    return TransactThenOpenImpl(Work, Arg);
+}
+UE_AUTORTFM_REGISTER_OPEN_FUNCTION(autortfm_transact_then_open);
+
 void RTFM_autortfm_commit(void (*Work)(void* Arg), void* Arg, FContext*)
 {
     autortfm_result Result = autortfm_transact(Work, Arg);
diff --git a/Engine/Source/Runtime/Core/Public/AutoRTFM/AutoRTFM.h b/Engine/Source/Runtime/Core/Public/AutoRTFM/AutoRTFM.h
--- a/Engine/Source/Runtime/Core/Public/AutoRTFM/AutoRTFM.h
+++ b/Engine/Source/Runtime/Core/Public/AutoRTFM/AutoRTFM.h
@@ -106,12 +106,18 @@ UE_AUTORTFM_FORCEINLINE bool autortfm_is_closed(void)
 // if this nested transaction succeeds.
 #if UE_AUTORTFM
 autortfm_result autortfm_transact(void (*work)(void* arg), void* arg);
+autortfm_result autortfm_transact_then_open(void (*work)(void* arg), void* arg);
 #else
 UE_AUTORTFM_FORCEINLINE autortfm_result autortfm_transact(void (*work)(void* arg), void* arg)
 {
 	work(arg);
     return autortfm_committed;
 }
+UE_AUTORTFM_FORCEINLINE autortfm_result autortfm_transact_then_open(void (*work)(void* arg), void* arg)
+{
+	work(arg);
+    return autortfm_committed;
+}
 #endif
 
 // Run the callback in a transaction like autortfm_transact, but abort program
@@ -362,6 +368,18 @@ UE_AUTORTFM_FORCEINLINE ETransactionResult Transact(const TFunctor& Functor)
         const_cast<void*>(static_cast<const void*>(&Functor))));
 }
 
+// This is just like calling Transact([&] { Open([&] { Functor(); }); });  
+// The reason we expose it is that it allows the caller's module to not
+// be compiled with the AutoRTFM instrumentation of functions if the only
+// thing that's being invoked is a function in the open.
+template<typename TFunctor>
+UE_AUTORTFM_FORCEINLINE ETransactionResult TransactThenOpen(const TFunctor& Functor)
+{
+	return static_cast<ETransactionResult>(autortfm_transact_then_open(
+		[] (void* Arg) { (*static_cast<const TFunctor*>(Arg))(); },
+		const_cast<void*>(static_cast<const void*>(&Functor))));
+}
+
 template<typename TFunctor>
 UE_AUTORTFM_FORCEINLINE void Commit(const TFunctor& Functor)
 {
