Tasks System: simplified the control flow of task execution, this removes 2 function calls for every task execution while doing the same

diff --git a/Engine/Source/Runtime/Core/Private/Tasks/Pipe.cpp b/Engine/Source/Runtime/Core/Private/Tasks/Pipe.cpp
--- a/Engine/Source/Runtime/Core/Private/Tasks/Pipe.cpp
+++ b/Engine/Source/Runtime/Core/Private/Tasks/Pipe.cpp
@@ -59,10 +59,9 @@ namespace UE::Tasks
 				);
 			}
 
-			virtual bool TryExecuteTaskVirtual() final
+			virtual void ExecuteTask() override final
 			{
 				checkNoEntry(); // the method won't be called because the task was initialized with `EExtendedTaskPriority::TaskEvent`
-				return true;
 			}
 		};
 
diff --git a/Engine/Source/Runtime/Core/Public/Async/TaskGraphInterfaces.h b/Engine/Source/Runtime/Core/Public/Async/TaskGraphInterfaces.h
--- a/Engine/Source/Runtime/Core/Public/Async/TaskGraphInterfaces.h
+++ b/Engine/Source/Runtime/Core/Public/Async/TaskGraphInterfaces.h
@@ -648,20 +648,14 @@ private:
 		FBaseGraphTask::Init(TEXT("GraphTask"), InPriority, InExtendedPriority);
 	}
 
-	virtual bool TryExecuteTaskVirtual() final
+	virtual void ExecuteTask() override final
 	{
-		return TryExecute(
-			[](UE::Tasks::Private::FTaskBase& Task)
-			{
-				TGraphTask& This = static_cast<TGraphTask&>(Task);
-				FGraphEventRef GraphEventRef{ &This };
-				TTask* TaskObject = This.TaskStorage.GetTypedPtr();
-				ENamedThreads::Type ThreadIndex = ENamedThreads::GetThreadIndex(TaskObject->GetDesiredThread());
+		FGraphEventRef GraphEventRef{ this };
+		TTask* TaskObject = TaskStorage.GetTypedPtr();
+		ENamedThreads::Type ThreadIndex = ENamedThreads::GetThreadIndex(TaskObject->GetDesiredThread());
 
-				TaskObject->DoTask(ThreadIndex, GraphEventRef);
-				DestructItem(TaskObject);
-			}
-		);
+		TaskObject->DoTask(ThreadIndex, GraphEventRef);
+		DestructItem(TaskObject);
 	}
 
 private:
@@ -695,10 +689,9 @@ public:
 	static void operator delete(void* Ptr);
 
 private:
-	virtual bool TryExecuteTaskVirtual() override
+	virtual void ExecuteTask() override final
 	{
 		checkNoEntry(); // graph events are never executed
-		return true;
 	}
 };
 
diff --git a/Engine/Source/Runtime/Core/Public/Tasks/Task.h b/Engine/Source/Runtime/Core/Public/Tasks/Task.h
--- a/Engine/Source/Runtime/Core/Public/Tasks/Task.h
+++ b/Engine/Source/Runtime/Core/Public/Tasks/Task.h
@@ -181,7 +181,7 @@ namespace UE::Tasks
 		{
 			check(IsValid());
 			FTaskHandle::Wait();
-			return ((Private::TTaskWithResult<ResultType>*)Pimpl.GetReference())->GetResult();
+			return static_cast<Private::TTaskWithResult<ResultType>*>(Pimpl.GetReference())->GetResult();
 		}
 
 	private:
diff --git a/Engine/Source/Runtime/Core/Public/Tasks/TaskPrivate.h b/Engine/Source/Runtime/Core/Public/Tasks/TaskPrivate.h
--- a/Engine/Source/Runtime/Core/Public/Tasks/TaskPrivate.h
+++ b/Engine/Source/Runtime/Core/Public/Tasks/TaskPrivate.h
@@ -168,9 +168,7 @@ namespace UE::Tasks
 				TaskTrace::Destroyed(GetTraceId());
 			}
 
-			// will be called to execute the task, must be implemented by a derived class that should call `FTaskBase::TryExecute` and pass the task body
-			// @see TExecutableTask::TryExecuteTaskVirtual
-			virtual bool TryExecuteTaskVirtual() = 0;
+			virtual void ExecuteTask() = 0;
 
 		public:
 			// returns true if it's valid to wait for the task completion.
@@ -409,12 +407,10 @@ namespace UE::Tasks
 			}
 
 		protected:
-			using FTaskBodyType = void(*)(FTaskBase&);
-
 			// tries to get execution permission and if successful, executes given task body and completes the task if there're no pending nested tasks. 
 			// does all required accounting before/after task execution. the task can be deleted as a result of this call.
 			// @returns true if the task was executed by the current thread
-			bool TryExecute(FTaskBodyType TaskBody)
+			bool TryExecuteTask()
 			{
 				if (!TrySetExecutionFlag())
 				{
@@ -435,8 +431,9 @@ namespace UE::Tasks
 				}
 
 				{
+					UE::FInheritedContextScope InheritedContextScope = RestoreInheritedContext();
 					TaskTrace::FTaskTimingEventScope TaskEventScope(GetTraceId());
-					TaskBody(*this);
+					ExecuteTask();
 				}
 
 				if (GetPipe() != nullptr)
@@ -487,12 +484,6 @@ namespace UE::Tasks
 
 			CORE_API void ClearPipe();
 
-			FORCEINLINE bool TryExecuteTask()
-			{
-				UE::FInheritedContextScope InheritedContextScope = RestoreInheritedContext();
-				return TryExecuteTaskVirtual();
-			}
-
 		private:
 			// A task can be locked for execution (by prerequisites or if it's not launched yet) or for completion (by nested tasks).
 			// This method is called to unlock the task and so can result in its scheduling (and execution) or completion
@@ -669,32 +660,24 @@ namespace UE::Tasks
 		};
 
 		// Task implementation that can be executed, as it stores task body. Generic version (for tasks that return non-void results).
-		// In most cases it should be allocated on the heap and used with TRefCountPtr, e.g. @see FTaskHandle. With care, can be allocated on the stack, e.g. see 
-		// WaitingTask in FTaskBase::Wait().
-		// Implements memory allocation from a pooled fixed-size allocator tuned for the everage UE task size
+		// In most cases it should be allocated on the heap and used with TRefCountPtr, e.g. @see FTaskHandle. 
 		template<typename TaskBodyType, typename ResultType = TInvokeResult_T<TaskBodyType>, typename Enable = void>
 		class TExecutableTaskBase : public TTaskWithResult<ResultType>
 		{
 			UE_NONCOPYABLE(TExecutableTaskBase);
 
 		public:
-			virtual bool TryExecuteTaskVirtual() final
+			virtual void ExecuteTask() override final
 			{
-				return FTaskBase::TryExecute(
-					[](FTaskBase& Task)
-					{
-						TExecutableTaskBase& This = static_cast<TExecutableTaskBase&>(Task);
-						new(&This.ResultStorage) ResultType{ Invoke(*This.TaskBodyStorage.GetTypedPtr()) };
+				new(&this->ResultStorage) ResultType{ Invoke(*TaskBodyStorage.GetTypedPtr()) };
 
-						// destroy the task body as soon as we are done with it, as it can have captured data sensitive to destruction order
-						DestructItem(This.TaskBodyStorage.GetTypedPtr());
-					}
-				);
+				// destroy the task body as soon as we are done with it, as it can have captured data sensitive to destruction order
+				DestructItem(TaskBodyStorage.GetTypedPtr());
 			}
 
 		protected:
-			TExecutableTaskBase(const TCHAR* InDebugName, TaskBodyType&& TaskBody, ETaskPriority InPriority, EExtendedTaskPriority InExtendedPriority) :
-				TTaskWithResult<ResultType>(InDebugName, InPriority, InExtendedPriority, 2)
+			TExecutableTaskBase(const TCHAR* InDebugName, TaskBodyType&& TaskBody, ETaskPriority InPriority, EExtendedTaskPriority InExtendedPriority) 
+				: TTaskWithResult<ResultType>(InDebugName, InPriority, InExtendedPriority, 2)
 				// 2 init refs: one for the initial reference (we don't increment it on passing to `TRefCountPtr`), and one for the internal 
 				// reference that keeps the task alive while it's in the system. is released either on task completion or by the scheduler after
 				// trying to execute the task
@@ -713,18 +696,12 @@ namespace UE::Tasks
 			UE_NONCOPYABLE(TExecutableTaskBase);
 
 		public:
-			virtual bool TryExecuteTaskVirtual() final
+			virtual void ExecuteTask() override final
 			{
-				return TryExecute(
-					[](FTaskBase& Task)
-					{
-						TExecutableTaskBase& This = static_cast<TExecutableTaskBase&>(Task);
-						Invoke(*This.TaskBodyStorage.GetTypedPtr());
+				Invoke(*TaskBodyStorage.GetTypedPtr());
 
-						// destroy the task body as soon as we are done with it, as it can have captured data sensitive to destruction order
-						DestructItem(This.TaskBodyStorage.GetTypedPtr());
-					}
-				);
+				// destroy the task body as soon as we are done with it, as it can have captured data sensitive to destruction order
+				DestructItem(TaskBodyStorage.GetTypedPtr());
 			}
 
 		protected:
@@ -799,10 +776,9 @@ namespace UE::Tasks
 				Init(InDebugName, ETaskPriority::Normal, EExtendedTaskPriority::TaskEvent);
 			}
 
-			virtual bool TryExecuteTaskVirtual() override
+			virtual void ExecuteTask() override final
 			{
 				checkNoEntry(); // never executed because it doesn't have a task body
-				return true;
 			}
 		};
 
