tasks integration: fixed waiting for tasks with timeout on named threads

diff --git a/Engine/Source/Runtime/Core/Private/Tasks/TaskPrivate.cpp b/Engine/Source/Runtime/Core/Private/Tasks/TaskPrivate.cpp
--- a/Engine/Source/Runtime/Core/Private/Tasks/TaskPrivate.cpp
+++ b/Engine/Source/Runtime/Core/Private/Tasks/TaskPrivate.cpp
@@ -191,7 +191,8 @@ namespace UE::Tasks
 			TRACE_CPUPROFILER_EVENT_SCOPE(Tasks::Wait);
 
 			// if we are on a named thread, handle waiting in TaskGraph-specific style
-			if (TryWaitOnNamedThread(*this))
+			// (named threads don't support waiting with timeout)
+			if (Timeout == FTimeout::Never() && TryWaitOnNamedThread(*this))
 			{
 				return true;
 			}
