Delegate perfect forwarding.

diff --git a/Engine/Source/Runtime/Core/Public/Delegates/DelegateInstancesImpl.h b/Engine/Source/Runtime/Core/Public/Delegates/DelegateInstancesImpl.h
--- a/Engine/Source/Runtime/Core/Public/Delegates/DelegateInstancesImpl.h
+++ b/Engine/Source/Runtime/Core/Public/Delegates/DelegateInstancesImpl.h
@@ -149,7 +149,7 @@ public:
 		checkSlow(IsSafeToExecute());
 
 		TPlacementNewer<FParmsWithPayload> PayloadAndParams;
-		this->Payload.ApplyAfter(PayloadAndParams, Params...);
+		this->Payload.ApplyAfter(PayloadAndParams, Forward<ParamTypes>(Params)...);
 		UserObjectPtr->ProcessEvent(CachedFunction, &PayloadAndParams);
 		return PayloadAndParams->GetResult();
 	}
@@ -161,7 +161,7 @@ public:
 			using FParmsWithPayload = TPayload<RetValType(typename TDecay<ParamTypes>::Type..., typename TDecay<VarTypes> ::Type...)>;
 
 			TPlacementNewer<FParmsWithPayload> PayloadAndParams;
-			this->Payload.ApplyAfter(PayloadAndParams, Params...);
+			this->Payload.ApplyAfter(PayloadAndParams, Forward<ParamTypes>(Params)...);
 			ActualUserObject->ProcessEvent(CachedFunction, &PayloadAndParams);
 			return true;
 		}
@@ -278,7 +278,7 @@ public:
 
 		checkSlow(MethodPtr != nullptr);
 
-		return this->Payload.ApplyAfter(MethodPtr, MutableUserObject, Params...);
+		return this->Payload.ApplyAfter(MethodPtr, MutableUserObject, Forward<ParamTypes>(Params)...);
 	}
 
 	bool ExecuteIfSafe(ParamTypes... Params) const final
@@ -295,7 +295,7 @@ public:
 
 			checkSlow(MethodPtr != nullptr);
 
-			(void)this->Payload.ApplyAfter(MethodPtr, MutableUserObject, Params...);
+			(void)this->Payload.ApplyAfter(MethodPtr, MutableUserObject, Forward<ParamTypes>(Params)...);
 
 			return true;
 		}
@@ -414,7 +414,7 @@ public:
 
 		checkSlow(MethodPtr != nullptr);
 
-		return this->Payload.ApplyAfter(MethodPtr, MutableUserObject, Params...);
+		return this->Payload.ApplyAfter(MethodPtr, MutableUserObject, Forward<ParamTypes>(Params)...);
 	}
 
 
@@ -429,7 +429,7 @@ public:
 
 		checkSlow(MethodPtr != nullptr);
 
-		(void)this->Payload.ApplyAfter(MethodPtr, MutableUserObject, Params...);
+		(void)this->Payload.ApplyAfter(MethodPtr, MutableUserObject, Forward<ParamTypes>(Params)...);
 
 		return true;
 	}
@@ -546,7 +546,7 @@ public:
 
 		checkSlow(MethodPtr != nullptr);
 
-		return this->Payload.ApplyAfter(MethodPtr, MutableUserObject, Params...);
+		return this->Payload.ApplyAfter(MethodPtr, MutableUserObject, Forward<ParamTypes>(Params)...);
 	}
 
 	bool ExecuteIfSafe(ParamTypes... Params) const final
@@ -562,7 +562,7 @@ public:
 
 			checkSlow(MethodPtr != nullptr);
 
-			(void)this->Payload.ApplyAfter(MethodPtr, MutableUserObject, Params...);
+			(void)this->Payload.ApplyAfter(MethodPtr, MutableUserObject, Forward<ParamTypes>(Params)...);
 
 			return true;
 		}
@@ -664,7 +664,7 @@ public:
 		// Call the static function
 		checkSlow(StaticFuncPtr != nullptr);
 
-		return this->Payload.ApplyAfter(StaticFuncPtr, Params...);
+		return this->Payload.ApplyAfter(StaticFuncPtr, Forward<ParamTypes>(Params)...);
 	}
 
 	bool ExecuteIfSafe(ParamTypes... Params) const final
@@ -672,7 +672,7 @@ public:
 		// Call the static function
 		checkSlow(StaticFuncPtr != nullptr);
 
-		(void)this->Payload.ApplyAfter(StaticFuncPtr, Params...);
+		(void)this->Payload.ApplyAfter(StaticFuncPtr, Forward<ParamTypes>(Params)...);
 
 		return true;
 	}
@@ -759,13 +759,13 @@ public:
 
 	RetValType Execute(ParamTypes... Params) const final
 	{
-		return this->Payload.ApplyAfter(Functor, Params...);
+		return this->Payload.ApplyAfter(Functor, Forward<ParamTypes>(Params)...);
 	}
 
 	bool ExecuteIfSafe(ParamTypes... Params) const final
 	{
 		// Functors are always considered safe to execute!
-		(void)this->Payload.ApplyAfter(Functor, Params...);
+		(void)this->Payload.ApplyAfter(Functor, Forward<ParamTypes>(Params)...);
 
 		return true;
 	}
@@ -858,14 +858,14 @@ public:
 
 	RetValType Execute(ParamTypes... Params) const final
 	{
-		return this->Payload.ApplyAfter(Functor, Params...);
+		return this->Payload.ApplyAfter(Functor, Forward<ParamTypes>(Params)...);
 	}
 
 	bool ExecuteIfSafe(ParamTypes... Params) const final
 	{
 		if (ContextObject.IsValid())
 		{
-			(void)this->Payload.ApplyAfter(Functor, Params...);
+			(void)this->Payload.ApplyAfter(Functor, Forward<ParamTypes>(Params)...);
 			return true;
 		}
 
diff --git a/Engine/Source/Runtime/Core/Public/Delegates/DelegateSignatureImpl.inl b/Engine/Source/Runtime/Core/Public/Delegates/DelegateSignatureImpl.inl
--- a/Engine/Source/Runtime/Core/Public/Delegates/DelegateSignatureImpl.inl
+++ b/Engine/Source/Runtime/Core/Public/Delegates/DelegateSignatureImpl.inl
@@ -512,7 +512,7 @@ public:
 		// Consider using ExecuteIfBound() instead.
 		checkSlow(LocalDelegateInstance != nullptr);
 
-		return LocalDelegateInstance->Execute(Params...);
+		return LocalDelegateInstance->Execute(Forward<ParamTypes>(Params)...);
 	}
 
 	/**
@@ -533,7 +533,7 @@ public:
 
 		if (const DelegateInstanceInterfaceType* Ptr = GetDelegateInstanceProtected())
 		{
-			return Ptr->ExecuteIfSafe(Params...);
+			return Ptr->ExecuteIfSafe(Forward<ParamTypes>(Params)...);
 		}
 
 		return false;
