Fixed EnumAddFlags and EnumRemoveFlags to not require operators on the enum, as how EnumHasAllFlags and EnumHasAnyFlags work.
Some refactoring to remove repeated typing and extraneous brackets.

diff --git a/Engine/Source/Runtime/Core/Public/Misc/EnumClassFlags.h b/Engine/Source/Runtime/Core/Public/Misc/EnumClassFlags.h
--- a/Engine/Source/Runtime/Core/Public/Misc/EnumClassFlags.h
+++ b/Engine/Source/Runtime/Core/Public/Misc/EnumClassFlags.h
@@ -27,23 +27,27 @@
 template<typename Enum>
 constexpr bool EnumHasAllFlags(Enum Flags, Enum Contains)
 {
-	return ( ( ( __underlying_type(Enum) )Flags ) & ( __underlying_type(Enum) )Contains ) == ( ( __underlying_type(Enum) )Contains );
+	using UnderlyingType = __underlying_type(Enum);
+	return ((UnderlyingType)Flags & (UnderlyingType)Contains) == (UnderlyingType)Contains;
 }
 
 template<typename Enum>
 constexpr bool EnumHasAnyFlags(Enum Flags, Enum Contains)
 {
-	return ( ( ( __underlying_type(Enum) )Flags ) & ( __underlying_type(Enum) )Contains ) != 0;
+	using UnderlyingType = __underlying_type(Enum);
+	return ((UnderlyingType)Flags & (UnderlyingType)Contains) != 0;
 }
 
 template<typename Enum>
 void EnumAddFlags(Enum& Flags, Enum FlagsToAdd)
 {
-	Flags |= FlagsToAdd;
+	using UnderlyingType = __underlying_type(Enum);
+	Flags = (Enum)((UnderlyingType)Flags | (UnderlyingType)FlagsToAdd);
 }
 
 template<typename Enum>
 void EnumRemoveFlags(Enum& Flags, Enum FlagsToRemove)
 {
-	Flags &= ~FlagsToRemove;
+	using UnderlyingType = __underlying_type(Enum);
+	Flags = (Enum)((UnderlyingType)Flags & ~(UnderlyingType)FlagsToRemove);
 }
\ No newline at end of file
