Attempt to work around ambiguous overloaded Swap problem.

diff --git a/Engine/Source/Runtime/Core/Public/Templates/SharedPointer.h b/Engine/Source/Runtime/Core/Public/Templates/SharedPointer.h
--- a/Engine/Source/Runtime/Core/Public/Templates/SharedPointer.h
+++ b/Engine/Source/Runtime/Core/Public/Templates/SharedPointer.h
@@ -387,13 +387,13 @@ public:
 	FORCEINLINE TSharedRef& operator=( TSharedRef const& InSharedRef )
 	{
 		TSharedRef Temp = InSharedRef;
-		Swap(Temp, *this);
+		::Swap(Temp, *this);
 		return *this;
 	}
 
 	FORCEINLINE TSharedRef& operator=( TSharedRef&& InSharedRef )
 	{
-		Swap(*this, InSharedRef);
+		::Swap(*this, InSharedRef);
 		return *this;
 	}
 
@@ -979,7 +979,7 @@ public:
 	FORCEINLINE TSharedPtr& operator=( TSharedPtr const& InSharedPtr )
 	{
 		TSharedPtr Temp = InSharedPtr;
-		Swap(Temp, *this);
+		::Swap(Temp, *this);
 		return *this;
 	}
 
@@ -1411,7 +1411,7 @@ public:
 	FORCEINLINE TWeakPtr& operator=( TWeakPtr const& InWeakPtr )
 	{
 		TWeakPtr Temp = InWeakPtr;
-		Swap(Temp, *this);
+		::Swap(Temp, *this);
 		return *this;
 	}
 
