FArchiveMD5 - Calling Reset() on an MD5 archive should reset the internal state

diff --git a/Engine/Source/Runtime/Core/Public/Misc/ArchiveMD5.h b/Engine/Source/Runtime/Core/Public/Misc/ArchiveMD5.h
--- a/Engine/Source/Runtime/Core/Public/Misc/ArchiveMD5.h
+++ b/Engine/Source/Runtime/Core/Public/Misc/ArchiveMD5.h
@@ -45,6 +45,12 @@ public:
 		return *this;
 	}
 
+	virtual void Reset() override
+	{
+		FArchive::Reset();
+		MD5 = FMD5();
+	}
+
 	void GetHash(FMD5Hash& Hash)
 	{
 		Hash.Set(MD5);
