Fixing missing pixel formats from the FOREACH_ENUM_EPIXELFORMAT

- PF_R9G9B9EXP5 and PF_P010 were missing

diff --git a/Engine/Source/Runtime/Core/Public/PixelFormat.h b/Engine/Source/Runtime/Core/Public/PixelFormat.h
--- a/Engine/Source/Runtime/Core/Public/PixelFormat.h
+++ b/Engine/Source/Runtime/Core/Public/PixelFormat.h
@@ -194,6 +194,8 @@ enum EPixelFormat : uint8
 	op(PF_R32G32B32F) \
 	op(PF_R8_SINT) \
 	op(PF_R64_UINT) \
+	op(PF_R9G9B9EXP5) \
+	op(PF_P010) \
 	op(PF_ASTC_4x4_NORM_RG) \
 	op(PF_ASTC_6x6_NORM_RG) \
 	op(PF_ASTC_8x8_NORM_RG) \
