[Core]
* Moved dllexport from type to methods/staticvar. This improves compile times, memory and performance in dll builds

diff --git a/Engine/Source/Runtime/Core/Private/Experimental/Containers/HazardPointer.cpp b/Engine/Source/Runtime/Core/Private/Experimental/Containers/HazardPointer.cpp
--- a/Engine/Source/Runtime/Core/Private/Experimental/Containers/HazardPointer.cpp
+++ b/Engine/Source/Runtime/Core/Private/Experimental/Containers/HazardPointer.cpp
@@ -21,7 +21,7 @@ FHazardPointerCollection::~FHazardPointerCollection()
 }
 
 template<bool Cached>
-CORE_API FHazardPointerCollection::FHazardRecord* FHazardPointerCollection::Grow()
+FHazardPointerCollection::FHazardRecord* FHazardPointerCollection::Grow()
 {
 	uint32 TotalNumBeforeLock = TotalNumHazardRecords;
 	//No empty Entry found: create new
diff --git a/Engine/Source/Runtime/Core/Private/HAL/Allocators/AnsiAllocator.h b/Engine/Source/Runtime/Core/Private/HAL/Allocators/AnsiAllocator.h
--- a/Engine/Source/Runtime/Core/Private/HAL/Allocators/AnsiAllocator.h
+++ b/Engine/Source/Runtime/Core/Private/HAL/Allocators/AnsiAllocator.h
@@ -16,7 +16,7 @@ namespace UE::Core::Private
 }
 
 /** Allocator that allocates memory using standard library functions. */
-class CORE_API FAnsiAllocator
+class FAnsiAllocator
 {
 public:
 	using SizeType = int32;
diff --git a/Engine/Source/Runtime/Core/Private/Internationalization/TextHistory.h b/Engine/Source/Runtime/Core/Private/Internationalization/TextHistory.h
--- a/Engine/Source/Runtime/Core/Private/Internationalization/TextHistory.h
+++ b/Engine/Source/Runtime/Core/Private/Internationalization/TextHistory.h
@@ -304,11 +304,11 @@ protected:
 };
 
 /** Handles history for FText::Format when passing named arguments */
-class CORE_API FTextHistory_NamedFormat : public FTextHistory_Generated
+class FTextHistory_NamedFormat : public FTextHistory_Generated
 {
 public:
 	FTextHistory_NamedFormat() = default;
-	FTextHistory_NamedFormat(FString&& InDisplayString, FTextFormat&& InSourceFmt, FFormatNamedArguments&& InArguments);
+	CORE_API FTextHistory_NamedFormat(FString&& InDisplayString, FTextFormat&& InSourceFmt, FFormatNamedArguments&& InArguments);
 
 	/** Allow moving */
 	FTextHistory_NamedFormat(FTextHistory_NamedFormat&& Other) = default;
@@ -317,11 +317,11 @@ public:
 	//~ Begin FTextHistory Interface
 	OVERRIDE_TEXT_HISTORY_STRINGIFICATION;
 	virtual ETextHistoryType GetType() const override { return ETextHistoryType::NamedFormat; }
-	virtual bool IdenticalTo(const FTextHistory& Other, const ETextIdenticalModeFlags CompareModeFlags) const override;
-	virtual FString BuildLocalizedDisplayString() const override;
-	virtual FString BuildInvariantDisplayString() const override;
-	virtual void Serialize(FStructuredArchive::FRecord Record) override;
-	virtual void GetHistoricFormatData(const FText& InText, TArray<FHistoricTextFormatData>& OutHistoricFormatData) const override;
+	CORE_API virtual bool IdenticalTo(const FTextHistory& Other, const ETextIdenticalModeFlags CompareModeFlags) const override;
+	CORE_API virtual FString BuildLocalizedDisplayString() const override;
+	CORE_API virtual FString BuildInvariantDisplayString() const override;
+	CORE_API virtual void Serialize(FStructuredArchive::FRecord Record) override;
+	CORE_API virtual void GetHistoricFormatData(const FText& InText, TArray<FHistoricTextFormatData>& OutHistoricFormatData) const override;
 	//~ End FTextHistory Interface
 
 private:
diff --git a/Engine/Source/Runtime/Core/Public/Android/AndroidErrorOutputDevice.h b/Engine/Source/Runtime/Core/Public/Android/AndroidErrorOutputDevice.h
--- a/Engine/Source/Runtime/Core/Public/Android/AndroidErrorOutputDevice.h
+++ b/Engine/Source/Runtime/Core/Public/Android/AndroidErrorOutputDevice.h
@@ -5,11 +5,11 @@
 #include "Misc/OutputDevice.h"
 #include "Misc/OutputDeviceError.h"
 
-class CORE_API FAndroidErrorOutputDevice : public FOutputDeviceError
+class FAndroidErrorOutputDevice : public FOutputDeviceError
 {
 public:
 	/** Constructor, initializing member variables */
-	FAndroidErrorOutputDevice();
+	CORE_API FAndroidErrorOutputDevice();
 
 	/**
 	 * Serializes the passed in data unless the current event is suppressed.
@@ -17,13 +17,13 @@ public:
 	 * @param	Data	Text to log
 	 * @param	Event	Event name used for suppression purposes
 	 */
-	virtual void Serialize( const TCHAR* Msg, ELogVerbosity::Type Verbosity, const class FName& Category ) override;
+	CORE_API virtual void Serialize( const TCHAR* Msg, ELogVerbosity::Type Verbosity, const class FName& Category ) override;
 
 	/**
 	 * Error handling function that is being called from within the system wide global
 	 * error handler, e.g. using structured exception handling on the PC.
 	 */
-	void HandleError();
+	CORE_API void HandleError();
 
 private:
 	int32		ErrorPos;
diff --git a/Engine/Source/Runtime/Core/Public/Android/AndroidJavaEnv.h b/Engine/Source/Runtime/Core/Public/Android/AndroidJavaEnv.h
--- a/Engine/Source/Runtime/Core/Public/Android/AndroidJavaEnv.h
+++ b/Engine/Source/Runtime/Core/Public/Android/AndroidJavaEnv.h
@@ -22,7 +22,7 @@ namespace AndroidJavaEnv
 
 // Helper class that automatically calls DeleteLocalRef on the passed-in Java object when goes out of scope
 template <typename T>
-class CORE_API FScopedJavaObject
+class FScopedJavaObject
 {
 public:
 	FScopedJavaObject(JNIEnv* InEnv, const T& InObjRef) :
@@ -90,24 +90,24 @@ CORE_API FScopedJavaObject<T> NewScopedJavaObject(JNIEnv* InEnv, const T& InObjR
 	return FScopedJavaObject<T>(InEnv, InObjRef);
 }
 
-class CORE_API FJavaHelper
+class FJavaHelper
 {
 public:
 	// Converts the java string to FString and calls DeleteLocalRef on the passed-in java string reference
-	static FString FStringFromLocalRef(JNIEnv* Env, jstring JavaString);
+	static CORE_API FString FStringFromLocalRef(JNIEnv* Env, jstring JavaString);
 	
 	// Converts the java string to FString and calls DeleteGlobalRef on the passed-in java string reference
-	static FString FStringFromGlobalRef(JNIEnv* Env, jstring JavaString);
+	static CORE_API FString FStringFromGlobalRef(JNIEnv* Env, jstring JavaString);
 	
 	// Converts the java string to FString, does NOT modify the passed-in java string reference
-	static FString FStringFromParam(JNIEnv* Env, jstring JavaString);
+	static CORE_API FString FStringFromParam(JNIEnv* Env, jstring JavaString);
 	
 	// Converts FString into a Java string wrapped in FScopedJavaObject
-	static FScopedJavaObject<jstring> ToJavaString(JNIEnv* Env, const FString& UnrealString);
+	static CORE_API FScopedJavaObject<jstring> ToJavaString(JNIEnv* Env, const FString& UnrealString);
 
 	// Converts a TArray<FStringView> into a Java string array wrapped in FScopedJavaObject. FStringView content is expected to be null terminated
-	static FScopedJavaObject<jobjectArray> ToJavaStringArray(JNIEnv* Env, const TArray<FStringView>& UnrealStrings);
+	static CORE_API FScopedJavaObject<jobjectArray> ToJavaStringArray(JNIEnv* Env, const TArray<FStringView>& UnrealStrings);
 
 	// Converts the java objectArray to an array of FStrings. jopbjectArray must be a String[] on the Java side
-	static TArray<FString> ObjectArrayToFStringTArray(JNIEnv* Env, jobjectArray ObjectArray);
+	static CORE_API TArray<FString> ObjectArrayToFStringTArray(JNIEnv* Env, jobjectArray ObjectArray);
 };
diff --git a/Engine/Source/Runtime/Core/Public/Android/AndroidPlatformAffinity.h b/Engine/Source/Runtime/Core/Public/Android/AndroidPlatformAffinity.h
--- a/Engine/Source/Runtime/Core/Public/Android/AndroidPlatformAffinity.h
+++ b/Engine/Source/Runtime/Core/Public/Android/AndroidPlatformAffinity.h
@@ -15,27 +15,27 @@ private:
 	static uint64 GetLittleCoreMask();
 	const static uint64 AllCores = 0xFFFFFFFFFFFFFFFF;
 public:
-	static const CORE_API uint64 GetMainGameMask()
+	static const uint64 GetMainGameMask()
 	{
 		return GameThreadMask;
 	}
 
-	static const CORE_API uint64 GetRenderingThreadMask()
+	static const uint64 GetRenderingThreadMask()
 	{
 		return RenderingThreadMask;
 	}
 
-	static const CORE_API uint64 GetRHIThreadMask()
+	static const uint64 GetRHIThreadMask()
 	{
 		return AllCores;
 	}
 
-	static const CORE_API uint64 GetRTHeartBeatMask()
+	static const uint64 GetRTHeartBeatMask()
 	{
 		return GetLittleCoreMask();
 	}
 
-	static const CORE_API uint64 GetPoolThreadMask()
+	static const uint64 GetPoolThreadMask()
 	{
 #if ANDROID_USE_NICE_VALUE_THREADPRIORITY
 		return AllCores;
@@ -44,37 +44,37 @@ public:
 #endif
 	}
 
-	static const CORE_API uint64 GetTaskGraphThreadMask()
+	static const uint64 GetTaskGraphThreadMask()
 	{
 		return AllCores;
 	}
 
-	static const CORE_API uint64 GetAudioRenderThreadMask()
+	static const uint64 GetAudioRenderThreadMask()
 	{
 		return GetLittleCoreMask();
 	}
 
-	static const CORE_API uint64 GetTaskGraphBackgroundTaskMask()
+	static const uint64 GetTaskGraphBackgroundTaskMask()
 	{
 		return GetLittleCoreMask();
 	}
 
-	static const CORE_API uint64 GetTaskGraphHighPriorityTaskMask()
+	static const uint64 GetTaskGraphHighPriorityTaskMask()
 	{
 		return AllCores;
 	}
 
-	static const CORE_API uint64 GetAsyncLoadingThreadMask()
+	static const uint64 GetAsyncLoadingThreadMask()
 	{
 		return AllCores;
 	}
 
-	static CORE_API EThreadPriority GetRenderingThreadPriority()
+	static EThreadPriority GetRenderingThreadPriority()
 	{
 		return TPri_SlightlyBelowNormal;
 	}
 
-	static CORE_API EThreadPriority GetRHIThreadPriority()
+	static EThreadPriority GetRHIThreadPriority()
 	{
 		return TPri_Normal;
 	}
diff --git a/Engine/Source/Runtime/Core/Public/Android/AndroidPlatformCrashContext.h b/Engine/Source/Runtime/Core/Public/Android/AndroidPlatformCrashContext.h
--- a/Engine/Source/Runtime/Core/Public/Android/AndroidPlatformCrashContext.h
+++ b/Engine/Source/Runtime/Core/Public/Android/AndroidPlatformCrashContext.h
@@ -5,7 +5,7 @@
 #include "GenericPlatform/GenericPlatformCrashContext.h"
 #include "Containers/Map.h"
 
-struct CORE_API FAndroidCrashContext : public FGenericCrashContext
+struct FAndroidCrashContext : public FGenericCrashContext
 {
 	/** Signal number */
 	int32 Signal;
@@ -19,7 +19,7 @@ struct CORE_API FAndroidCrashContext : public FGenericCrashContext
 	/** Thread context */
 	void* Context;
 
-	FAndroidCrashContext(ECrashContextType InType, const TCHAR* InErrorMessage);
+	CORE_API FAndroidCrashContext(ECrashContextType InType, const TCHAR* InErrorMessage);
 
 	~FAndroidCrashContext()
 	{
@@ -40,35 +40,35 @@ struct CORE_API FAndroidCrashContext : public FGenericCrashContext
 		Context = InContext;
 	}
 
-	virtual void GetPortableCallStack(const uint64* StackFrames, int32 NumStackFrames, TArray<FCrashStackFrame>& OutCallStack) const override;
-	virtual void AddPlatformSpecificProperties() const override;
+	CORE_API virtual void GetPortableCallStack(const uint64* StackFrames, int32 NumStackFrames, TArray<FCrashStackFrame>& OutCallStack) const override;
+	CORE_API virtual void AddPlatformSpecificProperties() const override;
 
-	void CaptureCrashInfo();
-	void StoreCrashInfo(bool bWriteLog) const;
+	CORE_API void CaptureCrashInfo();
+	CORE_API void StoreCrashInfo(bool bWriteLog) const;
 
 	static const int32 CrashReportMaxPathSize = 512;
 
-	static void Initialize();
+	static CORE_API void Initialize();
 	// Returns the report directory used for this crash context.
 	FString GetCurrentReportDirectoryPath() const {	return FString(ReportDirectory); }
 
 	// Returns the main crash directory for this session. This will not be correct for non-fatal reports.
-	static const FString GetGlobalCrashDirectoryPath();
+	static CORE_API const FString GetGlobalCrashDirectoryPath();
 	// Fills DirectoryNameOUT with the global crash directory for a fatal crash this session. This will not be correct for non-fatal reports.
-	static void GetGlobalCrashDirectoryPath(char(&DirectoryNameOUT)[CrashReportMaxPathSize]);
+	static CORE_API void GetGlobalCrashDirectoryPath(char(&DirectoryNameOUT)[CrashReportMaxPathSize]);
 
-	void AddAndroidCrashProperty(const FString& Key, const FString& Value);
+	CORE_API void AddAndroidCrashProperty(const FString& Key, const FString& Value);
 
-	void SetOverrideCallstack(const FString& OverrideCallstackIN);
+	CORE_API void SetOverrideCallstack(const FString& OverrideCallstackIN);
 
 	// generate an absolute path to a crash report folder.
-	static void GenerateReportDirectoryName(char(&DirectoryNameOUT)[CrashReportMaxPathSize]);
+	static CORE_API void GenerateReportDirectoryName(char(&DirectoryNameOUT)[CrashReportMaxPathSize]);
 
 	// expects externally allocated memory for all threads found in FThreadManager + main thread
-	void DumpAllThreadCallstacks(FAsyncThreadBackTrace* BackTrace, int NumThreads) const;
+	CORE_API void DumpAllThreadCallstacks(FAsyncThreadBackTrace* BackTrace, int NumThreads) const;
 
 	/** Async-safe ItoA */
-	static const ANSICHAR* ItoANSI(uint64 Val, uint64 Base, uint32 Len = 0);
+	static CORE_API const ANSICHAR* ItoANSI(uint64 Val, uint64 Base, uint32 Len = 0);
 
 	// temporary accessor to allow overriding of the portable callstack.
 	TArray<FCrashStackFrame>& GetPortableCallstack_TEMP() { return CallStack; }
@@ -76,7 +76,7 @@ struct CORE_API FAndroidCrashContext : public FGenericCrashContext
 protected:
 
 	/** Allow platform implementations to provide a callstack property. Primarily used when non-native code triggers a crash. */
-	virtual const TCHAR* GetCallstackProperty() const;
+	CORE_API virtual const TCHAR* GetCallstackProperty() const;
 
 private:
 	FString OverrideCallstack;
@@ -87,7 +87,7 @@ private:
 	char ReportDirectory[CrashReportMaxPathSize];
 };
 
-struct CORE_API FAndroidMemoryWarningContext : public FGenericMemoryWarningContext
+struct FAndroidMemoryWarningContext : public FGenericMemoryWarningContext
 {
 	FAndroidMemoryWarningContext() : LastTrimMemoryState(FAndroidPlatformMemory::ETrimValues::Unknown) {}
 
@@ -95,4 +95,4 @@ struct CORE_API FAndroidMemoryWarningContext : public FGenericMemoryWarningConte
 	FAndroidPlatformMemory::ETrimValues LastTrimMemoryState;
 };
 
-typedef FAndroidCrashContext FPlatformCrashContext;
\ No newline at end of file
+typedef FAndroidCrashContext FPlatformCrashContext;
diff --git a/Engine/Source/Runtime/Core/Public/Android/AndroidPlatformFile.h b/Engine/Source/Runtime/Core/Public/Android/AndroidPlatformFile.h
--- a/Engine/Source/Runtime/Core/Public/Android/AndroidPlatformFile.h
+++ b/Engine/Source/Runtime/Core/Public/Android/AndroidPlatformFile.h
@@ -15,17 +15,17 @@
 	Android File I/O implementation with additional utilities to deal
 	with Java side access.
 **/
-class CORE_API IAndroidPlatformFile : public IPhysicalPlatformFile
+class IAndroidPlatformFile : public IPhysicalPlatformFile
 {
 public:
-	static IAndroidPlatformFile & GetPlatformPhysical();
+	static CORE_API IAndroidPlatformFile & GetPlatformPhysical();
 
 #if USE_ANDROID_FILE
 	/**
 	 * Get the directory path to write log files to.
 	 * This is /temp0 in shipping, or a path inside /data for other configs.
 	 */
-	static const FString* GetOverrideLogDirectory();
+	static CORE_API const FString* GetOverrideLogDirectory();
 #endif
 
 #if USE_ANDROID_JNI
@@ -50,6 +50,6 @@ public:
 	// asset path.
 	virtual FString FileRootPath(const TCHAR* Filename) = 0;
 
-	virtual FString ConvertToAbsolutePathForExternalAppForRead(const TCHAR* Filename) override;
-	virtual FString ConvertToAbsolutePathForExternalAppForWrite(const TCHAR* Filename) override;
+	CORE_API virtual FString ConvertToAbsolutePathForExternalAppForRead(const TCHAR* Filename) override;
+	CORE_API virtual FString ConvertToAbsolutePathForExternalAppForWrite(const TCHAR* Filename) override;
 };
diff --git a/Engine/Source/Runtime/Core/Public/Android/AndroidPlatformMemory.h b/Engine/Source/Runtime/Core/Public/Android/AndroidPlatformMemory.h
--- a/Engine/Source/Runtime/Core/Public/Android/AndroidPlatformMemory.h
+++ b/Engine/Source/Runtime/Core/Public/Android/AndroidPlatformMemory.h
@@ -20,17 +20,17 @@ struct FPlatformMemoryStats : public FGenericPlatformMemoryStats
 /**
 * Android implementation of the memory OS functions
 **/
-struct CORE_API FAndroidPlatformMemory : public FGenericPlatformMemory
+struct FAndroidPlatformMemory : public FGenericPlatformMemory
 {
 	//~ Begin FGenericPlatformMemory Interface
-	static void Init();
-	static FPlatformMemoryStats GetStats();
-	static uint64 GetMemoryUsedFast();
-	static const FPlatformMemoryConstants& GetConstants();
-	static EPlatformMemorySizeBucket GetMemorySizeBucket();
-	static FMalloc* BaseAllocator();
-	static void* BinnedAllocFromOS(SIZE_T Size);
-	static void BinnedFreeToOS(void* Ptr, SIZE_T Size);
+	static CORE_API void Init();
+	static CORE_API FPlatformMemoryStats GetStats();
+	static CORE_API uint64 GetMemoryUsedFast();
+	static CORE_API const FPlatformMemoryConstants& GetConstants();
+	static CORE_API EPlatformMemorySizeBucket GetMemorySizeBucket();
+	static CORE_API FMalloc* BaseAllocator();
+	static CORE_API void* BinnedAllocFromOS(SIZE_T Size);
+	static CORE_API void BinnedFreeToOS(void* Ptr, SIZE_T Size);
 
 	class FPlatformVirtualMemoryBlock : public FBasicVirtualMemoryBlock
 	{
@@ -81,7 +81,7 @@ struct CORE_API FAndroidPlatformMemory : public FGenericPlatformMemory
 		static size_t GetVirtualSizeAlignment();
 	};
 
-	static bool GetLLMAllocFunctions(void*(*&OutAllocFunction)(size_t), void(*&OutFreeFunction)(void*, size_t), int32& OutAlignment);
+	static CORE_API bool GetLLMAllocFunctions(void*(*&OutAllocFunction)(size_t), void(*&OutFreeFunction)(void*, size_t), int32& OutAlignment);
 	//~ End FGenericPlatformMemory Interface
 	
 	// ANDROID ONLY:
@@ -98,7 +98,7 @@ struct CORE_API FAndroidPlatformMemory : public FGenericPlatformMemory
 	};
 	// called when OS (via JNI) reports memory trouble, triggers MemoryWarningHandler callback on game thread if set.
 	enum class EOSMemoryStatusCategory { OSTrim };
-	static void UpdateOSMemoryStatus(EOSMemoryStatusCategory OSMemoryStatusCategory, int value);
+	static CORE_API void UpdateOSMemoryStatus(EOSMemoryStatusCategory OSMemoryStatusCategory, int value);
 };
 
 typedef FAndroidPlatformMemory FPlatformMemory;
diff --git a/Engine/Source/Runtime/Core/Public/Android/AndroidPlatformMisc.h b/Engine/Source/Runtime/Core/Public/Android/AndroidPlatformMisc.h
--- a/Engine/Source/Runtime/Core/Public/Android/AndroidPlatformMisc.h
+++ b/Engine/Source/Runtime/Core/Public/Android/AndroidPlatformMisc.h
@@ -38,24 +38,24 @@ namespace EAndroidSurfaceOrientation
 /**
  * Android implementation of the misc OS functions
  */
-struct CORE_API FAndroidMisc : public FGenericPlatformMisc
+struct FAndroidMisc : public FGenericPlatformMisc
 {
-	static void RequestExit( bool Force, const TCHAR* CallSite = nullptr);
-	static bool RestartApplication();
-	static void LocalPrint(const TCHAR *Message);
+	static CORE_API void RequestExit( bool Force, const TCHAR* CallSite = nullptr);
+	static CORE_API bool RestartApplication();
+	static CORE_API void LocalPrint(const TCHAR *Message);
 	static bool IsLocalPrintThreadSafe() { return true; }
-	static void PlatformPreInit();
-	static void PlatformInit();
-	static void PlatformTearDown();
-	static void PlatformHandleSplashScreen(bool ShowSplashScreen);
+	static CORE_API void PlatformPreInit();
+	static CORE_API void PlatformInit();
+	static CORE_API void PlatformTearDown();
+	static CORE_API void PlatformHandleSplashScreen(bool ShowSplashScreen);
     static EDeviceScreenOrientation GetDeviceOrientation() { return DeviceOrientation; }
 	UE_DEPRECATED(5.1, "SetDeviceOrientation is deprecated. Use SetAllowedDeviceOrientation instead.")
-	static void SetDeviceOrientation(EDeviceScreenOrientation NewDeviceOrentation);
-	static void SetAllowedDeviceOrientation(EDeviceScreenOrientation NewAllowedDeviceOrientation);
+	static CORE_API void SetDeviceOrientation(EDeviceScreenOrientation NewDeviceOrentation);
+	static CORE_API void SetAllowedDeviceOrientation(EDeviceScreenOrientation NewAllowedDeviceOrientation);
 
 	// Change this to an Enum with Always allow, allow and deny
-	static void SetCellularPreference(int32 Value);
-	static int32 GetCellularPreference();
+	static CORE_API void SetCellularPreference(int32 Value);
+	static CORE_API int32 GetCellularPreference();
     
 	FORCEINLINE static int32 GetMaxPathLength()
 	{
@@ -63,19 +63,19 @@ struct CORE_API FAndroidMisc : public FGenericPlatformMisc
 	}
 
 	UE_DEPRECATED(4.21, "void FPlatformMisc::GetEnvironmentVariable(Name, Result, Length) is deprecated. Use FString FPlatformMisc::GetEnvironmentVariable(Name) instead.")
-	static void GetEnvironmentVariable(const TCHAR* VariableName, TCHAR* Result, int32 ResultLength);
+	static CORE_API void GetEnvironmentVariable(const TCHAR* VariableName, TCHAR* Result, int32 ResultLength);
 
-	static FString GetEnvironmentVariable(const TCHAR* VariableName);
-	static const TCHAR* GetSystemErrorMessage(TCHAR* OutBuffer, int32 BufferCount, int32 Error);
-	static EAppReturnType::Type MessageBoxExt( EAppMsgType::Type MsgType, const TCHAR* Text, const TCHAR* Caption );
-	static bool UseRenderThread();
-	static bool HasPlatformFeature(const TCHAR* FeatureName);
-	static bool ShouldDisablePluginAtRuntime(const FString& PluginName);
-	static bool SupportsES30();
+	static CORE_API FString GetEnvironmentVariable(const TCHAR* VariableName);
+	static CORE_API const TCHAR* GetSystemErrorMessage(TCHAR* OutBuffer, int32 BufferCount, int32 Error);
+	static CORE_API EAppReturnType::Type MessageBoxExt( EAppMsgType::Type MsgType, const TCHAR* Text, const TCHAR* Caption );
+	static CORE_API bool UseRenderThread();
+	static CORE_API bool HasPlatformFeature(const TCHAR* FeatureName);
+	static CORE_API bool ShouldDisablePluginAtRuntime(const FString& PluginName);
+	static CORE_API bool SupportsES30();
 
 public:
 
-	static bool AllowThreadHeartBeat();
+	static CORE_API bool AllowThreadHeartBeat();
 
 	struct FCPUStatTime{
 		uint64_t			TotalTime;
@@ -102,38 +102,38 @@ public:
 		
 	};
 
-	static FCPUState& GetCPUState();
-	static int32 NumberOfCores();
-	static int32 NumberOfCoresIncludingHyperthreads();
-	static bool SupportsLocalCaching();
-	static void CreateGuid(struct FGuid& Result);
-	static void SetCrashHandler(void (* CrashHandler)(const FGenericCrashContext& Context));
+	static CORE_API FCPUState& GetCPUState();
+	static CORE_API int32 NumberOfCores();
+	static CORE_API int32 NumberOfCoresIncludingHyperthreads();
+	static CORE_API bool SupportsLocalCaching();
+	static CORE_API void CreateGuid(struct FGuid& Result);
+	static CORE_API void SetCrashHandler(void (* CrashHandler)(const FGenericCrashContext& Context));
 	// NOTE: THIS FUNCTION IS DEFINED IN ANDROIDOPENGL.CPP
-	static void GetValidTargetPlatforms(class TArray<class FString>& TargetPlatformNames);
-	static bool GetUseVirtualJoysticks();
-	static bool SupportsTouchInput();
+	static CORE_API void GetValidTargetPlatforms(class TArray<class FString>& TargetPlatformNames);
+	static CORE_API bool GetUseVirtualJoysticks();
+	static CORE_API bool SupportsTouchInput();
 	static const TCHAR* GetDefaultDeviceProfileName() { return TEXT("Android_Default"); }
-	static bool GetVolumeButtonsHandledBySystem();
-	static void SetVolumeButtonsHandledBySystem(bool enabled);
+	static CORE_API bool GetVolumeButtonsHandledBySystem();
+	static CORE_API void SetVolumeButtonsHandledBySystem(bool enabled);
 	// Returns current volume, 0-15
-	static int GetVolumeState(double* OutTimeOfChangeInSec = nullptr);
+	static CORE_API int GetVolumeState(double* OutTimeOfChangeInSec = nullptr);
 
-	static int32 GetDeviceVolume();
+	static CORE_API int32 GetDeviceVolume();
 
 #if USE_ANDROID_FILE
-	static const TCHAR* GamePersistentDownloadDir();
-	static FString GetLoginId();
+	static CORE_API const TCHAR* GamePersistentDownloadDir();
+	static CORE_API FString GetLoginId();
 #endif
 #if USE_ANDROID_JNI
-	static FString GetDeviceId();
-	static FString GetUniqueAdvertisingId();
+	static CORE_API FString GetDeviceId();
+	static CORE_API FString GetUniqueAdvertisingId();
 #endif
-	static FString GetCPUVendor();
-	static FString GetCPUBrand();
-	static FString GetCPUChipset();
-	static FString GetPrimaryGPUBrand();
-	static void GetOSVersions(FString& out_OSVersionLabel, FString& out_OSSubVersionLabel);
-	static bool GetDiskTotalAndFreeSpace(const FString& InPath, uint64& TotalNumberOfBytes, uint64& NumberOfFreeBytes);
+	static CORE_API FString GetCPUVendor();
+	static CORE_API FString GetCPUBrand();
+	static CORE_API FString GetCPUChipset();
+	static CORE_API FString GetPrimaryGPUBrand();
+	static CORE_API void GetOSVersions(FString& out_OSVersionLabel, FString& out_OSSubVersionLabel);
+	static CORE_API bool GetDiskTotalAndFreeSpace(const FString& InPath, uint64& TotalNumberOfBytes, uint64& NumberOfFreeBytes);
 	
 	enum EBatteryState
 	{
@@ -150,97 +150,97 @@ public:
 		float						Temperature;    // in degrees of Celsius
 	};
 
-	static FBatteryState GetBatteryState();
-	static int GetBatteryLevel();
-	static bool IsRunningOnBattery();
-	static bool IsInLowPowerMode();
-	static float GetDeviceTemperatureLevel();
-	static bool AreHeadPhonesPluggedIn();
-	static ENetworkConnectionType GetNetworkConnectionType();
+	static CORE_API FBatteryState GetBatteryState();
+	static CORE_API int GetBatteryLevel();
+	static CORE_API bool IsRunningOnBattery();
+	static CORE_API bool IsInLowPowerMode();
+	static CORE_API float GetDeviceTemperatureLevel();
+	static CORE_API bool AreHeadPhonesPluggedIn();
+	static CORE_API ENetworkConnectionType GetNetworkConnectionType();
 #if USE_ANDROID_JNI
-	static bool HasActiveWiFiConnection();
+	static CORE_API bool HasActiveWiFiConnection();
 #endif
 
-	static void RegisterForRemoteNotifications();
-	static void UnregisterForRemoteNotifications();
-	static bool IsAllowedRemoteNotifications();
+	static CORE_API void RegisterForRemoteNotifications();
+	static CORE_API void UnregisterForRemoteNotifications();
+	static CORE_API bool IsAllowedRemoteNotifications();
 
 	/** @return Memory representing a true type or open type font provided by the platform as a default font for unreal to consume; empty array if the default font failed to load. */
-	static TArray<uint8> GetSystemFontBytes();
+	static CORE_API TArray<uint8> GetSystemFontBytes();
 
-	static IPlatformChunkInstall* GetPlatformChunkInstall();
+	static CORE_API IPlatformChunkInstall* GetPlatformChunkInstall();
 
-	static void PrepareMobileHaptics(EMobileHapticsType Type);
-	static void TriggerMobileHaptics();
-	static void ReleaseMobileHaptics();
-	static void ShareURL(const FString& URL, const FText& Description, int32 LocationHintX, int32 LocationHintY);
+	static CORE_API void PrepareMobileHaptics(EMobileHapticsType Type);
+	static CORE_API void TriggerMobileHaptics();
+	static CORE_API void ReleaseMobileHaptics();
+	static CORE_API void ShareURL(const FString& URL, const FText& Description, int32 LocationHintX, int32 LocationHintY);
 
-	static FString LoadTextFileFromPlatformPackage(const FString& RelativePath);
-	static bool FileExistsInPlatformPackage(const FString& RelativePath);
+	static CORE_API FString LoadTextFileFromPlatformPackage(const FString& RelativePath);
+	static CORE_API bool FileExistsInPlatformPackage(const FString& RelativePath);
 
 	// ANDROID ONLY:
-	static void SetVersionInfo(FString AndroidVersion, int32 InTargetSDKVersion, FString DeviceMake, FString DeviceModel, FString DeviceBuildNumber, FString OSLanguage);
-	static const FString GetAndroidVersion();
-	static int32 GetAndroidMajorVersion();
-	static int32 GetTargetSDKVersion();
-	static const FString GetDeviceMake();
-	static const FString GetDeviceModel();
-	static const FString GetOSLanguage();
-	static const FString GetDeviceBuildNumber();
-	static const FString GetProjectVersion();
-	static FString GetDefaultLocale();
-	static FString GetGPUFamily();
-	static FString GetGLVersion();
-	static bool SupportsFloatingPointRenderTargets();
-	static bool SupportsShaderFramebufferFetch();
-	static bool SupportsShaderIOBlocks();
+	static CORE_API void SetVersionInfo(FString AndroidVersion, int32 InTargetSDKVersion, FString DeviceMake, FString DeviceModel, FString DeviceBuildNumber, FString OSLanguage);
+	static CORE_API const FString GetAndroidVersion();
+	static CORE_API int32 GetAndroidMajorVersion();
+	static CORE_API int32 GetTargetSDKVersion();
+	static CORE_API const FString GetDeviceMake();
+	static CORE_API const FString GetDeviceModel();
+	static CORE_API const FString GetOSLanguage();
+	static CORE_API const FString GetDeviceBuildNumber();
+	static CORE_API const FString GetProjectVersion();
+	static CORE_API FString GetDefaultLocale();
+	static CORE_API FString GetGPUFamily();
+	static CORE_API FString GetGLVersion();
+	static CORE_API bool SupportsFloatingPointRenderTargets();
+	static CORE_API bool SupportsShaderFramebufferFetch();
+	static CORE_API bool SupportsShaderIOBlocks();
 #if USE_ANDROID_JNI
-	static int GetAndroidBuildVersion();
+	static CORE_API int GetAndroidBuildVersion();
 #endif
-	static bool IsSupportedAndroidDevice();
-	static void SetForceUnsupported(bool bInOverride);
-	static const TMap<FString, FString>& GetConfigRulesTMap();
-	static FString* GetConfigRulesVariable(const FString& Key);
+	static CORE_API bool IsSupportedAndroidDevice();
+	static CORE_API void SetForceUnsupported(bool bInOverride);
+	static CORE_API const TMap<FString, FString>& GetConfigRulesTMap();
+	static CORE_API FString* GetConfigRulesVariable(const FString& Key);
 
 	/* HasVulkanDriverSupport
 	 * @return true if this Android device supports a Vulkan API Unreal could use
 	 */
-	static bool HasVulkanDriverSupport();
+	static CORE_API bool HasVulkanDriverSupport();
 
 	/* IsVulkanAvailable
 	 * @return	true if there is driver support, we have an RHI, we are packaged with Vulkan support,
 	 *			and not we are not forcing GLES with a command line switch
 	 */
-	static bool IsVulkanAvailable();
+	static CORE_API bool IsVulkanAvailable();
 
 	/* ShouldUseVulkan
 	 * @return true if Vulkan is available, and not disabled by device profile cvar
 	 */
-	static bool ShouldUseVulkan();
-	static bool ShouldUseDesktopVulkan();
-	static FString GetVulkanVersion();
+	static CORE_API bool ShouldUseVulkan();
+	static CORE_API bool ShouldUseDesktopVulkan();
+	static CORE_API FString GetVulkanVersion();
 	typedef TFunction<void(void* NewNativeHandle)> ReInitWindowCallbackType;
-	static ReInitWindowCallbackType GetOnReInitWindowCallback();
-	static void SetOnReInitWindowCallback(ReInitWindowCallbackType InOnReInitWindowCallback);
+	static CORE_API ReInitWindowCallbackType GetOnReInitWindowCallback();
+	static CORE_API void SetOnReInitWindowCallback(ReInitWindowCallbackType InOnReInitWindowCallback);
 	typedef TFunction<void()> ReleaseWindowCallbackType;
-	static ReleaseWindowCallbackType GetOnReleaseWindowCallback();
-	static void SetOnReleaseWindowCallback(ReleaseWindowCallbackType InOnReleaseWindowCallback);
-	static FString GetOSVersion();
+	static CORE_API ReleaseWindowCallbackType GetOnReleaseWindowCallback();
+	static CORE_API void SetOnReleaseWindowCallback(ReleaseWindowCallbackType InOnReleaseWindowCallback);
+	static CORE_API FString GetOSVersion();
 	static bool GetOverrideResolution(int32 &ResX, int32& ResY) { return false; }
 	typedef TFunction<void()> OnPauseCallBackType;
-	static OnPauseCallBackType GetOnPauseCallback();
-	static void SetOnPauseCallback(OnPauseCallBackType InOnPauseCallback);
-	static void TriggerCrashHandler(ECrashContextType InType, const TCHAR* InErrorMessage, const TCHAR* OverrideCallstack = nullptr);
+	static CORE_API OnPauseCallBackType GetOnPauseCallback();
+	static CORE_API void SetOnPauseCallback(OnPauseCallBackType InOnPauseCallback);
+	static CORE_API void TriggerCrashHandler(ECrashContextType InType, const TCHAR* InErrorMessage, const TCHAR* OverrideCallstack = nullptr);
 
 	// To help track down issues with failing crash handler.
-	static FString GetFatalSignalMessage(int Signal, siginfo* Info);
-	static void OverrideFatalSignalHandler(void (*FatalSignalHandlerOverrideFunc)(int Signal, struct siginfo* Info, void* Context, uint32 CrashingThreadId));
+	static CORE_API FString GetFatalSignalMessage(int Signal, siginfo* Info);
+	static CORE_API void OverrideFatalSignalHandler(void (*FatalSignalHandlerOverrideFunc)(int Signal, struct siginfo* Info, void* Context, uint32 CrashingThreadId));
 	// To help track down issues with failing crash handler.
 
-	static bool IsInSignalHandler();
+	static CORE_API bool IsInSignalHandler();
 
 #if !UE_BUILD_SHIPPING
-	static bool IsDebuggerPresent();
+	static CORE_API bool IsDebuggerPresent();
 #endif
 
 	FORCEINLINE static void MemoryBarrier()
@@ -250,36 +250,36 @@ public:
 
 
 #if STATS || ENABLE_STATNAMEDEVENTS
-	static void BeginNamedEventFrame();
-	static void BeginNamedEvent(const struct FColor& Color, const TCHAR* Text);
-	static void BeginNamedEvent(const struct FColor& Color, const ANSICHAR* Text);
-	static void EndNamedEvent();
-	static void CustomNamedStat(const TCHAR* Text, float Value, const TCHAR* Graph, const TCHAR* Unit);
-	static void CustomNamedStat(const ANSICHAR* Text, float Value, const ANSICHAR* Graph, const ANSICHAR* Unit);
+	static CORE_API void BeginNamedEventFrame();
+	static CORE_API void BeginNamedEvent(const struct FColor& Color, const TCHAR* Text);
+	static CORE_API void BeginNamedEvent(const struct FColor& Color, const ANSICHAR* Text);
+	static CORE_API void EndNamedEvent();
+	static CORE_API void CustomNamedStat(const TCHAR* Text, float Value, const TCHAR* Graph, const TCHAR* Unit);
+	static CORE_API void CustomNamedStat(const ANSICHAR* Text, float Value, const ANSICHAR* Graph, const ANSICHAR* Unit);
 #endif
 
 #if (STATS || ENABLE_STATNAMEDEVENTS)
-	static int32 TraceMarkerFileDescriptor;
+	static CORE_API int32 TraceMarkerFileDescriptor;
 #endif
 	
 	// run time compatibility information
-	static FString AndroidVersion; // version of android we are running eg "4.0.4"
-	static int32 AndroidMajorVersion; // integer major version of Android we are running, eg 10
-	static int32 TargetSDKVersion; // Target SDK version, eg 29.
-	static FString DeviceMake; // make of the device we are running on eg. "samsung"
-	static FString DeviceModel; // model of the device we are running on eg "SAMSUNG-SGH-I437"
-	static FString DeviceBuildNumber; // platform image build number of device "R16NW.G960NKSU1ARD6"
-	static FString OSLanguage; // language code the device is set to
+	static CORE_API FString AndroidVersion; // version of android we are running eg "4.0.4"
+	static CORE_API int32 AndroidMajorVersion; // integer major version of Android we are running, eg 10
+	static CORE_API int32 TargetSDKVersion; // Target SDK version, eg 29.
+	static CORE_API FString DeviceMake; // make of the device we are running on eg. "samsung"
+	static CORE_API FString DeviceModel; // model of the device we are running on eg "SAMSUNG-SGH-I437"
+	static CORE_API FString DeviceBuildNumber; // platform image build number of device "R16NW.G960NKSU1ARD6"
+	static CORE_API FString OSLanguage; // language code the device is set to
 
 	// Build version of Android, i.e. API level.
-	static int32 AndroidBuildVersion;
+	static CORE_API int32 AndroidBuildVersion;
 
 	// Key/Value pair variables from the optional configuration.txt
-	static TMap<FString, FString> ConfigRulesVariables;
+	static CORE_API TMap<FString, FString> ConfigRulesVariables;
 
-	static bool VolumeButtonsHandledBySystem;
+	static CORE_API bool VolumeButtonsHandledBySystem;
 
-	static bool bNeedsRestartAfterPSOPrecompile;
+	static CORE_API bool bNeedsRestartAfterPSOPrecompile;
 
 	enum class ECoreFrequencyProperty
 	{
@@ -288,55 +288,55 @@ public:
 		MinFrequency,
 	};
 
-	static uint32 GetCoreFrequency(int32 CoreIndex, ECoreFrequencyProperty CoreFrequencyProperty);
+	static CORE_API uint32 GetCoreFrequency(int32 CoreIndex, ECoreFrequencyProperty CoreFrequencyProperty);
 
 	// Returns CPU temperature read from one of the configurable CPU sensors via android.CPUThermalSensorFilePath CVar or AndroidEngine.ini, [ThermalSensors] section.
 	// Doesn't guarantee to work on all devices. Some devices require root access rights to read sensors information, in that case 0.0 will be returned
-	static float GetCPUTemperature();
+	static CORE_API float GetCPUTemperature();
 
-	static void UpdateDeviceOrientation();
+	static CORE_API void UpdateDeviceOrientation();
 
 	static void SaveDeviceOrientation(EDeviceScreenOrientation NewDeviceOrentation) { DeviceOrientation = NewDeviceOrentation; }
 
 	// Window access is locked by the game thread before preinit and unlocked here after RHIInit (PlatformCreateDynamicRHI). 
-	static void UnlockAndroidWindow();
+	static CORE_API void UnlockAndroidWindow();
 	
-	static TArray<int32> GetSupportedNativeDisplayRefreshRates();
+	static CORE_API TArray<int32> GetSupportedNativeDisplayRefreshRates();
 
-	static bool SetNativeDisplayRefreshRate(int32 RefreshRate);
+	static CORE_API bool SetNativeDisplayRefreshRate(int32 RefreshRate);
 
-	static int32 GetNativeDisplayRefreshRate();
+	static CORE_API int32 GetNativeDisplayRefreshRate();
 
 	/**
 	 * Returns whether or not a 16 bit index buffer should be promoted to 32 bit on load, needed for some Android devices
 	 */
-	static bool Expand16BitIndicesTo32BitOnLoad();
+	static CORE_API bool Expand16BitIndicesTo32BitOnLoad();
 
 	/**
 	 * Will return true if we wish to propagate the alpha to the backbuffer
 	 */
-	static int GetMobilePropagateAlphaSetting();
+	static CORE_API int GetMobilePropagateAlphaSetting();
 
-	static bool SupportsBackbufferSampling();
+	static CORE_API bool SupportsBackbufferSampling();
 
-	static void SetMemoryWarningHandler(void (*Handler)(const FGenericMemoryWarningContext& Context));
-	static bool HasMemoryWarningHandler();
+	static CORE_API void SetMemoryWarningHandler(void (*Handler)(const FGenericMemoryWarningContext& Context));
+	static CORE_API bool HasMemoryWarningHandler();
 
 	// Android specific requesting of exit, *ONLY* use this function in signal handling code. Otherwise normal RequestExit functions
-	static void NonReentrantRequestExit();
+	static CORE_API void NonReentrantRequestExit();
 
 	// Register/Get thread names for Android specific threads
-	static void RegisterThreadName(const char* Name, uint32 ThreadId);
-	static const char* GetThreadName(uint32 ThreadId);
+	static CORE_API void RegisterThreadName(const char* Name, uint32 ThreadId);
+	static CORE_API const char* GetThreadName(uint32 ThreadId);
 
-	static void ShowConsoleWindow();
+	static CORE_API void ShowConsoleWindow();
 
-	static FDelegateHandle AddNetworkListener(FOnNetworkConnectionChangedDelegate&& InNewDelegate);
-	static bool RemoveNetworkListener(FDelegateHandle Handle);
+	static CORE_API FDelegateHandle AddNetworkListener(FOnNetworkConnectionChangedDelegate&& InNewDelegate);
+	static CORE_API bool RemoveNetworkListener(FDelegateHandle Handle);
 
 private:
-	static const ANSICHAR* CodeToString(int Signal, int si_code);
-	static EDeviceScreenOrientation DeviceOrientation;
+	static CORE_API const ANSICHAR* CodeToString(int Signal, int si_code);
+	static CORE_API EDeviceScreenOrientation DeviceOrientation;
 
 #if USE_ANDROID_JNI
 	enum class EAndroidScreenOrientation
@@ -357,7 +357,7 @@ private:
 		SCREEN_ORIENTATION_USER_PORTRAIT = 12,
 	};
 	
-	static int32 GetAndroidScreenOrientation(EDeviceScreenOrientation ScreenOrientation);
+	static CORE_API int32 GetAndroidScreenOrientation(EDeviceScreenOrientation ScreenOrientation);
 #endif // USE_ANDROID_JNI
 };
 
diff --git a/Engine/Source/Runtime/Core/Public/Android/AndroidPlatformOutputDevices.h b/Engine/Source/Runtime/Core/Public/Android/AndroidPlatformOutputDevices.h
--- a/Engine/Source/Runtime/Core/Public/Android/AndroidPlatformOutputDevices.h
+++ b/Engine/Source/Runtime/Core/Public/Android/AndroidPlatformOutputDevices.h
@@ -8,7 +8,7 @@
 #pragma once
 #include "GenericPlatform/GenericPlatformOutputDevices.h"
 
-struct CORE_API FAndroidOutputDevices : public FGenericPlatformOutputDevices
+struct FAndroidOutputDevices : public FGenericPlatformOutputDevices
 {
 };
 
diff --git a/Engine/Source/Runtime/Core/Public/Android/AndroidPlatformProcess.h b/Engine/Source/Runtime/Core/Public/Android/AndroidPlatformProcess.h
--- a/Engine/Source/Runtime/Core/Public/Android/AndroidPlatformProcess.h
+++ b/Engine/Source/Runtime/Core/Public/Android/AndroidPlatformProcess.h
@@ -26,22 +26,22 @@ public:
 /**
  * Android implementation of the Process OS functions
  **/
-struct CORE_API FAndroidPlatformProcess : public FGenericPlatformProcess
+struct FAndroidPlatformProcess : public FGenericPlatformProcess
 {
-	static void* GetDllHandle(const TCHAR* Filename);
-	static void FreeDllHandle(void* DllHandle);
-	static void* GetDllExport(void* DllHandle, const TCHAR* ProcName);
-	static const TCHAR* ComputerName();
-	static void SetThreadAffinityMask( uint64 AffinityMask );
-	static uint32 GetCurrentProcessId();
-	static uint32 GetCurrentCoreNumber();
-	static const TCHAR* BaseDir();
-	static const TCHAR* ExecutableName(bool bRemoveExtension = true);
-	static class FRunnableThread* CreateRunnableThread();
-	static bool CanLaunchURL(const TCHAR* URL);
-	static void LaunchURL(const TCHAR* URL, const TCHAR* Parms, FString* Error);
-	static FString GetGameBundleId();
-	static void SetThreadName(const TCHAR* ThreadName);
+	static CORE_API void* GetDllHandle(const TCHAR* Filename);
+	static CORE_API void FreeDllHandle(void* DllHandle);
+	static CORE_API void* GetDllExport(void* DllHandle, const TCHAR* ProcName);
+	static CORE_API const TCHAR* ComputerName();
+	static CORE_API void SetThreadAffinityMask( uint64 AffinityMask );
+	static CORE_API uint32 GetCurrentProcessId();
+	static CORE_API uint32 GetCurrentCoreNumber();
+	static CORE_API const TCHAR* BaseDir();
+	static CORE_API const TCHAR* ExecutableName(bool bRemoveExtension = true);
+	static CORE_API class FRunnableThread* CreateRunnableThread();
+	static CORE_API bool CanLaunchURL(const TCHAR* URL);
+	static CORE_API void LaunchURL(const TCHAR* URL, const TCHAR* Parms, FString* Error);
+	static CORE_API FString GetGameBundleId();
+	static CORE_API void SetThreadName(const TCHAR* ThreadName);
 };
 
 typedef FAndroidPlatformProcess FPlatformProcess;
diff --git a/Engine/Source/Runtime/Core/Public/Android/AndroidPlatformStackWalk.h b/Engine/Source/Runtime/Core/Public/Android/AndroidPlatformStackWalk.h
--- a/Engine/Source/Runtime/Core/Public/Android/AndroidPlatformStackWalk.h
+++ b/Engine/Source/Runtime/Core/Public/Android/AndroidPlatformStackWalk.h
@@ -22,20 +22,20 @@ struct FAsyncThreadBackTrace
 /**
 * Android platform stack walking
 */
-struct CORE_API FAndroidPlatformStackWalk : public FGenericPlatformStackWalk
+struct FAndroidPlatformStackWalk : public FGenericPlatformStackWalk
 {
 	typedef FGenericPlatformStackWalk Parent;
 
-	static void ProgramCounterToSymbolInfo(uint64 ProgramCounter, FProgramCounterSymbolInfo& out_SymbolInfo);
-	static uint32 CaptureStackBackTrace(uint64* BackTrace, uint32 MaxDepth, void* Context = nullptr);
-	static bool SymbolInfoToHumanReadableString(const FProgramCounterSymbolInfo& SymbolInfo, ANSICHAR* HumanReadableString, SIZE_T HumanReadableStringSize);
+	static CORE_API void ProgramCounterToSymbolInfo(uint64 ProgramCounter, FProgramCounterSymbolInfo& out_SymbolInfo);
+	static CORE_API uint32 CaptureStackBackTrace(uint64* BackTrace, uint32 MaxDepth, void* Context = nullptr);
+	static CORE_API bool SymbolInfoToHumanReadableString(const FProgramCounterSymbolInfo& SymbolInfo, ANSICHAR* HumanReadableString, SIZE_T HumanReadableStringSize);
 
-	static uint32 CaptureThreadStackBackTrace(uint64 ThreadId, uint64* BackTrace, uint32 MaxDepth, void* Context = nullptr);
-	static int CaptureThreadStackBackTraceAsync(FAsyncThreadBackTrace* BackTrace);
+	static CORE_API uint32 CaptureThreadStackBackTrace(uint64 ThreadId, uint64* BackTrace, uint32 MaxDepth, void* Context = nullptr);
+	static CORE_API int CaptureThreadStackBackTraceAsync(FAsyncThreadBackTrace* BackTrace);
 
-	static void HandleBackTraceSignal(siginfo* Info, void* Context);
+	static CORE_API void HandleBackTraceSignal(siginfo* Info, void* Context);
 
-	static bool InitStackWalking();
+	static CORE_API bool InitStackWalking();
 };
 
 typedef FAndroidPlatformStackWalk FPlatformStackWalk;
diff --git a/Engine/Source/Runtime/Core/Public/Android/AndroidPlatformTLS.h b/Engine/Source/Runtime/Core/Public/Android/AndroidPlatformTLS.h
--- a/Engine/Source/Runtime/Core/Public/Android/AndroidPlatformTLS.h
+++ b/Engine/Source/Runtime/Core/Public/Android/AndroidPlatformTLS.h
@@ -12,7 +12,7 @@
 /**
  * Android implementation of the TLS OS functions
  */
-struct CORE_API FAndroidTLS : public FGenericPlatformTLS
+struct FAndroidTLS : public FGenericPlatformTLS
 {
 	/**
 	 * Returns the currently executing thread's id
diff --git a/Engine/Source/Runtime/Core/Public/Android/AndroidPlatformTime.h b/Engine/Source/Runtime/Core/Public/Android/AndroidPlatformTime.h
--- a/Engine/Source/Runtime/Core/Public/Android/AndroidPlatformTime.h
+++ b/Engine/Source/Runtime/Core/Public/Android/AndroidPlatformTime.h
@@ -14,7 +14,7 @@
 /**
  * Android implementation of the Time OS functions
  */
-struct CORE_API FAndroidTime : public FGenericPlatformTime
+struct FAndroidTime : public FGenericPlatformTime
 {
 	// android uses BSD time code from GenericPlatformTime
 	static FORCEINLINE double Seconds()
diff --git a/Engine/Source/Runtime/Core/Public/Async/AsyncFileHandle.h b/Engine/Source/Runtime/Core/Public/Async/AsyncFileHandle.h
--- a/Engine/Source/Runtime/Core/Public/Async/AsyncFileHandle.h
+++ b/Engine/Source/Runtime/Core/Public/Async/AsyncFileHandle.h
@@ -187,7 +187,7 @@ protected:
 	}
 };
 
-class CORE_API IAsyncReadFileHandle
+class IAsyncReadFileHandle
 {
 public:
 	IAsyncReadFileHandle()
diff --git a/Engine/Source/Runtime/Core/Public/Async/MappedFileHandle.h b/Engine/Source/Runtime/Core/Public/Async/MappedFileHandle.h
--- a/Engine/Source/Runtime/Core/Public/Async/MappedFileHandle.h
+++ b/Engine/Source/Runtime/Core/Public/Async/MappedFileHandle.h
@@ -13,7 +13,7 @@ DECLARE_DWORD_ACCUMULATOR_STAT_EXTERN(TEXT("Num Mapped File Regions"), STAT_Mapp
 
 // Note on threading. Like the rest of the filesystem platform abstraction, these methods are threadsafe, but it is expected you are not concurrently _using_ these data structures. 
 
-class CORE_API IMappedFileRegion
+class IMappedFileRegion
 {
 	const uint8* MappedPtr;
 	size_t MappedSize;
@@ -81,7 +81,7 @@ public:
 	IMappedFileRegion& operator=(const IMappedFileRegion&) = delete;
 };
 
-class CORE_API IMappedFileHandle
+class IMappedFileHandle
 {
 	size_t MappedFileSize;
 
diff --git a/Engine/Source/Runtime/Core/Public/Async/TaskGraphInterfaces.h b/Engine/Source/Runtime/Core/Public/Async/TaskGraphInterfaces.h
--- a/Engine/Source/Runtime/Core/Public/Async/TaskGraphInterfaces.h
+++ b/Engine/Source/Runtime/Core/Public/Async/TaskGraphInterfaces.h
@@ -217,7 +217,7 @@ enum class UE_DEPRECATED(4.26, "No longer supported") EPowerSavingEligibility :
 	NotEligible			// When set high priority tasks will not be downgraded when power saving is required.
 };
 
-class CORE_API FAutoConsoleTaskPriority
+class FAutoConsoleTaskPriority
 {
 	FString RawSetting;
 	FString FullHelpText;
@@ -226,9 +226,9 @@ class CORE_API FAutoConsoleTaskPriority
 	ENamedThreads::Type TaskPriority;
 	ENamedThreads::Type TaskPriorityIfForcedToNormalThreadPriority;
 
-	static FString CreateFullHelpText(const TCHAR* Name, const TCHAR* OriginalHelp);
-	static FString ConfigStringFromPriorities(ENamedThreads::Type InThreadPriority, ENamedThreads::Type InTaskPriority, ENamedThreads::Type InTaskPriorityBackup);
-	void OnSettingChanged(IConsoleVariable* Variable);
+	static CORE_API FString CreateFullHelpText(const TCHAR* Name, const TCHAR* OriginalHelp);
+	static CORE_API FString ConfigStringFromPriorities(ENamedThreads::Type InThreadPriority, ENamedThreads::Type InTaskPriority, ENamedThreads::Type InTaskPriorityBackup);
+	CORE_API void OnSettingChanged(IConsoleVariable* Variable);
 
 public:
 	FAutoConsoleTaskPriority(const TCHAR* Name, const TCHAR* Help, ENamedThreads::Type DefaultThreadPriority, ENamedThreads::Type DefaultTaskPriority, ENamedThreads::Type DefaultTaskPriorityIfForcedToNormalThreadPriority = ENamedThreads::UnusedAnchor)
diff --git a/Engine/Source/Runtime/Core/Public/Async/TaskTrace.h b/Engine/Source/Runtime/Core/Public/Async/TaskTrace.h
--- a/Engine/Source/Runtime/Core/Public/Async/TaskTrace.h
+++ b/Engine/Source/Runtime/Core/Public/Async/TaskTrace.h
@@ -46,17 +46,17 @@ namespace TaskTrace
 	void CORE_API Completed(FId TaskId);
 	void CORE_API Destroyed(FId TaskId);
 
-	struct CORE_API FWaitingScope
+	struct FWaitingScope
 	{
-		explicit FWaitingScope(const TArray<FId>& Tasks); // waiting for given tasks completion
-		explicit FWaitingScope(FId TaskId);
-		~FWaitingScope();
+		CORE_API explicit FWaitingScope(const TArray<FId>& Tasks); // waiting for given tasks completion
+		CORE_API explicit FWaitingScope(FId TaskId);
+		CORE_API ~FWaitingScope();
 	};
 
-	struct CORE_API FTaskTimingEventScope
+	struct FTaskTimingEventScope
 	{
-		FTaskTimingEventScope(TaskTrace::FId InTaskId);
-		~FTaskTimingEventScope();
+		CORE_API FTaskTimingEventScope(TaskTrace::FId InTaskId);
+		CORE_API ~FTaskTimingEventScope();
 
 	private:
 		bool bIsActive = false;
diff --git a/Engine/Source/Runtime/Core/Public/Clang/ClangPlatformAtomics.h b/Engine/Source/Runtime/Core/Public/Clang/ClangPlatformAtomics.h
--- a/Engine/Source/Runtime/Core/Public/Clang/ClangPlatformAtomics.h
+++ b/Engine/Source/Runtime/Core/Public/Clang/ClangPlatformAtomics.h
@@ -14,7 +14,7 @@
 /**
  * GCC/Clang implementation of the Atomics OS functions
  **/
-struct CORE_API FClangPlatformAtomics : public FGenericPlatformAtomics
+struct FClangPlatformAtomics : public FGenericPlatformAtomics
 {
 	static FORCEINLINE int8 InterlockedIncrement(volatile int8* Value)
 	{
diff --git a/Engine/Source/Runtime/Core/Public/Containers/BackgroundableTicker.h b/Engine/Source/Runtime/Core/Public/Containers/BackgroundableTicker.h
--- a/Engine/Source/Runtime/Core/Public/Containers/BackgroundableTicker.h
+++ b/Engine/Source/Runtime/Core/Public/Containers/BackgroundableTicker.h
@@ -10,14 +10,14 @@
  * This works the same as the core FTSTicker, but on supported mobile platforms
  * it continues ticking while the app is running in the background.
  */
-class CORE_API FTSBackgroundableTicker
+class FTSBackgroundableTicker
 	: public FTSTicker
 {
 public:
-	static FTSBackgroundableTicker& GetCoreTicker();
+	static CORE_API FTSBackgroundableTicker& GetCoreTicker();
 
-	FTSBackgroundableTicker();
-	~FTSBackgroundableTicker();
+	CORE_API FTSBackgroundableTicker();
+	CORE_API ~FTSBackgroundableTicker();
 
 private:
 	FDelegateHandle CoreTickerHandle;
diff --git a/Engine/Source/Runtime/Core/Public/Containers/Ticker.h b/Engine/Source/Runtime/Core/Public/Containers/Ticker.h
--- a/Engine/Source/Runtime/Core/Public/Containers/Ticker.h
+++ b/Engine/Source/Runtime/Core/Public/Containers/Ticker.h
@@ -23,7 +23,7 @@ DECLARE_DELEGATE_RetVal_OneParam(bool, FTickerDelegate, float);
 /**
  * Thread-safe ticker class. Fires delegates after a delay.
  */
-class CORE_API FTSTicker
+class FTSTicker
 {
 private:
 	struct FElement;
@@ -32,7 +32,7 @@ public:
 	using FDelegateHandle = TWeakPtr<FElement>;
 
 	/** Singleton used for the ticker in Core / Launch. If you add a new ticker for a different subsystem, do not put that singleton here! **/
-	static FTSTicker& GetCoreTicker();
+	static CORE_API FTSTicker& GetCoreTicker();
 
 	/**
 	 * Add a new ticker with a given delay / interval
@@ -42,7 +42,7 @@ public:
 	 * @param InDelegate Delegate to fire after the delay
 	 * @param InDelay Delay until next fire; 0 means "next frame"
 	 */
-	FDelegateHandle AddTicker(const FTickerDelegate & InDelegate, float InDelay = 0.0f);
+	CORE_API FDelegateHandle AddTicker(const FTickerDelegate & InDelegate, float InDelay = 0.0f);
 
 	/**
 	* Add a new ticker with a given delay / interval.
@@ -53,7 +53,7 @@ public:
 	* @param InDelay Delay until next fire; 0 means "next frame"
 	* @param Function Function to execute. Should return true to fire after another InDelay time
 	*/
-	FDelegateHandle AddTicker(const TCHAR * InName, float InDelay, TFunction<bool(float)> Function);
+	CORE_API FDelegateHandle AddTicker(const TCHAR * InName, float InDelay, TFunction<bool(float)> Function);
 
 	/**
 	 * Removes a previously added ticker delegate.
@@ -63,7 +63,7 @@ public:
 	 *
 	 * @param Handle The handle of the ticker to remove.
 	 */
-	static void RemoveTicker(FDelegateHandle Handle);
+	static CORE_API void RemoveTicker(FDelegateHandle Handle);
 
 	/**
 	 * Fire all tickers who have passed their delay and reschedule the ones that return true
@@ -78,12 +78,12 @@ public:
 	 *
 	 * @param DeltaTime	time that has passed since the last tick call
 	 */
-	void Tick(float DeltaTime);
+	CORE_API void Tick(float DeltaTime);
 
 	/**
 	 * Resets the instance to its default state. Must be called from the ticking thread.
 	 */
-	void Reset();
+	CORE_API void Reset();
 
 private:
 	/** Internal structure to store a ticker delegate and related data **/
@@ -131,7 +131,7 @@ private:
 /**
  * Base class for thread-safe ticker objects
  */
-class CORE_API FTSTickerObjectBase
+class FTSTickerObjectBase
 {
 public:
 	UE_NONCOPYABLE(FTSTickerObjectBase);
@@ -142,10 +142,10 @@ public:
 	 * @param InDelay Delay until next fire; 0 means "next frame"
 	 * @param Ticker the ticker to register with. Defaults to FTSTicker::GetCoreTicker().
 	*/
-	FTSTickerObjectBase(float InDelay = 0.0f, FTSTicker& Ticker = FTSTicker::GetCoreTicker());
+	CORE_API FTSTickerObjectBase(float InDelay = 0.0f, FTSTicker& Ticker = FTSTicker::GetCoreTicker());
 
 	/** Virtual destructor. */
-	virtual ~FTSTickerObjectBase();
+	CORE_API virtual ~FTSTickerObjectBase();
 
 	/**
 	 * Pure virtual that must be overloaded by the inheriting class.
@@ -172,4 +172,4 @@ void ExecuteOnGameThread(const TCHAR* DebugName, FunctorType&& Functor)
 			return false;
 		}
 	);
-}
\ No newline at end of file
+}
diff --git a/Engine/Source/Runtime/Core/Public/Containers/UnrealString.h b/Engine/Source/Runtime/Core/Public/Containers/UnrealString.h
--- a/Engine/Source/Runtime/Core/Public/Containers/UnrealString.h
+++ b/Engine/Source/Runtime/Core/Public/Containers/UnrealString.h
@@ -52,7 +52,7 @@ int32        GetNum(const FString& String);
  * - Use UTF8TEXTVIEW("...") for string view literals (type is FUtf8StringView).
  * - Use \uxxxx or \Uxxxxxxxx escape sequences rather than \x to specify Unicode code points.
  */
-class CORE_API FString
+class FString
 {
 public:
 	using AllocatorType = TSizedDefaultAllocator<32>;
@@ -134,22 +134,22 @@ public:
 	}
 
 	/** Construct from null-terminated C string or nullptr  */
-	FString(const ANSICHAR* Str);
-	FString(const WIDECHAR* Str);
-	FString(const UTF8CHAR* Str);
-	FString(const UCS2CHAR* Str);
+	CORE_API FString(const ANSICHAR* Str);
+	CORE_API FString(const WIDECHAR* Str);
+	CORE_API FString(const UTF8CHAR* Str);
+	CORE_API FString(const UCS2CHAR* Str);
 
 	/** Construct from null-terminated C substring or nullptr */
-	FString(int32 Len, const ANSICHAR* Str);
-	FString(int32 Len, const WIDECHAR* Str);
-	FString(int32 Len, const UTF8CHAR* Str);
-	FString(int32 Len, const UCS2CHAR* Str);
+	CORE_API FString(int32 Len, const ANSICHAR* Str);
+	CORE_API FString(int32 Len, const WIDECHAR* Str);
+	CORE_API FString(int32 Len, const UTF8CHAR* Str);
+	CORE_API FString(int32 Len, const UCS2CHAR* Str);
 
 	/** Construct from null-terminated C string or nullptr with extra slack on top of original string length */
-	FString(const ANSICHAR* Str, int32 ExtraSlack);
-	FString(const WIDECHAR* Str, int32 ExtraSlack);
-	FString(const UTF8CHAR* Str, int32 ExtraSlack);
-	FString(const UCS2CHAR* Str, int32 ExtraSlack);
+	CORE_API FString(const ANSICHAR* Str, int32 ExtraSlack);
+	CORE_API FString(const WIDECHAR* Str, int32 ExtraSlack);
+	CORE_API FString(const UTF8CHAR* Str, int32 ExtraSlack);
+	CORE_API FString(const UCS2CHAR* Str, int32 ExtraSlack);
 
 	/** Construct from contiguous range of characters such as FStringView or FStringBuilderBase */
 	template <typename CharRangeType, typename TEnableIf<TIsCharRangeNotCArrayNotFString<CharRangeType>::Value>::Type* = nullptr>
@@ -195,7 +195,7 @@ public:
 	}
 #endif
 
-	FString& operator=(const TCHAR* Str);
+	CORE_API FString& operator=(const TCHAR* Str);
 
 	template <typename CharRangeType, typename TEnableIf<TIsTCharRangeNotCArray<CharRangeType>::Value>::Type* = nullptr>
 	FORCEINLINE FString& operator=(CharRangeType&& Range)
@@ -205,7 +205,7 @@ public:
 	}
 
 private:
-	void AssignRange(const TCHAR* Str, int32 Len);
+	CORE_API void AssignRange(const TCHAR* Str, int32 Len);
 
 public:
 	/**
@@ -284,8 +284,8 @@ public:
 	 *
 	 * @param Slack length of empty string to create
 	 */
-	void Empty();
-	void Empty(int32 Slack);
+	CORE_API void Empty();
+	CORE_API void Empty(int32 Slack);
 
 	/**
 	 * Test whether this string is empty
@@ -302,12 +302,12 @@ public:
 	 *
 	 * @param NewReservedSize The expected usage size (in characters, not including the terminator) after calling this function.
 	 */
-	void Reset(int32 NewReservedSize = 0);
+	CORE_API void Reset(int32 NewReservedSize = 0);
 
 	/**
 	 * Remove unallocated empty character space from the end of this string
 	 */
-	void Shrink();
+	CORE_API void Shrink();
 
 	/**
 	 * Tests if index is valid, i.e. greater than or equal to zero, and less than the number of characters in this string (excluding the null terminator).
@@ -350,7 +350,7 @@ public:
 
 #ifdef __OBJC__
 	/** Convert FString to Objective-C NSString */
-    NSString* GetNSString() const;
+    CORE_API NSString* GetNSString() const;
 #endif
 
 	/** 
@@ -358,10 +358,10 @@ public:
 	 *
 	 * @param Str can be null if Count is 0. Can be unterminated, Str[Count] isn't read.
 	 */
-	void AppendChars(const ANSICHAR* Str, int32 Count);
-	void AppendChars(const WIDECHAR* Str, int32 Count);
-	void AppendChars(const UCS2CHAR* Str, int32 Count);
-	void AppendChars(const UTF8CHAR* Str, int32 Count);
+	CORE_API void AppendChars(const ANSICHAR* Str, int32 Count);
+	CORE_API void AppendChars(const WIDECHAR* Str, int32 Count);
+	CORE_API void AppendChars(const UCS2CHAR* Str, int32 Count);
+	CORE_API void AppendChars(const UTF8CHAR* Str, int32 Count);
 
 	/** Append a string and return a reference to this */
 	template<class CharType>
@@ -395,7 +395,7 @@ public:
 	}
 
 	/** Append a single character and return a reference to this */
-	FString& AppendChar(TCHAR InChar);
+	CORE_API FString& AppendChar(TCHAR InChar);
 
 	/** Append a string and return a reference to this */
 	template <typename StrType>
@@ -422,8 +422,8 @@ public:
 		}
 	}
 
-	void InsertAt(int32 Index, TCHAR Character);
-	void InsertAt(int32 Index, const FString& Characters);
+	CORE_API void InsertAt(int32 Index, TCHAR Character);
+	CORE_API void InsertAt(int32 Index, const FString& Characters);
 
 	/**
 	 * Removes characters within the string.
@@ -432,7 +432,7 @@ public:
 	 * @param Count           The number of characters to remove.
 	 * @param bAllowShrinking Whether or not to reallocate to shrink the storage after removal.
 	 */
-	void RemoveAt(int32 Index, int32 Count = 1, bool bAllowShrinking = true);
+	CORE_API void RemoveAt(int32 Index, int32 Count = 1, bool bAllowShrinking = true);
 
 	/**
 	 * Removes the text from the start of the string if it exists.
@@ -476,7 +476,7 @@ public:
 	 * @param InPrefixLen length of InPrefix
 	 * @return true if the prefix was removed, otherwise false.
 	 */
-	bool RemoveFromStart(const TCHAR* InPrefix, int32 InPrefixLen, ESearchCase::Type SearchCase = ESearchCase::IgnoreCase);
+	CORE_API bool RemoveFromStart(const TCHAR* InPrefix, int32 InPrefixLen, ESearchCase::Type SearchCase = ESearchCase::IgnoreCase);
 
 	/**
 	 * Removes the text from the end of the string if it exists.
@@ -520,7 +520,7 @@ public:
 	 * @param InSuffixLen length of InSuffix
 	 * @return true if the suffix was removed, otherwise false.
 	 */
-	bool RemoveFromEnd(const TCHAR* InSuffix, int32 InSuffixLen, ESearchCase::Type SearchCase = ESearchCase::IgnoreCase);
+	CORE_API bool RemoveFromEnd(const TCHAR* InSuffix, int32 InSuffixLen, ESearchCase::Type SearchCase = ESearchCase::IgnoreCase);
 
 	/**
 	 * Concatenate this path with given path ensuring the / character is used between them
@@ -528,7 +528,7 @@ public:
 	 * @param Str       Pointer to an array of TCHARs (not necessarily null-terminated) to be concatenated onto the end of this.
 	 * @param StrLength Exact number of characters from Str to append.
 	 */
-	void PathAppend(const TCHAR* Str, int32 StrLength);
+	CORE_API void PathAppend(const TCHAR* Str, int32 StrLength);
 
 	/**
 	 * Concatenates an FString with a TCHAR.
@@ -569,18 +569,18 @@ public:
 	}
 
 private:
-	UE_NODISCARD static FString ConcatFF(const FString& Lhs, const FString& Rhs);
-	UE_NODISCARD static FString ConcatFF(FString&& Lhs, const FString& Rhs);
-	UE_NODISCARD static FString ConcatFF(const FString& Lhs, FString&& Rhs);
-	UE_NODISCARD static FString ConcatFF(FString&& Lhs,	FString&& Rhs);
-	UE_NODISCARD static FString ConcatFC(const FString& Lhs, const TCHAR* Rhs);
-	UE_NODISCARD static FString ConcatFC(FString&& Lhs,	const TCHAR* Rhs);
-	UE_NODISCARD static FString ConcatCF(const TCHAR* Lhs, const FString& Rhs);
-	UE_NODISCARD static FString ConcatCF(const TCHAR* Lhs, FString&& Rhs);
-	UE_NODISCARD static FString ConcatFR(const FString& Lhs, const TCHAR* Rhs, int32 RhsLen);
-	UE_NODISCARD static FString ConcatFR(FString&& Lhs,	const TCHAR* Rhs, int32 RhsLen);
-	UE_NODISCARD static FString ConcatRF(const TCHAR* Lhs, int32 LhsLen, const FString& Rhs);
-	UE_NODISCARD static FString ConcatRF(const TCHAR* Lhs, int32 LhsLen, FString&& Rhs);
+	UE_NODISCARD static CORE_API FString ConcatFF(const FString& Lhs, const FString& Rhs);
+	UE_NODISCARD static CORE_API FString ConcatFF(FString&& Lhs, const FString& Rhs);
+	UE_NODISCARD static CORE_API FString ConcatFF(const FString& Lhs, FString&& Rhs);
+	UE_NODISCARD static CORE_API FString ConcatFF(FString&& Lhs,	FString&& Rhs);
+	UE_NODISCARD static CORE_API FString ConcatFC(const FString& Lhs, const TCHAR* Rhs);
+	UE_NODISCARD static CORE_API FString ConcatFC(FString&& Lhs,	const TCHAR* Rhs);
+	UE_NODISCARD static CORE_API FString ConcatCF(const TCHAR* Lhs, const FString& Rhs);
+	UE_NODISCARD static CORE_API FString ConcatCF(const TCHAR* Lhs, FString&& Rhs);
+	UE_NODISCARD static CORE_API FString ConcatFR(const FString& Lhs, const TCHAR* Rhs, int32 RhsLen);
+	UE_NODISCARD static CORE_API FString ConcatFR(FString&& Lhs,	const TCHAR* Rhs, int32 RhsLen);
+	UE_NODISCARD static CORE_API FString ConcatRF(const TCHAR* Lhs, int32 LhsLen, const FString& Rhs);
+	UE_NODISCARD static CORE_API FString ConcatRF(const TCHAR* Lhs, int32 LhsLen, FString&& Rhs);
 
 public:
 	UE_NODISCARD FORCEINLINE friend FString operator+(const FString& Lhs, const FString& Rhs)	{ return ConcatFF(Lhs, Rhs); }
@@ -1039,7 +1039,7 @@ public:
 	}
 
 	/** Returns the string to the right of the specified location, counting forward from the left (from the beginning of the word). */
-	UE_NODISCARD FString RightChop( int32 Count ) const &;
+	UE_NODISCARD CORE_API FString RightChop( int32 Count ) const &;
 
 	UE_NODISCARD FORCEINLINE FString RightChop(int32 Count) &&
 	{
@@ -1054,11 +1054,11 @@ public:
 	}
 
 	/** Returns the substring from Start position for Count characters. */
-	UE_NODISCARD FString Mid(int32 Start, int32 Count) const &;
-	UE_NODISCARD FString Mid(int32 Start, int32 Count) &&;
+	UE_NODISCARD CORE_API FString Mid(int32 Start, int32 Count) const &;
+	UE_NODISCARD CORE_API FString Mid(int32 Start, int32 Count) &&;
 
 	/** Returns the substring from Start position to the end */
-	UE_NODISCARD FORCEINLINE FString Mid(int32 Start) const & { return RightChop(Start); }
+	UE_NODISCARD CORE_API FORCEINLINE FString Mid(int32 Start) const & { return RightChop(Start); }
 	UE_NODISCARD FORCEINLINE FString Mid(int32 Start) && { return ((FString&&)*this).RightChop(Start); }
 
 	/** Modifies the string such that it is now the substring from Start position for Count characters. */
@@ -1149,7 +1149,7 @@ public:
 	 *
 	 *        Consider using UE::String::FindLast() as an alternative.
 	 */
-	UE_NODISCARD int32 Find(const TCHAR* SubStr, int32 InSubStrLen, ESearchCase::Type SearchCase = ESearchCase::IgnoreCase,
+	UE_NODISCARD CORE_API int32 Find(const TCHAR* SubStr, int32 InSubStrLen, ESearchCase::Type SearchCase = ESearchCase::IgnoreCase,
 		ESearchDir::Type SearchDir = ESearchDir::FromStart, int32 StartPosition = INDEX_NONE) const;
 
 
@@ -1326,47 +1326,47 @@ public:
 	 * @param SearchDir			Indicates whether the search starts at the beginning or at the end ( defaults to ESearchDir::FromStart )
 	 * @return true if string is split, otherwise false
 	 */
-	bool Split(const FString& InS, FString* LeftS, FString* RightS, ESearchCase::Type SearchCase,
+	CORE_API bool Split(const FString& InS, FString* LeftS, FString* RightS, ESearchCase::Type SearchCase,
 		ESearchDir::Type SearchDir = ESearchDir::FromStart) const;
 
 	/** Split with ESearchCase::IgnoreCase and ESearchDir::FromStart. Allows compiler to avoid branches w/o inlining code. */
-	bool Split(const FString& InS, FString* LeftS, FString* RightS) const;
+	CORE_API bool Split(const FString& InS, FString* LeftS, FString* RightS) const;
 
 	/** Returns a new string with the characters of this converted to uppercase */
-	UE_NODISCARD FString ToUpper() const &;
+	UE_NODISCARD CORE_API FString ToUpper() const &;
 
 	/**
 	 * Converts all characters in this rvalue string to uppercase and moves it into the returned string.
 	 * @return a new string with the characters of this converted to uppercase
 	 */
-	UE_NODISCARD FString ToUpper() &&;
+	UE_NODISCARD CORE_API FString ToUpper() &&;
 
 	/** Converts all characters in this string to uppercase */
-	void ToUpperInline();
+	CORE_API void ToUpperInline();
 
 	/** Returns a new string with the characters of this converted to lowercase */
-	UE_NODISCARD FString ToLower() const &;
+	UE_NODISCARD CORE_API FString ToLower() const &;
 
 	/**
 	 * Converts all characters in this rvalue string to lowercase and moves it into the returned string.
 	 * @return a new string with the characters of this converted to lowercase
 	 */
-	UE_NODISCARD FString ToLower() &&;
+	UE_NODISCARD CORE_API FString ToLower() &&;
 
 	/** Converts all characters in this string to lowercase */
-	void ToLowerInline();
+	CORE_API void ToLowerInline();
 
 	/** Pad the left of this string for ChCount characters */
-	UE_NODISCARD FString LeftPad( int32 ChCount ) const;
+	UE_NODISCARD CORE_API FString LeftPad( int32 ChCount ) const;
 
 	/** Pad the right of this string for ChCount characters */
-	UE_NODISCARD FString RightPad( int32 ChCount ) const;
+	UE_NODISCARD CORE_API FString RightPad( int32 ChCount ) const;
 	
 	/** Returns true if the string only contains numeric characters */
-	UE_NODISCARD bool IsNumeric() const;
+	UE_NODISCARD CORE_API bool IsNumeric() const;
 	
 	/** Removes spaces from the string.  I.E. "Spaces Are Cool" --> "SpacesAreCool". */
-	void RemoveSpacesInline();
+	CORE_API void RemoveSpacesInline();
 
 	/**
 	 * Constructs FString object similarly to how classic sprintf works.
@@ -1400,8 +1400,8 @@ public:
 	}
 
 private:
-	static FString VARARGS PrintfImpl(const TCHAR* Fmt, ...);
-	static void VARARGS AppendfImpl(FString& AppendToMe, const TCHAR* Fmt, ...);
+	static CORE_API FString VARARGS PrintfImpl(const TCHAR* Fmt, ...);
+	static CORE_API void VARARGS AppendfImpl(FString& AppendToMe, const TCHAR* Fmt, ...);
 public:
 
 	/**
@@ -1410,7 +1410,7 @@ public:
 	 * @param InNamedArguments		A map of named arguments that match the tokens specified in InExpression
 	 * @return A string containing the formatted text
 	 */
-	UE_NODISCARD static FString Format(const TCHAR* InFormatString, const FStringFormatNamedArguments& InNamedArguments);
+	UE_NODISCARD static CORE_API FString Format(const TCHAR* InFormatString, const FStringFormatNamedArguments& InNamedArguments);
 
 	/**
 	 * Format the specified string using the specified arguments. Replaces instances of {0} with indices from the given array matching the index specified in the token
@@ -1418,10 +1418,10 @@ public:
 	 * @param InOrderedArguments	An array of ordered arguments that match the tokens specified in InExpression
 	 * @return A string containing the formatted text
 	 */
-	UE_NODISCARD static FString Format(const TCHAR* InFormatString, const FStringFormatOrderedArguments& InOrderedArguments);
+	UE_NODISCARD static CORE_API FString Format(const TCHAR* InFormatString, const FStringFormatOrderedArguments& InOrderedArguments);
 
 	/** Returns a string containing only the Ch character */
-	UE_NODISCARD static FString Chr( TCHAR Ch );
+	UE_NODISCARD static CORE_API FString Chr( TCHAR Ch );
 
 	/**
 	 * Returns a string that is full of a variable number of characters
@@ -1431,7 +1431,7 @@ public:
 	 * 
 	 * @return The string of NumCharacters characters.
 	 */
-	UE_NODISCARD static FString ChrN( int32 NumCharacters, TCHAR Char );
+	UE_NODISCARD static CORE_API FString ChrN( int32 NumCharacters, TCHAR Char );
 
 	/**
 	 * Serializes the string.
@@ -1485,7 +1485,7 @@ public:
 	 * @param SearchCase		Indicates whether the search is case sensitive or not ( defaults to ESearchCase::IgnoreCase )
 	 * @return true if this string begins with specified text, false otherwise
 	 */
-	UE_NODISCARD bool StartsWith(const TCHAR* InPrefix, int32 InPrefixLen, ESearchCase::Type SearchCase = ESearchCase::IgnoreCase) const;
+	UE_NODISCARD CORE_API bool StartsWith(const TCHAR* InPrefix, int32 InPrefixLen, ESearchCase::Type SearchCase = ESearchCase::IgnoreCase) const;
 
 	/**
 	 * Test whether this string ends with given suffix.
@@ -1528,7 +1528,7 @@ public:
 	 * @param SearchCase		Indicates whether the search is case sensitive or not ( defaults to ESearchCase::IgnoreCase )
 	 * @return true if this string ends with specified text, false otherwise
 	 */
-	UE_NODISCARD bool EndsWith(const TCHAR* InSuffix, int32 InSuffixLen, ESearchCase::Type SearchCase = ESearchCase::IgnoreCase ) const;
+	UE_NODISCARD CORE_API bool EndsWith(const TCHAR* InSuffix, int32 InSuffixLen, ESearchCase::Type SearchCase = ESearchCase::IgnoreCase ) const;
 
 	/**
 	 * Searches this string for a given wild card
@@ -1579,96 +1579,96 @@ public:
 	 * @return true if this string matches the *?-type wildcard given.
 	 * @warning This is a simple, SLOW routine. Use with caution
 	 */
-	UE_NODISCARD bool MatchesWildcard(const TCHAR* Wildcard, int32 WildcardLen, ESearchCase::Type SearchCase = ESearchCase::IgnoreCase) const;
+	UE_NODISCARD CORE_API bool MatchesWildcard(const TCHAR* Wildcard, int32 WildcardLen, ESearchCase::Type SearchCase = ESearchCase::IgnoreCase) const;
 
 	/**
 	 * Removes whitespace characters from the start and end of this string. Modifies the string in-place.
 	 */
-	void TrimStartAndEndInline();
+	CORE_API void TrimStartAndEndInline();
 
 	/**
 	 * Removes whitespace characters from the start and end of this string.
 	 * @note Unlike Trim() this function returns a copy, and does not mutate the string.
 	 */
-	UE_NODISCARD FString TrimStartAndEnd() const &;
+	UE_NODISCARD CORE_API FString TrimStartAndEnd() const &;
 
 	/**
 	 * Removes whitespace characters from the start and end of this string.
 	 * @note Unlike Trim() this function returns a copy, and does not mutate the string.
 	 */
-	UE_NODISCARD FString TrimStartAndEnd() &&;
+	UE_NODISCARD CORE_API FString TrimStartAndEnd() &&;
 
 	/**
 	 * Removes whitespace characters from the start of this string. Modifies the string in-place.
 	 */
-	void TrimStartInline();
+	CORE_API void TrimStartInline();
 
 	/**
 	 * Removes whitespace characters from the start of this string.
 	 * @note Unlike Trim() this function returns a copy, and does not mutate the string.
 	 */
-	UE_NODISCARD FString TrimStart() const &;
+	UE_NODISCARD CORE_API FString TrimStart() const &;
 
 	/**
 	 * Removes whitespace characters from the start of this string.
 	 * @note Unlike Trim() this function returns a copy, and does not mutate the string.
 	 */
-	UE_NODISCARD FString TrimStart() &&;
+	UE_NODISCARD CORE_API FString TrimStart() &&;
 
 	/**
 	 * Removes whitespace characters from the end of this string. Modifies the string in-place.
 	 */
-	void TrimEndInline();
+	CORE_API void TrimEndInline();
 
 	/**
 	 * Removes whitespace characters from the end of this string.
 	 * @note Unlike TrimTrailing() this function returns a copy, and does not mutate the string.
 	 */
-	UE_NODISCARD FString TrimEnd() const &;
+	UE_NODISCARD CORE_API FString TrimEnd() const &;
 
 	/**
 	 * Removes whitespace characters from the end of this string.
 	 * @note Unlike TrimTrailing() this function returns a copy, and does not mutate the string.
 	 */
-	UE_NODISCARD FString TrimEnd() &&;
+	UE_NODISCARD CORE_API FString TrimEnd() &&;
 
 	/** 
 	 * Trims the inner array after the null terminator.
 	 */
-	void TrimToNullTerminator();
+	CORE_API void TrimToNullTerminator();
 
 
 	/**
 	 * Trims wrapping quotation marks from this string.
 	 */
-	void TrimQuotesInline(bool* bQuotesRemoved = nullptr);
+	CORE_API void TrimQuotesInline(bool* bQuotesRemoved = nullptr);
 
 	/**
 	* Trims a single character from the start and end of the string (removes at max one instance in the beginning and end of the string).
 	* @see TrimChar for a variant that returns a modified copy of the string
 	*/
-	void TrimCharInline(const TCHAR CharacterToTrim, bool* bCharRemoved);
+	CORE_API void TrimCharInline(const TCHAR CharacterToTrim, bool* bCharRemoved);
 	
 	/**
 	 * Returns a copy of this string with wrapping quotation marks removed.
 	 */
-	UE_NODISCARD FString TrimQuotes(bool* bQuotesRemoved = nullptr ) const &;
+	UE_NODISCARD CORE_API FString TrimQuotes(bool* bQuotesRemoved = nullptr ) const &;
 
 	/**
 	 * Returns this string with wrapping quotation marks removed.
 	 */
-	UE_NODISCARD FString TrimQuotes(bool* bQuotesRemoved = nullptr) &&;
+	UE_NODISCARD CORE_API FString TrimQuotes(bool* bQuotesRemoved = nullptr) &&;
 	
 	/**
 	* Returns a copy of this string with wrapping CharacterToTrim removed (removes at max one instance in the beginning and end of the string).
 	* @see TrimCharInline for an inline variant
 	*/
-	UE_NODISCARD FString TrimChar(const TCHAR CharacterToTrim, bool* bCharRemoved = nullptr ) const &;
+	UE_NODISCARD CORE_API FString TrimChar(const TCHAR CharacterToTrim, bool* bCharRemoved = nullptr ) const &;
 
 	/**
 	* Returns a copy of this string with wrapping CharacterToTrim removed (removes at max one instance in the beginning and end of the string).
 	*/
-	UE_NODISCARD FString TrimChar(const TCHAR CharacterToTrim, bool* bCharRemoved = nullptr) &&;
+	UE_NODISCARD CORE_API FString TrimChar(const TCHAR CharacterToTrim, bool* bCharRemoved = nullptr) &&;
 	
 	/**
 	 * Breaks up a delimited string into elements of a string array.
@@ -1679,7 +1679,7 @@ public:
 	 *
 	 * @return	The number of elements in InArray
 	 */
-	int32 ParseIntoArray( TArray<FString>& OutArray, const TCHAR* pchDelim, bool InCullEmpty = true ) const;
+	CORE_API int32 ParseIntoArray( TArray<FString>& OutArray, const TCHAR* pchDelim, bool InCullEmpty = true ) const;
 
 	/**
 	 * Breaks up a delimited string into elements of a string array, using any whitespace and an 
@@ -1691,7 +1691,7 @@ public:
 	 *
 	 * @return	The number of elements in InArray
 	 */
-	int32 ParseIntoArrayWS( TArray<FString>& OutArray, const TCHAR* pchExtraDelim = nullptr, bool InCullEmpty = true ) const;
+	CORE_API int32 ParseIntoArrayWS( TArray<FString>& OutArray, const TCHAR* pchExtraDelim = nullptr, bool InCullEmpty = true ) const;
 
 	/**
 	* Breaks up a delimited string into elements of a string array, using line ending characters
@@ -1701,7 +1701,7 @@ public:
 	*
 	* @return	The number of elements in InArray
 	*/
-	int32 ParseIntoArrayLines(TArray<FString>& OutArray, bool InCullEmpty = true) const;
+	CORE_API int32 ParseIntoArrayLines(TArray<FString>& OutArray, bool InCullEmpty = true) const;
 
 	/**
 	* Breaks up a delimited string into elements of a string array, using the given delimiters
@@ -1713,7 +1713,7 @@ public:
 	*
 	* @return	The number of elements in InArray
 	*/
-	int32 ParseIntoArray(TArray<FString>& OutArray, const TCHAR*const* DelimArray, int32 NumDelims, bool InCullEmpty = true) const;
+	CORE_API int32 ParseIntoArray(TArray<FString>& OutArray, const TCHAR*const* DelimArray, int32 NumDelims, bool InCullEmpty = true) const;
 
 	/**
 	 * Takes an array of strings and removes any zero length entries.
@@ -1722,22 +1722,22 @@ public:
 	 *
 	 * @return	The number of elements left in InArray
 	 */
-	static int32 CullArray( TArray<FString>* InArray );
+	static CORE_API int32 CullArray( TArray<FString>* InArray );
 
 	/**
 	 * Returns a copy of this string, with the characters in reverse order
 	 */
-	UE_NODISCARD FString Reverse() const &;
+	UE_NODISCARD CORE_API FString Reverse() const &;
 
 	/**
 	 * Returns this string, with the characters in reverse order
 	 */
-	UE_NODISCARD FString Reverse() &&;
+	UE_NODISCARD CORE_API FString Reverse() &&;
 
 	/**
 	 * Reverses the order of characters in this string
 	 */
-	void ReverseString();
+	CORE_API void ReverseString();
 
 	/**
 	 * Replace all occurrences of a substring in this string
@@ -1747,7 +1747,7 @@ public:
 	 * @param SearchCase	Indicates whether the search is case sensitive or not ( defaults to ESearchCase::IgnoreCase )
 	 * @return a copy of this string with the replacement made
 	 */
-	UE_NODISCARD FString Replace(const TCHAR* From, const TCHAR* To, ESearchCase::Type SearchCase = ESearchCase::IgnoreCase) const &;
+	UE_NODISCARD CORE_API FString Replace(const TCHAR* From, const TCHAR* To, ESearchCase::Type SearchCase = ESearchCase::IgnoreCase) const &;
 
 	/**
 	 * Replace all occurrences of a substring in this string
@@ -1757,7 +1757,7 @@ public:
 	 * @param SearchCase	Indicates whether the search is case sensitive or not ( defaults to ESearchCase::IgnoreCase )
 	 * @return a copy of this string with the replacement made
 	 */
-	UE_NODISCARD FString Replace(const TCHAR* From, const TCHAR* To, ESearchCase::Type SearchCase = ESearchCase::IgnoreCase) &&;
+	UE_NODISCARD CORE_API FString Replace(const TCHAR* From, const TCHAR* To, ESearchCase::Type SearchCase = ESearchCase::IgnoreCase) &&;
 
 	/**
 	 * Replace all occurrences of SearchText with ReplacementText in this string.
@@ -1768,7 +1768,7 @@ public:
 	 *
 	 * @return	the number of occurrences of SearchText that were replaced.
 	 */
-	int32 ReplaceInline( const TCHAR* SearchText, const TCHAR* ReplacementText, ESearchCase::Type SearchCase = ESearchCase::IgnoreCase );
+	CORE_API int32 ReplaceInline( const TCHAR* SearchText, const TCHAR* ReplacementText, ESearchCase::Type SearchCase = ESearchCase::IgnoreCase );
 
 	/**
 	 * Replace all occurrences of a character with another.
@@ -1792,8 +1792,8 @@ public:
 
 private:
 	
-	void ReplaceCharInlineCaseSensitive(const TCHAR SearchChar, const TCHAR ReplacementChar);
-	void ReplaceCharInlineIgnoreCase(const TCHAR SearchChar, const TCHAR ReplacementChar);
+	CORE_API void ReplaceCharInlineCaseSensitive(const TCHAR SearchChar, const TCHAR ReplacementChar);
+	CORE_API void ReplaceCharInlineIgnoreCase(const TCHAR SearchChar, const TCHAR ReplacementChar);
 
 public:
 
@@ -1809,7 +1809,7 @@ public:
 	/**
 	 * Returns a copy of this string with all quote marks escaped (unless the quote is already escaped)
 	 */
-	UE_NODISCARD FString ReplaceQuotesWithEscapedQuotes() &&;
+	UE_NODISCARD CORE_API FString ReplaceQuotesWithEscapedQuotes() &&;
 
 	/**
 	 * Replaces certain characters with the "escaped" version of that character (i.e. replaces "\n" with "\\n").
@@ -1817,7 +1817,7 @@ public:
 	 *
 	 * @param	Chars	by default, replaces all supported characters; this parameter allows you to limit the replacement to a subset.
 	 */
-	void ReplaceCharWithEscapedCharInline( const TArray<TCHAR>* Chars = nullptr );
+	CORE_API void ReplaceCharWithEscapedCharInline( const TArray<TCHAR>* Chars = nullptr );
 
 	/**
 	 * Replaces certain characters with the "escaped" version of that character (i.e. replaces "\n" with "\\n").
@@ -1852,7 +1852,7 @@ public:
 	 * Removes the escape backslash for all supported characters, replacing the escape and character with the non-escaped version.  (i.e.
 	 * replaces "\\n" with "\n".  Counterpart to ReplaceCharWithEscapedCharInline().
 	 */
-	void ReplaceEscapedCharWithCharInline( const TArray<TCHAR>* Chars = nullptr );
+	CORE_API void ReplaceEscapedCharWithCharInline( const TArray<TCHAR>* Chars = nullptr );
 
 	/**
 	 * Removes the escape backslash for all supported characters, replacing the escape and character with the non-escaped version.  (i.e.
@@ -1881,7 +1881,7 @@ public:
 	 * Replaces all instances of '\t' with TabWidth number of spaces
 	 * @param InSpacesPerTab - Number of spaces that a tab represents
 	 */
-	void ConvertTabsToSpacesInline(const int32 InSpacesPerTab);
+	CORE_API void ConvertTabsToSpacesInline(const int32 InSpacesPerTab);
 
 	/** 
 	 * Replaces all instances of '\t' with TabWidth number of spaces
@@ -1907,10 +1907,10 @@ public:
 	}
 
 	// Takes the number passed in and formats the string in comma format ( 12345 becomes "12,345")
-	UE_NODISCARD static FString FormatAsNumber( int32 InNumber );
+	UE_NODISCARD static CORE_API FString FormatAsNumber( int32 InNumber );
 
 	// To allow more efficient memory handling, automatically adds one for the string termination.
-	void Reserve(int32 CharacterCount);
+	CORE_API void Reserve(int32 CharacterCount);
 
 	/**
 	 * Serializes a string as ANSI char array.
@@ -1919,7 +1919,7 @@ public:
 	 * @param	Ar				Archive to serialize with
 	 * @param	MinCharacters	Minimum number of characters to serialize.
 	 */
-	void SerializeAsANSICharArray( FArchive& Ar, int32 MinCharacters=0 ) const;
+	CORE_API void SerializeAsANSICharArray( FArchive& Ar, int32 MinCharacters=0 ) const;
 
 
 	/** Converts an integer to a string. */
@@ -1931,7 +1931,7 @@ public:
 	}
 
 	/** appends the integer InNum to this string */
-	void AppendInt( int32 InNum );
+	CORE_API void AppendInt( int32 InNum );
 
 	/**
 	 * Converts a string into a boolean value
@@ -1940,7 +1940,7 @@ public:
 	 *
 	 * @return The boolean value
 	 */
-	UE_NODISCARD bool ToBool() const;
+	UE_NODISCARD CORE_API bool ToBool() const;
 
 	/**
 	 * Converts a buffer to a string
@@ -1950,7 +1950,7 @@ public:
 	 *
 	 * @return the blob in string form
 	 */
-	UE_NODISCARD static FString FromBlob(const uint8* SrcBuffer,const uint32 SrcSize);
+	UE_NODISCARD static CORE_API FString FromBlob(const uint8* SrcBuffer,const uint32 SrcSize);
 
 	/**
 	 * Converts a string into a buffer
@@ -1960,7 +1960,7 @@ public:
 	 *
 	 * @return true if the conversion happened, false otherwise
 	 */
-	static bool ToBlob(const FString& Source,uint8* DestBuffer,const uint32 DestSize);
+	static CORE_API bool ToBlob(const FString& Source,uint8* DestBuffer,const uint32 DestSize);
 
 	/**
 	 * Converts a buffer to a string by hex-ifying the elements
@@ -1970,7 +1970,7 @@ public:
 	 *
 	 * @return the blob in string form
 	 */
-	UE_NODISCARD static FString FromHexBlob(const uint8* SrcBuffer,const uint32 SrcSize);
+	UE_NODISCARD static CORE_API FString FromHexBlob(const uint8* SrcBuffer,const uint32 SrcSize);
 
 	/**
 	 * Converts a string into a buffer
@@ -1980,7 +1980,7 @@ public:
 	 *
 	 * @return true if the conversion happened, false otherwise
 	 */
-	static bool ToHexBlob(const FString& Source,uint8* DestBuffer,const uint32 DestSize);
+	static CORE_API bool ToHexBlob(const FString& Source,uint8* DestBuffer,const uint32 DestSize);
 
 	/**
 	 * Converts a float string with the trailing zeros stripped
@@ -1991,7 +1991,7 @@ public:
 	 *
 	 * @return sanitized string version of float
 	 */
-	UE_NODISCARD static FString SanitizeFloat( double InFloat, const int32 InMinFractionalDigits = 1 );
+	UE_NODISCARD static CORE_API FString SanitizeFloat( double InFloat, const int32 InMinFractionalDigits = 1 );
 
 	/**
 	 * Joins a range of 'something that can be concatentated to strings with +=' together into a single string with separators.
@@ -2531,7 +2531,7 @@ CORE_API int32 FindMatchingClosingParenthesis(const FString& TargetString, const
 */
 CORE_API FString SlugStringForValidName(const FString& DisplayString, const TCHAR* ReplaceWith = TEXT(""));
 
-struct CORE_API FTextRange
+struct FTextRange
 {
 	FTextRange()
 		: BeginIndex(INDEX_NONE)
@@ -2586,7 +2586,7 @@ struct CORE_API FTextRange
 	/**
 	 * Produce an array of line ranges from the given text, breaking at any new-line characters
 	 */
-	static void CalculateLineRangesFromString(const FString& Input, TArray<FTextRange>& LineRanges);
+	static CORE_API void CalculateLineRangesFromString(const FString& Input, TArray<FTextRange>& LineRanges);
 
 	int32 BeginIndex;
 	int32 EndIndex;
diff --git a/Engine/Source/Runtime/Core/Public/CoreGlobals.h b/Engine/Source/Runtime/Core/Public/CoreGlobals.h
--- a/Engine/Source/Runtime/Core/Public/CoreGlobals.h
+++ b/Engine/Source/Runtime/Core/Public/CoreGlobals.h
@@ -71,20 +71,20 @@ CORE_API void BootTimingPoint(const ANSICHAR *Message);
 
 CORE_API void DumpBootTiming();
 
-struct CORE_API FScopedBootTiming
+struct FScopedBootTiming
 {
 	FString Message;
 	double StartTime;
-	FScopedBootTiming(const ANSICHAR *InMessage);
-	FScopedBootTiming(const ANSICHAR *InMessage, FName Suffix);
-	~FScopedBootTiming();
+	CORE_API FScopedBootTiming(const ANSICHAR *InMessage);
+	CORE_API FScopedBootTiming(const ANSICHAR *InMessage, FName Suffix);
+	CORE_API ~FScopedBootTiming();
 };
 
-struct CORE_API FEngineTrackedActivityScope
+struct FEngineTrackedActivityScope
 {
-	FEngineTrackedActivityScope(const TCHAR* Fmt, ...);
-	FEngineTrackedActivityScope(const FString& Str);
-	~FEngineTrackedActivityScope();
+	CORE_API FEngineTrackedActivityScope(const TCHAR* Fmt, ...);
+	CORE_API FEngineTrackedActivityScope(const FString& Str);
+	CORE_API ~FEngineTrackedActivityScope();
 };
 
 
@@ -105,7 +105,7 @@ extern CORE_API TCHAR GErrorHist[16384];
 // #crashReport: 2014-08-19 Combine into one, refactor.
 extern CORE_API TCHAR GErrorExceptionDescription[4096];
 
-struct CORE_API FCoreTexts
+struct FCoreTexts
 {
 	const FText& True;
 	const FText& False;
@@ -113,10 +113,10 @@ struct CORE_API FCoreTexts
 	const FText& No;
 	const FText& None;
 
-	static const FCoreTexts& Get();
+	static CORE_API const FCoreTexts& Get();
 
 	/** Invalidates existing references. Do not use FCoreTexts after calling. */
-	static void TearDown();
+	static CORE_API void TearDown();
 
 	// Non-copyable
 	FCoreTexts(const FCoreTexts&) = delete;
@@ -310,11 +310,11 @@ FORCEINLINE bool IsAllowCommandletAudio()
 #endif
 }
 
-class CORE_API FIsDuplicatingClassForReinstancing
+class FIsDuplicatingClassForReinstancing
 {
 public:
-	FIsDuplicatingClassForReinstancing& operator= (bool bOther);
-	operator bool() const;
+	CORE_API FIsDuplicatingClassForReinstancing& operator= (bool bOther);
+	CORE_API operator bool() const;
 };
 
 extern CORE_API bool GEdSelectionLock;
@@ -450,11 +450,11 @@ extern CORE_API bool GIsCookerLoadingPackage;
 /** Whether GWorld points to the play in editor world */
 extern CORE_API bool GIsPlayInEditorWorld;
 
-class CORE_API FPlayInEditorID
+class FPlayInEditorID
 {
 public:
-	FPlayInEditorID& operator= (int32 InOther);
-	operator int32() const;
+	CORE_API FPlayInEditorID& operator= (int32 InOther);
+	CORE_API operator int32() const;
 };
 /** In the editor, this is set to the specific world context PIEInstance that is being currently processed */
 extern CORE_API FPlayInEditorID GPlayInEditorID;
@@ -639,7 +639,7 @@ protected:
 	CORE_API FTaskTagScope(bool InTagOnlyIfNone, ETaskTag InTag);
 
 public:
-	CORE_API FTaskTagScope(ETaskTag InTag = ETaskTag::ENone) : FTaskTagScope(false, InTag)
+	FTaskTagScope(ETaskTag InTag = ETaskTag::ENone) : FTaskTagScope(false, InTag)
 	{
 
 	}
@@ -660,7 +660,7 @@ public:
 class FOptionalTaskTagScope : public FTaskTagScope
 {
 public:
-	CORE_API FOptionalTaskTagScope(ETaskTag InTag = ETaskTag::ENone) : FTaskTagScope(true, InTag)
+	FOptionalTaskTagScope(ETaskTag InTag = ETaskTag::ENone) : FTaskTagScope(true, InTag)
 	{
 
 	}
diff --git a/Engine/Source/Runtime/Core/Public/Experimental/Containers/HazardPointer.h b/Engine/Source/Runtime/Core/Public/Experimental/Containers/HazardPointer.h
--- a/Engine/Source/Runtime/Core/Public/Experimental/Containers/HazardPointer.h
+++ b/Engine/Source/Runtime/Core/Public/Experimental/Containers/HazardPointer.h
@@ -160,7 +160,7 @@ class FHazardPointerCollection
 	CORE_API void Delete(const HazardPointer_Impl::FHazardDeleter& Deleter, int32 CollectLimit);
 
 	template<bool Cached>
-	CORE_API FHazardRecord* Grow();
+	FHazardRecord* Grow();
 
 public:
 	FHazardPointerCollection()
diff --git a/Engine/Source/Runtime/Core/Public/Experimental/Misc/ExecutionResource.h b/Engine/Source/Runtime/Core/Public/Experimental/Misc/ExecutionResource.h
--- a/Engine/Source/Runtime/Core/Public/Experimental/Misc/ExecutionResource.h
+++ b/Engine/Source/Runtime/Core/Public/Experimental/Misc/ExecutionResource.h
@@ -16,18 +16,18 @@ public:
  * Per-thread stack-frame of execution resource currently held.
  * Can be retrieved and passed around to other task until not needed anymore.
  */
-class CORE_API FExecutionResourceContext
+class FExecutionResourceContext
 {
 public:
-	static TRefCountPtr<IExecutionResource> Get();
+	static CORE_API TRefCountPtr<IExecutionResource> Get();
 };
 
 /**
  * Used to push an execution resource on the stack-frame that can be retrieved by FExecutionResourceContext::Get().
  */
-class CORE_API FExecutionResourceContextScope
+class FExecutionResourceContextScope
 {
 public:
-	FExecutionResourceContextScope(TRefCountPtr<IExecutionResource> InExecutionResource);
-	~FExecutionResourceContextScope();
-};
\ No newline at end of file
+	CORE_API FExecutionResourceContextScope(TRefCountPtr<IExecutionResource> InExecutionResource);
+	CORE_API ~FExecutionResourceContextScope();
+};
diff --git a/Engine/Source/Runtime/Core/Public/FramePro/FrameProProfiler.h b/Engine/Source/Runtime/Core/Public/FramePro/FrameProProfiler.h
--- a/Engine/Source/Runtime/Core/Public/FramePro/FrameProProfiler.h
+++ b/Engine/Source/Runtime/Core/Public/FramePro/FrameProProfiler.h
@@ -7,31 +7,31 @@
 #if FRAMEPRO_ENABLED
 
 /** Wrapper for FramePro  */
-class CORE_API FFrameProProfiler
+class FFrameProProfiler
 {
 public:
-	static void Initialize();
-	static void TearDown();
+	static CORE_API void Initialize();
+	static CORE_API void TearDown();
 
 	/** Called to mark the start of each frame  */
-	static void FrameStart();
+	static CORE_API void FrameStart();
 
 	/** Begin a named event */
-	static void PushEvent(); // Event with no name, expected to be named at the end
-	static void PushEvent(const TCHAR* Text);
-	static void PushEvent(const ANSICHAR* Text);
+	static CORE_API void PushEvent(); // Event with no name, expected to be named at the end
+	static CORE_API void PushEvent(const TCHAR* Text);
+	static CORE_API void PushEvent(const ANSICHAR* Text);
 
 	/** End currently active named event */
-	static void PopEvent();
-	static void PopEvent(const TCHAR* Override);
-	static void PopEvent(const ANSICHAR* Override);
+	static CORE_API void PopEvent();
+	static CORE_API void PopEvent(const TCHAR* Override);
+	static CORE_API void PopEvent(const ANSICHAR* Override);
 
-	static void StartFrameProRecordingFromCommand(const TArray< FString >& Args);
-	static FString StartFrameProRecording(const FString& FilenameRoot, int32 MinScopeTime, bool bAppendDateTime=true);
-	static void StopFrameProRecording();
+	static CORE_API void StartFrameProRecordingFromCommand(const TArray< FString >& Args);
+	static CORE_API FString StartFrameProRecording(const FString& FilenameRoot, int32 MinScopeTime, bool bAppendDateTime=true);
+	static CORE_API void StopFrameProRecording();
 
-	static bool IsFrameProRecording();
-	static bool IsThreadContextReady();
+	static CORE_API bool IsFrameProRecording();
+	static CORE_API bool IsThreadContextReady();
 };
 
-#endif // FRAMEPRO_ENABLED
\ No newline at end of file
+#endif // FRAMEPRO_ENABLED
diff --git a/Engine/Source/Runtime/Core/Public/GenericPlatform/GenericPlatformAffinity.h b/Engine/Source/Runtime/Core/Public/GenericPlatform/GenericPlatformAffinity.h
--- a/Engine/Source/Runtime/Core/Public/GenericPlatform/GenericPlatformAffinity.h
+++ b/Engine/Source/Runtime/Core/Public/GenericPlatform/GenericPlatformAffinity.h
@@ -45,79 +45,79 @@ ENUM_CLASS_FLAGS(EThreadCreateFlags);
 class FGenericPlatformAffinity
 {
 public:
-	static const CORE_API uint64 GetMainGameMask()
+	static const uint64 GetMainGameMask()
 	{
 		return 0xFFFFFFFFFFFFFFFF;
 	}
 
-	static const CORE_API uint64 GetRenderingThreadMask()
+	static const uint64 GetRenderingThreadMask()
 	{
 		return 0xFFFFFFFFFFFFFFFF;
 	}
 
-	static const CORE_API uint64 GetRHIThreadMask()
+	static const uint64 GetRHIThreadMask()
 	{
 		return 0xFFFFFFFFFFFFFFFF;
 	}
 
-	static const CORE_API uint64 GetRHIFrameOffsetThreadMask()
+	static const uint64 GetRHIFrameOffsetThreadMask()
 	{
 		return 0xFFFFFFFFFFFFFFFF;
 	}
 
-	static const CORE_API uint64 GetRTHeartBeatMask()
+	static const uint64 GetRTHeartBeatMask()
 	{
 		return 0xFFFFFFFFFFFFFFFF;
 	}
 
-	static const CORE_API uint64 GetPoolThreadMask()
+	static const uint64 GetPoolThreadMask()
 	{
 		return 0xFFFFFFFFFFFFFFFF;
 	}
 
-	static const CORE_API uint64 GetTaskGraphThreadMask()
+	static const uint64 GetTaskGraphThreadMask()
 	{
 		return 0xFFFFFFFFFFFFFFFF;
 	}
 
-	static const CORE_API uint64 GetAudioRenderThreadMask()
+	static const uint64 GetAudioRenderThreadMask()
 	{
 		return 0xFFFFFFFFFFFFFFFF;
 	}
 
-	static const CORE_API uint64 GetNoAffinityMask()
+	static const uint64 GetNoAffinityMask()
 	{
 		return 0xFFFFFFFFFFFFFFFF;
 	}
 
-	static const CORE_API uint64 GetTaskGraphBackgroundTaskMask()
+	static const uint64 GetTaskGraphBackgroundTaskMask()
 	{
 		return 0xFFFFFFFFFFFFFFFF;
 	}
 
-	static const CORE_API uint64 GetTaskGraphHighPriorityTaskMask()
+	static const uint64 GetTaskGraphHighPriorityTaskMask()
 	{
 		return 0xFFFFFFFFFFFFFFFF;
 	}
 
-	static const CORE_API uint64 GetAsyncLoadingThreadMask()
+	static const uint64 GetAsyncLoadingThreadMask()
 	{
 		return 0xFFFFFFFFFFFFFFFF;
 	}
 
-	static const CORE_API uint64 GetIoDispatcherThreadMask()
+	static const uint64 GetIoDispatcherThreadMask()
 	{
 		return 0xFFFFFFFFFFFFFFFF;
 	}
 
-	static const CORE_API uint64 GetTraceThreadMask()
+	static const uint64 GetTraceThreadMask()
 	{
 		return 0xFFFFFFFFFFFFFFFF;
 	}
 
 	// @todo what do we think about having this as a function in this class? Should be make a whole new one? 
 	// scrap it and force the priority like before?
-	static CORE_API EThreadPriority GetRenderingThreadPriority()
+	static EThreadPriority GetRenderingThreadPriority()
 	{
 		return TPri_Normal;
 	}
@@ -127,7 +127,7 @@ public:
 		return EThreadCreateFlags::None;
 	}
 	
-	static CORE_API EThreadPriority GetRHIThreadPriority()
+	static EThreadPriority GetRHIThreadPriority()
 	{
 		return TPri_Normal;
 	}
diff --git a/Engine/Source/Runtime/Core/Public/GenericPlatform/GenericPlatformChunkInstall.h b/Engine/Source/Runtime/Core/Public/GenericPlatform/GenericPlatformChunkInstall.h
--- a/Engine/Source/Runtime/Core/Public/GenericPlatform/GenericPlatformChunkInstall.h
+++ b/Engine/Source/Runtime/Core/Public/GenericPlatform/GenericPlatformChunkInstall.h
@@ -132,7 +132,7 @@ enum class ENamedChunkType : uint8
 /**
 * Interface for platform specific chunk based install
 **/
-class CORE_API IPlatformChunkInstall
+class IPlatformChunkInstall
 {
 public:
 
@@ -513,13 +513,13 @@ protected:
 
 
 // temporary helper base class for platform chunk installers that have implemented named chunk support to provide FCustomChunk emulation
-class CORE_API FGenericPlatformChunkInstall_WithEmulatedCustomChunks : public FGenericPlatformChunkInstall
+class FGenericPlatformChunkInstall_WithEmulatedCustomChunks : public FGenericPlatformChunkInstall
 {
 public:
 	virtual bool SupportsIntelligentInstall() override final { return true; }
-	virtual bool IsChunkInstallationPending(const TArray<FCustomChunk>& ChunkTagsID) override final;
-	virtual bool InstallChunks(const TArray<FCustomChunk>& ChunkTagsID)  override final;
-	virtual bool UninstallChunks(const TArray<FCustomChunk>& ChunkTagsID) override final;
+	CORE_API virtual bool IsChunkInstallationPending(const TArray<FCustomChunk>& ChunkTagsID) override final;
+	CORE_API virtual bool InstallChunks(const TArray<FCustomChunk>& ChunkTagsID)  override final;
+	CORE_API virtual bool UninstallChunks(const TArray<FCustomChunk>& ChunkTagsID) override final;
 };
 
-PRAGMA_ENABLE_DEPRECATION_WARNINGS
\ No newline at end of file
+PRAGMA_ENABLE_DEPRECATION_WARNINGS
diff --git a/Engine/Source/Runtime/Core/Public/GenericPlatform/GenericPlatformCrashContext.h b/Engine/Source/Runtime/Core/Public/GenericPlatform/GenericPlatformCrashContext.h
--- a/Engine/Source/Runtime/Core/Public/GenericPlatform/GenericPlatformCrashContext.h
+++ b/Engine/Source/Runtime/Core/Public/GenericPlatform/GenericPlatformCrashContext.h
@@ -264,10 +264,10 @@ struct FSharedCrashContext
 struct FCrashContextExtendedWriter
 {
 	/** Adds a named buffer to the report. Intended for larger payloads. */
-	CORE_API virtual void AddBuffer(const TCHAR* Identifier, const uint8* Data, uint32 DataSize) = 0;
+	virtual void AddBuffer(const TCHAR* Identifier, const uint8* Data, uint32 DataSize) = 0;
 
 	/** Add a named buffer containing a string to the report. */
-	CORE_API virtual void AddString(const TCHAR* Identifier, const TCHAR* DataStr) = 0;
+	virtual void AddString(const TCHAR* Identifier, const TCHAR* DataStr) = 0;
 };
 
 /** Simple Delegate for additional crash context. */
@@ -461,13 +461,13 @@ public:
 	CORE_API static void ResetEngineData();
 
 	/** Accessor for engine data reset callback delegate */
-	CORE_API static FEngineDataResetDelegate& OnEngineDataResetDelegate() { return OnEngineDataReset; }
+	static FEngineDataResetDelegate& OnEngineDataResetDelegate() { return OnEngineDataReset; }
 
 	/** Updates (or adds if not already present) arbitrary engine data to the crash context (will remove the key if passed an empty string) */
 	CORE_API static void SetEngineData(const FString& Key, const FString& Value);
 
 	/** Accessor for engine data change callback delegate */
-	CORE_API static FEngineDataSetDelegate& OnEngineDataSetDelegate() { return OnEngineDataSet; }
+	static FEngineDataSetDelegate& OnEngineDataSetDelegate() { return OnEngineDataSet; }
 
 	/** Get the engine data dictionary */
 	CORE_API static const TMap<FString, FString>& GetEngineData();
@@ -476,13 +476,13 @@ public:
 	CORE_API static void ResetGameData();
 
 	/** Accessor for game data reset callback delegate */
-	CORE_API static FGameDataResetDelegate& OnGameDataResetDelegate() { return OnGameDataReset; }
+	static FGameDataResetDelegate& OnGameDataResetDelegate() { return OnGameDataReset; }
 
 	/** Updates (or adds if not already present) arbitrary game data to the crash context (will remove the key if passed an empty string) */
 	CORE_API static void SetGameData(const FString& Key, const FString& Value);
 
 	/** Accessor for game data change callback delegate */
-	CORE_API static FGameDataSetDelegate& OnGameDataSetDelegate() { return OnGameDataSet; }
+	static FGameDataSetDelegate& OnGameDataSetDelegate() { return OnGameDataSet; }
 
 	/** Get the game data dictionary */
 	CORE_API static const TMap<FString, FString>& GetGameData();
@@ -707,7 +707,7 @@ private:
 	FGenericCrashContext& operator=(const FGenericCrashContext&) = delete;
 };
 
-struct CORE_API FGenericMemoryWarningContext
+struct FGenericMemoryWarningContext
 {};
 
 namespace RecoveryService
