[Core]

diff --git a/Engine/Source/Runtime/Core/Public/GenericPlatform/GenericPlatformFile.h b/Engine/Source/Runtime/Core/Public/GenericPlatform/GenericPlatformFile.h
--- a/Engine/Source/Runtime/Core/Public/GenericPlatform/GenericPlatformFile.h
+++ b/Engine/Source/Runtime/Core/Public/GenericPlatform/GenericPlatformFile.h
@@ -92,7 +92,7 @@ ENUM_CLASS_FLAGS(EDirectoryVisitorFlags);
 /** 
  * File handle interface. 
 **/
-class CORE_API IFileHandle
+class IFileHandle
 {
 public:
 	/** Destructor, also the only way to close the file handle **/
@@ -158,7 +158,7 @@ public:
 	/////////// Utility Functions. These have a default implementation that uses the pure virtual operations.
 
 	/** Return the total size of the file **/
-	virtual int64		Size();
+	CORE_API virtual int64		Size();
 };
 
 
@@ -215,13 +215,13 @@ struct FFileStatData
 /**
 * File I/O Interface
 **/
-class CORE_API IPlatformFile
+class IPlatformFile
 {
 public:
 	/** Physical file system of the _platform_, never wrapped. **/
-	static IPlatformFile& GetPlatformPhysical();
+	static CORE_API IPlatformFile& GetPlatformPhysical();
 	/** Returns the name of the physical platform file type. */
-	static const TCHAR* GetPhysicalTypeName();
+	static CORE_API const TCHAR* GetPhysicalTypeName();
 	/** Destructor. */
 	virtual ~IPlatformFile() {}
 
@@ -458,7 +458,7 @@ public:
 	* @param Filename file to be opened
 	* @return Close the file by delete'ing the handle. A non-null return value does not mean the file exists, since that may not be determined yet.
 	*/
-	virtual IAsyncReadFileHandle* OpenAsyncRead(const TCHAR* Filename);
+	CORE_API virtual IAsyncReadFileHandle* OpenAsyncRead(const TCHAR* Filename);
 
 	/** Controls if the pak precacher should process precache requests.
 	* Requests below this threshold will not get precached. Without this throttle, quite a lot of memory
@@ -482,10 +482,10 @@ public:
 	}
 
 
-	virtual void GetTimeStampPair(const TCHAR* PathA, const TCHAR* PathB, FDateTime& OutTimeStampA, FDateTime& OutTimeStampB);
+	CORE_API virtual void GetTimeStampPair(const TCHAR* PathA, const TCHAR* PathB, FDateTime& OutTimeStampA, FDateTime& OutTimeStampB);
 
 	/** Return the modification time of a file in the local time of the calling code (GetTimeStamp returns UTC). Returns FDateTime::MinValue() on failure **/
-	virtual FDateTime	GetTimeStampLocal(const TCHAR* Filename);
+	CORE_API virtual FDateTime	GetTimeStampLocal(const TCHAR* Filename);
 
 	/**
 	 * Call the visitor once for each file or directory in a single directory. This function does not explore subdirectories.
@@ -493,7 +493,7 @@ public:
 	 * @param Visitor		Visitor to call for each element of the directory (see FDirectoryVisitor::Visit for the signature)
 	 * @return				false if the directory did not exist or if the visitor returned false.
 	**/
-	virtual bool IterateDirectory(const TCHAR* Directory, FDirectoryVisitorFunc Visitor);
+	CORE_API virtual bool IterateDirectory(const TCHAR* Directory, FDirectoryVisitorFunc Visitor);
 
 	/**
 	 * Call the visitor once for each file or directory in a single directory. This function does not explore subdirectories.
@@ -501,7 +501,7 @@ public:
 	 * @param Visitor		Visitor to call for each element of the directory (see FDirectoryStatVisitor::Visit for the signature)
 	 * @return				false if the directory did not exist or if the visitor returned false.
 	**/
-	virtual bool IterateDirectoryStat(const TCHAR* Directory, FDirectoryStatVisitorFunc Visitor);
+	CORE_API virtual bool IterateDirectoryStat(const TCHAR* Directory, FDirectoryStatVisitorFunc Visitor);
 
 	/** 
 	 * Call the Visit function of the visitor once for each file or directory in a directory tree. This function explores subdirectories.
@@ -509,7 +509,7 @@ public:
 	 * @param Visitor		Visitor to call for each element of the directory and each element of all subdirectories.
 	 * @return				false if the directory did not exist or if the visitor returned false.
 	**/
-	virtual bool IterateDirectoryRecursively(const TCHAR* Directory, FDirectoryVisitor& Visitor);
+	CORE_API virtual bool IterateDirectoryRecursively(const TCHAR* Directory, FDirectoryVisitor& Visitor);
 
 	/** 
 	 * Call the Visit function of the visitor once for each file or directory in a directory tree. This function explores subdirectories.
@@ -517,7 +517,7 @@ public:
 	 * @param Visitor		Visitor to call for each element of the directory and each element of all subdirectories.
 	 * @return				false if the directory did not exist or if the visitor returned false.
 	**/
-	virtual bool IterateDirectoryStatRecursively(const TCHAR* Directory, FDirectoryStatVisitor& Visitor);
+	CORE_API virtual bool IterateDirectoryStatRecursively(const TCHAR* Directory, FDirectoryStatVisitor& Visitor);
 
 	/**
 	 * Call the Visit function of the visitor once for each file or directory in a directory tree. This function explores subdirectories.
@@ -525,7 +525,7 @@ public:
 	 * @param Visitor		Visitor to call for each element of the directory and each element of all subdirectories (see FDirectoryVisitor::Visit for the signature).
 	 * @return				false if the directory did not exist or if the visitor returned false.
 	**/
-	virtual bool IterateDirectoryRecursively(const TCHAR* Directory, FDirectoryVisitorFunc Visitor);
+	CORE_API virtual bool IterateDirectoryRecursively(const TCHAR* Directory, FDirectoryVisitorFunc Visitor);
 
 	/**
 	 * Call the Visit function of the visitor once for each file or directory in a directory tree. This function explores subdirectories.
@@ -533,7 +533,7 @@ public:
 	 * @param Visitor		Visitor to call for each element of the directory and each element of all subdirectories (see FDirectoryStatVisitor::Visit for the signature).
 	 * @return				false if the directory did not exist or if the visitor returned false.
 	**/
-	virtual bool IterateDirectoryStatRecursively(const TCHAR* Directory, FDirectoryStatVisitorFunc Visitor);
+	CORE_API virtual bool IterateDirectoryStatRecursively(const TCHAR* Directory, FDirectoryStatVisitorFunc Visitor);
 		
 	/**
 	 * Finds all the files within the given directory, with optional file extension filter
@@ -542,7 +542,7 @@ public:
 	 * 							Otherwise FileExtension can be of the form .EXT or just EXT and only files with that extension will be returned.
 	 * @return FoundFiles		All the files that matched the optional FileExtension filter, or all files if none was specified.
 	 */
-	virtual void FindFiles(TArray<FString>& FoundFiles, const TCHAR* Directory, const TCHAR* FileExtension);
+	CORE_API virtual void FindFiles(TArray<FString>& FoundFiles, const TCHAR* Directory, const TCHAR* FileExtension);
 
 	/**
 	 * Finds all the files within the directory tree, with optional file extension filter
@@ -551,17 +551,17 @@ public:
 	 * 							Otherwise FileExtension can be of the form .EXT or just EXT and only files with that extension will be returned.
 	 * @return FoundFiles		All the files that matched the optional FileExtension filter, or all files if none was specified.
 	 */
-	virtual void FindFilesRecursively(TArray<FString>& FoundFiles, const TCHAR* Directory, const TCHAR* FileExtension);
+	CORE_API virtual void FindFilesRecursively(TArray<FString>& FoundFiles, const TCHAR* Directory, const TCHAR* FileExtension);
 
 	/** 
 	 * Delete all files and subdirectories in a directory, then delete the directory itself
 	 * @param Directory		The directory to delete.
 	 * @return				true if the directory was deleted or did not exist.
 	**/
-	virtual bool DeleteDirectoryRecursively(const TCHAR* Directory);
+	CORE_API virtual bool DeleteDirectoryRecursively(const TCHAR* Directory);
 
 	/** Create a directory, including any parent directories and return true if the directory was created or already existed. **/
-	virtual bool CreateDirectoryTree(const TCHAR* Directory);
+	CORE_API virtual bool CreateDirectoryTree(const TCHAR* Directory);
 
 	/** 
 	 * Copy a file. This will fail if the destination file already exists.
@@ -571,7 +571,7 @@ public:
 	 * @param WriteFlags		Destination file write options.
 	 * @return			true if the file was copied sucessfully.
 	**/
-	virtual bool CopyFile(const TCHAR* To, const TCHAR* From, EPlatformFileRead ReadFlags = EPlatformFileRead::None, EPlatformFileWrite WriteFlags = EPlatformFileWrite::None);
+	CORE_API virtual bool CopyFile(const TCHAR* To, const TCHAR* From, EPlatformFileRead ReadFlags = EPlatformFileRead::None, EPlatformFileWrite WriteFlags = EPlatformFileWrite::None);
 
 	/** 
 	 * Copy a file or a hierarchy of files (directory).
@@ -580,7 +580,7 @@ public:
 	 * @param bOverwriteAllExisting			Whether to overwrite everything that exists at target
 	 * @return								true if operation completed successfully.
 	 */
-	virtual bool CopyDirectoryTree(const TCHAR* DestinationDirectory, const TCHAR* Source, bool bOverwriteAllExisting);
+	CORE_API virtual bool CopyDirectoryTree(const TCHAR* DestinationDirectory, const TCHAR* Source, bool bOverwriteAllExisting);
 
 	/**
 	 * Converts passed in filename to use an absolute path (for reading).
@@ -589,7 +589,7 @@ public:
 	 * 
 	 * @return	filename using absolute path
 	 */
-	virtual FString ConvertToAbsolutePathForExternalAppForRead( const TCHAR* Filename );
+	CORE_API virtual FString ConvertToAbsolutePathForExternalAppForRead( const TCHAR* Filename );
 
 	/**
 	 * Converts passed in filename to use an absolute path (for writing)
@@ -598,7 +598,7 @@ public:
 	 * 
 	 * @return	filename using absolute path
 	 */
-	virtual FString ConvertToAbsolutePathForExternalAppForWrite( const TCHAR* Filename );
+	CORE_API virtual FString ConvertToAbsolutePathForExternalAppForWrite( const TCHAR* Filename );
 
 	/**
 	 * Helper class to send/receive data to the file server function
@@ -667,7 +667,7 @@ public:
 /**
 * Common base for physical platform File I/O Interface
 **/
-class CORE_API IPhysicalPlatformFile : public IPlatformFile
+class IPhysicalPlatformFile : public IPlatformFile
 {
 public:
 	//~ Begin IPlatformFile Interface
@@ -675,7 +675,7 @@ public:
 	{
 		return true;
 	}
-	virtual bool Initialize(IPlatformFile* Inner, const TCHAR* CmdLine) override;
+	CORE_API virtual bool Initialize(IPlatformFile* Inner, const TCHAR* CmdLine) override;
 	virtual IPlatformFile* GetLowerLevel() override
 	{
 		return nullptr;
@@ -692,7 +692,7 @@ public:
 };
 
 /* Interface class for FPakFile to allow usage from modules that cannot have a compile dependency on FPakFile */
-class CORE_API IPakFile
+class IPakFile
 {
 public:
 	virtual const FString& PakGetPakFilename() const = 0;
diff --git a/Engine/Source/Runtime/Core/Public/GenericPlatform/GenericPlatformMemory.h b/Engine/Source/Runtime/Core/Public/GenericPlatform/GenericPlatformMemory.h
--- a/Engine/Source/Runtime/Core/Public/GenericPlatform/GenericPlatformMemory.h
+++ b/Engine/Source/Runtime/Core/Public/GenericPlatform/GenericPlatformMemory.h
@@ -121,7 +121,7 @@ typedef FGenericPlatformMemoryConstants FPlatformMemoryConstants;
  * Struct used to hold common memory stats for all platforms.
  * These values may change over the entire life of the executable.
  */
-struct CORE_API FGenericPlatformMemoryStats : public FPlatformMemoryConstants
+struct FGenericPlatformMemoryStats : public FPlatformMemoryConstants
 {
 	/** The amount of physical memory currently available, in bytes. */
 	uint64 AvailablePhysical;
@@ -150,10 +150,10 @@ struct CORE_API FGenericPlatformMemoryStats : public FPlatformMemoryConstants
 		Warning,
 		Critical, // high risk of OOM conditions
 	};
-	EMemoryPressureStatus GetMemoryPressureStatus() const;
+	CORE_API EMemoryPressureStatus GetMemoryPressureStatus() const;
 
 	/** Default constructor, clears all variables. */
-	FGenericPlatformMemoryStats();
+	CORE_API FGenericPlatformMemoryStats();
 
 	struct FPlatformSpecificStat
 	{
@@ -166,9 +166,9 @@ struct CORE_API FGenericPlatformMemoryStats : public FPlatformMemoryConstants
 		{}
 	};
 
-	TArray<FPlatformSpecificStat> GetPlatformSpecificStats() const;
+	CORE_API TArray<FPlatformSpecificStat> GetPlatformSpecificStats() const;
 
-	uint64 GetAvailablePhysical(bool bExcludeExtraDevMemory) const;
+	CORE_API uint64 GetAvailablePhysical(bool bExcludeExtraDevMemory) const;
 
 	/** Called by FCsvProfiler::EndFrame to set platform specific CSV stats. */
 	void SetEndFrameCsvStats() const {}
@@ -195,22 +195,22 @@ struct FPlatformMemoryStats;
 #define FMemory_Alloca_Aligned(Size, Alignment) ((Size==0) ? 0 : ((Alignment <= 16) ? FMemory_Alloca(Size) : (void*)(((PTRINT)__FMemory_Alloca_Func(Size + Alignment-1) + Alignment-1) & ~(Alignment-1))))
 
 /** Generic implementation for most platforms, these tend to be unused and unimplemented. */
-struct CORE_API FGenericPlatformMemory
+struct FGenericPlatformMemory
 {
 	/** Set to true if we encounters out of memory. */
-	static bool bIsOOM;
+	static CORE_API bool bIsOOM;
 
 	/** Set to size of allocation that triggered out of memory, zero otherwise. */
-	static uint64 OOMAllocationSize;
+	static CORE_API uint64 OOMAllocationSize;
 
 	/** Set to alignment of allocation that triggered out of memory, zero otherwise. */
-	static uint32 OOMAllocationAlignment;
+	static CORE_API uint32 OOMAllocationAlignment;
 
 	/** Preallocated buffer to delete on out of memory. Used by OOM handling and crash reporting. */
-	static void* BackupOOMMemoryPool;
+	static CORE_API void* BackupOOMMemoryPool;
 
 	/** Size of BackupOOMMemoryPool in bytes. */
-	static uint32 BackupOOMMemoryPoolSize;
+	static CORE_API uint32 BackupOOMMemoryPoolSize;
 
 	/**
 	 * Various memory regions that can be used with memory stats. The exact meaning of
@@ -248,7 +248,7 @@ struct CORE_API FGenericPlatformMemory
 	};
 
 	/** Current allocator */
-	static EMemoryAllocatorToUse AllocatorToUse;
+	static CORE_API EMemoryAllocatorToUse AllocatorToUse;
 
 	/**
 	 * Flags used for shared memory creation/open
@@ -302,12 +302,12 @@ struct CORE_API FGenericPlatformMemory
 
 
 	/** Initializes platform memory specific constants. */
-	static void Init();
+	static CORE_API void Init();
 	
-	[[noreturn]] static void OnOutOfMemory(uint64 Size, uint32 Alignment);
+	[[noreturn]] static CORE_API void OnOutOfMemory(uint64 Size, uint32 Alignment);
 
 	/** Initializes the memory pools, should be called by the init function. */
-	static void SetupMemoryPools();
+	static CORE_API void SetupMemoryPools();
 
 	
 	/**
@@ -322,12 +322,12 @@ struct CORE_API FGenericPlatformMemory
 	/**
 	 * @return the default allocator.
 	 */
-	static FMalloc* BaseAllocator();
+	static CORE_API FMalloc* BaseAllocator();
 
 	/**
 	 * @return platform specific current memory statistics.
 	 */
-	static FPlatformMemoryStats GetStats();
+	static CORE_API FPlatformMemoryStats GetStats();
 
 	/**
 	* @return memory used for platforms that can do it quickly (without affecting stat unit much)
@@ -340,17 +340,17 @@ struct CORE_API FGenericPlatformMemory
 	/**
 	 * Writes all platform specific current memory statistics in the format usable by the malloc profiler.
 	 */
-	static void GetStatsForMallocProfiler( FGenericMemoryStats& out_Stats );
+	static CORE_API void GetStatsForMallocProfiler( FGenericMemoryStats& out_Stats );
 
 	/**
 	 * @return platform specific memory constants.
 	 */
-	static const FPlatformMemoryConstants& GetConstants();
+	static CORE_API const FPlatformMemoryConstants& GetConstants();
 	
 	/**
 	 * @return approximate physical RAM in GB.
 	 */
-	static uint32 GetPhysicalGBRam();
+	static CORE_API uint32 GetPhysicalGBRam();
 
 	/**
 	 * Changes the protection on a region of committed pages in the virtual address space.
@@ -361,7 +361,7 @@ struct CORE_API FGenericPlatformMemory
 	 * @param bCanWrite Can the memory be written to.
 	 * @return True if the specified pages' protection mode was changed.
 	 */
-	static bool PageProtect(void* const Ptr, const SIZE_T Size, const bool bCanRead, const bool bCanWrite);
+	static CORE_API bool PageProtect(void* const Ptr, const SIZE_T Size, const bool bCanRead, const bool bCanWrite);
 
 	/**
 	 * Allocates pages from the OS.
@@ -370,7 +370,7 @@ struct CORE_API FGenericPlatformMemory
 	 *
 	 * @return OS allocated pointer for use by binned allocator
 	 */
-	static void* BinnedAllocFromOS( SIZE_T Size );
+	static CORE_API void* BinnedAllocFromOS( SIZE_T Size );
 	
 	/**
 	 * Returns pages allocated by BinnedAllocFromOS to the OS.
@@ -378,7 +378,7 @@ struct CORE_API FGenericPlatformMemory
 	 * @param A pointer previously returned from BinnedAllocFromOS
 	 * @param Size size of the allocation previously passed to BinnedAllocFromOS
 	 */
-	static void BinnedFreeToOS( void* Ptr, SIZE_T Size );
+	static CORE_API void BinnedFreeToOS( void* Ptr, SIZE_T Size );
 	
 	/**
 	 * Performs initial setup for MiMalloc.
@@ -522,15 +522,15 @@ struct CORE_API FGenericPlatformMemory
 	}
 
 	/** Dumps basic platform memory statistics into the specified output device. */
-	static void DumpStats( FOutputDevice& Ar );
+	static CORE_API void DumpStats( FOutputDevice& Ar );
 
 	/** Dumps basic platform memory statistics and allocator specific statistics into the specified output device. */
-	static void DumpPlatformAndAllocatorStats( FOutputDevice& Ar );
+	static CORE_API void DumpPlatformAndAllocatorStats( FOutputDevice& Ar );
 
 	/**
 	 * Return which "high level", per platform, memory bucket we are in
 	 */
-	static EPlatformMemorySizeBucket GetMemorySizeBucket();
+	static CORE_API EPlatformMemorySizeBucket GetMemorySizeBucket();
 
 	/** @name Memory functions */
 
@@ -596,7 +596,7 @@ private:
 		B = Tmp;
 	}
 
-	static void MemswapGreaterThan8( void* Ptr1, void* Ptr2, SIZE_T Size );
+	static CORE_API void MemswapGreaterThan8( void* Ptr1, void* Ptr2, SIZE_T Size );
 
 public:
 	static inline void Memswap( void* Ptr1, void* Ptr2, SIZE_T Size )
@@ -706,7 +706,7 @@ public:
 	 *
 	 * @return pointer to FSharedMemoryRegion (or its descendants) if successful, NULL if not.
 	 */
-	static FSharedMemoryRegion* MapNamedSharedMemoryRegion(const FString& Name, bool bCreate, uint32 AccessMode, SIZE_T Size);
+	static CORE_API FSharedMemoryRegion* MapNamedSharedMemoryRegion(const FString& Name, bool bCreate, uint32 AccessMode, SIZE_T Size);
 
 	/**
 	 * Unmaps a name shared memory region
@@ -715,7 +715,7 @@ public:
 	 *
 	 * @return true if successful
 	 */
-	static bool UnmapNamedSharedMemoryRegion(FSharedMemoryRegion * MemoryRegion);
+	static CORE_API bool UnmapNamedSharedMemoryRegion(FSharedMemoryRegion * MemoryRegion);
 
 	/**
 	*	Gets whether this platform supports Fast VRAM memory
@@ -732,19 +732,19 @@ public:
 	* Returns true if debug memory has been assigned to the title for general use.
 	* Only applies to platforms with fixed memory and no paging.
 	*/
-	static bool IsExtraDevelopmentMemoryAvailable();
+	static CORE_API bool IsExtraDevelopmentMemoryAvailable();
 
 	/**
 	* Returns >0 if debug memory has been assigned to the title for general use.
 	* Only applies to platforms with fixed memory and no paging.
 	*/
-	static uint64 GetExtraDevelopmentMemorySize();
+	static CORE_API uint64 GetExtraDevelopmentMemorySize();
 
 	/**
 	* This function sets AllocFunction and FreeFunction and returns true, or just returns false.
 	* These functions are the platform dependant low low low level functions that LLM uses to allocate memory.
 	*/
-	static bool GetLLMAllocFunctions(void*(*&OutAllocFunction)(size_t), void(*&OutFreeFunction)(void*, size_t), int32& OutAlignment);
+	static CORE_API bool GetLLMAllocFunctions(void*(*&OutAllocFunction)(size_t), void(*&OutFreeFunction)(void*, size_t), int32& OutAlignment);
 
 	/**
 	* Called for all default tracker LLM allocations and frees, when LLM is enabled.
@@ -783,5 +783,5 @@ protected:
 	friend struct FGenericStatsUpdater;
 
 	/** Updates platform specific stats. This method is called through FGenericStatsUpdater from the task graph thread. */
-	static void InternalUpdateStats( const FPlatformMemoryStats& MemoryStats );
+	static CORE_API void InternalUpdateStats( const FPlatformMemoryStats& MemoryStats );
 };
diff --git a/Engine/Source/Runtime/Core/Public/GenericPlatform/GenericPlatformMisc.h b/Engine/Source/Runtime/Core/Public/GenericPlatform/GenericPlatformMisc.h
--- a/Engine/Source/Runtime/Core/Public/GenericPlatform/GenericPlatformMisc.h
+++ b/Engine/Source/Runtime/Core/Public/GenericPlatform/GenericPlatformMisc.h
@@ -327,12 +327,12 @@ CORE_API const TCHAR* LexToString(EAppReturnType::Type Value);
 /*
  * Holds a computed SHA256 hash.
  */
-struct CORE_API FSHA256Signature
+struct FSHA256Signature
 {
 	uint8 Signature[32];
 
 	/** Generates a hex string of the signature */
-	FString ToString() const;
+	CORE_API FString ToString() const;
 };
 
 enum class EMobileHapticsType : uint8
@@ -539,17 +539,17 @@ ENUM_CLASS_FLAGS(EContextSwitchFlags);
 /**
  * Generic implementation for most platforms
  */
-struct CORE_API FGenericPlatformMisc
+struct FGenericPlatformMisc
 {
 	/**
 	 * Called during appInit() after cmd line setup
 	 */
-	static void PlatformPreInit();
+	static CORE_API void PlatformPreInit();
 	static void PlatformInit() { }
 
 
-	static FDelegateHandle AddNetworkListener(FOnNetworkConnectionChangedDelegate&& InNewDelegate);
-	static bool RemoveNetworkListener(FDelegateHandle Handle);
+	static CORE_API FDelegateHandle AddNetworkListener(FOnNetworkConnectionChangedDelegate&& InNewDelegate);
+	static CORE_API bool RemoveNetworkListener(FDelegateHandle Handle);
 	
 	/**
 	* Called to dismiss splash screen
@@ -587,7 +587,7 @@ struct CORE_API FGenericPlatformMisc
 	 * @param VariableName The name of the variable (ie "Path")
 	 * @return Value of the variable, or an empty string if not set.
 	 */
-	static FString GetEnvironmentVariable(const TCHAR* VariableName);
+	static CORE_API FString GetEnvironmentVariable(const TCHAR* VariableName);
 
 	/**
 	 * Sets an environment variable to the local process's environment
@@ -595,7 +595,7 @@ struct CORE_API FGenericPlatformMisc
 	 * @param VariableName The name of the variable (ie "Path")
 	 * @param Value The string to set the variable to.
 	 */
-	static void SetEnvironmentVar(const TCHAR* VariableName, const TCHAR* Value);
+	static CORE_API void SetEnvironmentVar(const TCHAR* VariableName, const TCHAR* Value);
 
 	/**
 	 * Returns the maximum length of a path
@@ -608,7 +608,7 @@ struct CORE_API FGenericPlatformMisc
 	/**
 	 * return the delimiter between paths in the PATH environment variable.
 	 */
-	static const TCHAR* GetPathVarDelimiter();
+	static CORE_API const TCHAR* GetPathVarDelimiter();
 
 	/**
 	 * Retrieve the Mac address of the current adapter. This is generally platform-dependent code that other platforms must implement.
@@ -616,7 +616,7 @@ struct CORE_API FGenericPlatformMisc
 	 * @return array of bytes representing the Mac address, or empty array if unable to determine.
 	 */
 	UE_DEPRECATED(4.14, "GetMacAddress is deprecated. It is not reliable on all platforms")
-	static TArray<uint8> GetMacAddress();
+	static CORE_API TArray<uint8> GetMacAddress();
 
 	/**
 	 * Retrieve the Mac address of the current adapter as a string. This generally relies on GetMacAddress() being implemented for the platform, if possible.
@@ -624,7 +624,7 @@ struct CORE_API FGenericPlatformMisc
 	 * @return String representing the Mac address, or empty string.
 	 */
 	UE_DEPRECATED(4.14, "GetMacAddressString is deprecated. It is not reliable on all platforms")
-	static FString GetMacAddressString();
+	static CORE_API FString GetMacAddressString();
 
 	/**
 	 * Retrieve the Mac address of the current adapter as a hashed string (privacy). This generally relies on GetMacAddress() being implemented for the platform, if possible.
@@ -632,7 +632,7 @@ struct CORE_API FGenericPlatformMisc
 	 * @return String representing the hashed Mac address, or empty string.
 	 */
 	UE_DEPRECATED(4.14, "GetHasedMacAddressString is deprecated. It is not reliable on all platforms")
-	static FString GetHashedMacAddressString();
+	static CORE_API FString GetHashedMacAddressString();
 
 	/**
 	 * Returns a unique string for device identification. Differs from the deprecated GetUniqueDeviceId
@@ -649,18 +649,18 @@ struct CORE_API FGenericPlatformMisc
 	 *
 	 * @return the unique string generated by this platform for this device, or an empty string if one is not available.
 	 */
-	static FString GetDeviceId();
+	static CORE_API FString GetDeviceId();
 
 	/**
 	 * Returns a unique string for advertising identification
 	 *
 	 * @return the unique string generated by this platform for this device
 	 */
-	static FString GetUniqueAdvertisingId();
+	static CORE_API FString GetUniqueAdvertisingId();
 
 	// #CrashReport: 2015-02-24 Remove
 	/** Submits a crash report to a central server (release builds only) */
-	static void SubmitErrorReport( const TCHAR* InErrorHist, EErrorReportMode::Type InMode );
+	static CORE_API void SubmitErrorReport( const TCHAR* InErrorHist, EErrorReportMode::Type InMode );
 
 	/** Check to see if the platform is being viewed remotely. In such a mode we should aim to minimize screen refresh to get the best performance on the remote viewer */
 	static bool IsRemoteSession()
@@ -679,7 +679,7 @@ struct CORE_API FGenericPlatformMisc
 	}
 
 	/** Returns non-zero if a performance-sensitive diagnostic, like a memory sanitizer, is detected for this process */
-	static EProcessDiagnosticFlags GetProcessDiagnostics();
+	static CORE_API EProcessDiagnosticFlags GetProcessDiagnostics();
 
 	/** Break into the debugger, if IsDebuggerPresent returns true, otherwise do nothing  */
 	UE_DEPRECATED(4.19, "FPlatformMisc::DebugBreak is deprecated. Use the UE_DEBUG_BREAK() macro instead.")
@@ -696,7 +696,7 @@ struct CORE_API FGenericPlatformMisc
 	 *
 	 * @return	CPU vendor name
 	 */
-	static FString GetCPUVendor();
+	static CORE_API FString GetCPUVendor();
 
 	/**
 	 * On x86(-64) platforms, uses cpuid instruction to get the CPU signature
@@ -712,53 +712,53 @@ struct CORE_API FGenericPlatformMisc
 	 *			Bits 20-27	Extended family
 	 *			Bits 28-31	Reserved
 	 */
-	static uint32 GetCPUInfo();
+	static CORE_API uint32 GetCPUInfo();
 
 	/** @return whether this cpu supports certain required instructions or not */
-	static bool HasNonoptionalCPUFeatures();
+	static CORE_API bool HasNonoptionalCPUFeatures();
 	/** @return whether to check for specific CPU compatibility or not */
-	static bool NeedsNonoptionalCPUFeaturesCheck();
+	static CORE_API bool NeedsNonoptionalCPUFeaturesCheck();
 
 	/**
 	 * Uses cpuid instruction to get the CPU brand string
 	 *
 	 * @return	CPU brand string
 	 */
-	static FString GetCPUBrand();
+	static CORE_API FString GetCPUBrand();
 
 	/**
 	 * Returns the CPU chipset if known
 	 *
 	 * @return	CPU chipset string (or "Unknown")
 	 */
-	static FString GetCPUChipset();
+	static CORE_API FString GetCPUChipset();
 
 	/**
 	 * @return primary GPU brand string
 	 */
-	static FString GetPrimaryGPUBrand();
+	static CORE_API FString GetPrimaryGPUBrand();
 
 	/**
 	 * @return	"DeviceMake|DeviceModel" if possible, and "CPUVendor|CPUBrand" otherwise, optionally returns "DeviceMake|DeviceModel|CPUChipset" if known
 	 */
-	static FString GetDeviceMakeAndModel();
+	static CORE_API FString GetDeviceMakeAndModel();
 
-	static struct FGPUDriverInfo GetGPUDriverInfo(const FString& DeviceDescription);
+	static CORE_API struct FGPUDriverInfo GetGPUDriverInfo(const FString& DeviceDescription);
 
 	/**
 	 * Gets the OS Version and OS Subversion.
 	 */
-	static void GetOSVersions( FString& out_OSVersionLabel, FString& out_OSSubVersionLabel );
+	static CORE_API void GetOSVersions( FString& out_OSVersionLabel, FString& out_OSSubVersionLabel );
 
 	/**
 	 * Gets a string representing the numeric OS version (as opposed to a translated OS version that GetOSVersions returns).
 	 * The returned string should try to be brief and avoid newlines and symbols, but there's technically no restriction on the string it can return.
 	 * If the implementation does not support this, it should return an empty string.
 	 */
-	static FString GetOSVersion();
+	static CORE_API FString GetOSVersion();
 
 	/** Retrieves information about the total number of bytes and number of free bytes for the specified disk path. */
-	static bool GetDiskTotalAndFreeSpace( const FString& InPath, uint64& TotalNumberOfBytes, uint64& NumberOfFreeBytes );
+	static CORE_API bool GetDiskTotalAndFreeSpace( const FString& InPath, uint64& TotalNumberOfBytes, uint64& NumberOfFreeBytes );
 
 	/**
 	 * Gets Page Fault stats covering the lifetime of the process
@@ -834,7 +834,7 @@ struct CORE_API FGenericPlatformMisc
 	/**
 	 * Enforces strict memory load/store ordering across the memory barrier call.
 	 */
-	FORCENOINLINE static void MemoryBarrier();
+	CORE_API FORCENOINLINE static void MemoryBarrier();
 
 	/**
 	 * Set a handler to be called when there is a memory warning from the OS
@@ -860,28 +860,28 @@ struct CORE_API FGenericPlatformMisc
 
 	static void SetLastError(uint32 ErrorCode) {}
 
-	static void RaiseException( uint32 ExceptionCode );
+	static CORE_API void RaiseException( uint32 ExceptionCode );
 
 public:
 
 	/**
 	 * Platform specific function for adding a named event that can be viewed in external tool
 	 */
-	static void BeginNamedEvent(const struct FColor& Color, const TCHAR* Text);
-	static void BeginNamedEvent(const struct FColor& Color, const ANSICHAR* Text);
+	static CORE_API void BeginNamedEvent(const struct FColor& Color, const TCHAR* Text);
+	static CORE_API void BeginNamedEvent(const struct FColor& Color, const ANSICHAR* Text);
 
 	/**
 	 * Stats logging for named event 
 	 */
 	template<typename CharType>
 	static void StatNamedEvent(const CharType* Text);
-	static void TickStatNamedEvents();
-	static void LogNameEventStatsInit();
+	static CORE_API void TickStatNamedEvents();
+	static CORE_API void LogNameEventStatsInit();
 
 	/**
 	 * Platform specific function for closing a named event that can be viewed in external tool
 	 */
-	static void EndNamedEvent();
+	static CORE_API void EndNamedEvent();
 
 	/** Platform specific function for adding a named custom stat that can be viewed in external tool */
 	static void CustomNamedStat(const TCHAR* Text, float Value, const TCHAR* Graph, const TCHAR* Unit) {}
@@ -935,7 +935,7 @@ public:
 	 *	@param	InKeyValues			The mapping of key->value pairs to set
 	 *	@return	bool				true if the value was set correctly, false if not
 	 */
-	static bool SetStoredValues(const FString& InStoreId, const FString& InSectionName, const TMap<FString, FString>& InKeyValues);
+	static CORE_API bool SetStoredValues(const FString& InStoreId, const FString& InSectionName, const TMap<FString, FString>& InKeyValues);
 
 	/**
 	 *	Set the value for the given section and key in the platform specific key->value store
@@ -947,7 +947,7 @@ public:
 	 *	@param	InValue				The value to set
 	 *	@return	bool				true if the value was set correctly, false if not
 	 */
-	static bool SetStoredValue(const FString& InStoreId, const FString& InSectionName, const FString& InKeyName, const FString& InValue);
+	static CORE_API bool SetStoredValue(const FString& InStoreId, const FString& InSectionName, const FString& InKeyName, const FString& InValue);
 
 	/**
 	 *	Get the value for the given section and key from the platform specific key->value store
@@ -959,7 +959,7 @@ public:
 	 *	@param	OutValue			The value found
 	 *	@return	bool				true if the entry was found (and OutValue contains the result), false if not
 	 */
-	static bool GetStoredValue(const FString& InStoreId, const FString& InSectionName, const FString& InKeyName, FString& OutValue);
+	static CORE_API bool GetStoredValue(const FString& InStoreId, const FString& InSectionName, const FString& InKeyName, FString& OutValue);
 
 	/**
 	 *	Deletes value for the given section and key in the platform specific key->value store
@@ -970,7 +970,7 @@ public:
 	 *	@param	InKeyName			The name of the key to set the value for
 	 *	@return	bool				true if the value was deleted correctly, false if not found or couldn't delete
 	 */
-	static bool DeleteStoredValue(const FString& InStoreId, const FString& InSectionName, const FString& InKeyName);
+	static CORE_API bool DeleteStoredValue(const FString& InStoreId, const FString& InSectionName, const FString& InKeyName);
 
 	/**
 	 *	Deletes the given section and its contained values in the platform specific key->value store
@@ -980,20 +980,20 @@ public:
 	 *	@param	InSectionName		The section to delete.
 	 *	@return	bool				true if the section was deleted, false if not found or couldn't delete
 	 */
-	static bool DeleteStoredSection(const FString& InStoreId, const FString& InSectionName);
+	static CORE_API bool DeleteStoredSection(const FString& InStoreId, const FString& InSectionName);
 
 	/** Sends a message to a remote tool, and debugger consoles */
-	static void LowLevelOutputDebugString(const TCHAR *Message);
-	static void VARARGS LowLevelOutputDebugStringf(const TCHAR *Format, ... );
+	static CORE_API void LowLevelOutputDebugString(const TCHAR *Message);
+	static CORE_API void VARARGS LowLevelOutputDebugStringf(const TCHAR *Format, ... );
 
 	/** Whether LowLevelOutputDebugString[f] can receive structured output as lines of JSON. */
-	static bool IsLowLevelOutputDebugStringStructured();
+	static CORE_API bool IsLowLevelOutputDebugStringStructured();
 
 	/** Sets the default output to UTF8 */
-	static void SetUTF8Output();
+	static CORE_API void SetUTF8Output();
 
 	/** Prints string to the default output */
-	static void LocalPrint( const TCHAR* Str );
+	static CORE_API void LocalPrint( const TCHAR* Str );
 
 	/** Whether LocalPrint can be called from any thread without overlapping */
 	static bool IsLocalPrintThreadSafe() { return false;  }
@@ -1002,7 +1002,7 @@ public:
 	 * Whether the platform has a separate debug channel to stdout (eg. OutputDebugString on Windows). Used to suppress messages being output twice
 	 * if both go to the same place.
 	 */
-	static bool HasSeparateChannelForDebugOutput();
+	static CORE_API bool HasSeparateChannelForDebugOutput();
 
 	/**
 	 * Requests application exit.
@@ -1011,7 +1011,7 @@ public:
 	 *				  If false, request clean main-loop exit from the platform specific code.
 	 * @param	CallSite a diagnostic message identifying the callsite.
 	 */
-	static void RequestExit( bool Force, const TCHAR* CallSite = nullptr);
+	static CORE_API void RequestExit( bool Force, const TCHAR* CallSite = nullptr);
 
 	/**
 	 * Requests application exit with a specified return code. Name is different from RequestExit() so overloads of just one of functions are possible.
@@ -1021,18 +1021,18 @@ public:
 	 * @param   ReturnCode This value will be returned from the program (on the platforms where it's possible). Limited to 0-255 to conform with POSIX.
 	 * @param	CallSite a diagnostic message identifying the callsite.
 	 */
-	static void RequestExitWithStatus( bool Force, uint8 ReturnCode, const TCHAR* CallSite = nullptr);
+	static CORE_API void RequestExitWithStatus( bool Force, uint8 ReturnCode, const TCHAR* CallSite = nullptr);
 
 
 	/**
 	 * Requests application to restart
 	 */
-	static bool RestartApplication();
+	static CORE_API bool RestartApplication();
 
     /**
 	 * Requests application to restart appending the passed in CmdLine to the restart
 	 */
-	static bool RestartApplicationWithCmdLine(const char* CmdLine);
+	static CORE_API bool RestartApplicationWithCmdLine(const char* CmdLine);
 
 	/**
 	 * Returns the last system error code in string form.  NOTE: Only one return value is valid at a time!
@@ -1041,18 +1041,18 @@ public:
 	 * @param BufferLength the size of the buffer in character count
 	 * @param Error the error code to convert to string form
 	 */
-	static const TCHAR* GetSystemErrorMessage(TCHAR* OutBuffer, int32 BufferCount, int32 Error);
+	static CORE_API const TCHAR* GetSystemErrorMessage(TCHAR* OutBuffer, int32 BufferCount, int32 Error);
 
 	/** Copies text to the operating system clipboard. */
 	UE_DEPRECATED(4.18, "FPlatformMisc::ClipboardCopy() has been superseded by FPlatformApplicationMisc::ClipboardCopy()")
-	static void ClipboardCopy(const TCHAR* Str);
+	static CORE_API void ClipboardCopy(const TCHAR* Str);
 
 	/** Pastes in text from the operating system clipboard. */
 	UE_DEPRECATED(4.18, "FPlatformMisc::ClipboardPaste() has been superseded by FPlatformApplicationMisc::ClipboardPaste()")
-	static void ClipboardPaste(class FString& Dest);
+	static CORE_API void ClipboardPaste(class FString& Dest);
 
 	/** Create a new globally unique identifier. **/
-	static void CreateGuid(struct FGuid& Result);
+	static CORE_API void CreateGuid(struct FGuid& Result);
 
 	/**
 	 * Show a message box if possible, otherwise print a message and return the default
@@ -1061,7 +1061,7 @@ public:
 	 * @param Caption String indicating the title of the message box
 	 * @return Very strange convention...not really EAppReturnType, see implementation
 	 */
-	static EAppReturnType::Type MessageBoxExt( EAppMsgType::Type MsgType, const TCHAR* Text, const TCHAR* Caption );
+	static CORE_API EAppReturnType::Type MessageBoxExt( EAppMsgType::Type MsgType, const TCHAR* Text, const TCHAR* Caption );
 
 	/**
 	 * Handles Game Explorer, Firewall and FirstInstall commands, typically from the installer
@@ -1111,14 +1111,14 @@ public:
 	/**
 	 * @return platform specific path separator.
 	 */
-	static const TCHAR* GetDefaultPathSeparator();
+	static CORE_API const TCHAR* GetDefaultPathSeparator();
 
 	/**
 	 * Checks if platform wants to use a rendering thread on current device
 	 *
 	 * @return true if allowed, false if shouldn't use a separate rendering thread
 	 */
-	static bool UseRenderThread();
+	static CORE_API bool UseRenderThread();
 
 	/**
 	 * Checks if platform wants to allow an audio thread on current device (note: does not imply it will, only if okay given other criteria met)
@@ -1136,32 +1136,32 @@ public:
 	 *
 	 * @return true if allows, false if shouldn't allow thread heartbeat hang detection
 	 */
-	static bool AllowThreadHeartBeat();
+	static CORE_API bool AllowThreadHeartBeat();
 
 	/**
 	 * return the number of hardware CPU cores
 	 */
-	static int32 NumberOfCores();
+	static CORE_API int32 NumberOfCores();
 
 	/**
 	* @return a description of all the groups in the current system and their affinities
 	*/
-	static const FProcessorGroupDesc& GetProcessorGroupDesc();
+	static CORE_API const FProcessorGroupDesc& GetProcessorGroupDesc();
 
 	/**
 	 * return the number of logical CPU cores
 	 */
-	static int32 NumberOfCoresIncludingHyperthreads();
+	static CORE_API int32 NumberOfCoresIncludingHyperthreads();
 
 	/**
 	 * Return the number of worker threads we should spawn, based on number of cores
 	 */
-	static int32 NumberOfWorkerThreadsToSpawn();
+	static CORE_API int32 NumberOfWorkerThreadsToSpawn();
 
 	/**
 	 * Return the number of worker threads we should spawn to service IO, NOT based on number of cores
 	 */
-	static int32 NumberOfIOWorkerThreadsToSpawn();
+	static CORE_API int32 NumberOfIOWorkerThreadsToSpawn();
 
 	/**
 	 * Return the platform specific async IO system, or nullptr if the standard one should be used.
@@ -1179,37 +1179,37 @@ public:
 	}
 
 	/** Get the application root directory. */
-	static const TCHAR* RootDir();
+	static CORE_API const TCHAR* RootDir();
 
 	/** get additional directories which can be considered as root directories */
-	static TArray<FString> GetAdditionalRootDirectories();
+	static CORE_API TArray<FString> GetAdditionalRootDirectories();
 	/** add an additional root directory */
-	static void AddAdditionalRootDirectory(const FString& RootDir);
+	static CORE_API void AddAdditionalRootDirectory(const FString& RootDir);
 
 	/** Get the engine directory */
-	static const TCHAR* EngineDir();
+	static CORE_API const TCHAR* EngineDir();
 
 	/** Get the directory the application was launched from (useful for commandline utilities) */
-	static const TCHAR* LaunchDir();
+	static CORE_API const TCHAR* LaunchDir();
 
 	/** Function to store the current working directory for use with LaunchDir() */
-	static void CacheLaunchDir();
+	static CORE_API void CacheLaunchDir();
 
 	/**
 	 *	Return the project directory
 	 */
-	static const TCHAR* ProjectDir();
+	static CORE_API const TCHAR* ProjectDir();
 
 	/**
 	 *	Return the engine and project directory based from the executable.
 	 *  This function would return different results than EngineDir and Project dir if for example -basedir is set on cmd line and/or pak files are used
 	 */
-	static bool GetEngineAndProjectAbsoluteDirsFromExecutable(FString& OutProjectDir, FString& OutEngineDir);
+	static CORE_API bool GetEngineAndProjectAbsoluteDirsFromExecutable(FString& OutProjectDir, FString& OutEngineDir);
 
 	/**
 	*	Return the CloudDir.  CloudDir can be per-user.
 	*/
-	static FString CloudDir();
+	static CORE_API FString CloudDir();
 
 	/**
 	*	Return true if the PersistentDownloadDir is available.
@@ -1224,7 +1224,7 @@ public:
 	/** Ensure that there enough space in the persistent download dir. Returns true if there is sufficient space.
 	 *  If there is insufficient space and bAttemptToUseUI is true, the function will return false and the user may be prompted to clear space
 	 */
-	static bool CheckPersistentDownloadStorageSpaceAvailable( uint64 BytesRequired, bool bAttemptToUseUI );
+	static CORE_API bool CheckPersistentDownloadStorageSpaceAvailable( uint64 BytesRequired, bool bAttemptToUseUI );
 
 
 	/**
@@ -1232,49 +1232,49 @@ public:
 	 *	On some platforms, returns the writable directory for downloaded data that persists across play sessions.
 	 *	This dir is always per-game.
 	 */
-	static const TCHAR* GamePersistentDownloadDir();
+	static CORE_API const TCHAR* GamePersistentDownloadDir();
 
 	/**
 	 *	Return the temporary directory
 	 */
-	static const TCHAR* GameTemporaryDownloadDir();
+	static CORE_API const TCHAR* GameTemporaryDownloadDir();
 
-	static const TCHAR* GeneratedConfigDir();
+	static CORE_API const TCHAR* GeneratedConfigDir();
 
-	static const TCHAR* GetUBTPlatform();
+	static CORE_API const TCHAR* GetUBTPlatform();
 
-	static const TCHAR* GetUBTTarget();
+	static CORE_API const TCHAR* GetUBTTarget();
 
-	static void SetUBTTargetName(const TCHAR* InTargetName);
-	static const TCHAR* GetUBTTargetName();
+	static CORE_API void SetUBTTargetName(const TCHAR* InTargetName);
+	static CORE_API const TCHAR* GetUBTTargetName();
 
 	/** 
 	 * Determines the shader format for the platform
 	 *
 	 * @return	Returns the shader format to be used by that platform
 	 */
-	static const TCHAR* GetNullRHIShaderFormat();
+	static CORE_API const TCHAR* GetNullRHIShaderFormat();
 
 	/**
 	 * Returns the platform specific chunk based install interface
 	 *
 	 * @return	Returns the platform specific chunk based install implementation
 	 */
-	static IPlatformChunkInstall* GetPlatformChunkInstall();
+	static CORE_API IPlatformChunkInstall* GetPlatformChunkInstall();
 
 	/**
 	 * Returns the platform specific compression interface
 	 *
 	 * @return Returns the platform specific compression interface
 	 */
-	static IPlatformCompression* GetPlatformCompression();
+	static CORE_API IPlatformCompression* GetPlatformCompression();
 
 	/**
 	 * Returns the platform specific interface for communication with processes running on the host pc.
 	 * 
 	 * @return Returns the platform specific host communication interface.
 	 */
-	static IPlatformHostCommunication& GetPlatformHostCommunication();
+	static CORE_API IPlatformHostCommunication& GetPlatformHostCommunication();
 
 	/**
 	 * Has the OS execute a command and path pair (such as launch a browser)
@@ -1311,7 +1311,7 @@ public:
 	 *
 	 * @return whether the hash was computed successfully
 	 */
-	static bool GetSHA256Signature(const void* Data, uint32 ByteSize, FSHA256Signature& OutSignature);
+	static CORE_API bool GetSHA256Signature(const void* Data, uint32 ByteSize, FSHA256Signature& OutSignature);
 
 	/**
 	 * Get the default language (for localization) used by this platform.
@@ -1321,7 +1321,7 @@ public:
 	 *  - An optional four-letter ISO 15924 script code (eg, "Hans").
 	 *  - An optional two-letter ISO 3166-1 country code (eg, "CN").
 	 */
-	static FString GetDefaultLanguage();
+	static CORE_API FString GetDefaultLanguage();
 
 	/**
 	 * Get the default locale (for internationalization) used by this platform.
@@ -1330,13 +1330,13 @@ public:
 	 *  - An optional four-letter ISO 15924 script code (eg, "Hans").
 	 *  - An optional two-letter ISO 3166-1 country code (eg, "CN").
 	 */
-	static FString GetDefaultLocale();
+	static CORE_API FString GetDefaultLocale();
 
 	/**
 	 * Get the timezone identifier for this platform, or an empty string if the default timezone calculation will work.
 	 * @note This should return either an Olson timezone (eg, "America/Los_Angeles") or an offset from GMT/UTC (eg, "GMT-8:00").
 	 */
-	static FString GetTimeZoneId();
+	static CORE_API FString GetTimeZoneId();
 
 	/**
 	 *	Platform-specific Exec function
@@ -1352,7 +1352,7 @@ public:
 	}
 
 	/** @return Get the name of the platform specific file manager (eg, Explorer on Windows, Finder on OS X) */
-	static FText GetFileManagerName();
+	static CORE_API FText GetFileManagerName();
 
 	/** @return Whether filehandles can be opened on one thread and read/written on another thread */
 	static bool SupportsMultithreadedFileHandles()
@@ -1378,7 +1378,7 @@ public:
 	 * @return ensure is allowed
 	 */
 #if DO_ENSURE
-	static bool IsEnsureAllowed();
+	static CORE_API bool IsEnsureAllowed();
 #else
 	static bool IsEnsureAllowed() { return true; }
 #endif // DO_ENSURE
@@ -1387,7 +1387,7 @@ public:
 	 * Updates hotfixable ensure settings from config and commandline (config takes priority).
 	 */
 #if DO_ENSURE
-	static void UpdateHotfixableEnsureSettings();
+	static CORE_API void UpdateHotfixableEnsureSettings();
 #else
 	static void UpdateHotfixableEnsureSettings() {}
 #endif // DO_ENSURE
@@ -1395,7 +1395,7 @@ public:
 	/**
 	 * Ticks values that can be hotfixable in the config.
 	 */
-	static void TickHotfixables();
+	static CORE_API void TickHotfixables();
 
 	static void PromptForRemoteDebugging(bool bIsEnsure)
 	{
@@ -1441,7 +1441,7 @@ public:
 	 *
 	 * @return the default profile name.
 	 */
-	static const TCHAR* GetDefaultDeviceProfileName();
+	static CORE_API const TCHAR* GetDefaultDeviceProfileName();
 
 	/**
 	 * Gets the current battery level.
@@ -1464,12 +1464,12 @@ public:
 	 * Level is a relative value that is platform dependent. Lower is cooler, higher is warmer.
 	 * Level of -1 means Unimplemented.
 	 */
-	static float GetDeviceTemperatureLevel();
+	static CORE_API float GetDeviceTemperatureLevel();
 
 	/**
 	 * Allows a game/program/etc to control the game directory in a special place (for instance, monolithic programs that don't have .uprojects)
 	 */
-	static void SetOverrideProjectDir(const FString& InOverrideDir);
+	static CORE_API void SetOverrideProjectDir(const FString& InOverrideDir);
 
 	UE_DEPRECATED(4.18, "FPaths::SetOverrideGameDir() has been superseded by FPaths::SetOverrideProjectDir().")
 	static FORCEINLINE void SetOverrideGameDir(const FString& InOverrideDir) { return SetOverrideProjectDir(InOverrideDir); }
@@ -1478,7 +1478,7 @@ public:
 	 * Return an ordered list of target platforms this runtime can support (ie Android_DXT, Android
 	 * would mean that it prefers Android_DXT, but can use Android as well)
 	 */
-	static void GetValidTargetPlatforms(TArray<FString>& TargetPlatformNames);
+	static CORE_API void GetValidTargetPlatforms(TArray<FString>& TargetPlatformNames);
 
 	/*
 	 * Converts from a local user index to a FPlatformUserId
@@ -1486,7 +1486,7 @@ public:
 	 * @param LocalUserIndex	A logical index where 0 should always be a valid user
 	 * @returns					A platform-allocated opaque user id, or PLATFORMUSERID_NONE
 	 */
-	static FPlatformUserId GetPlatformUserForUserIndex(int32 LocalUserIndex);
+	static CORE_API FPlatformUserId GetPlatformUserForUserIndex(int32 LocalUserIndex);
 
 	/*
 	 * Converts from a FPlatformUserId to local user index
@@ -1494,7 +1494,7 @@ public:
 	 * @param LocalUserIndex	A platform-allocated opaque user id
 	 * @returns					A valid logical index, or INDEX_NONE
 	 */
-	static int32 GetUserIndexForPlatformUser(FPlatformUserId PlatformUser);
+	static CORE_API int32 GetUserIndexForPlatformUser(FPlatformUserId PlatformUser);
 
 	/**
 	 * Returns whether the platform wants to use a touch screen for virtual joysticks.
@@ -1552,7 +1552,7 @@ public:
 	{}
 
 	/** @return Memory representing a true type or open type font provided by the platform as a default font for unreal to consume; empty array if the default font failed to load. */
-	static TArray<uint8> GetSystemFontBytes();
+	static CORE_API TArray<uint8> GetSystemFontBytes();
 
 	/**
 	 * Returns whether WiFi connection is currently active
@@ -1565,12 +1565,12 @@ public:
 	/**
 	 * Returns the current status for the network connection
 	 */
-	static ENetworkConnectionStatus GetNetworkConnectionStatus();
+	static CORE_API ENetworkConnectionStatus GetNetworkConnectionStatus();
 
 	/**
 	 * Updates the current status for the network connection
 	 */
-	static void SetNetworkConnectionStatus(ENetworkConnectionStatus NewNetworkConnectionStatus);
+	static CORE_API void SetNetworkConnectionStatus(ENetworkConnectionStatus NewNetworkConnectionStatus);
 
 	/**
 	 * Returns whether WiFi connection is currently active
@@ -1600,20 +1600,20 @@ public:
 	/**
 	 * Returns whether the platform is running on battery power or not.
 	 */
-	static bool IsRunningOnBattery();
+	static CORE_API bool IsRunningOnBattery();
 
 	/**
 	 * Returns the current orientation of the device: will be either Portrait, LandscapeLeft, PortraitUpsideDown or LandscapeRight.
 	 * 
 	 * @return An EDeviceScreenOrientation value.
 	 */
-	static EDeviceScreenOrientation GetDeviceOrientation();
+	static CORE_API EDeviceScreenOrientation GetDeviceOrientation();
 
 	/**
 	 * Change the orientation of the device.
 	 */
 	UE_DEPRECATED(5.1, "SetDeviceOrientation is deprecated. Use SetAllowedDeviceOrientation instead.")
-	static void SetDeviceOrientation(EDeviceScreenOrientation NewDeviceOrientation);
+	static CORE_API void SetDeviceOrientation(EDeviceScreenOrientation NewDeviceOrientation);
 
 	/**
 	 * Returns the allowed orientation of the device. This is NOT the same as GetDeviceOrientation, which only returns Portrait, LandscapeLeft, 
@@ -1623,12 +1623,12 @@ public:
 	 * 
 	 * @return An EDeviceScreenOrientation value.
 	 */
-	static EDeviceScreenOrientation GetAllowedDeviceOrientation();
+	static CORE_API EDeviceScreenOrientation GetAllowedDeviceOrientation();
 
 	/**
 	 * Change the allowed orientation of the device. 
 	 */
-	static void SetAllowedDeviceOrientation(EDeviceScreenOrientation NewAllowedDeviceOrientation);
+	static CORE_API void SetAllowedDeviceOrientation(EDeviceScreenOrientation NewAllowedDeviceOrientation);
 
 	/**
 	 * Returns the device volume if the device is capable of returning that information.
@@ -1636,7 +1636,7 @@ public:
 	 *   0 : Muted
 	 * 100 : Full Volume
 	 */
-	static int32 GetDeviceVolume();
+	static CORE_API int32 GetDeviceVolume();
 
 	/**
 	 * Get (or create) the unique ID used to identify this computer.
@@ -1644,7 +1644,7 @@ public:
 	 * This is NOT based on a machine fingerprint.
 	 */
 	UE_DEPRECATED(4.14, "GetMachineId is deprecated. Use GetLoginId instead.")
-	static FGuid GetMachineId();
+	static CORE_API FGuid GetMachineId();
 
 	/**
 	 * Returns a unique string associated with the login account of the current machine.
@@ -1658,7 +1658,7 @@ public:
 	 *
 	 * @return a string containing the LoginID, or empty if not available.
 	 */
-	static FString GetLoginId();
+	static CORE_API FString GetLoginId();
 
 	/**
 	 * Get the Epic account ID for the user who last used the Launcher.
@@ -1667,58 +1667,58 @@ public:
 	 *
 	 * @return an empty string if the account ID was not present or it failed to read it for any reason.
 	 */
-	static FString GetEpicAccountId();
+	static CORE_API FString GetEpicAccountId();
 
 	/**
 	 * Gets a globally unique ID the represents a particular operating system install.
 	 * @returns an opaque string representing the ID, or an empty string if the platform doesn't support one.
 	 */
-	static FString GetOperatingSystemId();
+	static CORE_API FString GetOperatingSystemId();
 
 	/**
 	 * Gets the current mode of convertible laptops, i.e. Laptop or Tablet.
 	 *
 	 * @return The laptop mode, or Unknown if not known, or NotSupported if not a convertible laptop.
 	 */
-	static EConvertibleLaptopMode GetConvertibleLaptopMode();
+	static CORE_API EConvertibleLaptopMode GetConvertibleLaptopMode();
 
 	/**
 	 * Get a string description of the mode the engine was running in.
 	 */
-	static const TCHAR* GetEngineMode();
+	static CORE_API const TCHAR* GetEngineMode();
 
 	/**
 	 * Returns an array of the user's preferred languages in order of preference
 	 * @return An array of language IDs ordered from most preferred to least
 	 */
-	static TArray<FString> GetPreferredLanguages();
+	static CORE_API TArray<FString> GetPreferredLanguages();
 
 	/**
 	 * Returns the currency code associated with the device's locale
 	 * @return the currency code associated with the device's locale
  	*/
-	static FString GetLocalCurrencyCode();
+	static CORE_API FString GetLocalCurrencyCode();
 
 	/**
 	 * Returns the currency symbol associated with the device's locale
 	 * @return the currency symbol associated with the device's locale
 	 */
-	static FString GetLocalCurrencySymbol();
+	static CORE_API FString GetLocalCurrencySymbol();
 
 	/**
 	 * Requests permission to send remote notifications to the user's device.
 	 */
-	static void RegisterForRemoteNotifications();
+	static CORE_API void RegisterForRemoteNotifications();
 
 	/**
 	 * Returns whether or not the device has been registered to receive remote notifications.
 	 */
-	static bool IsRegisteredForRemoteNotifications();
+	static CORE_API bool IsRegisteredForRemoteNotifications();
 
 	/**
 	* Requests unregistering from receiving remote notifications on the user's device.
 	*/
-	static void UnregisterForRemoteNotifications();
+	static CORE_API void UnregisterForRemoteNotifications();
 
 	/**
 	 * Allows platform at runtime to disable unsupported plugins
@@ -1765,34 +1765,34 @@ public:
 		return false;
 	}
 
-	static bool RequestDeviceCheckToken(TFunction<void(const TArray<uint8>&)> QuerySucceededFunc, TFunction<void(const FString&, const FString&)> QueryFailedFunc);
+	static CORE_API bool RequestDeviceCheckToken(TFunction<void(const TArray<uint8>&)> QuerySucceededFunc, TFunction<void(const FString&, const FString&)> QueryFailedFunc);
 
 	UE_DEPRECATED(5.1, "Use named chunks instead")
-	static TArray<FCustomChunk> GetOnDemandChunksForPakchunkIndices(const TArray<int32>& PakchunkIndices);
+	static CORE_API TArray<FCustomChunk> GetOnDemandChunksForPakchunkIndices(const TArray<int32>& PakchunkIndices);
 
 	UE_DEPRECATED(5.1, "Use IPlatformChunkInstall::GetNamedChunksByType instead")
-	static TArray<FCustomChunk> GetAllOnDemandChunks();
+	static CORE_API TArray<FCustomChunk> GetAllOnDemandChunks();
 
 	UE_DEPRECATED(5.1, "Use IPlatformChunkInstall::GetNamedChunksByType instead")
-	static TArray<FCustomChunk> GetAllLanguageChunks();
+	static CORE_API TArray<FCustomChunk> GetAllLanguageChunks();
 
 	UE_DEPRECATED(5.1, "Use IPlatformChunkInstall::GetNamedChunksByType instead")
-	static TArray<FCustomChunk> GetCustomChunksByType(ECustomChunkType DesiredChunkType);
+	static CORE_API TArray<FCustomChunk> GetCustomChunksByType(ECustomChunkType DesiredChunkType);
 
 	/**
 	 * Loads a text file relative to the package root on platforms that distribute apps in package formats.
 	 * For other platforms, the path is relative to the root directory.
 	 */
-	static FString LoadTextFileFromPlatformPackage(const FString& RelativePath);
+	static CORE_API FString LoadTextFileFromPlatformPackage(const FString& RelativePath);
 
-	static bool FileExistsInPlatformPackage(const FString& RelativePath);
+	static CORE_API bool FileExistsInPlatformPackage(const FString& RelativePath);
 
 	/**
 	 * Frees any memory retained by FGenericPlatformMisc.
 	 */
-	static void TearDown();
+	static CORE_API void TearDown();
 
-	static void ParseChunkIdPakchunkIndexMapping(TArray<FString> ChunkIndexRedirects, TMap<int32, int32>& OutMapping);
+	static CORE_API void ParseChunkIdPakchunkIndexMapping(TArray<FString> ChunkIndexRedirects, TMap<int32, int32>& OutMapping);
 
 	static void PumpMessagesOutsideMainLoop()
 	{
@@ -1828,7 +1828,7 @@ public:
 		return PakchunkIndex;
 	}
 
-	static int32 GetPakchunkIndexFromPakFile(const FString& InFilename);
+	static CORE_API int32 GetPakchunkIndexFromPakFile(const FString& InFilename);
 
 	FORCEINLINE static bool Expand16BitIndicesTo32BitOnLoad()
 	{
@@ -1875,14 +1875,14 @@ public:
 	/**
 	 * Returns true if PGO is currently enabled
 	 */
-	static bool IsPGOEnabled();
+	static CORE_API bool IsPGOEnabled();
 
-	static void ShowConsoleWindow();
+	static CORE_API void ShowConsoleWindow();
 
 	/*
 	 * For mobile devices, this function will return true if we wish to propagate the alpha to the backbuffer
 	 */
-	static int GetMobilePropagateAlphaSetting();
+	static CORE_API int GetMobilePropagateAlphaSetting();
 
 	/**
 	*	Return if the game is running on cloud server
@@ -1919,14 +1919,14 @@ public:
 
 protected:
 	/** Whether the user should be prompted to allow for a remote debugger to be attached */
-	static bool bShouldPromptForRemoteDebugging;
+	static CORE_API bool bShouldPromptForRemoteDebugging;
 	/** Whether the user should be prompted to allow for a remote debugger to be attached on an ensure */
-	static bool bPromptForRemoteDebugOnEnsure;
+	static CORE_API bool bPromptForRemoteDebugOnEnsure;
 #endif	//#if !UE_BUILD_SHIPPING
 
-	static EDeviceScreenOrientation AllowedDeviceOrientation;
+	static CORE_API EDeviceScreenOrientation AllowedDeviceOrientation;
 
-	static ENetworkConnectionStatus CurrentNetworkConnectionStatus;
+	static CORE_API ENetworkConnectionStatus CurrentNetworkConnectionStatus;
 
 protected:
 	/**
@@ -1941,7 +1941,7 @@ protected:
 	 * @param bOutSetPhysicalCountToLogicalCount Whether NumberOfCores should be overridden to return
 	 *        NumberOfCoresIncludingHyperthreads.
 	 */
-	static void GetConfiguredCoreLimits(int32 PlatformNumPhysicalCores, int32 PlatformNumLogicalCores,
+	static CORE_API void GetConfiguredCoreLimits(int32 PlatformNumPhysicalCores, int32 PlatformNumLogicalCores,
 		bool& bOutFullyInitialized, int32& OutPhysicalCoreLimit, int32& OutLogicalCoreLimit,
 		bool& bOutSetPhysicalCountToLogicalCount);
 
diff --git a/Engine/Source/Runtime/Core/Public/GenericPlatform/GenericPlatformNamedPipe.h b/Engine/Source/Runtime/Core/Public/GenericPlatform/GenericPlatformNamedPipe.h
--- a/Engine/Source/Runtime/Core/Public/GenericPlatform/GenericPlatformNamedPipe.h
+++ b/Engine/Source/Runtime/Core/Public/GenericPlatform/GenericPlatformNamedPipe.h
@@ -10,15 +10,15 @@
 /**
  * Wrapper for platform named pipe communication.
  */
-class CORE_API FGenericPlatformNamedPipe
+class FGenericPlatformNamedPipe
 {
 public:
 
 	/** Default constructor. */
-	FGenericPlatformNamedPipe();
+	CORE_API FGenericPlatformNamedPipe();
 
 	/** Virtual destructor. */
-	virtual ~FGenericPlatformNamedPipe();
+	CORE_API virtual ~FGenericPlatformNamedPipe();
 
 	FGenericPlatformNamedPipe(const FGenericPlatformNamedPipe&) = delete;
 	FGenericPlatformNamedPipe& operator=(const FGenericPlatformNamedPipe&) = delete;
diff --git a/Engine/Source/Runtime/Core/Public/GenericPlatform/GenericPlatformProcess.h b/Engine/Source/Runtime/Core/Public/GenericPlatform/GenericPlatformProcess.h
--- a/Engine/Source/Runtime/Core/Public/GenericPlatform/GenericPlatformProcess.h
+++ b/Engine/Source/Runtime/Core/Public/GenericPlatform/GenericPlatformProcess.h
@@ -118,7 +118,7 @@ struct FProcHandle;
 /**
 * Generic implementation for most platforms, these tend to be unused and unimplemented
 **/
-struct CORE_API FGenericPlatformProcess
+struct FGenericPlatformProcess
 {
 	/**
 	 * Generic representation of a interprocess semaphore
@@ -174,13 +174,13 @@ struct CORE_API FGenericPlatformProcess
 	};
 
 	/** Load a DLL. **/
-	static void* GetDllHandle( const TCHAR* Filename );
+	static CORE_API void* GetDllHandle( const TCHAR* Filename );
 
 	/** Free a DLL. **/
-	static void FreeDllHandle( void* DllHandle );
+	static CORE_API void FreeDllHandle( void* DllHandle );
 
 	/** Lookup the address of a DLL function. **/
-	static void* GetDllExport( void* DllHandle, const TCHAR* ProcName );
+	static CORE_API void* GetDllExport( void* DllHandle, const TCHAR* ProcName );
 
 	/** Adds a directory to search in when resolving implicitly loaded or filename-only DLLs. **/
 	FORCEINLINE static void AddDllDirectory(const TCHAR* Directory)
@@ -211,28 +211,28 @@ struct CORE_API FGenericPlatformProcess
 	 *
 	 * @return the ProcessId of this process.
 	 */
-	static uint32 GetCurrentProcessId();
+	static CORE_API uint32 GetCurrentProcessId();
 
 	/**
 	 * Retrieves the current hardware CPU core
 	 *
 	 * @return the current hardware core.
 	 */
-	static uint32 GetCurrentCoreNumber();
+	static CORE_API uint32 GetCurrentCoreNumber();
 
 	/**	 
 	 * Change the thread processor affinity
 	 *
 	 * @param AffinityMask A bitfield indicating what processors the thread is allowed to run on.
 	 */
-	static void SetThreadAffinityMask( uint64 AffinityMask );
+	static CORE_API void SetThreadAffinityMask( uint64 AffinityMask );
 
 	/**
 	 * Change the thread processor priority
 	 *
 	 * @param NewPriority an EThreadPriority indicating what priority the thread is to run at.
 	 */
-	static void SetThreadPriority( EThreadPriority NewPriority );
+	static CORE_API void SetThreadPriority( EThreadPriority NewPriority );
 
 	/**
 	 * Helper function to set thread name of the current thread.
@@ -241,7 +241,7 @@ struct CORE_API FGenericPlatformProcess
 	static void SetThreadName( const TCHAR* ThreadName ) { }
 
 	/** Get the active stack size for the currently running thread **/
-	static uint32 GetStackSize();
+	static CORE_API uint32 GetStackSize();
 
 	/** Output information about the currently active thread **/
 	static void DumpThreadInfo( const TCHAR* MarkerName ) { }
@@ -262,22 +262,22 @@ struct CORE_API FGenericPlatformProcess
 	static void TeardownAudioThread() { }
 	
 	/** Content saved to the game or engine directories should be rerouted to user directories instead **/
-	static bool ShouldSaveToUserDir();
+	static CORE_API bool ShouldSaveToUserDir();
 
 	/** Get startup directory.  NOTE: Only one return value is valid at a time! **/
-	static const TCHAR* BaseDir();
+	static CORE_API const TCHAR* BaseDir();
 
 	/** Get user directory.  NOTE: Only one return value is valid at a time! **/
-	static const TCHAR* UserDir();
+	static CORE_API const TCHAR* UserDir();
 
 	/** Get the user settings directory.  NOTE: Only one return value is valid at a time! **/
-	static const TCHAR *UserSettingsDir();
+	static CORE_API const TCHAR *UserSettingsDir();
 
 	/** Get the user temporary directory.  NOTE: Only one return value is valid at a time! **/
-	static const TCHAR *UserTempDir();
+	static CORE_API const TCHAR *UserTempDir();
 
 	/** Get the user home directory.  NOTE: Only one return value is valid at a time! **/
-	static const TCHAR *UserHomeDir();
+	static CORE_API const TCHAR *UserHomeDir();
 
 	struct ApplicationSettingsContext
 	{
@@ -292,7 +292,7 @@ struct CORE_API FGenericPlatformProcess
 	};
 
 	/** Get application settings directory.  NOTE: Only one return value is valid at a time! **/
-	static const TCHAR* ApplicationSettingsDir();
+	static CORE_API const TCHAR* ApplicationSettingsDir();
 	
 	/** 
 	 * Get application settings directory for a given context.
@@ -301,19 +301,19 @@ struct CORE_API FGenericPlatformProcess
 	 * 
 	 * @return 			A string to the appropriate directory.
 	 */
-	static FString GetApplicationSettingsDir(const ApplicationSettingsContext& Settings);
+	static CORE_API FString GetApplicationSettingsDir(const ApplicationSettingsContext& Settings);
 
 	/** Get computer name.  NOTE: Only one return value is valid at a time! **/
-	static const TCHAR* ComputerName();
+	static CORE_API const TCHAR* ComputerName();
 
 	/** Get user name.  NOTE: Only one return value is valid at a time! **/
-	static const TCHAR* UserName(bool bOnlyAlphaNumeric = true);
-	static const TCHAR* ShaderDir();
-	static void SetShaderDir(const TCHAR*Where);
-	static void SetCurrentWorkingDirectoryToBaseDir();
+	static CORE_API const TCHAR* UserName(bool bOnlyAlphaNumeric = true);
+	static CORE_API const TCHAR* ShaderDir();
+	static CORE_API void SetShaderDir(const TCHAR*Where);
+	static CORE_API void SetCurrentWorkingDirectoryToBaseDir();
 
 	/** Get the current working directory (only really makes sense on desktop platforms) */
-	static FString GetCurrentWorkingDirectory();
+	static CORE_API FString GetCurrentWorkingDirectory();
 
 	/**
 	 * Sets the process limits.
@@ -332,17 +332,17 @@ struct CORE_API FGenericPlatformProcess
 	 *
 	 * @return The path to the directory.
 	 */
-	static const FString ShaderWorkingDir();
+	static CORE_API const FString ShaderWorkingDir();
 
 	/**	Clean the shader working directory. */
-	static void CleanShaderWorkingDir();
+	static CORE_API void CleanShaderWorkingDir();
 
 	/**
 	 * Return the path to the currently running executable
 	 *
 	 * @return 	Path of the currently running executable
 	 */
-	static const TCHAR* ExecutablePath();
+	static CORE_API const TCHAR* ExecutablePath();
 
 	/**
 	 * Return the name of the currently running executable
@@ -350,7 +350,7 @@ struct CORE_API FGenericPlatformProcess
 	 * @param	bRemoveExtension	true to remove the extension of the executable name, false to leave it intact
 	 * @return 	Name of the currently running executable
 	 */
-	static const TCHAR* ExecutableName(bool bRemoveExtension = true);
+	static CORE_API const TCHAR* ExecutableName(bool bRemoveExtension = true);
 
 	/**
 	 * Generates the path to the specified application or game.
@@ -363,7 +363,7 @@ struct CORE_API FGenericPlatformProcess
 	 * @param BuildConfiguration The build configuration of the game.
 	 * @return The generated application path.
 	 */
-	static FString GenerateApplicationPath( const FString& AppName, EBuildConfiguration BuildConfiguration);
+	static CORE_API FString GenerateApplicationPath( const FString& AppName, EBuildConfiguration BuildConfiguration);
 
 	/**
 	 * Return the prefix of dynamic library (e.g. lib)
@@ -371,7 +371,7 @@ struct CORE_API FGenericPlatformProcess
 	 * @return The prefix string.
 	 * @see GetModuleExtension, GetModulesDirectory
 	 */
-	static const TCHAR* GetModulePrefix();
+	static CORE_API const TCHAR* GetModulePrefix();
 
 	/**
 	 * Return the extension of dynamic library
@@ -379,12 +379,12 @@ struct CORE_API FGenericPlatformProcess
 	 * @return Extension of dynamic library.
 	 * @see GetModulePrefix, GetModulesDirectory
 	 */
-	static const TCHAR* GetModuleExtension();
+	static CORE_API const TCHAR* GetModuleExtension();
 
 	/**
 	 * Used only by platforms with DLLs, this gives the subdirectory from binaries to find the executables
 	 */
-	static const TCHAR* GetBinariesSubdirectory();
+	static CORE_API const TCHAR* GetBinariesSubdirectory();
 
 	/**
 	 * Used only by platforms with DLLs, this gives the full path to the main directory containing modules
@@ -392,7 +392,7 @@ struct CORE_API FGenericPlatformProcess
 	 * @return The path to the directory.
 	 * @see GetModulePrefix, GetModuleExtension
 	 */
-	static const FString GetModulesDirectory();
+	static CORE_API const FString GetModulesDirectory();
 	
 	/**
 	 * Launch a uniform resource locator (i.e. http://www.epicgames.com/unreal).
@@ -400,7 +400,7 @@ struct CORE_API FGenericPlatformProcess
 	 * task. The URL param must already be a valid URL. If you're looking for code 
 	 * to properly escape a URL fragment, use FGenericPlatformHttp::UrlEncode.
 	 */
-	static void LaunchURL( const TCHAR* URL, const TCHAR* Parms, FString* Error );
+	static CORE_API void LaunchURL( const TCHAR* URL, const TCHAR* Parms, FString* Error );
 
 	/**
 	 * Launch a uniform resource locator (i.e. http://www.epicgames.com/unreal).
@@ -412,19 +412,19 @@ struct CORE_API FGenericPlatformProcess
 	 * 
 	 * @return true if URL passed the filter and was launched, false if it was rejected by the filter.
 	 */
-	static bool LaunchURLFiltered(const TCHAR* URL, const TCHAR* Parms, FString* Error, const UE::Core::FURLRequestFilter& Filter);
+	static CORE_API bool LaunchURLFiltered(const TCHAR* URL, const TCHAR* Parms, FString* Error, const UE::Core::FURLRequestFilter& Filter);
 
 	/**
 	 * Checks if the platform can launch a uniform resource locator (i.e. http://www.epicgames.com/unreal).
 	 **/
-	static bool CanLaunchURL(const TCHAR* URL);
+	static CORE_API bool CanLaunchURL(const TCHAR* URL);
 	
 	/**
 	 * Retrieves the platform-specific bundle identifier or package name of the game
 	 *
 	 * @return The game's bundle identifier or package name.
 	 */
-	static FString GetGameBundleId();
+	static CORE_API FString GetGameBundleId();
 	
 	/**
 	 * Creates a new process and its primary thread. The new process runs the
@@ -441,7 +441,7 @@ struct CORE_API FGenericPlatformProcess
 	 * @param PipeReadChild			Optional HANDLE to pipe for redirecting input
 	 * @return	The process handle for use in other process functions
 	 */
-	static FProcHandle CreateProc( const TCHAR* URL, const TCHAR* Parms, bool bLaunchDetached, bool bLaunchHidden, bool bLaunchReallyHidden, uint32* OutProcessID, int32 PriorityModifier, const TCHAR* OptionalWorkingDirectory, void* PipeWriteChild, void* PipeReadChild = nullptr);
+	static CORE_API FProcHandle CreateProc( const TCHAR* URL, const TCHAR* Parms, bool bLaunchDetached, bool bLaunchHidden, bool bLaunchReallyHidden, uint32* OutProcessID, int32 PriorityModifier, const TCHAR* OptionalWorkingDirectory, void* PipeWriteChild, void* PipeReadChild = nullptr);
 	
 	/**
 	 * Creates a new process and its primary thread, with separate std pipes. The new process runs the
@@ -459,7 +459,7 @@ struct CORE_API FGenericPlatformProcess
 	 * @param PipeStdErrChild		Optional HANDLE to pipe for redirecting stderr
 	 * @return	The process handle for use in other process functions
 	 */
-	static FProcHandle CreateProc( const TCHAR* URL, const TCHAR* Parms, bool bLaunchDetached, bool bLaunchHidden, bool bLaunchReallyHidden, uint32* OutProcessID, int32 PriorityModifier, const TCHAR* OptionalWorkingDirectory, void* PipeWriteChild, void* PipeReadChild, void* PipeStdErrChild);
+	static CORE_API FProcHandle CreateProc( const TCHAR* URL, const TCHAR* Parms, bool bLaunchDetached, bool bLaunchHidden, bool bLaunchReallyHidden, uint32* OutProcessID, int32 PriorityModifier, const TCHAR* OptionalWorkingDirectory, void* PipeWriteChild, void* PipeReadChild, void* PipeStdErrChild);
 
 	/**
 	 * Opens an existing process. 
@@ -467,7 +467,7 @@ struct CORE_API FGenericPlatformProcess
 	 * @param ProcessID				The process id of the process for which we want to obtain a handle.
 	 * @return The process handle for use in other process functions
 	 */
-	static FProcHandle OpenProcess(uint32 ProcessID);
+	static CORE_API FProcHandle OpenProcess(uint32 ProcessID);
 
 	/**
 	 * Returns true if the specified process is running 
@@ -475,28 +475,28 @@ struct CORE_API FGenericPlatformProcess
 	 * @param ProcessHandle handle returned from FPlatformProcess::CreateProc
 	 * @return true if the process is still running
 	 */
-	static bool IsProcRunning( FProcHandle & ProcessHandle );
+	static CORE_API bool IsProcRunning( FProcHandle & ProcessHandle );
 	
 	/**
 	 * Waits for a process to stop
 	 *
 	 * @param ProcessHandle handle returned from FPlatformProcess::CreateProc
 	 */
-	static void WaitForProc( FProcHandle & ProcessHandle );
+	static CORE_API void WaitForProc( FProcHandle & ProcessHandle );
 
 	/**
 	 * Cleans up FProcHandle after we're done with it.
 	 *
 	 * @param ProcessHandle handle returned from FPlatformProcess::CreateProc.
 	 */
-	static void CloseProc( FProcHandle & ProcessHandle );
+	static CORE_API void CloseProc( FProcHandle & ProcessHandle );
 
 	/** Terminates a process
 	 *
 	 * @param ProcessHandle handle returned from FPlatformProcess::CreateProc
 	 * @param KillTree Whether the entire process tree should be terminated.
 	 */
-	static void TerminateProc( FProcHandle & ProcessHandle, bool KillTree = false );
+	static CORE_API void TerminateProc( FProcHandle & ProcessHandle, bool KillTree = false );
 
 	/** Terminates a process tree
 	 *
@@ -504,7 +504,7 @@ struct CORE_API FGenericPlatformProcess
 	 * @param Predicate that returns true if the process identified by ProcessId and ApplicationName
 	 *        should be terminated with its children, else that process and its children will be kept alive
 	 */
-	static void TerminateProcTreeWithPredicate(
+	static CORE_API void TerminateProcTreeWithPredicate(
 			FProcHandle& ProcessHandle,
 			TFunctionRef<bool(uint32 ProcessId, const TCHAR* ApplicationName)> Predicate);
 
@@ -523,22 +523,22 @@ struct CORE_API FGenericPlatformProcess
 	 * will not return until IsEngineExitRequested() is true (EWaitAndForkResult::Parent) or there was an error (EWaitAndForkResult::Error)
 	 * The signal the parent process expects is platform-specific (i.e. SIGRTMIN+1 on Linux). 
 	 */
-	static EWaitAndForkResult WaitAndFork();
+	static CORE_API EWaitAndForkResult WaitAndFork();
 
 	/** Retrieves the termination status of the specified process. **/
-	static bool GetProcReturnCode( FProcHandle & ProcHandle, int32* ReturnCode );
+	static CORE_API bool GetProcReturnCode( FProcHandle & ProcHandle, int32* ReturnCode );
 
 	/** Returns true if the specified application is running */
-	static bool IsApplicationRunning( uint32 ProcessId );
+	static CORE_API bool IsApplicationRunning( uint32 ProcessId );
 
 	/** Returns true if the specified application is running */
-	static bool IsApplicationRunning( const TCHAR* ProcName );
+	static CORE_API bool IsApplicationRunning( const TCHAR* ProcName );
 
 	/** Returns the Name of process given by the PID.  Returns Empty string "" if PID not found. */
-	static FString GetApplicationName( uint32 ProcessId );
+	static CORE_API FString GetApplicationName( uint32 ProcessId );
 
 	/** Outputs the virtual memory usage, of the process with the specified PID */
-	static bool GetApplicationMemoryUsage(uint32 ProcessId, SIZE_T* OutMemoryUsage);
+	static CORE_API bool GetApplicationMemoryUsage(uint32 ProcessId, SIZE_T* OutMemoryUsage);
 
 	/**
 	 * Executes a process, returning the return code, stdout, and stderr. This
@@ -549,13 +549,13 @@ struct CORE_API FGenericPlatformProcess
 	 * @OptionalWorkingDirectory may be 0
 	 * @OptionalbShouldEndWithParentProcess false by default. True to make sure the process is killed with the parent processor (Not Supported on all Platforms)
 	 */
-	static bool ExecProcess(const TCHAR* URL, const TCHAR* Params, int32* OutReturnCode, FString* OutStdOut, FString* OutStdErr, const TCHAR* OptionalWorkingDirectory = NULL, bool bShouldEndWithParentProcess = false);
+	static CORE_API bool ExecProcess(const TCHAR* URL, const TCHAR* Params, int32* OutReturnCode, FString* OutStdOut, FString* OutStdErr, const TCHAR* OptionalWorkingDirectory = NULL, bool bShouldEndWithParentProcess = false);
 
 	/**
 	 * Executes a process as administrator, requesting elevation as necessary. This
 	 * call blocks until the process has returned.
 	 */
-	static bool ExecElevatedProcess(const TCHAR* URL, const TCHAR* Params, int32* OutReturnCode);
+	static CORE_API bool ExecElevatedProcess(const TCHAR* URL, const TCHAR* Params, int32* OutReturnCode);
 
 	/**
 	 * Attempt to launch the provided file name in its default external application. Similar to FPlatformProcess::LaunchURL,
@@ -567,25 +567,25 @@ struct CORE_API FGenericPlatformProcess
 	 * @param	Verb		Optional verb to use when opening the file, if it applies for the platform.
 	 * @return true if the file is launched successfully, false otherwise.
 	 */
-	static bool LaunchFileInDefaultExternalApplication( const TCHAR* FileName, const TCHAR* Parms = NULL, ELaunchVerb::Type Verb = ELaunchVerb::Open, bool bPromptToOpenOnFailure = true );
+	static CORE_API bool LaunchFileInDefaultExternalApplication( const TCHAR* FileName, const TCHAR* Parms = NULL, ELaunchVerb::Type Verb = ELaunchVerb::Open, bool bPromptToOpenOnFailure = true );
 
 	/**
 	 * Attempt to "explore" the folder specified by the provided file path
 	 *
 	 * @param	FilePath	File path specifying a folder to explore
 	 */
-	static void ExploreFolder( const TCHAR* FilePath );
+	static CORE_API void ExploreFolder( const TCHAR* FilePath );
 
 #if PLATFORM_HAS_BSD_TIME 
 
 	/** Sleep this thread for Seconds.  0.0 means release the current time slice to let other threads get some attention. Uses stats.*/
-	static void Sleep( float Seconds );
+	static CORE_API void Sleep( float Seconds );
 	/** Sleep this thread for Seconds.  0.0 means release the current time slice to let other threads get some attention. */
-	static void SleepNoStats( float Seconds );
+	static CORE_API void SleepNoStats( float Seconds );
 	/** Sleep this thread infinitely. */
-	[[noreturn]] static void SleepInfinite();
+	[[noreturn]] static CORE_API void SleepInfinite();
 	/** Yield this thread so another may run for a while. */
-	static void YieldThread();
+	static CORE_API void YieldThread();
 
 #endif // PLATFORM_HAS_BSD_TIME
 
@@ -595,7 +595,7 @@ struct CORE_API FGenericPlatformProcess
 	* @param	Condition	Condition to evaluate.
 	* @param	SleepTime	Time to sleep
 	*/
-	static void ConditionalSleep(TFunctionRef<bool()> Condition, float SleepTime = 0.0f);
+	static CORE_API void ConditionalSleep(TFunctionRef<bool()> Condition, float SleepTime = 0.0f);
 
 	/**
 	 * Creates a new event.
@@ -606,7 +606,7 @@ struct CORE_API FGenericPlatformProcess
 	 */
 	// Message to others in the future, don't try to delete this function as it isn't exactly deprecated, but it should only ever be called from TEventPool::GetEventFromPool()
 	UE_DEPRECATED(5.0, "Please use GetSynchEventFromPool to create a new event, and ReturnSynchEventToPool to release the event.")
-	static class FEvent* CreateSynchEvent(bool bIsManualReset = false);
+	static CORE_API class FEvent* CreateSynchEvent(bool bIsManualReset = false);
 
 	/**
 	 * Gets an event from the pool or creates a new one if necessary.
@@ -615,12 +615,12 @@ struct CORE_API FGenericPlatformProcess
 	 * @return An event, or nullptr none could be created.
 	 * @see CreateSynchEvent, ReturnSynchEventToPool
 	 */
-	static class FEvent* GetSynchEventFromPool(bool bIsManualReset = false);
+	static CORE_API class FEvent* GetSynchEventFromPool(bool bIsManualReset = false);
 
 	/**
 	 * Deletes all the recycled sync events contained by the pools
 	 */
-	static void FlushPoolSyncEvents();
+	static CORE_API void FlushPoolSyncEvents();
 
 	/**
 	 * Returns an event to the pool.
@@ -628,14 +628,14 @@ struct CORE_API FGenericPlatformProcess
 	 * @param Event The event to return.
 	 * @see CreateSynchEvent, GetSynchEventFromPool
 	 */
-	static void ReturnSynchEventToPool(FEvent* Event);
+	static CORE_API void ReturnSynchEventToPool(FEvent* Event);
 
 	/**
 	 * Creates the platform-specific runnable thread. This should only be called from FRunnableThread::Create.
 	 *
 	 * @return The newly created thread
 	 */
-	static class FRunnableThread* CreateRunnableThread();
+	static CORE_API class FRunnableThread* CreateRunnableThread();
 
 	/**
 	 * Closes an anonymous pipe.
@@ -644,7 +644,7 @@ struct CORE_API FGenericPlatformProcess
 	 * @param WritePipe The handle to the write end of the pipe.
 	 * @see CreatePipe, ReadPipe
 	 */
-	static void ClosePipe( void* ReadPipe, void* WritePipe );
+	static CORE_API void ClosePipe( void* ReadPipe, void* WritePipe );
 
 	/**
 	 * Creates a writable anonymous pipe.
@@ -658,7 +658,7 @@ struct CORE_API FGenericPlatformProcess
 	 * @return true on success, false otherwise.
 	 * @see ClosePipe, ReadPipe
 	 */
-	static bool CreatePipe(void*& ReadPipe, void*& WritePipe, bool bWritePipeLocal = false);
+	static CORE_API bool CreatePipe(void*& ReadPipe, void*& WritePipe, bool bWritePipeLocal = false);
 
 	/**
 	 * Reads all pending data from an anonymous pipe, such as STDOUT or STDERROR of a process.
@@ -667,7 +667,7 @@ struct CORE_API FGenericPlatformProcess
 	 * @return A string containing the read data.
 	 * @see ClosePipe, CreatePipe
 	 */
-	static FString ReadPipe( void* ReadPipe );
+	static CORE_API FString ReadPipe( void* ReadPipe );
 
 	/**
 	 * Reads all pending data from an anonymous pipe, such as STDOUT or STDERROR of a process.
@@ -677,7 +677,7 @@ struct CORE_API FGenericPlatformProcess
 	 * @return true if successful (i.e. any data was read)
 	 * @see ClosePipe, CreatePipe
 	 */
-	static bool ReadPipeToArray(void* ReadPipe, TArray<uint8> & Output);
+	static CORE_API bool ReadPipeToArray(void* ReadPipe, TArray<uint8> & Output);
 
 	/**
 	* Sends the message to process through pipe
@@ -688,7 +688,7 @@ struct CORE_API FGenericPlatformProcess
 	* @return True if all bytes written successfully.
 	* @see CreatePipe, ClosePipe, ReadPipe
 	*/
-	static bool WritePipe(void* WritePipe, const FString& Message, FString* OutWritten = nullptr);
+	static CORE_API bool WritePipe(void* WritePipe, const FString& Message, FString* OutWritten = nullptr);
 
 	/**
 	* Sends data to process through pipe
@@ -700,7 +700,7 @@ struct CORE_API FGenericPlatformProcess
 	* @return True if all bytes written successfully.
 	* @see CreatePipe, ClosePipe, ReadPipe
 	*/
-	static bool WritePipe(void* WritePipe, const uint8* Data, const int32 DataLength, int32* OutDataLength = nullptr);
+	static CORE_API bool WritePipe(void* WritePipe, const uint8* Data, const int32 DataLength, int32* OutDataLength = nullptr);
 
 
 	/**
@@ -708,7 +708,7 @@ struct CORE_API FGenericPlatformProcess
 	 *
 	 * @return true if the platform can use multiple threads, false otherwise.
 	 */
-	static bool SupportsMultithreading();
+	static CORE_API bool SupportsMultithreading();
 	
 	/** Enables Real Time Mode on the current thread. */
 	static void SetRealTimeMode() { }
@@ -721,7 +721,7 @@ struct CORE_API FGenericPlatformProcess
 	 * @param MaxLocks Maximum amount of locks that the semaphore can have (pass 1 to make it act as mutex).
 	 * @return Pointer to heap allocated semaphore object. Caller is responsible for deletion.
 	 */
-	static FSemaphore* NewInterprocessSynchObject(const FString& Name, bool bCreate, uint32 MaxLocks = 1);
+	static CORE_API FSemaphore* NewInterprocessSynchObject(const FString& Name, bool bCreate, uint32 MaxLocks = 1);
 
 	/**
 	 * Creates or opens an interprocess synchronization object.
@@ -731,31 +731,31 @@ struct CORE_API FGenericPlatformProcess
 	 * @param MaxLocks Maximum amount of locks that the semaphore can have (pass 1 to make it act as mutex).
 	 * @return Pointer to heap allocated semaphore object. Caller is responsible for deletion.
 	 */
-	static FSemaphore* NewInterprocessSynchObject(const TCHAR* Name, bool bCreate, uint32 MaxLocks = 1);
+	static CORE_API FSemaphore* NewInterprocessSynchObject(const TCHAR* Name, bool bCreate, uint32 MaxLocks = 1);
 
 	/**
 	 * Deletes an interprocess synchronization object.
 	 *
 	 * @param Object object to destroy.
 	 */
-	static bool DeleteInterprocessSynchObject(FSemaphore * Object);
+	static CORE_API bool DeleteInterprocessSynchObject(FSemaphore * Object);
 
 	/**
 	 * Makes process run as a system service (daemon), i.e. detaches it from whatever user session it was initially run from.
 	 *
 	 * @return true if successful, false otherwise.
 	 */
-	static bool Daemonize();
+	static CORE_API bool Daemonize();
 
 	/**
 	 * Checks if we're the first instance. An instance can become first if the previous first instance quits before it.
 	 */
-	static bool IsFirstInstance();
+	static CORE_API bool IsFirstInstance();
 
 	/**
 	 * Tears down allocated process resources.
 	 */
-	static void TearDown();
+	static CORE_API void TearDown();
 
 	/**
 	 * force skip calling FThreadStats::WaitForStats()
@@ -765,12 +765,12 @@ struct CORE_API FGenericPlatformProcess
 	/**
 	 * specifies the thread to use for UObject reference collection
 	 */
-	static ENamedThreads::Type GetDesiredThreadForUObjectReferenceCollector();
+	static CORE_API ENamedThreads::Type GetDesiredThreadForUObjectReferenceCollector();
 
 	/**
 	 * allows a platform to override the threading configuration for reference collection
 	 */
-	static void ModifyThreadAssignmentForUObjectReferenceCollector( int32& NumThreads, int32& NumBackgroundThreads, ENamedThreads::Type& NormalThreadName, ENamedThreads::Type& BackgroundThreadName );
+	static CORE_API void ModifyThreadAssignmentForUObjectReferenceCollector( int32& NumThreads, int32& NumBackgroundThreads, ENamedThreads::Type& NormalThreadName, ENamedThreads::Type& BackgroundThreadName );
 
 	/**
 	 * Tells the processor to pause for implementation-specific amount of time. Is used for spin-loops to improve the speed at 
diff --git a/Engine/Source/Runtime/Core/Public/GenericPlatform/GenericPlatformStackWalk.h b/Engine/Source/Runtime/Core/Public/GenericPlatform/GenericPlatformStackWalk.h
--- a/Engine/Source/Runtime/Core/Public/GenericPlatform/GenericPlatformStackWalk.h
+++ b/Engine/Source/Runtime/Core/Public/GenericPlatform/GenericPlatformStackWalk.h
@@ -72,7 +72,7 @@ struct FProgramCounterSymbolInfo final
 	CORE_API FProgramCounterSymbolInfo();
 };
 
-struct CORE_API FProgramCounterSymbolInfoEx
+struct FProgramCounterSymbolInfoEx
 {
 	/** Module name. */
 	FString	ModuleName;
@@ -96,13 +96,13 @@ struct CORE_API FProgramCounterSymbolInfoEx
 	uint64	ProgramCounter;
 
 	/** Default constructor. */
-	FProgramCounterSymbolInfoEx(FString InModuleName = "", FString InFunctionName = "", FString InFilename = "", uint32 InLineNumber = 0, uint64 InSymbolDisplacement = 0, uint64 InOffsetInModule = 0, uint64 InProgramCounter = 0);
+	CORE_API FProgramCounterSymbolInfoEx(FString InModuleName = "", FString InFunctionName = "", FString InFilename = "", uint32 InLineNumber = 0, uint64 InSymbolDisplacement = 0, uint64 InOffsetInModule = 0, uint64 InProgramCounter = 0);
 };
 
 /**
  * Generic implementation for most platforms
  */
-struct CORE_API FGenericPlatformStackWalk
+struct FGenericPlatformStackWalk
 {
 	typedef FGenericPlatformStackWalk Base;
 
@@ -124,7 +124,7 @@ struct CORE_API FGenericPlatformStackWalk
 	/**
 	* Initializes options related to stack walking from ini, i.e. how detailed the stack walking should be, performance settings etc.
 	*/
-	static void Init();
+	static CORE_API void Init();
 
 	/**
 	 * Initializes stack traversal and symbol. Must be called before any other stack/symbol functions. Safe to reenter.
@@ -163,7 +163,7 @@ struct CORE_API FGenericPlatformStackWalk
 	 * @param	Context					Pointer to crash context, if any
 	 * @return	true if the symbol was found, otherwise false
 	 */ 
-	static bool ProgramCounterToHumanReadableString( int32 CurrentCallDepth, uint64 ProgramCounter, ANSICHAR* HumanReadableString, SIZE_T HumanReadableStringSize, FGenericCrashContext* Context = nullptr );
+	static CORE_API bool ProgramCounterToHumanReadableString( int32 CurrentCallDepth, uint64 ProgramCounter, ANSICHAR* HumanReadableString, SIZE_T HumanReadableStringSize, FGenericCrashContext* Context = nullptr );
 
 	/**
 	 * Converts the passed in symbol information to a human readable string and appends it to the passed in one.
@@ -183,10 +183,10 @@ struct CORE_API FGenericPlatformStackWalk
 	 * @param	Context					Pointer to crash context, if any
 	 * @return	true if the symbol was found, otherwise false
 	 */ 
-	static bool SymbolInfoToHumanReadableString( const FProgramCounterSymbolInfo& SymbolInfo, ANSICHAR* HumanReadableString, SIZE_T HumanReadableStringSize );
+	static CORE_API bool SymbolInfoToHumanReadableString( const FProgramCounterSymbolInfo& SymbolInfo, ANSICHAR* HumanReadableString, SIZE_T HumanReadableStringSize );
 
 	/** Same as above, but can be used with external applications. */
-	static bool SymbolInfoToHumanReadableStringEx( const FProgramCounterSymbolInfoEx& SymbolInfo, FString& out_HumanReadableString );
+	static CORE_API bool SymbolInfoToHumanReadableStringEx( const FProgramCounterSymbolInfoEx& SymbolInfo, FString& out_HumanReadableString );
 
 	/**
 	 * Converts the passed in program counter address to a symbol info struct, filling in module and filename, line number and displacement.
@@ -213,7 +213,7 @@ struct CORE_API FGenericPlatformStackWalk
 	 * @param	MaxDepth			Entries in BackTrace array
 	 * @param	Context				Optional thread context information
 	 */
-	static uint32 CaptureStackBackTrace( uint64* BackTrace, uint32 MaxDepth, void* Context = nullptr );
+	static CORE_API uint32 CaptureStackBackTrace( uint64* BackTrace, uint32 MaxDepth, void* Context = nullptr );
 
 	/**
 	 * Capture a stack backtrace for a specific thread.
@@ -223,7 +223,7 @@ struct CORE_API FGenericPlatformStackWalk
 	 * @param	MaxDepth			Entries in BackTrace array
 	 * @param	Context				Optional thread context information
 	 */
-	static uint32 CaptureThreadStackBackTrace(uint64 ThreadId, uint64* BackTrace, uint32 MaxDepth, void* Context = nullptr);
+	static CORE_API uint32 CaptureThreadStackBackTrace(uint64 ThreadId, uint64* BackTrace, uint32 MaxDepth, void* Context = nullptr);
 
 	/**
 	 * Walks the stack and appends the human readable string to the passed in one.
@@ -234,7 +234,7 @@ struct CORE_API FGenericPlatformStackWalk
 	 * @param	IgnoreCount			Number of stack entries to ignore (some are guaranteed to be in the stack walking code)
 	 * @param	Context				Optional thread context information
 	 */ 
-	FORCENOINLINE static void StackWalkAndDump( ANSICHAR* HumanReadableString, SIZE_T HumanReadableStringSize, int32 IgnoreCount, void* Context = nullptr );  // FORCENOINLINE so it can be counted during StackTrace
+	CORE_API FORCENOINLINE static void StackWalkAndDump( ANSICHAR* HumanReadableString, SIZE_T HumanReadableStringSize, int32 IgnoreCount, void* Context = nullptr );  // FORCENOINLINE so it can be counted during StackTrace
 
 	/**
 	 * Walks the stack and appends the human readable string to the passed in one.
@@ -245,7 +245,7 @@ struct CORE_API FGenericPlatformStackWalk
 	 * @param	ProgramCounter		Instruction address that should be at the top of the stack
 	 * @param	Context				Optional thread context information
 	 */ 
-	static void StackWalkAndDump( ANSICHAR* HumanReadableString, SIZE_T HumanReadableStringSize, void* ProgramCounter, void* Context = nullptr );
+	static CORE_API void StackWalkAndDump( ANSICHAR* HumanReadableString, SIZE_T HumanReadableStringSize, void* ProgramCounter, void* Context = nullptr );
 
 	/**
 	 * Walks the stack and updates the Stack array with the symbol information for each line in the stack.
@@ -256,7 +256,7 @@ struct CORE_API FGenericPlatformStackWalk
 	 * 
 	 * @return	The stack of symbols to return.
 	 */ 
-	static TArray<FProgramCounterSymbolInfo> GetStack(int32 IgnoreCount, int32 MaxDepth = 100, void* Context = nullptr);
+	static CORE_API TArray<FProgramCounterSymbolInfo> GetStack(int32 IgnoreCount, int32 MaxDepth = 100, void* Context = nullptr);
 
 	/**
 	* Walks the stack for the specified thread and appends the human readable string to the passed in one.
@@ -281,7 +281,7 @@ struct CORE_API FGenericPlatformStackWalk
 	 * @param   Flags				Used to pass additional information (see StackWalkFlags)
 	 * @param	Context				Optional thread context information
 	 */ 
-	static void StackWalkAndDumpEx( ANSICHAR* HumanReadableString, SIZE_T HumanReadableStringSize, int32 IgnoreCount, uint32 Flags, void* Context = nullptr );
+	static CORE_API void StackWalkAndDumpEx( ANSICHAR* HumanReadableString, SIZE_T HumanReadableStringSize, int32 IgnoreCount, uint32 Flags, void* Context = nullptr );
 	/**
 	 * Walks the stack and appends the human readable string to the passed in one.
 	 * @warning: The code assumes that HumanReadableString is large enough to contain the information.
@@ -292,7 +292,7 @@ struct CORE_API FGenericPlatformStackWalk
 	 * @param   Flags				Used to pass additional information (see StackWalkFlags)
 	 * @param	Context				Optional thread context information
 	 */ 
-	static void StackWalkAndDumpEx( ANSICHAR* HumanReadableString, SIZE_T HumanReadableStringSize, void* ProgramCounter, uint32 Flags, void* Context = nullptr );
+	static CORE_API void StackWalkAndDumpEx( ANSICHAR* HumanReadableString, SIZE_T HumanReadableStringSize, void* ProgramCounter, uint32 Flags, void* Context = nullptr );
 
 	/**
 	 * Returns the number of modules loaded by the currently running process.
@@ -321,11 +321,11 @@ struct CORE_API FGenericPlatformStackWalk
 	 *
 	 * @return	A map containing the meta-data (if any).
 	 */
-	static TMap<FName, FString> GetSymbolMetaData();
+	static CORE_API TMap<FName, FString> GetSymbolMetaData();
 
 protected:
 
 	/** Returns true if non-monolithic builds should produce full callstacks in the log (and load all debug symbols) */
-	static bool WantsDetailedCallstacksInNonMonolithicBuilds();
+	static CORE_API bool WantsDetailedCallstacksInNonMonolithicBuilds();
 
 };
diff --git a/Engine/Source/Runtime/Core/Public/GenericPlatform/GenericPlatformString.h b/Engine/Source/Runtime/Core/Public/GenericPlatform/GenericPlatformString.h
--- a/Engine/Source/Runtime/Core/Public/GenericPlatform/GenericPlatformString.h
+++ b/Engine/Source/Runtime/Core/Public/GenericPlatform/GenericPlatformString.h
@@ -419,7 +419,7 @@ private:
 	 * @param Src     Pointer to the null-terminated string being converted.
 	 */
 	template <typename DestEncoding, typename SourceEncoding>
-	static CORE_API void LogBogusChars(const SourceEncoding* Src);
+	static void LogBogusChars(const SourceEncoding* Src);
 
 
 	/**
@@ -429,5 +429,5 @@ private:
 	 * @param SrcSize Number of characters in the Src string.
 	 */
 	template <typename DestEncoding, typename SourceEncoding>
-	static CORE_API void LogBogusChars(const SourceEncoding* Src, int32 SrcSize);
+	static void LogBogusChars(const SourceEncoding* Src, int32 SrcSize);
 };
diff --git a/Engine/Source/Runtime/Core/Public/GenericPlatform/GenericPlatformSymbolication.h b/Engine/Source/Runtime/Core/Public/GenericPlatform/GenericPlatformSymbolication.h
--- a/Engine/Source/Runtime/Core/Public/GenericPlatform/GenericPlatformSymbolication.h
+++ b/Engine/Source/Runtime/Core/Public/GenericPlatform/GenericPlatformSymbolication.h
@@ -116,15 +116,15 @@ struct FGenericPlatformSymbolDatabaseKeyFuncs
 typedef TSet<FGenericPlatformSymbolDatabase, FGenericPlatformSymbolDatabaseKeyFuncs> FGenericPlatformSymbolDatabaseSet;
 
 
-struct CORE_API FGenericPlatformSymbolication
+struct FGenericPlatformSymbolication
 {
 	/**
 	 * @param Architecture Specify the architecture in case of universal binary, currently unused
 	 */
-	static bool LoadSymbolDatabaseForBinary(FString SourceFolder, FString BinaryPath, FString BinarySignature, TOptional<FString> Architecture, FGenericPlatformSymbolDatabase& OutDatabase);
-	static bool SaveSymbolDatabaseForBinary(FString TargetFolder, FString Name, FGenericPlatformSymbolDatabase& Database);
+	static CORE_API bool LoadSymbolDatabaseForBinary(FString SourceFolder, FString BinaryPath, FString BinarySignature, TOptional<FString> Architecture, FGenericPlatformSymbolDatabase& OutDatabase);
+	static CORE_API bool SaveSymbolDatabaseForBinary(FString TargetFolder, FString Name, FGenericPlatformSymbolDatabase& Database);
 	
-	static bool SymbolInfoForStrippedSymbol(FGenericPlatformSymbolDatabase const& Database, uint64 ProgramCounter, uint64 ModuleOffset, FString ModuleSignature, FProgramCounterSymbolInfo& Info);
+	static CORE_API bool SymbolInfoForStrippedSymbol(FGenericPlatformSymbolDatabase const& Database, uint64 ProgramCounter, uint64 ModuleOffset, FString ModuleSignature, FProgramCounterSymbolInfo& Info);
 	
 	static bool SymbolInfoForAddress(uint64 ProgramCounter, FProgramCounterSymbolInfo& Info) { return false; }
 };
diff --git a/Engine/Source/Runtime/Core/Public/GenericPlatform/GenericPlatformTime.h b/Engine/Source/Runtime/Core/Public/GenericPlatform/GenericPlatformTime.h
--- a/Engine/Source/Runtime/Core/Public/GenericPlatform/GenericPlatformTime.h
+++ b/Engine/Source/Runtime/Core/Public/GenericPlatform/GenericPlatformTime.h
@@ -39,7 +39,7 @@ struct FCPUTime
 /**
 * Generic implementation for most platforms
 **/
-struct CORE_API FGenericPlatformTime
+struct FGenericPlatformTime
 {
 #if PLATFORM_HAS_BSD_TIME 
 	/**
@@ -50,7 +50,7 @@ struct CORE_API FGenericPlatformTime
 	 *
 	 * @return	current time
 	 */
-	static double InitTiming();
+	static CORE_API double InitTiming();
 
 	static FORCEINLINE double Seconds()
 	{
@@ -73,10 +73,10 @@ struct CORE_API FGenericPlatformTime
 	}
 
 	/** Returns the system time. */
-	static void SystemTime( int32& Year, int32& Month, int32& DayOfWeek, int32& Day, int32& Hour, int32& Min, int32& Sec, int32& MSec );
+	static CORE_API void SystemTime( int32& Year, int32& Month, int32& DayOfWeek, int32& Day, int32& Hour, int32& Min, int32& Sec, int32& MSec );
 
 	/** Returns the UTC time. */
-	static void UtcTime( int32& Year, int32& Month, int32& DayOfWeek, int32& Day, int32& Hour, int32& Min, int32& Sec, int32& MSec );
+	static CORE_API void UtcTime( int32& Year, int32& Month, int32& DayOfWeek, int32& Day, int32& Hour, int32& Min, int32& Sec, int32& MSec );
 #endif
 
 	/**
@@ -86,7 +86,7 @@ struct CORE_API FGenericPlatformTime
 	 * @param DestSize Size of destination buffer in characters
 	 * @return Date string
 	 */
-	static TCHAR* StrDate( TCHAR* Dest, SIZE_T DestSize );
+	static CORE_API TCHAR* StrDate( TCHAR* Dest, SIZE_T DestSize );
 	/**
 	 * Get the system time
 	 * 
@@ -94,7 +94,7 @@ struct CORE_API FGenericPlatformTime
 	 * @param DestSize Size of destination buffer in characters
 	 * @return Time string
 	 */
-	static TCHAR* StrTime( TCHAR* Dest, SIZE_T DestSize );
+	static CORE_API TCHAR* StrTime( TCHAR* Dest, SIZE_T DestSize );
 
 	/**
 	 * Returns a timestamp string built from the current date and time.
@@ -102,7 +102,7 @@ struct CORE_API FGenericPlatformTime
 	 *
 	 * @return timestamp string
 	 */
-	static const TCHAR* StrTimestamp();
+	static CORE_API const TCHAR* StrTimestamp();
 
 	/**
 	 * Returns a pretty-string for a time given in seconds. (I.e. "4:31 min", "2:16:30 hours", etc)
@@ -110,7 +110,7 @@ struct CORE_API FGenericPlatformTime
 	 * @param Seconds Time in seconds
 	 * @return Time in a pretty formatted string
 	 */
-	static FString PrettyTime( double Seconds );
+	static CORE_API FString PrettyTime( double Seconds );
 
 	/** Updates CPU utilization, called through a delegate from the Core ticker. */
 	static bool UpdateCPUTime( float DeltaTime )
@@ -186,7 +186,7 @@ struct CORE_API FGenericPlatformTime
 	/**
 	 * @return seconds per cycle.
 	 */
-	static double GetSecondsPerCycle64();
+	static CORE_API double GetSecondsPerCycle64();
 	/** Converts cycles to milliseconds. */
 	static double ToMilliseconds64(const uint64 Cycles)
 	{
@@ -201,7 +201,7 @@ struct CORE_API FGenericPlatformTime
 
 protected:
 
-	static double SecondsPerCycle;
-	static double SecondsPerCycle64;
-	static double LastIntervalCPUTimeInSeconds;
+	static CORE_API double SecondsPerCycle;
+	static CORE_API double SecondsPerCycle64;
+	static CORE_API double LastIntervalCPUTimeInSeconds;
 };
