[Core]

diff --git a/Engine/Source/Runtime/Core/Public/GenericPlatform/GenericWidePlatformString.h b/Engine/Source/Runtime/Core/Public/GenericPlatform/GenericWidePlatformString.h
--- a/Engine/Source/Runtime/Core/Public/GenericPlatform/GenericWidePlatformString.h
+++ b/Engine/Source/Runtime/Core/Public/GenericPlatform/GenericWidePlatformString.h
@@ -40,7 +40,7 @@ public:
 	CORE_API static WIDECHAR* Strncpy(WIDECHAR* Dest, const WIDECHAR* Src, SIZE_T MaxLen);
 	CORE_API static WIDECHAR* Strcat(WIDECHAR* Dest, SIZE_T DestCount, const WIDECHAR* Src);
 
-	CORE_API static int32 Strcmp( const WIDECHAR* String1, const WIDECHAR* String2 )
+	static int32 Strcmp( const WIDECHAR* String1, const WIDECHAR* String2 )
 	{
 		// walk the strings, comparing them case sensitively
 		for (; *String1 || *String2; String1++, String2++)
@@ -54,7 +54,7 @@ public:
 		return 0;
 	}
 
-	CORE_API static int32 Strlen( const WIDECHAR* String )
+	static int32 Strlen( const WIDECHAR* String )
 	{
 		int32 Length = -1;
 
@@ -67,7 +67,7 @@ public:
 		return Length;
 	}
 
-	CORE_API static int32 Strnlen( const WIDECHAR* String, SIZE_T StringSize )
+	static int32 Strnlen( const WIDECHAR* String, SIZE_T StringSize )
 	{
 		int32 Length = -1;
 
@@ -108,7 +108,7 @@ public:
 	}
 #endif
 
-	CORE_API static const WIDECHAR* Strstr( const WIDECHAR* String, const WIDECHAR* Find)
+	static const WIDECHAR* Strstr( const WIDECHAR* String, const WIDECHAR* Find)
 	{
 		WIDECHAR Char1, Char2;
 		if ((Char1 = *Find++) != 0)
@@ -134,7 +134,7 @@ public:
 		return String;
 	}
 
-	CORE_API static const WIDECHAR* Strchr( const WIDECHAR* String, WIDECHAR C)
+	static const WIDECHAR* Strchr( const WIDECHAR* String, WIDECHAR C)
 	{
 		while (*String != C && *String != 0)
 		{
@@ -144,7 +144,7 @@ public:
 		return (*String == C) ? String : nullptr;
 	}
 
-	CORE_API static const WIDECHAR* Strrchr( const WIDECHAR* String, WIDECHAR C)
+	static const WIDECHAR* Strrchr( const WIDECHAR* String, WIDECHAR C)
 	{
 		const WIDECHAR *Last = nullptr;
 		
@@ -173,12 +173,12 @@ public:
 	CORE_API static float Atof(const WIDECHAR* String);
 	CORE_API static double Atod(const WIDECHAR* String);
 
-	CORE_API static FORCEINLINE int32 Atoi(const WIDECHAR* String)
+	static FORCEINLINE int32 Atoi(const WIDECHAR* String)
 	{
 		return Strtoi( String, NULL, 10 );
 	}
 	
-	CORE_API static FORCEINLINE int64 Atoi64(const WIDECHAR* String)
+	static FORCEINLINE int64 Atoi64(const WIDECHAR* String)
 	{
 		return Strtoi64( String, NULL, 10 );
 	}
@@ -192,7 +192,7 @@ public:
 	/**
 	 * Ansi implementation
 	 **/
-	CORE_API static FORCEINLINE ANSICHAR* Strcpy(ANSICHAR* Dest, SIZE_T DestCount, const ANSICHAR* Src)
+	static FORCEINLINE ANSICHAR* Strcpy(ANSICHAR* Dest, SIZE_T DestCount, const ANSICHAR* Src)
 	{
 // Skip suggestions about using strcpy_s instead.
 PRAGMA_DISABLE_DEPRECATION_WARNINGS
@@ -200,7 +200,7 @@ PRAGMA_DISABLE_DEPRECATION_WARNINGS
 PRAGMA_ENABLE_DEPRECATION_WARNINGS
 	}
 
-	CORE_API static FORCEINLINE ANSICHAR* Strncpy(ANSICHAR* Dest, const ANSICHAR* Src, SIZE_T MaxLen)
+	static FORCEINLINE ANSICHAR* Strncpy(ANSICHAR* Dest, const ANSICHAR* Src, SIZE_T MaxLen)
 	{
 // Skip suggestions about using strncpy_s instead.
 PRAGMA_DISABLE_DEPRECATION_WARNINGS
@@ -210,7 +210,7 @@ PRAGMA_ENABLE_DEPRECATION_WARNINGS
 		return Dest;
 	}
 
-	CORE_API static FORCEINLINE ANSICHAR* Strcat(ANSICHAR* Dest, SIZE_T DestCount, const ANSICHAR* Src)
+	static FORCEINLINE ANSICHAR* Strcat(ANSICHAR* Dest, SIZE_T DestCount, const ANSICHAR* Src)
 	{
 // Skip suggestions about using strcat_s instead.
 PRAGMA_DISABLE_DEPRECATION_WARNINGS
@@ -218,77 +218,77 @@ PRAGMA_DISABLE_DEPRECATION_WARNINGS
 PRAGMA_ENABLE_DEPRECATION_WARNINGS
 	}
 
-	CORE_API static FORCEINLINE int32 Strcmp( const ANSICHAR* String1, const ANSICHAR* String2 )
+	static FORCEINLINE int32 Strcmp( const ANSICHAR* String1, const ANSICHAR* String2 )
 	{
 		return strcmp(String1, String2);
 	}
 
-	CORE_API static FORCEINLINE int32 Strncmp( const ANSICHAR* String1, const ANSICHAR* String2, SIZE_T Count )
+	static FORCEINLINE int32 Strncmp( const ANSICHAR* String1, const ANSICHAR* String2, SIZE_T Count )
 	{
 		return strncmp( String1, String2, Count );
 	}
 
-	CORE_API static FORCEINLINE int32 Strlen( const ANSICHAR* String )
+	static FORCEINLINE int32 Strlen( const ANSICHAR* String )
 	{
 		return strlen( String );
 	}
 
-	CORE_API static FORCEINLINE int32 Strnlen( const ANSICHAR* String, SIZE_T StringSize )
+	static FORCEINLINE int32 Strnlen( const ANSICHAR* String, SIZE_T StringSize )
 	{
 		return strnlen( String, StringSize );
 	}
 
-	CORE_API static FORCEINLINE const ANSICHAR* Strstr( const ANSICHAR* String, const ANSICHAR* Find)
+	static FORCEINLINE const ANSICHAR* Strstr( const ANSICHAR* String, const ANSICHAR* Find)
 	{
 		return strstr(String, Find);
 	}
 
-	CORE_API static FORCEINLINE const ANSICHAR* Strchr( const ANSICHAR* String, ANSICHAR C)
+	static FORCEINLINE const ANSICHAR* Strchr( const ANSICHAR* String, ANSICHAR C)
 	{
 		return strchr(String, C);
 	}
 
-	CORE_API static FORCEINLINE const ANSICHAR* Strrchr( const ANSICHAR* String, ANSICHAR C)
+	static FORCEINLINE const ANSICHAR* Strrchr( const ANSICHAR* String, ANSICHAR C)
 	{
 		return strrchr(String, C);
 	}
 
-	CORE_API static FORCEINLINE int32 Atoi(const ANSICHAR* String)
+	static FORCEINLINE int32 Atoi(const ANSICHAR* String)
 	{
 		return atoi( String );
 	}
 
-	CORE_API static FORCEINLINE int64 Atoi64(const ANSICHAR* String)
+	static FORCEINLINE int64 Atoi64(const ANSICHAR* String)
 	{
 		return strtoll( String, NULL, 10 );
 	}
 
-	CORE_API static FORCEINLINE float Atof(const ANSICHAR* String)
+	static FORCEINLINE float Atof(const ANSICHAR* String)
 	{
 		return (float)atof( String );
 	}
 
-	CORE_API static FORCEINLINE double Atod(const ANSICHAR* String)
+	static FORCEINLINE double Atod(const ANSICHAR* String)
 	{
 		return atof( String );
 	}
 
-	CORE_API static FORCEINLINE int32 Strtoi( const ANSICHAR* Start, ANSICHAR** End, int32 Base )
+	static FORCEINLINE int32 Strtoi( const ANSICHAR* Start, ANSICHAR** End, int32 Base )
 	{
 		return strtol( Start, End, Base );
 	}
 
-	CORE_API static FORCEINLINE int64 Strtoi64( const ANSICHAR* Start, ANSICHAR** End, int32 Base )
+	static FORCEINLINE int64 Strtoi64( const ANSICHAR* Start, ANSICHAR** End, int32 Base )
 	{
 		return strtoll(Start, End, Base);
 	}
 
-	CORE_API static FORCEINLINE uint64 Strtoui64( const ANSICHAR* Start, ANSICHAR** End, int32 Base )
+	static FORCEINLINE uint64 Strtoui64( const ANSICHAR* Start, ANSICHAR** End, int32 Base )
 	{
 		return strtoull(Start, End, Base);
 	}
 
-	CORE_API static FORCEINLINE ANSICHAR* Strtok(ANSICHAR* StrToken, const ANSICHAR* Delim, ANSICHAR** Context)
+	static FORCEINLINE ANSICHAR* Strtok(ANSICHAR* StrToken, const ANSICHAR* Delim, ANSICHAR** Context)
 	{
 // Skip suggestions about using strtok_s instead.
 PRAGMA_DISABLE_DEPRECATION_WARNINGS
@@ -296,7 +296,7 @@ PRAGMA_DISABLE_DEPRECATION_WARNINGS
 PRAGMA_ENABLE_DEPRECATION_WARNINGS
 	}
 
-	static CORE_API int32 GetVarArgs( ANSICHAR* Dest, SIZE_T DestSize, const ANSICHAR*& Fmt, va_list ArgPtr )
+	static int32 GetVarArgs( ANSICHAR* Dest, SIZE_T DestSize, const ANSICHAR*& Fmt, va_list ArgPtr )
 	{
 		int32 Result = vsnprintf(Dest, DestSize, Fmt, ArgPtr);
 		return (Result != -1 && Result < (int32)DestSize) ? Result : -1;
@@ -404,7 +404,7 @@ PRAGMA_ENABLE_DEPRECATION_WARNINGS
 	 * UCS2 implementation
 	 **/
 
-	CORE_API static FORCEINLINE int32 Strlen( const UCS2CHAR* String )
+	static FORCEINLINE int32 Strlen( const UCS2CHAR* String )
 	{
 		int32 Result = 0;
 		while (*String++)
@@ -415,7 +415,7 @@ PRAGMA_ENABLE_DEPRECATION_WARNINGS
 		return Result;
 	}
 
-	CORE_API static FORCEINLINE int32 Strnlen( const UCS2CHAR* String, SIZE_T StringSize )
+	static FORCEINLINE int32 Strnlen( const UCS2CHAR* String, SIZE_T StringSize )
 	{
 		int32 Result = 0;
 		while (StringSize-- > 0 && *String++)
diff --git a/Engine/Source/Runtime/Core/Public/GenericPlatform/StandardPlatformString.h b/Engine/Source/Runtime/Core/Public/GenericPlatform/StandardPlatformString.h
--- a/Engine/Source/Runtime/Core/Public/GenericPlatform/StandardPlatformString.h
+++ b/Engine/Source/Runtime/Core/Public/GenericPlatform/StandardPlatformString.h
@@ -145,7 +145,7 @@ PRAGMA_ENABLE_DEPRECATION_WARNINGS
 	}
 
 #if PLATFORM_USE_SYSTEM_VSWPRINTF
-	static CORE_API int32 GetVarArgs( WIDECHAR* Dest, SIZE_T DestSize, const WIDECHAR*& Fmt, va_list ArgPtr )
+	static int32 GetVarArgs( WIDECHAR* Dest, SIZE_T DestSize, const WIDECHAR*& Fmt, va_list ArgPtr )
 	{
 #if PLATFORM_USE_LS_SPEC_FOR_WIDECHAR
 		// fix up the Fmt string, as fast as possible, without using an FString
@@ -315,7 +315,7 @@ PRAGMA_DISABLE_DEPRECATION_WARNINGS
 PRAGMA_ENABLE_DEPRECATION_WARNINGS
 	}
 
-	static CORE_API int32 GetVarArgs( ANSICHAR* Dest, SIZE_T DestSize, const ANSICHAR*& Fmt, va_list ArgPtr )
+	static int32 GetVarArgs( ANSICHAR* Dest, SIZE_T DestSize, const ANSICHAR*& Fmt, va_list ArgPtr )
 	{
 		int32 Result = vsnprintf(Dest, DestSize, Fmt, ArgPtr);
 		return (Result != -1 && Result < (int32)DestSize) ? Result : -1;
diff --git a/Engine/Source/Runtime/Core/Public/HAL/ConsoleManager.h b/Engine/Source/Runtime/Core/Public/HAL/ConsoleManager.h
--- a/Engine/Source/Runtime/Core/Public/HAL/ConsoleManager.h
+++ b/Engine/Source/Runtime/Core/Public/HAL/ConsoleManager.h
@@ -12,7 +12,7 @@
 class FOutputDevice;
 class UWorld;
 
-class CORE_API FConsoleManager :public IConsoleManager
+class FConsoleManager :public IConsoleManager
 {
 public:
 	/** constructor */
@@ -35,57 +35,57 @@ public:
 	}
 	
 	// internally needed or ECVF_RenderThreadSafe
-	IConsoleThreadPropagation* GetThreadPropagationCallback();
+	CORE_API IConsoleThreadPropagation* GetThreadPropagationCallback();
 	// internally needed or ECVF_RenderThreadSafe
-	bool IsThreadPropagationThread();
+	CORE_API bool IsThreadPropagationThread();
 
 	/** @param InVar must not be 0 */
-	virtual FString FindConsoleObjectName(const IConsoleObject* Obj) const override;
+	CORE_API virtual FString FindConsoleObjectName(const IConsoleObject* Obj) const override;
 
 	/** Can be moved out into some automated testing system */
-	void Test();
+	CORE_API void Test();
 
-	void OnCVarChanged();
+	CORE_API void OnCVarChanged();
 
-	virtual FConsoleVariableMulticastDelegate& OnCVarUnregistered()override;
+	CORE_API virtual FConsoleVariableMulticastDelegate& OnCVarUnregistered()override;
 
 	// interface IConsoleManager -----------------------------------
 
-	virtual IConsoleVariable* RegisterConsoleVariable(const TCHAR* Name, bool DefaultValue, const TCHAR* Help, uint32 Flags) override;
-	virtual IConsoleVariable* RegisterConsoleVariable(const TCHAR* Name, int32 DefaultValue, const TCHAR* Help, uint32 Flags) override;
-	virtual IConsoleVariable* RegisterConsoleVariable(const TCHAR* Name, float DefaultValue, const TCHAR* Help, uint32 Flags) override;
-	virtual IConsoleVariable* RegisterConsoleVariable(const TCHAR* Name, const TCHAR* DefaultValue, const TCHAR* Help, uint32 Flags) override;
-	virtual IConsoleVariable* RegisterConsoleVariable(const TCHAR* Name, const FString& DefaultValue, const TCHAR* Help, uint32 Flags) override;
-
-	virtual IConsoleVariable* RegisterConsoleVariableRef(const TCHAR* Name, bool& RefValue, const TCHAR* Help, uint32 Flags) override;
-	virtual IConsoleVariable* RegisterConsoleVariableRef(const TCHAR* Name, int32& RefValue, const TCHAR* Help, uint32 Flags) override;
-	virtual IConsoleVariable* RegisterConsoleVariableRef(const TCHAR* Name, float& RefValue, const TCHAR* Help, uint32 Flags) override;
-	virtual IConsoleVariable* RegisterConsoleVariableRef(const TCHAR* Name, FString& RefValue, const TCHAR* Help, uint32 Flags) override;
-	virtual IConsoleVariable* RegisterConsoleVariableBitRef(const TCHAR* CVarName, const TCHAR* FlagName, uint32 BitNumber, uint8* Force0MaskPtr, uint8* Force1MaskPtr, const TCHAR* Help, uint32 Flags) override;
+	CORE_API virtual IConsoleVariable* RegisterConsoleVariable(const TCHAR* Name, bool DefaultValue, const TCHAR* Help, uint32 Flags) override;
+	CORE_API virtual IConsoleVariable* RegisterConsoleVariable(const TCHAR* Name, int32 DefaultValue, const TCHAR* Help, uint32 Flags) override;
+	CORE_API virtual IConsoleVariable* RegisterConsoleVariable(const TCHAR* Name, float DefaultValue, const TCHAR* Help, uint32 Flags) override;
+	CORE_API virtual IConsoleVariable* RegisterConsoleVariable(const TCHAR* Name, const TCHAR* DefaultValue, const TCHAR* Help, uint32 Flags) override;
+	CORE_API virtual IConsoleVariable* RegisterConsoleVariable(const TCHAR* Name, const FString& DefaultValue, const TCHAR* Help, uint32 Flags) override;
+
+	CORE_API virtual IConsoleVariable* RegisterConsoleVariableRef(const TCHAR* Name, bool& RefValue, const TCHAR* Help, uint32 Flags) override;
+	CORE_API virtual IConsoleVariable* RegisterConsoleVariableRef(const TCHAR* Name, int32& RefValue, const TCHAR* Help, uint32 Flags) override;
+	CORE_API virtual IConsoleVariable* RegisterConsoleVariableRef(const TCHAR* Name, float& RefValue, const TCHAR* Help, uint32 Flags) override;
+	CORE_API virtual IConsoleVariable* RegisterConsoleVariableRef(const TCHAR* Name, FString& RefValue, const TCHAR* Help, uint32 Flags) override;
+	CORE_API virtual IConsoleVariable* RegisterConsoleVariableBitRef(const TCHAR* CVarName, const TCHAR* FlagName, uint32 BitNumber, uint8* Force0MaskPtr, uint8* Force1MaskPtr, const TCHAR* Help, uint32 Flags) override;
 	
-	virtual void CallAllConsoleVariableSinks() override;
-
-	virtual FConsoleVariableSinkHandle RegisterConsoleVariableSink_Handle(const FConsoleCommandDelegate& Command) override;
-	virtual void UnregisterConsoleVariableSink_Handle(FConsoleVariableSinkHandle Handle) override;
-
-	virtual IConsoleCommand* RegisterConsoleCommand(const TCHAR* Name, const TCHAR* Help, const FConsoleCommandDelegate& Command, uint32 Flags) override;
-	virtual IConsoleCommand* RegisterConsoleCommand(const TCHAR* Name, const TCHAR* Help, const FConsoleCommandWithArgsDelegate& Command, uint32 Flags) override;
-	virtual IConsoleCommand* RegisterConsoleCommand(const TCHAR* Name, const TCHAR* Help, const FConsoleCommandWithWorldDelegate& Command, uint32 Flags) override;
-	virtual IConsoleCommand* RegisterConsoleCommand(const TCHAR* Name, const TCHAR* Help, const FConsoleCommandWithWorldAndArgsDelegate& Command, uint32 Flags) override;
-	virtual IConsoleCommand* RegisterConsoleCommand(const TCHAR* Name, const TCHAR* Help, const FConsoleCommandWithArgsAndOutputDeviceDelegate& Command, uint32 Flags) override;
-	virtual IConsoleCommand* RegisterConsoleCommand(const TCHAR* Name, const TCHAR* Help, const FConsoleCommandWithWorldArgsAndOutputDeviceDelegate& Command, uint32 Flags) override;
-	virtual IConsoleCommand* RegisterConsoleCommand(const TCHAR* Name, const TCHAR* Help, const FConsoleCommandWithOutputDeviceDelegate& Command, uint32 Flags) override;
-	virtual IConsoleCommand* RegisterConsoleCommand(const TCHAR* Name, const TCHAR* Help, uint32 Flags) override;
-	virtual IConsoleObject* FindConsoleObject(const TCHAR* Name, bool bTrackFrequentCalls = true) const override;
-	virtual IConsoleVariable* FindConsoleVariable(const TCHAR* Name, bool bTrackFrequentCalls = true) const override;
-	virtual void ForEachConsoleObjectThatStartsWith(const FConsoleObjectVisitor& Visitor, const TCHAR* ThatStartsWith) const override;
-	virtual void ForEachConsoleObjectThatContains(const FConsoleObjectVisitor& Visitor, const TCHAR* ThatContains) const override;
-	virtual bool ProcessUserConsoleInput(const TCHAR* InInput, FOutputDevice& Ar, UWorld* InWorld) override;
-	virtual void AddConsoleHistoryEntry(const TCHAR* Key, const TCHAR* Input) override;
-	virtual void GetConsoleHistory(const TCHAR* Key, TArray<FString>& Out) override;
-	virtual bool IsNameRegistered(const TCHAR* Name) const override;	
-	virtual void RegisterThreadPropagation(uint32 ThreadId, IConsoleThreadPropagation* InCallback) override;
-	virtual void UnregisterConsoleObject( IConsoleObject* Object, bool bKeepState) override;
+	CORE_API virtual void CallAllConsoleVariableSinks() override;
+
+	CORE_API virtual FConsoleVariableSinkHandle RegisterConsoleVariableSink_Handle(const FConsoleCommandDelegate& Command) override;
+	CORE_API virtual void UnregisterConsoleVariableSink_Handle(FConsoleVariableSinkHandle Handle) override;
+
+	CORE_API virtual IConsoleCommand* RegisterConsoleCommand(const TCHAR* Name, const TCHAR* Help, const FConsoleCommandDelegate& Command, uint32 Flags) override;
+	CORE_API virtual IConsoleCommand* RegisterConsoleCommand(const TCHAR* Name, const TCHAR* Help, const FConsoleCommandWithArgsDelegate& Command, uint32 Flags) override;
+	CORE_API virtual IConsoleCommand* RegisterConsoleCommand(const TCHAR* Name, const TCHAR* Help, const FConsoleCommandWithWorldDelegate& Command, uint32 Flags) override;
+	CORE_API virtual IConsoleCommand* RegisterConsoleCommand(const TCHAR* Name, const TCHAR* Help, const FConsoleCommandWithWorldAndArgsDelegate& Command, uint32 Flags) override;
+	CORE_API virtual IConsoleCommand* RegisterConsoleCommand(const TCHAR* Name, const TCHAR* Help, const FConsoleCommandWithArgsAndOutputDeviceDelegate& Command, uint32 Flags) override;
+	CORE_API virtual IConsoleCommand* RegisterConsoleCommand(const TCHAR* Name, const TCHAR* Help, const FConsoleCommandWithWorldArgsAndOutputDeviceDelegate& Command, uint32 Flags) override;
+	CORE_API virtual IConsoleCommand* RegisterConsoleCommand(const TCHAR* Name, const TCHAR* Help, const FConsoleCommandWithOutputDeviceDelegate& Command, uint32 Flags) override;
+	CORE_API virtual IConsoleCommand* RegisterConsoleCommand(const TCHAR* Name, const TCHAR* Help, uint32 Flags) override;
+	CORE_API virtual IConsoleObject* FindConsoleObject(const TCHAR* Name, bool bTrackFrequentCalls = true) const override;
+	CORE_API virtual IConsoleVariable* FindConsoleVariable(const TCHAR* Name, bool bTrackFrequentCalls = true) const override;
+	CORE_API virtual void ForEachConsoleObjectThatStartsWith(const FConsoleObjectVisitor& Visitor, const TCHAR* ThatStartsWith) const override;
+	CORE_API virtual void ForEachConsoleObjectThatContains(const FConsoleObjectVisitor& Visitor, const TCHAR* ThatContains) const override;
+	CORE_API virtual bool ProcessUserConsoleInput(const TCHAR* InInput, FOutputDevice& Ar, UWorld* InWorld) override;
+	CORE_API virtual void AddConsoleHistoryEntry(const TCHAR* Key, const TCHAR* Input) override;
+	CORE_API virtual void GetConsoleHistory(const TCHAR* Key, TArray<FString>& Out) override;
+	CORE_API virtual bool IsNameRegistered(const TCHAR* Name) const override;	
+	CORE_API virtual void RegisterThreadPropagation(uint32 ThreadId, IConsoleThreadPropagation* InCallback) override;
+	CORE_API virtual void UnregisterConsoleObject( IConsoleObject* Object, bool bKeepState) override;
 
 private: // ----------------------------------------------------
 
@@ -116,25 +116,25 @@ private: // ----------------------------------------------------
 	 * @param Obj must not be 0
 	 * @return 0 if the name was already in use
 	 */
-	IConsoleObject* AddConsoleObject(const TCHAR* Name, IConsoleObject* Obj);
+	CORE_API IConsoleObject* AddConsoleObject(const TCHAR* Name, IConsoleObject* Obj);
 
 	/**
 	 * @param Stream must not be 0
 	 * @param Pattern must not be 0
 	 */
-	static bool MatchPartialName(const TCHAR* Stream, const TCHAR* Pattern);
+	static CORE_API bool MatchPartialName(const TCHAR* Stream, const TCHAR* Pattern);
 
 	/** Returns true if Pattern is found in Stream, case insensitive. */
-	static bool MatchSubstring(const TCHAR* Stream, const TCHAR* Pattern);
+	static CORE_API bool MatchSubstring(const TCHAR* Stream, const TCHAR* Pattern);
 
 	/**
 	 * Get string till whitespace, jump over whitespace
 	 * inefficient but this code is not performance critical
 	 */
-	static FString GetTextSection(const TCHAR* &It);
+	static CORE_API FString GetTextSection(const TCHAR* &It);
 
 	/** same as FindConsoleObject() but ECVF_CreatedFromIni are not filtered out (for internal use) */
-	IConsoleObject* FindConsoleObjectUnfiltered(const TCHAR* Name) const;
+	CORE_API IConsoleObject* FindConsoleObjectUnfiltered(const TCHAR* Name) const;
 
 	/**
 	 * Unregisters a console variable or command, if that object was registered.  For console variables, this will
@@ -144,11 +144,11 @@ private: // ----------------------------------------------------
 	 * @param	Name	Name of the console object to remove (not case sensitive)
 	 * @param	bKeepState	if the current state is kept in memory until a cvar with the same name is registered
 	 */
-	void UnregisterConsoleObject(const TCHAR* Name, bool bKeepState);
+	CORE_API void UnregisterConsoleObject(const TCHAR* Name, bool bKeepState);
 
 	// reads HistoryEntriesMap from the .ini file (if not already loaded)
-	void LoadHistoryIfNeeded();
+	CORE_API void LoadHistoryIfNeeded();
 
 	// writes HistoryEntriesMap to the .ini file
-	void SaveHistory();
+	CORE_API void SaveHistory();
 };
diff --git a/Engine/Source/Runtime/Core/Public/HAL/Event.h b/Engine/Source/Runtime/Core/Public/HAL/Event.h
--- a/Engine/Source/Runtime/Core/Public/HAL/Event.h
+++ b/Engine/Source/Runtime/Core/Public/HAL/Event.h
@@ -17,7 +17,7 @@
  * 
  * Consider using FEventRef as a safer and more convenient alternative.
  */
-class CORE_API FEvent
+class FEvent
 {
 public:
 
@@ -102,20 +102,20 @@ public:
 	// DO NOT MODIFY THESE
 
 	/** Advances stats associated with this event. Used to monitor wait->trigger history. */
-	void AdvanceStats();
+	CORE_API void AdvanceStats();
 
 protected:
 	/** Sends to the stats a special messages which encodes a wait for the event. */
-	void WaitForStats();
+	CORE_API void WaitForStats();
 
 	/** Send to the stats a special message which encodes a trigger for the event. */
-	void TriggerForStats();
+	CORE_API void TriggerForStats();
 
 	/** Resets start cycles to 0. */
-	void ResetForStats();
+	CORE_API void ResetForStats();
 
 	/** Counter used to generate an unique id for the events. */
-	static TAtomic<uint32> EventUniqueId;
+	static CORE_API TAtomic<uint32> EventUniqueId;
 
 	/** An unique id of this event. */
 	uint32 EventId;
@@ -131,12 +131,12 @@ enum class EEventMode { AutoReset, ManualReset };
  *
  * non-copyable, non-movable
  */
-class CORE_API FEventRef final
+class FEventRef final
 {
 public:
-	explicit FEventRef(EEventMode Mode = EEventMode::AutoReset);
+	CORE_API explicit FEventRef(EEventMode Mode = EEventMode::AutoReset);
 
-	~FEventRef();
+	CORE_API ~FEventRef();
 
 	FEventRef(const FEventRef&) = delete;
 	FEventRef& operator=(const FEventRef&) = delete;
@@ -160,10 +160,10 @@ private:
 /**
  * RAII-style shared and pooled `FEvent`
  */
-class CORE_API FSharedEventRef final
+class FSharedEventRef final
 {
 public:
-	explicit FSharedEventRef(EEventMode Mode = EEventMode::AutoReset);
+	CORE_API explicit FSharedEventRef(EEventMode Mode = EEventMode::AutoReset);
 
 	FEvent* operator->() const
 	{
@@ -172,4 +172,4 @@ public:
 
 private:
 	TSharedPtr<FEvent> Ptr;
-};
\ No newline at end of file
+};
diff --git a/Engine/Source/Runtime/Core/Public/HAL/FileManager.h b/Engine/Source/Runtime/Core/Public/HAL/FileManager.h
--- a/Engine/Source/Runtime/Core/Public/HAL/FileManager.h
+++ b/Engine/Source/Runtime/Core/Public/HAL/FileManager.h
@@ -54,7 +54,7 @@ enum EFileOpenFlags
 };
 
 
-class CORE_API IFileManager
+class IFileManager
 {
 protected:
 
@@ -64,7 +64,7 @@ protected:
 public:
 
 	/** Singleton access, platform specific, also calls PreInit() **/
-	static IFileManager& Get();
+	static CORE_API IFileManager& Get();
 
 	/** Allow the file manager to handle the commandline */
 	virtual void ProcessCommandLineOptions() = 0;
diff --git a/Engine/Source/Runtime/Core/Public/HAL/FileManagerGeneric.h b/Engine/Source/Runtime/Core/Public/HAL/FileManagerGeneric.h
--- a/Engine/Source/Runtime/Core/Public/HAL/FileManagerGeneric.h
+++ b/Engine/Source/Runtime/Core/Public/HAL/FileManagerGeneric.h
@@ -32,7 +32,7 @@
  * simple, unoptimized implementations of functions whose implementations
  * can be derived from other functions.
  */
-class CORE_API FFileManagerGeneric
+class FFileManagerGeneric
 	: public IFileManager
 {
 	// instead of caching the LowLevel, we call the singleton each time to never be incorrect
@@ -57,7 +57,7 @@ public:
 
 	// IFileManager interface
 
-	virtual void ProcessCommandLineOptions() override;
+	CORE_API virtual void ProcessCommandLineOptions() override;
 
 	virtual void SetSandboxEnabled(bool bInEnabled) override
 	{
@@ -86,25 +86,25 @@ public:
 	}
 #endif
 
-	bool Delete( const TCHAR* Filename, bool RequireExists=0, bool EvenReadOnly=0, bool Quiet=0 ) override;
-	bool IsReadOnly( const TCHAR* Filename ) override;
-	bool Move( const TCHAR* Dest, const TCHAR* Src, bool Replace=1, bool EvenIfReadOnly=0, bool Attributes=0, bool bDoNotRetryOrError=0 ) override;
-	bool FileExists( const TCHAR* Filename ) override;
-	bool DirectoryExists(const TCHAR* InDirectory) override;
-	void FindFiles( TArray<FString>& Result, const TCHAR* Filename, bool Files, bool Directories ) override;
-	void FindFilesRecursive( TArray<FString>& FileNames, const TCHAR* StartDirectory, const TCHAR* Filename, bool Files, bool Directories, bool bClearFileNames=true) override;
-	double GetFileAgeSeconds( const TCHAR* Filename ) override;
-	FDateTime GetTimeStamp( const TCHAR* Filename ) override;
-	FDateTime GetAccessTimeStamp( const TCHAR* Filename ) override;
-	void GetTimeStampPair(const TCHAR* PathA, const TCHAR* PathB, FDateTime& OutTimeStampA, FDateTime& OutTimeStampB);
-	bool SetTimeStamp( const TCHAR* Filename, FDateTime Timestamp ) override;
-	virtual FString GetFilenameOnDisk(const TCHAR* Filename) override;
-
-	virtual uint32	Copy( const TCHAR* Dest, const TCHAR* Src, bool Replace = 1, bool EvenIfReadOnly = 0, bool Attributes = 0, FCopyProgress* Progress = nullptr, EFileRead ReadFlags = FILEREAD_None, EFileWrite WriteFlags = FILEWRITE_None ) override;
-	virtual bool	MakeDirectory( const TCHAR* Path, bool Tree=0 ) override;
-	virtual bool	DeleteDirectory( const TCHAR* Path, bool RequireExists=0, bool Tree=0 ) override;
-
-	virtual FFileStatData GetStatData(const TCHAR* FilenameOrDirectory) override;
+	CORE_API bool Delete( const TCHAR* Filename, bool RequireExists=0, bool EvenReadOnly=0, bool Quiet=0 ) override;
+	CORE_API bool IsReadOnly( const TCHAR* Filename ) override;
+	CORE_API bool Move( const TCHAR* Dest, const TCHAR* Src, bool Replace=1, bool EvenIfReadOnly=0, bool Attributes=0, bool bDoNotRetryOrError=0 ) override;
+	CORE_API bool FileExists( const TCHAR* Filename ) override;
+	CORE_API bool DirectoryExists(const TCHAR* InDirectory) override;
+	CORE_API void FindFiles( TArray<FString>& Result, const TCHAR* Filename, bool Files, bool Directories ) override;
+	CORE_API void FindFilesRecursive( TArray<FString>& FileNames, const TCHAR* StartDirectory, const TCHAR* Filename, bool Files, bool Directories, bool bClearFileNames=true) override;
+	CORE_API double GetFileAgeSeconds( const TCHAR* Filename ) override;
+	CORE_API FDateTime GetTimeStamp( const TCHAR* Filename ) override;
+	CORE_API FDateTime GetAccessTimeStamp( const TCHAR* Filename ) override;
+	CORE_API void GetTimeStampPair(const TCHAR* PathA, const TCHAR* PathB, FDateTime& OutTimeStampA, FDateTime& OutTimeStampB);
+	CORE_API bool SetTimeStamp( const TCHAR* Filename, FDateTime Timestamp ) override;
+	CORE_API virtual FString GetFilenameOnDisk(const TCHAR* Filename) override;
+
+	CORE_API virtual uint32	Copy( const TCHAR* Dest, const TCHAR* Src, bool Replace = 1, bool EvenIfReadOnly = 0, bool Attributes = 0, FCopyProgress* Progress = nullptr, EFileRead ReadFlags = FILEREAD_None, EFileWrite WriteFlags = FILEWRITE_None ) override;
+	CORE_API virtual bool	MakeDirectory( const TCHAR* Path, bool Tree=0 ) override;
+	CORE_API virtual bool	DeleteDirectory( const TCHAR* Path, bool RequireExists=0, bool Tree=0 ) override;
+
+	CORE_API virtual FFileStatData GetStatData(const TCHAR* FilenameOrDirectory) override;
 
 	/**
 	 * Finds all the files within the given directory, with optional file extension filter.
@@ -116,7 +116,7 @@ public:
 	 *
 	 * @return FoundFiles, All the files that matched the optional FileExtension filter, or all files if none was specified.
 	 */
-	virtual void FindFiles(TArray<FString>& FoundFiles, const TCHAR* Directory, const TCHAR* FileExtension = nullptr) override;
+	CORE_API virtual void FindFiles(TArray<FString>& FoundFiles, const TCHAR* Directory, const TCHAR* FileExtension = nullptr) override;
 
 	/** 
 	 * Call the Visit function of the visitor once for each file or directory in a single directory. This function does not explore subdirectories.
@@ -124,8 +124,8 @@ public:
 	 * @param Visitor		Visitor to call for each element of the directory
 	 * @return				false if the directory did not exist or if the visitor returned false.
 	**/
-	bool IterateDirectory(const TCHAR* Directory, IPlatformFile::FDirectoryVisitor& Visitor) override;
-	bool IterateDirectory(const TCHAR* Directory, IPlatformFile::FDirectoryVisitorFunc Visitor) override;
+	CORE_API bool IterateDirectory(const TCHAR* Directory, IPlatformFile::FDirectoryVisitor& Visitor) override;
+	CORE_API bool IterateDirectory(const TCHAR* Directory, IPlatformFile::FDirectoryVisitorFunc Visitor) override;
 
 	/** 
 	 * Call the Visit function of the visitor once for each file or directory in a directory tree. This function explores subdirectories.
@@ -133,8 +133,8 @@ public:
 	 * @param Visitor		Visitor to call for each element of the directory and each element of all subdirectories.
 	 * @return				false if the directory did not exist or if the visitor returned false.
 	**/
-	bool IterateDirectoryRecursively(const TCHAR* Directory, IPlatformFile::FDirectoryVisitor& Visitor) override;
-	bool IterateDirectoryRecursively(const TCHAR* Directory, IPlatformFile::FDirectoryVisitorFunc Visitor) override;
+	CORE_API bool IterateDirectoryRecursively(const TCHAR* Directory, IPlatformFile::FDirectoryVisitor& Visitor) override;
+	CORE_API bool IterateDirectoryRecursively(const TCHAR* Directory, IPlatformFile::FDirectoryVisitorFunc Visitor) override;
 
 	/** 
 	 * Call the Visit function of the visitor once for each file or directory in a single directory. This function does not explore subdirectories.
@@ -142,8 +142,8 @@ public:
 	 * @param Visitor		Visitor to call for each element of the directory
 	 * @return				false if the directory did not exist or if the visitor returned false.
 	**/
-	bool IterateDirectoryStat(const TCHAR* Directory, IPlatformFile::FDirectoryStatVisitor& Visitor) override;
-	bool IterateDirectoryStat(const TCHAR* Directory, IPlatformFile::FDirectoryStatVisitorFunc Visitor) override;
+	CORE_API bool IterateDirectoryStat(const TCHAR* Directory, IPlatformFile::FDirectoryStatVisitor& Visitor) override;
+	CORE_API bool IterateDirectoryStat(const TCHAR* Directory, IPlatformFile::FDirectoryStatVisitorFunc Visitor) override;
 
 	/** 
 	 * Call the Visit function of the visitor once for each file or directory in a directory tree. This function explores subdirectories.
@@ -151,8 +151,8 @@ public:
 	 * @param Visitor		Visitor to call for each element of the directory and each element of all subdirectories.
 	 * @return				false if the directory did not exist or if the visitor returned false.
 	**/
-	bool IterateDirectoryStatRecursively(const TCHAR* Directory, IPlatformFile::FDirectoryStatVisitor& Visitor) override;
-	bool IterateDirectoryStatRecursively(const TCHAR* Directory, IPlatformFile::FDirectoryStatVisitorFunc Visitor) override;
+	CORE_API bool IterateDirectoryStatRecursively(const TCHAR* Directory, IPlatformFile::FDirectoryStatVisitor& Visitor) override;
+	CORE_API bool IterateDirectoryStatRecursively(const TCHAR* Directory, IPlatformFile::FDirectoryStatVisitorFunc Visitor) override;
 
 	/**
 	 * Converts passed in filename to use a relative path.
@@ -160,7 +160,7 @@ public:
 	 * @param	Filename	filename to convert to use a relative path
 	 * @return	filename using relative path
 	 */
-	static FString DefaultConvertToRelativePath( const TCHAR* Filename );
+	static CORE_API FString DefaultConvertToRelativePath( const TCHAR* Filename );
 
 	/**
 	 * Converts passed in filename to use a relative path.
@@ -168,7 +168,7 @@ public:
 	 * @param	Filename	filename to convert to use a relative path
 	 * @return	filename using relative path
 	 */
-	FString ConvertToRelativePath( const TCHAR* Filename ) override;
+	CORE_API FString ConvertToRelativePath( const TCHAR* Filename ) override;
 
 	/**
 	 * Converts passed in filename to use an absolute path (for reading)
@@ -176,7 +176,7 @@ public:
 	 * @param	Filename	filename to convert to use an absolute path, safe to pass in already using absolute path
 	 * @return	filename using absolute path
 	 */
-	FString ConvertToAbsolutePathForExternalAppForRead( const TCHAR* Filename ) override;
+	CORE_API FString ConvertToAbsolutePathForExternalAppForRead( const TCHAR* Filename ) override;
 
 	/**
 	 * Converts passed in filename to use an absolute path (for writing)
@@ -184,7 +184,7 @@ public:
 	 * @param	Filename	filename to convert to use an absolute path, safe to pass in already using absolute path
 	 * @return	filename using absolute path
 	 */
-	FString ConvertToAbsolutePathForExternalAppForWrite( const TCHAR* Filename ) override;
+	CORE_API FString ConvertToAbsolutePathForExternalAppForWrite( const TCHAR* Filename ) override;
 
 	/**
 	 *	Returns the size of a file. (Thread-safe)
@@ -192,7 +192,7 @@ public:
 	 *	@param Filename		Platform-independent Unreal filename.
 	 *	@return				File size in bytes or INDEX_NONE if the file didn't exist.
 	 **/
-	int64 FileSize( const TCHAR* Filename ) override;
+	CORE_API int64 FileSize( const TCHAR* Filename ) override;
 
 	/**
 	 * Sends a message to the file server, and will block until it's complete. Will return 
@@ -207,15 +207,15 @@ public:
 	}
 
 private:
-	FArchive * CreateFileReaderInternal(const TCHAR* Filename, uint32 ReadFlags, uint32 BufferSize);
-	FArchive* CreateFileWriterInternal(const TCHAR* Filename, uint32 WriteFlags, uint32 BufferSize);
+	CORE_API FArchive * CreateFileReaderInternal(const TCHAR* Filename, uint32 ReadFlags, uint32 BufferSize);
+	CORE_API FArchive* CreateFileWriterInternal(const TCHAR* Filename, uint32 WriteFlags, uint32 BufferSize);
 
 	/**
 	 * Helper called from Copy if Progress is available
 	 */
-	uint32	CopyWithProgress(const TCHAR* InDestFile, const TCHAR* InSrcFile, bool ReplaceExisting, bool EvenIfReadOnly, bool Attributes, FCopyProgress* Progress, EFileRead ReadFlags, EFileWrite WriteFlags);
+	CORE_API uint32	CopyWithProgress(const TCHAR* InDestFile, const TCHAR* InSrcFile, bool ReplaceExisting, bool EvenIfReadOnly, bool Attributes, FCopyProgress* Progress, EFileRead ReadFlags, EFileWrite WriteFlags);
 
-	void FindFilesRecursiveInternal( TArray<FString>& FileNames, const TCHAR* StartDirectory, const TCHAR* Filename, bool Files, bool Directories);
+	CORE_API void FindFilesRecursiveInternal( TArray<FString>& FileNames, const TCHAR* StartDirectory, const TCHAR* Filename, bool Files, bool Directories);
 };
 
 
@@ -223,13 +223,13 @@ private:
 	FArchiveFileReaderGeneric
 -----------------------------------------------------------------------------*/
 
-class CORE_API FArchiveFileReaderGeneric : public FArchive
+class FArchiveFileReaderGeneric : public FArchive
 {
 public:
-	FArchiveFileReaderGeneric( IFileHandle* InHandle, const TCHAR* InFilename, int64 InSize, uint32 InBufferSize = PLATFORM_FILE_READER_BUFFER_SIZE, uint32 InFlags = FILEREAD_None );
-	~FArchiveFileReaderGeneric();
+	CORE_API FArchiveFileReaderGeneric( IFileHandle* InHandle, const TCHAR* InFilename, int64 InSize, uint32 InBufferSize = PLATFORM_FILE_READER_BUFFER_SIZE, uint32 InFlags = FILEREAD_None );
+	CORE_API ~FArchiveFileReaderGeneric();
 
-	virtual void Seek( int64 InPos ) override final;
+	CORE_API virtual void Seek( int64 InPos ) override final;
 	virtual int64 Tell() override final
 	{
 		return Pos;
@@ -238,33 +238,33 @@ public:
 	{
 		return Size;
 	}
-	virtual bool Close() override final;
-	virtual void Serialize( void* V, int64 Length ) override final;
+	CORE_API virtual bool Close() override final;
+	CORE_API virtual void Serialize( void* V, int64 Length ) override final;
 	virtual FString GetArchiveName() const override
 	{
 		return Filename;
 	}
-	virtual void FlushCache() override final;
+	CORE_API virtual void FlushCache() override final;
 
-	virtual bool Precache(int64 PrecacheOffset, int64 PrecacheSize) override;
+	CORE_API virtual bool Precache(int64 PrecacheOffset, int64 PrecacheSize) override;
 
 protected:
-	bool InternalPrecache( int64 PrecacheOffset, int64 PrecacheSize );
+	CORE_API bool InternalPrecache( int64 PrecacheOffset, int64 PrecacheSize );
 	/** 
 	 * Platform specific seek
 	 * @param InPos - Offset from beginning of file to seek to
 	 * @return false on failure 
 	**/
-	virtual bool SeekLowLevel(int64 InPos);
+	CORE_API virtual bool SeekLowLevel(int64 InPos);
 	/** Close the file handle **/
-	virtual void CloseLowLevel();
+	CORE_API virtual void CloseLowLevel();
 	/** 
 	 * Platform specific read
 	 * @param Dest - Buffer to fill in
 	 * @param CountToRead - Number of bytes to read
 	 * @param OutBytesRead - Bytes actually read
 	**/
-	virtual void ReadLowLevel(uint8* Dest, int64 CountToRead, int64& OutBytesRead);
+	CORE_API virtual void ReadLowLevel(uint8* Dest, int64 CountToRead, int64& OutBytesRead);
 
 	/** Returns true if the archive should suppress logging in case of error */
 	bool IsSilent() const
@@ -300,21 +300,21 @@ protected:
 	FArchiveFileWriterGeneric
 -----------------------------------------------------------------------------*/
 
-class CORE_API FArchiveFileWriterGeneric : public FArchive
+class FArchiveFileWriterGeneric : public FArchive
 {
 public:
-	FArchiveFileWriterGeneric( IFileHandle* InHandle, const TCHAR* InFilename, int64 InPos, uint32 InBufferSize = PLATFORM_FILE_WRITER_BUFFER_SIZE, uint32 InFlags = FILEWRITE_None);
-	~FArchiveFileWriterGeneric();
+	CORE_API FArchiveFileWriterGeneric( IFileHandle* InHandle, const TCHAR* InFilename, int64 InPos, uint32 InBufferSize = PLATFORM_FILE_WRITER_BUFFER_SIZE, uint32 InFlags = FILEWRITE_None);
+	CORE_API ~FArchiveFileWriterGeneric();
 
-	virtual void Seek( int64 InPos ) override final;
+	CORE_API virtual void Seek( int64 InPos ) override final;
 	virtual int64 Tell() override final
 	{
 		return Pos;
 	}
-	virtual int64 TotalSize() override;
-	virtual bool Close() override final;
-	virtual void Serialize( void* V, int64 Length ) override final;
-	virtual void Flush() override final;
+	CORE_API virtual int64 TotalSize() override;
+	CORE_API virtual bool Close() override final;
+	CORE_API virtual void Serialize( void* V, int64 Length ) override final;
+	CORE_API virtual void Flush() override final;
 	virtual FString GetArchiveName() const override
 	{
 		return Filename;
@@ -326,25 +326,25 @@ protected:
 	 * @note Doesn't flush the handle itself, so this data may be cached by the OS and not yet written to disk!
 	 * @return true if there was buffer data and it was written successfully, false if there was nothing to flush or the write failed
 	 */
-	bool FlushBuffer();
+	CORE_API bool FlushBuffer();
 	/** 
 	 * Platform specific seek
 	 * @param InPos - Offset from beginning of file to seek to
 	 * @return false on failure 
 	**/
-	virtual bool SeekLowLevel(int64 InPos);
+	CORE_API virtual bool SeekLowLevel(int64 InPos);
 	/** 
 	 * Close the file handle
 	 * @return false on failure
 	**/
-	virtual bool CloseLowLevel();
+	CORE_API virtual bool CloseLowLevel();
 	/** 
 	 * Platform specific write
 	 * @param Src - Buffer to write out
 	 * @param CountToWrite - Number of bytes to write
 	 * @return false on failure 
 	**/
-	virtual bool WriteLowLevel(const uint8* Src, int64 CountToWrite);
+	CORE_API virtual bool WriteLowLevel(const uint8* Src, int64 CountToWrite);
 
 	/** 
 	 * Logs I/O error
@@ -353,7 +353,7 @@ protected:
 	 * logged in this message.
 	 * @param Message Brief description of what went wrong
 	 */
-	void LogWriteError(const TCHAR* Message);
+	CORE_API void LogWriteError(const TCHAR* Message);
 
 	/** Returns true if the archive should suppress logging in case of error */
 	bool IsSilent() const
diff --git a/Engine/Source/Runtime/Core/Public/HAL/IConsoleManager.h b/Engine/Source/Runtime/Core/Public/HAL/IConsoleManager.h
--- a/Engine/Source/Runtime/Core/Public/HAL/IConsoleManager.h
+++ b/Engine/Source/Runtime/Core/Public/HAL/IConsoleManager.h
@@ -714,7 +714,7 @@ public:
 /**
  * handles console commands and variables, registered console variables are released on destruction
  */
-struct CORE_API IConsoleManager
+struct IConsoleManager
 {
 	/**
 	 * Create a bool console variable
@@ -1038,7 +1038,7 @@ struct CORE_API IConsoleManager
 	 * @param DeviceProfileName If this is non-empty, the given deviceprofile will be loaded from the platform's inis and inserted into the CVars
 	 * @param Visit the callback to run for each CVar found
 	 */
-	static bool VisitPlatformCVarsForEmulation(FName PlatformName, const FString& DeviceProfileName, TFunctionRef<void(const FString& CVarName, const FString& CVarValue, EConsoleVariableFlags SetBy)> Visit);
+	static CORE_API bool VisitPlatformCVarsForEmulation(FName PlatformName, const FString& DeviceProfileName, TFunctionRef<void(const FString& CVarName, const FString& CVarValue, EConsoleVariableFlags SetBy)> Visit);
 #endif
 
 	virtual FConsoleVariableMulticastDelegate& OnCVarUnregistered() = 0;
@@ -1048,17 +1048,17 @@ protected:
 
 private:
 	/** Singleton for the console manager **/
-	static IConsoleManager* Singleton;
+	static CORE_API IConsoleManager* Singleton;
 
 	/** Function to create the singleton **/
-	static void SetupSingleton();
+	static CORE_API void SetupSingleton();
 };
 
 
 /**
  * auto registering console variable sinks (register a callback function that is called when ever a cvar is changes by the user, changes are grouped and happen in specific engine spots during the frame/main loop)
  */
-class CORE_API FAutoConsoleVariableSink
+class FAutoConsoleVariableSink
 {
 public:
 	/** Constructor, saves the argument for future removal from the console variable system **/
@@ -1081,7 +1081,7 @@ public:
 /**
  * Base class for autoregistering console commands.
  */
-class CORE_API FAutoConsoleObject
+class FAutoConsoleObject
 {
 protected:
 	/** Constructor, saves the argument for future removal from the console variable system **/
@@ -1109,9 +1109,9 @@ public:
 		IConsoleManager::Get().UnregisterConsoleObject(Target);
 	}
 
-	static TArray<const FAutoConsoleObject*>& AccessGeneralShaderChangeCvars();
-	static TArray<const FAutoConsoleObject*>& AccessMobileShaderChangeCvars();
-	static TArray<const FAutoConsoleObject*>& AccessDesktopShaderChangeCvars();
+	static CORE_API TArray<const FAutoConsoleObject*>& AccessGeneralShaderChangeCvars();
+	static CORE_API TArray<const FAutoConsoleObject*>& AccessMobileShaderChangeCvars();
+	static CORE_API TArray<const FAutoConsoleObject*>& AccessDesktopShaderChangeCvars();
 
 	/** returns the contained console object as an IConsoleVariable **/
 	FORCEINLINE IConsoleVariable* AsVariable()
@@ -1135,7 +1135,7 @@ private:
 /**
  * Autoregistering float, int or string console variable
  */
-class CORE_API FAutoConsoleVariable : private FAutoConsoleObject
+class FAutoConsoleVariable : private FAutoConsoleObject
 {
 public:
 	/**
@@ -1251,7 +1251,7 @@ public:
 	}
 };
 #else
-class CORE_API FAutoConsoleVariable
+class FAutoConsoleVariable
 {
 public:
 	FAutoConsoleVariable(const TCHAR* Name, int32 DefaultValue, const TCHAR* Help, uint32 Flags = ECVF_Default)
@@ -1272,7 +1272,7 @@ public:
 /**
  * Autoregistering float, int, bool, FString REF variable class...this changes that value when the console variable is changed. 
  */
-class CORE_API FAutoConsoleVariableRef : private FAutoConsoleObject
+class FAutoConsoleVariableRef : private FAutoConsoleObject
 {
 public:
 	/**
@@ -1391,7 +1391,7 @@ public:
 	}
 };
 #else
-class CORE_API FAutoConsoleVariableRef
+class FAutoConsoleVariableRef
 {
 public:
 	FAutoConsoleVariableRef(const TCHAR* Name, int32& RefValue, const TCHAR* Help, uint32 Flags = ECVF_Default)
@@ -1731,7 +1731,7 @@ private:
 /**
  * Autoregistering console command
  */
-class CORE_API FAutoConsoleCommand : private FAutoConsoleObject
+class FAutoConsoleCommand : private FAutoConsoleObject
 {
 public:
 	/**
@@ -1800,7 +1800,7 @@ public:
 	}
 };
 #else
-class CORE_API FAutoConsoleCommand
+class FAutoConsoleCommand
 {
 public:
 	FAutoConsoleCommand(const TCHAR* Name, const TCHAR* Help, const FConsoleCommandDelegate& Command, uint32 Flags = ECVF_Default)
@@ -1822,7 +1822,7 @@ public:
 /**
  * Autoregistering console command with a world
  */
-class CORE_API FAutoConsoleCommandWithWorld : private FAutoConsoleObject
+class FAutoConsoleCommandWithWorld : private FAutoConsoleObject
 {
 public:
 	/**
@@ -1844,7 +1844,7 @@ public:
 /**
  * Autoregistering console command with a world and arguments
  */
-class CORE_API FAutoConsoleCommandWithWorldAndArgs : private FAutoConsoleObject
+class FAutoConsoleCommandWithWorldAndArgs : private FAutoConsoleObject
 {
 public:	
 	/**
@@ -1864,7 +1864,7 @@ public:
 /**
  * Autoregistering console command with args and an output device
  */
-class CORE_API FAutoConsoleCommandWithArgsAndOutputDevice : private FAutoConsoleObject
+class FAutoConsoleCommandWithArgsAndOutputDevice : private FAutoConsoleObject
 {
 public:
 	
@@ -1877,7 +1877,7 @@ public:
 /**
  * Autoregistering console command with an output device
  */
-class CORE_API FAutoConsoleCommandWithOutputDevice : private FAutoConsoleObject
+class FAutoConsoleCommandWithOutputDevice : private FAutoConsoleObject
 {
 public:
 	/**
@@ -1897,7 +1897,7 @@ public:
 /**
  * Autoregistering console command with world, args, an output device
  */
-class CORE_API FAutoConsoleCommandWithWorldArgsAndOutputDevice : private FAutoConsoleObject
+class FAutoConsoleCommandWithWorldArgsAndOutputDevice : private FAutoConsoleObject
 {
 public:
 	/**
diff --git a/Engine/Source/Runtime/Core/Public/HAL/IPlatformFileCachedWrapper.h b/Engine/Source/Runtime/Core/Public/HAL/IPlatformFileCachedWrapper.h
--- a/Engine/Source/Runtime/Core/Public/HAL/IPlatformFileCachedWrapper.h
+++ b/Engine/Source/Runtime/Core/Public/HAL/IPlatformFileCachedWrapper.h
@@ -18,7 +18,7 @@
 
 class IAsyncReadFileHandle;
 
-class CORE_API FCachedFileHandle : public IFileHandle
+class FCachedFileHandle : public IFileHandle
 {
 public:
 	FCachedFileHandle(IFileHandle* InFileHandle, bool bInReadable, bool bInWritable)
@@ -252,7 +252,7 @@ private:
 	int32					CurrentCache;
 };
 
-class CORE_API FCachedReadPlatformFile : public IPlatformFile
+class FCachedReadPlatformFile : public IPlatformFile
 {
 	IPlatformFile*		LowerLevel;
 public:
diff --git a/Engine/Source/Runtime/Core/Public/HAL/IPlatformFileLogWrapper.h b/Engine/Source/Runtime/Core/Public/HAL/IPlatformFileLogWrapper.h
--- a/Engine/Source/Runtime/Core/Public/HAL/IPlatformFileLogWrapper.h
+++ b/Engine/Source/Runtime/Core/Public/HAL/IPlatformFileLogWrapper.h
@@ -39,7 +39,7 @@ extern bool bSuppressFileLog;
 
 class FLoggedPlatformFile;
 
-class CORE_API FLoggedFileHandle : public IFileHandle
+class FLoggedFileHandle : public IFileHandle
 {
 	TUniquePtr<IFileHandle>	FileHandle;
 	FString					Filename;
@@ -48,8 +48,8 @@ class CORE_API FLoggedFileHandle : public IFileHandle
 #endif
 public:
 
-	FLoggedFileHandle(IFileHandle* InFileHandle, const TCHAR* InFilename, FLoggedPlatformFile& InOwner);
-	virtual ~FLoggedFileHandle();
+	CORE_API FLoggedFileHandle(IFileHandle* InFileHandle, const TCHAR* InFilename, FLoggedPlatformFile& InOwner);
+	CORE_API virtual ~FLoggedFileHandle();
 
 	virtual int64		Tell() override
 	{
@@ -133,7 +133,7 @@ public:
 	}
 };
 
-class CORE_API FLoggedPlatformFile : public IPlatformFile
+class FLoggedPlatformFile : public IPlatformFile
 {
 	IPlatformFile* LowerLevel;
 
@@ -159,9 +159,9 @@ public:
 	using IPlatformFile::IterateDirectoryStat;
 	using IPlatformFile::IterateDirectoryStatRecursively;
 
-	virtual bool ShouldBeUsed(IPlatformFile* Inner, const TCHAR* CmdLine) const override;
+	CORE_API virtual bool ShouldBeUsed(IPlatformFile* Inner, const TCHAR* CmdLine) const override;
 
-	virtual bool Initialize(IPlatformFile* Inner, const TCHAR* CommandLineParam) override;
+	CORE_API virtual bool Initialize(IPlatformFile* Inner, const TCHAR* CommandLineParam) override;
 
 	IPlatformFile* GetLowerLevel() override
 	{
@@ -480,7 +480,7 @@ public:
 			OpenHandles.Remove(Filename);
 		}
 	}
-	void HandleDumpCommand(const TCHAR* Cmd, FOutputDevice& Ar);
+	CORE_API void HandleDumpCommand(const TCHAR* Cmd, FOutputDevice& Ar);
 #endif
 	virtual IAsyncReadFileHandle* OpenAsyncRead(const TCHAR* Filename) override
 	{
diff --git a/Engine/Source/Runtime/Core/Public/HAL/IPlatformFileManagedStorageWrapper.h b/Engine/Source/Runtime/Core/Public/HAL/IPlatformFileManagedStorageWrapper.h
--- a/Engine/Source/Runtime/Core/Public/HAL/IPlatformFileManagedStorageWrapper.h
+++ b/Engine/Source/Runtime/Core/Public/HAL/IPlatformFileManagedStorageWrapper.h
@@ -826,7 +826,7 @@ private:
 // NOTE: This is templated rather than a polymorphic wrapper because a lot code expects the physical layer not to be a wrapper.
 // It also has the benefit of not needing updating every time a new function is added to IPlatformFile.
 template<class BaseClass>
-class CORE_API TManagedStoragePlatformFile : public BaseClass
+class TManagedStoragePlatformFile : public BaseClass
 {
 private:
 	static bool IsReady()
@@ -1049,4 +1049,4 @@ public:
 
 		return bSuccess;
 	}
-};
\ No newline at end of file
+};
diff --git a/Engine/Source/Runtime/Core/Public/HAL/IPlatformFileOpenLogWrapper.h b/Engine/Source/Runtime/Core/Public/HAL/IPlatformFileOpenLogWrapper.h
--- a/Engine/Source/Runtime/Core/Public/HAL/IPlatformFileOpenLogWrapper.h
+++ b/Engine/Source/Runtime/Core/Public/HAL/IPlatformFileOpenLogWrapper.h
@@ -48,7 +48,7 @@ public:
 };
 
 
-class CORE_API FPlatformFileOpenLog : public IPlatformFile
+class FPlatformFileOpenLog : public IPlatformFile
 {
 protected:
 
diff --git a/Engine/Source/Runtime/Core/Public/HAL/LowLevelMemTracker.h b/Engine/Source/Runtime/Core/Public/HAL/LowLevelMemTracker.h
--- a/Engine/Source/Runtime/Core/Public/HAL/LowLevelMemTracker.h
+++ b/Engine/Source/Runtime/Core/Public/HAL/LowLevelMemTracker.h
@@ -508,7 +508,7 @@ struct FLLMTagInfo
 };
 
 /** The main LLM tracker class. */
-class CORE_API FLowLevelMemTracker
+class FLowLevelMemTracker
 {
 public:
 
@@ -521,18 +521,18 @@ public:
 			return Construct();
 	}
 
-	static FLowLevelMemTracker& Construct();
+	static CORE_API FLowLevelMemTracker& Construct();
 
-	static bool IsEnabled();
+	static CORE_API bool IsEnabled();
 
 	/**
 	 * We always start up running, but if the commandline disables us, we will do it later after main
 	 * (can't get the commandline early enough in a cross-platform way).
 	 */
-	void ProcessCommandLine(const TCHAR* CmdLine);
+	CORE_API void ProcessCommandLine(const TCHAR* CmdLine);
 
 	/** Return the total amount of memory being tracked. */
-	uint64 GetTotalTrackedMemory(ELLMTracker Tracker);
+	CORE_API uint64 GetTotalTrackedMemory(ELLMTracker Tracker);
 
 	/**
 	 * Records the use of memory that was allocated for a pointer.
@@ -545,9 +545,9 @@ public:
 	 * @param AllocType Type of allocation, FMalloc for regular allocations done through GMalloc, System for others.
 	 * @param bTrackInMemPro Whether to pass the allocation/free information on to FMemProProfiler.
 	 */
-	void OnLowLevelAlloc(ELLMTracker Tracker, const void* Ptr, uint64 Size, ELLMTag DefaultTag = ELLMTag::Untagged,
+	CORE_API void OnLowLevelAlloc(ELLMTracker Tracker, const void* Ptr, uint64 Size, ELLMTag DefaultTag = ELLMTag::Untagged,
 		ELLMAllocType AllocType = ELLMAllocType::None, bool bTrackInMemPro = true);
-	void OnLowLevelAlloc(ELLMTracker Tracker, const void* Ptr, uint64 Size, FName DefaultTag,
+	CORE_API void OnLowLevelAlloc(ELLMTracker Tracker, const void* Ptr, uint64 Size, FName DefaultTag,
 		ELLMAllocType AllocType = ELLMAllocType::None, bool bTrackInMemPro = true);
 
 	/**
@@ -559,7 +559,7 @@ public:
 	 * @param AllocType Type of allocation, FMalloc for regular allocations done through GMalloc, System for others.
 	 * @param bTrackInMemPro Whether to pass the allocation/free information on to FMemProProfiler.
 	 */
-	void OnLowLevelFree(ELLMTracker Tracker, const void* Ptr,
+	CORE_API void OnLowLevelFree(ELLMTracker Tracker, const void* Ptr,
 		ELLMAllocType AllocType = ELLMAllocType::None, bool bTrackInMemPro = true);
 
 	/**
@@ -571,101 +571,101 @@ public:
 	 * @param DefaultTag The tag to use if there is no tag already in scope on the callstack.
 	 * @param AllocType Type of allocation, FMalloc for regular allocations done through GMalloc, System for others.
 	 */
-	void OnLowLevelChangeInMemoryUse(ELLMTracker Tracker, int64 DeltaMemory, ELLMTag DefaultTag = ELLMTag::Untagged, ELLMAllocType AllocType = ELLMAllocType::None);
-	void OnLowLevelChangeInMemoryUse(ELLMTracker Tracker, int64 DeltaMemory, FName DefaultTag, ELLMAllocType AllocType = ELLMAllocType::None);
+	CORE_API void OnLowLevelChangeInMemoryUse(ELLMTracker Tracker, int64 DeltaMemory, ELLMTag DefaultTag = ELLMTag::Untagged, ELLMAllocType AllocType = ELLMAllocType::None);
+	CORE_API void OnLowLevelChangeInMemoryUse(ELLMTracker Tracker, int64 DeltaMemory, FName DefaultTag, ELLMAllocType AllocType = ELLMAllocType::None);
 
 	/** Call if an allocation is moved in memory, such as in a defragger. */
-	void OnLowLevelAllocMoved(ELLMTracker Tracker, const void* Dest, const void* Source,
+	CORE_API void OnLowLevelAllocMoved(ELLMTracker Tracker, const void* Dest, const void* Source,
 		ELLMAllocType AllocType = ELLMAllocType::None);
 
 	/** Updates memory stats and optionally publishes them to observers. */
-	void UpdateStatsPerFrame(const TCHAR* LogName=nullptr);
+	CORE_API void UpdateStatsPerFrame(const TCHAR* LogName=nullptr);
 	/** Updates memory stats. */
-	void Tick();
+	CORE_API void Tick();
 
 	/** Optionally set the amount of memory taken up before the game starts for executable and data segments .*/
-	void SetProgramSize(uint64 InProgramSize);
+	CORE_API void SetProgramSize(uint64 InProgramSize);
 
 	/** console command handler */
-	bool Exec(const TCHAR* Cmd, FOutputDevice& Ar);
+	CORE_API bool Exec(const TCHAR* Cmd, FOutputDevice& Ar);
 
 	/** Report whether the given TagSet is active, e.g. AssetTracking. */
-	bool IsTagSetActive(ELLMTagSet Set);
+	CORE_API bool IsTagSetActive(ELLMTagSet Set);
 
 	/** For some tag sets, it's really useful to reduce threads, to attribute allocations to assets, for instance. */
-	bool ShouldReduceThreads();
+	CORE_API bool ShouldReduceThreads();
 
 	/** Get the top active tag for the given tracker */
 	UE_DEPRECATED(4.27, "Tags have been changed to FNames and the old ELLMTag is now only the top-level coarse tag. Use GetActivateTagData instead to get the current Tag instead of its toplevel parent.")
-	int64 GetActiveTag(ELLMTracker Tracker);
+	CORE_API int64 GetActiveTag(ELLMTracker Tracker);
 
 	/** Get an opaque identifier for the top active tag for the given tracker. */
-	const UE::LLMPrivate::FTagData* GetActiveTagData(ELLMTracker Tracker, ELLMTagSet TagSet = ELLMTagSet::None);
+	CORE_API const UE::LLMPrivate::FTagData* GetActiveTagData(ELLMTracker Tracker, ELLMTagSet TagSet = ELLMTagSet::None);
 
 	/** Register custom ELLMTags. */
-	void RegisterPlatformTag(int32 Tag, const TCHAR* Name, FName StatName, FName SummaryStatName, int32 ParentTag = -1);
-	void RegisterProjectTag(int32 Tag, const TCHAR* Name, FName StatName, FName SummaryStatName, int32 ParentTag = -1);
+	CORE_API void RegisterPlatformTag(int32 Tag, const TCHAR* Name, FName StatName, FName SummaryStatName, int32 ParentTag = -1);
+	CORE_API void RegisterProjectTag(int32 Tag, const TCHAR* Name, FName StatName, FName SummaryStatName, int32 ParentTag = -1);
     
 	/** Get all tags being tracked. */
-	TArray<const UE::LLMPrivate::FTagData*> GetTrackedTags(ELLMTagSet TagSet = ELLMTagSet::None);
+	CORE_API TArray<const UE::LLMPrivate::FTagData*> GetTrackedTags(ELLMTagSet TagSet = ELLMTagSet::None);
 
 	/** Get all tags being tracked by the given tracker. */
-	TArray<const UE::LLMPrivate::FTagData*> GetTrackedTags(ELLMTracker Tracker, ELLMTagSet TagSet = ELLMTagSet::None);
+	CORE_API TArray<const UE::LLMPrivate::FTagData*> GetTrackedTags(ELLMTracker Tracker, ELLMTagSet TagSet = ELLMTagSet::None);
 
-	void GetTrackedTagsNamesWithAmount(TMap<FName, uint64>& TagsNamesWithAmount, ELLMTracker Tracker, ELLMTagSet TagSet);
-	void GetTrackedTagsNamesWithAmountFiltered(TMap<FName, uint64>& TagsNamesWithAmount, ELLMTracker Tracker, ELLMTagSet TagSet, TArray<FLLMTagSetAllocationFilter>& Filters);
+	CORE_API void GetTrackedTagsNamesWithAmount(TMap<FName, uint64>& TagsNamesWithAmount, ELLMTracker Tracker, ELLMTagSet TagSet);
+	CORE_API void GetTrackedTagsNamesWithAmountFiltered(TMap<FName, uint64>& TagsNamesWithAmount, ELLMTracker Tracker, ELLMTagSet TagSet, TArray<FLLMTagSetAllocationFilter>& Filters);
 
 	/** Look up the ELLMTag associated with the given display name. */
-	bool FindTagByName(const TCHAR* Name, uint64& OutTag, ELLMTagSet InTagSet = ELLMTagSet::None) const;
+	CORE_API bool FindTagByName(const TCHAR* Name, uint64& OutTag, ELLMTagSet InTagSet = ELLMTagSet::None) const;
 
 	UE_DEPRECATED(4.27, "Use FindTagDisplayName instead")
-	const TCHAR* FindTagName(uint64 Tag) const;
+	CORE_API const TCHAR* FindTagName(uint64 Tag) const;
 
 	/** Get the display name for the given ELLMTag. */
-	FName FindTagDisplayName(uint64 Tag) const;
+	CORE_API FName FindTagDisplayName(uint64 Tag) const;
 
 	/** Get the display name for the given FTagData. */
-	FName GetTagDisplayName(const UE::LLMPrivate::FTagData* TagData) const;
+	CORE_API FName GetTagDisplayName(const UE::LLMPrivate::FTagData* TagData) const;
 
 	/** Get the path name for the given FTagData from a chain of its parents' display names. */
-	FString GetTagDisplayPathName(const UE::LLMPrivate::FTagData* TagData) const;
+	CORE_API FString GetTagDisplayPathName(const UE::LLMPrivate::FTagData* TagData) const;
 
 	/** Get the unique identifier name for the given FTagData. */
-	FName GetTagUniqueName(const UE::LLMPrivate::FTagData* TagData) const;
+	CORE_API FName GetTagUniqueName(const UE::LLMPrivate::FTagData* TagData) const;
 
 	/** Get the amount of memory for an ELLMTag from the given tracker. */
-	int64 GetTagAmountForTracker(ELLMTracker Tracker, ELLMTag Tag, bool bPeakAmount = false);
+	CORE_API int64 GetTagAmountForTracker(ELLMTracker Tracker, ELLMTag Tag, bool bPeakAmount = false);
 
 	/** Get the amount of memory for a FTagData from the given tracker. */
-	int64 GetTagAmountForTracker(ELLMTracker Tracker, const UE::LLMPrivate::FTagData* TagData, bool bPeakAmount = false);
+	CORE_API int64 GetTagAmountForTracker(ELLMTracker Tracker, const UE::LLMPrivate::FTagData* TagData, bool bPeakAmount = false);
 
-    int64 GetTagAmountForTracker(ELLMTracker Tracker, FName Tag, ELLMTagSet TagSet, bool bPeakAmount = false);
+    CORE_API int64 GetTagAmountForTracker(ELLMTracker Tracker, FName Tag, ELLMTagSet TagSet, bool bPeakAmount = false);
 
 	/** Set the amount of memory for an ELLMTag for a given tracker and optionally update the total tracked memory. */
-	void SetTagAmountForTracker(ELLMTracker Tracker, ELLMTag Tag, int64 Amount, bool bAddToTotal);
+	CORE_API void SetTagAmountForTracker(ELLMTracker Tracker, ELLMTag Tag, int64 Amount, bool bAddToTotal);
 
 	/** Set the amount of memory for an tag name for a given tracker and optionally update the total tracked memory. */
-	void SetTagAmountForTracker(ELLMTracker Tracker, FName Tag, ELLMTagSet TagSet, int64 Amount, bool bAddToTotal);
+	CORE_API void SetTagAmountForTracker(ELLMTracker Tracker, FName Tag, ELLMTagSet TagSet, int64 Amount, bool bAddToTotal);
 	
 	/** Dump the display name of the current TagData for the given tracker to the output. */
-	uint64 DumpTag(ELLMTracker Tracker, const char* FileName, int LineNumber);
+	CORE_API uint64 DumpTag(ELLMTracker Tracker, const char* FileName, int LineNumber);
 
 	/** Publishes the active LLM stats in the active frame, useful for single targeted LLM snapshots. */
-	void PublishDataSingleFrame();
+	CORE_API void PublishDataSingleFrame();
 
 	enum class EDumpFormat
 	{
 		PlainText,
 		CSV,
 	};
-	void DumpToLog(EDumpFormat DumpFormat = EDumpFormat::PlainText, FOutputDevice* OutputDevice = nullptr);
+	CORE_API void DumpToLog(EDumpFormat DumpFormat = EDumpFormat::PlainText, FOutputDevice* OutputDevice = nullptr);
 
-	void OnPreFork();
+	CORE_API void OnPreFork();
 
 private:
-	FLowLevelMemTracker();
+	CORE_API FLowLevelMemTracker();
 
-	~FLowLevelMemTracker();
+	CORE_API ~FLowLevelMemTracker();
 
 	/**
 	 * Allocation and Setup of the data required for tracking allocations is done as late as possible to prevent
@@ -673,43 +673,43 @@ private:
 	 * than ParseCommandLine, so we do not complete all initialization in this function (e.g. features required only
 	 * for Update are omitted). Initialization done here will be torn down in ParseCommandLine if LLM is disabled.
 	 */
-	void BootstrapInitialise();
+	CORE_API void BootstrapInitialise();
 
 	bool IsBootstrapping() const { return bIsBootstrapping; }
 
 	/** Free all memory. This will put the tracker into a permanently disabled state. */
-	void Clear();
-	void InitialiseProgramSize();
+	CORE_API void Clear();
+	CORE_API void InitialiseProgramSize();
 
-	class UE::LLMPrivate::FLLMTracker* GetTracker(ELLMTracker Tracker);
+	CORE_API class UE::LLMPrivate::FLLMTracker* GetTracker(ELLMTracker Tracker);
 
-	void TickInternal();
-	void UpdateTags();
-	void SortTags(UE::LLMPrivate::FTagDataArray*& OutOldTagDatas);
-	void PublishDataPerFrame(const TCHAR* LogName);
+	CORE_API void TickInternal();
+	CORE_API void UpdateTags();
+	CORE_API void SortTags(UE::LLMPrivate::FTagDataArray*& OutOldTagDatas);
+	CORE_API void PublishDataPerFrame(const TCHAR* LogName);
 
-	void RegisterCustomTagInternal(int32 Tag, ELLMTagSet TagSet, const TCHAR* Name, FName StatName, FName SummaryStatName, int32 ParentTag = -1);
+	CORE_API void RegisterCustomTagInternal(int32 Tag, ELLMTagSet TagSet, const TCHAR* Name, FName StatName, FName SummaryStatName, int32 ParentTag = -1);
 	/**
 	* Called during C++ global static initialization when GMalloc is available (and hence FNames are unavailable)
 	* Creates the subset of tags necessary to record allocations during GMalloc and FName construction
 	*/
-	void BootstrapTagDatas();
+	CORE_API void BootstrapTagDatas();
 	/** Called when we have detected enabled in processcommandline, or as late as possible if callers access features that require full initialisation before then */
-	void FinishInitialise();
-	void InitialiseTagDatas_SetLLMTagNames();
-	void InitialiseTagDatas_FinishRegister();
-	void InitialiseTagDatas();
-	void ClearTagDatas();
-	void RegisterTagDeclaration(FLLMTagDeclaration& TagDeclaration);
-	UE::LLMPrivate::FTagData& RegisterTagData(FName Name, FName DisplayName, FName ParentName, FName StatName, FName SummaryStatName, bool bHasEnumTag, ELLMTag EnumTag, bool bIsStatTag, UE::LLMPrivate::ETagReferenceSource ReferenceSource, ELLMTagSet TagSet = ELLMTagSet::None);
+	CORE_API void FinishInitialise();
+	CORE_API void InitialiseTagDatas_SetLLMTagNames();
+	CORE_API void InitialiseTagDatas_FinishRegister();
+	CORE_API void InitialiseTagDatas();
+	CORE_API void ClearTagDatas();
+	CORE_API void RegisterTagDeclaration(FLLMTagDeclaration& TagDeclaration);
+	CORE_API UE::LLMPrivate::FTagData& RegisterTagData(FName Name, FName DisplayName, FName ParentName, FName StatName, FName SummaryStatName, bool bHasEnumTag, ELLMTag EnumTag, bool bIsStatTag, UE::LLMPrivate::ETagReferenceSource ReferenceSource, ELLMTagSet TagSet = ELLMTagSet::None);
 	/** Construct if not yet done the data on the given FTagData that relies on the presence of other TagDatas */
-	void FinishConstruct(UE::LLMPrivate::FTagData* TagData, UE::LLMPrivate::ETagReferenceSource ReferenceSource);
-	void ReportDuplicateTagName(UE::LLMPrivate::FTagData* TagData, UE::LLMPrivate::ETagReferenceSource ReferenceSource);
+	CORE_API void FinishConstruct(UE::LLMPrivate::FTagData* TagData, UE::LLMPrivate::ETagReferenceSource ReferenceSource);
+	CORE_API void ReportDuplicateTagName(UE::LLMPrivate::FTagData* TagData, UE::LLMPrivate::ETagReferenceSource ReferenceSource);
 
-	const UE::LLMPrivate::FTagData* FindOrAddTagData(ELLMTag EnumTag, UE::LLMPrivate::ETagReferenceSource ReferenceSource = UE::LLMPrivate::ETagReferenceSource::FunctionAPI);
-	const UE::LLMPrivate::FTagData* FindOrAddTagData(FName Name, ELLMTagSet TagSet, bool bIsStatData=false, UE::LLMPrivate::ETagReferenceSource ReferenceSource = UE::LLMPrivate::ETagReferenceSource::FunctionAPI);
-	const UE::LLMPrivate::FTagData* FindTagData(ELLMTag EnumTag, UE::LLMPrivate::ETagReferenceSource ReferenceSource = UE::LLMPrivate::ETagReferenceSource::FunctionAPI);
-	const UE::LLMPrivate::FTagData* FindTagData(FName Name, ELLMTagSet TagSet, UE::LLMPrivate::ETagReferenceSource ReferenceSource = UE::LLMPrivate::ETagReferenceSource::FunctionAPI);
+	CORE_API const UE::LLMPrivate::FTagData* FindOrAddTagData(ELLMTag EnumTag, UE::LLMPrivate::ETagReferenceSource ReferenceSource = UE::LLMPrivate::ETagReferenceSource::FunctionAPI);
+	CORE_API const UE::LLMPrivate::FTagData* FindOrAddTagData(FName Name, ELLMTagSet TagSet, bool bIsStatData=false, UE::LLMPrivate::ETagReferenceSource ReferenceSource = UE::LLMPrivate::ETagReferenceSource::FunctionAPI);
+	CORE_API const UE::LLMPrivate::FTagData* FindTagData(ELLMTag EnumTag, UE::LLMPrivate::ETagReferenceSource ReferenceSource = UE::LLMPrivate::ETagReferenceSource::FunctionAPI);
+	CORE_API const UE::LLMPrivate::FTagData* FindTagData(FName Name, ELLMTagSet TagSet, UE::LLMPrivate::ETagReferenceSource ReferenceSource = UE::LLMPrivate::ETagReferenceSource::FunctionAPI);
 
 	friend class FLLMPauseScope;
 	friend class FLLMScope;
@@ -753,13 +753,13 @@ private:
 	bool bPublishSingleFrame;
 	bool bCapturedSizeSnapshot;
 
-	static FLowLevelMemTracker* TrackerInstance;
+	static CORE_API FLowLevelMemTracker* TrackerInstance;
 public: // really internal but needs to be visible for LLM_IF_ENABLED macro
-	static bool bIsDisabled;
+	static CORE_API bool bIsDisabled;
 };
 
 /** LLM scope for tracking memory. */
-class CORE_API FLLMScope
+class FLLMScope
 {
 public:
 	FLLMScope(FName TagName, bool bIsStatTag, ELLMTagSet InTagSet, ELLMTracker InTracker, bool bOverride = true)
@@ -795,10 +795,10 @@ public:
 	}
 
 protected:
-	void Init(FName TagName, bool bIsStatTag, ELLMTagSet InTagSet, ELLMTracker InTracker, bool bOverride = true);
-	void Init(ELLMTag TagEnum, bool bIsStatTag, ELLMTagSet InTagSet, ELLMTracker InTracker, bool bOverride = true);
-	void Init(const UE::LLMPrivate::FTagData* TagData, bool bIsStatTag, ELLMTagSet InTagSet, ELLMTracker InTracker, bool bOverride = true);
-	void Destruct();
+	CORE_API void Init(FName TagName, bool bIsStatTag, ELLMTagSet InTagSet, ELLMTracker InTracker, bool bOverride = true);
+	CORE_API void Init(ELLMTag TagEnum, bool bIsStatTag, ELLMTagSet InTagSet, ELLMTracker InTracker, bool bOverride = true);
+	CORE_API void Init(const UE::LLMPrivate::FTagData* TagData, bool bIsStatTag, ELLMTagSet InTagSet, ELLMTracker InTracker, bool bOverride = true);
+	CORE_API void Destruct();
 
 	ELLMTracker Tracker{};
 	bool bEnabled = false;
@@ -806,14 +806,14 @@ protected:
 };
 
 /** LLM scope for pausing LLM (disables the allocation hooks). */
-class CORE_API FLLMPauseScope
+class FLLMPauseScope
 {
 public:
-	FLLMPauseScope(FName TagName, bool bIsStatTag, uint64 Amount, ELLMTracker TrackerToPause, ELLMAllocType InAllocType);
-	FLLMPauseScope(ELLMTag TagEnum, bool bIsStatTag, uint64 Amount, ELLMTracker TrackerToPause, ELLMAllocType InAllocType);
-	~FLLMPauseScope();
+	CORE_API FLLMPauseScope(FName TagName, bool bIsStatTag, uint64 Amount, ELLMTracker TrackerToPause, ELLMAllocType InAllocType);
+	CORE_API FLLMPauseScope(ELLMTag TagEnum, bool bIsStatTag, uint64 Amount, ELLMTracker TrackerToPause, ELLMAllocType InAllocType);
+	CORE_API ~FLLMPauseScope();
 protected:
-	void Init(FName TagName, ELLMTag EnumTag, bool bIsEnumTag, bool bIsStatTag, uint64 Amount,
+	CORE_API void Init(FName TagName, ELLMTag EnumTag, bool bIsEnumTag, bool bIsStatTag, uint64 Amount,
 		ELLMTracker TrackerToPause, ELLMAllocType InAllocType);
 	ELLMTracker PausedTracker;
 	ELLMAllocType AllocType;
@@ -821,11 +821,11 @@ protected:
 };
 
 /** LLM scope for inheriting tag from the given address. */
-class CORE_API FLLMScopeFromPtr
+class FLLMScopeFromPtr
 {
 public:
-	FLLMScopeFromPtr(void* Ptr, ELLMTracker Tracker);
-	~FLLMScopeFromPtr();
+	CORE_API FLLMScopeFromPtr(void* Ptr, ELLMTracker Tracker);
+	CORE_API ~FLLMScopeFromPtr();
 protected:
 	ELLMTracker Tracker;
 	bool bEnabled[static_cast<int32>(ELLMTagSet::Max)];
@@ -838,24 +838,24 @@ private:
 };
 
 /** Global instances to provide information about a tag to LLM. */
-class CORE_API FLLMTagDeclaration
+class FLLMTagDeclaration
 {
 public:
-	FLLMTagDeclaration(const TCHAR* InCPPName, const FName InDisplayName=NAME_None, FName InParentTagName = NAME_None, FName InStatName = NAME_None, FName InSummaryStatName = NAME_None, ELLMTagSet TagSet = ELLMTagSet::None);
+	CORE_API FLLMTagDeclaration(const TCHAR* InCPPName, const FName InDisplayName=NAME_None, FName InParentTagName = NAME_None, FName InStatName = NAME_None, FName InSummaryStatName = NAME_None, ELLMTagSet TagSet = ELLMTagSet::None);
 	FName GetUniqueName() const { return UniqueName; }
 
 	typedef void (*FCreationCallback)(FLLMTagDeclaration&);
 protected:
 
-	static void AddCreationCallback(FCreationCallback InCallback);
-	static void ClearCreationCallbacks();
-	static TArrayView<FCreationCallback> GetCreationCallbacks();
-	static FLLMTagDeclaration* GetList();
+	static CORE_API void AddCreationCallback(FCreationCallback InCallback);
+	static CORE_API void ClearCreationCallbacks();
+	static CORE_API TArrayView<FCreationCallback> GetCreationCallbacks();
+	static CORE_API FLLMTagDeclaration* GetList();
 
-	void Register();
+	CORE_API void Register();
 
 protected:
-	void ConstructUniqueName();
+	CORE_API void ConstructUniqueName();
 
 	const TCHAR* CPPName;
 	FName UniqueName;
@@ -871,7 +871,7 @@ protected:
 };
 
 /** Used to filter Tag allocations specifying a Name that matches in a TagSet, or TagSet alone or just by Name. */
-struct CORE_API FLLMTagSetAllocationFilter
+struct FLLMTagSetAllocationFilter
 {
 	/** Tag name to match */
 	FName Name;
diff --git a/Engine/Source/Runtime/Core/Public/HAL/MallocBinned2.h b/Engine/Source/Runtime/Core/Public/HAL/MallocBinned2.h
--- a/Engine/Source/Runtime/Core/Public/HAL/MallocBinned2.h
+++ b/Engine/Source/Runtime/Core/Public/HAL/MallocBinned2.h
@@ -127,7 +127,7 @@ enum class EBlockCanary : uint8
 //
 // Optimized virtual memory allocator.
 //
-class CORE_API FMallocBinned2 : public FMalloc
+class FMallocBinned2 : public FMalloc
 {
 	// Forward declares.
 	struct FPoolInfo;
@@ -469,12 +469,12 @@ class CORE_API FMallocBinned2 : public FMalloc
 public:
 
 
-	FMallocBinned2();
+	CORE_API FMallocBinned2();
 
-	virtual ~FMallocBinned2();
+	CORE_API virtual ~FMallocBinned2();
 
 	// FMalloc interface.
-	virtual bool IsInternallyThreadSafe() const override;
+	CORE_API virtual bool IsInternallyThreadSafe() const override;
 	FORCEINLINE virtual void* Malloc(SIZE_T Size, uint32 Alignment) override
 	{
 #if BINNED2_ALLOCATOR_STATS_VALIDATION
@@ -757,41 +757,41 @@ if (NewSize <= BINNED2_MAX_SMALL_POOL_SIZE && Alignment <= BINNED2_MINIMUM_ALIGN
 		return SizeOut;
 	}
 
-	virtual bool ValidateHeap() override;
-	virtual void Trim(bool bTrimThreadCaches) override;
-	virtual void SetupTLSCachesOnCurrentThread() override;
-	virtual void ClearAndDisableTLSCachesOnCurrentThread() override;
-	virtual const TCHAR* GetDescriptiveName() override;
-	virtual void UpdateStats() override;
-	virtual void OnMallocInitialized() override;
-	virtual void OnPreFork() override;
-	virtual void OnPostFork() override;
+	CORE_API virtual bool ValidateHeap() override;
+	CORE_API virtual void Trim(bool bTrimThreadCaches) override;
+	CORE_API virtual void SetupTLSCachesOnCurrentThread() override;
+	CORE_API virtual void ClearAndDisableTLSCachesOnCurrentThread() override;
+	CORE_API virtual const TCHAR* GetDescriptiveName() override;
+	CORE_API virtual void UpdateStats() override;
+	CORE_API virtual void OnMallocInitialized() override;
+	CORE_API virtual void OnPreFork() override;
+	CORE_API virtual void OnPostFork() override;
 	// End FMalloc interface.
 
-	void FlushCurrentThreadCache();
-	void* MallocExternalSmall(SIZE_T Size, uint32 Alignment);
-	void* MallocExternalLarge(SIZE_T Size, uint32 Alignment);
-	void* ReallocExternal(void* Ptr, SIZE_T NewSize, uint32 Alignment);
-	void FreeExternal(void *Ptr);
-	bool GetAllocationSizeExternal(void* Ptr, SIZE_T& SizeOut);
+	CORE_API void FlushCurrentThreadCache();
+	CORE_API void* MallocExternalSmall(SIZE_T Size, uint32 Alignment);
+	CORE_API void* MallocExternalLarge(SIZE_T Size, uint32 Alignment);
+	CORE_API void* ReallocExternal(void* Ptr, SIZE_T NewSize, uint32 Alignment);
+	CORE_API void FreeExternal(void *Ptr);
+	CORE_API bool GetAllocationSizeExternal(void* Ptr, SIZE_T& SizeOut);
 
-	void CanaryTest(const FFreeBlock* Block) const;
-	void CanaryFail(const FFreeBlock* Block) const;
+	CORE_API void CanaryTest(const FFreeBlock* Block) const;
+	CORE_API void CanaryFail(const FFreeBlock* Block) const;
 
 #if BINNED2_ALLOCATOR_STATS
-	int64 GetTotalAllocatedSmallPoolMemory() const;
+	CORE_API int64 GetTotalAllocatedSmallPoolMemory() const;
 #endif
-	virtual void GetAllocatorStats( FGenericMemoryStats& out_Stats ) override;
+	CORE_API virtual void GetAllocatorStats( FGenericMemoryStats& out_Stats ) override;
 	/** Dumps current allocator stats to the log. */
-	virtual void DumpAllocatorStats(class FOutputDevice& Ar) override;
+	CORE_API virtual void DumpAllocatorStats(class FOutputDevice& Ar) override;
 	
-	static uint16 SmallBlockSizesReversed[BINNED2_SMALL_POOL_COUNT]; // this is reversed to get the smallest elements on our main cache line
-	static FMallocBinned2* MallocBinned2;
-	static uint32 Binned2TlsSlot;
-	static uint32 PageSize;
-	static uint32 OsAllocationGranularity;
+	static CORE_API uint16 SmallBlockSizesReversed[BINNED2_SMALL_POOL_COUNT]; // this is reversed to get the smallest elements on our main cache line
+	static CORE_API FMallocBinned2* MallocBinned2;
+	static CORE_API uint32 Binned2TlsSlot;
+	static CORE_API uint32 PageSize;
+	static CORE_API uint32 OsAllocationGranularity;
 	// Mapping of sizes to small table indices
-	static uint8 MemSizeToIndex[1 + (BINNED2_MAX_SMALL_POOL_SIZE >> BINNED2_MINIMUM_ALIGNMENT_SHIFT)];
+	static CORE_API uint8 MemSizeToIndex[1 + (BINNED2_MAX_SMALL_POOL_SIZE >> BINNED2_MINIMUM_ALIGNMENT_SHIFT)];
 
 	FORCEINLINE uint32 BoundSizeToPoolIndex(SIZE_T Size) 
 	{
diff --git a/Engine/Source/Runtime/Core/Public/HAL/MallocDoubleFreeFinder.h b/Engine/Source/Runtime/Core/Public/HAL/MallocDoubleFreeFinder.h
--- a/Engine/Source/Runtime/Core/Public/HAL/MallocDoubleFreeFinder.h
+++ b/Engine/Source/Runtime/Core/Public/HAL/MallocDoubleFreeFinder.h
@@ -10,16 +10,16 @@ class FMalloc;
 class FOutputDevice;
 class UWorld;
 
-class CORE_API FMallocDoubleFreeFinder final : public FMallocCallstackHandler
+class FMallocDoubleFreeFinder final : public FMallocCallstackHandler
 {
 public:
-	FMallocDoubleFreeFinder(FMalloc* InMalloc);
+	CORE_API FMallocDoubleFreeFinder(FMalloc* InMalloc);
 
 #if UE_ALLOW_EXEC_COMMANDS
 	/**
 	 * Handles any commands passed in on the command line
 	 */
-	virtual bool Exec(UWorld* InWorld, const TCHAR* Cmd, FOutputDevice& Ar) override;
+	CORE_API virtual bool Exec(UWorld* InWorld, const TCHAR* Cmd, FOutputDevice& Ar) override;
 #endif
 
 	/**
@@ -27,15 +27,15 @@ public:
 	 * It will dump a callstack of the last allocator free most likely to have caused the problem to the log, if you have symbols loaded
 	 * Might be useful to pass an access violation ptr to this!
 	 */
-	void TrackSpecial(void* Ptr);
+	CORE_API void TrackSpecial(void* Ptr);
 
-	virtual void Init();
+	CORE_API virtual void Init();
 
-	static FMalloc* OverrideIfEnabled(FMalloc*InUsedAlloc);
+	static CORE_API FMalloc* OverrideIfEnabled(FMalloc*InUsedAlloc);
 
 protected:
-	virtual void TrackMalloc(void* Ptr, uint32 Size, int32 CallStackIndex);
-	virtual void TrackFree(void* Ptr, uint32 OldSize, int32 CallStackIndex);
+	CORE_API virtual void TrackMalloc(void* Ptr, uint32 Size, int32 CallStackIndex);
+	CORE_API virtual void TrackFree(void* Ptr, uint32 OldSize, int32 CallStackIndex);
 
 	struct TrackedAllocationData
 	{
diff --git a/Engine/Source/Runtime/Core/Public/HAL/MallocFrameProfiler.h b/Engine/Source/Runtime/Core/Public/HAL/MallocFrameProfiler.h
--- a/Engine/Source/Runtime/Core/Public/HAL/MallocFrameProfiler.h
+++ b/Engine/Source/Runtime/Core/Public/HAL/MallocFrameProfiler.h
@@ -11,33 +11,33 @@ class FMalloc;
 class FOutputDevice;
 class UWorld;
 
-class CORE_API FMallocFrameProfiler final : public FMallocCallstackHandler
+class FMallocFrameProfiler final : public FMallocCallstackHandler
 {
 public:
-	FMallocFrameProfiler(FMalloc* InMalloc);
+	CORE_API FMallocFrameProfiler(FMalloc* InMalloc);
 
-	virtual void Init() override;
+	CORE_API virtual void Init() override;
 	
 #if UE_ALLOW_EXEC_COMMANDS
 	/**
 	 * Handles any commands passed in on the command line
 	 */
-	virtual bool Exec(UWorld* InWorld, const TCHAR* Cmd, FOutputDevice& Ar) override;
+	CORE_API virtual bool Exec(UWorld* InWorld, const TCHAR* Cmd, FOutputDevice& Ar) override;
 #endif
 
 	/** 
 	* Called once per frame, gathers and sets all memory allocator statistics into the corresponding stats. MUST BE THREAD SAFE.
 	*/
-	virtual void UpdateStats() override;
+	CORE_API virtual void UpdateStats() override;
 
-	static FMalloc* OverrideIfEnabled(FMalloc*InUsedAlloc);
+	static CORE_API FMalloc* OverrideIfEnabled(FMalloc*InUsedAlloc);
 
 protected:
-	virtual bool IsDisabled() override;
+	CORE_API virtual bool IsDisabled() override;
 
-	virtual void TrackMalloc(void* Ptr, uint32 Size, int32 CallStackIndex);
-	virtual void TrackFree(void* Ptr, uint32 OldSize, int32 CallStackIndex);
-	virtual void TrackRealloc(void* OldPtr, void* NewPtr, uint32 NewSize, uint32 OldSize, int32 CallStackIndex);
+	CORE_API virtual void TrackMalloc(void* Ptr, uint32 Size, int32 CallStackIndex);
+	CORE_API virtual void TrackFree(void* Ptr, uint32 OldSize, int32 CallStackIndex);
+	CORE_API virtual void TrackRealloc(void* OldPtr, void* NewPtr, uint32 NewSize, uint32 OldSize, int32 CallStackIndex);
 
 protected:
 	bool bEnabled;
diff --git a/Engine/Source/Runtime/Core/Public/HAL/MallocLeakDetection.h b/Engine/Source/Runtime/Core/Public/HAL/MallocLeakDetection.h
--- a/Engine/Source/Runtime/Core/Public/HAL/MallocLeakDetection.h
+++ b/Engine/Source/Runtime/Core/Public/HAL/MallocLeakDetection.h
@@ -62,7 +62,7 @@ struct FMallocLeakReportOptions
 /**
  * Maintains a list of all pointers to currently allocated memory.
  */
-class CORE_API FMallocLeakDetection
+class FMallocLeakDetection
 {
 	struct FCallstackTrack
 	{
@@ -119,14 +119,14 @@ class CORE_API FMallocLeakDetection
 
 private:
 
-	FMallocLeakDetection();
-	~FMallocLeakDetection();
+	CORE_API FMallocLeakDetection();
+	CORE_API ~FMallocLeakDetection();
 
 	/** Track a callstack */
 
 	/** Stop tracking a callstack */
-	void AddCallstack(FCallstackTrack& Callstack);
-	void RemoveCallstack(FCallstackTrack& Callstack);
+	CORE_API void AddCallstack(FCallstackTrack& Callstack);
+	CORE_API void RemoveCallstack(FCallstackTrack& Callstack);
 
 	/** List of all currently allocated pointers */
 	TMap<void*, FCallstackTrack> OpenPointers;
@@ -167,38 +167,38 @@ private:
 
 public:	
 
-	static FMallocLeakDetection& Get();
-	static void HandleMallocLeakCommand(const TArray< FString >& Args);
+	static CORE_API FMallocLeakDetection& Get();
+	static CORE_API void HandleMallocLeakCommand(const TArray< FString >& Args);
 
 	/** Enable/disable collection of allocation with an optional filter on allocation size */
-	void SetAllocationCollection(bool bEnabled, int32 Size = 0);
+	CORE_API void SetAllocationCollection(bool bEnabled, int32 Size = 0);
 
 	/** Returns state of allocation collection */
 	bool IsAllocationCollectionEnabled(void) const { return bCaptureAllocs; }
 
 	/** Clear currently accumulated data */
-	void ClearData();
+	CORE_API void ClearData();
 
 	/** Dumps currently open callstacks */
-	int32 DumpOpenCallstacks(const TCHAR* FileName, const FMallocLeakReportOptions& Options = FMallocLeakReportOptions());
+	CORE_API int32 DumpOpenCallstacks(const TCHAR* FileName, const FMallocLeakReportOptions& Options = FMallocLeakReportOptions());
 
 	/** Perform a linear fit checkpoint of all open callstacks */
-	void CheckpointLinearFit();
+	CORE_API void CheckpointLinearFit();
 
 	/** Handles new allocated pointer */
-	void Malloc(void* Ptr, SIZE_T Size);
+	CORE_API void Malloc(void* Ptr, SIZE_T Size);
 
 	/** Handles reallocation */
-	void Realloc(void* OldPtr, SIZE_T OldSize, void* NewPtr, SIZE_T NewSize);
+	CORE_API void Realloc(void* OldPtr, SIZE_T OldSize, void* NewPtr, SIZE_T NewSize);
 
 	/** Removes allocated pointer from list */
-	void Free(void* Ptr);	
+	CORE_API void Free(void* Ptr);	
 
 	/** Disabled allocation tracking for this thread, @see MALLOCLEAK_IGNORE_SCOPE and FMallocLeakDetectionIgnoreScope. */
-	void SetDisabledForThisThread(const bool Disabled);
+	CORE_API void SetDisabledForThisThread(const bool Disabled);
 
 	/** Returns true of allocation tracking for this thread is  */
-	bool IsDisabledForThisThread() const;
+	CORE_API bool IsDisabledForThisThread() const;
 
 	/** Pushes context that will be associated with allocations. All open contexts will be displayed alongside
 	callstacks in a report.  */
@@ -206,13 +206,13 @@ public:
 	{
 		this->PushContext(*Context);
 	}
-	void PushContext(const TCHAR* Context);
+	CORE_API void PushContext(const TCHAR* Context);
 
 	/** Pops a context from the above */
-	void PopContext();
+	CORE_API void PopContext();
 
 	/** Returns */
-	void GetOpenCallstacks(TArray<uint32>& OutCallstacks, SIZE_T& OutTotalSize, const FMallocLeakReportOptions& Options = FMallocLeakReportOptions());
+	CORE_API void GetOpenCallstacks(TArray<uint32>& OutCallstacks, SIZE_T& OutTotalSize, const FMallocLeakReportOptions& Options = FMallocLeakReportOptions());
 };
 
 /**
diff --git a/Engine/Source/Runtime/Core/Public/HAL/MallocTimer.h b/Engine/Source/Runtime/Core/Public/HAL/MallocTimer.h
--- a/Engine/Source/Runtime/Core/Public/HAL/MallocTimer.h
+++ b/Engine/Source/Runtime/Core/Public/HAL/MallocTimer.h
@@ -9,7 +9,7 @@
 #define UE_TIME_VIRTUALMALLOC 0
 #endif
 
-struct CORE_API FScopedVirtualMallocTimer
+struct FScopedVirtualMallocTimer
 {
 	enum IndexType : int32 
 	{
@@ -32,8 +32,8 @@ struct CORE_API FScopedVirtualMallocTimer
 	};
 
 #if UE_TIME_VIRTUALMALLOC
-	static uint64 GTotalCycles[IndexType::Max][PlatormIndexType::PlatormIndexTypeMax];
-	static uint64 GTotalCounts[IndexType::Max][PlatormIndexType::PlatormIndexTypeMax];
+	static CORE_API uint64 GTotalCycles[IndexType::Max][PlatormIndexType::PlatormIndexTypeMax];
+	static CORE_API uint64 GTotalCounts[IndexType::Max][PlatormIndexType::PlatormIndexTypeMax];
 
 	int32 Index;
 	int32 PlatformTypeIndex;
@@ -52,7 +52,7 @@ struct CORE_API FScopedVirtualMallocTimer
 		FPlatformAtomics::InterlockedAdd((volatile int64*)&GTotalCycles[Index][PlatformTypeIndex], Add);
 	}
 
-	static void UpdateStats();
+	static CORE_API void UpdateStats();
 #else	//UE_TIME_VIRTUALMALLOC
 	FORCEINLINE FScopedVirtualMallocTimer(int32 InIndex = 0, int32 InPlatformTypeIndex = 0)
 	{
diff --git a/Engine/Source/Runtime/Core/Public/HAL/MemoryBase.h b/Engine/Source/Runtime/Core/Public/HAL/MemoryBase.h
--- a/Engine/Source/Runtime/Core/Public/HAL/MemoryBase.h
+++ b/Engine/Source/Runtime/Core/Public/HAL/MemoryBase.h
@@ -41,7 +41,7 @@ struct FGenericMemoryStats;
  * alloced by the system malloc routines, bypassing GMalloc. This is e.g. used by FMalloc
  * itself.
  */
-class CORE_API FUseSystemMallocForNew
+class FUseSystemMallocForNew
 {
 public:
 	/**
@@ -50,14 +50,14 @@ public:
 	 * @param	Size	Amount of memory to allocate (in bytes)
 	 * @return			A pointer to a block of memory with size Size or nullptr
 	 */
-	void* operator new(size_t Size);
+	CORE_API void* operator new(size_t Size);
 
 	/**
 	 * Overloaded delete operator using the system allocator
 	 *
 	 * @param	Ptr		Pointer to delete
 	 */
-	void operator delete(void* Ptr);
+	CORE_API void operator delete(void* Ptr);
 
 	/**
 	 * Overloaded array new operator using the system allocator.
@@ -76,7 +76,7 @@ public:
 };
 
 /** The global memory allocator's interface. */
-class CORE_API FMalloc  : 
+class FMalloc  : 
 	public FUseSystemMallocForNew,
 	public FExec
 {
@@ -90,7 +90,7 @@ public:
 	 * TryMalloc - like Malloc(), but may return a nullptr result if the allocation
 	 *             request cannot be satisfied.
 	 */
-	virtual void* TryMalloc( SIZE_T Count, uint32 Alignment=DEFAULT_ALIGNMENT );
+	CORE_API virtual void* TryMalloc( SIZE_T Count, uint32 Alignment=DEFAULT_ALIGNMENT );
 
 	/** 
 	 * Realloc
@@ -102,7 +102,7 @@ public:
 	 *              request cannot be satisfied. Note that in this case the memory
 	 *              pointed to by Original will still be valid
 	 */
-	virtual void* TryRealloc(void* Original, SIZE_T Count, uint32 Alignment=DEFAULT_ALIGNMENT);
+	CORE_API virtual void* TryRealloc(void* Original, SIZE_T Count, uint32 Alignment=DEFAULT_ALIGNMENT);
 
 	/** 
 	 * Free
@@ -157,7 +157,7 @@ public:
 	*	Initializes stats metadata. We need to do this as soon as possible, but cannot be done in the constructor
 	*	due to the FName::StaticInit
 	*/
-	virtual void InitializeStatsMetadata();
+	CORE_API virtual void InitializeStatsMetadata();
 
 #if UE_ALLOW_EXEC_COMMANDS
 	/**
@@ -170,10 +170,10 @@ public:
 #endif // UE_ALLOW_EXEC_COMMANDS
 
 	/** Called once per frame, gathers and sets all memory allocator statistics into the corresponding stats. MUST BE THREAD SAFE. */
-	virtual void UpdateStats();
+	CORE_API virtual void UpdateStats();
 
 	/** Writes allocator stats from the last update into the specified destination. */
-	virtual void GetAllocatorStats( FGenericMemoryStats& out_Stats );
+	CORE_API virtual void GetAllocatorStats( FGenericMemoryStats& out_Stats );
 
 	/** Dumps current allocator stats to the log. */
 	virtual void DumpAllocatorStats( class FOutputDevice& Ar )
@@ -229,6 +229,6 @@ protected:
 #if !UE_BUILD_SHIPPING
 public:
 	/** Limits the maximum single allocation, to this many bytes, for debugging */
-	static TAtomic<uint64> MaxSingleAlloc;
+	static CORE_API TAtomic<uint64> MaxSingleAlloc;
 #endif
 };
diff --git a/Engine/Source/Runtime/Core/Public/HAL/MemoryMisc.h b/Engine/Source/Runtime/Core/Public/HAL/MemoryMisc.h
--- a/Engine/Source/Runtime/Core/Public/HAL/MemoryMisc.h
+++ b/Engine/Source/Runtime/Core/Public/HAL/MemoryMisc.h
@@ -25,12 +25,12 @@ struct FGenericMemoryStats
 // reporting +/- difference in memory.
 // WARNING This will also capture differences in Threads which have nothing to do with the scope
 #if ENABLE_MEMORY_SCOPE_STATS
-class CORE_API FScopedMemoryStats
+class FScopedMemoryStats
 {
 public:
-	explicit FScopedMemoryStats(const TCHAR* Name);
+	CORE_API explicit FScopedMemoryStats(const TCHAR* Name);
 
-	~FScopedMemoryStats();
+	CORE_API ~FScopedMemoryStats();
 
 private:
 	const TCHAR* Text;
@@ -55,24 +55,24 @@ public:
 
 #if ENABLE_SHARED_MEMORY_TRACKER && PLATFORM_UNIX
 
-class CORE_API FSharedMemoryTracker
+class FSharedMemoryTracker
 {
 public:
 
 	/** Print the difference in memory pool sizes since the last call to this function exclusively */
-	static void PrintMemoryDiff(const TCHAR* Context);
+	static CORE_API void PrintMemoryDiff(const TCHAR* Context);
 
 
 	/** Store the memory pool size at construction and log the difference in memory that occurred during the lifetime of the tracker */
-	explicit FSharedMemoryTracker(const FString& InContext);
-	~FSharedMemoryTracker();
+	CORE_API explicit FSharedMemoryTracker(const FString& InContext);
+	CORE_API ~FSharedMemoryTracker();
 
 private:
 	const FString PrintContext;
 	const FExtendedPlatformMemoryStats StartStats;
 };
 #else
-class CORE_API FSharedMemoryTracker
+class FSharedMemoryTracker
 {
 public:
 
@@ -80,4 +80,4 @@ public:
 
 	explicit FSharedMemoryTracker(const FString& /*InContext*/) {}
 };
-#endif // ENABLE_SHARED_MEMORY_TRACKER  && PLATFORM_UNIX
\ No newline at end of file
+#endif // ENABLE_SHARED_MEMORY_TRACKER  && PLATFORM_UNIX
diff --git a/Engine/Source/Runtime/Core/Public/HAL/PlatformFileManager.h b/Engine/Source/Runtime/Core/Public/HAL/PlatformFileManager.h
--- a/Engine/Source/Runtime/Core/Public/HAL/PlatformFileManager.h
+++ b/Engine/Source/Runtime/Core/Public/HAL/PlatformFileManager.h
@@ -13,7 +13,7 @@
 /**
 * Platform File chain manager.
 **/
-class CORE_API FPlatformFileManager
+class FPlatformFileManager
 {
 #if USE_ATOMIC_PLATFORM_FILE
 	/** Currently used platform file. */
@@ -26,21 +26,21 @@ class CORE_API FPlatformFileManager
 public:
 
 	/** Constructor. */
-	FPlatformFileManager( );
+	CORE_API FPlatformFileManager( );
 
 	/**
 	 * Gets the currently used platform file.
 	 *
 	 * @return Reference to the currently used platform file.
 	 */
-	IPlatformFile& GetPlatformFile( );
+	CORE_API IPlatformFile& GetPlatformFile( );
 
 	/**
 	 * Sets the current platform file.
 	 *
 	 * @param NewTopmostPlatformFile Platform file to be used.
 	 */
-	void SetPlatformFile( IPlatformFile& NewTopmostPlatformFile );
+	CORE_API void SetPlatformFile( IPlatformFile& NewTopmostPlatformFile );
 
 	/**
 	 * Finds a platform file in the chain of active platform files.
@@ -48,7 +48,7 @@ public:
 	 * @param Name of the platform file.
 	 * @return Pointer to the active platform file or nullptr if the platform file was not found.
 	 */
-	IPlatformFile* FindPlatformFile( const TCHAR* Name );
+	CORE_API IPlatformFile* FindPlatformFile( const TCHAR* Name );
 
 	/**
 	 * Creates a new platform file instance.
@@ -56,22 +56,22 @@ public:
 	 * @param Name of the platform file to create.
 	 * @return Platform file instance of the platform file type was found, nullptr otherwise.
 	 */
-	IPlatformFile* GetPlatformFile( const TCHAR* Name );
+	CORE_API IPlatformFile* GetPlatformFile( const TCHAR* Name );
 
 	/**
 	 * calls Tick on the platform files in the TopmostPlatformFile chain
 	 */
-	void TickActivePlatformFile();
+	CORE_API void TickActivePlatformFile();
 
 	/**
 	* Performs additional initialization when the new async IO is enabled.
 	*/
-	void InitializeNewAsyncIO();
+	CORE_API void InitializeNewAsyncIO();
 
 	/**
 	 * Gets FPlatformFileManager Singleton.
 	 */
-	static FPlatformFileManager& Get( );
+	static CORE_API FPlatformFileManager& Get( );
 
 	/**
 	* Removes the specified file wrapper from the platform file wrapper chain.
@@ -81,7 +81,7 @@ public:
 	*
 	* @param The platform file to remove.
 	*/
-	void RemovePlatformFile(IPlatformFile* PlatformFileToRemove);
+	CORE_API void RemovePlatformFile(IPlatformFile* PlatformFileToRemove);
 
 	/**
 	* Inserts a new platform file into the platform file wrapper chain.
@@ -93,5 +93,5 @@ public:
 	* @param The platform file to insert.
 	* @return true if the platform file was inserted.
 	*/
-	bool InsertPlatformFile(IPlatformFile* NewPlatformFile);
+	CORE_API bool InsertPlatformFile(IPlatformFile* NewPlatformFile);
 };
diff --git a/Engine/Source/Runtime/Core/Public/HAL/PlatformMisc.h b/Engine/Source/Runtime/Core/Public/HAL/PlatformMisc.h
--- a/Engine/Source/Runtime/Core/Public/HAL/PlatformMisc.h
+++ b/Engine/Source/Runtime/Core/Public/HAL/PlatformMisc.h
@@ -69,7 +69,7 @@ private:
 	bool bStarted;
 };
 
-class CORE_API FScopedNamedEvent
+class FScopedNamedEvent
 {
 public:
 
@@ -89,7 +89,7 @@ public:
 	}	
 };
 
-class CORE_API FScopedProfilerColor
+class FScopedProfilerColor
 {
 public:
 
@@ -111,7 +111,7 @@ public:
 // BeginNamedEventStatic works the same as BeginNamedEvent, but should only be passed a compile-time string literal.
 // Some platform profilers can optimize the case where strings for certain events are constant.
 //
-class CORE_API FScopedNamedEventStatic
+class FScopedNamedEventStatic
 {
 public:
 
@@ -178,7 +178,7 @@ public:
 
 #else
 
-class CORE_API FScopedNamedEvent
+class FScopedNamedEvent
 {
 public:
 	UE_DEPRECATED(4.19, "FScopedNamedEvent is compiled out in shipping builds, use SCOPED_NAMED_EVENT or variant instead to compile correctly for all targets.")
@@ -192,7 +192,7 @@ public:
 	}
 };
 
-class CORE_API FScopedNamedEventStatic
+class FScopedNamedEventStatic
 {
 public:
 	UE_DEPRECATED(4.19, "FScopedNamedEventStatic is compiled out in shipping builds, use SCOPED_NAMED_EVENT or variant instead to compile correctly for all targets.")
@@ -217,7 +217,7 @@ public:
 #endif
 
 // For timing OnEnterBackground tasks. This can be time sensitive on some platforms
-class CORE_API FScopedEnterBackgroundEvent
+class FScopedEnterBackgroundEvent
 {
 public:
 
diff --git a/Engine/Source/Runtime/Core/Public/HAL/Runnable.h b/Engine/Source/Runtime/Core/Public/HAL/Runnable.h
--- a/Engine/Source/Runtime/Core/Public/HAL/Runnable.h
+++ b/Engine/Source/Runtime/Core/Public/HAL/Runnable.h
@@ -16,7 +16,7 @@
  * Run() is called where the real threaded work is done. Upon completion, Exit() is called to allow
  * correct clean up.
  */
-class CORE_API FRunnable
+class FRunnable
 {
 public:
 
diff --git a/Engine/Source/Runtime/Core/Public/HAL/RunnableThread.h b/Engine/Source/Runtime/Core/Public/HAL/RunnableThread.h
--- a/Engine/Source/Runtime/Core/Public/HAL/RunnableThread.h
+++ b/Engine/Source/Runtime/Core/Public/HAL/RunnableThread.h
@@ -16,7 +16,7 @@ class FTlsAutoCleanup;
  *
  * This interface specifies the methods used to manage a thread's life cycle.
  */
-class CORE_API FRunnableThread
+class FRunnableThread
 {
 	friend class FThreadSingletonInitializer;
 	friend class FTlsAutoCleanup;
@@ -25,12 +25,12 @@ class CORE_API FRunnableThread
 	friend class FForkProcessHelper;
 
 	/** Index of TLS slot for FRunnableThread pointer. */
-	static uint32 RunnableTlsSlot;
+	static CORE_API uint32 RunnableTlsSlot;
 
 public:
 
 	/** Gets a new Tls slot for storing the runnable thread pointer. */
-	static uint32 GetTlsSlot();
+	static CORE_API uint32 GetTlsSlot();
 
 	/**
 	 * Factory method to create a thread with the specified stack size and thread priority.
@@ -41,7 +41,7 @@ public:
 	 * @param InThreadPri Tells the thread whether it needs to adjust its priority or not. Defaults to normal priority
 	 * @return The newly created thread or nullptr if it failed
 	 */
-	static FRunnableThread* Create(
+	static CORE_API FRunnableThread* Create(
 		class FRunnable* InRunnable,
 		const TCHAR* ThreadName,
 		uint32 InStackSize = 0,
@@ -134,10 +134,10 @@ public:
 	}
 
 	/** Default constructor. */
-	FRunnableThread();
+	CORE_API FRunnableThread();
 
 	/** Virtual destructor */
-	virtual ~FRunnableThread();
+	CORE_API virtual ~FRunnableThread();
 
 	/**
 	 * @return a runnable thread that is executing this runnable, if return value is nullptr, it means the running thread can be game thread or a thread created outside the runnable interface
@@ -165,10 +165,10 @@ protected:
 		EThreadCreateFlags InCreateFlags = EThreadCreateFlags::None) = 0;
 
 	/** Stores this instance in the runnable thread TLS slot. */
-	void SetTls();
+	CORE_API void SetTls();
 
 	/** Deletes all FTlsAutoCleanup objects created for this thread. */
-	void FreeTls();
+	CORE_API void FreeTls();
 
 	/** Holds the name of the thread. */
 	FString ThreadName;
@@ -191,7 +191,7 @@ protected:
 private:
 
 	/** Called to setup a newly created RunnableThread */
-	static void SetupCreatedThread(FRunnableThread*& NewThread, class FRunnable* InRunnable, const TCHAR* ThreadName, uint32 InStackSize, EThreadPriority InThreadPri, uint64 InThreadAffinityMask, EThreadCreateFlags InCreateFlags);
+	static CORE_API void SetupCreatedThread(FRunnableThread*& NewThread, class FRunnable* InRunnable, const TCHAR* ThreadName, uint32 InStackSize, EThreadPriority InThreadPri, uint64 InThreadAffinityMask, EThreadCreateFlags InCreateFlags);
 
 	/** Used by the thread manager to tick threads in single-threaded mode */
 	virtual void Tick() {}
@@ -207,5 +207,5 @@ private:
 	/**
 	 * Called after the internal thread is created so it can register debug information
 	 */
-	void PostCreate(EThreadPriority ThreadPriority);
+	CORE_API void PostCreate(EThreadPriority ThreadPriority);
 };
diff --git a/Engine/Source/Runtime/Core/Public/HAL/Thread.h b/Engine/Source/Runtime/Core/Public/HAL/Thread.h
--- a/Engine/Source/Runtime/Core/Public/HAL/Thread.h
+++ b/Engine/Source/Runtime/Core/Public/HAL/Thread.h
@@ -20,7 +20,7 @@
  *		Thread.Join();
  * For more verbose example check `TestTypicalUseCase` in `ThreadTest.cpp`
  */
-class CORE_API FThread final
+class FThread final
 {
 public:
 	// indicates if the thread should be forked in case the owning process is forked
@@ -42,7 +42,7 @@ public:
 	* @param ThreadAffinity Tells the thread whether it needs to adjust its affinity or not. Defaults to no affinity
 	* @param IsForkable Tells the thread whether it can be forked. Defaults to NonForkable
 	*/
-	FThread(
+	CORE_API FThread(
 		TCHAR const* ThreadName,
 		TUniqueFunction<void()>&& ThreadFunction,
 		uint32 StackSize = 0,
@@ -53,7 +53,7 @@ public:
 
 	// with SingleThreadTickFunction that will be executed every frame if running with `-nothreading
 	// (FPlatformProcess::SupportsMultithreading() == false)
-	FThread(
+	CORE_API FThread(
 		TCHAR const* ThreadName,
 		TUniqueFunction<void()>&& ThreadFunction,
 		TUniqueFunction<void()>&& SingleThreadTickFunction,
@@ -72,12 +72,12 @@ public:
 	 * Move assignment operator.
 	 * Asserts if the instance is joinable.
 	 */
-	FThread& operator=(FThread&& Other);
+	CORE_API FThread& operator=(FThread&& Other);
 
 	/**
 	 * Destructor asserts if the instance is not joined or detached.
 	 */
-	~FThread();
+	CORE_API ~FThread();
 
 	/**
 	 * Checks if the thread object identifies an active thread of execution. 
@@ -85,7 +85,7 @@ public:
 	 * thread of execution and is therefore joinable.
 	 * @see Join
 	 */
-	bool IsJoinable() const;
+	CORE_API bool IsJoinable() const;
 
 	/**
 	 * Blocks the current thread until the thread identified by `this` finishes its execution.
@@ -94,14 +94,14 @@ public:
 	 * from multiple threads constitutes a data race that results in undefined behavior.
 	 * @see IsJoinable
 	 */
-	void Join();
+	CORE_API void Join();
 
 	static constexpr uint32 InvalidThreadId = ~uint32(0);
 
 	/**
 	 * @return Thread ID for this thread
 	 */
-	uint32 GetThreadId() const;
+	CORE_API uint32 GetThreadId() const;
 
 #if 0 // disabled as it doesn't work as intended
 
@@ -110,7 +110,7 @@ public:
 	 * Any allocated resources will be freed once the thread exits.
 	 * After calling detach `this` no longer owns any thread.
 	 */
-	void Detach();
+	CORE_API void Detach();
 
 #endif
 
