[Core]

diff --git a/Engine/Source/Runtime/Core/Public/HAL/ThreadHeartBeat.h b/Engine/Source/Runtime/Core/Public/HAL/ThreadHeartBeat.h
--- a/Engine/Source/Runtime/Core/Public/HAL/ThreadHeartBeat.h
+++ b/Engine/Source/Runtime/Core/Public/HAL/ThreadHeartBeat.h
@@ -32,26 +32,26 @@ DECLARE_DELEGATE_OneParam(FOnHangDelegate, uint32);
  * When the title is suspended, this thread is also suspended, and the local clock stops.
  * The delta is clamped so if we are resumed, the clock continues from where it left off.
  */
-class CORE_API FThreadHeartBeatClock
+class FThreadHeartBeatClock
 {
 	uint64 CurrentCycles;
 	uint64 LastRealTickCycles;
 	const uint64 MaxTimeStepCycles;
 
 public:
-	FThreadHeartBeatClock(double InMaxTimeStep);
+	CORE_API FThreadHeartBeatClock(double InMaxTimeStep);
 
-	void Tick();
-	double Seconds();
+	CORE_API void Tick();
+	CORE_API double Seconds();
 };
 
 /**
  * Thread heartbeat check class.
  * Used by crash handling code to check for hangs.
  */
-class CORE_API FThreadHeartBeat : public FRunnable
+class FThreadHeartBeat : public FRunnable
 {
-	static FThreadHeartBeat* Singleton;
+	static CORE_API FThreadHeartBeat* Singleton;
 
 	/** Holds per-thread info about the heartbeat */
 	struct FHeartBeatInfo
@@ -150,15 +150,15 @@ class CORE_API FThreadHeartBeat : public FRunnable
 	FOnHangDelegate OnHangDelegate;
 #endif
 
-	FThreadHeartBeat();
-	virtual ~FThreadHeartBeat();
+	CORE_API FThreadHeartBeat();
+	CORE_API virtual ~FThreadHeartBeat();
 
-	void InitSettings();
+	CORE_API void InitSettings();
 
-	void FORCENOINLINE OnHang(double HangDuration, uint32 ThreadThatHung);
-	void FORCENOINLINE OnPresentHang(double HangDuration);
+	CORE_API void FORCENOINLINE OnHang(double HangDuration, uint32 ThreadThatHung);
+	CORE_API void FORCENOINLINE OnPresentHang(double HangDuration);
 
-	bool IsEnabled();
+	CORE_API bool IsEnabled();
 
 public:
 
@@ -172,58 +172,58 @@ public:
 	};
 
 	/** Gets the heartbeat singleton */
-	static FThreadHeartBeat& Get();
-	static FThreadHeartBeat* GetNoInit();
+	static CORE_API FThreadHeartBeat& Get();
+	static CORE_API FThreadHeartBeat* GetNoInit();
 
 	/** Begin measuring heartbeat */
-	void Start();
+	CORE_API void Start();
 	/** Called from a thread once per frame to update the heartbeat time */
-	void HeartBeat(bool bReadConfig = false);
+	CORE_API void HeartBeat(bool bReadConfig = false);
 	/** Called from the rendering or RHI thread when the platform RHI presents a frame (supported platforms only). */
-	void PresentFrame();
+	CORE_API void PresentFrame();
 	/** Called by a supervising thread to check the threads' health */
-	uint32 CheckHeartBeat(double& OutHangDuration);
+	CORE_API uint32 CheckHeartBeat(double& OutHangDuration);
 	/** Called by a thread when it's no longer expecting to be ticked */
-	void KillHeartBeat();
+	CORE_API void KillHeartBeat();
 
 	/** Called from a thread once on entry to a function to be monitored */
-	void MonitorFunctionStart();
+	CORE_API void MonitorFunctionStart();
 	/** Called by a thread when a function has completed and no longer needs to be monitored */
-	void MonitorFunctionEnd();
+	CORE_API void MonitorFunctionEnd();
 	/** Called by a supervising thread to check all function calls' being monitored health */
-	uint32 CheckFunctionHeartBeat(double& OutHangDuration);
+	CORE_API uint32 CheckFunctionHeartBeat(double& OutHangDuration);
 
 	/* 
 		Called from a thread to register a checkpoint to be monitored 
 		@param EndCheckPoint name of the checkpoint that needs to be reached. TimeToReachCheckPoint the time duration we have to reach the specified checkpoint.
 	*/
-	void MonitorCheckpointStart(FName EndCheckPoint, double TimeToReachCheckpoint);
+	CORE_API void MonitorCheckpointStart(FName EndCheckPoint, double TimeToReachCheckpoint);
 	/* Called from a thread when a checkpoint has ended */
-	void MonitorCheckpointEnd(FName CheckPoint);
+	CORE_API void MonitorCheckpointEnd(FName CheckPoint);
 	/* Called by a supervising thread to check all checkpoints forward progress */
-	uint32 CheckCheckpointHeartBeat(double& OutHangDuration);
+	CORE_API uint32 CheckCheckpointHeartBeat(double& OutHangDuration);
 
 	/** 
 	 * Suspend heartbeat measuring for the current thread if the thread has already had a heartbeat 
 	 * @param bAllThreads If true, suspends heartbeat for all threads, not only the current one
 	 */
-	void SuspendHeartBeat(bool bAllThreads = false);
+	CORE_API void SuspendHeartBeat(bool bAllThreads = false);
 	/** 
 	 * Resume heartbeat measuring for the current thread 
 	 * @param bAllThreads If true, resumes heartbeat for all threads, not only the current one
 	 */
-	void ResumeHeartBeat(bool bAllThreads = false);
+	CORE_API void ResumeHeartBeat(bool bAllThreads = false);
 
 	/**
 	* Returns true/false if this thread is currently performing heartbeat monitoring
 	*/
-	bool IsBeating();
+	CORE_API bool IsBeating();
 
 	/** 
 	 * Sets a multiplier to the hang duration (>= 1.0).
 	 * Can be used to extend the duration during loading screens etc.
 	 */
-	void SetDurationMultiplier(double NewMultiplier);
+	CORE_API void SetDurationMultiplier(double NewMultiplier);
 
 	/*
 	* Get the Id of the last thread to trigger the hang detector.
@@ -258,9 +258,9 @@ public:
 	double GetHangDuration() const { return ConfigHangDuration; };
 
 	//~ Begin FRunnable Interface.
-	virtual bool Init();
-	virtual uint32 Run();
-	virtual void Stop();
+	CORE_API virtual bool Init();
+	CORE_API virtual uint32 Run();
+	CORE_API virtual void Stop();
 	//~ End FRunnable Interface
 };
 
@@ -331,9 +331,9 @@ public:
 #define WALK_STACK_ON_HITCH_DETECTED 0
 #endif
 
-class CORE_API FGameThreadHitchHeartBeatThreaded : public FRunnable
+class FGameThreadHitchHeartBeatThreaded : public FRunnable
 {
-	static FGameThreadHitchHeartBeatThreaded* Singleton;
+	static CORE_API FGameThreadHitchHeartBeatThreaded* Singleton;
 
 	/** Thread to run the worker FRunnable on */
 	FRunnableThread* Thread;
@@ -362,10 +362,10 @@ class CORE_API FGameThreadHitchHeartBeatThreaded : public FRunnable
 
 	FThreadHeartBeatClock Clock;
 
-	void InitSettings();
+	CORE_API void InitSettings();
 
-	FGameThreadHitchHeartBeatThreaded();
-	virtual ~FGameThreadHitchHeartBeatThreaded();
+	CORE_API FGameThreadHitchHeartBeatThreaded();
+	CORE_API virtual ~FGameThreadHitchHeartBeatThreaded();
 
 public:
 
@@ -376,35 +376,35 @@ public:
 	};
 
 	/** Gets the heartbeat singleton */
-	static FGameThreadHitchHeartBeatThreaded& Get();
-	static FGameThreadHitchHeartBeatThreaded* GetNoInit();
+	static CORE_API FGameThreadHitchHeartBeatThreaded& Get();
+	static CORE_API FGameThreadHitchHeartBeatThreaded* GetNoInit();
 
 	/**
 	* Called at the start of a frame to register the time we are looking to detect a hitch
 	*/
-	void FrameStart(bool bSkipThisFrame = false);
+	CORE_API void FrameStart(bool bSkipThisFrame = false);
 
-	double GetFrameStartTime();
-	double GetCurrentTime();
+	CORE_API double GetFrameStartTime();
+	CORE_API double GetCurrentTime();
 
 	/**
 	* Suspend heartbeat hitch detection. Must call ResumeHeartBeat later to resume.
 	*/
-	void SuspendHeartBeat();
+	CORE_API void SuspendHeartBeat();
 
 	/**
 	* Resume heartbeat hitch detection. Call only after first calling SuspendHeartBeat.
 	*/
-	void ResumeHeartBeat();
+	CORE_API void ResumeHeartBeat();
 
 	// No-op, used in FUnixSignalGameHitchHeartBeat
 	void Restart() {}
 	void PostFork() {}
 
 	//~ Begin FRunnable Interface.
-	virtual bool Init();
-	virtual uint32 Run();
-	virtual void Stop();
+	CORE_API virtual bool Init();
+	CORE_API virtual uint32 Run();
+	CORE_API virtual void Stop();
 	//~ End FRunnable Interface
 };
 
diff --git a/Engine/Source/Runtime/Core/Public/HAL/ThreadManager.h b/Engine/Source/Runtime/Core/Public/HAL/ThreadManager.h
--- a/Engine/Source/Runtime/Core/Public/HAL/ThreadManager.h
+++ b/Engine/Source/Runtime/Core/Public/HAL/ThreadManager.h
@@ -16,7 +16,7 @@ class FRunnableThread;
 /**
  * Manages runnables and runnable threads.
  */
-class CORE_API FThreadManager
+class FThreadManager
 {
 	/** Critical section for ThreadList */
 	FCriticalSection ThreadsCritical;
@@ -33,7 +33,7 @@ public:
 	* @param Thread thread object.
 	* @see RemoveThread
 	*/
-	void AddThread(uint32 ThreadId, FRunnableThread* Thread);
+	CORE_API void AddThread(uint32 ThreadId, FRunnableThread* Thread);
 
 	/**
 	* Used internally to remove thread object.
@@ -41,13 +41,13 @@ public:
 	* @param Thread thread object to be removed.
 	* @see AddThread
 	*/
-	void RemoveThread(FRunnableThread* Thread);
+	CORE_API void RemoveThread(FRunnableThread* Thread);
 
 	/** Get the number of registered threads */
 	int32 NumThreads() const { return Threads.Num(); }
 
 	/** Ticks all fake threads and their runnable objects. */
-	void Tick();
+	CORE_API void Tick();
 
 	/** Returns the name of a thread given its TLS id */
 	inline static const FString& GetThreadName(uint32 ThreadId)
@@ -75,29 +75,29 @@ public:
 		TArray<uint64, TInlineAllocator<100>> ProgramCounters;
 	};
 
-	void GetAllThreadStackBackTraces(TArray<FThreadStackBackTrace>& StackTraces);
+	CORE_API void GetAllThreadStackBackTraces(TArray<FThreadStackBackTrace>& StackTraces);
 #endif
 
 	/**
 	 * Enumerate each thread.
 	 *
 	 */
-	void ForEachThread(TFunction<void(uint32 ThreadId, FRunnableThread* Thread)> Func);
+	CORE_API void ForEachThread(TFunction<void(uint32 ThreadId, FRunnableThread* Thread)> Func);
 
 	/**
 	 * Access to the singleton object.
 	 *
 	 * @return Thread manager object.
 	 */
-	static FThreadManager& Get();
+	static CORE_API FThreadManager& Get();
 
 private:
 
 	friend class FForkProcessHelper;
 
 	/** Returns a list of registered forkable threads  */
-	TArray<FRunnableThread*> GetForkableThreads();
+	CORE_API TArray<FRunnableThread*> GetForkableThreads();
 
 	/** Returns internal name of a the thread given its TLS id */
-	const FString& GetThreadNameInternal(uint32 ThreadId);
+	CORE_API const FString& GetThreadNameInternal(uint32 ThreadId);
 };
diff --git a/Engine/Source/Runtime/Core/Public/HAL/TlsAutoCleanup.h b/Engine/Source/Runtime/Core/Public/HAL/TlsAutoCleanup.h
--- a/Engine/Source/Runtime/Core/Public/HAL/TlsAutoCleanup.h
+++ b/Engine/Source/Runtime/Core/Public/HAL/TlsAutoCleanup.h
@@ -8,7 +8,7 @@
 Base class for objects in TLS that support auto-cleanup. 
 Polymorphically deletes registered instances on thread exit.
 */
-class CORE_API FTlsAutoCleanup
+class FTlsAutoCleanup
 {
 public:
 	/** Virtual destructor. */
@@ -16,7 +16,7 @@ public:
 	{}
 
 	/** Register this instance to be auto-cleanup. */
-	void Register();
+	CORE_API void Register();
 };
 
 /** Wrapper for values to be stored in TLS that support auto-cleanup. */
diff --git a/Engine/Source/Runtime/Core/Public/HAL/UnrealMemory.h b/Engine/Source/Runtime/Core/Public/HAL/UnrealMemory.h
--- a/Engine/Source/Runtime/Core/Public/HAL/UnrealMemory.h
+++ b/Engine/Source/Runtime/Core/Public/HAL/UnrealMemory.h
@@ -37,11 +37,11 @@ FORCEINLINE void DoGamethreadHook(int32 Index)
 
 #if TIME_MALLOC
 
-struct CORE_API FScopedMallocTimer
+struct FScopedMallocTimer
 {
-	static uint64 GTotalCycles[4];
-	static uint64 GTotalCount[4];
-	static uint64 GTotalMisses[4];
+	static CORE_API uint64 GTotalCycles[4];
+	static CORE_API uint64 GTotalCount[4];
+	static CORE_API uint64 GTotalMisses[4];
 
 	int32 Index;
 	uint64 Cycles;
@@ -67,7 +67,7 @@ struct CORE_API FScopedMallocTimer
 	{
 		FPlatformAtomics::InterlockedAdd((volatile int64*)&GTotalMisses[InIndex], 1);
 	}
-	static void Spew();
+	static CORE_API void Spew();
 }; 
 #else
 struct FScopedMallocTimer
@@ -89,7 +89,7 @@ struct FScopedMallocTimer
 	FMemory.
 -----------------------------------------------------------------------------*/
 
-struct CORE_API FMemory
+struct FMemory
 {
 	/** Some allocators can be given hints to treat allocations differently depending on how the memory is used, it's lifetime etc. */
 	enum AllocationHints
@@ -207,10 +207,10 @@ struct CORE_API FMemory
 	// C style memory allocation stubs.
 	//
 
-	static void* Malloc(SIZE_T Count, uint32 Alignment = DEFAULT_ALIGNMENT);
-	static void* Realloc(void* Original, SIZE_T Count, uint32 Alignment = DEFAULT_ALIGNMENT);
-	static void Free(void* Original);
-	static SIZE_T GetAllocSize(void* Original);
+	static CORE_API void* Malloc(SIZE_T Count, uint32 Alignment = DEFAULT_ALIGNMENT);
+	static CORE_API void* Realloc(void* Original, SIZE_T Count, uint32 Alignment = DEFAULT_ALIGNMENT);
+	static CORE_API void Free(void* Original);
+	static CORE_API SIZE_T GetAllocSize(void* Original);
 
 	static FORCEINLINE_DEBUGGABLE void* MallocZeroed(SIZE_T Count, uint32 Alignment = DEFAULT_ALIGNMENT)
 	{
@@ -225,22 +225,22 @@ struct CORE_API FMemory
 	* and shrink containers to optimal sizes.
 	* This call is always fast and threadsafe with no locking.
 	*/
-	static SIZE_T QuantizeSize(SIZE_T Count, uint32 Alignment = DEFAULT_ALIGNMENT);
+	static CORE_API SIZE_T QuantizeSize(SIZE_T Count, uint32 Alignment = DEFAULT_ALIGNMENT);
 
 	/**
 	* Releases as much memory as possible. Must be called from the main thread.
 	*/
-	static void Trim(bool bTrimThreadCaches = true);
+	static CORE_API void Trim(bool bTrimThreadCaches = true);
 
 	/**
 	* Set up TLS caches on the current thread. These are the threads that we can trim.
 	*/
-	static void SetupTLSCachesOnCurrentThread();
+	static CORE_API void SetupTLSCachesOnCurrentThread();
 
 	/**
 	* Clears the TLS caches on the current thread and disables any future caching.
 	*/
-	static void ClearAndDisableTLSCachesOnCurrentThread();
+	static CORE_API void ClearAndDisableTLSCachesOnCurrentThread();
 
 	/**
 	 * A helper function that will perform a series of random heap allocations to test
@@ -249,42 +249,42 @@ struct CORE_API FMemory
 	 * where you call it in a good state, do something to corrupt memory, then call this again
 	 * and hopefully freeing some pointers will trigger a crash.
 	 */
-	static void TestMemory();
+	static CORE_API void TestMemory();
 	/**
 	* Called once main is started and we have -purgatorymallocproxy.
 	* This uses the purgatory malloc proxy to check if things are writing to stale pointers.
 	*/
-	static void EnablePurgatoryTests();
+	static CORE_API void EnablePurgatoryTests();
 	/**
 	* Called once main is started and we have -purgatorymallocproxy.
 	* This uses the purgatory malloc proxy to check if things are writing to stale pointers.
 	*/
-	static void EnablePoisonTests();
+	static CORE_API void EnablePoisonTests();
 	/**
 	* Set global allocator instead of creating it lazily on first allocation.
 	* Must only be called once and only if lazy init is disabled via a macro.
 	*/
-	static void ExplicitInit(FMalloc& Allocator);
+	static CORE_API void ExplicitInit(FMalloc& Allocator);
 
 	/**
 	* Functions to handle special memory given to the title from the platform
 	* This memory is allocated like a stack, it's never really freed
 	*/
-	static void RegisterPersistentAuxiliary(void* InMemory, SIZE_T InSize);
-	static void* MallocPersistentAuxiliary(SIZE_T InSize, uint32 InAlignment = 0);
-	static void FreePersistentAuxiliary(void* InPtr);
-	static bool IsPersistentAuxiliaryActive();
-	static void DisablePersistentAuxiliary();
-	static void EnablePersistentAuxiliary();
-	static SIZE_T GetUsedPersistentAuxiliary();
+	static CORE_API void RegisterPersistentAuxiliary(void* InMemory, SIZE_T InSize);
+	static CORE_API void* MallocPersistentAuxiliary(SIZE_T InSize, uint32 InAlignment = 0);
+	static CORE_API void FreePersistentAuxiliary(void* InPtr);
+	static CORE_API bool IsPersistentAuxiliaryActive();
+	static CORE_API void DisablePersistentAuxiliary();
+	static CORE_API void EnablePersistentAuxiliary();
+	static CORE_API SIZE_T GetUsedPersistentAuxiliary();
 private:
-	static void GCreateMalloc();
+	static CORE_API void GCreateMalloc();
 	// These versions are called either at startup or in the event of a crash
-	static void* MallocExternal(SIZE_T Count, uint32 Alignment = DEFAULT_ALIGNMENT);
-	static void* ReallocExternal(void* Original, SIZE_T Count, uint32 Alignment = DEFAULT_ALIGNMENT);
-	static void FreeExternal(void* Original);
-	static SIZE_T GetAllocSizeExternal(void* Original);
-	static SIZE_T QuantizeSizeExternal(SIZE_T Count, uint32 Alignment = DEFAULT_ALIGNMENT);
+	static CORE_API void* MallocExternal(SIZE_T Count, uint32 Alignment = DEFAULT_ALIGNMENT);
+	static CORE_API void* ReallocExternal(void* Original, SIZE_T Count, uint32 Alignment = DEFAULT_ALIGNMENT);
+	static CORE_API void FreeExternal(void* Original);
+	static CORE_API SIZE_T GetAllocSizeExternal(void* Original);
+	static CORE_API SIZE_T QuantizeSizeExternal(SIZE_T Count, uint32 Alignment = DEFAULT_ALIGNMENT);
 };
 
 #define INLINE_FMEMORY_OPERATION (0) // untested, but should work. Inlines FMemory::Malloc, etc
diff --git a/Engine/Source/Runtime/Core/Public/IO/IoChunkEncoding.h b/Engine/Source/Runtime/Core/Public/IO/IoChunkEncoding.h
--- a/Engine/Source/Runtime/Core/Public/IO/IoChunkEncoding.h
+++ b/Engine/Source/Runtime/Core/Public/IO/IoChunkEncoding.h
@@ -40,7 +40,7 @@ struct FIoChunkDecodingParams
  * The chunk encoding information is enocoded into a 16 byte
  * header followed by N block sizes.
  */
-class CORE_API FIoChunkEncoding
+class FIoChunkEncoding
 {
 public:
 	static constexpr uint32 ExpectedMagic = 0x2e696f; // .io
@@ -49,7 +49,7 @@ public:
 	static constexpr uint64 MaxSize = (uint64(1) << 40);
 
 	/** Header describing the encoded I/O chunk. */
-	struct CORE_API FHeader
+	struct FHeader
 	{
 		uint64 Magic: 24;
 		uint64 RawSize: 40;
@@ -58,24 +58,24 @@ public:
 		uint64 Flags: 8;
 		uint64 Pad: 8;
 
-		bool IsValid() const;
-		uint32 GetBlockSize() const;
-		uint32 GetBlockCount() const;
-		TConstArrayView<uint32> GetBlocks() const;
-		uint64 GetTotalHeaderSize() const;
+		CORE_API bool IsValid() const;
+		CORE_API uint32 GetBlockSize() const;
+		CORE_API uint32 GetBlockCount() const;
+		CORE_API TConstArrayView<uint32> GetBlocks() const;
+		CORE_API uint64 GetTotalHeaderSize() const;
 
-		static const FHeader* Decode(FMemoryView HeaderData);
+		static CORE_API const FHeader* Decode(FMemoryView HeaderData);
 	};
 
 	static_assert(sizeof(FHeader) == 16, "I/O chunk header size mismatch");
 
-	static bool Encode(const FIoChunkEncodingParams& Params, FMemoryView RawData, FIoBuffer& OutEncodedData);
-	static bool Encode(const FIoChunkEncodingParams& Params, FMemoryView RawData, FIoBuffer& OutHeader, FIoBuffer& OutEncodedBlocks);
-	static bool Decode(const FIoChunkDecodingParams& Params, FMemoryView EncodedBlocks, FMutableMemoryView OutRawData);
-	static bool Decode(FMemoryView EncodedData, FName CompressionFormat, FMemoryView EncryptionKey, FMutableMemoryView OutRawData, uint64 Offset = 0);
+	static CORE_API bool Encode(const FIoChunkEncodingParams& Params, FMemoryView RawData, FIoBuffer& OutEncodedData);
+	static CORE_API bool Encode(const FIoChunkEncodingParams& Params, FMemoryView RawData, FIoBuffer& OutHeader, FIoBuffer& OutEncodedBlocks);
+	static CORE_API bool Decode(const FIoChunkDecodingParams& Params, FMemoryView EncodedBlocks, FMutableMemoryView OutRawData);
+	static CORE_API bool Decode(FMemoryView EncodedData, FName CompressionFormat, FMemoryView EncryptionKey, FMutableMemoryView OutRawData, uint64 Offset = 0);
 
-	static TIoStatusOr<FIoOffsetAndLength> GetChunkRange(uint64 TotalRawSize, uint32 RawBlockSize, TConstArrayView<uint32> EncodedBlockSize, uint64 RawOffset, uint64 RawSize);
-	static TIoStatusOr<FIoOffsetAndLength> GetChunkRange(const FIoChunkDecodingParams& Params, uint64 RawSize);
-	static uint64 GetTotalEncodedSize(TConstArrayView<uint32> EncodedBlockSize);
-	static FIoStatus HashBlocks(const FHeader& Header, FMemoryView EncodedData, TArray<FIoHash>& OutHashes);
+	static CORE_API TIoStatusOr<FIoOffsetAndLength> GetChunkRange(uint64 TotalRawSize, uint32 RawBlockSize, TConstArrayView<uint32> EncodedBlockSize, uint64 RawOffset, uint64 RawSize);
+	static CORE_API TIoStatusOr<FIoOffsetAndLength> GetChunkRange(const FIoChunkDecodingParams& Params, uint64 RawSize);
+	static CORE_API uint64 GetTotalEncodedSize(TConstArrayView<uint32> EncodedBlockSize);
+	static CORE_API FIoStatus HashBlocks(const FHeader& Header, FMemoryView EncodedData, TArray<FIoHash>& OutHashes);
 };
diff --git a/Engine/Source/Runtime/Core/Public/IO/IoDispatcher.h b/Engine/Source/Runtime/Core/Public/IO/IoDispatcher.h
--- a/Engine/Source/Runtime/Core/Public/IO/IoDispatcher.h
+++ b/Engine/Source/Runtime/Core/Public/IO/IoDispatcher.h
@@ -81,8 +81,8 @@ public:
 
 	CORE_API void InitializeFileEnvironment(FStringView InPath, int32 InOrder = 0);
 
-	CORE_API const FString& GetPath() const { return Path; }
-	CORE_API int32 GetOrder() const { return Order; }
+	const FString& GetPath() const { return Path; }
+	int32 GetOrder() const { return Order; }
 
 private:
 	FString			Path;
@@ -633,7 +633,7 @@ struct FIoStoreCompressedReadResult
 class IIoStoreWriterReferenceChunkDatabase
 {
 public:
-	CORE_API virtual ~IIoStoreWriterReferenceChunkDatabase() = default;
+	virtual ~IIoStoreWriterReferenceChunkDatabase() = default;
 
 	/*
 	* Used by IIoStoreWriter to check and see if there's a reference chunk that matches the data that
@@ -653,19 +653,19 @@ public:
 	* Chunks provided *MUST* decompress to bits that hash to the exact value provided in InChunkKey (i.e. be exactly the same bits),
 	* and also be the same number of blocks (i.e. same CompressionBlockSize)
 	*/
-	CORE_API virtual bool RetrieveChunk(const TPair<FIoContainerId, FIoChunkHash>& InChunkKey, TUniqueFunction<void(TIoStatusOr<FIoStoreCompressedReadResult>)> InCompletionCallback) = 0;
+	virtual bool RetrieveChunk(const TPair<FIoContainerId, FIoChunkHash>& InChunkKey, TUniqueFunction<void(TIoStatusOr<FIoStoreCompressedReadResult>)> InCompletionCallback) = 0;
 
 	/* 
 	* Quick synchronous existence check that returns the number of blocks for the chunk. This is used to set up
 	* the necessary structures without needing to read the source data for the chunk.
 	*/
-	CORE_API virtual bool ChunkExists(const TPair<FIoContainerId, FIoChunkHash>& InChunkKey, uint32& OutNumChunkBlocks) = 0;
+	virtual bool ChunkExists(const TPair<FIoContainerId, FIoChunkHash>& InChunkKey, uint32& OutNumChunkBlocks) = 0;
 
 	/*
 	* Returns the compression block size that was used to break up the IoChunks in the source containers. If this is different than what we want, 
 	* then none of the chunks will ever match. Knowing this up front allows us to only match on hash
 	*/
-	CORE_API virtual uint32 GetCompressionBlockSize() const = 0;
+	virtual uint32 GetCompressionBlockSize() const = 0;
 };
 
 /**
@@ -677,27 +677,27 @@ public:
 class IIoStoreWriterHashDatabase
 {
 public:
-	CORE_API virtual ~IIoStoreWriterHashDatabase() = default;
-	CORE_API virtual bool FindHashForChunkId(const FIoChunkId& ChunkId, FIoChunkHash& OutHash) const = 0;
+	virtual ~IIoStoreWriterHashDatabase() = default;
+	virtual bool FindHashForChunkId(const FIoChunkId& ChunkId, FIoChunkHash& OutHash) const = 0;
 };
 
 
 class IIoStoreWriter
 {
 public:
-	CORE_API virtual ~IIoStoreWriter() = default;
+	virtual ~IIoStoreWriter() = default;
 
 	/**
 	*	If a reference database is provided, the IoStoreWriter implementation may elect to reuse compressed blocks
 	*	from previous containers instead of recompressing input data. This must be set before any writes are appended.
 	*/
-	CORE_API virtual void SetReferenceChunkDatabase(TSharedPtr<IIoStoreWriterReferenceChunkDatabase> ReferenceChunkDatabase) = 0;
-	CORE_API virtual void SetHashDatabase(TSharedPtr<IIoStoreWriterHashDatabase> HashDatabase, bool bVerifyHashDatabase) = 0;
-	CORE_API virtual void EnableDiskLayoutOrdering(const TArray<TUniquePtr<FIoStoreReader>>& PatchSourceReaders = TArray<TUniquePtr<FIoStoreReader>>()) = 0;
-	CORE_API virtual void Append(const FIoChunkId& ChunkId, FIoBuffer Chunk, const FIoWriteOptions& WriteOptions, uint64 OrderHint = MAX_uint64) = 0;
-	CORE_API virtual void Append(const FIoChunkId& ChunkId, IIoStoreWriteRequest* Request, const FIoWriteOptions& WriteOptions) = 0;
-	CORE_API virtual TIoStatusOr<FIoStoreWriterResult> GetResult() = 0;
-	CORE_API virtual void EnumerateChunks(TFunction<bool(FIoStoreTocChunkInfo&&)>&& Callback) const = 0;
+	virtual void SetReferenceChunkDatabase(TSharedPtr<IIoStoreWriterReferenceChunkDatabase> ReferenceChunkDatabase) = 0;
+	virtual void SetHashDatabase(TSharedPtr<IIoStoreWriterHashDatabase> HashDatabase, bool bVerifyHashDatabase) = 0;
+	virtual void EnableDiskLayoutOrdering(const TArray<TUniquePtr<FIoStoreReader>>& PatchSourceReaders = TArray<TUniquePtr<FIoStoreReader>>()) = 0;
+	virtual void Append(const FIoChunkId& ChunkId, FIoBuffer Chunk, const FIoWriteOptions& WriteOptions, uint64 OrderHint = MAX_uint64) = 0;
+	virtual void Append(const FIoChunkId& ChunkId, IIoStoreWriteRequest* Request, const FIoWriteOptions& WriteOptions) = 0;
+	virtual TIoStatusOr<FIoStoreWriterResult> GetResult() = 0;
+	virtual void EnumerateChunks(TFunction<bool(FIoStoreTocChunkInfo&&)>&& Callback) const = 0;
 };
 
 class FIoStoreReader
diff --git a/Engine/Source/Runtime/Core/Public/Internationalization/BreakIterator.h b/Engine/Source/Runtime/Core/Public/Internationalization/BreakIterator.h
--- a/Engine/Source/Runtime/Core/Public/Internationalization/BreakIterator.h
+++ b/Engine/Source/Runtime/Core/Public/Internationalization/BreakIterator.h
@@ -5,27 +5,27 @@
 #include "Templates/SharedPointer.h"
 #include "Internationalization/IBreakIterator.h"
 
-struct CORE_API FBreakIterator
+struct FBreakIterator
 {
 	/**
 	 * Create a new instance of a break iterator designed to locate character boundary points within a string
 	 */
-	static TSharedRef<IBreakIterator> CreateCharacterBoundaryIterator();
+	static CORE_API TSharedRef<IBreakIterator> CreateCharacterBoundaryIterator();
 
 	/**
 	 * Create a new instance of a break iterator designed to locate word boundary points within a string
 	 */
-	static TSharedRef<IBreakIterator> CreateWordBreakIterator();
+	static CORE_API TSharedRef<IBreakIterator> CreateWordBreakIterator();
 
 	/**
 	 * Create a new instance of a break iterator designed to locate appropriate soft-wrapping points within a string
 	 * Note that these aren't newlines as such, but rather places where the line could be soft-wrapped when displaying the text
 	 */
-	static TSharedRef<IBreakIterator> CreateLineBreakIterator();
+	static CORE_API TSharedRef<IBreakIterator> CreateLineBreakIterator();
 
 	/**
 	 * Create a new instance of a break iterator which is designed to locate appropriate soft-wrapping points within a string by searching for CamelCase boundaries
 	 * This may produce odd results when used with languages which don't have a concept of cases for characters, however should be fine for parsing variable or asset names
 	 */
-	static TSharedRef<IBreakIterator> CreateCamelCaseBreakIterator();
+	static CORE_API TSharedRef<IBreakIterator> CreateCamelCaseBreakIterator();
 };
diff --git a/Engine/Source/Runtime/Core/Public/Internationalization/Culture.h b/Engine/Source/Runtime/Core/Public/Internationalization/Culture.h
--- a/Engine/Source/Runtime/Core/Public/Internationalization/Culture.h
+++ b/Engine/Source/Runtime/Core/Public/Internationalization/Culture.h
@@ -19,7 +19,7 @@ class FLegacyCultureImplementation;
 typedef FLegacyCultureImplementation FCultureImplementation;
 #endif
 
-class CORE_API FCulture
+class FCulture
 {
 #if UE_ENABLE_ICU
 	friend class FText;
@@ -29,80 +29,80 @@ class CORE_API FCulture
 #endif
 
 public:
-	~FCulture();
+	CORE_API ~FCulture();
 	
-	static FCultureRef Create(TUniquePtr<FCultureImplementation>&& InImplementation);
+	static CORE_API FCultureRef Create(TUniquePtr<FCultureImplementation>&& InImplementation);
 
-	const FString& GetDisplayName() const;
+	CORE_API const FString& GetDisplayName() const;
 
-	const FString& GetEnglishName() const;
+	CORE_API const FString& GetEnglishName() const;
 
-	int GetKeyboardLayoutId() const;
+	CORE_API int GetKeyboardLayoutId() const;
 
-	int GetLCID() const;
+	CORE_API int GetLCID() const;
 
-	TArray<FString> GetPrioritizedParentCultureNames() const;
+	CORE_API TArray<FString> GetPrioritizedParentCultureNames() const;
 
-	static TArray<FString> GetPrioritizedParentCultureNames(const FString& LanguageCode, const FString& ScriptCode, const FString& RegionCode);
+	static CORE_API TArray<FString> GetPrioritizedParentCultureNames(const FString& LanguageCode, const FString& ScriptCode, const FString& RegionCode);
 
-	static FString CreateCultureName(const FString& LanguageCode, const FString& ScriptCode, const FString& RegionCode);
+	static CORE_API FString CreateCultureName(const FString& LanguageCode, const FString& ScriptCode, const FString& RegionCode);
 
-	static FString GetCanonicalName(const FString& Name);
+	static CORE_API FString GetCanonicalName(const FString& Name);
 
-	const FString& GetName() const;
+	CORE_API const FString& GetName() const;
 	
-	const FString& GetNativeName() const;
+	CORE_API const FString& GetNativeName() const;
 
-	const FString& GetUnrealLegacyThreeLetterISOLanguageName() const;
+	CORE_API const FString& GetUnrealLegacyThreeLetterISOLanguageName() const;
 
-	const FString& GetThreeLetterISOLanguageName() const;
+	CORE_API const FString& GetThreeLetterISOLanguageName() const;
 
-	const FString& GetTwoLetterISOLanguageName() const;
+	CORE_API const FString& GetTwoLetterISOLanguageName() const;
 
-	const FString& GetNativeLanguage() const;
+	CORE_API const FString& GetNativeLanguage() const;
 
-	const FString& GetRegion() const;
+	CORE_API const FString& GetRegion() const;
 
-	const FString& GetNativeRegion() const;
+	CORE_API const FString& GetNativeRegion() const;
 
-	const FString& GetScript() const;
+	CORE_API const FString& GetScript() const;
 
-	const FString& GetVariant() const;
+	CORE_API const FString& GetVariant() const;
 
-	bool IsRightToLeft() const;
+	CORE_API bool IsRightToLeft() const;
 
-	const FDecimalNumberFormattingRules& GetDecimalNumberFormattingRules() const;
+	CORE_API const FDecimalNumberFormattingRules& GetDecimalNumberFormattingRules() const;
 
-	const FDecimalNumberFormattingRules& GetPercentFormattingRules() const;
+	CORE_API const FDecimalNumberFormattingRules& GetPercentFormattingRules() const;
 
-	const FDecimalNumberFormattingRules& GetCurrencyFormattingRules(const FString& InCurrencyCode) const;
+	CORE_API const FDecimalNumberFormattingRules& GetCurrencyFormattingRules(const FString& InCurrencyCode) const;
 
 	/**
 	 * Get the correct plural form to use for the given number
 	 * @param PluralType The type of plural form to get (cardinal or ordinal)
 	 */
-	ETextPluralForm GetPluralForm(float Val,	const ETextPluralType PluralType) const;
-	ETextPluralForm GetPluralForm(double Val,	const ETextPluralType PluralType) const;
-	ETextPluralForm GetPluralForm(int8 Val,		const ETextPluralType PluralType) const;
-	ETextPluralForm GetPluralForm(int16 Val,	const ETextPluralType PluralType) const;
-	ETextPluralForm GetPluralForm(int32 Val,	const ETextPluralType PluralType) const;
-	ETextPluralForm GetPluralForm(int64 Val,	const ETextPluralType PluralType) const;
-	ETextPluralForm GetPluralForm(uint8 Val,	const ETextPluralType PluralType) const;
-	ETextPluralForm GetPluralForm(uint16 Val,	const ETextPluralType PluralType) const;
-	ETextPluralForm GetPluralForm(uint32 Val,	const ETextPluralType PluralType) const;
-	ETextPluralForm GetPluralForm(uint64 Val,	const ETextPluralType PluralType) const;
-	ETextPluralForm GetPluralForm(long Val,		const ETextPluralType PluralType) const;
+	CORE_API ETextPluralForm GetPluralForm(float Val,	const ETextPluralType PluralType) const;
+	CORE_API ETextPluralForm GetPluralForm(double Val,	const ETextPluralType PluralType) const;
+	CORE_API ETextPluralForm GetPluralForm(int8 Val,		const ETextPluralType PluralType) const;
+	CORE_API ETextPluralForm GetPluralForm(int16 Val,	const ETextPluralType PluralType) const;
+	CORE_API ETextPluralForm GetPluralForm(int32 Val,	const ETextPluralType PluralType) const;
+	CORE_API ETextPluralForm GetPluralForm(int64 Val,	const ETextPluralType PluralType) const;
+	CORE_API ETextPluralForm GetPluralForm(uint8 Val,	const ETextPluralType PluralType) const;
+	CORE_API ETextPluralForm GetPluralForm(uint16 Val,	const ETextPluralType PluralType) const;
+	CORE_API ETextPluralForm GetPluralForm(uint32 Val,	const ETextPluralType PluralType) const;
+	CORE_API ETextPluralForm GetPluralForm(uint64 Val,	const ETextPluralType PluralType) const;
+	CORE_API ETextPluralForm GetPluralForm(long Val,		const ETextPluralType PluralType) const;
 
 	/**
 	 * Get the plural forms supported by this culture
 	 * @param PluralType The type of plural form to get (cardinal or ordinal)
 	 */
-	const TArray<ETextPluralForm>& GetValidPluralForms(const ETextPluralType PluralType) const;
+	CORE_API const TArray<ETextPluralForm>& GetValidPluralForms(const ETextPluralType PluralType) const;
 	
-	void RefreshCultureDisplayNames(const TArray<FString>& InPrioritizedDisplayCultureNames, const bool bFullRefresh = true);
+	CORE_API void RefreshCultureDisplayNames(const TArray<FString>& InPrioritizedDisplayCultureNames, const bool bFullRefresh = true);
 
 private:
-	explicit FCulture(TUniquePtr<FCultureImplementation>&& InImplementation);
+	CORE_API explicit FCulture(TUniquePtr<FCultureImplementation>&& InImplementation);
 
 	TUniquePtr<FCultureImplementation> Implementation;
 
diff --git a/Engine/Source/Runtime/Core/Public/Internationalization/CultureFilter.h b/Engine/Source/Runtime/Core/Public/Internationalization/CultureFilter.h
--- a/Engine/Source/Runtime/Core/Public/Internationalization/CultureFilter.h
+++ b/Engine/Source/Runtime/Core/Public/Internationalization/CultureFilter.h
@@ -8,17 +8,17 @@ enum class EBuildConfiguration : uint8;
 enum class ELocalizationLoadFlags : uint8;
 
 /** Cache of enabled/disabled cultures loaded from config data */
-class CORE_API FCultureFilter
+class FCultureFilter
 {
 public:
 	/** Construct a culture filter based on the current build config and target, optionally filtering any loaded cultures against the set of available cultures */
-	explicit FCultureFilter(const TSet<FString>* AvailableCultures = nullptr);
+	CORE_API explicit FCultureFilter(const TSet<FString>* AvailableCultures = nullptr);
 
 	/** Construct a culture filter based on the given build config and target, optionally filtering any loaded cultures against the set of available cultures */
-	FCultureFilter(const EBuildConfiguration BuildConfig, const ELocalizationLoadFlags TargetFlags, const TSet<FString>* AvailableCultures = nullptr);
+	CORE_API FCultureFilter(const EBuildConfiguration BuildConfig, const ELocalizationLoadFlags TargetFlags, const TSet<FString>* AvailableCultures = nullptr);
 
 	/** Does the given culture pass this filter? */
-	bool IsCultureAllowed(const FString& Culture) const;
+	CORE_API bool IsCultureAllowed(const FString& Culture) const;
 
 private:
 	void Init(const EBuildConfiguration BuildConfig, const ELocalizationLoadFlags TargetFlags, const TSet<FString>* AvailableCultures);
diff --git a/Engine/Source/Runtime/Core/Public/Internationalization/Internationalization.h b/Engine/Source/Runtime/Core/Public/Internationalization/Internationalization.h
--- a/Engine/Source/Runtime/Core/Public/Internationalization/Internationalization.h
+++ b/Engine/Source/Runtime/Core/Public/Internationalization/Internationalization.h
@@ -56,7 +56,7 @@ public:
 	 * Get the current culture.
 	 * @note This function exists for legacy API parity with SetCurrentCulture and is equivalent to GetCurrentLanguage. It should *never* be used in internal localization/internationalization code!
 	 */
-	CORE_API FCultureRef GetCurrentCulture() const
+	FCultureRef GetCurrentCulture() const
 	{
 		return CurrentLanguage.ToSharedRef();
 	}
@@ -70,7 +70,7 @@ public:
 	/**
 	 * Get the current language (for localization).
 	 */
-	CORE_API FCultureRef GetCurrentLanguage() const
+	FCultureRef GetCurrentLanguage() const
 	{
 		return CurrentLanguage.ToSharedRef();
 	}
@@ -84,7 +84,7 @@ public:
 	/**
 	 * Get the current locale (for internationalization).
 	 */
-	CORE_API FCultureRef GetCurrentLocale() const
+	FCultureRef GetCurrentLocale() const
 	{
 		return CurrentLocale.ToSharedRef();
 	}
@@ -120,7 +120,7 @@ public:
 	 * Get the default culture specified by the OS.
 	 * @note This function exists for legacy API parity with GetCurrentCulture and is equivalent to GetDefaultLanguage. It should *never* be used in internal localization/internationalization code!
 	 */
-	CORE_API FCultureRef GetDefaultCulture() const
+	FCultureRef GetDefaultCulture() const
 	{
 		return DefaultLanguage.ToSharedRef();
 	}
@@ -128,7 +128,7 @@ public:
 	/**
 	 * Get the default language specified by the OS.
 	 */
-	CORE_API FCultureRef GetDefaultLanguage() const
+	FCultureRef GetDefaultLanguage() const
 	{
 		return DefaultLanguage.ToSharedRef();
 	}
@@ -136,7 +136,7 @@ public:
 	/**
 	 * Get the default locale specified by the OS.
 	 */
-	CORE_API FCultureRef GetDefaultLocale() const
+	FCultureRef GetDefaultLocale() const
 	{
 		return DefaultLocale.ToSharedRef();
 	}
@@ -144,7 +144,7 @@ public:
 	/**
 	 * Get the invariant culture that can be used when you don't care about localization/internationalization.
 	 */
-	CORE_API FCultureRef GetInvariantCulture() const
+	FCultureRef GetInvariantCulture() const
 	{
 		return InvariantCulture.ToSharedRef();
 	}
@@ -164,7 +164,7 @@ public:
 	 */
 	CORE_API void RestoreCultureState(const FCultureStateSnapshot& InSnapshot);
 
-	CORE_API bool IsInitialized() const {return bIsInitialized;}
+	bool IsInitialized() const {return bIsInitialized;}
 
 	/** Load and cache the data needed for every culture we know about (this is usually done per-culture as required) */
 	CORE_API void LoadAllCultureData();
@@ -208,7 +208,7 @@ public:
 
 	/** Broadcasts whenever the current culture changes */
 	DECLARE_EVENT(FInternationalization, FCultureChangedEvent)
-	CORE_API FCultureChangedEvent& OnCultureChanged() { return CultureChangedEvent; }
+	FCultureChangedEvent& OnCultureChanged() { return CultureChangedEvent; }
 
 private:
 	FInternationalization();
diff --git a/Engine/Source/Runtime/Core/Public/Internationalization/InternationalizationArchive.h b/Engine/Source/Runtime/Core/Public/Internationalization/InternationalizationArchive.h
--- a/Engine/Source/Runtime/Core/Public/Internationalization/InternationalizationArchive.h
+++ b/Engine/Source/Runtime/Core/Public/Internationalization/InternationalizationArchive.h
@@ -14,10 +14,10 @@
 
 class FLocMetadataObject;
 
-class CORE_API FArchiveEntry
+class FArchiveEntry
 {
 public:
-	FArchiveEntry(const FLocKey& InNamespace, const FLocKey& InKey, const FLocItem& InSource, const FLocItem& InTranslation, TSharedPtr<FLocMetadataObject> InKeyMetadataObj = nullptr, bool IsOptional = false);
+	CORE_API FArchiveEntry(const FLocKey& InNamespace, const FLocKey& InKey, const FLocItem& InSource, const FLocItem& InTranslation, TSharedPtr<FLocMetadataObject> InKeyMetadataObj = nullptr, bool IsOptional = false);
 
 	const FLocKey Namespace;
 	const FLocKey Key;
@@ -30,7 +30,7 @@ public:
 typedef TMultiMap< FLocKey, TSharedRef< FArchiveEntry > > FArchiveEntryByLocKeyContainer;
 typedef TMultiMap< FString, TSharedRef< FArchiveEntry >, FDefaultSetAllocator, FLocKeyMultiMapFuncs< TSharedRef< FArchiveEntry > > > FArchiveEntryByStringContainer;
 
-class CORE_API FInternationalizationArchive 
+class FInternationalizationArchive 
 {
 public:
 	enum class EFormatVersion : uint8
@@ -48,14 +48,14 @@ public:
 	{
 	}
 
-	bool AddEntry(const FLocKey& Namespace, const FLocKey& Key, const FLocItem& Source, const FLocItem& Translation, const TSharedPtr<FLocMetadataObject> KeyMetadataObj, const bool bOptional);
-	bool AddEntry(const TSharedRef<FArchiveEntry>& InEntry);
+	CORE_API bool AddEntry(const FLocKey& Namespace, const FLocKey& Key, const FLocItem& Source, const FLocItem& Translation, const TSharedPtr<FLocMetadataObject> KeyMetadataObj, const bool bOptional);
+	CORE_API bool AddEntry(const TSharedRef<FArchiveEntry>& InEntry);
 
-	void UpdateEntry(const TSharedRef<FArchiveEntry>& OldEntry, const TSharedRef<FArchiveEntry>& NewEntry);
+	CORE_API void UpdateEntry(const TSharedRef<FArchiveEntry>& OldEntry, const TSharedRef<FArchiveEntry>& NewEntry);
 
-	bool SetTranslation(const FLocKey& Namespace, const FLocKey& Key, const FLocItem& Source, const FLocItem& Translation, const TSharedPtr<FLocMetadataObject> KeyMetadataObj);
+	CORE_API bool SetTranslation(const FLocKey& Namespace, const FLocKey& Key, const FLocItem& Source, const FLocItem& Translation, const TSharedPtr<FLocMetadataObject> KeyMetadataObj);
 
-	TSharedPtr<FArchiveEntry> FindEntryByKey(const FLocKey& Namespace, const FLocKey& Key, const TSharedPtr<FLocMetadataObject> KeyMetadataObj) const;
+	CORE_API TSharedPtr<FArchiveEntry> FindEntryByKey(const FLocKey& Namespace, const FLocKey& Key, const TSharedPtr<FLocMetadataObject> KeyMetadataObj) const;
 
 	FArchiveEntryByLocKeyContainer::TConstIterator GetEntriesByKeyIterator() const
 	{
diff --git a/Engine/Source/Runtime/Core/Public/Internationalization/InternationalizationManifest.h b/Engine/Source/Runtime/Core/Public/Internationalization/InternationalizationManifest.h
--- a/Engine/Source/Runtime/Core/Public/Internationalization/InternationalizationManifest.h
+++ b/Engine/Source/Runtime/Core/Public/Internationalization/InternationalizationManifest.h
@@ -15,7 +15,7 @@
 
 class FLocMetadataObject;
 
-struct CORE_API FManifestContext
+struct FManifestContext
 {
 public:
 	FManifestContext()
@@ -35,12 +35,12 @@ public:
 	}
 
 	/** Copy ctor */
-	FManifestContext(const FManifestContext& Other);
+	CORE_API FManifestContext(const FManifestContext& Other);
 
-	FManifestContext& operator=(const FManifestContext& Other);
-	bool operator==(const FManifestContext& Other) const;
+	CORE_API FManifestContext& operator=(const FManifestContext& Other);
+	CORE_API bool operator==(const FManifestContext& Other) const;
 	inline bool operator!=(const FManifestContext& Other) const { return !(*this == Other); }
-	bool operator<(const FManifestContext& Other) const;
+	CORE_API bool operator<(const FManifestContext& Other) const;
 
 public:
 	FLocKey Key;
@@ -53,7 +53,7 @@ public:
 };
 
 
-struct CORE_API FLocItem
+struct FLocItem
 {
 public:
 	FLocItem()
@@ -75,15 +75,15 @@ public:
 	}
 
 	/** Copy ctor */
-	FLocItem(const FLocItem& Other);
+	CORE_API FLocItem(const FLocItem& Other);
 
-	FLocItem& operator=(const FLocItem& Other);
-	bool operator==(const FLocItem& Other) const;
+	CORE_API FLocItem& operator=(const FLocItem& Other);
+	CORE_API bool operator==(const FLocItem& Other) const;
 	inline bool operator!=(const FLocItem& Other) const { return !(*this == Other); }
-	bool operator<(const FLocItem& Other) const;
+	CORE_API bool operator<(const FLocItem& Other) const;
 
 	/** Similar functionality to == operator but ensures everything matches(ex. ignores COMPARISON_MODIFIER_PREFIX on metadata). */
-	bool IsExactMatch(const FLocItem& Other) const;
+	CORE_API bool IsExactMatch(const FLocItem& Other) const;
 
 public:
 	FString Text;
@@ -91,7 +91,7 @@ public:
 };
 
 
-class CORE_API FManifestEntry
+class FManifestEntry
 {
 public:
 	FManifestEntry(const FLocKey& InNamespace, const FLocItem& InSource)
@@ -101,9 +101,9 @@ public:
 	{
 	}
 
-	const FManifestContext* FindContext(const FLocKey& ContextKey, const TSharedPtr<FLocMetadataObject>& KeyMetadata = nullptr) const;
-	const FManifestContext* FindContextByKey(const FLocKey& ContextKey) const;
-	void MergeContextPlatformInfo(const FManifestContext& InContext);
+	CORE_API const FManifestContext* FindContext(const FLocKey& ContextKey, const TSharedPtr<FLocMetadataObject>& KeyMetadata = nullptr) const;
+	CORE_API const FManifestContext* FindContextByKey(const FLocKey& ContextKey) const;
+	CORE_API void MergeContextPlatformInfo(const FManifestContext& InContext);
 
 	const FLocKey Namespace;
 	const FLocItem Source;
@@ -115,7 +115,7 @@ typedef TMultiMap< FLocKey, TSharedRef< FManifestEntry > > FManifestEntryByLocKe
 typedef TMultiMap< FString, TSharedRef< FManifestEntry >, FDefaultSetAllocator, FLocKeyMultiMapFuncs< TSharedRef< FManifestEntry > > > FManifestEntryByStringContainer;
 
 
-class CORE_API FInternationalizationManifest 
+class FInternationalizationManifest 
 {
 public:
 	enum class EFormatVersion : uint8
@@ -138,15 +138,15 @@ public:
 	*
 	* @return Returns true if add was successful or a matching entry already exists, false is only returned in the case where a duplicate context was found with different text.
 	*/
-	bool AddSource(const FLocKey& Namespace, const FLocItem& Source, const FManifestContext& Context);
+	CORE_API bool AddSource(const FLocKey& Namespace, const FLocItem& Source, const FManifestContext& Context);
 
-	void UpdateEntry(const TSharedRef<FManifestEntry>& OldEntry, TSharedRef<FManifestEntry>& NewEntry);
+	CORE_API void UpdateEntry(const TSharedRef<FManifestEntry>& OldEntry, TSharedRef<FManifestEntry>& NewEntry);
 
-	TSharedPtr<FManifestEntry> FindEntryBySource(const FLocKey& Namespace, const FLocItem& Source) const;
+	CORE_API TSharedPtr<FManifestEntry> FindEntryBySource(const FLocKey& Namespace, const FLocItem& Source) const;
 
-	TSharedPtr<FManifestEntry> FindEntryByContext(const FLocKey& Namespace, const FManifestContext& Context) const;
+	CORE_API TSharedPtr<FManifestEntry> FindEntryByContext(const FLocKey& Namespace, const FManifestContext& Context) const;
 
-	TSharedPtr<FManifestEntry> FindEntryByKey(const FLocKey& Namespace, const FLocKey& Key, const FString* SourceText = nullptr) const;
+	CORE_API TSharedPtr<FManifestEntry> FindEntryByKey(const FLocKey& Namespace, const FLocKey& Key, const FString* SourceText = nullptr) const;
 
 	FManifestEntryByLocKeyContainer::TConstIterator GetEntriesByKeyIterator() const
 	{
diff --git a/Engine/Source/Runtime/Core/Public/Internationalization/InternationalizationMetadata.h b/Engine/Source/Runtime/Core/Public/Internationalization/InternationalizationMetadata.h
--- a/Engine/Source/Runtime/Core/Public/Internationalization/InternationalizationMetadata.h
+++ b/Engine/Source/Runtime/Core/Public/Internationalization/InternationalizationMetadata.h
@@ -29,10 +29,10 @@ class FLocMetadataObject;
 /**
  * A Metadata Value is a structure that can be a number of types.
  */
-class CORE_API FLocMetadataValue
+class FLocMetadataValue
 {
 public:
-	virtual ~FLocMetadataValue() = 0;
+	CORE_API virtual ~FLocMetadataValue() = 0;
 
 	virtual FString ToString() const = 0;
 
@@ -63,18 +63,18 @@ protected:
 
 	virtual FString GetTypeString() const = 0;
 
-	void ErrorMessage(const FString& InType);
+	CORE_API void ErrorMessage(const FString& InType);
 
 private:
-	FLocMetadataValue( const FLocMetadataValue& );
-	FLocMetadataValue& operator=( const FLocMetadataValue& );
+	CORE_API FLocMetadataValue( const FLocMetadataValue& );
+	CORE_API FLocMetadataValue& operator=( const FLocMetadataValue& );
 };
 
 
 /**
  * A LocMetadata Object is a structure holding an unordered set of name/value pairs.
  */
-class CORE_API FLocMetadataObject
+class FLocMetadataObject
 {
 public:
 	FLocMetadataObject()
@@ -83,7 +83,7 @@ public:
 	}
 
 	/** Copy ctor */
-	FLocMetadataObject( const FLocMetadataObject& Other );
+	CORE_API FLocMetadataObject( const FLocMetadataObject& Other );
 	
 	template<ELocMetadataType LocMetadataType>
 	TSharedPtr<FLocMetadataValue> GetField( const FString& FieldName )
@@ -125,45 +125,45 @@ public:
 		return false;
 	}
 
-	void SetField( const FString& FieldName, const TSharedPtr<FLocMetadataValue>& Value );
+	CORE_API void SetField( const FString& FieldName, const TSharedPtr<FLocMetadataValue>& Value );
 
-	void RemoveField(const FString& FieldName);
+	CORE_API void RemoveField(const FString& FieldName);
 
 	/** Get the field named FieldName as a string. Ensures that the field is present and is of type LocMetadata string. */
-	FString GetStringField(const FString& FieldName);
+	CORE_API FString GetStringField(const FString& FieldName);
 
 	/** Add a field named FieldName with value of StringValue */
-	void SetStringField( const FString& FieldName, const FString& StringValue );
+	CORE_API void SetStringField( const FString& FieldName, const FString& StringValue );
 
 	/** Get the field named FieldName as a boolean. Ensures that the field is present and is of type LocMetadata boolean. */
-	bool GetBoolField(const FString& FieldName);
+	CORE_API bool GetBoolField(const FString& FieldName);
 
 	/** Set a boolean field named FieldName and value of InValue */
-	void SetBoolField( const FString& FieldName, bool InValue );
+	CORE_API void SetBoolField( const FString& FieldName, bool InValue );
 
 	/** Get the field named FieldName as an array. Ensures that the field is present and is of type LocMetadata Array. */
-	TArray< TSharedPtr<FLocMetadataValue> > GetArrayField(const FString& FieldName);
+	CORE_API TArray< TSharedPtr<FLocMetadataValue> > GetArrayField(const FString& FieldName);
 
 	/** Set an array field named FieldName and value of Array */
-	void SetArrayField( const FString& FieldName, const TArray< TSharedPtr<FLocMetadataValue> >& Array );
+	CORE_API void SetArrayField( const FString& FieldName, const TArray< TSharedPtr<FLocMetadataValue> >& Array );
 
 	/** Get the field named FieldName as a LocMetadata object. Ensures that the field is present and is a LocMetadata object*/
-	TSharedPtr<FLocMetadataObject> GetObjectField(const FString& FieldName);
+	CORE_API TSharedPtr<FLocMetadataObject> GetObjectField(const FString& FieldName);
 
 	/** Set an ObjectField named FieldName and value of LocMetadataObject */
-	void SetObjectField( const FString& FieldName, const TSharedPtr<FLocMetadataObject>& LocMetadataObject );
+	CORE_API void SetObjectField( const FString& FieldName, const TSharedPtr<FLocMetadataObject>& LocMetadataObject );
 
-	FLocMetadataObject& operator=( const FLocMetadataObject& Other );
+	CORE_API FLocMetadataObject& operator=( const FLocMetadataObject& Other );
 
-	bool operator==(const FLocMetadataObject& Other) const;
-	bool operator<( const FLocMetadataObject& Other ) const;
+	CORE_API bool operator==(const FLocMetadataObject& Other) const;
+	CORE_API bool operator<( const FLocMetadataObject& Other ) const;
 
 	/** Similar functionality to == operator but ensures everything matches(ignores COMPARISON_MODIFIER_PREFIX). */
-	bool IsExactMatch( const FLocMetadataObject& Other ) const;
+	CORE_API bool IsExactMatch( const FLocMetadataObject& Other ) const;
 
-	static bool IsMetadataExactMatch( const FLocMetadataObject* const MetadataA, const FLocMetadataObject* const MetadataB );
+	static CORE_API bool IsMetadataExactMatch( const FLocMetadataObject* const MetadataA, const FLocMetadataObject* const MetadataB );
 
-	FString ToString() const;
+	CORE_API FString ToString() const;
 
 	CORE_API friend FArchive& operator<<(FArchive& Archive, FLocMetadataObject& Object);
 	CORE_API friend void operator<<(FStructuredArchive::FSlot Slot, FLocMetadataObject& Object);
@@ -173,20 +173,20 @@ public:
 	TMap< FString, TSharedPtr<FLocMetadataValue> > Values;
 
 	/** Special reserved character.  When encountered as a prefix metadata name(the key in the Values map), it will adjust the way comparisons are done. */
-	static const TCHAR* COMPARISON_MODIFIER_PREFIX;
+	static CORE_API const TCHAR* COMPARISON_MODIFIER_PREFIX;
 };
 
 /** A LocMetadata String Value. */
-class CORE_API FLocMetadataValueString : public FLocMetadataValue
+class FLocMetadataValueString : public FLocMetadataValue
 {
 public:
 	FLocMetadataValueString(const FString& InString) : Value(InString) {}
-	FLocMetadataValueString(FStructuredArchive::FSlot Slot);
-	static void Serialize(FLocMetadataValueString& Value, FStructuredArchive::FSlot Slot);
+	CORE_API FLocMetadataValueString(FStructuredArchive::FSlot Slot);
+	static CORE_API void Serialize(FLocMetadataValueString& Value, FStructuredArchive::FSlot Slot);
 
 	virtual FString AsString() override {return Value;}
 
-	virtual TSharedRef<FLocMetadataValue> Clone() const override;
+	CORE_API virtual TSharedRef<FLocMetadataValue> Clone() const override;
 	void SetString( const FString& InString ) { Value = InString; }
 
 	virtual FString ToString() const override {return Value;}
@@ -197,24 +197,24 @@ protected:
 
 	FString Value;
 
-	virtual bool EqualTo( const FLocMetadataValue& Other ) const override;
-	virtual bool LessThan( const FLocMetadataValue& Other ) const override;
+	CORE_API virtual bool EqualTo( const FLocMetadataValue& Other ) const override;
+	CORE_API virtual bool LessThan( const FLocMetadataValue& Other ) const override;
 
 	virtual FString GetTypeString() const override {return TEXT("String");}
 };
 
 
 /** A LocMetadata Boolean Value. */
-class CORE_API FLocMetadataValueBoolean : public FLocMetadataValue
+class FLocMetadataValueBoolean : public FLocMetadataValue
 {
 public:
 	FLocMetadataValueBoolean(bool InBool) : Value(InBool) {}
-	FLocMetadataValueBoolean(FStructuredArchive::FSlot Slot);
-	static void Serialize(FLocMetadataValueBoolean& Value, FStructuredArchive::FSlot Slot);
+	CORE_API FLocMetadataValueBoolean(FStructuredArchive::FSlot Slot);
+	static CORE_API void Serialize(FLocMetadataValueBoolean& Value, FStructuredArchive::FSlot Slot);
 
 	virtual bool AsBool() override {return Value;}
 
-	virtual TSharedRef<FLocMetadataValue> Clone() const override;
+	CORE_API virtual TSharedRef<FLocMetadataValue> Clone() const override;
 	
 	virtual FString ToString() const override {return Value ? TEXT("true") : TEXT("false");}
 
@@ -223,58 +223,58 @@ public:
 protected:
 	bool Value;
 
-	virtual bool EqualTo( const FLocMetadataValue& Other ) const override;
-	virtual bool LessThan( const FLocMetadataValue& Other ) const override;
+	CORE_API virtual bool EqualTo( const FLocMetadataValue& Other ) const override;
+	CORE_API virtual bool LessThan( const FLocMetadataValue& Other ) const override;
 
 	virtual FString GetTypeString() const override {return TEXT("Boolean");}
 };
 
 /** A LocMetadata Array Value. */
-class CORE_API FLocMetadataValueArray : public FLocMetadataValue
+class FLocMetadataValueArray : public FLocMetadataValue
 {
 public:
 	FLocMetadataValueArray(const TArray< TSharedPtr<FLocMetadataValue> >& InArray) : Value(InArray) {}
-	FLocMetadataValueArray(FStructuredArchive::FSlot Slot);
-	static void Serialize(FLocMetadataValueArray& Value, FStructuredArchive::FSlot Slot);
+	CORE_API FLocMetadataValueArray(FStructuredArchive::FSlot Slot);
+	static CORE_API void Serialize(FLocMetadataValueArray& Value, FStructuredArchive::FSlot Slot);
 
 	virtual TArray< TSharedPtr<FLocMetadataValue> > AsArray() override {return Value;}
 
-	virtual TSharedRef<FLocMetadataValue> Clone() const override;
+	CORE_API virtual TSharedRef<FLocMetadataValue> Clone() const override;
 	
-	virtual FString ToString() const override;
+	CORE_API virtual FString ToString() const override;
 
 	virtual ELocMetadataType GetType() const override {return ELocMetadataType::Array;}
 
 protected:
 	TArray< TSharedPtr<FLocMetadataValue> > Value;
 
-	virtual bool EqualTo( const FLocMetadataValue& Other ) const override;
-	virtual bool LessThan( const FLocMetadataValue& Other ) const override;
+	CORE_API virtual bool EqualTo( const FLocMetadataValue& Other ) const override;
+	CORE_API virtual bool LessThan( const FLocMetadataValue& Other ) const override;
 
 	virtual FString GetTypeString() const override {return TEXT("Array");}
 };
 
 /** A LocMetadata Object Value. */
-class CORE_API FLocMetadataValueObject : public FLocMetadataValue
+class FLocMetadataValueObject : public FLocMetadataValue
 {
 public:
 	FLocMetadataValueObject(TSharedPtr<FLocMetadataObject> InObject) : Value(InObject) {}
-	FLocMetadataValueObject(FStructuredArchive::FSlot Slot);
-	static void Serialize(FLocMetadataValueObject& Value, FStructuredArchive::FSlot Slot);
+	CORE_API FLocMetadataValueObject(FStructuredArchive::FSlot Slot);
+	static CORE_API void Serialize(FLocMetadataValueObject& Value, FStructuredArchive::FSlot Slot);
 
 	virtual TSharedPtr<FLocMetadataObject> AsObject() override {return Value;}
 
-	virtual TSharedRef<FLocMetadataValue> Clone() const override;
+	CORE_API virtual TSharedRef<FLocMetadataValue> Clone() const override;
 	
-	virtual FString ToString() const override;
+	CORE_API virtual FString ToString() const override;
 
 	virtual ELocMetadataType GetType() const override {return ELocMetadataType::Object;}
 
 protected:
 	TSharedPtr<FLocMetadataObject> Value;
 
-	virtual bool EqualTo( const FLocMetadataValue& Other ) const override;
-	virtual bool LessThan( const FLocMetadataValue& Other ) const override;
+	CORE_API virtual bool EqualTo( const FLocMetadataValue& Other ) const override;
+	CORE_API virtual bool LessThan( const FLocMetadataValue& Other ) const override;
 
 	virtual FString GetTypeString() const override {return TEXT("Object");}
 };
diff --git a/Engine/Source/Runtime/Core/Public/Internationalization/PolyglotTextData.h b/Engine/Source/Runtime/Core/Public/Internationalization/PolyglotTextData.h
--- a/Engine/Source/Runtime/Core/Public/Internationalization/PolyglotTextData.h
+++ b/Engine/Source/Runtime/Core/Public/Internationalization/PolyglotTextData.h
@@ -13,7 +13,7 @@
  * Polyglot data that may be registered to the text localization manager at runtime.
  * @note This struct is mirrored in NoExportTypes.h for UHT.
  */
-struct CORE_API FPolyglotTextData
+struct FPolyglotTextData
 {
 	/** Needed to avoid member access issues in UHT generated code */
 	friend struct Z_Construct_UScriptStruct_FPolyglotTextData_Statics;
@@ -27,125 +27,125 @@ public:
 	/**
 	 * Construct this polyglot data with an identity, optionally immediately providing the source text and native culture.
 	 */
-	FPolyglotTextData(const ELocalizedTextSourceCategory& InCategory, const FString& InNamespace, const FString& InKey, const FString& InNativeString, const FString& InNativeCulture = FString());
+	CORE_API FPolyglotTextData(const ELocalizedTextSourceCategory& InCategory, const FString& InNamespace, const FString& InKey, const FString& InNativeString, const FString& InNativeCulture = FString());
 
 	/**
 	 * Is this polyglot data valid and can be registered with the text localization manager?
 	 */
-	bool IsValid(FText* OutFailureReason = nullptr) const;
+	CORE_API bool IsValid(FText* OutFailureReason = nullptr) const;
 
 	/**
 	 * Set the category of this polyglot data.
 	 * @note This affects when and how the data is loaded into the text localization manager.
 	 */
-	void SetCategory(const ELocalizedTextSourceCategory InCategory);
+	CORE_API void SetCategory(const ELocalizedTextSourceCategory InCategory);
 
 	/**
 	 * Get the category of this polyglot data.
 	 */
-	ELocalizedTextSourceCategory GetCategory() const;
+	CORE_API ELocalizedTextSourceCategory GetCategory() const;
 
 	/**
 	 * Set the native culture of this polyglot data.
 	 * @note This may be empty, and if empty, will be inferred from the native culture of the text category.
 	 */
-	void SetNativeCulture(const FString& InNativeCulture);
+	CORE_API void SetNativeCulture(const FString& InNativeCulture);
 
 	/**
 	 * Get the native culture of this polyglot data.
 	 */
-	const FString& GetNativeCulture() const;
+	CORE_API const FString& GetNativeCulture() const;
 
 	/**
 	 * Resolve the native culture of this polyglot data, either using the native culture if specified, or inferring it from the category.
 	 */
-	FString ResolveNativeCulture() const;
+	CORE_API FString ResolveNativeCulture() const;
 
 	/**
 	 * Get the localized cultures of this polyglot data.
 	 */
-	TArray<FString> GetLocalizedCultures() const;
+	CORE_API TArray<FString> GetLocalizedCultures() const;
 
 	/**
 	 * Set the identity of the text created from this polyglot data.
 	 * @note The key must not be empty.
 	 */
-	void SetIdentity(const FString& InNamespace, const FString& InKey);
+	CORE_API void SetIdentity(const FString& InNamespace, const FString& InKey);
 
 	/**
 	 * Get the identity of the text created from this polyglot data.
 	 */
-	void GetIdentity(FString& OutNamespace, FString& OutKey) const;
+	CORE_API void GetIdentity(FString& OutNamespace, FString& OutKey) const;
 
 	/**
 	 * Get the namespace of the text created from this polyglot data.
 	 */
-	const FString& GetNamespace() const;
+	CORE_API const FString& GetNamespace() const;
 
 	/**
 	 * Get the key of the text created from this polyglot data.
 	 */
-	const FString& GetKey() const;
+	CORE_API const FString& GetKey() const;
 
 	/**
 	 * Set the native string of this polyglot data.
 	 * @note The native string must not be empty.
 	 */
-	void SetNativeString(const FString& InNativeString);
+	CORE_API void SetNativeString(const FString& InNativeString);
 
 	/**
 	 * Get the native string of this polyglot data.
 	 */
-	const FString& GetNativeString() const;
+	CORE_API const FString& GetNativeString() const;
 	
 	/**
 	 * Add a localized string to this polyglot data.
 	 * @note The native culture may also have a translation added via this function.
 	 */
-	void AddLocalizedString(const FString& InCulture, const FString& InLocalizedString);
+	CORE_API void AddLocalizedString(const FString& InCulture, const FString& InLocalizedString);
 
 	/**
 	 * Remove a localized string from this polyglot data.
 	 */
-	void RemoveLocalizedString(const FString& InCulture);
+	CORE_API void RemoveLocalizedString(const FString& InCulture);
 
 	/**
 	 * Get a localized string from this polyglot data.
 	 */
-	bool GetLocalizedString(const FString& InCulture, FString& OutLocalizedString) const;
+	CORE_API bool GetLocalizedString(const FString& InCulture, FString& OutLocalizedString) const;
 
 	/**
 	 * Clear the localized strings from this polyglot data.
 	 */
-	void ClearLocalizedStrings();
+	CORE_API void ClearLocalizedStrings();
 
 	/**
 	 * Set whether this polyglot data is a minimal patch.
 	 * @see bIsMinimalPatch.
 	 */
-	void IsMinimalPatch(const bool InIsMinimalPatch);
+	CORE_API void IsMinimalPatch(const bool InIsMinimalPatch);
 
 	/**
 	 * Get whether this polyglot data is a minimal patch.
 	 * @see bIsMinimalPatch.
 	 */
-	bool IsMinimalPatch() const;
+	CORE_API bool IsMinimalPatch() const;
 
 	/**
 	 * Get the text instance created from this polyglot data.
 	 */
-	FText GetText() const;
+	CORE_API FText GetText() const;
 
 private:
 	/**
 	 * Cache the text instance created from this polyglot data.
 	 */
-	void CacheText(FText* OutFailureReason = nullptr);
+	CORE_API void CacheText(FText* OutFailureReason = nullptr);
 
 	/**
 	 * Clear the cache of the text instance created from this polyglot data.
 	 */
-	void ClearCache();
+	CORE_API void ClearCache();
 
 	/**
 	 * The category of this polyglot data.
diff --git a/Engine/Source/Runtime/Core/Public/Internationalization/Regex.h b/Engine/Source/Runtime/Core/Public/Internationalization/Regex.h
--- a/Engine/Source/Runtime/Core/Public/Internationalization/Regex.h
+++ b/Engine/Source/Runtime/Core/Public/Internationalization/Regex.h
@@ -19,12 +19,12 @@ ENUM_CLASS_FLAGS(ERegexPatternFlags);
  * Implements a regular expression pattern.
  * @note DO NOT use this class as a file-level variable as its construction relies on the internationalization system being initialized!
  */
-class CORE_API FRegexPattern
+class FRegexPattern
 {
 	friend class FRegexMatcher;
 
 public:
-	explicit FRegexPattern(const FString& SourceString, ERegexPatternFlags Flags = ERegexPatternFlags::None);
+	CORE_API explicit FRegexPattern(const FString& SourceString, ERegexPatternFlags Flags = ERegexPatternFlags::None);
 
 	FRegexPattern(const FRegexPattern&) = default;
 	FRegexPattern& operator=(const FRegexPattern&) = default;
@@ -40,11 +40,11 @@ private:
  * Implements a regular expression pattern matcher.
  * @note DO NOT use this class as a file-level variable as its construction relies on the internationalization system being initialized!
  */
-class CORE_API FRegexMatcher
+class FRegexMatcher
 {
 public:
-	FRegexMatcher(const FRegexPattern& SourcePattern, const FString& Input);
-	FRegexMatcher(FRegexPattern&& SourcePattern, const FString& Input);
+	CORE_API FRegexMatcher(const FRegexPattern& SourcePattern, const FString& Input);
+	CORE_API FRegexMatcher(FRegexPattern&& SourcePattern, const FString& Input);
 
 	FRegexMatcher(const FRegexMatcher&) = delete;
 	FRegexMatcher& operator=(const FRegexMatcher&) = delete;
@@ -52,18 +52,18 @@ public:
 	FRegexMatcher(FRegexMatcher&&) = default;
 	FRegexMatcher& operator=(FRegexMatcher&&) = default;
 
-	bool FindNext();
+	CORE_API bool FindNext();
 
-	int32 GetMatchBeginning();
-	int32 GetMatchEnding();
+	CORE_API int32 GetMatchBeginning();
+	CORE_API int32 GetMatchEnding();
 
-	int32 GetCaptureGroupBeginning(const int32 Index);
-	int32 GetCaptureGroupEnding(const int32 Index);
-	FString GetCaptureGroup(const int32 Index);
+	CORE_API int32 GetCaptureGroupBeginning(const int32 Index);
+	CORE_API int32 GetCaptureGroupEnding(const int32 Index);
+	CORE_API FString GetCaptureGroup(const int32 Index);
 
-	int32 GetBeginLimit();
-	int32 GetEndLimit();
-	void SetLimits(const int32 BeginIndex, const int32 EndIndex);
+	CORE_API int32 GetBeginLimit();
+	CORE_API int32 GetEndLimit();
+	CORE_API void SetLimits(const int32 BeginIndex, const int32 EndIndex);
 
 private:
 	FRegexPattern Pattern;
diff --git a/Engine/Source/Runtime/Core/Public/Internationalization/StringTableCore.h b/Engine/Source/Runtime/Core/Public/Internationalization/StringTableCore.h
--- a/Engine/Source/Runtime/Core/Public/Internationalization/StringTableCore.h
+++ b/Engine/Source/Runtime/Core/Public/Internationalization/StringTableCore.h
@@ -29,7 +29,7 @@ class UStringTable;
 DECLARE_LOG_CATEGORY_EXTERN(LogStringTable, Log, All);
 
 /** Entry within a string table */
-class CORE_API FStringTableEntry
+class FStringTableEntry
 {
 public:
 	/** Create a new string table entry using the given data */
@@ -39,31 +39,31 @@ public:
 	}
 
 	/** Default constructor */
-	FStringTableEntry();
+	CORE_API FStringTableEntry();
 
 	/** Create a new string table entry using the given data */
-	FStringTableEntry(FStringTableConstRef InOwnerTable, FString InSourceString, FTextId InDisplayStringId);
+	CORE_API FStringTableEntry(FStringTableConstRef InOwnerTable, FString InSourceString, FTextId InDisplayStringId);
 
 	/** @return true if this entry is currently owned by a string table, false if it's been disowned (and should be re-cached) */
-	bool IsOwned() const;
+	CORE_API bool IsOwned() const;
 
 	/** Disown this string table entry. This is used to notify external code that has cached this entry that it needs to re-cache it from the string table */
-	void Disown();
+	CORE_API void Disown();
 
 	/** Is this string table entry owned by the given string table? */
-	bool IsOwnedBy(const FStringTable& InStringTable) const;
+	CORE_API bool IsOwnedBy(const FStringTable& InStringTable) const;
 
 	/** Get the source string of this string table entry */
-	const FString& GetSourceString() const;
+	CORE_API const FString& GetSourceString() const;
 
 	/** Get the display string of this string table entry */
-	FTextConstDisplayStringPtr GetDisplayString() const;
+	CORE_API FTextConstDisplayStringPtr GetDisplayString() const;
 
 	/** Get the display string ID of this string table entry */
-	FTextId GetDisplayStringId() const;
+	CORE_API FTextId GetDisplayStringId() const;
 
 	/** Get the placeholder source string to use for string table entries that are missing */
-	static const FString& GetPlaceholderSourceString();
+	static CORE_API const FString& GetPlaceholderSourceString();
 
 private:
 	/** The string table that owns us (if any) */
@@ -77,7 +77,7 @@ private:
 };
 
 /** String table implementation. Holds Key->SourceString pairs of text. */
-class CORE_API FStringTable : public TSharedFromThis<FStringTable, ESPMode::ThreadSafe>
+class FStringTable : public TSharedFromThis<FStringTable, ESPMode::ThreadSafe>
 {
 public:
 	/** Create a new string table */
@@ -87,78 +87,78 @@ public:
 	}
 
 	/** Default constructor */
-	FStringTable();
+	CORE_API FStringTable();
 
 	/** Destructor */
-	~FStringTable();
+	CORE_API ~FStringTable();
 
 	/** @return The asset that owns this string table instance (if any) */
-	UStringTable* GetOwnerAsset() const;
+	CORE_API UStringTable* GetOwnerAsset() const;
 
 	/** Set the asset that owns this string table instance (if any) */
-	void SetOwnerAsset(UStringTable* InOwnerAsset);
+	CORE_API void SetOwnerAsset(UStringTable* InOwnerAsset);
 
 	/** Has this string table been fully loaded yet? (used during asset loading) */
-	bool IsLoaded() const;
+	CORE_API bool IsLoaded() const;
 
 	/** Set whether this string table has been fully loaded yet */
-	void IsLoaded(const bool bInIsLoaded);
+	CORE_API void IsLoaded(const bool bInIsLoaded);
 
 	/** @return The namespace used by all entries in this string table */
-	FString GetNamespace() const;
+	CORE_API FString GetNamespace() const;
 
 	/** Set the namespace used by all entries in this string table */
-	void SetNamespace(const FString& InNamespace);
+	CORE_API void SetNamespace(const FString& InNamespace);
 
 	/** Get the source string used by the given entry (if any) */
-	bool GetSourceString(const FTextKey& InKey, FString& OutSourceString) const;
+	CORE_API bool GetSourceString(const FTextKey& InKey, FString& OutSourceString) const;
 
 	/** Set the source string used by the given entry (will replace any existing data for that entry) */
-	void SetSourceString(const FTextKey& InKey, const FString& InSourceString);
+	CORE_API void SetSourceString(const FTextKey& InKey, const FString& InSourceString);
 
 	/** Remove the given entry (including its meta-data) */
-	void RemoveSourceString(const FTextKey& InKey);
+	CORE_API void RemoveSourceString(const FTextKey& InKey);
 
 	/** Enumerate all source strings in the table. Return true from the enumerator to continue, or false to stop */
-	void EnumerateSourceStrings(const TFunctionRef<bool(const FString&, const FString&)>& InEnumerator) const;
-	void EnumerateKeysAndSourceStrings(const TFunctionRef<bool(const FTextKey&, const FString&)>& InEnumerator) const;
+	CORE_API void EnumerateSourceStrings(const TFunctionRef<bool(const FString&, const FString&)>& InEnumerator) const;
+	CORE_API void EnumerateKeysAndSourceStrings(const TFunctionRef<bool(const FTextKey&, const FString&)>& InEnumerator) const;
 
 	/** Clear all entries from the table (including their meta-data) */
-	void ClearSourceStrings(const int32 InSlack = 0);
+	CORE_API void ClearSourceStrings(const int32 InSlack = 0);
 
 	/** Find the entry with the given key (if any) */
-	FStringTableEntryConstPtr FindEntry(const FTextKey& InKey) const;
+	CORE_API FStringTableEntryConstPtr FindEntry(const FTextKey& InKey) const;
 
 	/** Given an entry, check to see if it exists in this table, and if so, get its key */
-	bool FindKey(const FStringTableEntryConstRef& InEntry, FString& OutKey) const;
-	bool FindKey(const FStringTableEntryConstRef& InEntry, FTextKey& OutKey) const;
+	CORE_API bool FindKey(const FStringTableEntryConstRef& InEntry, FString& OutKey) const;
+	CORE_API bool FindKey(const FStringTableEntryConstRef& InEntry, FTextKey& OutKey) const;
 
 	/** Get the meta-data with the given ID associated with the given entry, or an empty string if not found */
-	FString GetMetaData(const FTextKey& InKey, const FName InMetaDataId) const;
+	CORE_API FString GetMetaData(const FTextKey& InKey, const FName InMetaDataId) const;
 
 	/** Set the meta-data with the given ID associated with the given entry */
-	void SetMetaData(const FTextKey& InKey, const FName InMetaDataId, const FString& InMetaDataValue);
+	CORE_API void SetMetaData(const FTextKey& InKey, const FName InMetaDataId, const FString& InMetaDataValue);
 
 	/** Remove the meta-data with the given ID associated with the given entry */
-	void RemoveMetaData(const FTextKey& InKey, const FName InMetaDataId);
+	CORE_API void RemoveMetaData(const FTextKey& InKey, const FName InMetaDataId);
 
 	/** Enumerate all meta-data associated with the given entry. Return true from the enumerator to continue, or false to stop */
-	void EnumerateMetaData(const FTextKey& InKey, const TFunctionRef<bool(FName, const FString&)>& InEnumerator) const;
+	CORE_API void EnumerateMetaData(const FTextKey& InKey, const TFunctionRef<bool(FName, const FString&)>& InEnumerator) const;
 
 	/** Remove all meta-data associated with the given entry */
-	void ClearMetaData(const FTextKey& InKey);
+	CORE_API void ClearMetaData(const FTextKey& InKey);
 
 	/** Clear all meta-data from the table */
-	void ClearMetaData(const int32 InSlack = 0);
+	CORE_API void ClearMetaData(const int32 InSlack = 0);
 
 	/** Serialize this string table to/from an archive */
-	void Serialize(FArchive& Ar);
+	CORE_API void Serialize(FArchive& Ar);
 
 	/** Export the key, string, and meta-data information in this string table to a CSV file (does not export the namespace) */
-	bool ExportStrings(const FString& InFilename) const;
+	CORE_API bool ExportStrings(const FString& InFilename) const;
 
 	/** Import key, string, and meta-data information from a CSV file to this string table (does not import the namespace) */
-	bool ImportStrings(const FString& InFilename);
+	CORE_API bool ImportStrings(const FString& InFilename);
 
 private:
 	/** Pointer back to the asset that owns this table */
@@ -185,7 +185,7 @@ private:
 };
 
 /** Interface to allow Core code to access String Table assets from the Engine */
-class CORE_API IStringTableEngineBridge
+class IStringTableEngineBridge
 {
 public:
 	/**
@@ -304,24 +304,24 @@ protected:
 	virtual bool IsStringTableAssetBeingReplacedImpl(const UStringTable* InStringTableAsset) = 0;
 
 	/** Singleton instance, populated by the derived type */
-	static IStringTableEngineBridge* InstancePtr;
+	static CORE_API IStringTableEngineBridge* InstancePtr;
 
 	/** Whether String Table find/load is currently deferred (eg, during module load) */
-	static std::atomic<int8> DeferFindOrLoad;
+	static CORE_API std::atomic<int8> DeferFindOrLoad;
 };
 
 /** String table redirect utils */
-struct CORE_API FStringTableRedirects
+struct FStringTableRedirects
 {
 	/** Initialize the string table redirects */
-	static void InitStringTableRedirects();
+	static CORE_API void InitStringTableRedirects();
 
 	/** Redirect a table ID */
-	static void RedirectTableId(FName& InOutTableId);
+	static CORE_API void RedirectTableId(FName& InOutTableId);
 
 	/** Redirect a key */
-	static void RedirectKey(const FName InTableId, FTextKey& InOutKey);
+	static CORE_API void RedirectKey(const FName InTableId, FTextKey& InOutKey);
 
 	/** Redirect a table ID and key */
-	static void RedirectTableIdAndKey(FName& InOutTableId, FTextKey& InOutKey);
+	static CORE_API void RedirectTableIdAndKey(FName& InOutTableId, FTextKey& InOutKey);
 };
diff --git a/Engine/Source/Runtime/Core/Public/Internationalization/StringTableRegistry.h b/Engine/Source/Runtime/Core/Public/Internationalization/StringTableRegistry.h
--- a/Engine/Source/Runtime/Core/Public/Internationalization/StringTableRegistry.h
+++ b/Engine/Source/Runtime/Core/Public/Internationalization/StringTableRegistry.h
@@ -20,51 +20,51 @@ class UStringTable;
 template <typename FuncType> class TFunctionRef;
 
 /** Singleton registry of known string table instances */
-class CORE_API FStringTableRegistry
+class FStringTableRegistry
 {
 public:
 	/** Singleton accessor */
-	static FStringTableRegistry& Get();
+	static CORE_API FStringTableRegistry& Get();
 
 	/** Register a string table with the given ID */
-	void RegisterStringTable(const FName InTableId, FStringTableRef InTable);
+	CORE_API void RegisterStringTable(const FName InTableId, FStringTableRef InTable);
 
 	/** Unregister a string table with the given ID */
-	void UnregisterStringTable(const FName InTableId);
+	CORE_API void UnregisterStringTable(const FName InTableId);
 
 	/** Try and find a string table with the given ID */
-	FStringTablePtr FindMutableStringTable(const FName InTableId) const;
+	CORE_API FStringTablePtr FindMutableStringTable(const FName InTableId) const;
 
 	/** Try and find a string table with the given ID */
-	FStringTableConstPtr FindStringTable(const FName InTableId) const;
+	CORE_API FStringTableConstPtr FindStringTable(const FName InTableId) const;
 
 	/** Try and find a string table asset with the given ID */
-	UStringTable* FindStringTableAsset(const FName InTableId) const;
+	CORE_API UStringTable* FindStringTableAsset(const FName InTableId) const;
 
 	/** Enumerate all registered string tables */
-	void EnumerateStringTables(const TFunctionRef<bool(const FName&, const FStringTableConstRef&)>& InEnumerator) const;
+	CORE_API void EnumerateStringTables(const TFunctionRef<bool(const FName&, const FStringTableConstRef&)>& InEnumerator) const;
 
 	/** Given an FText, try and find the table ID and key for it */
 	UE_DEPRECATED(5.0, "FindTableIdAndKey is deprecated. Use FTextInspector::GetTableIdAndKey instead.")
-	bool FindTableIdAndKey(const FText& InText, FName& OutTableId, FString& OutKey) const;
+	CORE_API bool FindTableIdAndKey(const FText& InText, FName& OutTableId, FString& OutKey) const;
 
 	/** Log a missing string table entry (will only log each missing entry once to avoid spam) */
-	void LogMissingStringTableEntry(const FName InTableId, const FTextKey& InKey);
+	CORE_API void LogMissingStringTableEntry(const FName InTableId, const FTextKey& InKey);
 
 	/** Internal function called by LOCTABLE_NEW to create and register a new FStringTable instance */
-	void Internal_NewLocTable(const FName InTableId, const FString& InNamespace);
+	CORE_API void Internal_NewLocTable(const FName InTableId, const FString& InNamespace);
 
 	/** Internal function called by LOCTABLE_FROMFILE_X to create and register a FStringTable instance that has been populated from a file */
-	void Internal_LocTableFromFile(const FName InTableId, const FString& InNamespace, const FString& InFilePath, const FString& InRootPath);
+	CORE_API void Internal_LocTableFromFile(const FName InTableId, const FString& InNamespace, const FString& InFilePath, const FString& InRootPath);
 
 	/** Internal function called by LOCTABLE_SETSTRING to set the entry denoted by the given key to the given source string, within the given string table (table must have been registered already) */
-	void Internal_SetLocTableEntry(const FName InTableId, const FString& InKey, const FString& InSourceString);
+	CORE_API void Internal_SetLocTableEntry(const FName InTableId, const FString& InKey, const FString& InSourceString);
 
 	/** Internal function called by LOCTABLE_SETMETA to set meta-data for the entry denoted by the given key, within the given string table (table must have been registered already) */
-	void Internal_SetLocTableEntryMetaData(const FName InTableId, const FString& InKey, const FName InMetaDataId, const FString& InMetaData);
+	CORE_API void Internal_SetLocTableEntryMetaData(const FName InTableId, const FString& InKey, const FName InMetaDataId, const FString& InMetaData);
 
 	/** Internal function called by LOCTABLE to find the entry with by the given key within the given string table (redirects, will load assets if needed, and returns a dummy FText if not found) */
-	FText Internal_FindLocTableEntry(const FName InTableId, const FString& InKey, const EStringTableLoadingPolicy InLoadingPolicy) const;
+	CORE_API FText Internal_FindLocTableEntry(const FName InTableId, const FString& InKey, const EStringTableLoadingPolicy InLoadingPolicy) const;
 
 private:
 	/** Private constructor - use singleton accessor */
