[Core]

diff --git a/Engine/Source/Runtime/Core/Public/Internationalization/Text.h b/Engine/Source/Runtime/Core/Public/Internationalization/Text.h
--- a/Engine/Source/Runtime/Core/Public/Internationalization/Text.h
+++ b/Engine/Source/Runtime/Core/Public/Internationalization/Text.h
@@ -175,9 +175,9 @@ enum EMemoryUnitStandard
 	SI
 };
 
-struct CORE_API FNumberFormattingOptions
+struct FNumberFormattingOptions
 {
-	FNumberFormattingOptions();
+	CORE_API FNumberFormattingOptions();
 
 	bool AlwaysSign;
 	FNumberFormattingOptions& SetAlwaysSign( bool InValue ){ AlwaysSign = InValue; return *this; }
@@ -206,18 +206,18 @@ struct CORE_API FNumberFormattingOptions
 	friend CORE_API uint32 GetTypeHash( const FNumberFormattingOptions& Key );
 
 	/** Check to see if our formatting options match the other formatting options */
-	bool IsIdentical( const FNumberFormattingOptions& Other ) const;
+	CORE_API bool IsIdentical( const FNumberFormattingOptions& Other ) const;
 
 	/** Get the default number formatting options with grouping enabled */
-	static const FNumberFormattingOptions& DefaultWithGrouping();
+	static CORE_API const FNumberFormattingOptions& DefaultWithGrouping();
 
 	/** Get the default number formatting options with grouping disabled */
-	static const FNumberFormattingOptions& DefaultNoGrouping();
+	static CORE_API const FNumberFormattingOptions& DefaultNoGrouping();
 };
 
-struct CORE_API FNumberParsingOptions
+struct FNumberParsingOptions
 {
-	FNumberParsingOptions();
+	CORE_API FNumberParsingOptions();
 
 	bool UseGrouping;
 	FNumberParsingOptions& SetUseGrouping( bool InValue ){ UseGrouping = InValue; return *this; }
@@ -236,13 +236,13 @@ struct CORE_API FNumberParsingOptions
 	friend CORE_API uint32 GetTypeHash( const FNumberParsingOptions& Key );
 
 	/** Check to see if our parsing options match the other parsing options */
-	bool IsIdentical( const FNumberParsingOptions& Other ) const;
+	CORE_API bool IsIdentical( const FNumberParsingOptions& Other ) const;
 
 	/** Get the default number parsing options with grouping enabled */
-	static const FNumberParsingOptions& DefaultWithGrouping();
+	static CORE_API const FNumberParsingOptions& DefaultWithGrouping();
 
 	/** Get the default number parsing options with grouping disabled */
-	static const FNumberParsingOptions& DefaultNoGrouping();
+	static CORE_API const FNumberParsingOptions& DefaultNoGrouping();
 };
 
 /**
@@ -251,7 +251,7 @@ struct CORE_API FNumberParsingOptions
  * and is safe to be used as a function-level static.
  * See TextFormatter.cpp for the definition.
  */
-class CORE_API FTextFormat
+class FTextFormat
 {
 	friend class FTextFormatter;
 
@@ -269,83 +269,83 @@ public:
 	/**
 	 * Construct an instance using an empty FText.
 	 */
-	FTextFormat();
+	CORE_API FTextFormat();
 
 	/**
 	 * Construct an instance from an FText.
 	 * The text will be immediately compiled. 
 	 */
-	FTextFormat(const FText& InText);
+	CORE_API FTextFormat(const FText& InText);
 
 	/**
 	 * Construct an instance from an FText and custom format pattern definition.
 	 * The text will be immediately compiled.
 	 */
-	FTextFormat(const FText& InText, FTextFormatPatternDefinitionConstRef InCustomPatternDef);
+	CORE_API FTextFormat(const FText& InText, FTextFormatPatternDefinitionConstRef InCustomPatternDef);
 
 	/**
 	 * Construct an instance from an FString.
 	 * The string will be immediately compiled.
 	 */
-	static FTextFormat FromString(const FString& InString);
-	static FTextFormat FromString(FString&& InString);
+	static CORE_API FTextFormat FromString(const FString& InString);
+	static CORE_API FTextFormat FromString(FString&& InString);
 
 	/**
 	 * Construct an instance from an FString and custom format pattern definition.
 	 * The string will be immediately compiled.
 	 */
-	static FTextFormat FromString(const FString& InString, FTextFormatPatternDefinitionConstRef InCustomPatternDef);
-	static FTextFormat FromString(FString&& InString, FTextFormatPatternDefinitionConstRef InCustomPatternDef);
+	static CORE_API FTextFormat FromString(const FString& InString, FTextFormatPatternDefinitionConstRef InCustomPatternDef);
+	static CORE_API FTextFormat FromString(FString&& InString, FTextFormatPatternDefinitionConstRef InCustomPatternDef);
 
 	/**
 	 * Test to see whether this instance contains valid compiled data.
 	 */
-	bool IsValid() const;
+	CORE_API bool IsValid() const;
 
 	/**
 	 * Check whether this instance is considered identical to the other instance, based on the comparison flags provided.
 	 */
-	bool IdenticalTo(const FTextFormat& Other, const ETextIdenticalModeFlags CompareModeFlags) const;
+	CORE_API bool IdenticalTo(const FTextFormat& Other, const ETextIdenticalModeFlags CompareModeFlags) const;
 
 	/**
 	 * Get the source text that we're holding.
 	 * If we're holding a string then we'll construct a new text.
 	 */
-	FText GetSourceText() const;
+	CORE_API FText GetSourceText() const;
 
 	/**
 	 * Get the source string that we're holding.
 	 * If we're holding a text then we'll return its internal string.
 	 */
-	const FString& GetSourceString() const;
+	CORE_API const FString& GetSourceString() const;
 
 	/**
 	 * Get the type of expression currently compiled.
 	 */
-	EExpressionType GetExpressionType() const;
+	CORE_API EExpressionType GetExpressionType() const;
 
 	/**
 	 * Get the format pattern definition being used.
 	 */
-	FTextFormatPatternDefinitionConstRef GetPatternDefinition() const;
+	CORE_API FTextFormatPatternDefinitionConstRef GetPatternDefinition() const;
 
 	/**
 	 * Validate the format pattern is valid based on the rules of the given culture (or null to use the current language).
 	 * @return true if the pattern is valid, or false if not (false may also fill in OutValidationErrors).
 	 */
-	bool ValidatePattern(const FCulturePtr& InCulture, TArray<FString>& OutValidationErrors) const;
+	CORE_API bool ValidatePattern(const FCulturePtr& InCulture, TArray<FString>& OutValidationErrors) const;
 
 	/**
 	 * Append the names of any arguments to the given array.
 	 */
-	void GetFormatArgumentNames(TArray<FString>& OutArgumentNames) const;
+	CORE_API void GetFormatArgumentNames(TArray<FString>& OutArgumentNames) const;
 
 private:
 	/**
 	 * Construct an instance from an FString.
 	 * The string will be immediately compiled.
 	 */
-	FTextFormat(FString&& InString, FTextFormatPatternDefinitionConstRef InCustomPatternDef);
+	CORE_API FTextFormat(FString&& InString, FTextFormatPatternDefinitionConstRef InCustomPatternDef);
 
 	/** Cached compiled expression data */
 	TSharedRef<FTextFormatData, ESPMode::ThreadSafe> TextFormatData;
@@ -353,15 +353,15 @@ private:
 
 class FCulture;
 
-class CORE_API FText
+class FText
 {
 public:
 
-	static const FText& GetEmpty();
+	static CORE_API const FText& GetEmpty();
 
 public:
 
-	FText();
+	CORE_API FText();
 	FText(const FText&) = default;
 	FText(FText&&) = default;
 
@@ -371,30 +371,30 @@ public:
 	/**
 	 * Generate an FText that represents the passed number in the current culture
 	 */
-	static FText AsNumber(float Val,	const FNumberFormattingOptions* const Options = NULL, const FCulturePtr& TargetCulture = NULL);
-	static FText AsNumber(double Val,	const FNumberFormattingOptions* const Options = NULL, const FCulturePtr& TargetCulture = NULL);
-	static FText AsNumber(int8 Val,		const FNumberFormattingOptions* const Options = NULL, const FCulturePtr& TargetCulture = NULL);
-	static FText AsNumber(int16 Val,	const FNumberFormattingOptions* const Options = NULL, const FCulturePtr& TargetCulture = NULL);
-	static FText AsNumber(int32 Val,	const FNumberFormattingOptions* const Options = NULL, const FCulturePtr& TargetCulture = NULL);
-	static FText AsNumber(int64 Val,	const FNumberFormattingOptions* const Options = NULL, const FCulturePtr& TargetCulture = NULL);
-	static FText AsNumber(uint8 Val,	const FNumberFormattingOptions* const Options = NULL, const FCulturePtr& TargetCulture = NULL);
-	static FText AsNumber(uint16 Val,	const FNumberFormattingOptions* const Options = NULL, const FCulturePtr& TargetCulture = NULL);
-	static FText AsNumber(uint32 Val,	const FNumberFormattingOptions* const Options = NULL, const FCulturePtr& TargetCulture = NULL);
-	static FText AsNumber(uint64 Val,	const FNumberFormattingOptions* const Options = NULL, const FCulturePtr& TargetCulture = NULL);
-	static FText AsNumber(long Val,		const FNumberFormattingOptions* const Options = NULL, const FCulturePtr& TargetCulture = NULL);
-
-
-	static FText AsCurrency(float Val,  const FString& CurrencyCode = FString(), const FNumberFormattingOptions* const Options = NULL, const FCulturePtr& TargetCulture = NULL);
-	static FText AsCurrency(double Val, const FString& CurrencyCode = FString(), const FNumberFormattingOptions* const Options = NULL, const FCulturePtr& TargetCulture = NULL);
-	static FText AsCurrency(int8 Val,   const FString& CurrencyCode = FString(), const FNumberFormattingOptions* const Options = NULL, const FCulturePtr& TargetCulture = NULL);
-	static FText AsCurrency(int16 Val,  const FString& CurrencyCode = FString(), const FNumberFormattingOptions* const Options = NULL, const FCulturePtr& TargetCulture = NULL);
-	static FText AsCurrency(int32 Val,  const FString& CurrencyCode = FString(), const FNumberFormattingOptions* const Options = NULL, const FCulturePtr& TargetCulture = NULL);
-	static FText AsCurrency(int64 Val,  const FString& CurrencyCode = FString(), const FNumberFormattingOptions* const Options = NULL, const FCulturePtr& TargetCulture = NULL);
-	static FText AsCurrency(uint8 Val,  const FString& CurrencyCode = FString(), const FNumberFormattingOptions* const Options = NULL, const FCulturePtr& TargetCulture = NULL);
-	static FText AsCurrency(uint16 Val, const FString& CurrencyCode = FString(), const FNumberFormattingOptions* const Options = NULL, const FCulturePtr& TargetCulture = NULL);
-	static FText AsCurrency(uint32 Val, const FString& CurrencyCode = FString(), const FNumberFormattingOptions* const Options = NULL, const FCulturePtr& TargetCulture = NULL);
-	static FText AsCurrency(uint64 Val, const FString& CurrencyCode = FString(), const FNumberFormattingOptions* const Options = NULL, const FCulturePtr& TargetCulture = NULL);
-	static FText AsCurrency(long Val,   const FString& CurrencyCode = FString(), const FNumberFormattingOptions* const Options = NULL, const FCulturePtr& TargetCulture = NULL);
+	static CORE_API FText AsNumber(float Val,	const FNumberFormattingOptions* const Options = NULL, const FCulturePtr& TargetCulture = NULL);
+	static CORE_API FText AsNumber(double Val,	const FNumberFormattingOptions* const Options = NULL, const FCulturePtr& TargetCulture = NULL);
+	static CORE_API FText AsNumber(int8 Val,		const FNumberFormattingOptions* const Options = NULL, const FCulturePtr& TargetCulture = NULL);
+	static CORE_API FText AsNumber(int16 Val,	const FNumberFormattingOptions* const Options = NULL, const FCulturePtr& TargetCulture = NULL);
+	static CORE_API FText AsNumber(int32 Val,	const FNumberFormattingOptions* const Options = NULL, const FCulturePtr& TargetCulture = NULL);
+	static CORE_API FText AsNumber(int64 Val,	const FNumberFormattingOptions* const Options = NULL, const FCulturePtr& TargetCulture = NULL);
+	static CORE_API FText AsNumber(uint8 Val,	const FNumberFormattingOptions* const Options = NULL, const FCulturePtr& TargetCulture = NULL);
+	static CORE_API FText AsNumber(uint16 Val,	const FNumberFormattingOptions* const Options = NULL, const FCulturePtr& TargetCulture = NULL);
+	static CORE_API FText AsNumber(uint32 Val,	const FNumberFormattingOptions* const Options = NULL, const FCulturePtr& TargetCulture = NULL);
+	static CORE_API FText AsNumber(uint64 Val,	const FNumberFormattingOptions* const Options = NULL, const FCulturePtr& TargetCulture = NULL);
+	static CORE_API FText AsNumber(long Val,		const FNumberFormattingOptions* const Options = NULL, const FCulturePtr& TargetCulture = NULL);
+
+
+	static CORE_API FText AsCurrency(float Val,  const FString& CurrencyCode = FString(), const FNumberFormattingOptions* const Options = NULL, const FCulturePtr& TargetCulture = NULL);
+	static CORE_API FText AsCurrency(double Val, const FString& CurrencyCode = FString(), const FNumberFormattingOptions* const Options = NULL, const FCulturePtr& TargetCulture = NULL);
+	static CORE_API FText AsCurrency(int8 Val,   const FString& CurrencyCode = FString(), const FNumberFormattingOptions* const Options = NULL, const FCulturePtr& TargetCulture = NULL);
+	static CORE_API FText AsCurrency(int16 Val,  const FString& CurrencyCode = FString(), const FNumberFormattingOptions* const Options = NULL, const FCulturePtr& TargetCulture = NULL);
+	static CORE_API FText AsCurrency(int32 Val,  const FString& CurrencyCode = FString(), const FNumberFormattingOptions* const Options = NULL, const FCulturePtr& TargetCulture = NULL);
+	static CORE_API FText AsCurrency(int64 Val,  const FString& CurrencyCode = FString(), const FNumberFormattingOptions* const Options = NULL, const FCulturePtr& TargetCulture = NULL);
+	static CORE_API FText AsCurrency(uint8 Val,  const FString& CurrencyCode = FString(), const FNumberFormattingOptions* const Options = NULL, const FCulturePtr& TargetCulture = NULL);
+	static CORE_API FText AsCurrency(uint16 Val, const FString& CurrencyCode = FString(), const FNumberFormattingOptions* const Options = NULL, const FCulturePtr& TargetCulture = NULL);
+	static CORE_API FText AsCurrency(uint32 Val, const FString& CurrencyCode = FString(), const FNumberFormattingOptions* const Options = NULL, const FCulturePtr& TargetCulture = NULL);
+	static CORE_API FText AsCurrency(uint64 Val, const FString& CurrencyCode = FString(), const FNumberFormattingOptions* const Options = NULL, const FCulturePtr& TargetCulture = NULL);
+	static CORE_API FText AsCurrency(long Val,   const FString& CurrencyCode = FString(), const FNumberFormattingOptions* const Options = NULL, const FCulturePtr& TargetCulture = NULL);
 
 	/**
 	 * Generate an FText that represents the passed number as currency in the current culture.
@@ -402,90 +402,90 @@ public:
 	 * Keep in mind the CurrencyCode is completely independent of the culture it's displayed in (and they do not imply one another).
 	 * For example: FText::AsCurrencyBase(650, TEXT("EUR")); would return an FText of "<EUR>6.50" in most English cultures (en_US/en_UK) and "6,50<EUR>" in Spanish (es_ES) (where <EUR> is U+20AC)
 	 */
-	static FText AsCurrencyBase(int64 BaseVal, const FString& CurrencyCode, const FCulturePtr& TargetCulture = NULL, int32 ForceDecimalPlaces = -1);
+	static CORE_API FText AsCurrencyBase(int64 BaseVal, const FString& CurrencyCode, const FCulturePtr& TargetCulture = NULL, int32 ForceDecimalPlaces = -1);
 
 	/**
 	 * Generate an FText that represents the passed number as a percentage in the current culture
 	 */
-	static FText AsPercent(float Val,	const FNumberFormattingOptions* const Options = NULL, const FCulturePtr& TargetCulture = NULL);
-	static FText AsPercent(double Val,	const FNumberFormattingOptions* const Options = NULL, const FCulturePtr& TargetCulture = NULL);
+	static CORE_API FText AsPercent(float Val,	const FNumberFormattingOptions* const Options = NULL, const FCulturePtr& TargetCulture = NULL);
+	static CORE_API FText AsPercent(double Val,	const FNumberFormattingOptions* const Options = NULL, const FCulturePtr& TargetCulture = NULL);
 
 	/**
 	 * Generate an FText that represents the passed number as a date and/or time in the current culture
 	 * @note The overload using a custom pattern uses strftime-like syntax (see FDateTime::ToFormattedString)
 	 */
-	static FText AsDate(const FDateTime& DateTime, const EDateTimeStyle::Type DateStyle = EDateTimeStyle::Default, const FString& TimeZone = FString(), const FCulturePtr& TargetCulture = NULL);
-	static FText AsDateTime(const FDateTime& DateTime, const EDateTimeStyle::Type DateStyle = EDateTimeStyle::Default, const EDateTimeStyle::Type TimeStyle = EDateTimeStyle::Default, const FString& TimeZone = FString(), const FCulturePtr& TargetCulture = NULL);
-	static FText AsDateTime(const FDateTime& DateTime, const FString& CustomPattern, const FString& TimeZone = FString(), const FCulturePtr& TargetCulture = NULL);
-	static FText AsTime(const FDateTime& DateTime, const EDateTimeStyle::Type TimeStyle = EDateTimeStyle::Default, const FString& TimeZone = FString(), const FCulturePtr& TargetCulture = NULL);
-	static FText AsTimespan(const FTimespan& Timespan, const FCulturePtr& TargetCulture = NULL);
+	static CORE_API FText AsDate(const FDateTime& DateTime, const EDateTimeStyle::Type DateStyle = EDateTimeStyle::Default, const FString& TimeZone = FString(), const FCulturePtr& TargetCulture = NULL);
+	static CORE_API FText AsDateTime(const FDateTime& DateTime, const EDateTimeStyle::Type DateStyle = EDateTimeStyle::Default, const EDateTimeStyle::Type TimeStyle = EDateTimeStyle::Default, const FString& TimeZone = FString(), const FCulturePtr& TargetCulture = NULL);
+	static CORE_API FText AsDateTime(const FDateTime& DateTime, const FString& CustomPattern, const FString& TimeZone = FString(), const FCulturePtr& TargetCulture = NULL);
+	static CORE_API FText AsTime(const FDateTime& DateTime, const EDateTimeStyle::Type TimeStyle = EDateTimeStyle::Default, const FString& TimeZone = FString(), const FCulturePtr& TargetCulture = NULL);
+	static CORE_API FText AsTimespan(const FTimespan& Timespan, const FCulturePtr& TargetCulture = NULL);
 
 	/**
 	 * Gets the time zone string that represents a non-specific, zero offset, culture invariant time zone.
 	 */
-	static FString GetInvariantTimeZone();
+	static CORE_API FString GetInvariantTimeZone();
 
 	/**
 	 * Generate an FText that represents the passed number as a memory size in the current culture
 	 */
-	static FText AsMemory(uint64 NumBytes, const FNumberFormattingOptions* const Options = NULL, const FCulturePtr& TargetCulture = NULL, EMemoryUnitStandard UnitStandard = EMemoryUnitStandard::IEC);
+	static CORE_API FText AsMemory(uint64 NumBytes, const FNumberFormattingOptions* const Options = NULL, const FCulturePtr& TargetCulture = NULL, EMemoryUnitStandard UnitStandard = EMemoryUnitStandard::IEC);
 
 	/**
 	 * Generate an FText that represents the passed number as a memory size in the current culture
 	 */
-	static FText AsMemory(uint64 NumBytes, EMemoryUnitStandard UnitStandard);
+	static CORE_API FText AsMemory(uint64 NumBytes, EMemoryUnitStandard UnitStandard);
 
 	/**
 	 * Attempts to find an existing FText using the representation found in the loc tables for the specified namespace and key
 	 * @return true if OutText was properly set; otherwise false and OutText will be untouched
 	 */
-	static bool FindText( const FTextKey& Namespace, const FTextKey& Key, FText& OutText, const FString* const SourceString = nullptr );
+	static CORE_API bool FindText( const FTextKey& Namespace, const FTextKey& Key, FText& OutText, const FString* const SourceString = nullptr );
 
 	/**
 	 * Attempts to create an FText instance from a string table ID and key (this is the same as the LOCTABLE macro, except this can also work with non-literal string values).
 	 * @return The found text, or a dummy FText if not found.
 	 */
-	static FText FromStringTable(const FName InTableId, const FString& InKey, const EStringTableLoadingPolicy InLoadingPolicy = EStringTableLoadingPolicy::FindOrLoad);
+	static CORE_API FText FromStringTable(const FName InTableId, const FString& InKey, const EStringTableLoadingPolicy InLoadingPolicy = EStringTableLoadingPolicy::FindOrLoad);
 
 	/**
 	 * Generate an FText representing the pass name
 	 */
-	static FText FromName( const FName& Val);
+	static CORE_API FText FromName( const FName& Val);
 	
 	/**
 	 * Generate an FText representing the passed in string
 	 */
-	static FText FromString( const FString& String );
-	static FText FromString( FString&& String );
+	static CORE_API FText FromString( const FString& String );
+	static CORE_API FText FromString( FString&& String );
 
 	/**
 	 * Generate a FText representing the passed string view
 	 */
-	static FText FromStringView(FStringView InString);
+	static CORE_API FText FromStringView(FStringView InString);
 
 	/**
 	 * Generate a culture invariant FText representing the passed in string
 	 */
-	static FText AsCultureInvariant( const FString& String );
-	static FText AsCultureInvariant( FString&& String );
+	static CORE_API FText AsCultureInvariant( const FString& String );
+	static CORE_API FText AsCultureInvariant( FString&& String );
 
 	/**
 	 * Generate a culture invariant FText representing the passed in FText
 	 */
-	static FText AsCultureInvariant( FText Text );
+	static CORE_API FText AsCultureInvariant( FText Text );
 
-	const FString& ToString() const;
+	CORE_API const FString& ToString() const;
 
 	/** Deep build of the source string for this FText, climbing the history hierarchy */
-	FString BuildSourceString() const;
+	CORE_API FString BuildSourceString() const;
 
-	bool IsNumeric() const;
+	CORE_API bool IsNumeric() const;
 
-	int32 CompareTo( const FText& Other, const ETextComparisonLevel::Type ComparisonLevel = ETextComparisonLevel::Default ) const;
-	int32 CompareToCaseIgnored( const FText& Other ) const;
+	CORE_API int32 CompareTo( const FText& Other, const ETextComparisonLevel::Type ComparisonLevel = ETextComparisonLevel::Default ) const;
+	CORE_API int32 CompareToCaseIgnored( const FText& Other ) const;
 
-	bool EqualTo( const FText& Other, const ETextComparisonLevel::Type ComparisonLevel = ETextComparisonLevel::Default ) const;
-	bool EqualToCaseIgnored( const FText& Other ) const;
+	CORE_API bool EqualTo( const FText& Other, const ETextComparisonLevel::Type ComparisonLevel = ETextComparisonLevel::Default ) const;
+	CORE_API bool EqualToCaseIgnored( const FText& Other ) const;
 
 	/**
 	 * Check to see if this FText is identical to the other FText
@@ -495,14 +495,14 @@ public:
 	 *
 	 * @note If you actually want to perform a full lexical comparison, then you need to use EqualTo instead.
 	 */
-	bool IdenticalTo( const FText& Other, const ETextIdenticalModeFlags CompareModeFlags = ETextIdenticalModeFlags::None ) const;
+	CORE_API bool IdenticalTo( const FText& Other, const ETextIdenticalModeFlags CompareModeFlags = ETextIdenticalModeFlags::None ) const;
 
-	class CORE_API FSortPredicate
+	class FSortPredicate
 	{
 	public:
-		FSortPredicate(const ETextComparisonLevel::Type ComparisonLevel = ETextComparisonLevel::Default);
+		CORE_API FSortPredicate(const ETextComparisonLevel::Type ComparisonLevel = ETextComparisonLevel::Default);
 
-		bool operator()(const FText& A, const FText& B) const;
+		CORE_API bool operator()(const FText& A, const FText& B) const;
 
 	private:
 #if UE_ENABLE_ICU
@@ -511,43 +511,43 @@ public:
 #endif
 	};
 
-	bool IsEmpty() const;
+	CORE_API bool IsEmpty() const;
 
-	bool IsEmptyOrWhitespace() const;
+	CORE_API bool IsEmptyOrWhitespace() const;
 
 	/**
 	 * Transforms the text to lowercase in a culture correct way.
 	 * @note The returned instance is linked to the original and will be rebuilt if the active culture is changed.
 	 */
-	FText ToLower() const;
+	CORE_API FText ToLower() const;
 
 	/**
 	 * Transforms the text to uppercase in a culture correct way.
 	 * @note The returned instance is linked to the original and will be rebuilt if the active culture is changed.
 	 */
-	FText ToUpper() const;
+	CORE_API FText ToUpper() const;
 
 	/**
 	 * Removes any whitespace characters from the start of the text.
 	 */
-	static FText TrimPreceding( const FText& );
+	static CORE_API FText TrimPreceding( const FText& );
 
 	/**
 	 * Removes any whitespace characters from the end of the text.
 	 */
-	static FText TrimTrailing( const FText& );
+	static CORE_API FText TrimTrailing( const FText& );
 
 	/**
 	 * Removes any whitespace characters from the start and end of the text.
 	 */
-	static FText TrimPrecedingAndTrailing( const FText& );
+	static CORE_API FText TrimPrecedingAndTrailing( const FText& );
 
 	/**
 	 * Check to see if the given character is considered whitespace by the current culture
 	 */
-	static bool IsWhitespace( const TCHAR Char );
+	static CORE_API bool IsWhitespace( const TCHAR Char );
 
-	static void GetFormatPatternParameters(const FTextFormat& Fmt, TArray<FString>& ParameterNames);
+	static CORE_API void GetFormatPatternParameters(const FTextFormat& Fmt, TArray<FString>& ParameterNames);
 
 	/**
 	 * Format the given map of key->value pairs as named arguments within the given format pattern
@@ -559,8 +559,8 @@ public:
 	 * @param InArguments The map of key->value pairs to inject into the format pattern
 	 * @return The formatted FText
 	 */
-	static FText Format(FTextFormat Fmt, const FFormatNamedArguments& InArguments);
-	static FText Format(FTextFormat Fmt, FFormatNamedArguments&& InArguments);
+	static CORE_API FText Format(FTextFormat Fmt, const FFormatNamedArguments& InArguments);
+	static CORE_API FText Format(FTextFormat Fmt, FFormatNamedArguments&& InArguments);
 
 	/**
 	 * Format the given list values as ordered arguments within the given format pattern
@@ -572,8 +572,8 @@ public:
 	 * @param InArguments The list of values to inject into the format pattern
 	 * @return The formatted FText
 	 */
-	static FText Format(FTextFormat Fmt, const FFormatOrderedArguments& InArguments);
-	static FText Format(FTextFormat Fmt, FFormatOrderedArguments&& InArguments);
+	static CORE_API FText Format(FTextFormat Fmt, const FFormatOrderedArguments& InArguments);
+	static CORE_API FText Format(FTextFormat Fmt, FFormatOrderedArguments&& InArguments);
 
 	/**
 	 * Format the given list of variadic values as ordered arguments within the given format pattern
@@ -638,8 +638,8 @@ public:
 	 * @param Args An array of formattable values to join together
 	 * @return The joined FText
 	 */
-	static FText Join(const FText& Delimiter, const FFormatOrderedArguments& Args);
-	static FText Join(const FText& Delimiter, const TArray<FText>& Args);
+	static CORE_API FText Join(const FText& Delimiter, const FFormatOrderedArguments& Args);
+	static CORE_API FText Join(const FText& Delimiter, const TArray<FText>& Args);
 
 	/**
 	 * Join an arbitrary list of formattable items together, separated by the given delimiter
@@ -664,7 +664,7 @@ public:
 	 *
 	 * @param TextGenerator the text generator object that will generate the text
 	 */
-	static FText FromTextGenerator( const TSharedRef<ITextGenerator>& TextGenerator );
+	static CORE_API FText FromTextGenerator( const TSharedRef<ITextGenerator>& TextGenerator );
 
 	DECLARE_DELEGATE_RetVal_OneParam( TSharedRef<ITextGenerator>, FCreateTextGeneratorDelegate, FStructuredArchive::FRecord );
 	/**
@@ -672,7 +672,7 @@ public:
 	 *
 	 * @param TypeID the name under which to look up the factory function
 	 */
-	static FCreateTextGeneratorDelegate FindRegisteredTextGenerator( FName TypeID );
+	static CORE_API FCreateTextGeneratorDelegate FindRegisteredTextGenerator( FName TypeID );
 
 	/**
 	 * Registers a factory function to be used with serialization of text generators within FText.
@@ -680,7 +680,7 @@ public:
 	 * @param TypeID the name under which to register the factory function. Must match ITextGenerator::GetTypeID().
 	 * @param FactoryFunction the factory function to create the generator instance
 	 */
-	static void RegisterTextGenerator( FName TypeID, FCreateTextGeneratorDelegate FactoryFunction );
+	static CORE_API void RegisterTextGenerator( FName TypeID, FCreateTextGeneratorDelegate FactoryFunction );
 
 	/**
 	 * Registers a standard text generator factory function.
@@ -715,7 +715,7 @@ public:
 	 *
 	 * @see RegisterTextGenerator
 	 */
-	static void UnregisterTextGenerator( FName TypeID );
+	static CORE_API void UnregisterTextGenerator( FName TypeID );
 
 	/**
 	 * Unregisters a standard text generator factory function.
@@ -730,47 +730,47 @@ public:
 		UnregisterTextGenerator( T::TypeID );
 	}
 
-	bool IsTransient() const;
-	bool IsCultureInvariant() const;
-	bool IsInitializedFromString() const;
-	bool IsFromStringTable() const;
+	CORE_API bool IsTransient() const;
+	CORE_API bool IsCultureInvariant() const;
+	CORE_API bool IsInitializedFromString() const;
+	CORE_API bool IsFromStringTable() const;
 
-	bool ShouldGatherForLocalization() const;
+	CORE_API bool ShouldGatherForLocalization() const;
 
 #if WITH_EDITORONLY_DATA
 	/**
 	 * Constructs a new FText with the SourceString of the specified text but with the specified namespace and key
 	 */
-	static FText ChangeKey( const FTextKey& Namespace, const FTextKey& Key, const FText& Text );
+	static CORE_API FText ChangeKey( const FTextKey& Namespace, const FTextKey& Key, const FText& Text );
 #endif
 
 private:
 	/** Special constructor used to create StaticEmptyText without also allocating a history object */
 	enum class EInitToEmptyString : uint8 { Value };
-	explicit FText( EInitToEmptyString );
+	CORE_API explicit FText( EInitToEmptyString );
 
-	explicit FText( TSharedRef<ITextData, ESPMode::ThreadSafe> InTextData );
+	CORE_API explicit FText( TSharedRef<ITextData, ESPMode::ThreadSafe> InTextData );
 
-	explicit FText( FString&& InSourceString );
+	CORE_API explicit FText( FString&& InSourceString );
 
-	FText( FName InTableId, FString InKey, const EStringTableLoadingPolicy InLoadingPolicy );
+	CORE_API FText( FName InTableId, FString InKey, const EStringTableLoadingPolicy InLoadingPolicy );
 
-	FText( FString&& InSourceString, const FTextKey& InNamespace, const FTextKey& InKey, uint32 InFlags=0 );
+	CORE_API FText( FString&& InSourceString, const FTextKey& InNamespace, const FTextKey& InKey, uint32 InFlags=0 );
 
-	static void SerializeText( FArchive& Ar, FText& Value );
-	static void SerializeText(FStructuredArchive::FSlot Slot, FText& Value);
+	static CORE_API void SerializeText( FArchive& Ar, FText& Value );
+	static CORE_API void SerializeText(FStructuredArchive::FSlot Slot, FText& Value);
 
 	/** Get any historic text format data from the history used by this FText */
-	void GetHistoricFormatData(TArray<FHistoricTextFormatData>& OutHistoricFormatData) const;
+	CORE_API void GetHistoricFormatData(TArray<FHistoricTextFormatData>& OutHistoricFormatData) const;
 
 	/** Get any historic numeric format data from the history used by this FText */
-	bool GetHistoricNumericData(FHistoricTextNumericData& OutHistoricNumericData) const;
+	CORE_API bool GetHistoricNumericData(FHistoricTextNumericData& OutHistoricNumericData) const;
 
 	/** Rebuilds the FText under the current culture if needed */
-	void Rebuild() const;
+	CORE_API void Rebuild() const;
 
-	static FText FormatNamedImpl(FTextFormat&& Fmt, FFormatNamedArguments&& InArguments);
-	static FText FormatOrderedImpl(FTextFormat&& Fmt, FFormatOrderedArguments&& InArguments);
+	static CORE_API FText FormatNamedImpl(FTextFormat&& Fmt, FFormatNamedArguments&& InArguments);
+	static CORE_API FText FormatOrderedImpl(FTextFormat&& Fmt, FFormatOrderedArguments&& InArguments);
 
 private:
 	template<typename T1, typename T2>
@@ -812,7 +812,7 @@ public:
 	friend class FScopedTextIdentityPreserver;
 };
 
-class CORE_API FFormatArgumentValue
+class FFormatArgumentValue
 {
 public:
 	FFormatArgumentValue()
@@ -821,7 +821,7 @@ public:
 	{
 	}
 
-	FFormatArgumentValue(const class FCbValue& Value);
+	CORE_API FFormatArgumentValue(const class FCbValue& Value);
 
 	FFormatArgumentValue(const int32 Value)
 		: Type(EFormatArgumentType::Int)
@@ -879,14 +879,14 @@ public:
 
 	friend void operator<<(FStructuredArchive::FSlot Slot, FFormatArgumentValue& Value);
 
-	bool IdenticalTo(const FFormatArgumentValue& Other, const ETextIdenticalModeFlags CompareModeFlags) const;
+	CORE_API bool IdenticalTo(const FFormatArgumentValue& Other, const ETextIdenticalModeFlags CompareModeFlags) const;
 
-	FString ToFormattedString(const bool bInRebuildText, const bool bInRebuildAsSource) const;
-	void ToFormattedString(const bool bInRebuildText, const bool bInRebuildAsSource, FString& OutResult) const;
+	CORE_API FString ToFormattedString(const bool bInRebuildText, const bool bInRebuildAsSource) const;
+	CORE_API void ToFormattedString(const bool bInRebuildText, const bool bInRebuildAsSource, FString& OutResult) const;
 
-	FString ToExportedString(const bool bStripPackageNamespace = false) const;
-	void ToExportedString(FString& OutResult, const bool bStripPackageNamespace = false) const;
-	const TCHAR* FromExportedString(const TCHAR* InBuffer);
+	CORE_API FString ToExportedString(const bool bStripPackageNamespace = false) const;
+	CORE_API void ToExportedString(FString& OutResult, const bool bStripPackageNamespace = false) const;
+	CORE_API const TCHAR* FromExportedString(const TCHAR* InBuffer);
 
 	FORCEINLINE EFormatArgumentType::Type GetType() const
 	{
@@ -952,16 +952,16 @@ inline TSharedRef<ITextGenerator> FText::CreateTextGenerator(FStructuredArchive:
  * @note The primary consumer of this type is Blueprints (via a UHT mirror node). It is *not* expected that this be used in general C++ as FFormatArgumentValue is a much better type.
  * The UHT struct is located here: Engine\Source\Runtime\Engine\Classes\Kismet\KismetTextLibrary.h
  */
-struct CORE_API FFormatArgumentData
+struct FFormatArgumentData
 {
 	FFormatArgumentData()
 	{
 		ResetValue();
 	}
 
-	void ResetValue();
+	CORE_API void ResetValue();
 
-	FFormatArgumentValue ToArgumentValue() const;
+	CORE_API FFormatArgumentValue ToArgumentValue() const;
 
 	friend void operator<<(FStructuredArchive::FSlot Slot, FFormatArgumentData& Value);
 
@@ -1029,7 +1029,7 @@ FText FText::FormatOrdered( FTextFormat Fmt, TArguments&&... Args )
 }
 
 /** Used to gather information about a historic text format operation */
-class CORE_API FHistoricTextFormatData
+class FHistoricTextFormatData
 {
 public:
 	FHistoricTextFormatData()
@@ -1054,7 +1054,7 @@ public:
 };
 
 /** Used to gather information about a historic numeric format operation */
-class CORE_API FHistoricTextNumericData
+class FHistoricTextNumericData
 {
 public:
 	enum class EType : uint8
@@ -1086,18 +1086,18 @@ public:
 };
 
 /** A snapshot of an FText at a point in time that can be used to detect changes in the FText, including live-culture changes */
-class CORE_API FTextSnapshot
+class FTextSnapshot
 {
 public:
-	FTextSnapshot();
+	CORE_API FTextSnapshot();
 
-	explicit FTextSnapshot(const FText& InText);
+	CORE_API explicit FTextSnapshot(const FText& InText);
 
 	/** Check to see whether the given text is identical to the text this snapshot was made from */
-	bool IdenticalTo(const FText& InText) const;
+	CORE_API bool IdenticalTo(const FText& InText) const;
 
 	/** Check to see whether the display string of the given text is identical to the display string this snapshot was made from */
-	bool IsDisplayStringEqualTo(const FText& InText) const;
+	CORE_API bool IsDisplayStringEqualTo(const FText& InText) const;
 
 private:
 
@@ -1123,30 +1123,30 @@ private:
 	uint32 Flags = 0;
 };
 
-class CORE_API FTextInspector
+class FTextInspector
 {
 private:
 	FTextInspector() {}
 	~FTextInspector() {}
 
 public:
-	static bool ShouldGatherForLocalization(const FText& Text);
-	static TOptional<FString> GetNamespace(const FText& Text);
-	static TOptional<FString> GetKey(const FText& Text);
-	static FTextId GetTextId(const FText& Text);
-	static const FString* GetSourceString(const FText& Text);
-	static const FString& GetDisplayString(const FText& Text);
+	static CORE_API bool ShouldGatherForLocalization(const FText& Text);
+	static CORE_API TOptional<FString> GetNamespace(const FText& Text);
+	static CORE_API TOptional<FString> GetKey(const FText& Text);
+	static CORE_API FTextId GetTextId(const FText& Text);
+	static CORE_API const FString* GetSourceString(const FText& Text);
+	static CORE_API const FString& GetDisplayString(const FText& Text);
 	UE_DEPRECATED(5.0, "GetSharedDisplayString is no longer guaranteed to return a valid result and should NOT be used! If you wanted to get the text ID, use FTextInspector::GetTextId instead. If you wanted a key for unique text instances, use FTextInspector::GetSharedDataId instead.")
-	static FTextConstDisplayStringPtr GetSharedDisplayString(const FText& Text);
-	static bool GetTableIdAndKey(const FText& Text, FName& OutTableId, FString& OutKey);
-	static bool GetTableIdAndKey(const FText& Text, FName& OutTableId, FTextKey& OutKey);
-	static uint32 GetFlags(const FText& Text);
-	static void GetHistoricFormatData(const FText& Text, TArray<FHistoricTextFormatData>& OutHistoricFormatData);
-	static bool GetHistoricNumericData(const FText& Text, FHistoricTextNumericData& OutHistoricNumericData);
-	static const void* GetSharedDataId(const FText& Text);
+	static CORE_API FTextConstDisplayStringPtr GetSharedDisplayString(const FText& Text);
+	static CORE_API bool GetTableIdAndKey(const FText& Text, FName& OutTableId, FString& OutKey);
+	static CORE_API bool GetTableIdAndKey(const FText& Text, FName& OutTableId, FTextKey& OutKey);
+	static CORE_API uint32 GetFlags(const FText& Text);
+	static CORE_API void GetHistoricFormatData(const FText& Text, TArray<FHistoricTextFormatData>& OutHistoricFormatData);
+	static CORE_API bool GetHistoricNumericData(const FText& Text, FHistoricTextNumericData& OutHistoricNumericData);
+	static CORE_API const void* GetSharedDataId(const FText& Text);
 };
 
-class CORE_API FTextStringHelper
+class FTextStringHelper
 {
 public:
 	/**
@@ -1160,7 +1160,7 @@ public:
 	 *
 	 * @return The parsed FText instance.
 	 */
-	static FText CreateFromBuffer(const TCHAR* Buffer, const TCHAR* TextNamespace = nullptr, const TCHAR* PackageNamespace = nullptr, const bool bRequiresQuotes = false);
+	static CORE_API FText CreateFromBuffer(const TCHAR* Buffer, const TCHAR* TextNamespace = nullptr, const TCHAR* PackageNamespace = nullptr, const bool bRequiresQuotes = false);
 
 	/**
 	 * Attempt to extract an FText instance from the given stream of text.
@@ -1173,10 +1173,10 @@ public:
 	 *
 	 * @return The updated buffer after we parsed this text, or nullptr on failure
 	 */
-	static const TCHAR* ReadFromBuffer(const TCHAR* Buffer, FText& OutValue, const TCHAR* TextNamespace = nullptr, const TCHAR* PackageNamespace = nullptr, const bool bRequiresQuotes = false);
+	static CORE_API const TCHAR* ReadFromBuffer(const TCHAR* Buffer, FText& OutValue, const TCHAR* TextNamespace = nullptr, const TCHAR* PackageNamespace = nullptr, const bool bRequiresQuotes = false);
 	
 	UE_DEPRECATED(4.22, "FTextStringHelper::ReadFromString is deprecated. Use FTextStringHelper::ReadFromBuffer instead.")
-	static bool ReadFromString(const TCHAR* Buffer, FText& OutValue, const TCHAR* TextNamespace = nullptr, const TCHAR* PackageNamespace = nullptr, int32* OutNumCharsRead = nullptr, const bool bRequiresQuotes = false, const EStringTableLoadingPolicy InLoadingPolicy = EStringTableLoadingPolicy::FindOrLoad);
+	static CORE_API bool ReadFromString(const TCHAR* Buffer, FText& OutValue, const TCHAR* TextNamespace = nullptr, const TCHAR* PackageNamespace = nullptr, int32* OutNumCharsRead = nullptr, const bool bRequiresQuotes = false, const EStringTableLoadingPolicy InLoadingPolicy = EStringTableLoadingPolicy::FindOrLoad);
 
 	/**
 	 * Write the given FText instance to a stream of text
@@ -1186,64 +1186,64 @@ public:
 	 * @param bRequiresQuotes		 True if the written text literal must be surrounded by quotes (eg, when saving as a delimited list)
 	 * @param bStripPackageNamespace True to strip the package namespace from the written NSLOCTEXT value (eg, when saving cooked data)
 	 */
-	static void WriteToBuffer(FString& Buffer, const FText& Value, const bool bRequiresQuotes = false, const bool bStripPackageNamespace = false);
+	static CORE_API void WriteToBuffer(FString& Buffer, const FText& Value, const bool bRequiresQuotes = false, const bool bStripPackageNamespace = false);
 	
 	UE_DEPRECATED(4.22, "FTextStringHelper::WriteToString is deprecated. Use FTextStringHelper::WriteToBuffer instead.")
-	static bool WriteToString(FString& Buffer, const FText& Value, const bool bRequiresQuotes = false);
+	static CORE_API bool WriteToString(FString& Buffer, const FText& Value, const bool bRequiresQuotes = false);
 
 	/**
 	 * Test to see whether a given buffer contains complex text.
 	 *
 	 * @return True if it does, false otherwise
 	 */
-	static bool IsComplexText(const TCHAR* Buffer);
+	static CORE_API bool IsComplexText(const TCHAR* Buffer);
 
 private:
 	static const TCHAR* ReadFromBuffer_ComplexText(const TCHAR* Buffer, FText& OutValue, const TCHAR* TextNamespace, const TCHAR* PackageNamespace);
 };
 
-class CORE_API FTextBuilder
+class FTextBuilder
 {
 public:
 	/**
 	 * Increase the running indentation of the builder.
 	 */
-	void Indent();
+	CORE_API void Indent();
 
 	/**
 	 * Decrease the running indentation of the builder.
 	 */
-	void Unindent();
+	CORE_API void Unindent();
 
 	/**
 	 * Append an empty line to the builder, indented by the running indentation of the builder.
 	 */
-	void AppendLine();
+	CORE_API void AppendLine();
 
 	/**
 	 * Append the given text line to the builder, indented by the running indentation of the builder.
 	 */
-	void AppendLine(const FText& Text);
+	CORE_API void AppendLine(const FText& Text);
 
 	/**
 	 * Append the given string line to the builder, indented by the running indentation of the builder.
 	 */
-	void AppendLine(const FString& String);
+	CORE_API void AppendLine(const FString& String);
 
 	/**
 	 * Append the given name line to the builder, indented by the running indentation of the builder.
 	 */
-	void AppendLine(const FName& Name);
+	CORE_API void AppendLine(const FName& Name);
 
 	/**
 	 * Append the given formatted text line to the builder, indented by the running indentation of the builder.
 	 */
-	void AppendLineFormat(const FTextFormat& Pattern, const FFormatNamedArguments& Arguments);
+	CORE_API void AppendLineFormat(const FTextFormat& Pattern, const FFormatNamedArguments& Arguments);
 
 	/**
 	 * Append the given formatted text line to the builder, indented by the running indentation of the builder.
 	 */
-	void AppendLineFormat(const FTextFormat& Pattern, const FFormatOrderedArguments& Arguments);
+	CORE_API void AppendLineFormat(const FTextFormat& Pattern, const FFormatOrderedArguments& Arguments);
 
 	/**
 	 * Append the given formatted text line to the builder, indented by the running indentation of the builder.
@@ -1260,30 +1260,30 @@ public:
 	/**
 	 * Clear the builder and reset it to its default state.
 	 */
-	void Clear();
+	CORE_API void Clear();
 
 	/**
 	 * Check to see if the builder has any data.
 	 */
-	bool IsEmpty();
+	CORE_API bool IsEmpty();
 
 	/**
 	 * Build the current set of input into a FText.
 	 */
-	FText ToText() const;
+	CORE_API FText ToText() const;
 
 private:
-	void BuildAndAppendLine(FString&& Data);
-	void BuildAndAppendLine(FText&& Data);
+	CORE_API void BuildAndAppendLine(FString&& Data);
+	CORE_API void BuildAndAppendLine(FText&& Data);
 
 	TArray<FText> Lines;
 	int32 IndentCount = 0;
 };
 
 /** Unicode character helper functions */
-struct CORE_API FUnicodeChar
+struct FUnicodeChar
 {
-	static bool CodepointToString(const uint32 InCodepoint, FString& OutString);
+	static CORE_API bool CodepointToString(const uint32 InCodepoint, FString& OutString);
 };
 
 /**
@@ -1312,7 +1312,7 @@ namespace TextBiDi
 	};
 
 	/** Defines the interface for a re-usable BiDi object */
-	class CORE_API ITextBiDi
+	class ITextBiDi
 	{
 	public:
 		virtual ~ITextBiDi() {}
diff --git a/Engine/Source/Runtime/Core/Public/Internationalization/TextChronoFormatter.h b/Engine/Source/Runtime/Core/Public/Internationalization/TextChronoFormatter.h
--- a/Engine/Source/Runtime/Core/Public/Internationalization/TextChronoFormatter.h
+++ b/Engine/Source/Runtime/Core/Public/Internationalization/TextChronoFormatter.h
@@ -9,11 +9,11 @@
  * Utility for performing low-level localized chronological formats.
  * The implementation can be found in LegacyText.cpp and ICUText.cpp.
  */
-class CORE_API FTextChronoFormatter
+class FTextChronoFormatter
 {
 public:
-	static FString AsDate(const FDateTime& DateTime, const EDateTimeStyle::Type DateStyle, const FString& TimeZone, const FCulture& TargetCulture);
-	static FString AsDateTime(const FDateTime& DateTime, const EDateTimeStyle::Type DateStyle, const EDateTimeStyle::Type TimeStyle, const FString& TimeZone, const FCulture& TargetCulture);
-	static FString AsDateTime(const FDateTime& DateTime, const FString& CustomPattern, const FString& TimeZone, const FCulture& TargetCulture);
-	static FString AsTime(const FDateTime& DateTime, const EDateTimeStyle::Type TimeStyle, const FString& TimeZone, const FCulture& TargetCulture);
+	static CORE_API FString AsDate(const FDateTime& DateTime, const EDateTimeStyle::Type DateStyle, const FString& TimeZone, const FCulture& TargetCulture);
+	static CORE_API FString AsDateTime(const FDateTime& DateTime, const EDateTimeStyle::Type DateStyle, const EDateTimeStyle::Type TimeStyle, const FString& TimeZone, const FCulture& TargetCulture);
+	static CORE_API FString AsDateTime(const FDateTime& DateTime, const FString& CustomPattern, const FString& TimeZone, const FCulture& TargetCulture);
+	static CORE_API FString AsTime(const FDateTime& DateTime, const EDateTimeStyle::Type TimeStyle, const FString& TimeZone, const FCulture& TargetCulture);
 };
diff --git a/Engine/Source/Runtime/Core/Public/Internationalization/TextComparison.h b/Engine/Source/Runtime/Core/Public/Internationalization/TextComparison.h
--- a/Engine/Source/Runtime/Core/Public/Internationalization/TextComparison.h
+++ b/Engine/Source/Runtime/Core/Public/Internationalization/TextComparison.h
@@ -22,12 +22,12 @@ namespace ETextComparisonLevel
  * Utility for performing low-level text comparison.
  * The implementation can be found in LegacyText.cpp and ICUText.cpp.
  */
-class CORE_API FTextComparison
+class FTextComparison
 {
 public:
-	static int32 CompareTo(const FString& A, const FString& B, const ETextComparisonLevel::Type ComparisonLevel = ETextComparisonLevel::Default);
-	static int32 CompareToCaseIgnored(const FString& A, const FString& B);
+	static CORE_API int32 CompareTo(const FString& A, const FString& B, const ETextComparisonLevel::Type ComparisonLevel = ETextComparisonLevel::Default);
+	static CORE_API int32 CompareToCaseIgnored(const FString& A, const FString& B);
 
-	static bool EqualTo(const FString& A, const FString& B, const ETextComparisonLevel::Type ComparisonLevel = ETextComparisonLevel::Default);
-	static bool EqualToCaseIgnored(const FString& A, const FString& B);
+	static CORE_API bool EqualTo(const FString& A, const FString& B, const ETextComparisonLevel::Type ComparisonLevel = ETextComparisonLevel::Default);
+	static CORE_API bool EqualToCaseIgnored(const FString& A, const FString& B);
 };
diff --git a/Engine/Source/Runtime/Core/Public/Internationalization/TextFormatter.h b/Engine/Source/Runtime/Core/Public/Internationalization/TextFormatter.h
--- a/Engine/Source/Runtime/Core/Public/Internationalization/TextFormatter.h
+++ b/Engine/Source/Runtime/Core/Public/Internationalization/TextFormatter.h
@@ -18,17 +18,17 @@ struct FPrivateTextFormatArguments;
 /**
  * Definition of the pattern used during a text format.
  */
-class CORE_API FTextFormatPatternDefinition : public TSharedFromThis<FTextFormatPatternDefinition, ESPMode::ThreadSafe>
+class FTextFormatPatternDefinition : public TSharedFromThis<FTextFormatPatternDefinition, ESPMode::ThreadSafe>
 {
 public:
 	/** Constructor */
-	FTextFormatPatternDefinition();
+	CORE_API FTextFormatPatternDefinition();
 
 	/** Singleton access to the default instance */
-	static FTextFormatPatternDefinitionConstRef GetDefault();
+	static CORE_API FTextFormatPatternDefinitionConstRef GetDefault();
 
 	/** Get the text format definitions used when formatting text */
-	const FTokenDefinitions& GetTextFormatDefinitions() const;
+	CORE_API const FTokenDefinitions& GetTextFormatDefinitions() const;
 
 	/** Is the given character one that an escape token may escape? */
 	FORCEINLINE bool IsValidEscapeChar(const TCHAR InChar) const
@@ -70,34 +70,34 @@ private:
 /**
  * A text formatter is responsible for formatting text patterns using a set of named or ordered arguments.
  */
-class CORE_API FTextFormatter
+class FTextFormatter
 {
 public:
 	/** Callback function used to compile an argument modifier. Takes an argument modifier string and pattern definition, then returns the compiled result. */
 	typedef TFunction<TSharedPtr<ITextFormatArgumentModifier>(const FTextFormatString&, FTextFormatPatternDefinitionConstRef)> FCompileTextArgumentModifierFuncPtr;
 
 	/** Singleton access */
-	static FTextFormatter& Get();
+	static CORE_API FTextFormatter& Get();
 
-	void RegisterTextArgumentModifier(const FTextFormatString& InKeyword, FCompileTextArgumentModifierFuncPtr InCompileFunc);
-	void UnregisterTextArgumentModifier(const FTextFormatString& InKeyword);
-	FCompileTextArgumentModifierFuncPtr FindTextArgumentModifier(const FTextFormatString& InKeyword) const;
+	CORE_API void RegisterTextArgumentModifier(const FTextFormatString& InKeyword, FCompileTextArgumentModifierFuncPtr InCompileFunc);
+	CORE_API void UnregisterTextArgumentModifier(const FTextFormatString& InKeyword);
+	CORE_API FCompileTextArgumentModifierFuncPtr FindTextArgumentModifier(const FTextFormatString& InKeyword) const;
 
 	/** Low-level versions of Format. You probably want to use FText::Format(...) rather than call these directly */
-	static FText Format(FTextFormat&& InFmt, FFormatNamedArguments&& InArguments, const bool bInRebuildText, const bool bInRebuildAsSource);
-	static FText Format(FTextFormat&& InFmt, FFormatOrderedArguments&& InArguments, const bool bInRebuildText, const bool bInRebuildAsSource);
-	static FText Format(FTextFormat&& InFmt, TArray<FFormatArgumentData>&& InArguments, const bool bInRebuildText, const bool bInRebuildAsSource);
+	static CORE_API FText Format(FTextFormat&& InFmt, FFormatNamedArguments&& InArguments, const bool bInRebuildText, const bool bInRebuildAsSource);
+	static CORE_API FText Format(FTextFormat&& InFmt, FFormatOrderedArguments&& InArguments, const bool bInRebuildText, const bool bInRebuildAsSource);
+	static CORE_API FText Format(FTextFormat&& InFmt, TArray<FFormatArgumentData>&& InArguments, const bool bInRebuildText, const bool bInRebuildAsSource);
 
 	/** Low-level version of Format that returns a string. This should typically only be used externally when rebuilding the display string for some formatted text */
-	static FString FormatStr(const FTextFormat& InFmt, const FFormatNamedArguments& InArguments, const bool bInRebuildText, const bool bInRebuildAsSource);
-	static FString FormatStr(const FTextFormat& InFmt, const FFormatOrderedArguments& InArguments, const bool bInRebuildText, const bool bInRebuildAsSource);
-	static FString FormatStr(const FTextFormat& InFmt, const TArray<FFormatArgumentData>& InArguments, const bool bInRebuildText, const bool bInRebuildAsSource);
+	static CORE_API FString FormatStr(const FTextFormat& InFmt, const FFormatNamedArguments& InArguments, const bool bInRebuildText, const bool bInRebuildAsSource);
+	static CORE_API FString FormatStr(const FTextFormat& InFmt, const FFormatOrderedArguments& InArguments, const bool bInRebuildText, const bool bInRebuildAsSource);
+	static CORE_API FString FormatStr(const FTextFormat& InFmt, const TArray<FFormatArgumentData>& InArguments, const bool bInRebuildText, const bool bInRebuildAsSource);
 
 	/** Incredibly low-level version of format. You should only be calling this if you're implementing a custom argument modifier type that itself needs to format using the private arguments */
-	static FString Format(const FTextFormat& InFmt, const FPrivateTextFormatArguments& InFormatArgs);
+	static CORE_API FString Format(const FTextFormat& InFmt, const FPrivateTextFormatArguments& InFormatArgs);
 
 	/** Incredibly low-level version of FFormatArgumentValue::ToFormattedString. You should only be calling this if you're implementing a custom argument modifier type that itself needs to convert the argument to a string */
-	static void ArgumentValueToFormattedString(const FFormatArgumentValue& InValue, const FPrivateTextFormatArguments& InFormatArgs, FString& OutResult);
+	static CORE_API void ArgumentValueToFormattedString(const FFormatArgumentValue& InValue, const FPrivateTextFormatArguments& InFormatArgs, FString& OutResult);
 
 private:
 	static int32 EstimateArgumentValueLength(const FFormatArgumentValue& ArgumentValue);
diff --git a/Engine/Source/Runtime/Core/Public/Internationalization/TextGeneratorBase.h b/Engine/Source/Runtime/Core/Public/Internationalization/TextGeneratorBase.h
--- a/Engine/Source/Runtime/Core/Public/Internationalization/TextGeneratorBase.h
+++ b/Engine/Source/Runtime/Core/Public/Internationalization/TextGeneratorBase.h
@@ -9,16 +9,16 @@
 /**
  * Base class implementation for ITextGenerator.
  */
-class CORE_API FTextGeneratorBase : public ITextGenerator
+class FTextGeneratorBase : public ITextGenerator
 {
 public: // Serialization
 	/**
 	 * Gets the type ID of this generator. The type ID is used to reconstruct this type for serialization and must be registered with FText::RegisterTextGenerator().
 	 */
-	virtual FName GetTypeID() const override;
+	CORE_API virtual FName GetTypeID() const override;
 
 	/**
 	 * Serializes this generator.
 	 */
-	virtual void Serialize(FStructuredArchive::FRecord Record) override;
+	CORE_API virtual void Serialize(FStructuredArchive::FRecord Record) override;
 };
diff --git a/Engine/Source/Runtime/Core/Public/Internationalization/TextKey.h b/Engine/Source/Runtime/Core/Public/Internationalization/TextKey.h
--- a/Engine/Source/Runtime/Core/Public/Internationalization/TextKey.h
+++ b/Engine/Source/Runtime/Core/Public/Internationalization/TextKey.h
@@ -56,13 +56,13 @@ namespace TextKeyUtil
  * Optimized representation of a case-sensitive string, as used by localization keys.
  * This references an entry within a internal table to avoid memory duplication, as well as offering optimized comparison and hashing performance.
  */
-class CORE_API FTextKey
+class FTextKey
 {
 public:
-	FTextKey();
-	FTextKey(const TCHAR* InStr);
-	FTextKey(const FString& InStr);
-	FTextKey(FString&& InStr);
+	CORE_API FTextKey();
+	CORE_API FTextKey(const TCHAR* InStr);
+	CORE_API FTextKey(const FString& InStr);
+	CORE_API FTextKey(FString&& InStr);
 
 	/** Get the underlying chars buffer this text key represents */
 	FORCEINLINE const TCHAR* GetChars() const
@@ -89,22 +89,22 @@ public:
 	}
 
 	/** Serialize this text key as if it were an FString */
-	void SerializeAsString(FArchive& Ar);
+	CORE_API void SerializeAsString(FArchive& Ar);
 
 	/** Serialize this text key including its hash value (this method is sensitive to hashing algorithm changes, so only use it for generated files that can be rebuilt from another source) */
-	void SerializeWithHash(FArchive& Ar);
+	CORE_API void SerializeWithHash(FArchive& Ar);
 
 	/** Serialize this text key including its hash value, discarding the hash on load (to upgrade from an older hashing algorithm) */
-	void SerializeDiscardHash(FArchive& Ar);
+	CORE_API void SerializeDiscardHash(FArchive& Ar);
 
 	/** Serialize this text key as if it were an FString */
-	void SerializeAsString(FStructuredArchiveSlot Slot);
+	CORE_API void SerializeAsString(FStructuredArchiveSlot Slot);
 
 	/** Serialize this text key including its hash value (this method is sensitive to hashing algorithm changes, so only use it for generated files that can be rebuilt from another source) */
-	void SerializeWithHash(FStructuredArchiveSlot Slot);
+	CORE_API void SerializeWithHash(FStructuredArchiveSlot Slot);
 
 	/** Serialize this text key including its hash value, discarding the hash on load (to upgrade from an older hashing algorithm) */
-	void SerializeDiscardHash(FStructuredArchiveSlot Slot);
+	CORE_API void SerializeDiscardHash(FStructuredArchiveSlot Slot);
 
 	/** Is this text key empty? */
 	FORCEINLINE bool IsEmpty() const
@@ -113,13 +113,13 @@ public:
 	}
 
 	/** Reset this text key to be empty */
-	void Reset();
+	CORE_API void Reset();
 
 	/** Compact any slack within the internal table */
-	static void CompactDataStructures();
+	static CORE_API void CompactDataStructures();
 
 	/** Do not use any FTextKey or FTextId after calling this */
-	static void TearDown();
+	static CORE_API void TearDown();
 
 private:
 	/** Pointer to the string buffer we reference from the internal table */
@@ -136,7 +136,7 @@ private:
 /**
  * Optimized representation of a text identity (a namespace and key pair).
  */
-class CORE_API FTextId
+class FTextId
 {
 public:
 	FTextId()
diff --git a/Engine/Source/Runtime/Core/Public/Internationalization/TextLocalizationManager.h b/Engine/Source/Runtime/Core/Public/Internationalization/TextLocalizationManager.h
--- a/Engine/Source/Runtime/Core/Public/Internationalization/TextLocalizationManager.h
+++ b/Engine/Source/Runtime/Core/Public/Internationalization/TextLocalizationManager.h
@@ -38,7 +38,7 @@ enum class ETextLocalizationManagerInitializedFlags : uint8
 ENUM_CLASS_FLAGS(ETextLocalizationManagerInitializedFlags);
 
 /** Singleton class that manages display strings for FText. */
-class CORE_API FTextLocalizationManager
+class FTextLocalizationManager
 {
 	friend CORE_API void BeginPreInitTextLocalization();
 	friend CORE_API void BeginInitTextLocalization();
@@ -101,64 +101,64 @@ private:
 	bool bIsLocalizationLocked;
 #endif
 
-	FTextLocalizationManager();
+	CORE_API FTextLocalizationManager();
 	friend class FLazySingleton;
 	
 public:
 
-	~FTextLocalizationManager();
+	CORE_API ~FTextLocalizationManager();
 	UE_NONCOPYABLE(FTextLocalizationManager);
 
 	/** Singleton accessor */
-	static FTextLocalizationManager& Get();
-	static void TearDown();
+	static CORE_API FTextLocalizationManager& Get();
+	static CORE_API void TearDown();
 
-	void DumpMemoryInfo() const;
-	void CompactDataStructures();
+	CORE_API void DumpMemoryInfo() const;
+	CORE_API void CompactDataStructures();
 
 #if ENABLE_LOC_TESTING
 	/**
 	 * Dumps the current live table state to the log, optionally filtering it based on the given wildcard arguments.
 	 * @note Calling this function with no filters specified will dump the entire live table.
 	 */
-	void DumpLiveTable(const FString* NamespaceFilter = nullptr, const FString* KeyFilter = nullptr, const FString* DisplayStringFilter = nullptr, const FLogCategoryBase* CategoryOverride = nullptr) const;
+	CORE_API void DumpLiveTable(const FString* NamespaceFilter = nullptr, const FString* KeyFilter = nullptr, const FString* DisplayStringFilter = nullptr, const FLogCategoryBase* CategoryOverride = nullptr) const;
 #endif
 
 	/**
 	 * Get the language that will be requested during localization initialization, based on the hierarchy of: command line -> configs -> OS default.
 	 */
-	FString GetRequestedLanguageName() const;
+	CORE_API FString GetRequestedLanguageName() const;
 
 	/**
 	 * Get the locale that will be requested during localization initialization, based on the hierarchy of: command line -> configs -> OS default.
 	 */
-	FString GetRequestedLocaleName() const;
+	CORE_API FString GetRequestedLocaleName() const;
 
 	/**
 	 * Given a localization category, get the native culture for the category (if known).
 	 * @return The native culture for the given localization category, or an empty string if the native culture is unknown.
 	 */
-	FString GetNativeCultureName(const ELocalizedTextSourceCategory InCategory) const;
+	CORE_API FString GetNativeCultureName(const ELocalizedTextSourceCategory InCategory) const;
 
 	/**
 	 * Get a list of culture names that we have localized resource data for (ELocalizationLoadFlags controls which resources should be checked).
 	 */
-	TArray<FString> GetLocalizedCultureNames(const ELocalizationLoadFlags InLoadFlags) const;
+	CORE_API TArray<FString> GetLocalizedCultureNames(const ELocalizationLoadFlags InLoadFlags) const;
 
 	/**
 	 * Register a localized text source with the text localization manager.
 	 */
-	void RegisterTextSource(const TSharedRef<ILocalizedTextSource>& InLocalizedTextSource, const bool InRefreshResources = true);
+	CORE_API void RegisterTextSource(const TSharedRef<ILocalizedTextSource>& InLocalizedTextSource, const bool InRefreshResources = true);
 
 	/**
 	 * Register a polyglot text data with the text localization manager.
 	 */
-	void RegisterPolyglotTextData(const FPolyglotTextData& InPolyglotTextData, const bool InAddDisplayString = true);
-	void RegisterPolyglotTextData(TArrayView<const FPolyglotTextData> InPolyglotTextDataArray, const bool InAddDisplayStrings = true);
+	CORE_API void RegisterPolyglotTextData(const FPolyglotTextData& InPolyglotTextData, const bool InAddDisplayString = true);
+	CORE_API void RegisterPolyglotTextData(TArrayView<const FPolyglotTextData> InPolyglotTextDataArray, const bool InAddDisplayStrings = true);
 
 	/**	Finds and returns the display string with the given namespace and key, if it exists.
 	 *	Additionally, if a source string is specified and the found localized display string was not localized from that source string, null will be returned. */
-	FTextConstDisplayStringPtr FindDisplayString(const FTextKey& Namespace, const FTextKey& Key, const FString* const SourceString = nullptr) const;
+	CORE_API FTextConstDisplayStringPtr FindDisplayString(const FTextKey& Namespace, const FTextKey& Key, const FString* const SourceString = nullptr) const;
 
 	/**	Returns a display string with the given namespace and key.
 	 *	If no display string exists, it will be created using the source string or an empty string if no source string is provided.
@@ -166,11 +166,11 @@ public:
 	 *		... but it was not localized from the specified source string, the display string will be set to the specified source and returned.
 	 *		... and it was localized from the specified source string (or none was provided), the display string will be returned.
 	*/
-	FTextConstDisplayStringRef GetDisplayString(const FTextKey& Namespace, const FTextKey& Key, const FString* const SourceString);
+	CORE_API FTextConstDisplayStringRef GetDisplayString(const FTextKey& Namespace, const FTextKey& Key, const FString* const SourceString);
 
 #if WITH_EDITORONLY_DATA
 	/** If an entry exists for the specified namespace and key, returns true and provides the localization resource identifier from which it was loaded. Otherwise, returns false. */
-	bool GetLocResID(const FTextKey& Namespace, const FTextKey& Key, FString& OutLocResId) const;
+	CORE_API bool GetLocResID(const FTextKey& Namespace, const FTextKey& Key, FString& OutLocResId) const;
 #endif
 
 	UE_DEPRECATED(5.0, "FindNamespaceAndKeyFromDisplayString no longer functions! Use FTextInspector::GetTextId instead.")
@@ -183,11 +183,11 @@ public:
 	 *	Returns true if the display string has been or was already associated with the namespace and key.
 	 *	Returns false if the display string was already associated with another namespace and key or the namespace and key are already in use by another display string.
 	 */
-	bool AddDisplayString(const FTextDisplayStringRef& DisplayString, const FTextKey& Namespace, const FTextKey& Key);
+	CORE_API bool AddDisplayString(const FTextDisplayStringRef& DisplayString, const FTextKey& Namespace, const FTextKey& Key);
 
 	/** Updates display string entries and adds new display string entries based on localizations found in a specified localization resource. */
-	void UpdateFromLocalizationResource(const FString& LocalizationResourceFilePath);
-	void UpdateFromLocalizationResource(const FTextLocalizationResource& TextLocalizationResource);
+	CORE_API void UpdateFromLocalizationResource(const FString& LocalizationResourceFilePath);
+	CORE_API void UpdateFromLocalizationResource(const FTextLocalizationResource& TextLocalizationResource);
 
 	/**
 	 * Wait for any current async tasks to finish.
@@ -199,104 +199,104 @@ public:
 	 * data to load post-initialization (eg, because you're using chunked localization targets or are explicitly loading plugins). A good place to add a wait in that case 
 	 * might be at the end of your loading screen, before showing your main menu or game world.
 	 */
-	void WaitForAsyncTasks();
+	CORE_API void WaitForAsyncTasks();
 
 	/**
 	 * Reloads resources for the current culture.
 	 * @note The loading is async, see WaitUntilLoadingCompletes.
 	 */
-	void RefreshResources();
+	CORE_API void RefreshResources();
 
 	/**
 	 * Called when paths containing additional localization target data (LocRes) are mounted, to allow the display strings to dynamically update without waiting for a refresh.
 	 * @note The loading is async, see WaitUntilLoadingCompletes.
 	 * @see FCoreDelegates::GatherAdditionalLocResPathsCallback.
 	 */
-	void HandleLocalizationTargetsMounted(TArrayView<const FString> LocalizationTargetPaths);
+	CORE_API void HandleLocalizationTargetsMounted(TArrayView<const FString> LocalizationTargetPaths);
 
 	 /**
 	  * Called when paths containing additional localization target data (LocRes) are unmounted, to allow the display strings to dynamically update without waiting for a refresh.
 	  * @note The loading is async, see WaitUntilLoadingCompletes.
 	  * @see FCoreDelegates::GatherAdditionalLocResPathsCallback.
 	  */
-	void HandleLocalizationTargetsUnmounted(TArrayView<const FString> LocalizationTargetPaths);
+	CORE_API void HandleLocalizationTargetsUnmounted(TArrayView<const FString> LocalizationTargetPaths);
 
 	/**
 	 * Returns the current text revision number. This value can be cached when caching information from the text localization manager.
 	 * If the revision does not match, cached information may be invalid and should be recached.
 	 */
-	uint16 GetTextRevision() const;
+	CORE_API uint16 GetTextRevision() const;
 
 	/**
 	 * Attempts to find a local revision history for the given text ID.
 	 * This will only be set if the display string has been changed since the localization manager version has been changed (eg, if it has been edited while keeping the same key).
 	 * @return The local revision, or 0 if there have been no changes since a global history change.
 	 */
-	uint16 GetLocalRevisionForTextId(const FTextId& InTextId) const;
+	CORE_API uint16 GetLocalRevisionForTextId(const FTextId& InTextId) const;
 
 	/**
 	 * Get both the global and local revision for the given text ID.
 	 * @see GetTextRevision and GetLocalRevisionForTextId.
 	 */
-	void GetTextRevisions(const FTextId& InTextId, uint16& OutGlobalTextRevision, uint16& OutLocalTextRevision) const;
+	CORE_API void GetTextRevisions(const FTextId& InTextId, uint16& OutGlobalTextRevision, uint16& OutLocalTextRevision) const;
 
 #if WITH_EDITOR
 	/**
 	 * Enable the game localization preview using the current "preview language" setting, or the native culture if no "preview language" is set.
 	 * @note This is the same as calling EnableGameLocalizationPreview with the current "preview language" setting.
 	 */
-	void EnableGameLocalizationPreview();
+	CORE_API void EnableGameLocalizationPreview();
 
 	/**
 	 * Enable the game localization preview using the given language, or the native language if the culture name is empty.
 	 * @note This will also lockdown localization editing if the given language is a non-native game language (to avoid accidentally baking out translations as source data in assets).
 	 */
-	void EnableGameLocalizationPreview(const FString& CultureName);
+	CORE_API void EnableGameLocalizationPreview(const FString& CultureName);
 
 	/**
 	 * Disable the game localization preview.
 	 * @note This is the same as calling EnableGameLocalizationPreview with the native game language (or an empty string).
 	 */
-	void DisableGameLocalizationPreview();
+	CORE_API void DisableGameLocalizationPreview();
 
 	/**
 	 * Is the game localization preview enabled for a non-native language?
 	 */
-	bool IsGameLocalizationPreviewEnabled() const;
+	CORE_API bool IsGameLocalizationPreviewEnabled() const;
 
 	/**
 	 * Notify that the game localization preview should automatically enable itself under certain circumstances 
 	 * (such as changing the preview language via the UI) due to a state change (such as PIE starting).
 	 * @note This must be paired with a call to PopAutoEnableGameLocalizationPreview.
 	 */
-	void PushAutoEnableGameLocalizationPreview();
+	CORE_API void PushAutoEnableGameLocalizationPreview();
 
 	/**
 	 * Notify that the game localization preview should no longer automatically enable itself under certain circumstances 
 	 * (such as changing the preview language via the UI) due to a state change (such as PIE ending).
 	 * @note This must be paired with a call to PushAutoEnableGameLocalizationPreview.
 	 */
-	void PopAutoEnableGameLocalizationPreview();
+	CORE_API void PopAutoEnableGameLocalizationPreview();
 
 	/**
 	 * Should the game localization preview automatically enable itself under certain circumstances?
 	 */
-	bool ShouldGameLocalizationPreviewAutoEnable() const;
+	CORE_API bool ShouldGameLocalizationPreviewAutoEnable() const;
 
 	/**
 	 * Configure the "preview language" setting used for the game localization preview.
 	 */
-	void ConfigureGameLocalizationPreviewLanguage(const FString& CultureName);
+	CORE_API void ConfigureGameLocalizationPreviewLanguage(const FString& CultureName);
 
 	/**
 	 * Get the configured "preview language" setting used for the game localization preview (if any).
 	 */
-	FString GetConfiguredGameLocalizationPreviewLanguage() const;
+	CORE_API FString GetConfiguredGameLocalizationPreviewLanguage() const;
 
 	/**
 	 * Is the localization of this game currently locked? (ie, can it be edited in the UI?).
 	 */
-	bool IsLocalizationLocked() const;
+	CORE_API bool IsLocalizationLocked() const;
 #endif
 
 	/** Event type for immediately reacting to changes in display strings for text. */
@@ -305,46 +305,46 @@ public:
 
 private:
 	/** Callback for when a PAK file is loaded. Async loads any chunk specific localization resources. */
-	void OnPakFileMounted(const IPakFile& PakFile);
+	CORE_API void OnPakFileMounted(const IPakFile& PakFile);
 
 	/** Callback for changes in culture. Async loads the new culture's localization resources. */
-	void OnCultureChanged();
+	CORE_API void OnCultureChanged();
 
 	/** Loads localization resources for the specified culture, optionally loading localization resources that are editor-specific or game-specific. */
-	void LoadLocalizationResourcesForCulture_Sync(TArrayView<const TSharedPtr<ILocalizedTextSource>> AvailableTextSources, const FString& CultureName, const ELocalizationLoadFlags LocLoadFlags);
-	void LoadLocalizationResourcesForCulture_Async(const FString& CultureName, const ELocalizationLoadFlags LocLoadFlags);
+	CORE_API void LoadLocalizationResourcesForCulture_Sync(TArrayView<const TSharedPtr<ILocalizedTextSource>> AvailableTextSources, const FString& CultureName, const ELocalizationLoadFlags LocLoadFlags);
+	CORE_API void LoadLocalizationResourcesForCulture_Async(const FString& CultureName, const ELocalizationLoadFlags LocLoadFlags);
 
 	/** Loads localization resources for the specified prioritized cultures, optionally loading localization resources that are editor-specific or game-specific. */
-	void LoadLocalizationResourcesForPrioritizedCultures_Sync(TArrayView<const TSharedPtr<ILocalizedTextSource>> AvailableTextSources, TArrayView<const FString> PrioritizedCultureNames, const ELocalizationLoadFlags LocLoadFlags);
-	void LoadLocalizationResourcesForPrioritizedCultures_Async(TArrayView<const FString> PrioritizedCultureNames, const ELocalizationLoadFlags LocLoadFlags);
+	CORE_API void LoadLocalizationResourcesForPrioritizedCultures_Sync(TArrayView<const TSharedPtr<ILocalizedTextSource>> AvailableTextSources, TArrayView<const FString> PrioritizedCultureNames, const ELocalizationLoadFlags LocLoadFlags);
+	CORE_API void LoadLocalizationResourcesForPrioritizedCultures_Async(TArrayView<const FString> PrioritizedCultureNames, const ELocalizationLoadFlags LocLoadFlags);
 
 	/** Loads the specified localization targets for the specified prioritized cultures, optionally loading localization resources that are editor-specific or game-specific. */
-	void LoadLocalizationTargetsForPrioritizedCultures_Sync(TArrayView<const TSharedPtr<ILocalizedTextSource>> AvailableTextSources, TArrayView<const FString> LocalizationTargetPaths, TArrayView<const FString> PrioritizedCultureNames, const ELocalizationLoadFlags LocLoadFlags);
-	void LoadLocalizationTargetsForPrioritizedCultures_Async(TArrayView<const FString> LocalizationTargetPaths, TArrayView<const FString> PrioritizedCultureNames, const ELocalizationLoadFlags LocLoadFlags);
+	CORE_API void LoadLocalizationTargetsForPrioritizedCultures_Sync(TArrayView<const TSharedPtr<ILocalizedTextSource>> AvailableTextSources, TArrayView<const FString> LocalizationTargetPaths, TArrayView<const FString> PrioritizedCultureNames, const ELocalizationLoadFlags LocLoadFlags);
+	CORE_API void LoadLocalizationTargetsForPrioritizedCultures_Async(TArrayView<const FString> LocalizationTargetPaths, TArrayView<const FString> PrioritizedCultureNames, const ELocalizationLoadFlags LocLoadFlags);
 
 	/** Loads the chunked localization resource data for the specified chunk */
-	void LoadChunkedLocalizationResources_Sync(TArrayView<const TSharedPtr<ILocalizedTextSource>> AvailableTextSources, const int32 ChunkId, const FString& PakFilename);
-	void LoadChunkedLocalizationResources_Async(const int32 ChunkId, const FString& PakFilename);
+	CORE_API void LoadChunkedLocalizationResources_Sync(TArrayView<const TSharedPtr<ILocalizedTextSource>> AvailableTextSources, const int32 ChunkId, const FString& PakFilename);
+	CORE_API void LoadChunkedLocalizationResources_Async(const int32 ChunkId, const FString& PakFilename);
 
 	/** Queue the given task to run async (where possible), chained to any existing AsyncLocalizationTask */
-	void QueueAsyncTask(TUniqueFunction<void()>&& Task);
+	CORE_API void QueueAsyncTask(TUniqueFunction<void()>&& Task);
 
 	/** Updates display string entries and adds new display string entries based on provided native text. */
-	void UpdateFromNative(FTextLocalizationResource&& TextLocalizationResource, const bool bDirtyTextRevision = true);
+	CORE_API void UpdateFromNative(FTextLocalizationResource&& TextLocalizationResource, const bool bDirtyTextRevision = true);
 
 	/** Updates display string entries and adds new display string entries based on provided localizations. */
-	void UpdateFromLocalizations(FTextLocalizationResource&& TextLocalizationResource, const bool bDirtyTextRevision = true);
+	CORE_API void UpdateFromLocalizations(FTextLocalizationResource&& TextLocalizationResource, const bool bDirtyTextRevision = true);
 
 	/** Dirties the local revision counter for the given text ID by incrementing it (or adding it) */
-	void DirtyLocalRevisionForTextId(const FTextId& InTextId);
+	CORE_API void DirtyLocalRevisionForTextId(const FTextId& InTextId);
 
 	/** Dirties the text revision counter by incrementing it, causing a revision mismatch for any information cached before this happens.  */
-	void DirtyTextRevision();
+	CORE_API void DirtyTextRevision();
 #if ENABLE_LOC_TESTING
 	/** A helper function that leetifies all of the display strings when the LEET culture is active. */
-	void LeetifyAllDisplayStrings();
+	CORE_API void LeetifyAllDisplayStrings();
 	/** A helper function that converts all of the display strings to show the localization key and namespace associated with the string when the keys culture is active. */
-	void KeyifyAllDisplayStrings();
+	CORE_API void KeyifyAllDisplayStrings();
 #endif
 
 	/** Array of registered localized text sources, sorted by priority (@see RegisterTextSource) */
diff --git a/Engine/Source/Runtime/Core/Public/Internationalization/TextLocalizationResource.h b/Engine/Source/Runtime/Core/Public/Internationalization/TextLocalizationResource.h
--- a/Engine/Source/Runtime/Core/Public/Internationalization/TextLocalizationResource.h
+++ b/Engine/Source/Runtime/Core/Public/Internationalization/TextLocalizationResource.h
@@ -15,7 +15,7 @@
 class FArchive;
 
 /** Utility class for working with Localization MetaData Resource (LocMeta) files. */
-class CORE_API FTextLocalizationMetaDataResource
+class FTextLocalizationMetaDataResource
 {
 public:
 	FTextLocalizationMetaDataResource() = default;
@@ -36,20 +36,20 @@ public:
 	TArray<FString> CompiledCultures;
 
 	/** Load the given LocMeta file into this resource. */
-	bool LoadFromFile(const FString& FilePath);
+	CORE_API bool LoadFromFile(const FString& FilePath);
 
 	/** Load the given LocMeta archive into this resource. */
-	bool LoadFromArchive(FArchive& Archive, const FString& LocMetaID);
+	CORE_API bool LoadFromArchive(FArchive& Archive, const FString& LocMetaID);
 
 	/** Save this resource to the given LocMeta file. */
-	bool SaveToFile(const FString& FilePath);
+	CORE_API bool SaveToFile(const FString& FilePath);
 
 	/** Save this resource to the given LocMeta archive. */
-	bool SaveToArchive(FArchive& Archive, const FString& LocMetaID);
+	CORE_API bool SaveToArchive(FArchive& Archive, const FString& LocMetaID);
 };
 
 /** Utility class for working with Localization Resource (LocRes) files. */
-class CORE_API FTextLocalizationResource
+class FTextLocalizationResource
 {
 public:
 	FTextLocalizationResource() = default;
@@ -85,32 +85,32 @@ public:
 	}
 
 	/** Add a single entry to this resource. */
-	void AddEntry(const FTextKey& InNamespace, const FTextKey& InKey, const FString& InSourceString, const FString& InLocalizedString, const int32 InPriority, const FTextKey& InLocResID = FTextKey());
-	void AddEntry(const FTextKey& InNamespace, const FTextKey& InKey, const FString& InSourceString, const FTextConstDisplayStringRef& InLocalizedString, const int32 InPriority, const FTextKey& InLocResID = FTextKey());
-	void AddEntry(const FTextKey& InNamespace, const FTextKey& InKey, const uint32 InSourceStringHash, const FString& InLocalizedString, const int32 InPriority, const FTextKey& InLocResID = FTextKey());
-	void AddEntry(const FTextKey& InNamespace, const FTextKey& InKey, const uint32 InSourceStringHash, const FTextConstDisplayStringRef& InLocalizedString, const int32 InPriority, const FTextKey& InLocResID = FTextKey());
+	CORE_API void AddEntry(const FTextKey& InNamespace, const FTextKey& InKey, const FString& InSourceString, const FString& InLocalizedString, const int32 InPriority, const FTextKey& InLocResID = FTextKey());
+	CORE_API void AddEntry(const FTextKey& InNamespace, const FTextKey& InKey, const FString& InSourceString, const FTextConstDisplayStringRef& InLocalizedString, const int32 InPriority, const FTextKey& InLocResID = FTextKey());
+	CORE_API void AddEntry(const FTextKey& InNamespace, const FTextKey& InKey, const uint32 InSourceStringHash, const FString& InLocalizedString, const int32 InPriority, const FTextKey& InLocResID = FTextKey());
+	CORE_API void AddEntry(const FTextKey& InNamespace, const FTextKey& InKey, const uint32 InSourceStringHash, const FTextConstDisplayStringRef& InLocalizedString, const int32 InPriority, const FTextKey& InLocResID = FTextKey());
 
 	/** Is this resource empty? */
-	bool IsEmpty() const;
+	CORE_API bool IsEmpty() const;
 
 	/** Load all LocRes files in the specified directory into this resource. */
-	void LoadFromDirectory(const FString& DirectoryPath, const int32 Priority);
+	CORE_API void LoadFromDirectory(const FString& DirectoryPath, const int32 Priority);
 
 	/** Load the given LocRes file into this resource. */
-	bool LoadFromFile(const FString& FilePath, const int32 Priority);
+	CORE_API bool LoadFromFile(const FString& FilePath, const int32 Priority);
 
 	/** Load the given LocRes archive into this resource. */
-	bool LoadFromArchive(FArchive& Archive, const FTextKey& LocResID, const int32 Priority);
+	CORE_API bool LoadFromArchive(FArchive& Archive, const FTextKey& LocResID, const int32 Priority);
 
 	/** Save this resource to the given LocRes file. */
-	bool SaveToFile(const FString& FilePath) const;
+	CORE_API bool SaveToFile(const FString& FilePath) const;
 
 	/** Save this resource to the given LocRes archive. */
-	bool SaveToArchive(FArchive& Archive, const FTextKey& LocResID) const;
+	CORE_API bool SaveToArchive(FArchive& Archive, const FTextKey& LocResID) const;
 
 private:
 	/** Test whether the new entry should replace the current entry, optionally logging a conflict */
-	static bool ShouldReplaceEntry(const FTextKey& Namespace, const FTextKey& Key, const FEntry& CurrentEntry, const FEntry& NewEntry);
+	static CORE_API bool ShouldReplaceEntry(const FTextKey& Namespace, const FTextKey& Key, const FEntry& CurrentEntry, const FEntry& NewEntry);
 };
 
 namespace TextLocalizationResourceUtil
diff --git a/Engine/Source/Runtime/Core/Public/Internationalization/TextLocalizationResourceVersion.h b/Engine/Source/Runtime/Core/Public/Internationalization/TextLocalizationResourceVersion.h
--- a/Engine/Source/Runtime/Core/Public/Internationalization/TextLocalizationResourceVersion.h
+++ b/Engine/Source/Runtime/Core/Public/Internationalization/TextLocalizationResourceVersion.h
@@ -6,18 +6,18 @@
 #include "Misc/Guid.h"
 
 /** Version information for compiled LocMeta (Localization MetaData Resource) and LocRes (Localization Resource) files */
-struct CORE_API FTextLocalizationResourceVersion
+struct FTextLocalizationResourceVersion
 {
 	/**
 	 * Magic number identifying a LocMeta file.
 	 */
-	static const FGuid LocMetaMagic;
+	static CORE_API const FGuid LocMetaMagic;
 
 	/**
 	 * Magic number identifying a LocRes file.
 	 * @note Legacy LocRes files will be missing this as it wasn't added until version 1.
 	 */
-	static const FGuid LocResMagic;
+	static CORE_API const FGuid LocResMagic;
 
 	/**
 	 * Data versions for LocMeta files.
diff --git a/Engine/Source/Runtime/Core/Public/Internationalization/TextTransformer.h b/Engine/Source/Runtime/Core/Public/Internationalization/TextTransformer.h
--- a/Engine/Source/Runtime/Core/Public/Internationalization/TextTransformer.h
+++ b/Engine/Source/Runtime/Core/Public/Internationalization/TextTransformer.h
@@ -9,9 +9,9 @@
  * Utility for performing low-level localized transforms.
  * The implementation can be found in LegacyText.cpp and ICUText.cpp.
  */
-class CORE_API FTextTransformer
+class FTextTransformer
 {
 public:
-	static FString ToLower(const FString& InStr);
-	static FString ToUpper(const FString& InStr);
+	static CORE_API FString ToLower(const FString& InStr);
+	static CORE_API FString ToUpper(const FString& InStr);
 };
diff --git a/Engine/Source/Runtime/Core/Public/Linux/LinuxPlatformMisc.h b/Engine/Source/Runtime/Core/Public/Linux/LinuxPlatformMisc.h
--- a/Engine/Source/Runtime/Core/Public/Linux/LinuxPlatformMisc.h
+++ b/Engine/Source/Runtime/Core/Public/Linux/LinuxPlatformMisc.h
@@ -13,7 +13,7 @@
  */
 
 // A stub for now. Will allow to override any functions if needed
-struct CORE_API FLinuxPlatformMisc : public FUnixPlatformMisc
+struct FLinuxPlatformMisc : public FUnixPlatformMisc
 {
 };
 
diff --git a/Engine/Source/Runtime/Core/Public/Linux/LinuxPlatformProcess.h b/Engine/Source/Runtime/Core/Public/Linux/LinuxPlatformProcess.h
--- a/Engine/Source/Runtime/Core/Public/Linux/LinuxPlatformProcess.h
+++ b/Engine/Source/Runtime/Core/Public/Linux/LinuxPlatformProcess.h
@@ -12,10 +12,10 @@
 /**
  * Linux implementation of the Process OS functions
  */
-struct CORE_API FLinuxPlatformProcess : public FUnixPlatformProcess
+struct FLinuxPlatformProcess : public FUnixPlatformProcess
 {
-	static const TCHAR* BaseDir();
-	static const TCHAR* GetBinariesSubdirectory();
+	static CORE_API const TCHAR* BaseDir();
+	static CORE_API const TCHAR* GetBinariesSubdirectory();
 };
 
 typedef FLinuxPlatformProcess FPlatformProcess;
diff --git a/Engine/Source/Runtime/Core/Public/Logging/LogScopedCategoryAndVerbosityOverride.h b/Engine/Source/Runtime/Core/Public/Logging/LogScopedCategoryAndVerbosityOverride.h
--- a/Engine/Source/Runtime/Core/Public/Logging/LogScopedCategoryAndVerbosityOverride.h
+++ b/Engine/Source/Runtime/Core/Public/Logging/LogScopedCategoryAndVerbosityOverride.h
@@ -11,7 +11,7 @@
  * This is what allow a UE_LOG(LogHAL, Log, TEXT("...")); within a UE_LOG statement to know what the category and verbosity is.
  * When one of these goes out of scope, it restores the previous values.
 **/
-class CORE_API FLogScopedCategoryAndVerbosityOverride
+class FLogScopedCategoryAndVerbosityOverride
 {
 public:
 	/** STructure to aggregate a category and verbosity pair **/
@@ -39,13 +39,13 @@ private:
 	FOverride Backup;
 public:
 	/** Back up the existing category and verbosity pair, then sets them.*/
-	FLogScopedCategoryAndVerbosityOverride(FName Category, ELogVerbosity::Type Verbosity);
+	CORE_API FLogScopedCategoryAndVerbosityOverride(FName Category, ELogVerbosity::Type Verbosity);
 
 	/** Restore the category and verbosity overrides to the previous value.*/
-	~FLogScopedCategoryAndVerbosityOverride();
+	CORE_API ~FLogScopedCategoryAndVerbosityOverride();
 
 	/** Manages a TLS slot with the current overrides for category and verbosity. **/
-	static FOverride* GetTLSCurrent();
+	static CORE_API FOverride* GetTLSCurrent();
 };
 
 typedef FLogScopedCategoryAndVerbosityOverride FScopedCategoryAndVerbosityOverride;
diff --git a/Engine/Source/Runtime/Core/Public/Logging/MessageLog.h b/Engine/Source/Runtime/Core/Public/Logging/MessageLog.h
--- a/Engine/Source/Runtime/Core/Public/Logging/MessageLog.h
+++ b/Engine/Source/Runtime/Core/Public/Logging/MessageLog.h
@@ -108,14 +108,14 @@ public:
 	 * This allows systems that implement IMessageLog to receive messages.
 	 */
 	DECLARE_DELEGATE_RetVal_OneParam( TSharedRef<class IMessageLog>, FGetLog, const FName& );
-	CORE_API static FGetLog& OnGetLog() { return GetLog; }
+	static FGetLog& OnGetLog() { return GetLog; }
 
 	/**
 	 * Delegate used when message selection changes.
 	 * This is used to select object in the scene according to objects referenced in message tokens.
 	 */
 	DECLARE_DELEGATE_OneParam( FMessageSelectionChanged, TArray< TSharedRef<FTokenizedMessage> >& );
-	CORE_API static FMessageSelectionChanged& OnMessageSelectionChanged() { return MessageSelectionChanged; }
+	static FMessageSelectionChanged& OnMessageSelectionChanged() { return MessageSelectionChanged; }
 
 	/** Helper function to convert message log severity to log verbosity */
 	CORE_API static ELogVerbosity::Type GetLogVerbosity( EMessageSeverity::Type InSeverity );
diff --git a/Engine/Source/Runtime/Core/Public/Math/BasicMathExpressionEvaluator.h b/Engine/Source/Runtime/Core/Public/Math/BasicMathExpressionEvaluator.h
--- a/Engine/Source/Runtime/Core/Public/Math/BasicMathExpressionEvaluator.h
+++ b/Engine/Source/Runtime/Core/Public/Math/BasicMathExpressionEvaluator.h
@@ -79,14 +79,14 @@ namespace ExpressionParser
 }
 
 /** A basic math expression evaluator */
-class CORE_API FBasicMathExpressionEvaluator
+class FBasicMathExpressionEvaluator
 {
 public:
 	/** Constructor that sets up the parser's lexer and compiler */
-	FBasicMathExpressionEvaluator();
+	CORE_API FBasicMathExpressionEvaluator();
 
 	/** Evaluate the given expression, resulting in either a double value, or an error */
-	TValueOrError<double, FExpressionError> Evaluate(const TCHAR* InExpression, double InExistingValue = 0) const;
+	CORE_API TValueOrError<double, FExpressionError> Evaluate(const TCHAR* InExpression, double InExistingValue = 0) const;
 	
 private:
 	FTokenDefinitions TokenDefinitions;
diff --git a/Engine/Source/Runtime/Core/Public/Math/Color.h b/Engine/Source/Runtime/Core/Public/Math/Color.h
--- a/Engine/Source/Runtime/Core/Public/Math/Color.h
+++ b/Engine/Source/Runtime/Core/Public/Math/Color.h
@@ -276,7 +276,7 @@ struct FLinearColor
 		return FMath::Abs(this->R - ColorB.R) < Tolerance && FMath::Abs(this->G - ColorB.G) < Tolerance && FMath::Abs(this->B - ColorB.B) < Tolerance && FMath::Abs(this->A - ColorB.A) < Tolerance;
 	}
 
-	CORE_API FLinearColor CopyWithNewOpacity(float NewOpacicty) const
+	FLinearColor CopyWithNewOpacity(float NewOpacicty) const
 	{
 		FLinearColor NewCopy = *this;
 		NewCopy.A = NewOpacicty;
diff --git a/Engine/Source/Runtime/Core/Public/Math/ColorList.h b/Engine/Source/Runtime/Core/Public/Math/ColorList.h
--- a/Engine/Source/Runtime/Core/Public/Math/ColorList.h
+++ b/Engine/Source/Runtime/Core/Public/Math/ColorList.h
@@ -12,133 +12,133 @@
 
 class FString;
 
-class CORE_API FColorList
+class FColorList
 {
 public:
 	typedef TMap< FString, const FColor* > TColorsMap;
 	typedef TArray< const FColor* > TColorsLookup;
 
 	// Common colors.	
-	static const FColor White;
-	static const FColor Red;
-	static const FColor Green;
-	static const FColor Blue;
-	static const FColor Magenta;
-	static const FColor Cyan;
-	static const FColor Yellow;
-	static const FColor Black;
-	static const FColor Aquamarine;
-	static const FColor BakerChocolate;
-	static const FColor BlueViolet;
-	static const FColor Brass;
-	static const FColor BrightGold;
-	static const FColor Brown;
-	static const FColor Bronze;
-	static const FColor BronzeII;
-	static const FColor CadetBlue;
-	static const FColor CoolCopper;
-	static const FColor Copper;
-	static const FColor Coral;
-	static const FColor CornFlowerBlue;
-	static const FColor DarkBrown;
-	static const FColor DarkGreen;
-	static const FColor DarkGreenCopper;
-	static const FColor DarkOliveGreen;
-	static const FColor DarkOrchid;
-	static const FColor DarkPurple;
-	static const FColor DarkSlateBlue;
-	static const FColor DarkSlateGrey;
-	static const FColor DarkTan;
-	static const FColor DarkTurquoise;
-	static const FColor DarkWood;
-	static const FColor DimGrey;
-	static const FColor DustyRose;
-	static const FColor Feldspar;
-	static const FColor Firebrick;
-	static const FColor ForestGreen;
-	static const FColor Gold;
-	static const FColor Goldenrod;
-	static const FColor Grey;
-	static const FColor GreenCopper;
-	static const FColor GreenYellow;
-	static const FColor HunterGreen;
-	static const FColor IndianRed;
-	static const FColor Khaki;
-	static const FColor LightBlue;
-	static const FColor LightGrey;
-	static const FColor LightSteelBlue;
-	static const FColor LightWood;
-	static const FColor LimeGreen;
-	static const FColor MandarianOrange;
-	static const FColor Maroon;
-	static const FColor MediumAquamarine;
-	static const FColor MediumBlue;
-	static const FColor MediumForestGreen;
-	static const FColor MediumGoldenrod;
-	static const FColor MediumOrchid;
-	static const FColor MediumSeaGreen;
-	static const FColor MediumSlateBlue;
-	static const FColor MediumSpringGreen;
-	static const FColor MediumTurquoise;
-	static const FColor MediumVioletRed;
-	static const FColor MediumWood;
-	static const FColor MidnightBlue;
-	static const FColor NavyBlue;
-	static const FColor NeonBlue;
-	static const FColor NeonPink;
-	static const FColor NewMidnightBlue;
-	static const FColor NewTan;
-	static const FColor OldGold;
-	static const FColor Orange;
-	static const FColor OrangeRed;
-	static const FColor Orchid;
-	static const FColor PaleGreen;
-	static const FColor Pink;
-	static const FColor Plum;
-	static const FColor Quartz;
-	static const FColor RichBlue;
-	static const FColor Salmon;
-	static const FColor Scarlet;
-	static const FColor SeaGreen;
-	static const FColor SemiSweetChocolate;
-	static const FColor Sienna;
-	static const FColor Silver;
-	static const FColor SkyBlue;
-	static const FColor SlateBlue;
-	static const FColor SpicyPink;
-	static const FColor SpringGreen;
-	static const FColor SteelBlue;
-	static const FColor SummerSky;
-	static const FColor Tan;
-	static const FColor Thistle;
-	static const FColor Turquoise;
-	static const FColor VeryDarkBrown;
-	static const FColor VeryLightGrey;
-	static const FColor Violet;
-	static const FColor VioletRed;
-	static const FColor Wheat;
-	static const FColor YellowGreen;
+	static CORE_API const FColor White;
+	static CORE_API const FColor Red;
+	static CORE_API const FColor Green;
+	static CORE_API const FColor Blue;
+	static CORE_API const FColor Magenta;
+	static CORE_API const FColor Cyan;
+	static CORE_API const FColor Yellow;
+	static CORE_API const FColor Black;
+	static CORE_API const FColor Aquamarine;
+	static CORE_API const FColor BakerChocolate;
+	static CORE_API const FColor BlueViolet;
+	static CORE_API const FColor Brass;
+	static CORE_API const FColor BrightGold;
+	static CORE_API const FColor Brown;
+	static CORE_API const FColor Bronze;
+	static CORE_API const FColor BronzeII;
+	static CORE_API const FColor CadetBlue;
+	static CORE_API const FColor CoolCopper;
+	static CORE_API const FColor Copper;
+	static CORE_API const FColor Coral;
+	static CORE_API const FColor CornFlowerBlue;
+	static CORE_API const FColor DarkBrown;
+	static CORE_API const FColor DarkGreen;
+	static CORE_API const FColor DarkGreenCopper;
+	static CORE_API const FColor DarkOliveGreen;
+	static CORE_API const FColor DarkOrchid;
+	static CORE_API const FColor DarkPurple;
+	static CORE_API const FColor DarkSlateBlue;
+	static CORE_API const FColor DarkSlateGrey;
+	static CORE_API const FColor DarkTan;
+	static CORE_API const FColor DarkTurquoise;
+	static CORE_API const FColor DarkWood;
+	static CORE_API const FColor DimGrey;
+	static CORE_API const FColor DustyRose;
+	static CORE_API const FColor Feldspar;
+	static CORE_API const FColor Firebrick;
+	static CORE_API const FColor ForestGreen;
+	static CORE_API const FColor Gold;
+	static CORE_API const FColor Goldenrod;
+	static CORE_API const FColor Grey;
+	static CORE_API const FColor GreenCopper;
+	static CORE_API const FColor GreenYellow;
+	static CORE_API const FColor HunterGreen;
+	static CORE_API const FColor IndianRed;
+	static CORE_API const FColor Khaki;
+	static CORE_API const FColor LightBlue;
+	static CORE_API const FColor LightGrey;
+	static CORE_API const FColor LightSteelBlue;
+	static CORE_API const FColor LightWood;
+	static CORE_API const FColor LimeGreen;
+	static CORE_API const FColor MandarianOrange;
+	static CORE_API const FColor Maroon;
+	static CORE_API const FColor MediumAquamarine;
+	static CORE_API const FColor MediumBlue;
+	static CORE_API const FColor MediumForestGreen;
+	static CORE_API const FColor MediumGoldenrod;
+	static CORE_API const FColor MediumOrchid;
+	static CORE_API const FColor MediumSeaGreen;
+	static CORE_API const FColor MediumSlateBlue;
+	static CORE_API const FColor MediumSpringGreen;
+	static CORE_API const FColor MediumTurquoise;
+	static CORE_API const FColor MediumVioletRed;
+	static CORE_API const FColor MediumWood;
+	static CORE_API const FColor MidnightBlue;
+	static CORE_API const FColor NavyBlue;
+	static CORE_API const FColor NeonBlue;
+	static CORE_API const FColor NeonPink;
+	static CORE_API const FColor NewMidnightBlue;
+	static CORE_API const FColor NewTan;
+	static CORE_API const FColor OldGold;
+	static CORE_API const FColor Orange;
+	static CORE_API const FColor OrangeRed;
+	static CORE_API const FColor Orchid;
+	static CORE_API const FColor PaleGreen;
+	static CORE_API const FColor Pink;
+	static CORE_API const FColor Plum;
+	static CORE_API const FColor Quartz;
+	static CORE_API const FColor RichBlue;
+	static CORE_API const FColor Salmon;
+	static CORE_API const FColor Scarlet;
+	static CORE_API const FColor SeaGreen;
+	static CORE_API const FColor SemiSweetChocolate;
+	static CORE_API const FColor Sienna;
+	static CORE_API const FColor Silver;
+	static CORE_API const FColor SkyBlue;
+	static CORE_API const FColor SlateBlue;
+	static CORE_API const FColor SpicyPink;
+	static CORE_API const FColor SpringGreen;
+	static CORE_API const FColor SteelBlue;
+	static CORE_API const FColor SummerSky;
+	static CORE_API const FColor Tan;
+	static CORE_API const FColor Thistle;
+	static CORE_API const FColor Turquoise;
+	static CORE_API const FColor VeryDarkBrown;
+	static CORE_API const FColor VeryLightGrey;
+	static CORE_API const FColor Violet;
+	static CORE_API const FColor VioletRed;
+	static CORE_API const FColor Wheat;
+	static CORE_API const FColor YellowGreen;
 
 	/** Initializes list of common colors. */
-	void CreateColorMap();
+	CORE_API void CreateColorMap();
 
 	/** Returns a color based on ColorName. If not found, returns White. */
-	const FColor& GetFColorByName( const TCHAR* ColorName ) const;
+	CORE_API const FColor& GetFColorByName( const TCHAR* ColorName ) const;
 
 	/** Returns a linear color based on ColorName. If not found, returns White. */
-	const FLinearColor GetFLinearColorByName( const TCHAR* ColorName ) const;
+	CORE_API const FLinearColor GetFLinearColorByName( const TCHAR* ColorName ) const;
 
 	/** Returns true if color is valid common colors. If not found returns false. */
-	bool IsValidColorName( const TCHAR* ColorName ) const;
+	CORE_API bool IsValidColorName( const TCHAR* ColorName ) const;
 
 	/** Returns index of color based on ColorName. If not found returns 0. */
-	int32 GetColorIndex( const TCHAR* ColorName ) const;
+	CORE_API int32 GetColorIndex( const TCHAR* ColorName ) const;
 
 	/** Returns a color based on index. If index is invalid, returns White. */
-	const FColor& GetFColorByIndex( int32 ColorIndex ) const;
+	CORE_API const FColor& GetFColorByIndex( int32 ColorIndex ) const;
 
 	/** Returns color's name based on index. If index is invalid, returns BadIndex. */
-	const FString& GetColorNameByIndex( int32 ColorIndex ) const;
+	CORE_API const FString& GetColorNameByIndex( int32 ColorIndex ) const;
 
 	/** Returns the number of colors. */
 	int32 GetColorsNum() const
@@ -147,10 +147,10 @@ public:
 	}
 
 	/** Prints to log all colors information. */
-	void LogColors();
+	CORE_API void LogColors();
 
 protected:
-	void InitializeColor( const TCHAR* ColorName, const FColor* ColorPtr, int32& CurrentIndex );
+	CORE_API void InitializeColor( const TCHAR* ColorName, const FColor* ColorPtr, int32& CurrentIndex );
 
 	/** List of common colors. */
 	TColorsMap ColorsMap;
diff --git a/Engine/Source/Runtime/Core/Public/Math/CurveEdInterface.h b/Engine/Source/Runtime/Core/Public/Math/CurveEdInterface.h
--- a/Engine/Source/Runtime/Core/Public/Math/CurveEdInterface.h
+++ b/Engine/Source/Runtime/Core/Public/Math/CurveEdInterface.h
@@ -7,7 +7,7 @@
 #include "Math/InterpCurvePoint.h"
 
 /** Interface that allows the CurveEditor to edit this type of object. */
-class CORE_API FCurveEdInterface
+class FCurveEdInterface
 {
 public:
 	/** Get number of keyframes in curve. */
