[Core]

diff --git a/Engine/Source/Runtime/Core/Public/Math/UnitConversion.h b/Engine/Source/Runtime/Core/Public/Math/UnitConversion.h
--- a/Engine/Source/Runtime/Core/Public/Math/UnitConversion.h
+++ b/Engine/Source/Runtime/Core/Public/Math/UnitConversion.h
@@ -78,20 +78,20 @@ enum class EUnitType
 template<typename NumericType> struct FNumericUnit;
 
 /** Unit settings accessed globally through FUnitConversion::Settings() */
-class CORE_API FUnitSettings
+class FUnitSettings
 {
 public:
 
-	FUnitSettings();
+	CORE_API FUnitSettings();
 
 	/** Check whether unit display is globally enabled or disabled */
-	bool ShouldDisplayUnits() const;
-	void SetShouldDisplayUnits(bool bInGlobalUnitDisplay);
+	CORE_API bool ShouldDisplayUnits() const;
+	CORE_API void SetShouldDisplayUnits(bool bInGlobalUnitDisplay);
 	
 	/** Get/Set the specific valid units to display the specified type of unit in */
-	const TArray<EUnit>& GetDisplayUnits(EUnitType InType) const;
-	void SetDisplayUnits(EUnitType InType, const TArray<EUnit>& Units);
-	void SetDisplayUnits(EUnitType InType, EUnit Units);
+	CORE_API const TArray<EUnit>& GetDisplayUnits(EUnitType InType) const;
+	CORE_API void SetDisplayUnits(EUnitType InType, const TArray<EUnit>& Units);
+	CORE_API void SetDisplayUnits(EUnitType InType, EUnit Units);
 
 	/** Returns an event delegate that is executed when a display setting has changed. (GlobalUnitDisplay or DefaultInputUnits) */
 	DECLARE_EVENT(FUnitSettings, FDisplaySettingChanged);
@@ -109,25 +109,25 @@ private:
 	FDisplaySettingChanged SettingChangedEvent;
 };
 
-struct CORE_API FUnitConversion
+struct FUnitConversion
 {
 	/** Get the global settings for unit conversion/display */
-	static FUnitSettings& Settings();
+	static CORE_API FUnitSettings& Settings();
 
 	/** Check whether it is possible to convert a number between the two specified units */
-	static bool AreUnitsCompatible(EUnit From, EUnit To);
+	static CORE_API bool AreUnitsCompatible(EUnit From, EUnit To);
 
 	/** Check whether a unit is of the specified type */
-	static bool IsUnitOfType(EUnit Unit, EUnitType Type);
+	static CORE_API bool IsUnitOfType(EUnit Unit, EUnitType Type);
 
 	/** Get the type of the specified unit */
-	static EUnitType GetUnitType(EUnit);
+	static CORE_API EUnitType GetUnitType(EUnit);
 
 	/** Get the display string for the the specified unit type */
-	static const TCHAR* GetUnitDisplayString(EUnit Unit);
+	static CORE_API const TCHAR* GetUnitDisplayString(EUnit Unit);
 
 	/** Helper function to find a unit from a string (name or display string) */
-	static TOptional<EUnit> UnitFromString(const TCHAR* UnitString);
+	static CORE_API TOptional<EUnit> UnitFromString(const TCHAR* UnitString);
 
 public:
 
diff --git a/Engine/Source/Runtime/Core/Public/MemPro/MemProProfiler.h b/Engine/Source/Runtime/Core/Public/MemPro/MemProProfiler.h
--- a/Engine/Source/Runtime/Core/Public/MemPro/MemProProfiler.h
+++ b/Engine/Source/Runtime/Core/Public/MemPro/MemProProfiler.h
@@ -19,12 +19,12 @@
 #include "MemPro/MemPro.h"
 #include "Containers/StaticArray.h"
 
-class CORE_API FMemProProfiler
+class FMemProProfiler
 {
 public:
-	static void Init(const TCHAR* CmdLine);
+	static CORE_API void Init(const TCHAR* CmdLine);
 
-	static bool IsUsingPort( uint32 Port );
+	static CORE_API bool IsUsingPort( uint32 Port );
 
 	static inline bool IsStarted()
 	{
@@ -39,8 +39,8 @@ public:
 		return IsStarted() && MemProLLMTagsEnabled[(int32)Tag];
 	}
 
-	static void TrackTag( ELLMTag Tag );
-	static void TrackTagsByName( const TCHAR* TagNamesStr );
+	static CORE_API void TrackTag( ELLMTag Tag );
+	static CORE_API void TrackTagsByName( const TCHAR* TagNamesStr );
 #endif //ENABLE_LOW_LEVEL_MEM_TRACKER
 };
 
diff --git a/Engine/Source/Runtime/Core/Public/Microsoft/MicrosoftPlatformCrashContext.h b/Engine/Source/Runtime/Core/Public/Microsoft/MicrosoftPlatformCrashContext.h
--- a/Engine/Source/Runtime/Core/Public/Microsoft/MicrosoftPlatformCrashContext.h
+++ b/Engine/Source/Runtime/Core/Public/Microsoft/MicrosoftPlatformCrashContext.h
@@ -5,38 +5,38 @@
 #include "CoreTypes.h"
 #include "GenericPlatform/GenericPlatformCrashContext.h"
 
-struct CORE_API FMicrosoftPlatformCrashContext : public FGenericCrashContext
+struct FMicrosoftPlatformCrashContext : public FGenericCrashContext
 {
 	FMicrosoftPlatformCrashContext(ECrashContextType InType, const TCHAR* InErrorMessage)
 		: FGenericCrashContext(InType, InErrorMessage)
 	{
 	}
 
-	void CaptureAllThreadContexts();
+	CORE_API void CaptureAllThreadContexts();
 
-	virtual void SetPortableCallStack(const uint64* StackTrace, int32 StackTraceDepth) override;
-	virtual void AddPortableThreadCallStack(uint32 ThreadId, const TCHAR* ThreadName, const uint64* StackFrames, int32 NumStackFrames) override;
+	CORE_API virtual void SetPortableCallStack(const uint64* StackTrace, int32 StackTraceDepth) override;
+	CORE_API virtual void AddPortableThreadCallStack(uint32 ThreadId, const TCHAR* ThreadName, const uint64* StackFrames, int32 NumStackFrames) override;
 
 protected:
 	// Helpers
 	typedef TArray<void*, TInlineAllocator<128>> FModuleHandleArray;
 
-	static void GetProcModuleHandles(const FProcHandle& Process, FModuleHandleArray& OutHandles);
+	static CORE_API void GetProcModuleHandles(const FProcHandle& Process, FModuleHandleArray& OutHandles);
 
-	static void ConvertProgramCountersToStackFrames(
+	static CORE_API void ConvertProgramCountersToStackFrames(
 		const FProcHandle& Process,
 		const FModuleHandleArray& SortedModuleHandles,
 		const uint64* ProgramCounters,
 		int32 NumPCs,
 		TArray<FCrashStackFrame>& OutStackFrames);
 
-	virtual bool GetPlatformAllThreadContextsString(FString& OutStr) const override;
+	CORE_API virtual bool GetPlatformAllThreadContextsString(FString& OutStr) const override;
 
-	static void AddThreadContextString(
+	static CORE_API void AddThreadContextString(
 		uint32 CrashedThreadId,
 		uint32 ThreadId,
 		const FString& ThreadName,
 		const TArray<FCrashStackFrame>& StackFrames,
 		FString& OutStr);
 
-};
\ No newline at end of file
+};
diff --git a/Engine/Source/Runtime/Core/Public/Microsoft/MicrosoftPlatformStackWalk.h b/Engine/Source/Runtime/Core/Public/Microsoft/MicrosoftPlatformStackWalk.h
--- a/Engine/Source/Runtime/Core/Public/Microsoft/MicrosoftPlatformStackWalk.h
+++ b/Engine/Source/Runtime/Core/Public/Microsoft/MicrosoftPlatformStackWalk.h
@@ -8,14 +8,14 @@
 /**
  * Microsoft-shared implementation of the stack walking.
  **/
-struct CORE_API FMicrosoftPlatformStackWalk
+struct FMicrosoftPlatformStackWalk
 	: public FGenericPlatformStackWalk
 {
 
 protected:
 	// Extract debug info for a module from the module header in memory. 
 	// Can directly read the information even when the current target can't load the symbols itself or use certain DbgHelp APIs.
-	static bool ExtractInfoFromModule(void* ProcessHandle, void* ModuleHandle, FStackWalkModuleInfo& OutInfo);
+	static CORE_API bool ExtractInfoFromModule(void* ProcessHandle, void* ModuleHandle, FStackWalkModuleInfo& OutInfo);
 
 	//
 	// If the thread is not the calling thread, it should already have been suspended by the calling code.
@@ -31,5 +31,5 @@ protected:
 	 * @param	OutDepth			Number of stack frames returned.
 	 * @return	EXCEPTION_EXECUTE_HANDLER
 	 */
-	static int32 CaptureStackTraceInternal(uint64* OutBacktrace, uint32 MaxDepth, void* Context, void* ThreadHandle, uint32* OutDepth);
-};
\ No newline at end of file
+	static CORE_API int32 CaptureStackTraceInternal(uint64* OutBacktrace, uint32 MaxDepth, void* Context, void* ThreadHandle, uint32* OutDepth);
+};
diff --git a/Engine/Source/Runtime/Core/Public/Misc/AES.h b/Engine/Source/Runtime/Core/Public/Misc/AES.h
--- a/Engine/Source/Runtime/Core/Public/Misc/AES.h
+++ b/Engine/Source/Runtime/Core/Public/Misc/AES.h
@@ -16,7 +16,7 @@
 // current code is meant for keeping backwards compatiblity with existing data. Better way
 // to use AES would be integrated with authentication, for example, in AES-GCM mode.
 
-struct CORE_API FAES
+struct FAES
 {
 	static constexpr uint32 AESBlockSize = 16;
 
@@ -65,7 +65,7 @@ struct CORE_API FAES
 	* @param NumBytes the size of the buffer
 	* @param Key An FAESKey object containing the encryption key
 	*/
-	static void EncryptData(uint8* Contents, uint64 NumBytes, const FAESKey& Key);
+	static CORE_API void EncryptData(uint8* Contents, uint64 NumBytes, const FAESKey& Key);
 
 	/**
 	 * Encrypts a chunk of data using a specific key
@@ -74,7 +74,7 @@ struct CORE_API FAES
 	 * @param NumBytes the size of the buffer
 	 * @param Key a null terminated string that is a 32 bytes long
 	 */
-	static void EncryptData(uint8* Contents, uint64 NumBytes, const ANSICHAR* Key);
+	static CORE_API void EncryptData(uint8* Contents, uint64 NumBytes, const ANSICHAR* Key);
 
 	/**
 	* Encrypts a chunk of data using a specific key
@@ -84,7 +84,7 @@ struct CORE_API FAES
 	* @param Key a byte array that is a 32 byte length
 	* @param NumKeyBytes length of Key byte array, must be 32
 	*/
-	static void EncryptData(uint8* Contents, uint64 NumBytes, const uint8* KeyBytes, uint32 NumKeyBytes);
+	static CORE_API void EncryptData(uint8* Contents, uint64 NumBytes, const uint8* KeyBytes, uint32 NumKeyBytes);
 
 	/**
 	* Decrypts a chunk of data using a specific key
@@ -93,7 +93,7 @@ struct CORE_API FAES
 	* @param NumBytes the size of the buffer
 	* @param Key An FAESKey object containing the decryption key
 	*/
-	static void DecryptData(uint8* Contents, uint64 NumBytes, const FAESKey& Key);
+	static CORE_API void DecryptData(uint8* Contents, uint64 NumBytes, const FAESKey& Key);
 
 	/**
 	 * Decrypts a chunk of data using a specific key
@@ -102,7 +102,7 @@ struct CORE_API FAES
 	 * @param NumBytes the size of the buffer
 	 * @param Key a null terminated string that is a 32 bytes long
 	 */
-	static void DecryptData(uint8* Contents, uint64 NumBytes, const ANSICHAR* Key);
+	static CORE_API void DecryptData(uint8* Contents, uint64 NumBytes, const ANSICHAR* Key);
 
 	/**
 	* Decrypts a chunk of data using a specific key
@@ -112,5 +112,5 @@ struct CORE_API FAES
 	* @param Key a byte array that is a 32 byte length
 	* @param NumKeyBytes length of Key byte array, must be 32
 	*/
-	static void DecryptData(uint8* Contents, uint64 NumBytes, const uint8* KeyBytes, uint32 NumKeyBytes);
+	static CORE_API void DecryptData(uint8* Contents, uint64 NumBytes, const uint8* KeyBytes, uint32 NumKeyBytes);
 };
diff --git a/Engine/Source/Runtime/Core/Public/Misc/App.h b/Engine/Source/Runtime/Core/Public/Misc/App.h
--- a/Engine/Source/Runtime/Core/Public/Misc/App.h
+++ b/Engine/Source/Runtime/Core/Public/Misc/App.h
@@ -26,7 +26,7 @@
 /**
  * Provides information about the application.
  */
-class CORE_API FApp
+class FApp
 {
 public:
 
@@ -35,21 +35,21 @@ public:
 	 *
 	 * @return The branch name.
 	 */
-	static FString GetBranchName();
+	static CORE_API FString GetBranchName();
 
 	/**
 	 * Gets the application's build configuration, i.e. Debug or Shipping.
 	 *
 	 * @return The build configuration.
 	 */
-	static EBuildConfiguration GetBuildConfiguration();
+	static CORE_API EBuildConfiguration GetBuildConfiguration();
 
 	/**
 	 * Gets the target type of the current application (eg. client, server, etc...)
 	 *
 	 * @return The build target type
 	 */
-	static EBuildTargetType GetBuildTargetType();
+	static CORE_API EBuildTargetType GetBuildTargetType();
 
 #if UE_BUILD_DEVELOPMENT
 	/**
@@ -57,7 +57,7 @@ public:
 	 *
 	 * @param Whether we're running in debug game or not.
 	 */
-	static void SetDebugGame(bool bIsDebugGame);
+	static CORE_API void SetDebugGame(bool bIsDebugGame);
 #endif
 
 	/*
@@ -65,53 +65,53 @@ public:
 	*
 	* @return The build version
 	*/
-	static const TCHAR* GetBuildVersion();
+	static CORE_API const TCHAR* GetBuildVersion();
 	
 	/* 
 	 * Gets the URL for a job which created these binaries. 
 	 * This may not be accessible if you are not part of the organization that created this build. 
 	 * May be overridden with FCoreDelegates::OnGetBuildURL
 	 */
-	static const TCHAR* GetBuildURL();
+	static CORE_API const TCHAR* GetBuildURL();
 
 	/* 
 	 * Gets the URL for the currently running Horde job/step if any. 
 	 * May be overridden with FCoreDelegates::OnGetExecutingJobURL
 	 */ 
-	static const TCHAR* GetExecutingJobURL();
+	static CORE_API const TCHAR* GetExecutingJobURL();
 
 	/* Returns whether the binaries were built with debug info */
-	static bool GetIsWithDebugInfo();
+	static CORE_API bool GetIsWithDebugInfo();
 
 	/**
 	 * Gets the date at which this application was built.
 	 *
 	 * @return Build date string.
 	 */
-	static FString GetBuildDate();
+	static CORE_API FString GetBuildDate();
 
 	/**
 	 * Gets the name of the graphics RHI currently in use.
 	 *
 	 * @return name of Graphics RHI
 	 */
-	static FString GetGraphicsRHI();
+	static CORE_API FString GetGraphicsRHI();
 
 	/**
 	 * Sets the Graphics RHI currently in use
 	 */
-	static void SetGraphicsRHI(FString RHIString);
+	static CORE_API void SetGraphicsRHI(FString RHIString);
 
 
 	/**
 	 * Gets the value of ENGINE_IS_PROMOTED_BUILD.
 	 */
-	static int32 GetEngineIsPromotedBuild();
+	static CORE_API int32 GetEngineIsPromotedBuild();
 
 	/**
 	 * Gets the identifier for the unreal engine
 	 */
-	static FString GetEpicProductIdentifier();
+	static CORE_API FString GetEpicProductIdentifier();
 
 	/**
 	 * Gets the name of the current project.
@@ -238,7 +238,7 @@ public:
 	 *
 	 * @return Zen store project id.
 	 */
-	static FString GetZenStoreProjectId();
+	static CORE_API FString GetZenStoreProjectId();
 
 	/**
 	 * Gets the globally unique identifier of this application instance.
@@ -311,7 +311,7 @@ public:
 	/**
 	 * Initializes the application session.
 	 */
-	static void InitializeSession();
+	static CORE_API void InitializeSession();
 
 	/**
 	 * Check whether the specified user is authorized to interact with this session.
@@ -430,7 +430,7 @@ public:
 	 *
 	 * @return true if the application is installed, false otherwise.
 	 */
-	static bool IsInstalled();
+	static CORE_API bool IsInstalled();
 
 	/**
 	 * Checks whether the engine components of this application have been installed.
@@ -442,7 +442,7 @@ public:
 	 *
 	 * @return true if the engine is installed, false otherwise.
 	 */
-	static bool IsEngineInstalled();
+	static CORE_API bool IsEngineInstalled();
 
 	/**
 	 * Checks whether this application runs unattended.
@@ -452,7 +452,7 @@ public:
 	 *
 	 * @return true if the application runs unattended, false otherwise.
 	 */
-	static bool IsUnattended();
+	static CORE_API bool IsUnattended();
 
 	/**
 	 * Checks whether the application should run multi-threaded for performance critical features.
@@ -462,7 +462,7 @@ public:
 	 *
 	 * @return true if this isn't a server, has more than one core, does not have a -onethread command line options, etc.
 	 */
-	static bool ShouldUseThreadingForPerformance();
+	static CORE_API bool ShouldUseThreadingForPerformance();
 
 	/**
 	 * Checks whether application is in benchmark mode.
@@ -650,7 +650,7 @@ public:
 	 *
 	 * @return the current timecode.
 	 */
-	static FTimecode GetTimecode();
+	static CORE_API FTimecode GetTimecode();
 
 	/**
 	 * Get the frame rate of the current frame time.
@@ -658,7 +658,7 @@ public:
 	 *
 	 * @return the current timecode frame rate.
 	 */
-	static FFrameRate GetTimecodeFrameRate();
+	static CORE_API FFrameRate GetTimecodeFrameRate();
 
 	/**
 	 * Gets a frame number generated by the engine's timecode provider.
@@ -725,19 +725,19 @@ public:
 	 * 
 	 * @return Volume multiplier to use when app loses focus
 	 */
-	static float GetUnfocusedVolumeMultiplier();
+	static CORE_API float GetUnfocusedVolumeMultiplier();
 
 	/**
 	* Sets the Unfocused Volume Multiplier
 	*/
-	static void SetUnfocusedVolumeMultiplier(float InVolumeMultiplier);
+	static CORE_API void SetUnfocusedVolumeMultiplier(float InVolumeMultiplier);
 
 	/**
 	 * Sets if VRFocus should be used.
 	 *
 	 * @param  bInUseVRFocus	new bUseVRFocus value
 	 */
-	static void SetUseVRFocus(bool bInUseVRFocus);
+	static CORE_API void SetUseVRFocus(bool bInUseVRFocus);
 
 	/**
 	 * Gets if VRFocus should be used
@@ -752,7 +752,7 @@ public:
 	 *
 	 * @param  bInVRFocus	new VRFocus value
 	 */
-	static void SetHasVRFocus(bool bInHasVRFocus);
+	static CORE_API void SetHasVRFocus(bool bInHasVRFocus);
 
 	/**
 	 * Gets VRFocus, which indicates that the application should continue to render 
@@ -768,92 +768,92 @@ public:
 	 *
 	 * @param  InHasFocusFunction	address of a function that can query the focus state, typically &FPlatformApplicationMisc::IsThisApplicationForeground
 	 */
-	static void SetHasFocusFunction(bool (*InHasFocusFunction)());
+	static CORE_API void SetHasFocusFunction(bool (*InHasFocusFunction)());
 
 	/**
 	 * Gets Focus, Indicates that the application should continue to render
 	 * Audio and Video as if it had window focus, even though it may not.
 	 */
-	static bool HasFocus();
+	static CORE_API bool HasFocus();
 
 	/* If the random seed started with a constant or on time, can be affected by -FIXEDSEED or -BENCHMARK */
-	static bool bUseFixedSeed;
+	static CORE_API bool bUseFixedSeed;
 
 	/* Print all initial startup logging */
-	static void PrintStartupLogMessages();
+	static CORE_API void PrintStartupLogMessages();
 
 private:
 
 #if UE_BUILD_DEVELOPMENT
 	/** The current build configuration */
-	static bool bIsDebugGame;
+	static CORE_API bool bIsDebugGame;
 #endif
 
 	/** Holds the session identifier. */
-	static FGuid SessionId;
+	static CORE_API FGuid SessionId;
 
 	/** Holds the session name. */
-	static FString SessionName;
+	static CORE_API FString SessionName;
 
 	/** Holds the name of the user that launched session. */
-	static FString SessionOwner;
+	static CORE_API FString SessionOwner;
 
 	/** Holds the name the graphics RHI currently in use*/
-	static FString GraphicsRHI;
+	static CORE_API FString GraphicsRHI;
 
 	/** List of authorized session users. */
-	static TArray<FString> SessionUsers;
+	static CORE_API TArray<FString> SessionUsers;
 
 	/** Holds a flag indicating whether this is a standalone session. */
-	static bool Standalone;
+	static CORE_API bool Standalone;
 
 	/** Holds a flag Whether we are in benchmark mode or not. */
-	static bool bIsBenchmarking;
+	static CORE_API bool bIsBenchmarking;
 
 	/** Holds a flag whether we want to use a fixed time step or not. */
-	static bool bUseFixedTimeStep;
+	static CORE_API bool bUseFixedTimeStep;
 
 	/** Holds time step if a fixed delta time is wanted. */
-	static double FixedDeltaTime;
+	static CORE_API double FixedDeltaTime;
 
 	/** Holds current time. */
-	static double CurrentTime;
+	static CORE_API double CurrentTime;
 
 	/** Holds previous value of CurrentTime. */
-	static double LastTime;
+	static CORE_API double LastTime;
 
 	/** Holds current delta time in seconds. */
-	static double DeltaTime;
+	static CORE_API double DeltaTime;
 
 	/** Holds time we spent sleeping in UpdateTimeAndHandleMaxTickRate() if our frame time was smaller than one allowed by target FPS. */
-	static double IdleTime;
+	static CORE_API double IdleTime;
 
 	/** Holds the amount of IdleTime that was LONGER than we tried to sleep. The OS can't sleep the exact amount of time, so this measures that overshoot. */
-	static double IdleTimeOvershoot;
+	static CORE_API double IdleTimeOvershoot;
 
 	/** Holds overall game time. */
-	static double GameTime;
+	static CORE_API double GameTime;
 
 	/** Holds the current frame time and framerate. */
-	static TOptional<FQualifiedFrameTime> CurrentFrameTime;
+	static CORE_API TOptional<FQualifiedFrameTime> CurrentFrameTime;
 
 	/** Holds if we should generate a drop frame timecode when the frame rate does support it. */
-	static bool bUseDropFrameFormatWhenSupported;
+	static CORE_API bool bUseDropFrameFormatWhenSupported;
 
 	/** Use to affect the app volume when it loses focus */
-	static float VolumeMultiplier;
+	static CORE_API float VolumeMultiplier;
 
 	/** Read from config to define the volume when app loses focus */
-	static float UnfocusedVolumeMultiplier;
+	static CORE_API float UnfocusedVolumeMultiplier;
 
 	/** Holds a flag indicating if VRFocus should be used */
-	static bool bUseVRFocus;
+	static CORE_API bool bUseVRFocus;
 
 	/** Holds a flag indicating if app has focus in side the VR headset */
-	static bool bHasVRFocus;
+	static CORE_API bool bHasVRFocus;
 
 	/** Holds a function address that can indicate if application has focus */
-	static bool (*HasFocusFunction)();
+	static CORE_API bool (*HasFocusFunction)();
 };
 
 
diff --git a/Engine/Source/Runtime/Core/Public/Misc/ArchiveMD5.h b/Engine/Source/Runtime/Core/Public/Misc/ArchiveMD5.h
--- a/Engine/Source/Runtime/Core/Public/Misc/ArchiveMD5.h
+++ b/Engine/Source/Runtime/Core/Public/Misc/ArchiveMD5.h
@@ -13,7 +13,7 @@
 /**
  * FArchive adapter for FMD5
  */
-class CORE_API FArchiveMD5 : public FArchive
+class FArchiveMD5 : public FArchive
 {
 public:
 	inline FArchiveMD5()
@@ -23,7 +23,7 @@ public:
 		SetIsPersistent(false);
 	}
 
-	virtual FString GetArchiveName() const;
+	CORE_API virtual FString GetArchiveName() const;
 
 	void Serialize(void* Data, int64 Num) override
 	{
@@ -65,4 +65,4 @@ public:
 
 protected:
 	FMD5 MD5;
-};
\ No newline at end of file
+};
diff --git a/Engine/Source/Runtime/Core/Public/Misc/AssertionMacros.h b/Engine/Source/Runtime/Core/Public/Misc/AssertionMacros.h
--- a/Engine/Source/Runtime/Core/Public/Misc/AssertionMacros.h
+++ b/Engine/Source/Runtime/Core/Public/Misc/AssertionMacros.h
@@ -59,36 +59,36 @@ CORE_API TFunction<bool(const FEnsureHandlerArgs& Args)> GetEnsureHandler();
  * These functions offer debugging and diagnostic functionality and its presence 
  * depends on compiler switches.
  **/
-struct CORE_API FDebug
+struct FDebug
 {
 	/** Logs final assert message and exits the program. */
-	static void VARARGS AssertFailed(const ANSICHAR* Expr, const ANSICHAR* File, int32 Line, const TCHAR* Format = TEXT(""), ...);
-	static void AssertFailedV(const ANSICHAR* Expr, const ANSICHAR* File, int32 Line, const TCHAR* Format, va_list Args);
+	static CORE_API void VARARGS AssertFailed(const ANSICHAR* Expr, const ANSICHAR* File, int32 Line, const TCHAR* Format = TEXT(""), ...);
+	static CORE_API void AssertFailedV(const ANSICHAR* Expr, const ANSICHAR* File, int32 Line, const TCHAR* Format, va_list Args);
 
 	/** Triggers a fatal error, using the error formatted to GErrorHist via a previous call to FMsg*/
-	static void ProcessFatalError(void* ProgramCounter);
+	static CORE_API void ProcessFatalError(void* ProgramCounter);
 
 	// returns true if an assert has occurred
-	static bool HasAsserted();
+	static CORE_API bool HasAsserted();
 
 	// returns true if an ensure is currently in progress (e.g. the RenderThread is ensuring)
-	static bool IsEnsuring();
+	static CORE_API bool IsEnsuring();
 
 	// returns the number of times an ensure has failed in this instance.
-	static SIZE_T GetNumEnsureFailures();
+	static CORE_API SIZE_T GetNumEnsureFailures();
 
 	/** Dumps the stack trace into the log, meant to be used for debugging purposes. */
-	static void DumpStackTraceToLog(const ELogVerbosity::Type LogVerbosity);
+	static CORE_API void DumpStackTraceToLog(const ELogVerbosity::Type LogVerbosity);
 
 	/** Dumps the stack trace into the log with a custom heading, meant to be used for debugging purposes. */
-	static void DumpStackTraceToLog(const TCHAR* Heading, const ELogVerbosity::Type LogVerbosity);
+	static CORE_API void DumpStackTraceToLog(const TCHAR* Heading, const ELogVerbosity::Type LogVerbosity);
 
 #if DO_CHECK || DO_GUARD_SLOW || DO_ENSURE
 public:
-	static bool VARARGS CheckVerifyFailedImpl(const ANSICHAR* Expr, const ANSICHAR* File, int32 Line, void* ProgramCounter, const TCHAR* Format, ...);
+	static CORE_API bool VARARGS CheckVerifyFailedImpl(const ANSICHAR* Expr, const ANSICHAR* File, int32 Line, void* ProgramCounter, const TCHAR* Format, ...);
 private:
-	static void VARARGS LogAssertFailedMessageImpl(const ANSICHAR* Expr, const ANSICHAR* File, int32 Line, void* ProgramCounter, const TCHAR* Fmt, ...);
-	static void LogAssertFailedMessageImplV(const ANSICHAR* Expr, const ANSICHAR* File, int32 Line, void* ProgramCounter, const TCHAR* Fmt, va_list Args);
+	static CORE_API void VARARGS LogAssertFailedMessageImpl(const ANSICHAR* Expr, const ANSICHAR* File, int32 Line, void* ProgramCounter, const TCHAR* Fmt, ...);
+	static CORE_API void LogAssertFailedMessageImplV(const ANSICHAR* Expr, const ANSICHAR* File, int32 Line, void* ProgramCounter, const TCHAR* Fmt, va_list Args);
 
 public:
 //	/**
@@ -108,11 +108,11 @@ public:
 	 * 
 	 * Don't change the name of this function, it's used to detect ensures by the crash reporter.
 	 */
-	static void EnsureFailed( const ANSICHAR* Expr, const ANSICHAR* File, int32 Line, void* ProgramCounter, const TCHAR* Msg );
+	static CORE_API void EnsureFailed( const ANSICHAR* Expr, const ANSICHAR* File, int32 Line, void* ProgramCounter, const TCHAR* Msg );
 
 private:
-	static bool VARARGS OptionallyLogFormattedEnsureMessageReturningFalseImpl(bool bLog, const ANSICHAR* Expr, const ANSICHAR* File, int32 Line, void* ProgramCounter, const TCHAR* FormattedMsg, ...);
-	static bool OptionallyLogFormattedEnsureMessageReturningFalseImpl(bool bLog, const ANSICHAR* Expr, const ANSICHAR* File, int32 Line, void* ProgramCounter, const TCHAR* FormattedMsg, va_list Args);
+	static CORE_API bool VARARGS OptionallyLogFormattedEnsureMessageReturningFalseImpl(bool bLog, const ANSICHAR* Expr, const ANSICHAR* File, int32 Line, void* ProgramCounter, const TCHAR* FormattedMsg, ...);
+	static CORE_API bool OptionallyLogFormattedEnsureMessageReturningFalseImpl(bool bLog, const ANSICHAR* Expr, const ANSICHAR* File, int32 Line, void* ProgramCounter, const TCHAR* FormattedMsg, va_list Args);
 
 public:
 	/**
@@ -157,7 +157,7 @@ public:
 	* @param	Message		Multi-line message with a callstack
 	*
 	*/
-	static void LogFormattedMessageWithCallstack(const FName& LogName, const ANSICHAR* File, int32 Line, const TCHAR* Heading, const TCHAR* Message, ELogVerbosity::Type Verbosity);
+	static CORE_API void LogFormattedMessageWithCallstack(const FName& LogName, const ANSICHAR* File, int32 Line, const TCHAR* Heading, const TCHAR* Message, ELogVerbosity::Type Verbosity);
 };
 
 /*----------------------------------------------------------------------------
diff --git a/Engine/Source/Runtime/Core/Public/Misc/AsyncTaskNotification.h b/Engine/Source/Runtime/Core/Public/Misc/AsyncTaskNotification.h
--- a/Engine/Source/Runtime/Core/Public/Misc/AsyncTaskNotification.h
+++ b/Engine/Source/Runtime/Core/Public/Misc/AsyncTaskNotification.h
@@ -110,18 +110,18 @@ struct FAsyncNotificationStateData
 /**
  * Provides notifications for an on-going asynchronous task.
  */
-class CORE_API FAsyncTaskNotification
+class FAsyncTaskNotification
 {
 public:
 	/**
 	 * Create an asynchronous task notification.
 	 */
-	FAsyncTaskNotification(const FAsyncTaskNotificationConfig& InConfig);
+	CORE_API FAsyncTaskNotification(const FAsyncTaskNotificationConfig& InConfig);
 
 	/**
 	 * Destroy the asynchronous task notification.
 	 */
-	~FAsyncTaskNotification();
+	CORE_API ~FAsyncTaskNotification();
 
 	/**
 	 * Non-copyable.
@@ -132,64 +132,64 @@ public:
 	/**
 	 * Movable.
 	 */
-	FAsyncTaskNotification(FAsyncTaskNotification&& InOther);
-	FAsyncTaskNotification& operator=(FAsyncTaskNotification&& InOther);
+	CORE_API FAsyncTaskNotification(FAsyncTaskNotification&& InOther);
+	CORE_API FAsyncTaskNotification& operator=(FAsyncTaskNotification&& InOther);
 
 	/**
 	 * Set the title text of this notification.
 	 */
-	void SetTitleText(const FText& InTitleText, const bool bClearProgressText = true);
+	CORE_API void SetTitleText(const FText& InTitleText, const bool bClearProgressText = true);
 
 	/**
 	 * Set the progress text of this notification.
 	 */
-	void SetProgressText(const FText& InProgressText);
+	CORE_API void SetProgressText(const FText& InProgressText);
 
 	/**
 	 * Set the prompt text of this notification, if needed.
 	 */
-	void SetPromptText(const FText& InPromptText);
+	CORE_API void SetPromptText(const FText& InPromptText);
 
 	/**
 	 * Set the Hyperlink for this notification.
 	 */
-	void SetHyperlink(const FSimpleDelegate& InHyperlink, const FText& InHyperlinkText = FText());
+	CORE_API void SetHyperlink(const FSimpleDelegate& InHyperlink, const FText& InHyperlinkText = FText());
 
 	/**
 	 * Set the task as complete.
 	 */
-	void SetComplete(const bool bSuccess = true);
+	CORE_API void SetComplete(const bool bSuccess = true);
 
 	/**
 	 * Update the text and set the task as complete.
 	 */
-	void SetComplete(const FText& InTitleText, const FText& InProgressText, const bool bSuccess = true);
+	CORE_API void SetComplete(const FText& InTitleText, const FText& InProgressText, const bool bSuccess = true);
 
 	/**
 	 * Set the notification state. Provide finer control than SetComplete by setting every field of the state.
 	 */
-	void SetNotificationState(const FAsyncNotificationStateData& InState);
+	CORE_API void SetNotificationState(const FAsyncNotificationStateData& InState);
 
 	/**
 	 * Set whether this task be canceled.
 	 */
-	void SetCanCancel(const TAttribute<bool>& InCanCancel);
+	CORE_API void SetCanCancel(const TAttribute<bool>& InCanCancel);
 
 	/**
 	 * Set whether to keep this notification open on success.
 	*/
-	void SetKeepOpenOnSuccess(const TAttribute<bool>& InKeepOpenOnSuccess);
+	CORE_API void SetKeepOpenOnSuccess(const TAttribute<bool>& InKeepOpenOnSuccess);
 
 	/**
 	 * Set whether to keep this notification open on failure.
 	 */
-	void SetKeepOpenOnFailure(const TAttribute<bool>& InKeepOpenOnFailure);
+	CORE_API void SetKeepOpenOnFailure(const TAttribute<bool>& InKeepOpenOnFailure);
 
 	/**
 	 * Return the notification prompt action.
 	 * The action resets to `None` when the notification state changes.
 	 */
-	EAsyncTaskNotificationPromptAction GetPromptAction() const;
+	CORE_API EAsyncTaskNotificationPromptAction GetPromptAction() const;
 
 private:
 	/** Pointer to the real notification implementation */
diff --git a/Engine/Source/Runtime/Core/Public/Misc/AutomationEvent.h b/Engine/Source/Runtime/Core/Public/Misc/AutomationEvent.h
--- a/Engine/Source/Runtime/Core/Public/Misc/AutomationEvent.h
+++ b/Engine/Source/Runtime/Core/Public/Misc/AutomationEvent.h
@@ -13,7 +13,7 @@ enum class EAutomationEventType : uint8
 	Error
 };
 
-struct CORE_API FAutomationEvent
+struct FAutomationEvent
 {
 public:
 	FAutomationEvent(EAutomationEventType InType, const FString& InMessage)
@@ -38,7 +38,7 @@ public:
 	FGuid Artifact;
 };
 
-struct CORE_API FAutomationExecutionEntry
+struct FAutomationExecutionEntry
 {
 	FAutomationEvent Event;
 	FString Filename;
@@ -61,6 +61,6 @@ struct CORE_API FAutomationExecutionEntry
 	{
 	}
 
-	FString ToString() const;
-	FString ToStringFormattedEditorLog() const;
+	CORE_API FString ToString() const;
+	CORE_API FString ToStringFormattedEditorLog() const;
 };
diff --git a/Engine/Source/Runtime/Core/Public/Misc/AutomationTest.h b/Engine/Source/Runtime/Core/Public/Misc/AutomationTest.h
--- a/Engine/Source/Runtime/Core/Public/Misc/AutomationTest.h
+++ b/Engine/Source/Runtime/Core/Public/Misc/AutomationTest.h
@@ -186,7 +186,7 @@ struct FAutomationTelemetryData
 };
 
 /** Simple class to store the results of the execution of a automation test */
-class CORE_API FAutomationTestExecutionInfo
+class FAutomationTestExecutionInfo
 {
 public:
 	/** Constructor */
@@ -204,19 +204,19 @@ public:
 	}
 
 	/** Helper method to clear out the results from a previous execution */
-	void Clear();
+	CORE_API void Clear();
 
-	int32 RemoveAllEvents(EAutomationEventType EventType);
+	CORE_API int32 RemoveAllEvents(EAutomationEventType EventType);
 
-	int32 RemoveAllEvents(TFunctionRef<bool(FAutomationEvent&)> FilterPredicate);
+	CORE_API int32 RemoveAllEvents(TFunctionRef<bool(FAutomationEvent&)> FilterPredicate);
 
 	/** Any errors that occurred during execution */
 	const TArray<FAutomationExecutionEntry>& GetEntries() const { return Entries; }
 
-	void AddEvent(const FAutomationEvent& Event, int StackOffset = 0, bool bCaptureStack = true);
+	CORE_API void AddEvent(const FAutomationEvent& Event, int StackOffset = 0, bool bCaptureStack = true);
 
-	void AddWarning(const FString& WarningMessage);
-	void AddError(const FString& ErrorMessage);
+	CORE_API void AddWarning(const FString& WarningMessage);
+	CORE_API void AddError(const FString& ErrorMessage);
 
 	int32 GetWarningTotal() const { return Warnings; }
 	int32 GetErrorTotal() const { return Errors; }
@@ -268,7 +268,7 @@ private:
 };
 
 /** Simple class to store the automation test info */
-class CORE_API FAutomationTestInfo
+class FAutomationTestInfo
 {
 public:
 
@@ -728,7 +728,7 @@ struct FAutomationScreenshotData
 	}
 };
 
-struct CORE_API FAutomationScreenshotCompareResults
+struct FAutomationScreenshotCompareResults
 {
 	FGuid UniqueId;
 	FString ErrorMessage;
@@ -737,7 +737,7 @@ struct CORE_API FAutomationScreenshotCompareResults
 	bool bWasNew = false;
 	bool bWasSimilar = false;
 
-	FAutomationEvent ToAutomationEvent(const FString& ScreenhotName) const;
+	CORE_API FAutomationEvent ToAutomationEvent(const FString& ScreenhotName) const;
 };
 
 enum class EAutomationComparisonToleranceLevel : uint8
@@ -814,7 +814,7 @@ DECLARE_MULTICAST_DELEGATE_TwoParams(FOnPerformanceDataRetrieved, bool /*bSucces
 DECLARE_MULTICAST_DELEGATE_OneParam(FOnTestEvent, FAutomationTestBase*);
 
 /** Class representing the main framework for running automation tests */
-class CORE_API FAutomationTestFramework
+class FAutomationTestFramework
 {
 public:
 	/** Called right before automated test is about to begin */
@@ -846,7 +846,7 @@ public:
 	 *
 	 * @return The singleton instance of the framework.
 	 */
-	static FAutomationTestFramework& Get();
+	static CORE_API FAutomationTestFramework& Get();
 	static FAutomationTestFramework& GetInstance() { return Get(); }
 
 	/**
@@ -854,7 +854,7 @@ public:
 	 * to do something like generate project files, or create new projects it should use this directory, rather
 	 * than pollute other areas of the machine.
 	 */
-	FString GetUserAutomationDirectory() const;
+	CORE_API FString GetUserAutomationDirectory() const;
 
 	/**
 	 * Register a automation test into the framework. The automation test may or may not be necessarily valid
@@ -867,28 +867,28 @@ public:
 	 * @return	true if the test was successfully registered; false if a test was already registered under the same
 	 *			name as before
 	 */
-	bool RegisterAutomationTest( const FString& InTestNameToRegister, FAutomationTestBase* InTestToRegister );
+	CORE_API bool RegisterAutomationTest( const FString& InTestNameToRegister, FAutomationTestBase* InTestToRegister );
 
 	/**
 	 * Unregister a automation test with the provided name from the framework.
 	 *
 	 * @return true if the test was successfully unregistered; false if a test with that name was not found in the framework.
 	 */
-	bool UnregisterAutomationTest( const FString& InTestNameToUnregister );
+	CORE_API bool UnregisterAutomationTest( const FString& InTestNameToUnregister );
 
 	/**
 	 * Enqueues a latent command for execution on a subsequent frame
 	 *
 	 * @param NewCommand - The new command to enqueue for deferred execution
 	 */
-	void EnqueueLatentCommand(TSharedPtr<IAutomationLatentCommand> NewCommand);
+	CORE_API void EnqueueLatentCommand(TSharedPtr<IAutomationLatentCommand> NewCommand);
 
 	/**
 	 * Enqueues a network command for execution in accordance with this workers role
 	 *
 	 * @param NewCommand - The new command to enqueue for network execution
 	 */
-	void EnqueueNetworkCommand(TSharedPtr<IAutomationNetworkCommand> NewCommand);
+	CORE_API void EnqueueNetworkCommand(TSharedPtr<IAutomationNetworkCommand> NewCommand);
 
 	/**
 	 * Checks if a provided test is contained within the framework.
@@ -897,19 +897,19 @@ public:
 	 *
 	 * @return	true if the provided test is within the framework; false otherwise
 	 */
-	bool ContainsTest( const FString& InTestName ) const;
+	CORE_API bool ContainsTest( const FString& InTestName ) const;
 		
 	/**
 	 * Attempt to run all fast smoke tests that are valid for the current application configuration.
 	 *
 	 * @return	true if all smoke tests run were successful, false if any failed
 	 */
-	bool RunSmokeTests();
+	CORE_API bool RunSmokeTests();
 
 	/**
 	 * Reset status of worker (delete local files, etc)
 	 */
-	void ResetTests();
+	CORE_API void ResetTests();
 
 	/**
 	 * Attempt to start the specified test.
@@ -917,33 +917,33 @@ public:
 	 * @param	InTestToRun			Name of the test that should be run
 	 * @param	InRoleIndex			Identifier for which worker in this group that should execute a command
 	 */
-	void StartTestByName( const FString& InTestToRun, const int32 InRoleIndex );
+	CORE_API void StartTestByName( const FString& InTestToRun, const int32 InRoleIndex );
 
 	/**
 	 * Stop the current test and return the results of execution
 	 *
 	 * @return	true if the test ran successfully, false if it did not (or the test could not be found/was invalid)
 	 */
-	bool StopTest( FAutomationTestExecutionInfo& OutExecutionInfo );
+	CORE_API bool StopTest( FAutomationTestExecutionInfo& OutExecutionInfo );
 
 	/**
 	 * Execute all latent functions that complete during update
 	 *
 	 * @return - true if the latent command queue is now empty and the test is complete
 	 */
-	bool ExecuteLatentCommands();
+	CORE_API bool ExecuteLatentCommands();
 
 	/**
 	 * Execute the next network command if you match the role, otherwise just dequeue
 	 *
 	 * @return - true if any network commands were in the queue to give subsequent latent commands a chance to execute next frame
 	 */
-	bool ExecuteNetworkCommands();
+	CORE_API bool ExecuteNetworkCommands();
 
 	/**
 	 * Dequeue all latent and network commands
 	 */
-	void DequeueAllCommands();
+	CORE_API void DequeueAllCommands();
 
 	/**
 	 * Whether there is no latent command in queue
@@ -956,7 +956,7 @@ public:
 	/**
 	 * Load any modules that are not loaded by default and have test classes in them
 	 */
-	void LoadTestModules();
+	CORE_API void LoadTestModules();
 
 	/**
 	 * Populates the provided array with the names of all tests in the framework that are valid to run for the current
@@ -964,35 +964,35 @@ public:
 	 *
 	 * @param	TestInfo	Array to populate with the test information
 	 */
-	void GetValidTestNames( TArray<FAutomationTestInfo>& TestInfo ) const;
+	CORE_API void GetValidTestNames( TArray<FAutomationTestInfo>& TestInfo ) const;
 
 	/**
 	 * Whether the testing framework should allow content to be tested or not.  Intended to block developer directories.
 	 * @param Path - Full path to the content in question
 	 * @return - Whether this content should have tests performed on it
 	 */
-	bool ShouldTestContent(const FString& Path) const;
+	CORE_API bool ShouldTestContent(const FString& Path) const;
 
 	/**
 	 * Sets whether we want to include content in developer directories in automation testing
 	 */
-	void SetDeveloperDirectoryIncluded(const bool bInDeveloperDirectoryIncluded);
+	CORE_API void SetDeveloperDirectoryIncluded(const bool bInDeveloperDirectoryIncluded);
 
 	/**
 	* Sets which set of tests to pull from.
 	*/
-	void SetRequestedTestFilter(const uint32 InRequestedTestFlags);
+	CORE_API void SetRequestedTestFilter(const uint32 InRequestedTestFlags);
 	
 
 	/**
 	 * Accessor for delegate called when a png screenshot is captured 
 	 */
-	FOnTestScreenshotCaptured& OnScreenshotCaptured();
+	CORE_API FOnTestScreenshotCaptured& OnScreenshotCaptured();
 
 	/**
 	 * Accessor for delegate called when a png screenshot is captured and a frame trace
 	 */
-	FOnTestScreenshotAndTraceCaptured& OnScreenshotAndTraceCaptured();
+	CORE_API FOnTestScreenshotAndTraceCaptured& OnScreenshotAndTraceCaptured();
 
 	/**
 	 * Sets forcing smoke tests.
@@ -1021,7 +1021,7 @@ public:
 	 *
 	 * @param	AnalyticsItem	Log item to add to the current test
 	 */
-	void AddAnalyticsItemToCurrentTest( const FString& AnalyticsItem );
+	CORE_API void AddAnalyticsItemToCurrentTest( const FString& AnalyticsItem );
 
 	/**
 	 * Returns the actively executing test or null if there isn't one
@@ -1034,18 +1034,18 @@ public:
 	/**
 	 * Whether to skip stack walk while iterating for listing the tests
 	 */
-	static bool NeedSkipStackWalk();
+	static CORE_API bool NeedSkipStackWalk();
 
 	/**
 	 * Whether to output blueprint functional test metadata to the log when test is running
 	 */
-	static bool NeedLogBPTestMetadata();
+	static CORE_API bool NeedLogBPTestMetadata();
 
-	void NotifyScreenshotComparisonComplete(const FAutomationScreenshotCompareResults& CompareResults);
-	void NotifyTestDataRetrieved(bool bWasNew, const FString& JsonData);
-	void NotifyPerformanceDataRetrieved(bool bSuccess, const FString& ErrorMessage);
+	CORE_API void NotifyScreenshotComparisonComplete(const FAutomationScreenshotCompareResults& CompareResults);
+	CORE_API void NotifyTestDataRetrieved(bool bWasNew, const FString& JsonData);
+	CORE_API void NotifyPerformanceDataRetrieved(bool bSuccess, const FString& ErrorMessage);
 
-	void NotifyScreenshotTakenAndCompared();
+	CORE_API void NotifyScreenshotTakenAndCompared();
 
 	/**
 	 * Internal helper method designed to check if the given test is able to run in the current environment.
@@ -1056,7 +1056,7 @@ public:
 	 *
 	 * @return	true if the test is able to run; false if it is unable to run.
 	 */
-	bool CanRunTestInEnvironment(const FString& InTestToRun, FString* OutReason, bool* OutWarn) const;
+	CORE_API bool CanRunTestInEnvironment(const FString& InTestToRun, FString* OutReason, bool* OutWarn) const;
 
 private:
 
@@ -1178,10 +1178,10 @@ private:
 
 	friend class FAutomationTestOutputDevice;
 	/** Helper method called to prepare settings for automation testing to follow */
-	void PrepForAutomationTests();
+	CORE_API void PrepForAutomationTests();
 
 	/** Helper method called after automation testing is complete to restore settings to how they should be */
-	void ConcludeAutomationTests();
+	CORE_API void ConcludeAutomationTests();
 
 	/**
 	 * Helper method to dump the contents of the provided test name to execution info map to the provided feedback context
@@ -1189,7 +1189,7 @@ private:
 	 * @param	InContext		Context to dump the execution info to
 	 * @param	InInfoToDump	Execution info that should be dumped to the provided feedback context
 	 */
-	void DumpAutomationTestExecutionInfo( const TMap<FString, FAutomationTestExecutionInfo>& InInfoToDump );
+	CORE_API void DumpAutomationTestExecutionInfo( const TMap<FString, FAutomationTestExecutionInfo>& InInfoToDump );
 
 	/**
 	 * Internal helper method designed to simply start the provided test name.
@@ -1197,7 +1197,7 @@ private:
 	 * @param	InTestToRun			Name of the test that should be run
 	 * @param	OutExecutionInfo	Results of executing the test
 	 */
-	void InternalStartTest( const FString& InTestToRun );
+	CORE_API void InternalStartTest( const FString& InTestToRun );
 
 	/**
 	 * Internal helper method designed to stop current executing test and return the results of execution.
@@ -1205,17 +1205,17 @@ private:
 	 * @return	true if the test was successfully run; false if it was not, could not be found, or is invalid for
 	 *			the current application settings
 	 */
-	bool InternalStopTest(FAutomationTestExecutionInfo& OutExecutionInfo);
+	CORE_API bool InternalStopTest(FAutomationTestExecutionInfo& OutExecutionInfo);
 
 	/** Constructor */
-	FAutomationTestFramework();
+	CORE_API FAutomationTestFramework();
 
 	/** Destructor */
-	~FAutomationTestFramework();
+	CORE_API ~FAutomationTestFramework();
 
 	// Copy constructor and assignment operator intentionally left unimplemented
-	FAutomationTestFramework( const FAutomationTestFramework& );
-	FAutomationTestFramework& operator=( const FAutomationTestFramework& );
+	CORE_API FAutomationTestFramework( const FAutomationTestFramework& );
+	CORE_API FAutomationTestFramework& operator=( const FAutomationTestFramework& );
 
 	/** Specialized output device used for automation testing */
 	FAutomationTestOutputDevice AutomationTestOutputDevice;
@@ -1268,7 +1268,7 @@ private:
 };
 
 /** Simple abstract base class for all automation tests */
-class CORE_API FAutomationTestBase
+class FAutomationTestBase
 {
 public:
 	/**
@@ -1293,10 +1293,10 @@ public:
 	}
 
 	/** Log flags */
-	static bool bSuppressLogWarnings;
-	static bool bSuppressLogErrors;
-	static bool bElevateLogWarningsToErrors;
-	static TArray<FString> SuppressedLogCategories;
+	static CORE_API bool bSuppressLogWarnings;
+	static CORE_API bool bSuppressLogErrors;
+	static CORE_API bool bElevateLogWarningsToErrors;
+	static CORE_API TArray<FString> SuppressedLogCategories;
 
 	/**
 	 * Pure virtual method; returns the flags associated with the given automation test
@@ -1327,14 +1327,14 @@ public:
 	virtual uint32 GetRequiredDeviceNum() const = 0;
 
 	/** Clear any execution info/results from a prior running of this test */
-	void ClearExecutionInfo();
+	CORE_API void ClearExecutionInfo();
 
 	/**
 	 * Adds an error message to this test
 	 *
 	 * @param	InError	Error message to add to this test
 	 */
-	virtual void AddError( const FString& InError, int32 StackOffset = 0 );
+	CORE_API virtual void AddError( const FString& InError, int32 StackOffset = 0 );
 
 	/**
 	 * Adds an error message to this test if the condition is false
@@ -1343,7 +1343,7 @@ public:
 	 * @param   InError	   Error message to add to this test
 	 * @return	False if there was an error
 	 */
-	virtual bool AddErrorIfFalse( bool bCondition, const FString& InError, int32 StackOffset = 0 );
+	CORE_API virtual bool AddErrorIfFalse( bool bCondition, const FString& InError, int32 StackOffset = 0 );
 
 	/**
 	 * Adds an error message to this test
@@ -1352,7 +1352,7 @@ public:
 	 * @param	InFilename	The filename the error originated in
 	 * @param	InLineNumber	The line number in the file this error originated in
 	 */
-	virtual void AddErrorS(const FString& InError, const FString& InFilename, int32 InLineNumber);
+	CORE_API virtual void AddErrorS(const FString& InError, const FString& InFilename, int32 InLineNumber);
 
 	/**
 	 * Adds an warning message to this test
@@ -1361,35 +1361,35 @@ public:
 	 * @param	InFilename	The filename the error originated in
 	 * @param	InLineNumber	The line number in the file this error originated in
 	 */
-	virtual void AddWarningS(const FString& InWarning, const FString& InFilename, int32 InLineNumber);
+	CORE_API virtual void AddWarningS(const FString& InWarning, const FString& InFilename, int32 InLineNumber);
 
 	/**
 	 * Adds a warning to this test
 	 *
 	 * @param	InWarning	Warning message to add to this test
 	 */
-	virtual void AddWarning( const FString& InWarning, int32 StackOffset = 0);
+	CORE_API virtual void AddWarning( const FString& InWarning, int32 StackOffset = 0);
 
 	/**
 	 * Adds a log item to this test
 	 *
 	 * @param	InLogItem	Log item to add to this test
 	 */
-	virtual void AddInfo( const FString& InLogItem, int32 StackOffset = 0, bool bCaptureStack = false);
+	CORE_API virtual void AddInfo( const FString& InLogItem, int32 StackOffset = 0, bool bCaptureStack = false);
 
 	/**
 	 * Adds an automation event directly into the execution log.
 	 *
 	 * @param	InLogItem	Log item to add to this test
 	 */
-	virtual void AddEvent(const FAutomationEvent& InEvent, int32 StackOffset = 0, bool bCaptureStack = false);
+	CORE_API virtual void AddEvent(const FAutomationEvent& InEvent, int32 StackOffset = 0, bool bCaptureStack = false);
 
 	/**
 	 * Adds a analytics string to parse later
 	 *
 	 * @param	InLogItem	Log item to add to this test
 	 */
-	virtual void AddAnalyticsItem(const FString& InAnalyticsItem);
+	CORE_API virtual void AddAnalyticsItem(const FString& InAnalyticsItem);
 
 	/**
 	 * Adds a telemetry data point measurement
@@ -1398,7 +1398,7 @@ public:
 	 * @param	Measurement	Value to associate to the data point
 	 * @param	Context		optional context associated with the data point
 	 */
-	virtual void AddTelemetryData(const FString& DataPoint, double Measurement, const FString& Context = TEXT(""));
+	CORE_API virtual void AddTelemetryData(const FString& DataPoint, double Measurement, const FString& Context = TEXT(""));
 
 	/**
 	 * Adds several telemetry data point measurements
@@ -1406,40 +1406,40 @@ public:
 	 * @param	ValuePairs	value pair of Name and Measurement of several Data points
 	 * @param	Context		optional context associated with the data point
 	 */
-	virtual void AddTelemetryData(const TMap<FString, double>& ValuePairs, const FString& Context = TEXT(""));
+	CORE_API virtual void AddTelemetryData(const TMap<FString, double>& ValuePairs, const FString& Context = TEXT(""));
 
 	/**
 	 * Set telemetry storage name
 	 *
 	 * @param	StorageName	Name of the data storage
 	 */
-	virtual void SetTelemetryStorage(const FString& StorageName);
+	CORE_API virtual void SetTelemetryStorage(const FString& StorageName);
 
 	/**
 	 * Returns whether this test has any errors associated with it or not
 	 *
 	 * @return true if this test has at least one error associated with it; false if not
 	 */
-	bool HasAnyErrors() const;
+	CORE_API bool HasAnyErrors() const;
 
 	/**
 	* Returns whether this test has encountered all expected log messages defined for it
 	* @param VerbosityType Optionally specify to check by log level. Defaults to all.
 	* @return true if this test has encountered all expected messages; false if not
 	*/
-	bool HasMetExpectedMessages(ELogVerbosity::Type VerbosityType = ELogVerbosity::All);
+	CORE_API bool HasMetExpectedMessages(ELogVerbosity::Type VerbosityType = ELogVerbosity::All);
 
 	/**
 	* Returns whether this test has encountered all expected errors defined for it
 	*
 	* @return true if this test has encountered all expected errors; false if not
 	*/
-	bool HasMetExpectedErrors();
+	CORE_API bool HasMetExpectedErrors();
 
 	/**
 	 * Return the last success state for this test
 	 */
-	bool GetLastExecutionSuccessState();
+	CORE_API bool GetLastExecutionSuccessState();
 
 	/**
 	 * [Deprecated] Use AddError(msg) instead to change the state of the test to a failure
@@ -1459,22 +1459,22 @@ public:
 	 *
 	 * @param	OutInfo	Execution info to be populated with the same data contained within this test's execution info
 	 */
-	void GetExecutionInfo( FAutomationTestExecutionInfo& OutInfo ) const;
+	CORE_API void GetExecutionInfo( FAutomationTestExecutionInfo& OutInfo ) const;
 
 	/** 
 	 * Helper function that will generate a list of sub-tests via GetTests
 	 */
-	void GenerateTestNames( TArray<FAutomationTestInfo>& TestInfo ) const;
+	CORE_API void GenerateTestNames( TArray<FAutomationTestInfo>& TestInfo ) const;
 
 	/**
 	 * Helper function that determines if the given log category matches the expected category, inclusively (so an Error counts as a Warning)
 	*/
-	static bool LogCategoryMatchesSeverityInclusive(ELogVerbosity::Type Actual, ELogVerbosity::Type MaximumVerbosity);
+	static CORE_API bool LogCategoryMatchesSeverityInclusive(ELogVerbosity::Type Actual, ELogVerbosity::Type MaximumVerbosity);
 
 	/**
 	 * Enables log settings from config
 	*/
-	static void LoadDefaultLogSettings();
+	static CORE_API void LoadDefaultLogSettings();
 	/**
 	
 	* Adds a regex pattern to an internal list that this test will expect to encounter in logs (of the specified verbosity) during its execution. If an expected pattern
@@ -1486,7 +1486,7 @@ public:
 	* @param Occurrences - How many times to expect this message string to be seen. If > 0, the message must be seen the exact number of times
 	* specified or the test will fail. If == 0, the message must be seen one or more times (with no upper limit) or the test will fail.
 	*/
-	void AddExpectedMessage(FString ExpectedPatternString, ELogVerbosity::Type ExpectedVerbosity, EAutomationExpectedMessageFlags::MatchType CompareType = EAutomationExpectedMessageFlags::Contains, int32 Occurrences = 1);
+	CORE_API void AddExpectedMessage(FString ExpectedPatternString, ELogVerbosity::Type ExpectedVerbosity, EAutomationExpectedMessageFlags::MatchType CompareType = EAutomationExpectedMessageFlags::Contains, int32 Occurrences = 1);
 	
 	/**
 	* Adds a regex pattern to an internal list that this test will expect to encounter in logs (of all severities) during its execution. If an expected pattern
@@ -1497,7 +1497,7 @@ public:
 	* @param Occurrences - How many times to expect this message string to be seen. If > 0, the message must be seen the exact number of times
 	* specified or the test will fail. If == 0, the message must be seen one or more times (with no upper limit) or the test will fail.
 	*/
-	void AddExpectedMessage(FString ExpectedPatternString, EAutomationExpectedMessageFlags::MatchType CompareType = EAutomationExpectedMessageFlags::Contains, int32 Occurrences = 1);
+	CORE_API void AddExpectedMessage(FString ExpectedPatternString, EAutomationExpectedMessageFlags::MatchType CompareType = EAutomationExpectedMessageFlags::Contains, int32 Occurrences = 1);
 
 	/**
 	* Populate the provided expected log messages object with the expected messages contained within the test. Not particularly efficient,
@@ -1505,7 +1505,7 @@ public:
 	* @param Verbosity - Optionally filter the returned messages by verbosity. This is inclusive, so Warning will return Warnings, Errors, etc.
 	* @param OutInfo - Array of Expected Messages to be populated with the same data contained within this test's expected messages list
 	*/
-	void GetExpectedMessages(TArray<FAutomationExpectedMessage>& OutInfo, ELogVerbosity::Type Verbosity = ELogVerbosity::All) const;
+	CORE_API void GetExpectedMessages(TArray<FAutomationExpectedMessage>& OutInfo, ELogVerbosity::Type Verbosity = ELogVerbosity::All) const;
 
 	/**
 	* Adds a regex pattern to an internal list that this test will expect to encounter in error or warning logs during its execution. If an expected pattern
@@ -1516,7 +1516,7 @@ public:
 	* @param Occurrences - How many times to expect this error string to be seen. If > 0, the error must be seen the exact number of times
 	* specified or the test will fail. If == 0, the error must be seen one or more times (with no upper limit) or the test will fail.
 	*/
-	void AddExpectedError(FString ExpectedPatternString, EAutomationExpectedErrorFlags::MatchType CompareType = EAutomationExpectedErrorFlags::Contains, int32 Occurrences = 1);
+	CORE_API void AddExpectedError(FString ExpectedPatternString, EAutomationExpectedErrorFlags::MatchType CompareType = EAutomationExpectedErrorFlags::Contains, int32 Occurrences = 1);
 
 	/**
 	 * Is this a complex tast - if so it will be a stress test.
@@ -1635,20 +1635,20 @@ public:
 
 public:
 
-	bool TestEqual(const TCHAR* What, const int32 Actual, const int32 Expected);
-	bool TestEqual(const TCHAR* What, const int64 Actual, const int64 Expected);
+	CORE_API bool TestEqual(const TCHAR* What, const int32 Actual, const int32 Expected);
+	CORE_API bool TestEqual(const TCHAR* What, const int64 Actual, const int64 Expected);
 #if PLATFORM_64BITS
-	bool TestEqual(const TCHAR* What, const SIZE_T Actual, const SIZE_T Expected);
+	CORE_API bool TestEqual(const TCHAR* What, const SIZE_T Actual, const SIZE_T Expected);
 #endif
-	bool TestEqual(const TCHAR* What, const float Actual, const float Expected, float Tolerance = UE_KINDA_SMALL_NUMBER);
-	bool TestEqual(const TCHAR* What, const double Actual, const double Expected, double Tolerance = UE_KINDA_SMALL_NUMBER);
-	bool TestEqual(const TCHAR* What, const FVector Actual, const FVector Expected, float Tolerance = UE_KINDA_SMALL_NUMBER);
-	bool TestEqual(const TCHAR* What, const FTransform Actual, const FTransform Expected, float Tolerance = UE_KINDA_SMALL_NUMBER);
-	bool TestEqual(const TCHAR* What, const FRotator Actual, const FRotator Expected, float Tolerance = UE_KINDA_SMALL_NUMBER);
-	bool TestEqual(const TCHAR* What, const FColor Actual, const FColor Expected);
-	bool TestEqual(const TCHAR* What, const FLinearColor Actual, const FLinearColor Expected);
-	bool TestEqual(const TCHAR* What, const TCHAR* Actual, const TCHAR* Expected);
-	bool TestEqualInsensitive(const TCHAR* What, const TCHAR* Actual, const TCHAR* Expected);
+	CORE_API bool TestEqual(const TCHAR* What, const float Actual, const float Expected, float Tolerance = UE_KINDA_SMALL_NUMBER);
+	CORE_API bool TestEqual(const TCHAR* What, const double Actual, const double Expected, double Tolerance = UE_KINDA_SMALL_NUMBER);
+	CORE_API bool TestEqual(const TCHAR* What, const FVector Actual, const FVector Expected, float Tolerance = UE_KINDA_SMALL_NUMBER);
+	CORE_API bool TestEqual(const TCHAR* What, const FTransform Actual, const FTransform Expected, float Tolerance = UE_KINDA_SMALL_NUMBER);
+	CORE_API bool TestEqual(const TCHAR* What, const FRotator Actual, const FRotator Expected, float Tolerance = UE_KINDA_SMALL_NUMBER);
+	CORE_API bool TestEqual(const TCHAR* What, const FColor Actual, const FColor Expected);
+	CORE_API bool TestEqual(const TCHAR* What, const FLinearColor Actual, const FLinearColor Expected);
+	CORE_API bool TestEqual(const TCHAR* What, const TCHAR* Actual, const TCHAR* Expected);
+	CORE_API bool TestEqualInsensitive(const TCHAR* What, const TCHAR* Actual, const TCHAR* Expected);
 
 	bool TestEqual(const FString& What, const int32 Actual, const int32 Expected)
 	{
@@ -1746,11 +1746,11 @@ public:
 		return TestEqual(*What, Actual, Expected);
 	}
 
-	bool TestNearlyEqual(const TCHAR* What, const float Actual, const float Expected, float Tolerance = UE_KINDA_SMALL_NUMBER);
-	bool TestNearlyEqual(const TCHAR* What, const double Actual, const double Expected, double Tolerance = UE_KINDA_SMALL_NUMBER);
-	bool TestNearlyEqual(const TCHAR* What, const FVector Actual, const FVector Expected, float Tolerance = UE_KINDA_SMALL_NUMBER);
-	bool TestNearlyEqual(const TCHAR* What, const FTransform Actual, const FTransform Expected, float Tolerance = UE_KINDA_SMALL_NUMBER);
-	bool TestNearlyEqual(const TCHAR* What, const FRotator Actual, const FRotator Expected, float Tolerance = UE_KINDA_SMALL_NUMBER);
+	CORE_API bool TestNearlyEqual(const TCHAR* What, const float Actual, const float Expected, float Tolerance = UE_KINDA_SMALL_NUMBER);
+	CORE_API bool TestNearlyEqual(const TCHAR* What, const double Actual, const double Expected, double Tolerance = UE_KINDA_SMALL_NUMBER);
+	CORE_API bool TestNearlyEqual(const TCHAR* What, const FVector Actual, const FVector Expected, float Tolerance = UE_KINDA_SMALL_NUMBER);
+	CORE_API bool TestNearlyEqual(const TCHAR* What, const FTransform Actual, const FTransform Expected, float Tolerance = UE_KINDA_SMALL_NUMBER);
+	CORE_API bool TestNearlyEqual(const TCHAR* What, const FRotator Actual, const FRotator Expected, float Tolerance = UE_KINDA_SMALL_NUMBER);
 
 	bool TestNearlyEqual(const FString& What, const float Actual, const float Expected, float Tolerance = UE_KINDA_SMALL_NUMBER)
 	{
@@ -1786,7 +1786,7 @@ public:
 	 *
 	 * @see TestFalse
 	 */
-	bool TestFalse(const TCHAR* What, bool Value);
+	CORE_API bool TestFalse(const TCHAR* What, bool Value);
 
 	bool TestFalse(const FString& What, bool Value)
 	{
@@ -1899,7 +1899,7 @@ public:
 	 *
 	 * @see TestNotNull
 	 */
-	bool TestNull(const TCHAR* What, const void* Pointer);
+	CORE_API bool TestNull(const TCHAR* What, const void* Pointer);
 
 	bool TestNull(const FString& What, const void* Pointer)
 	{
@@ -1939,7 +1939,7 @@ public:
 	 *
 	 * @see TestFalse
 	 */
-	bool TestTrue(const TCHAR* What, bool Value);
+	CORE_API bool TestTrue(const TCHAR* What, bool Value);
 
 	bool TestTrue(const FString& What, bool Value)
 	{
@@ -1999,7 +1999,7 @@ protected:
 	virtual void SetTestContext(FString Context) { TestParameterContext = Context; }
 
 	/** Extracts a combined EAutomationTestFlags value from a string representation using tag notation "[Filter_1]...[Filter_n][Tag_1]...[Tag_m]" */
-	uint32 ExtractAutomationTestFlags(FString InTagNotation);
+	CORE_API uint32 ExtractAutomationTestFlags(FString InTagNotation);
 
 protected:
 
@@ -2031,14 +2031,14 @@ private:
 	*
 	* @return true if this message matches any of the expected messages
 	*/
-	bool IsExpectedMessage(const FString& Message, const ELogVerbosity::Type& Verbosity = ELogVerbosity::All);
+	CORE_API bool IsExpectedMessage(const FString& Message, const ELogVerbosity::Type& Verbosity = ELogVerbosity::All);
 
 	/**
 	 * Sets whether the test has succeeded or not
 	 *
 	 * @param	bSuccessful	true to mark the test successful, false to mark the test as failed
 	 */
-	void InternalSetSuccessState(bool bSuccessful);
+	CORE_API void InternalSetSuccessState(bool bSuccessful);
 
 	/* Log messages to be expected while processing this test.*/
 	TArray<FAutomationExpectedMessage> ExpectedMessages;
@@ -2047,7 +2047,7 @@ private:
 	FCriticalSection ActionCS;
 };
 
-class CORE_API FBDDAutomationTestBase : public FAutomationTestBase
+class FBDDAutomationTestBase : public FAutomationTestBase
 { 
 public:
 	FBDDAutomationTestBase(const FString& InName, const bool bInComplexTask)
@@ -2231,7 +2231,7 @@ private:
 
 DECLARE_DELEGATE(FDoneDelegate);
 
-class CORE_API FAutomationSpecBase 
+class FAutomationSpecBase 
 	: public FAutomationTestBase
 	, public TSharedFromThis<FAutomationSpecBase>
 {
diff --git a/Engine/Source/Runtime/Core/Public/Misc/Base64.h b/Engine/Source/Runtime/Core/Public/Misc/Base64.h
--- a/Engine/Source/Runtime/Core/Public/Misc/Base64.h
+++ b/Engine/Source/Runtime/Core/Public/Misc/Base64.h
@@ -19,7 +19,7 @@ enum class EBase64Mode : uint8
 /**
  * Class for encoding/decoding Base64 data (RFC 4648)
  */
-struct CORE_API FBase64
+struct FBase64
 {
 	/**
 	 * Encodes a FString into a Base64 string
@@ -29,7 +29,7 @@ struct CORE_API FBase64
 	 *
 	 * @return A string that encodes the binary data in a way that can be safely transmitted via various Internet protocols
 	 */
-	static FString Encode(const FString& Source, EBase64Mode Mode = EBase64Mode::Standard);
+	static CORE_API FString Encode(const FString& Source, EBase64Mode Mode = EBase64Mode::Standard);
 
 	/**
 	 * Encodes a binary uint8 array into a Base64 string
@@ -39,7 +39,7 @@ struct CORE_API FBase64
 	 *
 	 * @return A string that encodes the binary data in a way that can be safely transmitted via various Internet protocols
 	 */
-	static FString Encode(const TArray<uint8>& Source, EBase64Mode Mode = EBase64Mode::Standard);
+	static CORE_API FString Encode(const TArray<uint8>& Source, EBase64Mode Mode = EBase64Mode::Standard);
 
 	/**
 	 * Encodes the source into a Base64 string
@@ -50,7 +50,7 @@ struct CORE_API FBase64
 	 *
 	 * @return Base64 encoded string containing the binary data.
 	 */
-	static FString Encode(const uint8* Source, uint32 Length, EBase64Mode Mode = EBase64Mode::Standard);
+	static CORE_API FString Encode(const uint8* Source, uint32 Length, EBase64Mode Mode = EBase64Mode::Standard);
 
 	/**
 	 * Encodes the source into a Base64 string, storing it in a preallocated buffer.
@@ -85,7 +85,7 @@ struct CORE_API FBase64
 	 *
 	 * @return true if the buffer was decoded, false if it was invalid.
 	 */
-	static bool Decode(const FString& Source, FString& OutDest, EBase64Mode Mode = EBase64Mode::Standard);
+	static CORE_API bool Decode(const FString& Source, FString& OutDest, EBase64Mode Mode = EBase64Mode::Standard);
 
 	/**
 	 * Decodes a Base64 string into an array of bytes
@@ -96,7 +96,7 @@ struct CORE_API FBase64
 	 *
 	 * @return true if the buffer was decoded, false if it was invalid.
 	 */
-	static bool Decode(const FString& Source, TArray<uint8>& Dest, EBase64Mode Mode = EBase64Mode::Standard);
+	static CORE_API bool Decode(const FString& Source, TArray<uint8>& Dest, EBase64Mode Mode = EBase64Mode::Standard);
 
 	/**
 	 * Decodes a Base64 string into a preallocated buffer
@@ -117,7 +117,7 @@ struct CORE_API FBase64
 	*
 	* @return The size in bytes of the decoded data
 	*/
-	static uint32 GetDecodedDataSize(const FString& Source);
+	static CORE_API uint32 GetDecodedDataSize(const FString& Source);
 
 	/**
 	* Determine the decoded data size for the incoming base64 encoded string
@@ -140,4 +140,4 @@ struct CORE_API FBase64
 	{
 		return ((NumChars * 3) + 3) / 4;
 	}
-};
\ No newline at end of file
+};
diff --git a/Engine/Source/Runtime/Core/Public/Misc/BufferedOutputDevice.h b/Engine/Source/Runtime/Core/Public/Misc/BufferedOutputDevice.h
--- a/Engine/Source/Runtime/Core/Public/Misc/BufferedOutputDevice.h
+++ b/Engine/Source/Runtime/Core/Public/Misc/BufferedOutputDevice.h
@@ -10,7 +10,7 @@
 #include "Misc/OutputDeviceRedirector.h"
 
 /** Buffered output device. */
-class CORE_API FBufferedOutputDevice : public FOutputDevice
+class FBufferedOutputDevice : public FOutputDevice
 {
 protected:
 	TArray<FBufferedLine>		BufferedLines;
@@ -19,9 +19,9 @@ protected:
 
 public:
 	void	SetVerbosity(ELogVerbosity::Type Verbosity) { FilterLevel = Verbosity; }
-	void	Serialize(const TCHAR* InData, ELogVerbosity::Type Verbosity, const FName& Category) override;
-	void	GetContents(TArray<FBufferedLine>& DestBuffer);
+	CORE_API void	Serialize(const TCHAR* InData, ELogVerbosity::Type Verbosity, const FName& Category) override;
+	CORE_API void	GetContents(TArray<FBufferedLine>& DestBuffer);
 
 	/** Pushes buffered lines into the specified output device. */
-	void	RedirectTo(FOutputDevice& Ar);
+	CORE_API void	RedirectTo(FOutputDevice& Ar);
 };
diff --git a/Engine/Source/Runtime/Core/Public/Misc/CString.h b/Engine/Source/Runtime/Core/Public/Misc/CString.h
--- a/Engine/Source/Runtime/Core/Public/Misc/CString.h
+++ b/Engine/Source/Runtime/Core/Public/Misc/CString.h
@@ -42,11 +42,11 @@ namespace ESearchDir
 }
 
 /** Helper class used to convert CString into a boolean value. */
-struct CORE_API FToBoolHelper
+struct FToBoolHelper
 {
-	static bool FromCStringAnsi( const ANSICHAR* String );
-	static bool FromCStringWide( const WIDECHAR* String );
-	static bool FromCStringUtf8( const UTF8CHAR* String );
+	static CORE_API bool FromCStringAnsi( const ANSICHAR* String );
+	static CORE_API bool FromCStringWide( const WIDECHAR* String );
+	static CORE_API bool FromCStringUtf8( const UTF8CHAR* String );
 };
 
 /**
diff --git a/Engine/Source/Runtime/Core/Public/Misc/Change.h b/Engine/Source/Runtime/Core/Public/Misc/Change.h
--- a/Engine/Source/Runtime/Core/Public/Misc/Change.h
+++ b/Engine/Source/Runtime/Core/Public/Misc/Change.h
@@ -28,7 +28,7 @@ class UObject;
  * FChange may eventually be deprecated. You should subclass 
  * FSwapChange and FCommandChange to implement these different styles.
  */
-class CORE_API FChange
+class FChange
 {
 
 public:
@@ -60,7 +60,7 @@ public:
 	virtual FString ToString() const = 0;
 
 	/** Prints this change to the log, including sub-changes if there are any.  For compound changes, there might be multiple lines.  You should not need to override this function. */
-	virtual void PrintToLog( FFeedbackContext& FeedbackContext, const int32 IndentLevel = 0 );
+	CORE_API virtual void PrintToLog( FFeedbackContext& FeedbackContext, const int32 IndentLevel = 0 );
 
 	/** Virtual destructor */
 	virtual ~FChange()
@@ -89,7 +89,7 @@ private:
  *   1) apply the change to the given UObject
  *   2) return a new FSwapChange that does the "opposite" action
  */
-class CORE_API FSwapChange : public FChange
+class FSwapChange : public FChange
 {
 
 public:
@@ -117,7 +117,7 @@ public:
  * To use FCommandChange you must implement Apply() and Revert()
  * Revert() is called to "Undo" and Apply() is called to "Redo"
  */
-class CORE_API FCommandChange : public FChange
+class FCommandChange : public FChange
 {
 
 public:
@@ -139,7 +139,7 @@ public:
 
 
 
-struct CORE_API FCompoundChangeInput
+struct FCompoundChangeInput
 {
 	FCompoundChangeInput()
 	{
@@ -165,7 +165,7 @@ private:
  * FCompoundChange applies a sequence of FSwapChanges.
  * The changes are executed in reverse order (this is like a mini undo stack)
  */
-class CORE_API FCompoundChange : public FSwapChange
+class FCompoundChange : public FSwapChange
 {
 
 public:
@@ -177,9 +177,9 @@ public:
 	}
 
 	// Parent class overrides
-	virtual TUniquePtr<FChange> Execute( UObject* Object ) override;
-	virtual FString ToString() const override;
-	virtual void PrintToLog( class FFeedbackContext& FeedbackContext, const int32 IndentLevel = 0 ) override;
+	CORE_API virtual TUniquePtr<FChange> Execute( UObject* Object ) override;
+	CORE_API virtual FString ToString() const override;
+	CORE_API virtual void PrintToLog( class FFeedbackContext& FeedbackContext, const int32 IndentLevel = 0 ) override;
 
 
 private:
diff --git a/Engine/Source/Runtime/Core/Public/Misc/CompressedGrowableBuffer.h b/Engine/Source/Runtime/Core/Public/Misc/CompressedGrowableBuffer.h
--- a/Engine/Source/Runtime/Core/Public/Misc/CompressedGrowableBuffer.h
+++ b/Engine/Source/Runtime/Core/Public/Misc/CompressedGrowableBuffer.h
@@ -10,7 +10,7 @@
  * Growable compressed buffer. Usage is to append frequently but only request and therefore decompress
  * very infrequently. The prime usage case is the memory profiler keeping track of full call stacks.
  */
-struct CORE_API FCompressedGrowableBuffer
+struct FCompressedGrowableBuffer
 {
 private:
 	/**
@@ -21,7 +21,7 @@ private:
 	{
 		Redirect
 	};
-	FCompressedGrowableBuffer(EVS2015Redirector, int32 MaxPendingBufferSize, ECompressionFlags CompressionFlags);
+	CORE_API FCompressedGrowableBuffer(EVS2015Redirector, int32 MaxPendingBufferSize, ECompressionFlags CompressionFlags);
 public:
 	/**
 	 * Constructor
@@ -34,17 +34,17 @@ public:
 		// Make sure to remove the EVS2015Redirector constructor when this constructor is removed
 		: FCompressedGrowableBuffer(EVS2015Redirector::Redirect, MaxPendingBufferSize, CompressionFlags)
 	{}
-	FCompressedGrowableBuffer(int32 MaxPendingBufferSize, FName COmpressionFormat, ECompressionFlags CompressionFlags=COMPRESS_None);
+	CORE_API FCompressedGrowableBuffer(int32 MaxPendingBufferSize, FName COmpressionFormat, ECompressionFlags CompressionFlags=COMPRESS_None);
 
 	/**
 	 * Locks the buffer for reading. Needs to be called before calls to Access and needs
 	 * to be matched up with Unlock call.
 	 */
-	void Lock();
+	CORE_API void Lock();
 	/**
 	 * Unlocks the buffer and frees temporary resources used for accessing.
 	 */
-	void Unlock();
+	CORE_API void Unlock();
 
 	/**
 	 * Appends passed in data to the buffer. The data needs to be less than the max
@@ -54,7 +54,7 @@ public:
 	 * @param	Size	Size of data in bytes.
 	 * @return	Offset of data, used for retrieval later on
 	 */
-	int32 Append( void* Data, int32 Size );
+	CORE_API int32 Append( void* Data, int32 Size );
 
 	/**
 	 * Accesses the data at passed in offset and returns it. The memory is read-only and
@@ -63,7 +63,7 @@ public:
 	 *
 	 * @param	Offset	Offset to return corresponding data for
 	 */
-	void* Access( int32 Offset );
+	CORE_API void* Access( int32 Offset );
 
 	/**
 	 * @return	Number of entries appended.
diff --git a/Engine/Source/Runtime/Core/Public/Misc/ConfigCacheIni.h b/Engine/Source/Runtime/Core/Public/Misc/ConfigCacheIni.h
--- a/Engine/Source/Runtime/Core/Public/Misc/ConfigCacheIni.h
+++ b/Engine/Source/Runtime/Core/Public/Misc/ConfigCacheIni.h
@@ -411,7 +411,7 @@ public:
 	CORE_API bool Combine( const FString& Filename);
 	CORE_API void CombineFromBuffer(const FString& Buffer, const FString& FileHint);
 	UE_DEPRECATED(5.1, "Use CombineFromBuffer that takes FileHint")
-	CORE_API void CombineFromBuffer(const FString& Buffer)
+	void CombineFromBuffer(const FString& Buffer)
 	{
 		CombineFromBuffer(Buffer, TEXT("Unknown file, using deprecated function"));
 	}
@@ -525,7 +525,7 @@ public:
 	 */
 	CORE_API void ProcessInputFileContents(FStringView Contents, const FString& FileHint);
 	UE_DEPRECATED(5.1, "Use ProcessInputFileContents that takes FileHint")
-	CORE_API void ProcessInputFileContents(const FString& Buffer)
+	void ProcessInputFileContents(const FString& Buffer)
 	{
 		ProcessInputFileContents(Buffer, TEXT("Unknown file, using deprecated function"));
 	}
@@ -580,7 +580,7 @@ public:
 	CORE_API void AddDynamicLayerToHierarchy(const FString& Filename);
 
 	UE_DEPRECATED(5.0, "Call AddDynamicLayerToHierarchy. You also may need to call GetConfigFilename to get the right FConfigFile")
-	CORE_API void AddDynamicLayerToHeirarchy(const FString& Filename) { AddDynamicLayerToHierarchy(Filename); }
+	void AddDynamicLayerToHeirarchy(const FString& Filename) { AddDynamicLayerToHierarchy(Filename); }
 
 	friend FArchive& operator<<(FArchive& Ar, FConfigFile& ConfigFile);
 private:
@@ -635,31 +635,31 @@ enum class EConfigCacheType : uint8
 };
 
 // Set of all cached config files.
-class CORE_API FConfigCacheIni
+class FConfigCacheIni
 {
 public:
 	// Basic functions.
-	FConfigCacheIni(EConfigCacheType Type);
+	CORE_API FConfigCacheIni(EConfigCacheType Type);
 
 	/** DO NOT USE. This constructor is for internal usage only for hot-reload purposes. */
-	FConfigCacheIni();
+	CORE_API FConfigCacheIni();
 
-	virtual ~FConfigCacheIni();
+	CORE_API virtual ~FConfigCacheIni();
 
 	/**
 	* Disables any file IO by the config cache system
 	*/
-	virtual void DisableFileOperations();
+	CORE_API virtual void DisableFileOperations();
 
 	/**
 	* Re-enables file IO by the config cache system
 	*/
-	virtual void EnableFileOperations();
+	CORE_API virtual void EnableFileOperations();
 
 	/**
 	 * Returns whether or not file operations are disabled
 	 */
-	virtual bool AreFileOperationsDisabled();
+	CORE_API virtual bool AreFileOperationsDisabled();
 
 	/**
 	 * @return true after after the basic .ini files have been loaded
@@ -687,7 +687,7 @@ public:
 	*
 	* NOTE: The function naming is weird because you can't apparently have an overridden function differnt only by template type params
 	*/
-	virtual void Parse1ToNSectionOfStrings(const TCHAR* Section, const TCHAR* KeyOne, const TCHAR* KeyN, TMap<FString, TArray<FString> >& OutMap, const FString& Filename);
+	CORE_API virtual void Parse1ToNSectionOfStrings(const TCHAR* Section, const TCHAR* KeyOne, const TCHAR* KeyN, TMap<FString, TArray<FString> >& OutMap, const FString& Filename);
 
 	/**
 	* Parses apart an ini section that contains a list of 1-to-N mappings of names in the following format
@@ -707,7 +707,7 @@ public:
 	*
 	* NOTE: The function naming is weird because you can't apparently have an overridden function differnt only by template type params
 	*/
-	virtual void Parse1ToNSectionOfNames(const TCHAR* Section, const TCHAR* KeyOne, const TCHAR* KeyN, TMap<FName, TArray<FName> >& OutMap, const FString& Filename);
+	CORE_API virtual void Parse1ToNSectionOfNames(const TCHAR* Section, const TCHAR* KeyOne, const TCHAR* KeyN, TMap<FName, TArray<FName> >& OutMap, const FString& Filename);
 
 	/**
 	 * Finds the in-memory config file for a config cache filename.
@@ -716,7 +716,7 @@ public:
 	 *
 	 * @return The existing config file or null if it does not exist in memory
 	 */
-	FConfigFile* FindConfigFile(const FString& Filename);
+	CORE_API FConfigFile* FindConfigFile(const FString& Filename);
 
 	/**
 	 * Finds, loads, or creates the in-memory config file for a config cache filename.
@@ -725,20 +725,20 @@ public:
 	 * 
 	 * @return A new or existing config file
 	 */
-	FConfigFile* Find(const FString& InFilename);
+	CORE_API FConfigFile* Find(const FString& InFilename);
 
 	/**
 	 * Reports whether an FConfigFile* is pointing to a config file inside of this
 	 * Used for downstream functions to check whether a config file they were passed came from this ConfigCacheIni or from 
 	 * a different source such as LoadLocalIniFile
 	 */
-	bool ContainsConfigFile(const FConfigFile* ConfigFile) const;
+	CORE_API bool ContainsConfigFile(const FConfigFile* ConfigFile) const;
 
 	UE_DEPRECATED(5.0, "CreateIfNotFound is deprecated, please use the overload without this parameter or FindConfigFile")
-	FConfigFile* Find(const FString& Filename, bool CreateIfNotFound);
+	CORE_API FConfigFile* Find(const FString& Filename, bool CreateIfNotFound);
 
 	/** Finds Config file that matches the base name such as "Engine" */
-	FConfigFile* FindConfigFileWithBaseName(FName BaseName);
+	CORE_API FConfigFile* FindConfigFileWithBaseName(FName BaseName);
 
 	FConfigFile& Add(const FString& Filename, const FConfigFile& File)
 	{
@@ -749,30 +749,30 @@ public:
 		delete OtherFiles.FindRef(Filename);
 		return OtherFiles.Remove(Filename);
 	}
-	TArray<FString> GetFilenames();
+	CORE_API TArray<FString> GetFilenames();
 
 
-	void Flush(bool bRemoveFromCache, const FString& Filename=TEXT(""));
+	CORE_API void Flush(bool bRemoveFromCache, const FString& Filename=TEXT(""));
 
-	void LoadFile( const FString& InFilename, const FConfigFile* Fallback = NULL, const TCHAR* PlatformString = NULL );
-	void SetFile( const FString& InFilename, const FConfigFile* NewConfigFile );
-	void UnloadFile( const FString& Filename );
-	void Detach( const FString& Filename );
+	CORE_API void LoadFile( const FString& InFilename, const FConfigFile* Fallback = NULL, const TCHAR* PlatformString = NULL );
+	CORE_API void SetFile( const FString& InFilename, const FConfigFile* NewConfigFile );
+	CORE_API void UnloadFile( const FString& Filename );
+	CORE_API void Detach( const FString& Filename );
 
-	bool GetString( const TCHAR* Section, const TCHAR* Key, FString& Value, const FString& Filename );
-	bool GetText( const TCHAR* Section, const TCHAR* Key, FText& Value, const FString& Filename );
-	bool GetSection( const TCHAR* Section, TArray<FString>& Result, const FString& Filename );
-	bool DoesSectionExist(const TCHAR* Section, const FString& Filename);
+	CORE_API bool GetString( const TCHAR* Section, const TCHAR* Key, FString& Value, const FString& Filename );
+	CORE_API bool GetText( const TCHAR* Section, const TCHAR* Key, FText& Value, const FString& Filename );
+	CORE_API bool GetSection( const TCHAR* Section, TArray<FString>& Result, const FString& Filename );
+	CORE_API bool DoesSectionExist(const TCHAR* Section, const FString& Filename);
 	/**
 	 * @param Force Whether to create the Section on Filename if it did not exist previously.
 	 * @param Const If Const (and not Force), then it will not modify File->Dirty. If not Const (or Force is true), then File->Dirty will be set to true.
 	 */
-	FConfigSection* GetSectionPrivate( const TCHAR* Section, const bool Force, const bool Const, const FString& Filename );
-	void SetString( const TCHAR* Section, const TCHAR* Key, const TCHAR* Value, const FString& Filename );
-	void SetText( const TCHAR* Section, const TCHAR* Key, const FText& Value, const FString& Filename );
-	bool RemoveKey( const TCHAR* Section, const TCHAR* Key, const FString& Filename );
-	bool EmptySection( const TCHAR* Section, const FString& Filename );
-	bool EmptySectionsMatchingString( const TCHAR* SectionString, const FString& Filename );
+	CORE_API FConfigSection* GetSectionPrivate( const TCHAR* Section, const bool Force, const bool Const, const FString& Filename );
+	CORE_API void SetString( const TCHAR* Section, const TCHAR* Key, const TCHAR* Value, const FString& Filename );
+	CORE_API void SetText( const TCHAR* Section, const TCHAR* Key, const FText& Value, const FString& Filename );
+	CORE_API bool RemoveKey( const TCHAR* Section, const TCHAR* Key, const FString& Filename );
+	CORE_API bool EmptySection( const TCHAR* Section, const FString& Filename );
+	CORE_API bool EmptySectionsMatchingString( const TCHAR* SectionString, const FString& Filename );
 
 	/**
 	 * For a base ini name, gets the config cache filename key that is used by other functions like Find.
@@ -782,14 +782,14 @@ public:
 	 *
 	 * @return Filename key used by other cache functions
 	 */
-	FString GetConfigFilename(const TCHAR* BaseIniName);
+	CORE_API FString GetConfigFilename(const TCHAR* BaseIniName);
 
 	/**
 	 * Retrieve a list of all of the config files stored in the cache
 	 *
 	 * @param ConfigFilenames Out array to receive the list of filenames
 	 */
-	void GetConfigFilenames(TArray<FString>& ConfigFilenames);
+	CORE_API void GetConfigFilenames(TArray<FString>& ConfigFilenames);
 
 	/**
 	 * Retrieve the names for all sections contained in the file specified by Filename
@@ -799,7 +799,7 @@ public:
 	 *
 	 * @return	true if the file specified was successfully found;
 	 */
-	bool GetSectionNames( const FString& Filename, TArray<FString>& out_SectionNames );
+	CORE_API bool GetSectionNames( const FString& Filename, TArray<FString>& out_SectionNames );
 
 	/**
 	 * Retrieve the names of sections which contain data for the specified PerObjectConfig class.
@@ -811,9 +811,9 @@ public:
 	 *
 	 * @return	true if the file specified was found and it contained at least 1 section for the specified class
 	 */
-	bool GetPerObjectConfigSections( const FString& Filename, const FString& SearchClass, TArray<FString>& out_SectionNames, int32 MaxResults=1024 );
+	CORE_API bool GetPerObjectConfigSections( const FString& Filename, const FString& SearchClass, TArray<FString>& out_SectionNames, int32 MaxResults=1024 );
 
-	void Exit();
+	CORE_API void Exit();
 
 	/**
 	 * Prints out the entire config set, or just a single file if an ini is specified
@@ -821,71 +821,71 @@ public:
 	 * @param Ar the device to write to
 	 * @param IniName An optional ini name to restrict the writing to (Engine or WrangleContent) - meant to be used with "final" .ini files (not Default*)
 	 */
-	void Dump(FOutputDevice& Ar, const TCHAR* IniName=NULL);
+	CORE_API void Dump(FOutputDevice& Ar, const TCHAR* IniName=NULL);
 
 	/**
 	 * Dumps memory stats for each file in the config cache to the specified archive.
 	 *
 	 * @param	Ar	the output device to dump the results to
 	 */
-	virtual void ShowMemoryUsage( FOutputDevice& Ar );
+	CORE_API virtual void ShowMemoryUsage( FOutputDevice& Ar );
 
 	/**
 	 * USed to get the max memory usage for the FConfigCacheIni
 	 *
 	 * @return the amount of memory in byes
 	 */
-	virtual SIZE_T GetMaxMemoryUsage();
+	CORE_API virtual SIZE_T GetMaxMemoryUsage();
 
 	/**
 	 * allows to iterate through all key value pairs
 	 * @return false:error e.g. Section or Filename not found
 	 */
-	virtual bool ForEachEntry(const FKeyValueSink& Visitor, const TCHAR* Section, const FString& Filename);
+	CORE_API virtual bool ForEachEntry(const FKeyValueSink& Visitor, const TCHAR* Section, const FString& Filename);
 
 	// Derived functions.
-	FString GetStr
+	CORE_API FString GetStr
 	(
 		const TCHAR*		Section, 
 		const TCHAR*		Key, 
 		const FString&	Filename 
 	);
-	bool GetInt
+	CORE_API bool GetInt
 	(
 		const TCHAR*		Section,
 		const TCHAR*		Key,
 		int32&				Value,
 		const FString&	Filename
 	);
-	bool GetInt64
+	CORE_API bool GetInt64
 	(
 		const TCHAR* Section,
 		const TCHAR* Key,
 		int64& Value,
 		const FString& Filename
 	);
-	bool GetFloat
+	CORE_API bool GetFloat
 	(
 		const TCHAR*		Section,
 		const TCHAR*		Key,
 		float&				Value,
 		const FString&	Filename
 	);
-	bool GetDouble
+	CORE_API bool GetDouble
 	(
 		const TCHAR*		Section,
 		const TCHAR*		Key,
 		double&				Value,
 		const FString&	Filename
 	);
-	bool GetBool
+	CORE_API bool GetBool
 	(
 		const TCHAR*		Section,
 		const TCHAR*		Key,
 		bool&				Value,
 		const FString&	Filename
 	);
-	int32 GetArray
+	CORE_API int32 GetArray
 	(
 		const TCHAR*		Section,
 		const TCHAR*		Key,
@@ -898,40 +898,40 @@ public:
 	 * @param out_Arr - Array to load into
 	 * @param Filename - Ini file to load from
 	 */
-	int32 GetSingleLineArray
+	CORE_API int32 GetSingleLineArray
 	(
 		const TCHAR*		Section,
 		const TCHAR*		Key,
 		TArray<FString>&	out_Arr,
 		const FString&	Filename
 	);
-	bool GetColor
+	CORE_API bool GetColor
 	(
 		const TCHAR*		Section,
 		const TCHAR*		Key,
 		FColor&				Value,
 		const FString&	Filename
 	);
-	bool GetVector2D(
+	CORE_API bool GetVector2D(
 		const TCHAR*   Section,
 		const TCHAR*   Key,
 		FVector2D&     Value,
 		const FString& Filename);
-	bool GetVector
+	CORE_API bool GetVector
 	(
 		const TCHAR*		Section,
 		const TCHAR*		Key,
 		FVector&			Value,
 		const FString&	Filename
 	);
-	bool GetVector4
+	CORE_API bool GetVector4
 	(
 		const TCHAR*		Section,
 		const TCHAR*		Key,
 		FVector4&			Value,
 		const FString&	Filename
 	);
-	bool GetRotator
+	CORE_API bool GetRotator
 	(
 		const TCHAR*		Section,
 		const TCHAR*		Key,
@@ -1012,35 +1012,35 @@ public:
 		}
 	}
 
-	void SetInt
+	CORE_API void SetInt
 	(
 		const TCHAR*		Section,
 		const TCHAR*		Key,
 		int32					Value,
 		const FString&	Filename
 	);
-	void SetFloat
+	CORE_API void SetFloat
 	(
 		const TCHAR*		Section,
 		const TCHAR*		Key,
 		float				Value,
 		const FString&	Filename
 	);
-	void SetDouble
+	CORE_API void SetDouble
 	(
 		const TCHAR*		Section,
 		const TCHAR*		Key,
 		double				Value,
 		const FString&	Filename
 	);
-	void SetBool
+	CORE_API void SetBool
 	(
 		const TCHAR*		Section,
 		const TCHAR*		Key,
 		bool				Value,
 		const FString&	Filename
 	);
-	void SetArray
+	CORE_API void SetArray
 	(
 		const TCHAR*			Section,
 		const TCHAR*			Key,
@@ -1053,40 +1053,40 @@ public:
 	 * @param out_Arr - Array to save from
 	 * @param Filename - Ini file to save to
 	 */
-	void SetSingleLineArray
+	CORE_API void SetSingleLineArray
 	(
 		const TCHAR*			Section,
 		const TCHAR*			Key,
 		const TArray<FString>&	In_Arr,
 		const FString&		Filename
 	);
-	void SetColor
+	CORE_API void SetColor
 	(
 		const TCHAR*		Section,
 		const TCHAR*		Key,
 		FColor				Value,
 		const FString&	Filename
 	);
-	void SetVector2D(
+	CORE_API void SetVector2D(
 		const TCHAR*   Section,
 		const TCHAR*   Key,
 		FVector2D      Value,
 		const FString& Filename);
-	void SetVector
+	CORE_API void SetVector
 	(
 		const TCHAR*		Section,
 		const TCHAR*		Key,
 		FVector				Value,
 		const FString&	Filename
 	);
-	void SetVector4
+	CORE_API void SetVector4
 	(
 		const TCHAR*		Section,
 		const TCHAR*		Key,
 		const FVector4&		Value,
 		const FString&	Filename
 	);
-	void SetRotator
+	CORE_API void SetRotator
 	(
 		const TCHAR*		Section,
 		const TCHAR*		Key,
@@ -1100,14 +1100,14 @@ public:
 	 * Creates GConfig, loads the standard global ini files (Engine, Editor, etc),
 	 * fills out GEngineIni, etc. and marks GConfig as ready for use
 	 */
-	static void InitializeConfigSystem();
+	static CORE_API void InitializeConfigSystem();
 
 	/**
 	 * Returns the Custom Config string, which if set will load additional config files from Config/Custom/{CustomConfig}/DefaultX.ini to allow different types of builds.
 	 * It can be set from a game Target.cs file with CustomConfig = "Name".
 	 * Or in development, it can be overridden with a -CustomConfig=Name command line parameter.
 	 */
-	static const FString& GetCustomConfigString();
+	static CORE_API const FString& GetCustomConfigString();
 
 	/**
 	 * Calculates the name of a dest (generated) .ini file for a given base (ie Engine, Game, etc)
@@ -1118,7 +1118,7 @@ public:
 	 *
 	 * @return Standardized .ini filename
 	 */
-	static FString GetDestIniFilename(const TCHAR* BaseIniName, const TCHAR* PlatformName, const TCHAR* GeneratedConfigDir);
+	static CORE_API FString GetDestIniFilename(const TCHAR* BaseIniName, const TCHAR* PlatformName, const TCHAR* GeneratedConfigDir);
 
 	/**
 	 * Loads and generates a destination ini file and adds it to GConfig:
@@ -1138,7 +1138,7 @@ public:
 	 * @param GeneratedConfigDir The location where generated config files are made.
 	 * @return true if the final ini was created successfully.
 	 */
-	static bool LoadGlobalIniFile(FString& FinalIniFilename, const TCHAR* BaseIniName, const TCHAR* Platform = NULL, bool bForceReload = false, bool bRequireDefaultIni = false, bool bAllowGeneratedIniWhenCooked = true, bool bAllowRemoteConfig = true, const TCHAR* GeneratedConfigDir = *FPaths::GeneratedConfigDir(), FConfigCacheIni* ConfigSystem=GConfig);
+	static CORE_API bool LoadGlobalIniFile(FString& FinalIniFilename, const TCHAR* BaseIniName, const TCHAR* Platform = NULL, bool bForceReload = false, bool bRequireDefaultIni = false, bool bAllowGeneratedIniWhenCooked = true, bool bAllowRemoteConfig = true, const TCHAR* GeneratedConfigDir = *FPaths::GeneratedConfigDir(), FConfigCacheIni* ConfigSystem=GConfig);
 
 	/**
 	 * Load an ini file directly into an FConfigFile, and nothing is written to GConfig or disk. 
@@ -1152,7 +1152,7 @@ public:
 	 * @param bForceReload force reload the ini file from disk this is required if you make changes to the ini file not using the config system as the hierarchy cache will not be updated in this case
 	 * @return true if the ini file was loaded successfully
 	 */
-	static bool LoadLocalIniFile(FConfigFile& ConfigFile, const TCHAR* IniName, bool bIsBaseIniName, const TCHAR* Platform=NULL, bool bForceReload=false);
+	static CORE_API bool LoadLocalIniFile(FConfigFile& ConfigFile, const TCHAR* IniName, bool bIsBaseIniName, const TCHAR* Platform=NULL, bool bForceReload=false);
 
 	/**
 	 * Load an ini file directly into an FConfigFile from the specified config folders, optionally writing to disk. 
@@ -1171,14 +1171,14 @@ public:
 	 * @param GeneratedConfigDir The location where generated config files are made.
 	 * @return true if the ini file was loaded successfully
 	 */
-	static bool LoadExternalIniFile(FConfigFile& ConfigFile, const TCHAR* IniName, const TCHAR* EngineConfigDir, const TCHAR* SourceConfigDir, bool bIsBaseIniName, const TCHAR* Platform=NULL, bool bForceReload=false, bool bWriteDestIni=false, bool bAllowGeneratedIniWhenCooked = true, const TCHAR* GeneratedConfigDir = *FPaths::GeneratedConfigDir());
+	static CORE_API bool LoadExternalIniFile(FConfigFile& ConfigFile, const TCHAR* IniName, const TCHAR* EngineConfigDir, const TCHAR* SourceConfigDir, bool bIsBaseIniName, const TCHAR* Platform=NULL, bool bForceReload=false, bool bWriteDestIni=false, bool bAllowGeneratedIniWhenCooked = true, const TCHAR* GeneratedConfigDir = *FPaths::GeneratedConfigDir());
 
 	/**
 	 * Needs to be called after GConfig is set and LoadCoalescedFile was called.
 	 * Loads the state of console variables.
 	 * Works even if the variable is registered after the ini file was loaded.
 	 */
-	static void LoadConsoleVariablesFromINI();
+	static CORE_API void LoadConsoleVariablesFromINI();
 
 	/**
 	 * Normalizes file paths to INI files.
@@ -1208,7 +1208,7 @@ public:
 	 *	@param NonNormalizedPath	The path to the INI file we want to access.
 	 *	@return A normalized version of the path (may be the same as the input).
 	 */
-	static FString NormalizeConfigIniPath(const FString& NonNormalizedPath);
+	static CORE_API FString NormalizeConfigIniPath(const FString& NonNormalizedPath);
 
 	/**
 	 * This helper function searches the cache before trying to load the ini file using LoadLocalIniFile. 
@@ -1219,7 +1219,7 @@ public:
 	 * @param Platform The platform to use for Base ini names, NULL means to use the current platform
 	 * @return FConfigFile Found or loaded FConfigFile
 	 */
-	static FConfigFile* FindOrLoadPlatformConfig(FConfigFile& LocalFile, const TCHAR* IniName, const TCHAR* Platform = NULL);
+	static CORE_API FConfigFile* FindOrLoadPlatformConfig(FConfigFile& LocalFile, const TCHAR* IniName, const TCHAR* Platform = NULL);
 
 	/**
 	 * Attempts to find the platform config in the cache.
@@ -1228,14 +1228,14 @@ public:
 	 * @param Platform The platform to use for Base ini names, NULL means to use the current platform
 	 * @return FConfigFile Found FConfigFile
 	 */
-	static FConfigFile* FindPlatformConfig(const TCHAR* IniName, const TCHAR* Platform);
+	static CORE_API FConfigFile* FindPlatformConfig(const TCHAR* IniName, const TCHAR* Platform);
 	
 	/**
 	 * Save the current config cache state into a file for bootstrapping other processes.
 	 */
-	void SaveCurrentStateForBootstrap(const TCHAR* Filename);
+	CORE_API void SaveCurrentStateForBootstrap(const TCHAR* Filename);
 
-	void Serialize(FArchive& Ar);
+	CORE_API void Serialize(FArchive& Ar);
 
 
 	struct FKnownConfigFiles
@@ -1276,38 +1276,38 @@ public:
 	 *
 	 * return True if the engine ini was loaded
 	 */
-	static bool InitializeKnownConfigFiles(FConfigContext& Context);
+	static CORE_API bool InitializeKnownConfigFiles(FConfigContext& Context);
 
 	/**
 	 * Returns true if the given name is one of the known configs, where the matching G****Ini property is going to match the 
 	 * base name ("Engine" returns true, which means GEngineIni's value is just "Engine")
 	 */
-	bool IsKnownConfigName(FName ConfigName);
+	CORE_API bool IsKnownConfigName(FName ConfigName);
 
 	/**
 	 * Create GConfig from a saved file
 	 */
-	static bool CreateGConfigFromSaved(const TCHAR* Filename);
+	static CORE_API bool CreateGConfigFromSaved(const TCHAR* Filename);
 
 	/**
 	 * Retrieve the fully processed ini system for another platform. The editor will start loading these 
 	 * in the background on startup
 	 */
-	static FConfigCacheIni* ForPlatform(FName PlatformName);
+	static CORE_API FConfigCacheIni* ForPlatform(FName PlatformName);
 
 	/**
 	 * Wipe all cached platform configs. Next ForPlatform call will load on-demand the platform configs
 	 */
-	static void ClearOtherPlatformConfigs();
+	static CORE_API void ClearOtherPlatformConfigs();
 
 private:
 #if WITH_EDITOR
 	/** We only auto-initialize other platform configs in the editor to not slow down programs like ShaderCOmpileWorker */
-	static void AsyncInitializeConfigForPlatforms();
+	static CORE_API void AsyncInitializeConfigForPlatforms();
 #endif
 
 	/** Serialize a bootstrapping state into or from an archive */
-	void SerializeStateForBootstrap_Impl(FArchive& Ar);
+	CORE_API void SerializeStateForBootstrap_Impl(FArchive& Ar);
 
 	/** true if file operations should not be performed */
 	bool bAreFileOperationsDisabled;
diff --git a/Engine/Source/Runtime/Core/Public/Misc/ConfigContext.h b/Engine/Source/Runtime/Core/Public/Misc/ConfigContext.h
--- a/Engine/Source/Runtime/Core/Public/Misc/ConfigContext.h
+++ b/Engine/Source/Runtime/Core/Public/Misc/ConfigContext.h
@@ -26,7 +26,7 @@
 class FConfigCacheIni;
 class FConfigFile;
 
-class CORE_API FConfigContext
+class FConfigContext
 {
 public:
 
@@ -122,19 +122,19 @@ public:
 	/**
 	 * Call to make before attempting parallel config init
 	 */
-	static void EnsureRequiredGlobalPathsHaveBeenInitialized();
+	static CORE_API void EnsureRequiredGlobalPathsHaveBeenInitialized();
 
 	/**
 	 * Use the context to perform the actual load operation. Note that this is where you specify the Ini name (for instance "Engine"), meaning
 	 * you can use the same context for multiple configs (Engine, Game, Input, etc)
 	 */
-	bool Load(const TCHAR* IniName);
+	CORE_API bool Load(const TCHAR* IniName);
 
 	/**
 	 * Use the context to perform the actual load operation as above, but will return the generated final ini filename (in the case of GConfig, this would
 	 * be the key used to look up into GConfig, for example)
 	 */
-	bool Load(const TCHAR* IniName, FString& OutFilename);
+	CORE_API bool Load(const TCHAR* IniName, FString& OutFilename);
 
 
 	// because the hierarchy can jump between platforms, we cache off some directories per chained-platform
@@ -147,7 +147,7 @@ public:
 	/**
 	 * Return the paths to use to find hierarchical config files for the given platform (note that these are independent of the ini name)
 	 */
-	const FPerPlatformDirs& GetPerPlatformDirs(const FString& PlatformName);
+	CORE_API const FPerPlatformDirs& GetPerPlatformDirs(const FString& PlatformName);
 
 
 	// @todo make these private and friend the FCOnfigCacheIni and FConfigFile once everything is a member function!!!!
@@ -190,17 +190,17 @@ protected:
 	bool bDoNotResetConfigFile = false;
 	bool bCacheOnNextLoad = true;
 
-	FConfigContext(FConfigCacheIni* InConfigSystem, bool InIsHierarchicalConfig, const FString& InPlatform, FConfigFile* DestConfigFile=nullptr);
+	CORE_API FConfigContext(FConfigCacheIni* InConfigSystem, bool InIsHierarchicalConfig, const FString& InPlatform, FConfigFile* DestConfigFile=nullptr);
 
-	FConfigContext& ResetBaseIni(const TCHAR* InBaseIniName);
-	void CachePaths();
+	CORE_API FConfigContext& ResetBaseIni(const TCHAR* InBaseIniName);
+	CORE_API void CachePaths();
 
-	bool PrepareForLoad(bool& bPerformLoad);
-	bool PerformLoad();
+	CORE_API bool PrepareForLoad(bool& bPerformLoad);
+	CORE_API bool PerformLoad();
 
-	void AddStaticLayersToHierarchy();
-	bool GenerateDestIniFile();
-	FString PerformFinalExpansions(const FString& InString, const FString& Platform);
+	CORE_API void AddStaticLayersToHierarchy();
+	CORE_API bool GenerateDestIniFile();
+	CORE_API FString PerformFinalExpansions(const FString& InString, const FString& Platform);
 };
 
 bool DoesConfigFileExistWrapper(const TCHAR* IniFile, const TSet<FString>* IniCacheSet = nullptr);
