[Core]

diff --git a/Engine/Source/Runtime/Core/Public/Misc/CoreAsyncTaskNotificationImpl.h b/Engine/Source/Runtime/Core/Public/Misc/CoreAsyncTaskNotificationImpl.h
--- a/Engine/Source/Runtime/Core/Public/Misc/CoreAsyncTaskNotificationImpl.h
+++ b/Engine/Source/Runtime/Core/Public/Misc/CoreAsyncTaskNotificationImpl.h
@@ -88,29 +88,29 @@ public:
 /**
  * Basic asynchronous task notification that just logs status changes.
  */
-class CORE_API FCoreAsyncTaskNotificationImpl : public IAsyncTaskNotificationImpl
+class FCoreAsyncTaskNotificationImpl : public IAsyncTaskNotificationImpl
 {
 public:
-	FCoreAsyncTaskNotificationImpl();
-	virtual ~FCoreAsyncTaskNotificationImpl();
+	CORE_API FCoreAsyncTaskNotificationImpl();
+	CORE_API virtual ~FCoreAsyncTaskNotificationImpl();
 
 	//~ IAsyncTaskNotificationImpl
-	virtual void Initialize(const FAsyncTaskNotificationConfig& InConfig) override;
-	virtual void SetTitleText(const FText& InTitleText, const bool bClearProgressText) override;
-	virtual void SetProgressText(const FText& InProgressText) override;
-	virtual void SetPromptText(const FText& InPromptText) override;
-	virtual void SetHyperlink(const FSimpleDelegate& InHyperlink, const FText& InHyperlinkText) override;
-	virtual void SetComplete(const bool bSuccess) override;
-	virtual void SetComplete(const FText& InTitleText, const FText& InProgressText, const bool bSuccess) override;
-	virtual void SetNotificationState(const FAsyncNotificationStateData& InState) override;
-	virtual void SetCanCancel(const TAttribute<bool>& InCanCancel) override;
-	virtual void SetKeepOpenOnSuccess(const TAttribute<bool>& InKeepOpenOnSuccess) override;
-	virtual void SetKeepOpenOnFailure(const TAttribute<bool>& InKeepOpenOnFailure) override;
-	virtual EAsyncTaskNotificationPromptAction GetPromptAction() const override;
+	CORE_API virtual void Initialize(const FAsyncTaskNotificationConfig& InConfig) override;
+	CORE_API virtual void SetTitleText(const FText& InTitleText, const bool bClearProgressText) override;
+	CORE_API virtual void SetProgressText(const FText& InProgressText) override;
+	CORE_API virtual void SetPromptText(const FText& InPromptText) override;
+	CORE_API virtual void SetHyperlink(const FSimpleDelegate& InHyperlink, const FText& InHyperlinkText) override;
+	CORE_API virtual void SetComplete(const bool bSuccess) override;
+	CORE_API virtual void SetComplete(const FText& InTitleText, const FText& InProgressText, const bool bSuccess) override;
+	CORE_API virtual void SetNotificationState(const FAsyncNotificationStateData& InState) override;
+	CORE_API virtual void SetCanCancel(const TAttribute<bool>& InCanCancel) override;
+	CORE_API virtual void SetKeepOpenOnSuccess(const TAttribute<bool>& InKeepOpenOnSuccess) override;
+	CORE_API virtual void SetKeepOpenOnFailure(const TAttribute<bool>& InKeepOpenOnFailure) override;
+	CORE_API virtual EAsyncTaskNotificationPromptAction GetPromptAction() const override;
 
 protected:
 	/** Update the notification (the critical section is held while this function is called) */
-	virtual void UpdateNotification();
+	CORE_API virtual void UpdateNotification();
 
 	/** The current state of this notification */
 	EAsyncTaskNotificationState State = EAsyncTaskNotificationState::Pending;
@@ -143,7 +143,7 @@ private:
 /**
  * Factory to allow other systems (such as Slate) to override the default asynchronous task notification implementation.
  */
-class CORE_API FAsyncTaskNotificationFactory
+class FAsyncTaskNotificationFactory
 {
 	friend class FAsyncTaskNotification;
 
@@ -154,23 +154,23 @@ public:
 	/**
 	 * Get the factory singleton.
 	 */
-	static FAsyncTaskNotificationFactory& Get();
+	static CORE_API FAsyncTaskNotificationFactory& Get();
 
 	/**
 	 * Register a factory function.
 	 */
-	void RegisterFactory(const FName InName, const FFactoryFunc& InFunc);
+	CORE_API void RegisterFactory(const FName InName, const FFactoryFunc& InFunc);
 
 	/**
 	 * Unregister a factory function.
 	 */
-	void UnregisterFactory(const FName InName);
+	CORE_API void UnregisterFactory(const FName InName);
 
 private:
 	/**
 	 * Invoke the active factory function (if any), or return a default instance.
 	 */
-	FImplPointerType InvokeFactory() const;
+	CORE_API FImplPointerType InvokeFactory() const;
 
 	/** Registered factories */
 	TArray<TTuple<FName, FFactoryFunc>> Factories;
diff --git a/Engine/Source/Runtime/Core/Public/Misc/CoreDelegates.h b/Engine/Source/Runtime/Core/Public/Misc/CoreDelegates.h
--- a/Engine/Source/Runtime/Core/Public/Misc/CoreDelegates.h
+++ b/Engine/Source/Runtime/Core/Public/Misc/CoreDelegates.h
@@ -70,172 +70,172 @@ struct FCrashOverrideParameters
 	CORE_API ~FCrashOverrideParameters();
 };
 
-class CORE_API FCoreDelegates
+class FCoreDelegates
 {
 public:
 	// Callback for platform handling when flushing async loads.
-	static TMulticastDelegate<void()> OnAsyncLoadingFlush;
+	static CORE_API TMulticastDelegate<void()> OnAsyncLoadingFlush;
 
 	// Callback for a game thread interruption point when a async load flushing. Used to updating UI during long loads.
-	static TMulticastDelegate<void()> OnAsyncLoadingFlushUpdate;
+	static CORE_API TMulticastDelegate<void()> OnAsyncLoadingFlushUpdate;
 
 	// Callback on the game thread when an async load is started. This goes off before the packages has finished loading
-	static TMulticastDelegate<void(const FString&)> OnAsyncLoadPackage;
+	static CORE_API TMulticastDelegate<void(const FString&)> OnAsyncLoadPackage;
 
-	static TMulticastDelegate<void(const FString&)> OnSyncLoadPackage;
+	static CORE_API TMulticastDelegate<void(const FString&)> OnSyncLoadPackage;
 
 	// get a hotfix delegate
-	static TDelegate<void(void*, int32)>& GetHotfixDelegate(EHotfixDelegates::Type HotFix);
+	static CORE_API TDelegate<void(void*, int32)>& GetHotfixDelegate(EHotfixDelegates::Type HotFix);
 
 	// Callback when a user logs in/out of the platform.
-	static TMulticastDelegate<void(bool, int32, int32)> OnUserLoginChangedEvent;
+	static CORE_API TMulticastDelegate<void(bool, int32, int32)> OnUserLoginChangedEvent;
 
 	// Callback when controllers disconnected / reconnected
 	UE_DEPRECATED(5.1, "OnControllerConnectionChange, use IPlatformInputDeviceMapper::GetOnInputDeviceConnectionChange() instead")
-	static TMulticastDelegate<void(bool, FPlatformUserId, int32)> OnControllerConnectionChange;
+	static CORE_API TMulticastDelegate<void(bool, FPlatformUserId, int32)> OnControllerConnectionChange;
 
 	// Callback when a single controller pairing changes
 	UE_DEPRECATED(5.1, "OnControllerPairingChange, use IPlatformInputDeviceMapper::GetOnInputDevicePairingChange() instead")
-	static TMulticastDelegate<void(int32 ControllerIndex, FPlatformUserId NewUserPlatformId, FPlatformUserId OldUserPlatformId)> OnControllerPairingChange;
+	static CORE_API TMulticastDelegate<void(int32 ControllerIndex, FPlatformUserId NewUserPlatformId, FPlatformUserId OldUserPlatformId)> OnControllerPairingChange;
 
 	// Callback when a user changes the safe frame size
-	static TMulticastDelegate<void()> OnSafeFrameChangedEvent;
+	static CORE_API TMulticastDelegate<void()> OnSafeFrameChangedEvent;
 
 	// Callback for mounting all the pak files in default locations
-	static TDelegate<int32(const TArray<FString>&)> OnMountAllPakFiles;
+	static CORE_API TDelegate<int32(const TArray<FString>&)> OnMountAllPakFiles;
 
 	// Callback to prompt the pak system to mount a pak file
-	static TDelegate<IPakFile*(const FString&, int32)> MountPak;
+	static CORE_API TDelegate<IPakFile*(const FString&, int32)> MountPak;
 
 	// Callback to prompt the pak system to unmount a pak file.
-	static TDelegate<bool(const FString&)> OnUnmountPak;
+	static CORE_API TDelegate<bool(const FString&)> OnUnmountPak;
 
 	// Callback to optimize memeory for currently mounted paks
-	static TDelegate<void()> OnOptimizeMemoryUsageForMountedPaks;
+	static CORE_API TDelegate<void()> OnOptimizeMemoryUsageForMountedPaks;
 
 	// After a pakfile is mounted this is called
 	UE_DEPRECATED(5.3, "This delegate is not thread-safe, please use GetOnPakFileMounted2().")
-	static TMulticastDelegate<void(const IPakFile&)> OnPakFileMounted2;
-	static TTSMulticastDelegate<void(const IPakFile&)>& GetOnPakFileMounted2();
+	static CORE_API TMulticastDelegate<void(const IPakFile&)> OnPakFileMounted2;
+	static CORE_API TTSMulticastDelegate<void(const IPakFile&)>& GetOnPakFileMounted2();
 
 	// After a file is added this is called
-	static TMulticastDelegate<void(const FString&)> NewFileAddedDelegate;
+	static CORE_API TMulticastDelegate<void(const FString&)> NewFileAddedDelegate;
 
 	// After an attempt to mount all pak files, but none wre found, this is called
-	static TMulticastDelegate<void()> NoPakFilesMountedDelegate;
+	static CORE_API TMulticastDelegate<void()> NoPakFilesMountedDelegate;
 
 	// When a file is opened for read from a pak file.
 	UE_DEPRECATED(5.3, "This delegate is not thread-safe, please use Use FCoreDelegates::GetOnFileOpenedForReadFromPakFile()")
-	static TMulticastDelegate<void(const TCHAR* PakFile, const TCHAR* FileName), FNotThreadSafeNotCheckedDelegateUserPolicy> OnFileOpenedForReadFromPakFile;
-	static TTSMulticastDelegate<void(const TCHAR* PakFile, const TCHAR* FileName)>& GetOnFileOpenedForReadFromPakFile();
+	static CORE_API TMulticastDelegate<void(const TCHAR* PakFile, const TCHAR* FileName), FNotThreadSafeNotCheckedDelegateUserPolicy> OnFileOpenedForReadFromPakFile;
+	static CORE_API TTSMulticastDelegate<void(const TCHAR* PakFile, const TCHAR* FileName)>& GetOnFileOpenedForReadFromPakFile();
 
 	typedef TSharedPtr<class IMovieStreamer, ESPMode::ThreadSafe> FMovieStreamerPtr;
 
     // Delegate used to register a movie streamer with any movie player modules that bind to this delegate
     // Designed to be called when a platform specific movie streamer plugin starts up so that it doesn't need to implement a register for all movie player plugins
-    static TMulticastDelegate<void(FMovieStreamerPtr)> RegisterMovieStreamerDelegate;
+    static CORE_API TMulticastDelegate<void(FMovieStreamerPtr)> RegisterMovieStreamerDelegate;
     // Delegate used to un-register a movie streamer with any movie player modules that bind to this delegate
     // Designed to be called when a platform specific movie streamer plugin shuts down so that it doesn't need to implement a register for all movie player plugins
-    static TMulticastDelegate<void(FMovieStreamerPtr)> UnRegisterMovieStreamerDelegate;
+    static CORE_API TMulticastDelegate<void(FMovieStreamerPtr)> UnRegisterMovieStreamerDelegate;
 
 	// Callback when an ensure has occurred
-	static TMulticastDelegate<void()> OnHandleSystemEnsure;
+	static CORE_API TMulticastDelegate<void()> OnHandleSystemEnsure;
 
 	// Callback when an error (crash) has occurred
-	static TMulticastDelegate<void()> OnHandleSystemError;
+	static CORE_API TMulticastDelegate<void()> OnHandleSystemError;
 
 	// Called when an actor label is changed
-	static TMulticastDelegate<void(AActor*)> OnActorLabelChanged;
+	static CORE_API TMulticastDelegate<void(AActor*)> OnActorLabelChanged;
 
-	static TMulticastDelegate<void(const FGuid&, const FAES::FAESKey&)>& GetRegisterEncryptionKeyMulticastDelegate();
-	static TDelegate<void(uint8[32])>& GetPakEncryptionKeyDelegate();
-	static TDelegate<void(TArray<uint8>&, TArray<uint8>&)>& GetPakSigningKeysDelegate();
+	static CORE_API TMulticastDelegate<void(const FGuid&, const FAES::FAESKey&)>& GetRegisterEncryptionKeyMulticastDelegate();
+	static CORE_API TDelegate<void(uint8[32])>& GetPakEncryptionKeyDelegate();
+	static CORE_API TDelegate<void(TArray<uint8>&, TArray<uint8>&)>& GetPakSigningKeysDelegate();
 
 	
 
 #if WITH_EDITOR
 	// Called before the editor displays a modal window, allowing other windows the opportunity to disable themselves to avoid reentrant calls
-	static FSimpleMulticastDelegate PreModal;
+	static CORE_API FSimpleMulticastDelegate PreModal;
 
 	// Called after the editor dismisses a modal window, allowing other windows the opportunity to disable themselves to avoid reentrant calls
-	static FSimpleMulticastDelegate PostModal;
+	static CORE_API FSimpleMulticastDelegate PostModal;
     
     // Called before the editor displays a Slate (non-platform) modal window, allowing other windows the opportunity to disable themselves to avoid reentrant calls
-    static FSimpleMulticastDelegate PreSlateModal;
+    static CORE_API FSimpleMulticastDelegate PreSlateModal;
     
     // Called after the editor dismisses a Slate (non-platform) modal window, allowing other windows the opportunity to disable themselves to avoid reentrant calls
-    static FSimpleMulticastDelegate PostSlateModal;
+    static CORE_API FSimpleMulticastDelegate PostSlateModal;
     
 #endif	//WITH_EDITOR
 
 #if ALLOW_OTHER_PLATFORM_CONFIG
 	// Called when the CVar (ConsoleManager) needs to retrieve CVars for a deviceprofile for another platform - this dramatically simplifies module dependencies
 	typedef TMap<FName, FString> FCVarKeyValueMap;
-	static TDelegate<FCVarKeyValueMap(const FString& DeviceProfileName)> GatherDeviceProfileCVars;
+	static CORE_API TDelegate<FCVarKeyValueMap(const FString& DeviceProfileName)> GatherDeviceProfileCVars;
 #endif
 
 	// Called when an error occurred.
-	static FSimpleMulticastDelegate OnShutdownAfterError;
+	static CORE_API FSimpleMulticastDelegate OnShutdownAfterError;
 
 	// Called when appInit is called, very early in startup
-	static FSimpleMulticastDelegate OnInit;
+	static CORE_API FSimpleMulticastDelegate OnInit;
 
 	// Called during FEngineLoop::PreInit after GWarn & GError have been first set so that they can be overridden before anything in PreInit uses them
-	static FSimpleMulticastDelegate OnOutputDevicesInit;
+	static CORE_API FSimpleMulticastDelegate OnOutputDevicesInit;
 
 	// Called at the end of UEngine::Init, right before loading PostEngineInit modules for both normal execution and commandlets
-	static FSimpleMulticastDelegate OnPostEngineInit;
+	static CORE_API FSimpleMulticastDelegate OnPostEngineInit;
 
 	// Called after all modules have been loaded for all phases
-	static FSimpleMulticastDelegate OnAllModuleLoadingPhasesComplete;
+	static CORE_API FSimpleMulticastDelegate OnAllModuleLoadingPhasesComplete;
 
 	// Called at the very end of engine initialization, right before the engine starts ticking. This is not called for commandlets
-	static FSimpleMulticastDelegate OnFEngineLoopInitComplete;
+	static CORE_API FSimpleMulticastDelegate OnFEngineLoopInitComplete;
 
 	// Called when the application is about to exit.
-	static FSimpleMulticastDelegate OnExit;
+	static CORE_API FSimpleMulticastDelegate OnExit;
 
 	// Called when before the application is exiting.
-	static FSimpleMulticastDelegate OnPreExit;
+	static CORE_API FSimpleMulticastDelegate OnPreExit;
 
 	// Called before the engine exits. Separate from OnPreExit as OnEnginePreExit occurs before shutting down any core modules.
-	static FSimpleMulticastDelegate OnEnginePreExit;
+	static CORE_API FSimpleMulticastDelegate OnEnginePreExit;
 
 	/** Delegate for gathering up additional localization paths that are unknown to the engine core (such as plugins) */
-	static TMulticastDelegate<void(TArray<FString>&)> GatherAdditionalLocResPathsCallback;
+	static CORE_API TMulticastDelegate<void(TArray<FString>&)> GatherAdditionalLocResPathsCallback;
 
 	/** Color picker color has changed, please refresh as needed*/
-	static FSimpleMulticastDelegate ColorPickerChanged;
+	static CORE_API FSimpleMulticastDelegate ColorPickerChanged;
 
 	/** requests to open a message box */
-	static TDelegate<EAppReturnType::Type(EAppMsgType::Type, const FText&, const FText&)> ModalErrorMessage;
+	static CORE_API TDelegate<EAppReturnType::Type(EAppMsgType::Type, const FText&, const FText&)> ModalErrorMessage;
 
 	/** Called when the user accepts an invitation to the current game */
-	static TMulticastDelegate<void(const FString&, const FString&)> OnInviteAccepted;
+	static CORE_API TMulticastDelegate<void(const FString&, const FString&)> OnInviteAccepted;
 
 	// Called at the beginning of a frame
-	static FSimpleMulticastDelegate OnBeginFrame;
+	static CORE_API FSimpleMulticastDelegate OnBeginFrame;
 
 	// Called at the moment of sampling the input (currently on the gamethread)
-	static FSimpleMulticastDelegate OnSamplingInput;
+	static CORE_API FSimpleMulticastDelegate OnSamplingInput;
 
 	// Called at the end of a frame
-	static FSimpleMulticastDelegate OnEndFrame;
+	static CORE_API FSimpleMulticastDelegate OnEndFrame;
 
 	// Called at the beginning of a frame on the renderthread
-	static FSimpleMulticastDelegate OnBeginFrameRT;
+	static CORE_API FSimpleMulticastDelegate OnBeginFrameRT;
 
 	// Called at the end of a frame on the renderthread
-	static FSimpleMulticastDelegate OnEndFrameRT;
+	static CORE_API FSimpleMulticastDelegate OnEndFrameRT;
 
 
 	/** called before world origin shifting */
-	static TMulticastDelegate<void(class UWorld*, FIntVector, FIntVector)> PreWorldOriginOffset;
+	static CORE_API TMulticastDelegate<void(class UWorld*, FIntVector, FIntVector)> PreWorldOriginOffset;
 	/** called after world origin shifting */
-	static TMulticastDelegate<void(class UWorld*, FIntVector, FIntVector)> PostWorldOriginOffset;
+	static CORE_API TMulticastDelegate<void(class UWorld*, FIntVector, FIntVector)> PostWorldOriginOffset;
 
 	/** called when the main loop would otherwise starve. */
-	static TDelegate<void()> StarvedGameLoop;
+	static CORE_API TDelegate<void()> StarvedGameLoop;
 
 	// IOS-style temperature updates, allowing game to scale down to let temp drop (to avoid thermal throttling on mobile, for instance) */
 	// There is a parellel enum in ApplicationLifecycleComponent
@@ -249,219 +249,219 @@ public:
 
 		NumSeverities,
 	};
-	static TMulticastDelegate<void(ETemperatureSeverity)> OnTemperatureChange;
+	static CORE_API TMulticastDelegate<void(ETemperatureSeverity)> OnTemperatureChange;
 
 	/** Called when the OS goes into low power mode */
-	static TMulticastDelegate<void(bool)> OnLowPowerMode;
+	static CORE_API TMulticastDelegate<void(bool)> OnLowPowerMode;
 
 	DECLARE_MULTICAST_DELEGATE_OneParam(FOnNetworkConnectionChanged, ENetworkConnectionType /*ConnectionType*/);
-	static FOnNetworkConnectionChanged OnNetworkConnectionChanged;
+	static CORE_API FOnNetworkConnectionChanged OnNetworkConnectionChanged;
 
-	static TTSMulticastDelegate<void(const TCHAR* IniFilename, int32& ResponderCount)>& TSCountPreLoadConfigFileRespondersDelegate();
-	static TTSMulticastDelegate<void(const TCHAR* IniFilename, FString& LoadedContents)>& TSPreLoadConfigFileDelegate();
-	static TTSMulticastDelegate<void(const TCHAR* IniFilename, const FString& ContentsToSave, int32& SavedCount)>& TSPreSaveConfigFileDelegate();
-	static TTSMulticastDelegate<void(const FConfigFile*)>& TSOnFConfigCreated();
-	static TTSMulticastDelegate<void(const FConfigFile*)>& TSOnFConfigDeleted();
-	static TTSMulticastDelegate<void(const TCHAR* IniFilename, const TCHAR* SectionName, const TCHAR* Key)>& TSOnConfigValueRead();
-	static TTSMulticastDelegate<void(const TCHAR* IniFilename, const TCHAR* SectionName)>& TSOnConfigSectionRead();
-	static TTSMulticastDelegate<void(const TCHAR* IniFilename, const TCHAR* SectionName)>& TSOnConfigSectionNameRead();
-	static TTSMulticastDelegate<void(const FString& IniFilename, const TSet<FString>& SectionNames)>& TSOnConfigSectionsChanged();
+	static CORE_API TTSMulticastDelegate<void(const TCHAR* IniFilename, int32& ResponderCount)>& TSCountPreLoadConfigFileRespondersDelegate();
+	static CORE_API TTSMulticastDelegate<void(const TCHAR* IniFilename, FString& LoadedContents)>& TSPreLoadConfigFileDelegate();
+	static CORE_API TTSMulticastDelegate<void(const TCHAR* IniFilename, const FString& ContentsToSave, int32& SavedCount)>& TSPreSaveConfigFileDelegate();
+	static CORE_API TTSMulticastDelegate<void(const FConfigFile*)>& TSOnFConfigCreated();
+	static CORE_API TTSMulticastDelegate<void(const FConfigFile*)>& TSOnFConfigDeleted();
+	static CORE_API TTSMulticastDelegate<void(const TCHAR* IniFilename, const TCHAR* SectionName, const TCHAR* Key)>& TSOnConfigValueRead();
+	static CORE_API TTSMulticastDelegate<void(const TCHAR* IniFilename, const TCHAR* SectionName)>& TSOnConfigSectionRead();
+	static CORE_API TTSMulticastDelegate<void(const TCHAR* IniFilename, const TCHAR* SectionName)>& TSOnConfigSectionNameRead();
+	static CORE_API TTSMulticastDelegate<void(const FString& IniFilename, const TSet<FString>& SectionNames)>& TSOnConfigSectionsChanged();
 
-	static TMulticastDelegate<void(const TCHAR* SectionName, const TCHAR* IniFilename, uint32 SetBy, bool bAllowCheating)> OnApplyCVarFromIni;
+	static CORE_API TMulticastDelegate<void(const TCHAR* SectionName, const TCHAR* IniFilename, uint32 SetBy, bool bAllowCheating)> OnApplyCVarFromIni;
 
-	static TMulticastDelegate<void(uint32 ResX, uint32 ResY)> OnSystemResolutionChanged;
+	static CORE_API TMulticastDelegate<void(uint32 ResX, uint32 ResY)> OnSystemResolutionChanged;
 
 #if WITH_EDITOR
 	// called when a target platform changes it's return value of supported formats.  This is so anything caching those results can reset (like cached shaders for cooking)
-	static TMulticastDelegate<void(const ITargetPlatform*)> OnTargetPlatformChangedSupportedFormats;
+	static CORE_API TMulticastDelegate<void(const ITargetPlatform*)> OnTargetPlatformChangedSupportedFormats;
 
 	// Called when a feature level is disabled by the user.
-	static TMulticastDelegate<void(int, const FName&)> OnFeatureLevelDisabled;
+	static CORE_API TMulticastDelegate<void(int, const FName&)> OnFeatureLevelDisabled;
 #endif
 
 	/** IOS-style application lifecycle delegates */
 
 	// This is called when the application is about to be deactivated (e.g., due to a phone call or SMS or the sleep button).
 	// The game should be paused if possible, etc...
-	static TMulticastDelegate<void()> ApplicationWillDeactivateDelegate;
+	static CORE_API TMulticastDelegate<void()> ApplicationWillDeactivateDelegate;
 
 	// Called when the application has been reactivated (reverse any processing done in the Deactivate delegate)
-	static TMulticastDelegate<void()> ApplicationHasReactivatedDelegate;
+	static CORE_API TMulticastDelegate<void()> ApplicationHasReactivatedDelegate;
 
 	// This is called when the application is being backgrounded (e.g., due to switching
 	// to another app or closing it via the home button)
 	// The game should release shared resources, save state, etc..., since it can be
 	// terminated from the background state without any further warning.
-	static TMulticastDelegate<void()> ApplicationWillEnterBackgroundDelegate; // for instance, hitting the home button
+	static CORE_API TMulticastDelegate<void()> ApplicationWillEnterBackgroundDelegate; // for instance, hitting the home button
 
 	// Called when the application is returning to the foreground (reverse any processing done in the EnterBackground delegate)
-	static TMulticastDelegate<void()> ApplicationHasEnteredForegroundDelegate;
+	static CORE_API TMulticastDelegate<void()> ApplicationHasEnteredForegroundDelegate;
 
 	// This *may* be called when the application is getting terminated by the OS.
 	// There is no guarantee that this will ever be called on a mobile device,
 	// save state when ApplicationWillEnterBackgroundDelegate is called instead.
 	UE_DEPRECATED(5.3, "This delegate is not thread safe while it can be used concurrently. Please use to GetApplicationWillTerminateDelegate() instead.")
-	static TMulticastDelegate<void()> ApplicationWillTerminateDelegate;
+	static CORE_API TMulticastDelegate<void()> ApplicationWillTerminateDelegate;
 
-	static TTSMulticastDelegate<void()>& GetApplicationWillTerminateDelegate();
+	static CORE_API TTSMulticastDelegate<void()>& GetApplicationWillTerminateDelegate();
 
 	// Some platform have a System UI Overlay that can draw on top of the application.
 	// The game might want to be notified so it can pause, etc...
 	// Parameter (bool) should be true if the system UI is displayed, otherwise false should be passed
-	static TMulticastDelegate<void(bool)> ApplicationSystemUIOverlayStateChangedDelegate;
+	static CORE_API TMulticastDelegate<void(bool)> ApplicationSystemUIOverlayStateChangedDelegate;
 
 	// Called when in the background, if the OS is giving CPU time to the device. It is very likely
 	// this will never be called due to mobile OS backgrounded CPU restrictions. But if, for instance,
 	// VOIP is active on iOS, the will be getting called
-	static TMulticastDelegate<void(float DeltaTime)> MobileBackgroundTickDelegate;
+	static CORE_API TMulticastDelegate<void(float DeltaTime)> MobileBackgroundTickDelegate;
 
 	// Called when the OS needs control of the music (parameter is true) or when the OS returns
 	// control of the music to the application (parameter is false). This can happen due to a
 	// phone call or timer or other OS-level event. This is currently triggered only on iOS
 	// devices.
-	static TMulticastDelegate<void(bool)> UserMusicInterruptDelegate;
+	static CORE_API TMulticastDelegate<void(bool)> UserMusicInterruptDelegate;
 	
 	// [iOS only] Called when the mute switch is detected as changed or when the
 	// volume changes. Parameter 1 is the mute switch state (true is muted, false is
 	// unmuted). Parameter 2 is the volume as an integer from 0 to 100.
-	static TMulticastDelegate<void(bool, int)> AudioMuteDelegate;
+	static CORE_API TMulticastDelegate<void(bool, int)> AudioMuteDelegate;
 	
 	// [iOS only] Called when the audio device changes
 	// For instance, when the headphones are plugged in or removed
-	static TMulticastDelegate<void(bool)> AudioRouteChangedDelegate;
+	static CORE_API TMulticastDelegate<void(bool)> AudioRouteChangedDelegate;
 
 	// Generally, events triggering UserMusicInterruptDelegate or AudioMuteDelegate happen only
 	// when a change occurs. When a system comes online needing the current audio state but the
 	// event has already been broadcast, calling ApplicationRequestAudioState will force the
 	// UserMusicInterruptDelegate and AudioMuteDelegate to be called again if the low-level
 	// application layer supports it. Currently, this is available only on iOS.
-	static TMulticastDelegate<void()> ApplicationRequestAudioState;
+	static CORE_API TMulticastDelegate<void()> ApplicationRequestAudioState;
 	
 	// Called when the OS is running low on resources and asks the application to free up any cached resources, drop graphics quality etc.
-	static TMulticastDelegate<void()> ApplicationShouldUnloadResourcesDelegate;
+	static CORE_API TMulticastDelegate<void()> ApplicationShouldUnloadResourcesDelegate;
 
 	// Called with arguments passed to the application on statup, perhaps meta data passed on by another application which launched this one.
-	static TMulticastDelegate<void(const TArray<FString>&)> ApplicationReceivedStartupArgumentsDelegate;
+	static CORE_API TMulticastDelegate<void(const TArray<FString>&)> ApplicationReceivedStartupArgumentsDelegate;
 
 	/** IOS-style push notification delegates */
 
 	// called when the user grants permission to register for remote notifications
-	static TMulticastDelegate<void(TArray<uint8>)> ApplicationRegisteredForRemoteNotificationsDelegate;
+	static CORE_API TMulticastDelegate<void(TArray<uint8>)> ApplicationRegisteredForRemoteNotificationsDelegate;
 
 	// called when the user grants permission to register for notifications
-	static TMulticastDelegate<void(int)> ApplicationRegisteredForUserNotificationsDelegate;
+	static CORE_API TMulticastDelegate<void(int)> ApplicationRegisteredForUserNotificationsDelegate;
 
 	// called when the application fails to register for remote notifications
-	static TMulticastDelegate<void(FString)> ApplicationFailedToRegisterForRemoteNotificationsDelegate;
+	static CORE_API TMulticastDelegate<void(FString)> ApplicationFailedToRegisterForRemoteNotificationsDelegate;
 
 	// called when the application receives a remote notification
-	static TMulticastDelegate<void(FString, int)> ApplicationReceivedRemoteNotificationDelegate;
+	static CORE_API TMulticastDelegate<void(FString, int)> ApplicationReceivedRemoteNotificationDelegate;
 
 	// called when the application receives a local notification
-	static TMulticastDelegate<void(FString, int, int)> ApplicationReceivedLocalNotificationDelegate;
+	static CORE_API TMulticastDelegate<void(FString, int, int)> ApplicationReceivedLocalNotificationDelegate;
 
     // called when the application receives notice to perform a background fetch
-    static TMulticastDelegate<void()> ApplicationPerformFetchDelegate;
+    static CORE_API TMulticastDelegate<void()> ApplicationPerformFetchDelegate;
 
     // called when the application receives notice that a background download has completed
-    static TMulticastDelegate<void(FString)> ApplicationBackgroundSessionEventDelegate;
+    static CORE_API TMulticastDelegate<void(FString)> ApplicationBackgroundSessionEventDelegate;
 
 	/** Sent when a device screen orientation changes */
-	static TMulticastDelegate<void(int32)> ApplicationReceivedScreenOrientationChangedNotificationDelegate;
+	static CORE_API TMulticastDelegate<void(int32)> ApplicationReceivedScreenOrientationChangedNotificationDelegate;
 
 	/** Checks to see if the stat is already enabled */
-	static TMulticastDelegate<void(const TCHAR*, bool&, bool&)> StatCheckEnabled;
+	static CORE_API TMulticastDelegate<void(const TCHAR*, bool&, bool&)> StatCheckEnabled;
 
 	/** Sent after each stat is enabled */
-	static TMulticastDelegate<void(const TCHAR*)> StatEnabled;
+	static CORE_API TMulticastDelegate<void(const TCHAR*)> StatEnabled;
 
 	/** Sent after each stat is disabled */
-	static TMulticastDelegate<void(const TCHAR*)> StatDisabled;
+	static CORE_API TMulticastDelegate<void(const TCHAR*)> StatDisabled;
 
 	/** Sent when all stats need to be disabled */
-	static TMulticastDelegate<void(const bool)> StatDisableAll;
+	static CORE_API TMulticastDelegate<void(const bool)> StatDisableAll;
 
 	// Called when an application is notified that the application license info has been updated.
 	// The new license data should be polled and steps taken based on the results (i.e. halt application if license is no longer valid).
-	static TMulticastDelegate<void()> ApplicationLicenseChange;
+	static CORE_API TMulticastDelegate<void()> ApplicationLicenseChange;
 
 	/** Sent when the platform changed its laptop mode (for convertible laptops).*/
-	static TMulticastDelegate<void(EConvertibleLaptopMode)> PlatformChangedLaptopMode;
+	static CORE_API TMulticastDelegate<void(EConvertibleLaptopMode)> PlatformChangedLaptopMode;
 
 	/** Sent when the platform needs the user to fix headset tracking on startup (Most platforms do not need this.) */
-	static TMulticastDelegate<void()> VRHeadsetTrackingInitializingAndNeedsHMDToBeTrackedDelegate;
+	static CORE_API TMulticastDelegate<void()> VRHeadsetTrackingInitializingAndNeedsHMDToBeTrackedDelegate;
 
 	/** Sent when the platform finds that needed headset tracking on startup has completed (Most platforms do not need this.) */
-	static TMulticastDelegate<void()> VRHeadsetTrackingInitializedDelegate;
+	static CORE_API TMulticastDelegate<void()> VRHeadsetTrackingInitializedDelegate;
 
 	/** Sent when the platform requests a low-level VR recentering */
-	static TMulticastDelegate<void()> VRHeadsetRecenter;
+	static CORE_API TMulticastDelegate<void()> VRHeadsetRecenter;
 
 	/** Sent when connection to VR HMD is lost */
-	static TMulticastDelegate<void()> VRHeadsetLost;
+	static CORE_API TMulticastDelegate<void()> VRHeadsetLost;
 
 	/** Sent when connection to VR HMD is restored */
-	static TMulticastDelegate<void()> VRHeadsetReconnected;
+	static CORE_API TMulticastDelegate<void()> VRHeadsetReconnected;
 
 	/** Sent when connection to VR HMD connection is refused by the player */
-	static TMulticastDelegate<void()> VRHeadsetConnectCanceled;
+	static CORE_API TMulticastDelegate<void()> VRHeadsetConnectCanceled;
 
 	/** Sent when the VR HMD detects that it has been put on by the player. */
-	static TMulticastDelegate<void()> VRHeadsetPutOnHead;
+	static CORE_API TMulticastDelegate<void()> VRHeadsetPutOnHead;
 
 	/** Sent when the VR HMD detects that it has been taken off by the player. */
-	static TMulticastDelegate<void()> VRHeadsetRemovedFromHead;
+	static CORE_API TMulticastDelegate<void()> VRHeadsetRemovedFromHead;
 
 	/** Sent when a 3DOF VR controller is recentered */
-	static TMulticastDelegate<void()> VRControllerRecentered;
+	static CORE_API TMulticastDelegate<void()> VRControllerRecentered;
 
 	/** Sent when application code changes the user activity hint string for analytics, crash reports, etc */
-	static TMulticastDelegate<void(const FString&)> UserActivityStringChanged;
+	static CORE_API TMulticastDelegate<void(const FString&)> UserActivityStringChanged;
 
 	/** Sent when application code changes the currently active game session. The exact semantics of this will vary between games but it is useful for analytics, crash reports, etc  */
-	static TMulticastDelegate<void(const FString&)> GameSessionIDChanged;
+	static CORE_API TMulticastDelegate<void(const FString&)> GameSessionIDChanged;
 
 	/** Sent when application code changes game state. The exact semantics of this will vary between games but it is useful for analytics, crash reports, etc  */
-	static TMulticastDelegate<void(const FString&)> GameStateClassChanged;
+	static CORE_API TMulticastDelegate<void(const FString&)> GameStateClassChanged;
 
 	/** Sent by application code to set params that customize crash reporting behavior. */
-	static TMulticastDelegate<void(const FCrashOverrideParameters&)> CrashOverrideParamsChanged;
+	static CORE_API TMulticastDelegate<void(const FCrashOverrideParameters&)> CrashOverrideParamsChanged;
 	
 	/** Sent by engine code when the "vanilla" status of the engine changes */
-	static TMulticastDelegate<void(bool)> IsVanillaProductChanged;
+	static CORE_API TMulticastDelegate<void(bool)> IsVanillaProductChanged;
 
 	// Callback for platform specific very early init code.
-	static TMulticastDelegate<void()>& GetPreMainInitDelegate();
+	static CORE_API TMulticastDelegate<void()>& GetPreMainInitDelegate();
 	
 	/** Sent when GConfig is finished initializing */
-	static TTSMulticastDelegate<void()>& TSConfigReadyForUse();
+	static CORE_API TTSMulticastDelegate<void()>& TSConfigReadyForUse();
 
 	/** Callback for notifications regarding changes of the rendering thread. */
 
 	/** Sent just after the rendering thread has been created. */
-	static TMulticastDelegate<void()> PostRenderingThreadCreated;
+	static CORE_API TMulticastDelegate<void()> PostRenderingThreadCreated;
 	/* Sent just before the rendering thread is destroyed. */
-	static TMulticastDelegate<void()> PreRenderingThreadDestroyed;
+	static CORE_API TMulticastDelegate<void()> PreRenderingThreadDestroyed;
 
 	// Callback to allow custom resolution of package names. Arguments are InRequestedName, OutResolvedName.
 	// Should return True of resolution occured.
-	static TArray<TDelegate<bool(const FString&, FString&)>> PackageNameResolvers;
+	static CORE_API TArray<TDelegate<bool(const FString&, FString&)>> PackageNameResolvers;
 
 	// Called to request that systems free whatever memory they are able to. Called early in LoadMap.
 	// Caller is responsible for flushing rendering etc. See UEngine::TrimMemory
-	static FSimpleMulticastDelegate& GetMemoryTrimDelegate();
+	static CORE_API FSimpleMulticastDelegate& GetMemoryTrimDelegate();
 
 	// Called to request that low level allocator free whatever memory they are able to. 
-	static FSimpleMulticastDelegate& GetLowLevelAllocatorMemoryTrimDelegate();
+	static CORE_API FSimpleMulticastDelegate& GetLowLevelAllocatorMemoryTrimDelegate();
 
 	// Called to request that low level allocator must refreshed
-	static FSimpleMulticastDelegate& GetRefreshLowLevelAllocatorDelegate();
+	static CORE_API FSimpleMulticastDelegate& GetRefreshLowLevelAllocatorDelegate();
 
 	// Called when OOM event occurs, after backup memory has been freed, so there's some hope of being effective
-	static FSimpleMulticastDelegate& GetOutOfMemoryDelegate();
+	static CORE_API FSimpleMulticastDelegate& GetOutOfMemoryDelegate();
 
 	// Called from TerminateOnOutOfMemory in D3D11Util.cpp/D3D12Util.cpp
-	static TMulticastDelegate<void(const uint64, const uint64)>& GetGPUOutOfMemoryDelegate();
+	static CORE_API TMulticastDelegate<void(const uint64, const uint64)>& GetGPUOutOfMemoryDelegate();
 
 	enum class EOnScreenMessageSeverity : uint8
 	{
@@ -477,25 +477,25 @@ public:
 	// {
 	//		OutMessages.Add(FCoreDelegates::EOnScreenMessageSeverity::Info, FText::Format(LOCTEXT("MyMessage", "My Status: {0}"), SomeStatus));
 	// }
-	static TMulticastDelegate<void(FSeverityMessageMap&)> OnGetOnScreenMessages;
+	static CORE_API TMulticastDelegate<void(FSeverityMessageMap&)> OnGetOnScreenMessages;
 
-	static TDelegate<bool()> IsLoadingMovieCurrentlyPlaying;
+	static CORE_API TDelegate<bool()> IsLoadingMovieCurrentlyPlaying;
 
 	// Callback to allow user code to prevent url from being launched from FPlatformProcess::LaunchURL. Used to apply http allow list
 	// Return true for to launch the url
-	static TDelegate<bool(const TCHAR* URL)> ShouldLaunchUrl;
+	static CORE_API TDelegate<bool(const TCHAR* URL)> ShouldLaunchUrl;
 
 	// Callback to allow user code to implement a custom implementation for FPlatformProcess::LaunchURL. (Windows Only)
-	static TDelegate<void(const FString& URL, FString* Error)> LaunchCustomHandlerForURL;
+	static CORE_API TDelegate<void(const FString& URL, FString* Error)> LaunchCustomHandlerForURL;
 	
 	// Callback when the application has been activated by protocol (with optional user id, depending on the platform)
-	static TMulticastDelegate<void(const FString& Parameter, FPlatformUserId UserId /*= PLATFORMUSERID_NONE*/)> OnActivatedByProtocol;
+	static CORE_API TMulticastDelegate<void(const FString& Parameter, FPlatformUserId UserId /*= PLATFORMUSERID_NONE*/)> OnActivatedByProtocol;
 
 	/** Sent when GC finish destroy takes more time than expected */
-	static TMulticastDelegate<void(const FString&)> OnGCFinishDestroyTimeExtended;
+	static CORE_API TMulticastDelegate<void(const FString&)> OnGCFinishDestroyTimeExtended;
 
 	/** Called when the application's network initializes or shutdowns on platforms where the network stack is not always available */
-	static TMulticastDelegate<void(bool bIsNetworkInitialized)> ApplicationNetworkInitializationChanged;
+	static CORE_API TMulticastDelegate<void(bool bIsNetworkInitialized)> ApplicationNetworkInitializationChanged;
 
 	/**
 	 * Called when the connection state as reported by the platform changes
@@ -503,7 +503,7 @@ public:
 	 * @param LastConnectionState last state of the connection
 	 * @param ConnectionState current state of the connection
 	 */
-	static TMulticastDelegate<void(ENetworkConnectionStatus LastConnectionState, ENetworkConnectionStatus ConnectionState)> OnNetworkConnectionStatusChanged;
+	static CORE_API TMulticastDelegate<void(ENetworkConnectionStatus LastConnectionState, ENetworkConnectionStatus ConnectionState)> OnNetworkConnectionStatusChanged;
 
 	// Callback to let code read or write specialized binary data that is generated at Stage time, for optimizing data right before 
 	// final game data is being written to disk
@@ -527,40 +527,40 @@ public:
 		}
 	};
 	
-	static TTSMulticastDelegate<void(FExtraBinaryConfigData&)>& TSAccessExtraBinaryConfigData();
+	static CORE_API TTSMulticastDelegate<void(FExtraBinaryConfigData&)>& TSAccessExtraBinaryConfigData();
 
 	using FAttachShaderReadRequestFunc = TFunctionRef<class FIoRequest(const class FIoChunkId&, FGraphEventRef)>;
-	static TDelegate<void(TArrayView<const FSHAHash>, FAttachShaderReadRequestFunc)> PreloadPackageShaderMaps;
-	static TDelegate<void(TArrayView<const FSHAHash>)> ReleasePreloadedPackageShaderMaps;
+	static CORE_API TDelegate<void(TArrayView<const FSHAHash>, FAttachShaderReadRequestFunc)> PreloadPackageShaderMaps;
+	static CORE_API TDelegate<void(TArrayView<const FSHAHash>)> ReleasePreloadedPackageShaderMaps;
 	/** Called when the verbosity of a log category is changed */
-	static TMulticastDelegate<void(const FLogCategoryName& CategoryName, ELogVerbosity::Type OldVerbosity, ELogVerbosity::Type NewVerbosity)> OnLogVerbosityChanged;
+	static CORE_API TMulticastDelegate<void(const FLogCategoryName& CategoryName, ELogVerbosity::Type OldVerbosity, ELogVerbosity::Type NewVerbosity)> OnLogVerbosityChanged;
 
 	UE_DEPRECATED(5.1, "Use FPackageStore::Mount() instead")
-	static TDelegate<TSharedPtr<class IPackageStore>()> CreatePackageStore;
+	static CORE_API TDelegate<TSharedPtr<class IPackageStore>()> CreatePackageStore;
 
 	// Called immediately before the parent process will start responding to signals to fork
-	static FSimpleMulticastDelegate OnParentBeginFork;
+	static CORE_API FSimpleMulticastDelegate OnParentBeginFork;
 	// Called each time immediately before the parent process forks itself
-	static FSimpleMulticastDelegate OnParentPreFork;
+	static CORE_API FSimpleMulticastDelegate OnParentPreFork;
 
 	// Called immediately after the process spawned a fork
-	static TMulticastDelegate<void(EForkProcessRole ProcessRole)> OnPostFork;
+	static CORE_API TMulticastDelegate<void(EForkProcessRole ProcessRole)> OnPostFork;
 	// Called at the end of the frame where the process spawned a fork
-	static FSimpleMulticastDelegate OnChildEndFramePostFork;
+	static CORE_API FSimpleMulticastDelegate OnChildEndFramePostFork;
 
 	// Called when FParse::Command succeeds. Used to audit named commands usage.
 	DECLARE_MULTICAST_DELEGATE_OneParam(FOnNamedCommandParsed, const TCHAR* /*Cmd*/);
-	static FOnNamedCommandParsed OnNamedCommandParsed;
+	static CORE_API FOnNamedCommandParsed OnNamedCommandParsed;
 
 	// Called when FExec::Exec is called in a context where it's disallowed (UE_ALLOW_EXEC_COMMANDS = 0).
 	DECLARE_MULTICAST_DELEGATE_OneParam(FOnDisallowedExecCommandCalled, const TCHAR* /*Cmd*/);
-	static FOnDisallowedExecCommandCalled OnDisallowedExecCommandCalled;
+	static CORE_API FOnDisallowedExecCommandCalled OnDisallowedExecCommandCalled;
 
 	// Extension point for projects to report the URL for a continuous integration job which produced these binaries.
-	static TDelegate<const TCHAR*()> OnGetBuildURL;
+	static CORE_API TDelegate<const TCHAR*()> OnGetBuildURL;
 	
 	// Extension point for projects to report the URL for a continuous integration job which is currently executing this process. 
-	static TDelegate<const TCHAR*()> OnGetExecutingJobURL;
+	static CORE_API TDelegate<const TCHAR*()> OnGetExecutingJobURL;
 	
 private:
 
diff --git a/Engine/Source/Runtime/Core/Public/Misc/CoreMisc.h b/Engine/Source/Runtime/Core/Public/Misc/CoreMisc.h
--- a/Engine/Source/Runtime/Core/Public/Misc/CoreMisc.h
+++ b/Engine/Source/Runtime/Core/Public/Misc/CoreMisc.h
@@ -24,29 +24,29 @@ class UWorld;
  * Exec handler that registers itself and is being routed via StaticExec.
  * Note: Not intended for use with UObjects!
  */
-class CORE_API FSelfRegisteringExec : public FExec
+class FSelfRegisteringExec : public FExec
 {
 public:
 	/** Constructor, registering this instance. */
-	FSelfRegisteringExec();
+	CORE_API FSelfRegisteringExec();
 	/** Destructor, unregistering this instance. */
-	virtual ~FSelfRegisteringExec();
+	CORE_API virtual ~FSelfRegisteringExec();
 
 	/** Routes a command to the self-registered execs. */
-	static bool StaticExec( UWorld* Inworld, const TCHAR* Cmd, FOutputDevice& Ar );
+	static CORE_API bool StaticExec( UWorld* Inworld, const TCHAR* Cmd, FOutputDevice& Ar );
 };
 
 /** Registers a static Exec function using FSelfRegisteringExec. */
-class CORE_API FStaticSelfRegisteringExec : public FSelfRegisteringExec
+class FStaticSelfRegisteringExec : public FSelfRegisteringExec
 {
 public:
 
 	/** Initialization constructor. */
-	FStaticSelfRegisteringExec(bool (*InStaticExecFunc)(UWorld* Inworld, const TCHAR* Cmd,FOutputDevice& Ar));
+	CORE_API FStaticSelfRegisteringExec(bool (*InStaticExecFunc)(UWorld* Inworld, const TCHAR* Cmd,FOutputDevice& Ar));
 
 	//~ Begin Exec Interface
 #if UE_ALLOW_EXEC_COMMANDS
-	virtual bool Exec( UWorld* InWorld, const TCHAR* Cmd, FOutputDevice& Ar ) override;
+	CORE_API virtual bool Exec( UWorld* InWorld, const TCHAR* Cmd, FOutputDevice& Ar ) override;
 #endif
 	//~ End Exec Interface
 
@@ -56,15 +56,15 @@ private:
 };
 
 /** Registers a static Exec_Dev function using FSelfRegisteringExec. */
-class CORE_API FStaticSelfRegisteringExec_Dev: public FSelfRegisteringExec
+class FStaticSelfRegisteringExec_Dev: public FSelfRegisteringExec
 {
 public:
 
 	/** Initialization constructor. */
-	explicit FStaticSelfRegisteringExec_Dev(bool (*InStaticExecFunc)(UWorld* Inworld, const TCHAR* Cmd, FOutputDevice& Ar));
+	CORE_API explicit FStaticSelfRegisteringExec_Dev(bool (*InStaticExecFunc)(UWorld* Inworld, const TCHAR* Cmd, FOutputDevice& Ar));
 
 	//~ Begin Exec Interface
-	virtual bool Exec_Dev(UWorld* InWorld, const TCHAR* Cmd, FOutputDevice& Ar) override;
+	CORE_API virtual bool Exec_Dev(UWorld* InWorld, const TCHAR* Cmd, FOutputDevice& Ar) override;
 	//~ End Exec Interface
 
 private:
@@ -73,15 +73,15 @@ private:
 };
 
 /** Registers a static Exec_Editor function using FSelfRegisteringExec. */
-class CORE_API FStaticSelfRegisteringExec_Editor: public FSelfRegisteringExec
+class FStaticSelfRegisteringExec_Editor: public FSelfRegisteringExec
 {
 public:
 
 	/** Initialization constructor. */
-	explicit FStaticSelfRegisteringExec_Editor(bool (*InStaticExecFunc)(UWorld* Inworld, const TCHAR* Cmd, FOutputDevice& Ar));
+	CORE_API explicit FStaticSelfRegisteringExec_Editor(bool (*InStaticExecFunc)(UWorld* Inworld, const TCHAR* Cmd, FOutputDevice& Ar));
 
 	//~ Begin Exec Interface
-	virtual bool Exec_Editor(UWorld* InWorld, const TCHAR* Cmd, FOutputDevice& Ar) override;
+	CORE_API virtual bool Exec_Editor(UWorld* InWorld, const TCHAR* Cmd, FOutputDevice& Ar) override;
 	//~ End Exec Interface
 
 private:
@@ -98,10 +98,10 @@ public:
 };
 
 
-struct CORE_API FMaintenance
+struct FMaintenance
 {
 	/** deletes log files older than a number of days specified in the Engine ini file */
-	static void DeleteOldLogs();
+	static CORE_API void DeleteOldLogs();
 };
 
 /*-----------------------------------------------------------------------------
@@ -210,7 +210,7 @@ FORCEINLINE bool IsRunningClientOnly()
 /**
  * Helper for obtaining the default Url configuration
  */
-struct CORE_API FUrlConfig
+struct FUrlConfig
 {
 	FString DefaultProtocol;
 	FString DefaultName;
@@ -222,23 +222,23 @@ struct CORE_API FUrlConfig
 	/**
 	 * Initialize with defaults from ini
 	 */
-	void Init();
+	CORE_API void Init();
 
 	/**
 	 * Reset state
 	 */
-	void Reset();
+	CORE_API void Reset();
 };
 
 bool CORE_API StringHasBadDashes(const TCHAR* Str);
 
 /** Helper structure for boolean values in config */
-struct CORE_API FBoolConfigValueHelper
+struct FBoolConfigValueHelper
 {
 private:
 	bool bValue;
 public:
-	FBoolConfigValueHelper(const TCHAR* Section, const TCHAR* Key, const FString& Filename = GEditorIni);
+	CORE_API FBoolConfigValueHelper(const TCHAR* Section, const TCHAR* Key, const FString& Filename = GEditorIni);
 
 	operator bool() const
 	{
@@ -303,10 +303,10 @@ private:
 /** 
  * Scoped struct used to push and pop a script exception handler
  */
-struct CORE_API FScopedScriptExceptionHandler
+struct FScopedScriptExceptionHandler
 {
-	explicit FScopedScriptExceptionHandler(const FScriptExceptionHandlerFunc& InFunc);
-	~FScopedScriptExceptionHandler();
+	CORE_API explicit FScopedScriptExceptionHandler(const FScriptExceptionHandlerFunc& InFunc);
+	CORE_API ~FScopedScriptExceptionHandler();
 };
 
 /**
diff --git a/Engine/Source/Runtime/Core/Public/Misc/Crc.h b/Engine/Source/Runtime/Core/Public/Misc/Crc.h
--- a/Engine/Source/Runtime/Core/Public/Misc/Crc.h
+++ b/Engine/Source/Runtime/Core/Public/Misc/Crc.h
@@ -17,18 +17,18 @@ struct TTuple;
 /** 
  * CRC hash generation for different types of input data
  **/
-struct CORE_API FCrc
+struct FCrc
 {
 	/** lookup table with precalculated CRC values - slicing by 8 implementation */
-	static uint32 CRCTablesSB8[8][256];
+	static CORE_API uint32 CRCTablesSB8[8][256];
 
 	/** initializes the CRC lookup table. Must be called before any of the
 		CRC functions are used. */
-	static void Init();
+	static CORE_API void Init();
 
 	/** generates CRC hash of the memory area */
 	typedef uint32 (*MemCrc32Functor)( const void* Data, int32 Length, uint32 CRC );
-	static MemCrc32Functor MemCrc32Func;
+	static CORE_API MemCrc32Functor MemCrc32Func;
 	static FORCEINLINE uint32 MemCrc32(const void* Data, int32 Length, uint32 CRC = 0)
 	{
 		return MemCrc32Func(Data, Length, CRC);
@@ -86,9 +86,9 @@ struct CORE_API FCrc
 	 */
 
 	/** lookup table with precalculated CRC values */
-	static uint32 CRCTable_DEPRECATED[256];
+	static CORE_API uint32 CRCTable_DEPRECATED[256];
 	/** lookup table with precalculated CRC values - slicing by 8 implementation */
-	static uint32 CRCTablesSB8_DEPRECATED[8][256];
+	static CORE_API uint32 CRCTablesSB8_DEPRECATED[8][256];
 
 	/** String CRC. */
 	template <typename CharType>
@@ -133,7 +133,7 @@ struct CORE_API FCrc
 	template <typename CharType> static inline uint32 Strihash_DEPRECATED( const int32 DataLen, const CharType* Data );
 
 	/** generates CRC hash of the memory area */
-	static uint32 MemCrc_DEPRECATED( const void* Data, int32 Length, uint32 CRC=0 );
+	static CORE_API uint32 MemCrc_DEPRECATED( const void* Data, int32 Length, uint32 CRC=0 );
 };
 
 template <>
diff --git a/Engine/Source/Runtime/Core/Public/Misc/DataDrivenPlatformInfoRegistry.h b/Engine/Source/Runtime/Core/Public/Misc/DataDrivenPlatformInfoRegistry.h
--- a/Engine/Source/Runtime/Core/Public/Misc/DataDrivenPlatformInfoRegistry.h
+++ b/Engine/Source/Runtime/Core/Public/Misc/DataDrivenPlatformInfoRegistry.h
@@ -250,34 +250,34 @@ private:
 };
 
 
-struct CORE_API FDataDrivenPlatformInfoRegistry
+struct FDataDrivenPlatformInfoRegistry
 {
 
 	/**
 	* Get the global set of data driven platform information
 	*/
-	static const TMap<FName, FDataDrivenPlatformInfo>& GetAllPlatformInfos();
+	static CORE_API const TMap<FName, FDataDrivenPlatformInfo>& GetAllPlatformInfos();
 
 	/**
 	 * Gets a set of platform names based on GetAllPlatformInfos, their AdditionalRestrictedFolders, and possibly filtered based on what editor has support compiled for
 	 * This is not necessarily the same as IniParents, although there is overlap - IniParents come from chaining DDPIs, so those will be in GetAllPlatformInfos already to be checked 
 	 */
-	static const TArray<FString>& GetValidPlatformDirectoryNames();
+	static CORE_API const TArray<FString>& GetValidPlatformDirectoryNames();
 
 	/**
 	 * Get the data driven platform info for a given platform. If the platform doesn't have any on disk,
 	 * this will return a default constructed FConfigDataDrivenPlatformInfo
 	 */
-	static const FDataDrivenPlatformInfo& GetPlatformInfo(const FString& PlatformName);
-	static const FDataDrivenPlatformInfo& GetPlatformInfo(FName PlatformName);
-	static const FDataDrivenPlatformInfo& GetPlatformInfo(const char* PlatformName);
+	static CORE_API const FDataDrivenPlatformInfo& GetPlatformInfo(const FString& PlatformName);
+	static CORE_API const FDataDrivenPlatformInfo& GetPlatformInfo(FName PlatformName);
+	static CORE_API const FDataDrivenPlatformInfo& GetPlatformInfo(const char* PlatformName);
 
 
 	/**
 	 * Get sorted platorm infos or names. PlatformType is used to choose between true platforms and platform groups (or other fake platforms)
 	 */
-	static const TArray<FName> GetSortedPlatformNames(EPlatformInfoType PlatformType);
-	static const TArray<const FDataDrivenPlatformInfo*>& GetSortedPlatformInfos(EPlatformInfoType PlatformType);
+	static CORE_API const TArray<FName> GetSortedPlatformNames(EPlatformInfoType PlatformType);
+	static CORE_API const TArray<const FDataDrivenPlatformInfo*>& GetSortedPlatformInfos(EPlatformInfoType PlatformType);
 
 	UE_DEPRECATED(5.1, "Use GetSoprtedPlatformNames that takes a PlatformType parameter")
 	static const TArray<FName> GetSortedPlatformNames()
@@ -294,17 +294,17 @@ struct CORE_API FDataDrivenPlatformInfoRegistry
 	 * Gets a list of all known confidential platforms (note these are just the platforms you have access to, so, for example PS4 won't be
 	 * returned if you are not a PS4 licensee)
 	 */
-	static const TArray<FName>& GetConfidentialPlatforms();
+	static CORE_API const TArray<FName>& GetConfidentialPlatforms();
 
 	/**
 	 * Returns the number of discovered ini files that can be loaded with LoadDataDrivenIniFile
 	 */
-	static int32 GetNumDataDrivenIniFiles();
+	static CORE_API int32 GetNumDataDrivenIniFiles();
 
 	/**
 	 * Load the given ini file, and 
 	 */
-	static bool LoadDataDrivenIniFile(int32 Index, FConfigFile& IniFile, FString& PlatformName);
+	static CORE_API bool LoadDataDrivenIniFile(int32 Index, FConfigFile& IniFile, FString& PlatformName);
 
 
 #if DDPI_HAS_EXTENDED_PLATFORMINFO_DATA
@@ -321,31 +321,31 @@ struct CORE_API FDataDrivenPlatformInfoRegistry
 		// for instance WindowsClient
 		TargetPlatform,
 	};
-	static bool HasCompiledSupportForPlatform(FName PlatformName, EPlatformNameType PlatformNameType);
+	static CORE_API bool HasCompiledSupportForPlatform(FName PlatformName, EPlatformNameType PlatformNameType);
 
 	/**
 	 * Wipes out cached device status for all devices in a platform (or all platforms if PlatformName is empty)
 	 */
-	static void ClearDeviceStatus(FName PlatformName);
+	static CORE_API void ClearDeviceStatus(FName PlatformName);
 
-	static FDataDrivenPlatformInfo& DeviceIdToInfo(FString DeviceId, FString* OutDeviceName = nullptr);
+	static CORE_API FDataDrivenPlatformInfo& DeviceIdToInfo(FString DeviceId, FString* OutDeviceName = nullptr);
 
 	/**
 	 * Retrieve the full list of all known preview platforms from all DDPIs
 	*/
-	static const TArray<struct FPreviewPlatformMenuItem>& GetAllPreviewPlatformMenuItems();
+	static CORE_API const TArray<struct FPreviewPlatformMenuItem>& GetAllPreviewPlatformMenuItems();
 
 	/** Option to hide a platform from the user interface at runtime */
-	static bool IsPlatformHiddenFromUI(FName PlatformName);
-	static void SetPlatformHiddenFromUI(FName PlatformName);
+	static CORE_API bool IsPlatformHiddenFromUI(FName PlatformName);
+	static CORE_API void SetPlatformHiddenFromUI(FName PlatformName);
 
 private:
 	/**
 	 * Get a modifiable DDPI object, for Turnkey to update it's info
 	 */
 	friend class FTurkeySupportModule;
-	static FDataDrivenPlatformInfo& GetMutablePlatformInfo(FName PlatformName);
-	static TMap<FName, FDataDrivenPlatformInfo>& GetMutablePlatformInfos();
+	static CORE_API FDataDrivenPlatformInfo& GetMutablePlatformInfo(FName PlatformName);
+	static CORE_API TMap<FName, FDataDrivenPlatformInfo>& GetMutablePlatformInfos();
 
 #endif
 
diff --git a/Engine/Source/Runtime/Core/Public/Misc/DateTime.h b/Engine/Source/Runtime/Core/Public/Misc/DateTime.h
--- a/Engine/Source/Runtime/Core/Public/Misc/DateTime.h
+++ b/Engine/Source/Runtime/Core/Public/Misc/DateTime.h
@@ -561,7 +561,7 @@ public:
 	 * @return Time of day.
 	 * @see FromUnixTimestamp
 	 */
-	CORE_API int64 ToUnixTimestamp() const
+	int64 ToUnixTimestamp() const
 	{
 		return (Ticks - FDateTime(1970, 1, 1).Ticks) / ETimespan::TicksPerSecond;
 	}
@@ -572,7 +572,7 @@ public:
 	 * @return Time of day.
 	 * @see FromUnixTimestamp
 	 */
-	CORE_API double ToUnixTimestampDecimal() const
+	double ToUnixTimestampDecimal() const
 	{
 		return double(Ticks - FDateTime(1970, 1, 1).Ticks) / ETimespan::TicksPerSecond;
 	}
diff --git a/Engine/Source/Runtime/Core/Public/Misc/DefaultValueHelper.h b/Engine/Source/Runtime/Core/Public/Misc/DefaultValueHelper.h
--- a/Engine/Source/Runtime/Core/Public/Misc/DefaultValueHelper.h
+++ b/Engine/Source/Runtime/Core/Public/Misc/DefaultValueHelper.h
@@ -14,7 +14,7 @@ struct FLinearColor;
 //	FDefaultValueHelper
 //
 
-class CORE_API FDefaultValueHelper
+class FDefaultValueHelper
 {
 	/** Something like TEXT macro for chars is needed */
 	static TCHAR TS(const TCHAR* Sign);
@@ -51,40 +51,40 @@ class CORE_API FDefaultValueHelper
 
 public:
 
-	static FString GetUnqualifiedEnumValue(const FString& Source);
+	static CORE_API FString GetUnqualifiedEnumValue(const FString& Source);
 
-	static bool HasWhitespaces(const FString& Source);
+	static CORE_API bool HasWhitespaces(const FString& Source);
 
-	static FString RemoveWhitespaces(const FString& Source);
+	static CORE_API FString RemoveWhitespaces(const FString& Source);
 
 	/** Shell parameters list: " TypeName ( A a, B b ) " -> "A a, B b" */
-	static bool GetParameters(const FString& Source, const FString& TypeName, FString& OutForm);
+	static CORE_API bool GetParameters(const FString& Source, const FString& TypeName, FString& OutForm);
 
 	/** returns if given strings are equal, ignores initial and final white spaces in Source */
-	static bool Is(const FString& Source, const TCHAR* CompareStr);
+	static CORE_API bool Is(const FString& Source, const TCHAR* CompareStr);
 
 	/**   
 	 * source forms:	TypeName( TEXT ("ABC") ), TEXT("ABC"), TypeName("ABC"), "ABC"
 	 * output form:		ABC
 	 */
-	static bool StringFromCppString(const FString& Source, const FString& TypeName, FString& OutForm);
+	static CORE_API bool StringFromCppString(const FString& Source, const FString& TypeName, FString& OutForm);
 
 	/*
 	 *	Following functions accept c++ style representations of numbers.
 	 *  e.g. 13.5e-2f for float or -0x123 for int
 	 */
-	static bool IsStringValidInteger(const FString& Source);
+	static CORE_API bool IsStringValidInteger(const FString& Source);
 
-	static bool IsStringValidFloat(const FString& Source);
+	static CORE_API bool IsStringValidFloat(const FString& Source);
 
 	/** accepted form: " %f, %f, %f" */
-	static bool IsStringValidVector(const FString& Source);
+	static CORE_API bool IsStringValidVector(const FString& Source);
 
 	/** accepted form: " %f, %f, %f" */
-	static bool IsStringValidRotator(const FString& Source);
+	static CORE_API bool IsStringValidRotator(const FString& Source);
 
 	/** accepted form: " %f, %f, %f " or " %f, %f, %f, %f " (alpha is optional)  */
-	static bool IsStringValidLinearColor(const FString& Source);
+	static CORE_API bool IsStringValidLinearColor(const FString& Source);
 
 
 	/**
@@ -95,7 +95,7 @@ public:
 	 *
 	 * @return true if the conversion happened, false otherwise
 	 */
-	static bool ParseInt(const FString& Source, int32& OutVal);
+	static CORE_API bool ParseInt(const FString& Source, int32& OutVal);
 
 	/**
 	 * Converts a string into a int64.
@@ -105,7 +105,7 @@ public:
 	 *
 	 * @return true if the conversion happened, false otherwise
 	 */
-	static bool ParseInt64(const FString& Source, int64& OutVal);
+	static CORE_API bool ParseInt64(const FString& Source, int64& OutVal);
 
 	/**
 	 * Converts a string into a float.
@@ -115,7 +115,7 @@ public:
 	 *
 	 * @return true if the conversion happened, false otherwise
 	 */
-	static bool ParseFloat(const FString& Source, float& OutVal);
+	static CORE_API bool ParseFloat(const FString& Source, float& OutVal);
 
 	/**
 	 * Converts a string into a double.
@@ -125,7 +125,7 @@ public:
 	 *
 	 * @return true if the conversion happened, false otherwise
 	 */
-	static bool ParseDouble(const FString& Source, double& OutVal);
+	static CORE_API bool ParseDouble(const FString& Source, double& OutVal);
 
 	/**
 	 * Converts a string into a FVector. Accepted form: " %f, %f, %f "
@@ -135,8 +135,8 @@ public:
 	 *
 	 * @return true if the conversion happened, false otherwise
 	 */
-	static bool ParseVector(const FString& Source, FVector3f& OutVal);
-	static bool ParseVector(const FString& Source, FVector3d& OutVal);
+	static CORE_API bool ParseVector(const FString& Source, FVector3f& OutVal);
+	static CORE_API bool ParseVector(const FString& Source, FVector3d& OutVal);
 
 	/**
 	 * Converts a string into a FVector. Accepted form: " %f, %f "
@@ -146,8 +146,8 @@ public:
 	 *
 	 * @return true if the conversion happened, false otherwise
 	 */
-	static bool ParseVector2D(const FString& Source, FVector2f& OutVal);
-	static bool ParseVector2D(const FString& Source, FVector2d& OutVal);
+	static CORE_API bool ParseVector2D(const FString& Source, FVector2f& OutVal);
+	static CORE_API bool ParseVector2D(const FString& Source, FVector2d& OutVal);
 
 	/**
 	* Converts a string into a FVector4. Accepted form: " %f, %f, %f, %f "
@@ -157,8 +157,8 @@ public:
 	*
 	* @return true if the conversion happened, false otherwise
 	*/
-	static bool ParseVector4(const FString& Source, FVector4f& OutVal);
-	static bool ParseVector4(const FString& Source, FVector4d& OutVal);
+	static CORE_API bool ParseVector4(const FString& Source, FVector4f& OutVal);
+	static CORE_API bool ParseVector4(const FString& Source, FVector4d& OutVal);
 
 	/**
 	 * Converts a string into a FRotator. Accepted form: " %f, %f, %f "
@@ -168,8 +168,8 @@ public:
 	 *
 	 * @return true if the conversion happened, false otherwise
 	 */
-	static bool ParseRotator(const FString& Source, FRotator3f& OutVal);
-	static bool ParseRotator(const FString& Source, FRotator3d& OutVal);
+	static CORE_API bool ParseRotator(const FString& Source, FRotator3f& OutVal);
+	static CORE_API bool ParseRotator(const FString& Source, FRotator3d& OutVal);
 
 	/**
 	 * Converts a string into a FLinearColor. 
@@ -180,7 +180,7 @@ public:
 	 *
 	 * @return true if the conversion happened, false otherwise
 	 */
-	static bool ParseLinearColor(const FString& Source, FLinearColor& OutVal);
+	static CORE_API bool ParseLinearColor(const FString& Source, FLinearColor& OutVal);
 
 	/**
 	 * Converts a string into a FLinearColor. 
@@ -191,5 +191,5 @@ public:
 	 *
 	 * @return true if the conversion happened, false otherwise
 	 */
-	static bool ParseColor(const FString& Source, FColor& OutVal);
+	static CORE_API bool ParseColor(const FString& Source, FColor& OutVal);
 };
diff --git a/Engine/Source/Runtime/Core/Public/Misc/DelayedAutoRegister.h b/Engine/Source/Runtime/Core/Public/Misc/DelayedAutoRegister.h
--- a/Engine/Source/Runtime/Core/Public/Misc/DelayedAutoRegister.h
+++ b/Engine/Source/Runtime/Core/Public/Misc/DelayedAutoRegister.h
@@ -23,12 +23,12 @@ enum class EDelayedRegisterRunPhase : uint8
 	NumPhases,
 };
 
-struct CORE_API FDelayedAutoRegisterHelper
+struct FDelayedAutoRegisterHelper
 {
 
-	FDelayedAutoRegisterHelper(EDelayedRegisterRunPhase RunPhase, TFunction<void()> RegistrationFunction);
+	CORE_API FDelayedAutoRegisterHelper(EDelayedRegisterRunPhase RunPhase, TFunction<void()> RegistrationFunction);
 
-	static void RunAndClearDelayedAutoRegisterDelegates(EDelayedRegisterRunPhase RunPhase);
+	static CORE_API void RunAndClearDelayedAutoRegisterDelegates(EDelayedRegisterRunPhase RunPhase);
 };
 
 #if UE_ENABLE_INCLUDE_ORDER_DEPRECATED_IN_5_3
diff --git a/Engine/Source/Runtime/Core/Public/Misc/EmbeddedCommunication.h b/Engine/Source/Runtime/Core/Public/Misc/EmbeddedCommunication.h
--- a/Engine/Source/Runtime/Core/Public/Misc/EmbeddedCommunication.h
+++ b/Engine/Source/Runtime/Core/Public/Misc/EmbeddedCommunication.h
@@ -16,7 +16,7 @@
 typedef TMap<FString, FString> FEmbeddedCommunicationMap;
 
 // wraps parameters and a completion delegate
-struct CORE_API FEmbeddedCallParamsHelper
+struct FEmbeddedCallParamsHelper
 {
 	// the command for this call. something like "console" would be a command for the engine to run a console command
 	FString Command;
@@ -29,7 +29,7 @@ struct CORE_API FEmbeddedCallParamsHelper
 	TFunction<void(const FEmbeddedCommunicationMap&, FString)> OnCompleteDelegate;
 };
 
-class CORE_API FEmbeddedDelegates
+class FEmbeddedDelegates
 {
 public:
 
@@ -37,20 +37,20 @@ public:
 	DECLARE_MULTICAST_DELEGATE_OneParam(FEmbeddedCommunicationParamsDelegate, const FEmbeddedCallParamsHelper&);
 
 	// calling in from native wrapper to engine
-	static FEmbeddedCommunicationParamsDelegate& GetNativeToEmbeddedParamsDelegateForSubsystem(FName SubsystemName);
+	static CORE_API FEmbeddedCommunicationParamsDelegate& GetNativeToEmbeddedParamsDelegateForSubsystem(FName SubsystemName);
 
 	// calling out from engine to native wrapper
-	static FEmbeddedCommunicationParamsDelegate& GetEmbeddedToNativeParamsDelegateForSubsystem(FName SubsystemName);
+	static CORE_API FEmbeddedCommunicationParamsDelegate& GetEmbeddedToNativeParamsDelegateForSubsystem(FName SubsystemName);
 
 	// returns true if NativeToEmbedded delegate for subsystem exists
-	static bool IsEmbeddedSubsystemAvailable(FName SubsystemName);
+	static CORE_API bool IsEmbeddedSubsystemAvailable(FName SubsystemName);
 	
 	// FTSTicker-like delegate, to bind things to be ticked at a regular interval while the game thread is otherwise asleep.
-	static FSimpleMulticastDelegate SleepTickDelegate;
+	static CORE_API FSimpleMulticastDelegate SleepTickDelegate;
 
 	// get/set an object by name, thread safe
-	static void SetNamedObject(const FString& Name, void* Object);
-	static void* GetNamedObject(const FString& Name);
+	static CORE_API void SetNamedObject(const FString& Name, void* Object);
+	static CORE_API void* GetNamedObject(const FString& Name);
 	
 private:
 
@@ -67,41 +67,41 @@ private:
 };
 
 
-class CORE_API FEmbeddedCommunication
+class FEmbeddedCommunication
 {
 public:
 
 	// called early in UE lifecycle - RunOnGameThread can be called before this is called
-	static void Init();
+	static CORE_API void Init();
 
 	// force some ticking to happen - used to process messages during otherwise blocking operations like boot
-	static void ForceTick(int ID, float MinTimeSlice=0.1f, float MaxTimeSlice=0.5f);
+	static CORE_API void ForceTick(int ID, float MinTimeSlice=0.1f, float MaxTimeSlice=0.5f);
 
 	// queue up a function to call on game thread
-	static void RunOnGameThread(int Priority, TFunction<void()> Lambda);
+	static CORE_API void RunOnGameThread(int Priority, TFunction<void()> Lambda);
 
 	// wake up the game thread to process something put onto the game thread
-	static void WakeGameThread();
+	static CORE_API void WakeGameThread();
 
 	// called from game thread to pull off
-	static bool TickGameThread(float DeltaTime);
+	static CORE_API bool TickGameThread(float DeltaTime);
 	
 	// tell UE to stay awake (or allow it to sleep when nothing to do). Repeated calls with the same Requester are
 	// allowed, but bNeedsRendering must match
-	static void KeepAwake(FName Requester, bool bNeedsRendering);
-	static void AllowSleep(FName Requester);
+	static CORE_API void KeepAwake(FName Requester, bool bNeedsRendering);
+	static CORE_API void AllowSleep(FName Requester);
 	
-	static void UELogFatal(const TCHAR* String);
-	static void UELogError(const TCHAR* String);
-	static void UELogWarning(const TCHAR* String);
-	static void UELogDisplay(const TCHAR* String);
-	static void UELogLog(const TCHAR* String);
-	static void UELogVerbose(const TCHAR* String);
-
-	static bool IsAwakeForTicking();
-	static bool IsAwakeForRendering();
+	static CORE_API void UELogFatal(const TCHAR* String);
+	static CORE_API void UELogError(const TCHAR* String);
+	static CORE_API void UELogWarning(const TCHAR* String);
+	static CORE_API void UELogDisplay(const TCHAR* String);
+	static CORE_API void UELogLog(const TCHAR* String);
+	static CORE_API void UELogVerbose(const TCHAR* String);
+
+	static CORE_API bool IsAwakeForTicking();
+	static CORE_API bool IsAwakeForRendering();
 	
-	static FString GetDebugInfo();
+	static CORE_API FString GetDebugInfo();
 };
 
 // RAII for keep awake functionality
diff --git a/Engine/Source/Runtime/Core/Public/Misc/EngineBuildSettings.h b/Engine/Source/Runtime/Core/Public/Misc/EngineBuildSettings.h
--- a/Engine/Source/Runtime/Core/Public/Misc/EngineBuildSettings.h
+++ b/Engine/Source/Runtime/Core/Public/Misc/EngineBuildSettings.h
@@ -6,26 +6,26 @@
 #include "Containers/UnrealString.h"
 #include "CoreTypes.h"
 
-class CORE_API FEngineBuildSettings
+class FEngineBuildSettings
 {
 public:
 	/**
 	 * @return True if the build was gotten from perforce
 	 */
-	static bool IsPerforceBuild();
+	static CORE_API bool IsPerforceBuild();
 
 	/**
 	 * @return True if the build is for internal projects only
 	 */
-	static bool IsInternalBuild();
+	static CORE_API bool IsInternalBuild();
 
 	/**
 	 * @return True if the current installation is a source distribution.
 	 */
-	static bool IsSourceDistribution();
+	static CORE_API bool IsSourceDistribution();
 
 	/**
 	 * @return True if a given engine distribution contains source (as opposed to, say, Launcher builds)
 	 */
-	static bool IsSourceDistribution(const FString& RootDir);
+	static CORE_API bool IsSourceDistribution(const FString& RootDir);
 };
diff --git a/Engine/Source/Runtime/Core/Public/Misc/EngineNetworkCustomVersion.h b/Engine/Source/Runtime/Core/Public/Misc/EngineNetworkCustomVersion.h
--- a/Engine/Source/Runtime/Core/Public/Misc/EngineNetworkCustomVersion.h
+++ b/Engine/Source/Runtime/Core/Public/Misc/EngineNetworkCustomVersion.h
@@ -45,7 +45,7 @@ enum UE_DEPRECATED(5.1, "Using custom versions instead going forward, see FEngin
 	HISTORY_ENGINENETVERSION_LATEST = HISTORY_ENGINENETVERSION_PLUS_ONE - 1,
 };
 
-struct CORE_API FEngineNetworkCustomVersion
+struct FEngineNetworkCustomVersion
 {
 	enum Type
 	{
@@ -93,12 +93,12 @@ struct CORE_API FEngineNetworkCustomVersion
 	};
 
 	// The GUID for this custom version number
-	const static FGuid Guid;
+	CORE_API const static FGuid Guid;
 
 	FEngineNetworkCustomVersion() = delete;
 };
 
-struct CORE_API FGameNetworkCustomVersion
+struct FGameNetworkCustomVersion
 {
 	enum Type
 	{
@@ -111,7 +111,7 @@ struct CORE_API FGameNetworkCustomVersion
 	};
 
 	// The GUID for this custom version number
-	const static FGuid Guid;
+	CORE_API const static FGuid Guid;
 
 	FGameNetworkCustomVersion() = delete;
-};
\ No newline at end of file
+};
diff --git a/Engine/Source/Runtime/Core/Public/Misc/EngineVersion.h b/Engine/Source/Runtime/Core/Public/Misc/EngineVersion.h
--- a/Engine/Source/Runtime/Core/Public/Misc/EngineVersion.h
+++ b/Engine/Source/Runtime/Core/Public/Misc/EngineVersion.h
@@ -8,7 +8,7 @@
 #include "Serialization/StructuredArchive.h"
 
 /** Utility functions. */
-class CORE_API FEngineVersion : public FEngineVersionBase
+class FEngineVersion : public FEngineVersionBase
 {
 public:
 
@@ -16,34 +16,34 @@ public:
 	FEngineVersion() = default;
 
 	/** Constructs a version from the given components. */
-	FEngineVersion(uint16 InMajor, uint16 InMinor, uint16 InPatch, uint32 InChangelist, const FString &InBranch);
+	CORE_API FEngineVersion(uint16 InMajor, uint16 InMinor, uint16 InPatch, uint32 InChangelist, const FString &InBranch);
 
 	/** Sets the version to the given values. */
-	void Set(uint16 InMajor, uint16 InMinor, uint16 InPatch, uint32 InChangelist, const FString &InBranch);
+	CORE_API void Set(uint16 InMajor, uint16 InMinor, uint16 InPatch, uint32 InChangelist, const FString &InBranch);
 
 	/** Clears the object. */
-	void Empty();
+	CORE_API void Empty();
 
 	/** Checks whether this engine version is an exact match for another engine version */
-	bool ExactMatch(const FEngineVersion& Other) const;
+	CORE_API bool ExactMatch(const FEngineVersion& Other) const;
 
 	/** Checks compatibility with another version object. */
-	bool IsCompatibleWith(const FEngineVersionBase &Other) const;
+	CORE_API bool IsCompatibleWith(const FEngineVersionBase &Other) const;
 
 	/** Generates a version string */
-	FString ToString(EVersionComponent LastComponent = EVersionComponent::Branch) const;
+	CORE_API FString ToString(EVersionComponent LastComponent = EVersionComponent::Branch) const;
 
 	/** Parses a version object from a string. Returns true on success. */
-	static bool Parse(const FString &Text, FEngineVersion &OutVersion);
+	static CORE_API bool Parse(const FString &Text, FEngineVersion &OutVersion);
 
 	/** Gets the current engine version */
-	static const FEngineVersion& Current();
+	static CORE_API const FEngineVersion& Current();
 
 	/** Gets the earliest version which this engine maintains strict API and package compatibility with */
-	static const FEngineVersion& CompatibleWith();
+	static CORE_API const FEngineVersion& CompatibleWith();
 
 	/** Clears the current and compatible-with engine versions */
-	static void TearDown();
+	static CORE_API void TearDown();
 
 	/** Serialization functions */
 	friend CORE_API void operator<<(class FArchive &Ar, FEngineVersion &Version);
@@ -55,7 +55,7 @@ public:
 		return Branch.Replace( TEXT( "+" ), TEXT( "/" ) );
 	}
 
-	const FString& GetBranchDescriptor() const;
+	CORE_API const FString& GetBranchDescriptor() const;
 		
 
 private:
diff --git a/Engine/Source/Runtime/Core/Public/Misc/EngineVersionBase.h b/Engine/Source/Runtime/Core/Public/Misc/EngineVersionBase.h
--- a/Engine/Source/Runtime/Core/Public/Misc/EngineVersionBase.h
+++ b/Engine/Source/Runtime/Core/Public/Misc/EngineVersionBase.h
@@ -25,7 +25,7 @@ enum class EVersionComparison
 
 
 /** Base class for the EngineVersion class. Holds basic version numbers. */
-class CORE_API FEngineVersionBase
+class FEngineVersionBase
 {
 public:
 
@@ -33,10 +33,10 @@ public:
 	FEngineVersionBase() = default;
 
 	/** Constructs a version from the given components. */
-	FEngineVersionBase(uint16 InMajor, uint16 InMinor, uint16 InPatch = 0, uint32 InChangelist = 0);
+	CORE_API FEngineVersionBase(uint16 InMajor, uint16 InMinor, uint16 InPatch = 0, uint32 InChangelist = 0);
 
 	/** Returns the changelist number corresponding to this version. */
-	uint32 GetChangelist() const;	
+	CORE_API uint32 GetChangelist() const;	
 
 	/** Returns the Major version number corresponding to this version. */
 	FORCEINLINE uint16 GetMajor() const
@@ -57,19 +57,19 @@ public:
 	}
 
 	/** Checks if the changelist number represents licensee changelist number. */
-	bool IsLicenseeVersion() const;
+	CORE_API bool IsLicenseeVersion() const;
 
 	/** Returns whether the current version is empty. */
-	bool IsEmpty() const;
+	CORE_API bool IsEmpty() const;
 
 	/** Returns whether the engine version has a changelist component. */
-	bool HasChangelist() const; 
+	CORE_API bool HasChangelist() const; 
 
 	/** Returns the newest of two versions, and the component at which they differ */
-	static EVersionComparison GetNewest(const FEngineVersionBase &First, const FEngineVersionBase &Second, EVersionComponent *OutComponent);
+	static CORE_API EVersionComparison GetNewest(const FEngineVersionBase &First, const FEngineVersionBase &Second, EVersionComponent *OutComponent);
 
 	/** Encodes a licensee changelist number (by setting the top bit) */
-	static uint32 EncodeLicenseeChangelist(uint32 Changelist);
+	static CORE_API uint32 EncodeLicenseeChangelist(uint32 Changelist);
 
 protected:
 
diff --git a/Engine/Source/Runtime/Core/Public/Misc/Exec.h b/Engine/Source/Runtime/Core/Public/Misc/Exec.h
--- a/Engine/Source/Runtime/Core/Public/Misc/Exec.h
+++ b/Engine/Source/Runtime/Core/Public/Misc/Exec.h
@@ -25,10 +25,10 @@ class UWorld;
 #endif
 
 // Any object that is capable of taking commands.
-class CORE_API FExec
+class FExec
 {
 public:
-	virtual ~FExec();
+	CORE_API virtual ~FExec();
 
 #if UE_ALLOW_EXEC_COMMANDS
 
@@ -41,18 +41,18 @@ public:
 	 *
 	 * @return	true if command was handled, false otherwise
 	 */
-	virtual bool Exec( UWorld* InWorld, const TCHAR* Cmd, FOutputDevice& Ar );
+	CORE_API virtual bool Exec( UWorld* InWorld, const TCHAR* Cmd, FOutputDevice& Ar );
 #else // UE_ALLOW_EXEC_COMMANDS
 
 	/**
 	 * final override of Exec that asserts if called
 	 */
-	virtual bool Exec( UWorld* InWorld, const TCHAR* Cmd, FOutputDevice& Ar ) final;
+	CORE_API virtual bool Exec( UWorld* InWorld, const TCHAR* Cmd, FOutputDevice& Ar ) final;
 	
 	/**
 	 * final override of Exec to replace overrides where a default value for Ar is provided
 	 */
-	virtual bool Exec( UWorld* InWorld, const TCHAR* Cmd ) final;
+	CORE_API virtual bool Exec( UWorld* InWorld, const TCHAR* Cmd ) final;
 
 #endif // !UE_ALLOW_EXEC_COMMANDS
 
diff --git a/Engine/Source/Runtime/Core/Public/Misc/ExpressionParserTypes.h b/Engine/Source/Runtime/Core/Public/Misc/ExpressionParserTypes.h
--- a/Engine/Source/Runtime/Core/Public/Misc/ExpressionParserTypes.h
+++ b/Engine/Source/Runtime/Core/Public/Misc/ExpressionParserTypes.h
@@ -86,7 +86,7 @@ enum class EParseState
 };
 
 /** A token stream wraps up a raw string, providing accessors into it for consuming tokens */
-class CORE_API FTokenStream
+class FTokenStream
 {
 public:
 
@@ -96,53 +96,53 @@ public:
 	 * Optionally supply a token to accumulate into
 	 * Returns a string token for the stream, or empty on error
 	 */
-	TOptional<FStringToken> ParseToken(TFunctionRef<EParseState(TCHAR)> Pred, FStringToken* Accumulate = nullptr) const;
+	CORE_API TOptional<FStringToken> ParseToken(TFunctionRef<EParseState(TCHAR)> Pred, FStringToken* Accumulate = nullptr) const;
 
 	/** Attempt parse out the specified pre-defined string from the current read position (or accumulating into the specified existing token) */
-	TOptional<FStringToken> ParseToken(const TCHAR* Symbol, FStringToken* Accumulate = nullptr) const;
-	TOptional<FStringToken> ParseTokenIgnoreCase(const TCHAR* Symbol, FStringToken* Accumulate = nullptr) const;
+	CORE_API TOptional<FStringToken> ParseToken(const TCHAR* Symbol, FStringToken* Accumulate = nullptr) const;
+	CORE_API TOptional<FStringToken> ParseTokenIgnoreCase(const TCHAR* Symbol, FStringToken* Accumulate = nullptr) const;
 
 	/** Return a string token for the next character in the stream (or accumulating into the specified existing token) */
-	TOptional<FStringToken> ParseSymbol(FStringToken* Accumulate = nullptr) const;
+	CORE_API TOptional<FStringToken> ParseSymbol(FStringToken* Accumulate = nullptr) const;
 
 	/** Attempt parse out the specified pre-defined string from the current read position (or accumulating into the specified existing token) */
-	TOptional<FStringToken> ParseSymbol(TCHAR Symbol, FStringToken* Accumulate = nullptr) const;
+	CORE_API TOptional<FStringToken> ParseSymbol(TCHAR Symbol, FStringToken* Accumulate = nullptr) const;
 
 	/** Parse a whitespace token */
-	TOptional<FStringToken> ParseWhitespace(FStringToken* Accumulate = nullptr) const;
+	CORE_API TOptional<FStringToken> ParseWhitespace(FStringToken* Accumulate = nullptr) const;
 
 	/** Generate a token for the specified number of chars, at the current read position (or end of Accumulate) */
-	TOptional<FStringToken> GenerateToken(int32 NumChars, FStringToken* Accumulate = nullptr) const;
+	CORE_API TOptional<FStringToken> GenerateToken(int32 NumChars, FStringToken* Accumulate = nullptr) const;
 
 public:
 
 	/** Constructor. The stream is only valid for the lifetime of the string provided */
-	FTokenStream(const TCHAR* In);
+	CORE_API FTokenStream(const TCHAR* In);
 	
 	/** Peek at the character at the specified offset from the current read position */
-	TCHAR PeekChar(int32 Offset = 0) const;
+	CORE_API TCHAR PeekChar(int32 Offset = 0) const;
 
 	/** Get the number of characters remaining in the stream after the current read position */
-	int32 CharsRemaining() const;
+	CORE_API int32 CharsRemaining() const;
 
 	/** Check if it is valid to read (the optional number of characters) from the specified position */
-	bool IsReadPosValid(const TCHAR* InPos, int32 MinNumChars = 1) const;
+	CORE_API bool IsReadPosValid(const TCHAR* InPos, int32 MinNumChars = 1) const;
 
 	/** Check if the stream is empty */
-	bool IsEmpty() const;
+	CORE_API bool IsEmpty() const;
 
 	/** Get the current read position from the start of the stream */
-	int32 GetPosition() const;
+	CORE_API int32 GetPosition() const;
 
 	const TCHAR* GetStart() const { return Start; }
 	const TCHAR* GetRead() const { return ReadPos; }
 	const TCHAR* GetEnd() const { return End; }
 
 	/** Get the error context from the current read position */
-	FString GetErrorContext() const;
+	CORE_API FString GetErrorContext() const;
 
 	/** Set the current read position to the character proceeding the specified token */
-	void SetReadPos(const FStringToken& Token);
+	CORE_API void SetReadPos(const FStringToken& Token);
 
 private:
 
@@ -176,7 +176,7 @@ DEFINE_EXPRESSION_NODE_TYPE(double, 0x8444A8A3, 0x19AE4E13, 0xBCFA75EE, 0x39982B
  * 	Can be constructed from any C++ type that has a corresponding DEFINE_EXPRESSION_NODE_TYPE.
  * 	Evaluation behaviour (unary/binary operator etc) is defined in the expression grammar, rather than the type itself.
  */
-class CORE_API FExpressionNode : FNoncopyable
+class FExpressionNode : FNoncopyable
 {
 public:
 
@@ -190,21 +190,21 @@ public:
 		typename TEnableIf<!TPointerIsConvertibleFromTo<T, FExpressionNode>::Value>::Type* = nullptr
 		);
 
-	~FExpressionNode();
+	CORE_API ~FExpressionNode();
 
 	/** Move construction/assignment */
-	FExpressionNode(FExpressionNode&& In);
-	FExpressionNode& operator=(FExpressionNode&& In);
+	CORE_API FExpressionNode(FExpressionNode&& In);
+	CORE_API FExpressionNode& operator=(FExpressionNode&& In);
 
 	/** Get the type identifier of this node */
-	const FGuid& GetTypeId() const;
+	CORE_API const FGuid& GetTypeId() const;
 
 	/** Cast this node to the specified type. Will return nullptr if the types do not match. */
 	template<typename T>
 	const T* Cast() const;
 
 	/** Copy this node and its wrapped data */
-	FExpressionNode Copy() const;
+	CORE_API FExpressionNode Copy() const;
 
 private:
 
@@ -212,8 +212,8 @@ private:
 	static constexpr uint32 MaxStackAllocationSize = 64 - sizeof(FGuid);
 
 	/** Helper accessor to the data interface. Returns null for empty containers. */
-	Impl::IExpressionNodeStorage* GetData();
-	const Impl::IExpressionNodeStorage* GetData() const;
+	CORE_API Impl::IExpressionNodeStorage* GetData();
+	CORE_API const Impl::IExpressionNodeStorage* GetData() const;
 
 	/** TypeID - 16 bytes */
 	FGuid TypeId;
@@ -418,26 +418,26 @@ private:
 typedef TOperatorEvaluationEnvironment<> FOperatorEvaluationEnvironment;
 
 /** Class used to consume tokens from a string */
-class CORE_API FExpressionTokenConsumer
+class FExpressionTokenConsumer
 {
 public:
 	/** Construction from a raw string. The consumer is only valid as long as the string is valid */
-	FExpressionTokenConsumer(const TCHAR* InExpression);
+	CORE_API FExpressionTokenConsumer(const TCHAR* InExpression);
 
 	/** Extract the list of tokens from this consumer */
-	TArray<FExpressionToken> Extract();
+	CORE_API TArray<FExpressionToken> Extract();
 
 	/** Add an expression node to the consumer, specifying the FStringToken this node relates to.
 	 *	Adding a node to the consumer will move its stream read position to the end of the added token.
 	 */
-	void Add(const FStringToken& SourceToken, FExpressionNode&& Node);
+	CORE_API void Add(const FStringToken& SourceToken, FExpressionNode&& Node);
 
 	/** Get the expression stream */
 	FTokenStream& GetStream() { return Stream; }
 
 private:
-	FExpressionTokenConsumer(const FExpressionTokenConsumer&);
-	FExpressionTokenConsumer& operator=(const FExpressionTokenConsumer&);
+	CORE_API FExpressionTokenConsumer(const FExpressionTokenConsumer&);
+	CORE_API FExpressionTokenConsumer& operator=(const FExpressionTokenConsumer&);
 
 	/** Array of added tokens */
 	TArray<FExpressionToken> Tokens;
@@ -455,7 +455,7 @@ typedef TOptional<FExpressionError> (FExpressionDefinition)(FExpressionTokenCons
 
 
 /** A lexeme dictionary defining how to lex an expression. */
-class CORE_API FTokenDefinitions
+class FTokenDefinitions
 {
 public:
 	FTokenDefinitions() : bIgnoreWhitespace(false) {}
@@ -464,7 +464,7 @@ public:
 	void IgnoreWhitespace() { bIgnoreWhitespace = true; }
 
 	/** Define a token by way of a function to be invoked to attempt to parse a token from a stream */
-	void DefineToken(TFunction<FExpressionDefinition>&& Definition);
+	CORE_API void DefineToken(TFunction<FExpressionDefinition>&& Definition);
 
 public:
 
@@ -472,11 +472,11 @@ public:
 	bool DoesIgnoreWhitespace() { return bIgnoreWhitespace; }
 
 	/** Consume a token for the specified consumer */
-	TOptional<FExpressionError> ConsumeTokens(FExpressionTokenConsumer& Consumer) const;
+	CORE_API TOptional<FExpressionError> ConsumeTokens(FExpressionTokenConsumer& Consumer) const;
 
 private:
 
-	TOptional<FExpressionError> ConsumeToken(FExpressionTokenConsumer& Consumer) const;
+	CORE_API TOptional<FExpressionError> ConsumeToken(FExpressionTokenConsumer& Consumer) const;
 
 private:
 
@@ -517,7 +517,7 @@ struct FOpParameters
 };
 
 /** A lexical gammer defining how to parse an expression. Clients must define the tokens and operators to be interpreted by the parser. */
-class CORE_API FExpressionGrammar
+class FExpressionGrammar
 {
 public:
 	/** Define a grouping operator from two expression node types */
@@ -561,16 +561,16 @@ public:
 public:
 
 	/** Retrieve the corresponding grouping token for the specified open group type, or nullptr if it's not a group token */
-	const FGuid* GetGrouping(const FGuid& TypeId) const;
+	CORE_API const FGuid* GetGrouping(const FGuid& TypeId) const;
 
 	/** Check if this grammar defines a pre-unary operator for the specified symbol */
-	bool HasPreUnaryOperator(const FGuid& TypeId) const;
+	CORE_API bool HasPreUnaryOperator(const FGuid& TypeId) const;
 	
 	/** Check if this grammar defines a post-unary operator for the specified symbol */
-	bool HasPostUnaryOperator(const FGuid& TypeId) const;
+	CORE_API bool HasPostUnaryOperator(const FGuid& TypeId) const;
 
 	/** Get the binary operator precedence and associativity parameters, for the specified symbol, if any */
-	const FOpParameters* GetBinaryOperatorDefParameters(const FGuid& TypeId) const;
+	CORE_API const FOpParameters* GetBinaryOperatorDefParameters(const FGuid& TypeId) const;
 
 private:
 
diff --git a/Engine/Source/Runtime/Core/Public/Misc/FbxErrors.h b/Engine/Source/Runtime/Core/Public/Misc/FbxErrors.h
--- a/Engine/Source/Runtime/Core/Public/Misc/FbxErrors.h
+++ b/Engine/Source/Runtime/Core/Public/Misc/FbxErrors.h
@@ -12,169 +12,169 @@
  * This file contains known map errors that can be referenced by name.
  * Documentation for these errors is assumed to lie in UDN at: Engine\Documentation\Source\Shared\Editor\MapErrors
  */
-struct CORE_API FFbxErrors
+struct FFbxErrors
 {
 	/** Generic */
-	static FLazyName Generic_ImportingNewObjectFailed;
+	static CORE_API FLazyName Generic_ImportingNewObjectFailed;
 
-	static FLazyName Generic_ReimportingObjectFailed;
+	static CORE_API FLazyName Generic_ReimportingObjectFailed;
 
-	static FLazyName Generic_LoadingSceneFailed;
+	static CORE_API FLazyName Generic_LoadingSceneFailed;
 
-	static FLazyName Generic_InvalidCharacterInName;
+	static CORE_API FLazyName Generic_InvalidCharacterInName;
 
-	static FLazyName Generic_SameNameAssetExists;
+	static CORE_API FLazyName Generic_SameNameAssetExists;
 
-	static FLazyName Generic_SameNameAssetOverriding;
+	static CORE_API FLazyName Generic_SameNameAssetOverriding;
 
-	static FLazyName Generic_CannotDeleteReferenced;
+	static CORE_API FLazyName Generic_CannotDeleteReferenced;
 
-	static FLazyName Generic_FBXFileParseFailed;
+	static CORE_API FLazyName Generic_FBXFileParseFailed;
 
-	static FLazyName Generic_MeshNotFound;
+	static CORE_API FLazyName Generic_MeshNotFound;
 
-	static FLazyName Generic_CannotDetectImportType;
+	static CORE_API FLazyName Generic_CannotDetectImportType;
 
 	/** Mesh Generic **/
 
-	static FLazyName Generic_Mesh_NoGeometry;
+	static CORE_API FLazyName Generic_Mesh_NoGeometry;
 
-	static FLazyName Generic_Mesh_SmallGeometry;
+	static CORE_API FLazyName Generic_Mesh_SmallGeometry;
 
-	static FLazyName Generic_Mesh_TriangulationFailed;
+	static CORE_API FLazyName Generic_Mesh_TriangulationFailed;
 
-	static FLazyName Generic_Mesh_ConvertSmoothingGroupFailed;
+	static CORE_API FLazyName Generic_Mesh_ConvertSmoothingGroupFailed;
 
-	static FLazyName Generic_Mesh_UnsupportingSmoothingGroup;
+	static CORE_API FLazyName Generic_Mesh_UnsupportingSmoothingGroup;
 
-	static FLazyName Generic_Mesh_MaterialIndexInconsistency;
+	static CORE_API FLazyName Generic_Mesh_MaterialIndexInconsistency;
 
-	static FLazyName Generic_Mesh_MeshNotFound;
+	static CORE_API FLazyName Generic_Mesh_MeshNotFound;
 
-	static FLazyName Generic_Mesh_NoSmoothingGroup;
+	static CORE_API FLazyName Generic_Mesh_NoSmoothingGroup;
 
-	static FLazyName Generic_Mesh_LOD_InvalidIndex;
+	static CORE_API FLazyName Generic_Mesh_LOD_InvalidIndex;
 
-	static FLazyName Generic_Mesh_LOD_NoFileSelected;
+	static CORE_API FLazyName Generic_Mesh_LOD_NoFileSelected;
 
-	static FLazyName Generic_Mesh_LOD_MultipleFilesSelected;
+	static CORE_API FLazyName Generic_Mesh_LOD_MultipleFilesSelected;
 
-	static FLazyName Generic_Mesh_SkinxxNameError;
+	static CORE_API FLazyName Generic_Mesh_SkinxxNameError;
 
-	static FLazyName Generic_Mesh_TooManyLODs;
+	static CORE_API FLazyName Generic_Mesh_TooManyLODs;
 
-	static FLazyName Generic_Mesh_TangentsComputeError;
+	static CORE_API FLazyName Generic_Mesh_TangentsComputeError;
 
-	static FLazyName Generic_Mesh_NoReductionModuleAvailable;
+	static CORE_API FLazyName Generic_Mesh_NoReductionModuleAvailable;
 
-	static FLazyName Generic_Mesh_TooMuchUVChannels;
+	static CORE_API FLazyName Generic_Mesh_TooMuchUVChannels;
 
 	/** Static Mesh **/
-	static FLazyName StaticMesh_TooManyMaterials;
+	static CORE_API FLazyName StaticMesh_TooManyMaterials;
 
-	static FLazyName StaticMesh_UVSetLayoutProblem;
+	static CORE_API FLazyName StaticMesh_UVSetLayoutProblem;
 
-	static FLazyName StaticMesh_NoTriangles;
+	static CORE_API FLazyName StaticMesh_NoTriangles;
 
-	static FLazyName StaticMesh_BuildError;
+	static CORE_API FLazyName StaticMesh_BuildError;
 
-	static FLazyName StaticMesh_AllTrianglesDegenerate;
+	static CORE_API FLazyName StaticMesh_AllTrianglesDegenerate;
 
-	static FLazyName StaticMesh_AdjacencyOptionForced;
+	static CORE_API FLazyName StaticMesh_AdjacencyOptionForced;
 
 	/** SkeletalMesh **/
-	static FLazyName SkeletalMesh_DifferentRoots;
+	static CORE_API FLazyName SkeletalMesh_DifferentRoots;
 
-	static FLazyName SkeletalMesh_DuplicateBones;
+	static CORE_API FLazyName SkeletalMesh_DuplicateBones;
 
-	static FLazyName SkeletalMesh_NoInfluences;
+	static CORE_API FLazyName SkeletalMesh_NoInfluences;
 
-	static FLazyName SkeletalMesh_TooManyInfluences;
+	static CORE_API FLazyName SkeletalMesh_TooManyInfluences;
 
-	static FLazyName SkeletalMesh_RestoreSortingMismatchedStrips;
+	static CORE_API FLazyName SkeletalMesh_RestoreSortingMismatchedStrips;
 
-	static FLazyName SkeletalMesh_RestoreSortingNoSectionMatch;
+	static CORE_API FLazyName SkeletalMesh_RestoreSortingNoSectionMatch;
 
-	static FLazyName SkeletalMesh_RestoreSortingForSectionNumber;
+	static CORE_API FLazyName SkeletalMesh_RestoreSortingForSectionNumber;
 
-	static FLazyName SkeletalMesh_NoMeshFoundOnRoot;
+	static CORE_API FLazyName SkeletalMesh_NoMeshFoundOnRoot;
 	
-	static FLazyName SkeletalMesh_InvalidRoot;
+	static CORE_API FLazyName SkeletalMesh_InvalidRoot;
 
-	static FLazyName SkeletalMesh_InvalidBone;
+	static CORE_API FLazyName SkeletalMesh_InvalidBone;
 
-	static FLazyName SkeletalMesh_InvalidNode;
+	static CORE_API FLazyName SkeletalMesh_InvalidNode;
 
-	static FLazyName SkeletalMesh_NoWeightsOnDeformer;
+	static CORE_API FLazyName SkeletalMesh_NoWeightsOnDeformer;
 
-	static FLazyName SkeletalMesh_NoBindPoseInScene;
+	static CORE_API FLazyName SkeletalMesh_NoBindPoseInScene;
 
-	static FLazyName SkeletalMesh_NoAssociatedCluster;
+	static CORE_API FLazyName SkeletalMesh_NoAssociatedCluster;
 
-	static FLazyName SkeletalMesh_NoBoneFound;
+	static CORE_API FLazyName SkeletalMesh_NoBoneFound;
 
-	static FLazyName SkeletalMesh_InvalidBindPose;
+	static CORE_API FLazyName SkeletalMesh_InvalidBindPose;
 
-	static FLazyName SkeletalMesh_MultipleRoots;
+	static CORE_API FLazyName SkeletalMesh_MultipleRoots;
 
-	static FLazyName SkeletalMesh_BonesAreMissingFromBindPose;
+	static CORE_API FLazyName SkeletalMesh_BonesAreMissingFromBindPose;
 
-	static FLazyName SkeletalMesh_VertMissingInfluences;
+	static CORE_API FLazyName SkeletalMesh_VertMissingInfluences;
 
-	static FLazyName SkeletalMesh_SectionWithNoTriangle;
+	static CORE_API FLazyName SkeletalMesh_SectionWithNoTriangle;
 
-	static FLazyName SkeletalMesh_TooManyVertices;
+	static CORE_API FLazyName SkeletalMesh_TooManyVertices;
 
-	static FLazyName SkeletalMesh_FailedToCreatePhyscisAsset;
+	static CORE_API FLazyName SkeletalMesh_FailedToCreatePhyscisAsset;
 
-	static FLazyName SkeletalMesh_SkeletonRecreateError;
+	static CORE_API FLazyName SkeletalMesh_SkeletonRecreateError;
 
-	static FLazyName SkeletalMesh_ExceedsMaxBoneCount;
+	static CORE_API FLazyName SkeletalMesh_ExceedsMaxBoneCount;
 
-	static FLazyName SkeletalMesh_NoUVSet;
+	static CORE_API FLazyName SkeletalMesh_NoUVSet;
 
-	static FLazyName SkeletalMesh_LOD_MissingBone;
+	static CORE_API FLazyName SkeletalMesh_LOD_MissingBone;
 
-	static FLazyName SkeletalMesh_LOD_FailedToImport;
+	static CORE_API FLazyName SkeletalMesh_LOD_FailedToImport;
 
-	static FLazyName SkeletalMesh_LOD_RootNameIncorrect;
+	static CORE_API FLazyName SkeletalMesh_LOD_RootNameIncorrect;
 
-	static FLazyName SkeletalMesh_LOD_BonesDoNotMatch;
+	static CORE_API FLazyName SkeletalMesh_LOD_BonesDoNotMatch;
 
-	static FLazyName SkeletalMesh_LOD_IncorrectParent;
+	static CORE_API FLazyName SkeletalMesh_LOD_IncorrectParent;
 
-	static FLazyName SkeletalMesh_LOD_HasSoftVerts;
+	static CORE_API FLazyName SkeletalMesh_LOD_HasSoftVerts;
 
-	static FLazyName SkeletalMesh_LOD_MissingSocketBone;
+	static CORE_API FLazyName SkeletalMesh_LOD_MissingSocketBone;
 
-	static FLazyName SkeletalMesh_LOD_MissingMorphTarget;
+	static CORE_API FLazyName SkeletalMesh_LOD_MissingMorphTarget;
 
-	static FLazyName SkeletalMesh_FillImportDataFailed;
+	static CORE_API FLazyName SkeletalMesh_FillImportDataFailed;
 
-	static FLazyName SkeletalMesh_InvalidPosition;
+	static CORE_API FLazyName SkeletalMesh_InvalidPosition;
 
 	/** Animation **/
-	static FLazyName Animation_CouldNotFindRootTrack;
+	static CORE_API FLazyName Animation_CouldNotFindRootTrack;
 
-	static FLazyName Animation_CouldNotBuildSkeleton;
+	static CORE_API FLazyName Animation_CouldNotBuildSkeleton;
 
-	static FLazyName Animation_CouldNotFindTrack;
+	static CORE_API FLazyName Animation_CouldNotFindTrack;
 
-	static FLazyName Animation_ZeroLength;
+	static CORE_API FLazyName Animation_ZeroLength;
 
-	static FLazyName Animation_RootTrackMismatch;
+	static CORE_API FLazyName Animation_RootTrackMismatch;
 
-	static FLazyName Animation_DuplicatedBone;
+	static CORE_API FLazyName Animation_DuplicatedBone;
 
-	static FLazyName Animation_MissingBones;
+	static CORE_API FLazyName Animation_MissingBones;
 
-	static FLazyName Animation_InvalidData;
+	static CORE_API FLazyName Animation_InvalidData;
 
-	static FLazyName Animation_TransformError;
+	static CORE_API FLazyName Animation_TransformError;
 
-	static FLazyName Animation_DifferentLength;
+	static CORE_API FLazyName Animation_DifferentLength;
 
-	static FLazyName Animation_CurveNotFound;
+	static CORE_API FLazyName Animation_CurveNotFound;
 };
 
 /**
@@ -184,7 +184,7 @@ class FFbxErrorToken : public FDocumentationToken
 {
 public:
 	/** Factory method, tokens can only be constructed as shared refs */
-	CORE_API static TSharedRef<FFbxErrorToken> Create( const FName& InErrorName )
+	static TSharedRef<FFbxErrorToken> Create( const FName& InErrorName )
 	{
 		return MakeShareable(new FFbxErrorToken(InErrorName));
 	}
