[Core]

diff --git a/Engine/Source/Runtime/Core/Public/Misc/FeedbackContext.h b/Engine/Source/Runtime/Core/Public/Misc/FeedbackContext.h
--- a/Engine/Source/Runtime/Core/Public/Misc/FeedbackContext.h
+++ b/Engine/Source/Runtime/Core/Public/Misc/FeedbackContext.h
@@ -24,16 +24,16 @@ class SBuildProgressWidget;
 struct FSlowTask;
 
 /** A context for displaying modal warning messages. */
-class CORE_API FFeedbackContext
+class FFeedbackContext
 	: public FOutputDevice
 {
 public:
-	virtual void Serialize(const TCHAR* V, ELogVerbosity::Type Verbosity, const FName& Category) override;
-	virtual void Serialize(const TCHAR* V, ELogVerbosity::Type Verbosity, const FName& Category, double Time) override;
-	virtual void SerializeRecord(const UE::FLogRecord& Record) override;
+	CORE_API virtual void Serialize(const TCHAR* V, ELogVerbosity::Type Verbosity, const FName& Category) override;
+	CORE_API virtual void Serialize(const TCHAR* V, ELogVerbosity::Type Verbosity, const FName& Category, double Time) override;
+	CORE_API virtual void SerializeRecord(const UE::FLogRecord& Record) override;
 
 	/** Ask the user a binary question, returning their answer */
-	virtual bool YesNof( const FText& Question );
+	CORE_API virtual bool YesNof( const FText& Question );
 	
 	/**
 	 * Whether or not the user has canceled out of the progress dialog
@@ -56,11 +56,11 @@ public:
 
 
 	/**** Legacy API - not deprecated as it's still in heavy use, but superceded by FScopedSlowTask ****/
-	void BeginSlowTask( const FText& Task, bool ShowProgressDialog, bool bShowCancelButton=false );
-	void UpdateProgress( int32 Numerator, int32 Denominator );
-	void StatusUpdate( int32 Numerator, int32 Denominator, const FText& StatusText );
-	void StatusForceUpdate( int32 Numerator, int32 Denominator, const FText& StatusText );
-	void EndSlowTask();
+	CORE_API void BeginSlowTask( const FText& Task, bool ShowProgressDialog, bool bShowCancelButton=false );
+	CORE_API void UpdateProgress( int32 Numerator, int32 Denominator );
+	CORE_API void StatusUpdate( int32 Numerator, int32 Denominator, const FText& StatusText );
+	CORE_API void StatusForceUpdate( int32 Numerator, int32 Denominator, const FText& StatusText );
+	CORE_API void EndSlowTask();
 	/**** end legacy API ****/
 
 protected:
@@ -93,10 +93,10 @@ protected:
 	virtual void ProgressReported( const float TotalProgressInterp, FText DisplayMessage ) {}
 
 	/** Called to check whether we are playing in editor when starting a slow task */
-	virtual bool IsPlayingInEditor() const;
+	CORE_API virtual bool IsPlayingInEditor() const;
 
-	void FormatLine(FStringBuilderBase& Out, const TCHAR* V, ELogVerbosity::Type Verbosity, const FName& Category, double Time, ELogVerbosity::Type* OutVerbosity = nullptr) const;
-	void FormatRecordLine(FStringBuilderBase& Out, const UE::FLogRecord& Record, ELogVerbosity::Type* OutVerbosity = nullptr) const;
+	CORE_API void FormatLine(FStringBuilderBase& Out, const TCHAR* V, ELogVerbosity::Type Verbosity, const FName& Category, double Time, ELogVerbosity::Type* OutVerbosity = nullptr) const;
+	CORE_API void FormatRecordLine(FStringBuilderBase& Out, const UE::FLogRecord& Record, ELogVerbosity::Type* OutVerbosity = nullptr) const;
 
 public:
 	virtual FContextSupplier* GetContext() const { return nullptr; }
@@ -108,8 +108,8 @@ public:
 
 	bool	TreatWarningsAsErrors = false;
 
-	FFeedbackContext();
-	virtual ~FFeedbackContext();
+	CORE_API FFeedbackContext();
+	CORE_API virtual ~FFeedbackContext();
 
 	/** Gets warnings history */
 	void GetWarnings(TArray<FString>& OutWarnings) const
@@ -149,11 +149,11 @@ public:
 	}
 
 private:
-	void AddToHistory(const TCHAR* V, ELogVerbosity::Type Verbosity, const FName& Category, double Time);
-	void AddRecordToHistory(const UE::FLogRecord& Record);
+	CORE_API void AddToHistory(const TCHAR* V, ELogVerbosity::Type Verbosity, const FName& Category, double Time);
+	CORE_API void AddRecordToHistory(const UE::FLogRecord& Record);
 
-	FFeedbackContext(const FFeedbackContext&);
-	FFeedbackContext& operator=(const FFeedbackContext&);
+	CORE_API FFeedbackContext(const FFeedbackContext&);
+	CORE_API FFeedbackContext& operator=(const FFeedbackContext&);
 
 	/** Warnings history */
 	TArray<FString> Warnings;
@@ -193,10 +193,10 @@ protected:
 	TArray<TUniquePtr<FSlowTask>> LegacyAPIScopes;
 
 	/** Ask that the UI be updated as a result of the scope stack changing */
-	void RequestUpdateUI(bool bForceUpdate = false);
+	CORE_API void RequestUpdateUI(bool bForceUpdate = false);
 
 	/** Update the UI as a result of the scope stack changing */
-	void UpdateUI();
+	CORE_API void UpdateUI();
 
 	/**
 	 * Adds a new warning message to warnings history.
diff --git a/Engine/Source/Runtime/Core/Public/Misc/FeedbackContextMarkup.h b/Engine/Source/Runtime/Core/Public/Misc/FeedbackContextMarkup.h
--- a/Engine/Source/Runtime/Core/Public/Misc/FeedbackContextMarkup.h
+++ b/Engine/Source/Runtime/Core/Public/Misc/FeedbackContextMarkup.h
@@ -22,15 +22,15 @@ class FText;
  * Update the progress and set a status message for the current operation:
  *		@progress 'Compiling source code...' 50%
  */
-class CORE_API FFeedbackContextMarkup
+class FFeedbackContextMarkup
 {
 public:
 
 	/** Markup stack manipulation. */
-	static bool ParseCommand(const FString& Line, FFeedbackContext* Warn);
+	static CORE_API bool ParseCommand(const FString& Line, FFeedbackContext* Warn);
 
 	/** Utility functions for dealing with external processes. */
-	static bool PipeProcessOutput(const FText& Description, const FString& URL, const FString& Params, FFeedbackContext* Warn, int32* OutExitCode);
+	static CORE_API bool PipeProcessOutput(const FText& Description, const FString& URL, const FString& Params, FFeedbackContext* Warn, int32* OutExitCode);
 
 private:
 
diff --git a/Engine/Source/Runtime/Core/Public/Misc/FileHelper.h b/Engine/Source/Runtime/Core/Public/Misc/FileHelper.h
--- a/Engine/Source/Runtime/Core/Public/Misc/FileHelper.h
+++ b/Engine/Source/Runtime/Core/Public/Misc/FileHelper.h
@@ -24,7 +24,7 @@ template <typename FuncType> class TFunctionRef;
 /*-----------------------------------------------------------------------------
 	FFileHelper
 -----------------------------------------------------------------------------*/
-struct CORE_API FFileHelper
+struct FFileHelper
 {
 	enum class EHashOptions
 	{
@@ -57,7 +57,7 @@ struct CORE_API FFileHelper
 	 * Load a text file to an FString.
 	 * Supports all combination of ANSI/Unicode files and platforms.
 	*/
-	static void BufferToString( FString& Result, const uint8* Buffer, int32 Size );
+	static CORE_API void BufferToString( FString& Result, const uint8* Buffer, int32 Size );
 
 	/**
 	 * Load a binary file to a dynamic array with two uninitialized bytes at end as padding.
@@ -66,7 +66,7 @@ struct CORE_API FFileHelper
 	 * @param Filename  The file to read
 	 * @param Flags     Flags to pass to IFileManager::CreateFileReader
 	*/
-	static bool LoadFileToArray( TArray<uint8>& Result, const TCHAR* Filename, uint32 Flags = 0 );
+	static CORE_API bool LoadFileToArray( TArray<uint8>& Result, const TCHAR* Filename, uint32 Flags = 0 );
 
 	/**
 	 * Load a binary file to a dynamic array with two uninitialized bytes at end as padding.
@@ -75,7 +75,7 @@ struct CORE_API FFileHelper
 	 * @param Filename  The file to read
 	 * @param Flags     Flags to pass to IFileManager::CreateFileReader
 	*/
-	static bool LoadFileToArray( TArray64<uint8>& Result, const TCHAR* Filename, uint32 Flags = 0 );
+	static CORE_API bool LoadFileToArray( TArray64<uint8>& Result, const TCHAR* Filename, uint32 Flags = 0 );
 
 	/**
 	 * Load a text file to an FString. Supports all combination of ANSI/Unicode files and platforms.
@@ -84,7 +84,7 @@ struct CORE_API FFileHelper
 	 * @param Archive      Name of the archive to load from
 	 * @param VerifyFlags  Flags controlling the hash verification behavior ( see EHashOptions )
 	 */
-	static bool LoadFileToString(FString& Result, FArchive& Reader, EHashOptions VerifyFlags = EHashOptions::None);
+	static CORE_API bool LoadFileToString(FString& Result, FArchive& Reader, EHashOptions VerifyFlags = EHashOptions::None);
 
 	/**
 	 * Load a text file to an FString. Supports all combination of ANSI/Unicode files and platforms.
@@ -93,7 +93,7 @@ struct CORE_API FFileHelper
 	 * @param Filename     Name of the file to load
 	 * @param VerifyFlags  Flags controlling the hash verification behavior ( see EHashOptions )
 	 */
-	static bool LoadFileToString( FString& Result, const TCHAR* Filename, EHashOptions VerifyFlags = EHashOptions::None, uint32 ReadFlags = 0 );
+	static CORE_API bool LoadFileToString( FString& Result, const TCHAR* Filename, EHashOptions VerifyFlags = EHashOptions::None, uint32 ReadFlags = 0 );
 
 	/**
 	 * Load a text file to an FString. Supports all combination of ANSI/Unicode files and platforms.
@@ -103,7 +103,7 @@ struct CORE_API FFileHelper
 	 * @param Filename     Name of the file to load
 	 * @param VerifyFlags  Flags controlling the hash verification behavior ( see EHashOptions )
 	 */
-	static bool LoadFileToString(FString& Result, IPlatformFile* PlatformFile, const TCHAR* Filename, EHashOptions VerifyFlags = EHashOptions::None, uint32 ReadFlags = 0);
+	static CORE_API bool LoadFileToString(FString& Result, IPlatformFile* PlatformFile, const TCHAR* Filename, EHashOptions VerifyFlags = EHashOptions::None, uint32 ReadFlags = 0);
 
 	/**
 	 * Load a text file to an array of strings. Supports all combination of ANSI/Unicode files and platforms.
@@ -111,10 +111,10 @@ struct CORE_API FFileHelper
 	 * @param Result       String representation of the loaded file
 	 * @param Filename     Name of the file to load
 	 */
-	static bool LoadFileToStringArray( TArray<FString>& Result, const TCHAR* Filename );
+	static CORE_API bool LoadFileToStringArray( TArray<FString>& Result, const TCHAR* Filename );
 
 	UE_DEPRECATED(4.26, "LoadFileToStringArray no longer supports VerifyFlags. You can use UE::String::ParseLines to split up a string loaded with LoadFileToString")
-	static bool LoadFileToStringArray(TArray<FString>& Result, const TCHAR* Filename, EHashOptions VerifyFlags);
+	static CORE_API bool LoadFileToStringArray(TArray<FString>& Result, const TCHAR* Filename, EHashOptions VerifyFlags);
 
 	/**
 	 * Load a text file to an array of strings, filtered by a user-defined predicate. Supports all combination of ANSI/Unicode files and platforms.
@@ -123,10 +123,10 @@ struct CORE_API FFileHelper
 	 * @param Filename     Name of the file to load
 	 * @param Predicate    Condition for whether or not to add the line to the array
 	 */
-	static bool LoadFileToStringArrayWithPredicate(TArray<FString>& Result, const TCHAR* Filename, TFunctionRef<bool(const FString&)> Predicate);
+	static CORE_API bool LoadFileToStringArrayWithPredicate(TArray<FString>& Result, const TCHAR* Filename, TFunctionRef<bool(const FString&)> Predicate);
 
 	UE_DEPRECATED(4.26, "LoadFileToStringArrayWithPredicate no longer supports VerifyFlags. You can use UE::String::ParseLines to split up a string loaded with LoadFileToString")
-	static bool LoadFileToStringArrayWithPredicate(TArray<FString>& Result, const TCHAR* Filename, TFunctionRef<bool(const FString&)> Predicate, EHashOptions VerifyFlags);
+	static CORE_API bool LoadFileToStringArrayWithPredicate(TArray<FString>& Result, const TCHAR* Filename, TFunctionRef<bool(const FString&)> Predicate, EHashOptions VerifyFlags);
 
 	/**
 	 * Load a text file and invoke a visitor for each line. Supports all combination of ANSI/Unicode files and platforms.
@@ -134,29 +134,29 @@ struct CORE_API FFileHelper
 	 * @param Filename     Name of the file to load
 	 * @param Visitor      Visitor to invoke for each non-empty line in the file
 	 */
-	static bool LoadFileToStringWithLineVisitor(const TCHAR* Filename, TFunctionRef<void(FStringView Line)> Visitor);
+	static CORE_API bool LoadFileToStringWithLineVisitor(const TCHAR* Filename, TFunctionRef<void(FStringView Line)> Visitor);
 
 	/**
 	 * Save a binary array to a file.
 	 */
-	static bool SaveArrayToFile(TArrayView<const uint8> Array, const TCHAR* Filename, IFileManager* FileManager=&IFileManager::Get(), uint32 WriteFlags = 0);
+	static CORE_API bool SaveArrayToFile(TArrayView<const uint8> Array, const TCHAR* Filename, IFileManager* FileManager=&IFileManager::Get(), uint32 WriteFlags = 0);
 
 	/**
 	 * Save a binary array to a file.
 	 */
-	static bool SaveArrayToFile( const TArray64<uint8>& Array, const TCHAR* Filename, IFileManager* FileManager = &IFileManager::Get(), uint32 WriteFlags = 0 );
+	static CORE_API bool SaveArrayToFile( const TArray64<uint8>& Array, const TCHAR* Filename, IFileManager* FileManager = &IFileManager::Get(), uint32 WriteFlags = 0 );
 
 	/**
 	 * Write the FString to a file.
 	 * Supports all combination of ANSI/Unicode files and platforms.
 	 */
-	static bool SaveStringToFile( FStringView String, const TCHAR* Filename, EEncodingOptions EncodingOptions = EEncodingOptions::AutoDetect, IFileManager* FileManager = &IFileManager::Get(), uint32 WriteFlags = 0 );
+	static CORE_API bool SaveStringToFile( FStringView String, const TCHAR* Filename, EEncodingOptions EncodingOptions = EEncodingOptions::AutoDetect, IFileManager* FileManager = &IFileManager::Get(), uint32 WriteFlags = 0 );
 
 	/**
 	 * Write the FString to a file.
 	 * Supports all combination of ANSI/Unicode files and platforms.
 	 */
-	static bool SaveStringArrayToFile( const TArray<FString>& Lines, const TCHAR* Filename, EEncodingOptions EncodingOptions = EEncodingOptions::AutoDetect, IFileManager* FileManager = &IFileManager::Get(), uint32 WriteFlags = 0 );
+	static CORE_API bool SaveStringArrayToFile( const TArray<FString>& Lines, const TCHAR* Filename, EEncodingOptions EncodingOptions = EEncodingOptions::AutoDetect, IFileManager* FileManager = &IFileManager::Get(), uint32 WriteFlags = 0 );
 
 	/**
 	 * Saves a 24/32Bit BMP file to disk for debug image dump purposes
@@ -179,7 +179,7 @@ struct CORE_API FFileHelper
 	 *
 	 * @return true if success
 	 */
-	static bool CreateBitmap( const TCHAR* Pattern, int32 DataWidth, int32 DataHeight, const struct FColor* Data, FIntRect* SubRectangle = NULL, IFileManager* FileManager = &IFileManager::Get(), FString* OutFilename = NULL, bool bInWriteAlpha = false, EColorChannel ColorChannel = EColorChannel::All);
+	static CORE_API bool CreateBitmap( const TCHAR* Pattern, int32 DataWidth, int32 DataHeight, const struct FColor* Data, FIntRect* SubRectangle = NULL, IFileManager* FileManager = &IFileManager::Get(), FString* OutFilename = NULL, bool bInWriteAlpha = false, EColorChannel ColorChannel = EColorChannel::All);
 
 	/**
 	 * Generates the next unique bitmap filename with a specified extension
@@ -191,7 +191,7 @@ struct CORE_API FFileHelper
 	 *
 	 * @return true if success
 	 */
-	static bool GenerateNextBitmapFilename(const FString& Pattern, const FString& Extension, FString& OutFilename, IFileManager* FileManager = &IFileManager::Get());
+	static CORE_API bool GenerateNextBitmapFilename(const FString& Pattern, const FString& Extension, FString& OutFilename, IFileManager* FileManager = &IFileManager::Get());
 	
 	/**
 	 * Generates the next unique bitmap filename with a specified extension
@@ -202,7 +202,7 @@ struct CORE_API FFileHelper
 	 *
 	 * @return true if success
 	 */
-	static void GenerateDateTimeBasedBitmapFilename(const FString& Pattern, const FString& Extension, FString& OutFilename);
+	static CORE_API void GenerateDateTimeBasedBitmapFilename(const FString& Pattern, const FString& Extension, FString& OutFilename);
 	
 	/**
 	 *	Load the given ANSI text file to an array of strings - one FString per line of the file.
@@ -214,7 +214,7 @@ struct CORE_API FFileHelper
 	 *
 	 *	@return	bool				true if successful, false if not
 	 */
-	static bool LoadANSITextFileToStrings(const TCHAR* InFilename, IFileManager* InFileManager, TArray<FString>& OutStrings);
+	static CORE_API bool LoadANSITextFileToStrings(const TCHAR* InFilename, IFileManager* InFileManager, TArray<FString>& OutStrings);
 
 	/**
 	* Checks to see if a filename is valid for saving.
@@ -223,7 +223,7 @@ struct CORE_API FFileHelper
 	* @param Filename	Filename, with or without path information, to check.
 	* @param OutError	If an error occurs, this is the reason why
 	*/
-	static bool IsFilenameValidForSaving(const FString& Filename, FText& OutError);
+	static CORE_API bool IsFilenameValidForSaving(const FString& Filename, FText& OutError);
 
 	enum class UE_DEPRECATED(5.0, "EChannelMask has been deprecated in favor of EColorChannel") EChannelMask
 	{
@@ -236,7 +236,7 @@ struct CORE_API FFileHelper
 
 	PRAGMA_DISABLE_DEPRECATION_WARNINGS
 	UE_DEPRECATED(5.0, "EChannelMask has been deprecated in favor of EColorChannel, please use the other CreateBitmap() method.")
-	static bool CreateBitmap(const TCHAR* Pattern, int32 DataWidth, int32 DataHeight, const struct FColor* Data, FIntRect* SubRectangle, IFileManager* FileManager, FString* OutFilename, bool bInWriteAlpha, EChannelMask ChannelMask);
+	static CORE_API bool CreateBitmap(const TCHAR* Pattern, int32 DataWidth, int32 DataHeight, const struct FColor* Data, FIntRect* SubRectangle, IFileManager* FileManager, FString* OutFilename, bool bInWriteAlpha, EChannelMask ChannelMask);
 	PRAGMA_ENABLE_DEPRECATION_WARNINGS
 };
 
diff --git a/Engine/Source/Runtime/Core/Public/Misc/Fnv.h b/Engine/Source/Runtime/Core/Public/Misc/Fnv.h
--- a/Engine/Source/Runtime/Core/Public/Misc/Fnv.h
+++ b/Engine/Source/Runtime/Core/Public/Misc/Fnv.h
@@ -11,9 +11,9 @@
 /** 
  * FNV hash generation for different types of input data
  **/
-struct CORE_API FFnv
+struct FFnv
 {
 	/** generates FNV hash of the memory area */
-	static uint32 MemFnv32( const void* Data, int32 Length, uint32 FNV=0 );
-    static uint64 MemFnv64( const void* Data, int32 Length, uint64 FNV=0 );
+	static CORE_API uint32 MemFnv32( const void* Data, int32 Length, uint32 FNV=0 );
+    static CORE_API uint64 MemFnv64( const void* Data, int32 Length, uint64 FNV=0 );
 };
diff --git a/Engine/Source/Runtime/Core/Public/Misc/Fork.h b/Engine/Source/Runtime/Core/Public/Misc/Fork.h
--- a/Engine/Source/Runtime/Core/Public/Misc/Fork.h
+++ b/Engine/Source/Runtime/Core/Public/Misc/Fork.h
@@ -9,8 +9,8 @@
 	#define DEFAULT_SERVER_FAKE_FORKS 0
 #endif
 
-class CORE_API FRunnable;
-class CORE_API FRunnableThread;
+class FRunnable;
+class FRunnableThread;
 
 enum class EForkProcessRole : uint8
 {
@@ -31,7 +31,7 @@ enum class EForkProcessRole : uint8
  * From the fork point any memory page that gets written into by a forked process will be transferred into a unique page for this process.
  * 
  */
-class CORE_API FForkProcessHelper 
+class FForkProcessHelper 
 {
 public:
 
@@ -40,34 +40,34 @@ public:
 	 * This could be a process on a fork-supported platform that will launch real child processes. (-WaitAndFork is set)
 	 * Or it could be a process that will simulate forking by tranforming itself into a child process via fake forking (-FakeForking is set)
 	 */
-	static bool IsForkRequested();
+	static CORE_API bool IsForkRequested();
 
 	/**
 	 * Are we a forked process that supports multithreading
 	 * This only becomes true after its safe to be multithread.
 	 * Since a process can be forked mid-tick, there is a period of time where IsForkedChildProcess is true but IsForkedMultithreadInstance will be false
 	 */
-	static bool IsForkedMultithreadInstance();
+	static CORE_API bool IsForkedMultithreadInstance();
 
 	/**
 	 * Is this a process that was forked
 	 */
-	static bool IsForkedChildProcess();
+	static CORE_API bool IsForkedChildProcess();
 
 	/**
 	 * Sets the forked child process flag and index given to this child process
 	 */
-	static void SetIsForkedChildProcess(uint16 ChildIndex=1);
+	static CORE_API void SetIsForkedChildProcess(uint16 ChildIndex=1);
 
 	/**
 	* Returns the unique index of this forked child process. Index 0 is for the master server
 	*/
-	static uint16 GetForkedChildProcessIndex();
+	static CORE_API uint16 GetForkedChildProcessIndex();
 
 	/**
 	 * Event triggered when a fork occurred on the child process and its safe to create real threads
 	 */
-	static void OnForkingOccured();
+	static CORE_API void OnForkingOccured();
 
 	/**
 	 * Tells if we allow multithreading on forked processes.
@@ -75,7 +75,7 @@ public:
 	 * Enabled via -PostForkThreading
 	 * Disabled via -DisablePostForkThreading
 	 */
-	static bool SupportsMultithreadingPostFork();
+	static CORE_API bool SupportsMultithreadingPostFork();
 
 	/**
 	 * Creates a thread according to the environment it's in:
@@ -85,7 +85,7 @@ public:
 	 *		If called on the original master process: will create a forkable thread that is ticked in the main thread pre-fork but becomes a real thread post-fork
 	 *      If called on a forked child process: will create a real thread immediately
 	 */
-	static FRunnableThread* CreateForkableThread(
+	static CORE_API FRunnableThread* CreateForkableThread(
 		class FRunnable* InRunnable,
 		const TCHAR* InThreadName,
 		uint32 InStackSize = 0,
@@ -100,21 +100,21 @@ public:
 	 * E.g. notifies GMalloc to optimize for memory sharing across parent/child process
 	 * Note: This will be called multiple times on the parent before each fork. 
 	 */
-	static void LowLevelPreFork();
+	static CORE_API void LowLevelPreFork();
 
 	/**
 	 * Performs low-level cross-platform actions that should happen immediately AFTER forking in the PARENT process in a well-specified order.
 	 * Runs before any higher level code like calling into game-level constructs.
 	 * E.g. notifies GMalloc to optimize for memory sharing across parent/child process
 	 */
-	static void LowLevelPostForkParent();
+	static CORE_API void LowLevelPostForkParent();
 
 	/**
 	 * Performs low-level cross-platform actions that should happen immediately AFTER forking in the CHILD process in a well-specified order.
 	 * Runs before any higher level code like calling into game-level constructs.
 	 * E.g. notifies GMalloc to optimize for memory sharing across parent/child process
 	 */
-	static void LowLevelPostForkChild(uint16 ChildIndex=1);
+	static CORE_API void LowLevelPostForkChild(uint16 ChildIndex=1);
 };
 
 
diff --git a/Engine/Source/Runtime/Core/Public/Misc/HierarchicalLogArchive.h b/Engine/Source/Runtime/Core/Public/Misc/HierarchicalLogArchive.h
--- a/Engine/Source/Runtime/Core/Public/Misc/HierarchicalLogArchive.h
+++ b/Engine/Source/Runtime/Core/Public/Misc/HierarchicalLogArchive.h
@@ -10,10 +10,10 @@
 
 class FArchive;
 
-struct CORE_API FHierarchicalLogArchive : private FArchiveProxy
+struct FHierarchicalLogArchive : private FArchiveProxy
 {
 public:
-	FHierarchicalLogArchive(FArchive& InInnerArchive);
+	CORE_API FHierarchicalLogArchive(FArchive& InInnerArchive);
 
 	struct FIndentScope
 	{
@@ -47,7 +47,7 @@ public:
 	}
 
 private:
-	void WriteLine(const FString& InLine, bool bIndent = false);
+	CORE_API void WriteLine(const FString& InLine, bool bIndent = false);
 
 	int32 Indentation;
 };
diff --git a/Engine/Source/Runtime/Core/Public/Misc/InteractiveProcess.h b/Engine/Source/Runtime/Core/Public/Misc/InteractiveProcess.h
--- a/Engine/Source/Runtime/Core/Public/Misc/InteractiveProcess.h
+++ b/Engine/Source/Runtime/Core/Public/Misc/InteractiveProcess.h
@@ -34,7 +34,7 @@ DECLARE_DELEGATE_OneParam(FOnInteractiveProcessOutput, const FString&);
 /**
 * Implements an external process that can be interacted.
 */
-class CORE_API FInteractiveProcess
+class FInteractiveProcess
 	: public FRunnable
 {
 public:
@@ -46,7 +46,7 @@ public:
 	* @param InParams The command line parameters.
 	* @param InHidden Whether the window of the process should be hidden.
 	*/
-	FInteractiveProcess(const FString& InURL, const FString& InParams, bool InHidden, bool LongTime = false);
+	CORE_API FInteractiveProcess(const FString& InURL, const FString& InParams, bool InHidden, bool LongTime = false);
 
 	/**
 	* Creates a new interactive process.
@@ -56,17 +56,17 @@ public:
 	* @param InWorkingDir The URL of the working dir where the executable should launch.
 	* @param InHidden Whether the window of the process should be hidden.
 	*/
-	FInteractiveProcess(const FString& InURL, const FString& InParams, const FString& InWorkingDir, bool InHidden, bool LongTime = false);
+	CORE_API FInteractiveProcess(const FString& InURL, const FString& InParams, const FString& InWorkingDir, bool InHidden, bool LongTime = false);
 
 	/** Destructor. */
-	~FInteractiveProcess();
+	CORE_API ~FInteractiveProcess();
 
 	/**
 	* Gets the duration of time that the task has been running.
 	*
 	* @return Time duration.
 	*/
-	FTimespan GetDuration() const;
+	CORE_API FTimespan GetDuration() const;
 
 	/**
 	* Gets the Process Handle. The instance can be invalid if the process was not created.
@@ -93,7 +93,7 @@ public:
 	*
 	* @return True if succeed
 	*/
-	bool Launch();
+	CORE_API bool Launch();
 
 	/**
 	* Returns a delegate that is executed when the process has been canceled.
@@ -131,14 +131,14 @@ public:
 	*
 	* @param Message to be sent
 	*/
-	void SendWhenReady(const FString &Message);
+	CORE_API void SendWhenReady(const FString &Message);
 
 	/**
 	* Sends the data message when process is ready
 	*
 	* @param Data to be sent
 	*/
-	void SendWhenReady(const TArray<uint8> &Data);
+	CORE_API void SendWhenReady(const TArray<uint8> &Data);
 
 	/**
 	* Returns the return code from the exited process
@@ -168,7 +168,7 @@ public:
 		return true;
 	}
 
-	virtual uint32 Run() override;
+	CORE_API virtual uint32 Run() override;
 
 	virtual void Stop() override
 	{
@@ -184,12 +184,12 @@ protected:
 	*
 	* @param Output The output string to process.
 	*/
-	void ProcessOutput(const FString& Output);
+	CORE_API void ProcessOutput(const FString& Output);
 
 	/**
 	 * Takes the first message to be sent from MessagesToProcess, if there is one, and sends it to process
 	 */
-	void SendMessageToProcessIf();
+	CORE_API void SendMessageToProcessIf();
 
 private:
 	// Whether the process is being canceled. */
diff --git a/Engine/Source/Runtime/Core/Public/Misc/LocalTimestampDirectoryVisitor.h b/Engine/Source/Runtime/Core/Public/Misc/LocalTimestampDirectoryVisitor.h
--- a/Engine/Source/Runtime/Core/Public/Misc/LocalTimestampDirectoryVisitor.h
+++ b/Engine/Source/Runtime/Core/Public/Misc/LocalTimestampDirectoryVisitor.h
@@ -15,7 +15,7 @@ struct FDateTime;
 /**
  * Visitor to gather local files with their timestamps.
  */
-class CORE_API FLocalTimestampDirectoryVisitor
+class FLocalTimestampDirectoryVisitor
 	: public IPlatformFile::FDirectoryVisitor
 {
 public:
@@ -32,13 +32,13 @@ public:
 	 * @param bInCacheDirectories - Whether to cache the directories.
 	 * @param bInMakeLowerCase - Whether to lower case filenames and directories.
 	 */
-	FLocalTimestampDirectoryVisitor( IPlatformFile& InFileInterface, const TArray<FString>& InDirectoriesToIgnore, const TArray<FString>& InDirectoriesToNotRecurse, bool bInCacheDirectories = false, bool bInMakeLowerCase = false );
+	CORE_API FLocalTimestampDirectoryVisitor( IPlatformFile& InFileInterface, const TArray<FString>& InDirectoriesToIgnore, const TArray<FString>& InDirectoriesToNotRecurse, bool bInCacheDirectories = false, bool bInMakeLowerCase = false );
 
 public:
 
 	// IPlatformFile::FDirectoryVisitor interface
 
-	virtual bool Visit(const TCHAR* FilenameOrDirectory, bool bIsDirectory);
+	CORE_API virtual bool Visit(const TCHAR* FilenameOrDirectory, bool bIsDirectory);
 
 private:
 
diff --git a/Engine/Source/Runtime/Core/Public/Misc/MapErrors.h b/Engine/Source/Runtime/Core/Public/Misc/MapErrors.h
--- a/Engine/Source/Runtime/Core/Public/Misc/MapErrors.h
+++ b/Engine/Source/Runtime/Core/Public/Misc/MapErrors.h
@@ -11,240 +11,240 @@
  * This file contains known map errors that can be referenced by name.
  * Documentation for these errors is assumed to lie in UDN at: Engine\Documentation\Source\Shared\Editor\MapErrors
  */
-struct CORE_API FMapErrors
+struct FMapErrors
 {
     /**  Lighting */
 
 	/** {LightActor} has same light GUID as {LightActor} (Duplicate and replace the orig with the new one): Duplicate and replace the original with the new one. */
-    static FLazyName MatchingLightGUID;
+    static CORE_API FLazyName MatchingLightGUID;
  
     /**  Actor */
 
 	/** {Actor} has same GUID as {Actor} (Duplicate and replace the orig with the new one): Duplicate and replace the original with the new one. */
-    static FLazyName MatchingActorGUID;
+    static CORE_API FLazyName MatchingActorGUID;
 
 	/** {Actor} : Large actor casts a shadow and will cause an extreme performance hit unless bUseBooleanEnvironmentShadowing is set to true: A large actor has been set to cast shadows - this will cause extreme performance issues and should have bUseBooleanEnvironmentShadowing set to true. */
-    static FLazyName ActorLargeShadowCaster;
+    static CORE_API FLazyName ActorLargeShadowCaster;
 
 	 /** {Volume} causes damage, but has no damagetype defined.: This warning is caused when there is a volume that is set to cause damage but doesn't have a damage type defined.  A damage type is important because it tells the game code how to handle a actor's reaction to taking damage.  This can be solved by going to the actor's Property Window->Physics Volume and setting the 'DamageType' property. */
-    static FLazyName NoDamageType;
+    static CORE_API FLazyName NoDamageType;
    
 	/** {Actor} : Brush has non-coplanar polygons: This warning is caused when there is a brush in the level that has non-coplanar polygons.  This is usually caused by using brush editing tools in geometry mode in extreme ways and can cause missing polygons in the level.  This warning can be resolved by deleting the brush and recreating it. */
-    static FLazyName NonCoPlanarPolys;
+    static CORE_API FLazyName NonCoPlanarPolys;
     
 	/** {Actor} in same location as {Another Actor}: This warning is caused when there is a actor that is in the same exact location as another actor.  This warning is usually the result of a accidental duplication or paste operation.  It can be fixed by deleting one of the actors, or disregarded if the placement was intentional. */
-    static FLazyName SameLocation;
+    static CORE_API FLazyName SameLocation;
     
 	/** {Actor} has invalid DrawScale/ DrawScale3D: This warning is caused when either DrawScale, DrawScale3D X, DrawScale3D Y, or DrawScale 3D Z is equal to zero.  Meaning that the actor will not be shown because it is being scaled to zero on one of its axis.  To solve this problem, change any DrawScale's that are zero to be non-zero by selecting the actor and changing its drawscale at the bottom of the main UnrealEd window. */
-    static FLazyName InvalidDrawscale;
+    static CORE_API FLazyName InvalidDrawscale;
     
 	/** {Actor} is obsolete and must be removed!: This warning is caused when there is an instance of a actor in a level that has been marked deprecated.  This is usually because a actor was marked deprecated after the level was created, but the map was never updated.  This can be fixed by simply deleting the actor. */
-    static FLazyName ActorIsObselete;
+    static CORE_API FLazyName ActorIsObselete;
     
 	/** {Actor} bStatic true, but has Physics set to something other than PHYS_None!: This warning is caused when an actor has its bStatic flag set to true but its Physics is set to PHYS_None.  Since bStatic means that the actor will not be moving, having Physics set to PHYS_None is contradictory.  Actors set with the bStatic flag are also not ticked(updated).  This error can be solved by going to the actor's properties and changing its Physics to PHYS_None. */
-    static FLazyName StaticPhysNone;
+    static CORE_API FLazyName StaticPhysNone;
     
 	/** {Actor} : Volume actor has NULL collision component - please delete: The specified volume actor has a NULL collision component and should probably be deleted. */
-    static FLazyName VolumeActorCollisionComponentNULL;
+    static CORE_API FLazyName VolumeActorCollisionComponentNULL;
     
 	/** {Actor} : Volume actor has a collision component with 0 radius - please delete: The specified volume actor has a zero radius for its collision component and should probably be deleted. */
-    static FLazyName VolumeActorZeroRadius;
+    static CORE_API FLazyName VolumeActorZeroRadius;
     
 	/** {Actor} (LOD {Index}) has hand-painted vertex colors that no longer match the original StaticMesh {StaticMesh}: It looks like the original mesh was changed since this instance's vertex colors were painted down - this may need a refresh. */
-    static FLazyName VertexColorsNotMatchOriginalMesh;
+    static CORE_API FLazyName VertexColorsNotMatchOriginalMesh;
 
 	/** {ActorName} has collision enabled but StaticMesh ({StaticMeshName}) has no simple or complex collision. */
-	static FLazyName CollisionEnabledNoCollisionGeom;
+	static CORE_API FLazyName CollisionEnabledNoCollisionGeom;
 
 	/** Actor casts dynamic shadows and has a BoundsScale greater than 1! This will have a large performance hit: Serious performance warning... either reduce BoundsScale to be <= 1 or remove dynamic shadows... */
-    static FLazyName ShadowCasterUsingBoundsScale;
+    static CORE_API FLazyName ShadowCasterUsingBoundsScale;
     
-	static FLazyName MultipleSkyLights;
+	static CORE_API FLazyName MultipleSkyLights;
 
-	static FLazyName MultipleSkyAtmospheres;
-	static FLazyName MultipleSkyAtmosphereTypes;
+	static CORE_API FLazyName MultipleSkyAtmospheres;
+	static CORE_API FLazyName MultipleSkyAtmosphereTypes;
 
 	/** Actor gas invalid hair strands material */
-	static FLazyName InvalidHairStrandsMaterial;
+	static CORE_API FLazyName InvalidHairStrandsMaterial;
 
 	/** {ActorName} has WorldTrace blocked. It will be considered to be world geometry */
-	static FLazyName InvalidTrace;
+	static CORE_API FLazyName InvalidTrace;
 
     /**  BSP Brush */
 
 	/** {Brush Actor} : Brush has zero polygons - please delete!: This warning indicates that you have a brush in the level that doesn't have any polygons associated with it.  The brush should be deleted as it isn't doing anything useful. */
-    static FLazyName BrushZeroPolygons;
+    static CORE_API FLazyName BrushZeroPolygons;
     
 	/** Run 'Clean BSP Materials' to clear {count} references: This warning indicates that there are material references on brush faces that aren't contributing to the BSP, and that applying the Tools->'Clean BSP Materials' operation can clean up these references. */
-    static FLazyName CleanBSPMaterials;
+    static CORE_API FLazyName CleanBSPMaterials;
     
 	/** {Actor} : Brush has NULL BrushComponent property - please delete!: This warning is caused when there is a Brush with a "None" ConstraintInstance component.  It is usually found in older maps where duplication was used to create a Brush and can be fixed by deleting the Brush causing the warning and creating a new one. */
-    static FLazyName BrushComponentNull;
+    static CORE_API FLazyName BrushComponentNull;
     
 	/** {Brush} : Brush is planar: Planar brush used - please note that this may not work well with collision. */
-    static FLazyName PlanarBrush;
+    static CORE_API FLazyName PlanarBrush;
 
 	/** Camera */
 
 	/** Camera has AspectRatio=0 - please set this to something non-zero */
-	static FLazyName CameraAspectRatioIsZero;
+	static CORE_API FLazyName CameraAspectRatioIsZero;
 
     /**  Class */
 
 	/** { }::{ } is obsolete and must be removed (Class is abstract): NEEDS DESCRIPTION */
-    static FLazyName AbstractClass;
+    static CORE_API FLazyName AbstractClass;
     
 	/** { }::{ } is obsolete and must be removed (Class is deprecated): NEEDS DESCRIPTION */
-    static FLazyName DeprecatedClass;
+    static CORE_API FLazyName DeprecatedClass;
     
 
 	/** Foliage */
 
 	/** Foliage instances for a missing static mesh have been removed.: **TODO** */
-	static FLazyName FoliageMissingStaticMesh;
+	static CORE_API FLazyName FoliageMissingStaticMesh;
 
 	/** Foliage in this map is missing {MissingCount} cluster component(s) for static mesh {MeshName}. Opening the Foliage tool will fix this problem.: **TODO** */
-	static FLazyName FoliageMissingClusterComponent;
+	static CORE_API FLazyName FoliageMissingClusterComponent;
 
     /**  Landscape */
 
 	/** {LandscapeComponent} : Fixed up deleted layer weightmap: */
-    static FLazyName FixedUpDeletedLayerWeightmap;
+    static CORE_API FLazyName FixedUpDeletedLayerWeightmap;
     
 	/** {LandscapeComponent} : Fixed up incorrect layer weightmap texture index: */
-    static FLazyName FixedUpIncorrectLayerWeightmap;
+    static CORE_API FLazyName FixedUpIncorrectLayerWeightmap;
     
 	/** {LandscapeComponent} : Fixed up invalid material instances for a landscape component when its weightmap allocations don't match the material's layer configuration : */
-	static FLazyName FixedUpInvalidLandscapeMaterialInstances;
+	static CORE_API FLazyName FixedUpInvalidLandscapeMaterialInstances;
 
 	/** Fixed up shared weightmap texture for layer {Layer} in component {Component} (shares with {Name}): */
-    static FLazyName FixedUpSharedLayerWeightmap;
+    static CORE_API FLazyName FixedUpSharedLayerWeightmap;
 
 	/** Landscape ({ProxyName}) has overlapping render components at location ({X, Y}): */
-	static FLazyName LandscapeComponentPostLoad_Warning;
+	static CORE_API FLazyName LandscapeComponentPostLoad_Warning;
 
 	/** Landscape proxy {ProxyName} of {LevelName} points to a LandscapeActor that is not currently loaded. This will lose the property upon save.  */
-	static FLazyName LandscapeLazyObjectPtrDeprecation_Warning;
+	static CORE_API FLazyName LandscapeLazyObjectPtrDeprecation_Warning;
 
 	/** Landscape proxy {Proxy} had one or more shared properties not in sync with its parent landscape actor {Parent}. This has been fixed but the proxy needs to be saved in order to ensure cooking behaves as expected..  */
-	static FLazyName LandscapeProxyFixupSharedData_Warning;
+	static CORE_API FLazyName LandscapeProxyFixupSharedData_Warning;
 
     /**  Level */
 
 	/** Duplicate level info: Two WorldInfos somehow exist... */
-    static FLazyName DuplicateLevelInfo;
+    static CORE_API FLazyName DuplicateLevelInfo;
     
 	/** Map should have KillZ set.: This warning is caused when the map's KillZ in the WorldInfo properties is set to the default value. All maps should specify a KillZ appropriate for the level so that players cannot simply fall forever until they reach the playable world bounds. */
-    static FLazyName NoKillZ;
+    static CORE_API FLazyName NoKillZ;
   
 
     /**  Lighting */
 
 	/** {Actor} : Light actor has NULL LightComponent property - please delete!: This warning is caused when there is a Light actor with a "None" LightComponent.  It is usually found in older maps where duplication was used to create the actor and can be fixed by deleting the actor causing the warning and creating a new one. */
-    static FLazyName LightComponentNull;
+    static CORE_API FLazyName LightComponentNull;
     
 	/** Maps need lighting rebuilt: This warning is caused when lighting has been invalidated by moving or modifying a light actor.  This can cause problems because the rendered lighting in the level is not accurately representing the current state of lights in the level.  This error can be solved by going to the Build menu and rebuilding lighting for a map. */
-    static FLazyName RebuildLighting;
+    static CORE_API FLazyName RebuildLighting;
    
 	/** Component is a static type but has invalid lightmap settings!  Indirect lighting will be black.  Common causes are lightmap resolution of 0, LightmapCoordinateIndex out of bounds. */
-	static FLazyName StaticComponentHasInvalidLightmapSettings;
+	static CORE_API FLazyName StaticComponentHasInvalidLightmapSettings;
    
 	/** Component is a using CastVolumetricTranslucentShadow but this feature is disabled for the project! Turn on r.Shadow.TranslucentPerObject.ProjectEnabled in a project ini if required. */
-	static FLazyName PrimitiveComponentHasInvalidTranslucentShadowSetting;
+	static CORE_API FLazyName PrimitiveComponentHasInvalidTranslucentShadowSetting;
 
 	/** Navigation */
 
 	/** Paths need to be rebuilt: **TODO** */
-	static FLazyName RebuildPaths;
+	static CORE_API FLazyName RebuildPaths;
 
     /**  Particle System */
 
 	/** {Actor} : Emitter actor has NULL ParticleSystemComponent property - please delete!: This warning is caused when there is a Emitter actor with a "None" ParticleSystemComponent.  It is usually found in older maps where duplication was used to create the actor and can be fixed by deleting the actor causing the warning and creating a new one. */
-    static FLazyName ParticleSystemComponentNull;
+    static CORE_API FLazyName ParticleSystemComponentNull;
     
 	/** PSysComp has an empty parameter actor reference at index {Index} ({Actor}): Param.Actor should not be NULL. */
-    static FLazyName PSysCompErrorEmptyActorRef;
+    static CORE_API FLazyName PSysCompErrorEmptyActorRef;
     
 	/** PSysComp has an empty parameter material reference at index {Index} ({Actor}): Param.Material should not be NULL. */
-    static FLazyName PSysCompErrorEmptyMaterialRef;
+    static CORE_API FLazyName PSysCompErrorEmptyMaterialRef;
     
 
     /**  Skeletal Mesh */
 
 	/** {Actor} : SkeletalMeshActor has no PhysicsAsset assigned.: In order for a SkeletalMesh to have an accurate bounding box, it needs to have a PhysicsAsset assigned in its SkeletalMeshComponent. An incorrect or inaccurate bounding box can lead to the mesh vanishing when its origin is not in view, or to poor shadow resolution because the bounding box is too big. */
-    static FLazyName SkelMeshActorNoPhysAsset;
+    static CORE_API FLazyName SkelMeshActorNoPhysAsset;
     
 	/** {Actor} : Skeletal mesh actor has NULL SkeletalMeshComponent property: The specified SkeletalMeshActor has a NULL SkeletalMeshComponent. */
-    static FLazyName SkeletalMeshComponent;
+    static CORE_API FLazyName SkeletalMeshComponent;
     
 	/** {Actor} : Skeletal mesh actor has NULL SkeletalMesh property: The specified SkeletalMeshActor has a NULL SkeletalMesh. */
-    static FLazyName SkeletalMeshNull;
+    static CORE_API FLazyName SkeletalMeshNull;
 
     /**  Sound */
 
 	/** {Actor} : Ambient sound actor has NULL AudioComponent property - please delete!: This warning is caused when there is a Ambient sound actor with a "None" AudioComponent.  It is usually found in older maps where duplication was used to create a actor and can be fixed by deleting the actor causing the warning and creating a new one. */
-    static FLazyName AudioComponentNull;
+    static CORE_API FLazyName AudioComponentNull;
 
 	/** Ambient sound actor's AudioComponent has a NULL SoundCue property!: This warning is caused when there is a AmbientSound actor with a NULL SoundCue property.  This is a problem because the actor won't actually be playing any sounds.  This can be fixed by first choosing a sound cue in the generic browser and then going to the actor's Property Window->Audio Category->Audio Component and setting the 'SoundCue' property. */
-    static FLazyName SoundCueNull;
+    static CORE_API FLazyName SoundCueNull;
 
     /**  Static Mesh */
 
 	/** {Static Mesh Actor} : Static mesh actor has NULL StaticMesh property: This warning is caused when there is a static mesh actor in a level with a NULL StaticMesh property.  This can be a problem because the actor exists and is using memory, but doesn't have a static mesh to actually draw.  This warning is usually the result of creating a StaticMesh actor without first selecting a StaticMesh in the generic browser.  This warning can be fixed by first selecting a static mesh in the generic browser and then going to the StaticMesh actor's Property Window->StaticMeshActor Category->StaticMeshComponent->StaticMeshComponent Category to set the 'StaticMesh' Property. */
-    static FLazyName StaticMeshNull;
+    static CORE_API FLazyName StaticMeshNull;
     
 	/** {Actor} : Static mesh actor has NULL StaticMeshComponent property - please delete!: This warning is caused when there is a static mesh actor with a "None" StaticMeshComponent component.  It is usually found in older maps where duplication was used to create the actor and can be fixed by deleting the actor causing the warning and creating a new one. */
-    static FLazyName StaticMeshComponent;
+    static CORE_API FLazyName StaticMeshComponent;
     
 	/** {StaticMesh} has simple collision but is being scaled non-uniformly - collision creation will fail: Simple collision cannot be used with non-uniform scale. Please either fix the scale or the collision type. */
-    static FLazyName SimpleCollisionButNonUniformScale;
+    static CORE_API FLazyName SimpleCollisionButNonUniformScale;
 
 	/** More overridden materials {Count} on static mesh component than are referenced {Count} in source mesh {StaticMesh}: **TODO** */
-    static FLazyName MoreMaterialsThanReferenced;
+    static CORE_API FLazyName MoreMaterialsThanReferenced;
 
 	/** More Nanite overridden materials {Count} on static mesh component than are referenced {Count} in source mesh {StaticMesh}: **TODO** */
-	static FLazyName MoreNaniteMaterialsThanReferenced;
+	static CORE_API FLazyName MoreNaniteMaterialsThanReferenced;
 
 	 /** {Count} element(s) with zero triangles in static mesh {StaticMesh}: **TODO** */
-    static FLazyName ElementsWithZeroTriangles;
+    static CORE_API FLazyName ElementsWithZeroTriangles;
   
     /**  Volume */
 
 	/** LevelStreamingVolume is not in the persistent level - please delete: This warning is caused when there is a level streaming volume that does not exist in the persistent level.  This can be problematic because the volume will not be considered when checking to see if a streaming level should be loaded or unloaded.  You can fix this problem by deleting the level streaming volume and recreating it. */
-    static FLazyName LevelStreamingVolume;
+    static CORE_API FLazyName LevelStreamingVolume;
     
 	/** No levels are associated with streaming volume.: This warning is caused when there are no levels associated with a LevelStreamingVolume, making it non-functional. This problem can be fixed by associating one or more streaming levels with the offending LevelStreamingVolume. */
-    static FLazyName NoLevelsAssociated;
+    static CORE_API FLazyName NoLevelsAssociated;
     
 
     /**  Uncategorized */
 
 	/** Filename is too long ({FilenameLength} characters) - this may interfere with cooking for consoles. Unreal filenames should be no longer than {Length} characters. Filename value: {Filename}: Please rename the file to be within the length specified. */
-    static FLazyName FilenameIsTooLongForCooking;
+    static CORE_API FLazyName FilenameIsTooLongForCooking;
     
 	/** {ObjectName} : Externally referenced */
-	static FLazyName UsingExternalObject;
+	static CORE_API FLazyName UsingExternalObject;
 
     /**  Actors */
 
 	/** {Actor} : Repaired painted vertex colors: Painted vertex colors were repaired on this actor. */
-    static FLazyName RepairedPaintedVertexColors;
+    static CORE_API FLazyName RepairedPaintedVertexColors;
 
 	/** Hierarchical LOD **/
 
 	/** {LODActor} : Static mesh is missing for the built LODActor.  Did you remove the asset? */
-	static FLazyName LODActorMissingStaticMesh;
+	static CORE_API FLazyName LODActorMissingStaticMesh;
 
 	/** {LODActor} : Actor is missing. The actor might have been removed. We recommend you to build LOD again. */
-	static FLazyName LODActorMissingActor;
+	static CORE_API FLazyName LODActorMissingActor;
 
 	/** {LODActor} : NoActor is assigned. We recommend to delete this actor. */
-	static FLazyName LODActorNoActorFound;
+	static CORE_API FLazyName LODActorNoActorFound;
 
 	/** Virtual Textures */
 
 	/** {ObjectName} is using a virtual texture ('{TextureName}') on an unsupported property ('{PropertyName}'). */
-	static FLazyName InvalidVirtualTextureUsage;
+	static CORE_API FLazyName InvalidVirtualTextureUsage;
 };
 
 /**
diff --git a/Engine/Source/Runtime/Core/Public/Misc/MemStack.h b/Engine/Source/Runtime/Core/Public/Misc/MemStack.h
--- a/Engine/Source/Runtime/Core/Public/Misc/MemStack.h
+++ b/Engine/Source/Runtime/Core/Public/Misc/MemStack.h
@@ -32,7 +32,7 @@ enum EMemOned
 };
 
 
-class CORE_API FPageAllocator
+class FPageAllocator
 {
 public:
 	enum
@@ -46,16 +46,16 @@ public:
 	typedef TLockFreeFixedSizeAllocator<PageSize, PLATFORM_CACHE_LINE_SIZE, FThreadSafeCounter> TPageAllocator;
 #endif
 
-	static FPageAllocator& Get();
+	static CORE_API FPageAllocator& Get();
 
-	~FPageAllocator();
-	void* Alloc();
-	void Free(void* Mem);
-	void* AllocSmall();
-	void FreeSmall(void* Mem);
-	uint64 BytesUsed();
-	uint64 BytesFree();
-	void LatchProtectedMode();
+	CORE_API ~FPageAllocator();
+	CORE_API void* Alloc();
+	CORE_API void Free(void* Mem);
+	CORE_API void* AllocSmall();
+	CORE_API void FreeSmall(void* Mem);
+	CORE_API uint64 BytesUsed();
+	CORE_API uint64 BytesFree();
+	CORE_API void LatchProtectedMode();
 private:
 
 	static FPageAllocator* Instance;
@@ -74,11 +74,11 @@ private:
  * Items are allocated via PushBytes() or the specialized operator new()s.
  * Items are freed en masse by using FMemMark to Pop() them.
  **/
-class CORE_API FMemStackBase //-V1062
+class FMemStackBase //-V1062
 {
 public:
 
-	FMemStackBase();
+	CORE_API FMemStackBase();
 
 	FMemStackBase(const FMemStackBase&) = delete;
 	FMemStackBase(FMemStackBase&& Other)
@@ -157,10 +157,10 @@ public:
 		return NumMarks;
 	}
 	/** @return the number of bytes allocated for this FMemStack that are currently in use. */
-	int32 GetByteCount() const;
+	CORE_API int32 GetByteCount() const;
 
 	// Returns true if the pointer was allocated using this allocator
-	bool ContainsPointer(const void* Pointer) const;
+	CORE_API bool ContainsPointer(const void* Pointer) const;
 
 	// Friends.
 	friend class FMemMark;
@@ -195,10 +195,10 @@ private:
 	 * Allocate a new chunk of memory of at least MinSize size,
 	 * updates the memory stack's Chunks table and ActiveChunks counter.
 	 */
-	void AllocateNewChunk( int32 MinSize );
+	CORE_API void AllocateNewChunk( int32 MinSize );
 
 	/** Frees the chunks above the specified chunk on the stack. */
-	void FreeChunks( FTaggedMemory* NewTopChunk );
+	CORE_API void FreeChunks( FTaggedMemory* NewTopChunk );
 
 	// Variables.
 	uint8*			Top;				// Top of current chunk (Top<=End).
@@ -535,4 +535,4 @@ private:
 	FMemStackBase::FTaggedMemory* SavedChunk;
 	bool bPopped;
 	FMemMark* NextTopmostMark;
-};
\ No newline at end of file
+};
diff --git a/Engine/Source/Runtime/Core/Public/Misc/MemoryReadStream.h b/Engine/Source/Runtime/Core/Public/Misc/MemoryReadStream.h
--- a/Engine/Source/Runtime/Core/Public/Misc/MemoryReadStream.h
+++ b/Engine/Source/Runtime/Core/Public/Misc/MemoryReadStream.h
@@ -10,16 +10,16 @@
 /**
  * Represents a block of memory, but abstracts away the underlying layout
  */
-class CORE_API IMemoryReadStream
+class IMemoryReadStream
 {
 public:
-	static TRefCountPtr<IMemoryReadStream> CreateFromCopy(const void* InMemory, int64 InSize);
-	static TRefCountPtr<IMemoryReadStream> CreateFromCopy(IMemoryReadStream* InStream);
-	static TRefCountPtr<IMemoryReadStream> CreateFromBuffer(void* InMemory, int64 InSize, bool bOwnPointer);
+	static CORE_API TRefCountPtr<IMemoryReadStream> CreateFromCopy(const void* InMemory, int64 InSize);
+	static CORE_API TRefCountPtr<IMemoryReadStream> CreateFromCopy(IMemoryReadStream* InStream);
+	static CORE_API TRefCountPtr<IMemoryReadStream> CreateFromBuffer(void* InMemory, int64 InSize, bool bOwnPointer);
 
 	virtual const void* Read(int64& OutSize, int64 InOffset, int64 InSize) = 0;
 	virtual int64 GetSize() = 0;
-	virtual void CopyTo(void* Buffer, int64 InOffset, int64 InSize);
+	CORE_API virtual void CopyTo(void* Buffer, int64 InOffset, int64 InSize);
 	virtual void EnsureReadNonBlocking() {};		// Incur any potential wait such that future Read() or CopyTo() calls are wait and thread switch free.
 
 	FORCEINLINE uint32 AddRef() const { return uint32(NumRefs.Increment()); }
diff --git a/Engine/Source/Runtime/Core/Public/Misc/MessageDialog.h b/Engine/Source/Runtime/Core/Public/Misc/MessageDialog.h
--- a/Engine/Source/Runtime/Core/Public/Misc/MessageDialog.h
+++ b/Engine/Source/Runtime/Core/Public/Misc/MessageDialog.h
@@ -13,19 +13,19 @@ class FText;
  * These functions open a message dialog and display the specified informations
  * there.
  **/
-struct CORE_API FMessageDialog
+struct FMessageDialog
 {
 	/** Pops up a message dialog box containing the input string.
 	 * @param Message Text of message to show
 	 * @param Title Optional title to use (defaults to "Message")
 	*/
-	static void Debugf( const FText& Message );
-	static void Debugf( const FText& Message, const FText& Title );
+	static CORE_API void Debugf( const FText& Message );
+	static CORE_API void Debugf( const FText& Message, const FText& Title );
 	UE_DEPRECATED(5.3, "Use the overload of Debugf that takes the Title by-value.")
-	static void Debugf( const FText& Message, const FText* OptTitle );
+	static CORE_API void Debugf( const FText& Message, const FText* OptTitle );
 
 	/** Pops up a message dialog box containing the last system error code in string form. */
-	static void ShowLastError();
+	static CORE_API void ShowLastError();
 
 	/**
 	 * Open a modal message box dialog
@@ -33,10 +33,10 @@ struct CORE_API FMessageDialog
 	 * @param Message Text of message to show
 	 * @param Title Optional title to use (defaults to "Message")
 	*/
-	static EAppReturnType::Type Open( EAppMsgType::Type MessageType, const FText& Message);
-	static EAppReturnType::Type Open( EAppMsgType::Type MessageType, const FText& Message, const FText& Title);
+	static CORE_API EAppReturnType::Type Open( EAppMsgType::Type MessageType, const FText& Message);
+	static CORE_API EAppReturnType::Type Open( EAppMsgType::Type MessageType, const FText& Message, const FText& Title);
 	UE_DEPRECATED(5.3, "Use the overload of Open that takes the Title by-value.")
-	static EAppReturnType::Type Open( EAppMsgType::Type MessageType, const FText& Message, const FText* OptTitle);
+	static CORE_API EAppReturnType::Type Open( EAppMsgType::Type MessageType, const FText& Message, const FText* OptTitle);
 
 	/**
 	 * Open a modal message box dialog
@@ -45,8 +45,8 @@ struct CORE_API FMessageDialog
 	 * @param Message Text of message to show
 	 * @param Title Optional title to use (defaults to "Message")
 	*/
-	static EAppReturnType::Type Open(EAppMsgType::Type MessageType, EAppReturnType::Type DefaultValue, const FText& Message);
-	static EAppReturnType::Type Open(EAppMsgType::Type MessageType, EAppReturnType::Type DefaultValue, const FText& Message, const FText& Title);
+	static CORE_API EAppReturnType::Type Open(EAppMsgType::Type MessageType, EAppReturnType::Type DefaultValue, const FText& Message);
+	static CORE_API EAppReturnType::Type Open(EAppMsgType::Type MessageType, EAppReturnType::Type DefaultValue, const FText& Message, const FText& Title);
 	UE_DEPRECATED(5.3, "Use the overload of Open that takes the Title by-value.")
-	static EAppReturnType::Type Open(EAppMsgType::Type MessageType, EAppReturnType::Type DefaultValue, const FText& Message, const FText* OptTitle);
+	static CORE_API EAppReturnType::Type Open(EAppMsgType::Type MessageType, EAppReturnType::Type DefaultValue, const FText& Message, const FText* OptTitle);
 };
diff --git a/Engine/Source/Runtime/Core/Public/Misc/MonitoredProcess.h b/Engine/Source/Runtime/Core/Public/Misc/MonitoredProcess.h
--- a/Engine/Source/Runtime/Core/Public/Misc/MonitoredProcess.h
+++ b/Engine/Source/Runtime/Core/Public/Misc/MonitoredProcess.h
@@ -32,7 +32,7 @@ DECLARE_DELEGATE_OneParam(FOnMonitoredProcessOutput, FString)
 /**
  * Implements an external process that can be monitored.
  */
-class CORE_API FMonitoredProcess
+class FMonitoredProcess
 	: public FRunnable, FSingleThreadRunnable
 {
 public:
@@ -45,7 +45,7 @@ public:
 	 * @param InHidden Whether the window of the process should be hidden.
 	 * @param InCreatePipes Whether the output should be redirected to the caller.
 	 */
-	FMonitoredProcess( const FString& InURL, const FString& InParams, bool InHidden, bool InCreatePipes = true );
+	CORE_API FMonitoredProcess( const FString& InURL, const FString& InParams, bool InHidden, bool InCreatePipes = true );
 
 	/**
 	* Creates a new monitored process.
@@ -56,10 +56,10 @@ public:
 	* @param InWorkingDir The URL of the working dir where the executable should launch.
 	* @param InCreatePipes Whether the output should be redirected to the caller.
 	*/
-	FMonitoredProcess( const FString& InURL, const FString& InParams, const FString& InWorkingDir, bool InHidden, bool InCreatePipes = true );
+	CORE_API FMonitoredProcess( const FString& InURL, const FString& InParams, const FString& InWorkingDir, bool InHidden, bool InCreatePipes = true );
 
 	/** Destructor. */
-	virtual ~FMonitoredProcess();
+	CORE_API virtual ~FMonitoredProcess();
 
 public:
 
@@ -79,7 +79,7 @@ public:
 	 *
 	 * @return Time duration.
 	 */
-	FTimespan GetDuration() const;
+	CORE_API FTimespan GetDuration() const;
 
 	/**
 	 * Gets the Process Handle. The instance can be invalid if the process was not created.
@@ -104,10 +104,10 @@ public:
 	*
 	* @return true if the process is running, false otherwise.
 	*/
-	bool Update();
+	CORE_API bool Update();
 
 	/** Launches the process. */
-	virtual bool Launch();
+	CORE_API virtual bool Launch();
 
 	/**
 	 * Sets the sleep interval to be used in the main thread loop.
@@ -179,7 +179,7 @@ public:
 		return true;
 	}
 
-	virtual uint32 Run() override;
+	CORE_API virtual uint32 Run() override;
 
 	virtual void Stop() override
 	{
@@ -198,17 +198,17 @@ protected:
 	/**
 	* FSingleThreadRunnable interface
 	*/
-	void Tick() override;
+	CORE_API void Tick() override;
 
 	/**
 	 * Processes the given output string.
 	 *
 	 * @param Output The output string to process.
 	 */
-	void ProcessOutput( const FString& Output );
+	CORE_API void ProcessOutput( const FString& Output );
 
 protected:
-	void TickInternal();
+	CORE_API void TickInternal();
 
 
 	// Whether the process is being canceled. */
@@ -275,22 +275,22 @@ protected:
 };
 
 
-class CORE_API FSerializedUATProcess : public FMonitoredProcess
+class FSerializedUATProcess : public FMonitoredProcess
 {
 public:
 	/**
 	 * Get the host-platform-specific path to the UAT running script
 	 */
-	static FString GetUATPath();
+	static CORE_API FString GetUATPath();
 
 public:
-	FSerializedUATProcess(const FString& RunUATCommandline);
+	CORE_API FSerializedUATProcess(const FString& RunUATCommandline);
 
 	/**
 	 * Run UAT, serially with other FSerializedUATProcess objects. Because the actual call is delayed, this will
 	 * always return true, and the LaunchFailedDelegate will be called later if an error happens
 	 */
-	virtual bool Launch() override;
+	CORE_API virtual bool Launch() override;
 
 	/**
 	 * Returns a delegate that is executed when the process has been canceled.
@@ -305,9 +305,9 @@ public:
 
 private:
 
-	bool LaunchNext();
-	bool LaunchInternal();
-	static void CancelQueue();
+	CORE_API bool LaunchNext();
+	CORE_API bool LaunchInternal();
+	static CORE_API void CancelQueue();
 
 	// When this one completes, run the next in line
 	FSerializedUATProcess* NextProcessToRun = nullptr;
@@ -316,7 +316,7 @@ private:
 	// of Launch in the parent class, since it's async
 	FSimpleDelegate LaunchFailedDelegate;
 
-	static FCriticalSection Serializer;
-	static bool bHasSucceededOnce;
-	static FSerializedUATProcess* HeadProcess;
+	static CORE_API FCriticalSection Serializer;
+	static CORE_API bool bHasSucceededOnce;
+	static CORE_API FSerializedUATProcess* HeadProcess;
 };
diff --git a/Engine/Source/Runtime/Core/Public/Misc/NamePermissionList.h b/Engine/Source/Runtime/Core/Public/Misc/NamePermissionList.h
--- a/Engine/Source/Runtime/Core/Public/Misc/NamePermissionList.h
+++ b/Engine/Source/Runtime/Core/Public/Misc/NamePermissionList.h
@@ -17,74 +17,74 @@ class FString;
 /** List of owner names that requested a specific item filtered, allows unregistering specific set of changes by a given plugin or system */
 typedef TArray<FName> FPermissionListOwners;
 
-class CORE_API FNamePermissionList : public TSharedFromThis<FNamePermissionList>
+class FNamePermissionList : public TSharedFromThis<FNamePermissionList>
 {
 public:
 	FNamePermissionList() {}
 	virtual ~FNamePermissionList() {}
 
 	/** Returns true if passes filter restrictions using exact match */
-	bool PassesFilter(const FName Item) const;
+	CORE_API bool PassesFilter(const FName Item) const;
 
 	/** 
 	 * Add item to DenyList, this specific item will be filtered out.
 	 * @return whether the filters changed.
 	 */
-	bool AddDenyListItem(const FName OwnerName, const FName Item);
+	CORE_API bool AddDenyListItem(const FName OwnerName, const FName Item);
 
 	/**
 	 * Add item to allowlist after which all items not in the allowlist will be filtered out.
 	 * @return whether the filters changed.
 	 */
-	bool AddAllowListItem(const FName OwnerName, const FName Item);
+	CORE_API bool AddAllowListItem(const FName OwnerName, const FName Item);
 
 	/**
 	 * Removes a previously-added item from the DenyList.
 	 * @return whether the filters changed.
 	 */
-	bool RemoveDenyListItem(const FName OwnerName, const FName Item);
+	CORE_API bool RemoveDenyListItem(const FName OwnerName, const FName Item);
 
 	/**
 	 * Removes a previously-added item from the allowlist.
 	 * @return whether the filters changed.
 	 */
-	bool RemoveAllowListItem(const FName OwnerName, const FName Item);
+	CORE_API bool RemoveAllowListItem(const FName OwnerName, const FName Item);
 
 	/**
 	 * Set to filter out all items.
 	 * @return whether the filters changed.
 	 */
-	bool AddDenyListAll(const FName OwnerName);
+	CORE_API bool AddDenyListAll(const FName OwnerName);
 	
 	/** True if has filters active */
-	bool HasFiltering() const;
+	CORE_API bool HasFiltering() const;
 
 	/** Gathers the names of all the owners in this DenyList. */
-	TArray<FName> GetOwnerNames() const;
+	CORE_API TArray<FName> GetOwnerNames() const;
 
 	/** 
 	* Removes all filtering changes associated with a specific owner name.
 	 * @return whether the filters changed.
 	 */
-	bool UnregisterOwner(const FName OwnerName);
+	CORE_API bool UnregisterOwner(const FName OwnerName);
 
 	/**
 	 * Removes all filtering changes associated with the specified list of owner names.
 	 * @return whether the filters changed.
 	 */
-	bool UnregisterOwners(const TArray<FName>& OwnerNames);
+	CORE_API bool UnregisterOwners(const TArray<FName>& OwnerNames);
 
 	/**
 	 * Add the specified filters to this one.
 	 * @return whether the filters changed.
 	 */
-	bool Append(const FNamePermissionList& Other);
+	CORE_API bool Append(const FNamePermissionList& Other);
 
 	/**
 	* Unregisters specified owners then adds specified filters in one operation (to avoid multiple filters changed events).
 	* @return whether the filters changed.
 	*/
-	bool UnregisterOwnersAndAppend(const TArray<FName>& OwnerNamesToRemove, const FNamePermissionList& FiltersToAdd);
+	CORE_API bool UnregisterOwnersAndAppend(const TArray<FName>& OwnerNamesToRemove, const FNamePermissionList& FiltersToAdd);
 
 	/** Get raw DenyList */
 	const TMap<FName, FPermissionListOwners>& GetDenyList() const { return DenyList; }
@@ -122,7 +122,7 @@ enum class EPathPermissionListType
 	ClassPaths	// Class permission list
 };
 
-class CORE_API FPathPermissionList : public TSharedFromThis<FPathPermissionList>
+class FPathPermissionList : public TSharedFromThis<FPathPermissionList>
 {
 public:
 	FPathPermissionList(EPathPermissionListType InType = EPathPermissionListType::Default) 
@@ -132,124 +132,124 @@ public:
 	virtual ~FPathPermissionList() {}
 	
 	/** Returns true if passes filter restrictions using exact match */
-	bool PassesFilter(const FStringView Item) const;
+	CORE_API bool PassesFilter(const FStringView Item) const;
 
 	/** Returns true if passes filter restrictions using exact match */
-	bool PassesFilter(const FName Item) const;
+	CORE_API bool PassesFilter(const FName Item) const;
 
 	/** Returns true if passes filter restrictions using exact match */
-	bool PassesFilter(const TCHAR* Item) const;
+	CORE_API bool PassesFilter(const TCHAR* Item) const;
 
 	/** Returns true if passes filter restrictions for path */
-	bool PassesStartsWithFilter(const FStringView Item, const bool bAllowParentPaths = false) const;
+	CORE_API bool PassesStartsWithFilter(const FStringView Item, const bool bAllowParentPaths = false) const;
 
 	/** Returns true if passes filter restrictions for path */
-	bool PassesStartsWithFilter(const FName Item, const bool bAllowParentPaths = false) const;
+	CORE_API bool PassesStartsWithFilter(const FName Item, const bool bAllowParentPaths = false) const;
 
 	/** Returns true if passes filter restrictions for path */
-	bool PassesStartsWithFilter(const TCHAR* Item, const bool bAllowParentPaths = false) const;
+	CORE_API bool PassesStartsWithFilter(const TCHAR* Item, const bool bAllowParentPaths = false) const;
 
 	/** 
 	 * Add item to DenyList, this specific item will be filtered out.
 	 * @return whether the filters changed.
 	 */
-	bool AddDenyListItem(const FName OwnerName, const FStringView Item);
+	CORE_API bool AddDenyListItem(const FName OwnerName, const FStringView Item);
 
 	/**
 	 * Add item to DenyList, this specific item will be filtered out.
 	 * @return whether the filters changed.
 	 */
-	bool AddDenyListItem(const FName OwnerName, const FName Item);
+	CORE_API bool AddDenyListItem(const FName OwnerName, const FName Item);
 
 	/**
 	 * Add item to DenyList, this specific item will be filtered out.
 	 * @return whether the filters changed.
 	 */
-	bool AddDenyListItem(const FName OwnerName, const TCHAR* Item);
+	CORE_API bool AddDenyListItem(const FName OwnerName, const TCHAR* Item);
 
 	/**
 	* Remove item from the DenyList
 	* @return whether the filters changed
 	*/
-	bool RemoveDenyListItem(const FName OwnerName, const FStringView Item);
+	CORE_API bool RemoveDenyListItem(const FName OwnerName, const FStringView Item);
 
 	/**
 	* Remove item from the DenyList
 	* @return whether the filters changed
 	*/
-	bool RemoveDenyListItem(const FName OwnerName, const FName Item);
+	CORE_API bool RemoveDenyListItem(const FName OwnerName, const FName Item);
 
 	/**
 	* Remove item from the DenyList
 	* @return whether the filters changed
 	*/
-	bool RemoveDenyListItem(const FName OwnerName, const TCHAR* Item);
+	CORE_API bool RemoveDenyListItem(const FName OwnerName, const TCHAR* Item);
 
 	/**
 	 * Add item to allowlist after which all items not in the allowlist will be filtered out.
 	 * @return whether the filters changed.
 	 */
-	bool AddAllowListItem(const FName OwnerName, const FStringView Item);
+	CORE_API bool AddAllowListItem(const FName OwnerName, const FStringView Item);
 
 	/**
 	 * Add item to allowlist after which all items not in the allowlist will be filtered out.
 	 * @return whether the filters changed.
 	 */
-	bool AddAllowListItem(const FName OwnerName, const FName Item);
+	CORE_API bool AddAllowListItem(const FName OwnerName, const FName Item);
 
 	/**
 	 * Add item to allowlist after which all items not in the allowlist will be filtered out.
 	 * @return whether the filters changed.
 	 */
-	bool AddAllowListItem(const FName OwnerName, const TCHAR* Item);
+	CORE_API bool AddAllowListItem(const FName OwnerName, const TCHAR* Item);
 
 	/**
 	* Remove item from the AllowList
 	* @return whether the filters changed
 	*/
-	bool RemoveAllowListItem(const FName OwnerName, const FStringView Item);
+	CORE_API bool RemoveAllowListItem(const FName OwnerName, const FStringView Item);
 
 	/**
 	* Remove item from the AllowList
 	* @return whether the filters changed
 	*/
-	bool RemoveAllowListItem(const FName OwnerName, const FName Item);
+	CORE_API bool RemoveAllowListItem(const FName OwnerName, const FName Item);
 
 	/**
 	* Remove item from the AllowList
 	* @return whether the filters changed
 	*/
-	bool RemoveAllowListItem(const FName OwnerName, const TCHAR* Item);
+	CORE_API bool RemoveAllowListItem(const FName OwnerName, const TCHAR* Item);
 
 	/**
 	 * Set to filter out all items.
 	 * @return whether the filters changed.
 	 */
-	bool AddDenyListAll(const FName OwnerName);
+	CORE_API bool AddDenyListAll(const FName OwnerName);
 	
 	/** True if has filters active */
-	bool HasFiltering() const;
+	CORE_API bool HasFiltering() const;
 
 	/** Gathers the names of all the owners in this DenyList. */
-	TArray<FName> GetOwnerNames() const;
+	CORE_API TArray<FName> GetOwnerNames() const;
 
 	/**
 	 * Removes all filtering changes associated with a specific owner name.
 	 * @return whether the filters changed.
 	 */
-	bool UnregisterOwner(const FName OwnerName);
+	CORE_API bool UnregisterOwner(const FName OwnerName);
 
 	/**
 	 * Removes all filtering changes associated with the specified list of owner names.
 	 * @return whether the filters changed.
 	 */
-	bool UnregisterOwners(const TArray<FName>& OwnerNames);
+	CORE_API bool UnregisterOwners(const TArray<FName>& OwnerNames);
 	
 	/**
 	 * Add the specified filters to this one. Rules are not applied, direct append lists.
 	 * @return whether the filters changed.
 	 */
-	bool Append(const FPathPermissionList& Other);
+	CORE_API bool Append(const FPathPermissionList& Other);
 
 	/**
 	 * Combine two filters.
@@ -257,13 +257,13 @@ public:
 	 * Result will contain AllowList paths that pass both filters.
 	 * @return new combined filter.
 	 */
-	FPathPermissionList CombinePathFilters(const FPathPermissionList& OtherFilter) const;
+	CORE_API FPathPermissionList CombinePathFilters(const FPathPermissionList& OtherFilter) const;
 
 	/**
 	* Unregisters specified owners then adds specified filters in one operation (to avoid multiple filters changed events).
 	* @return whether the filters changed.
 	*/
-	bool UnregisterOwnersAndAppend(const TArray<FName>& OwnerNamesToRemove, const FPathPermissionList& FiltersToAdd);
+	CORE_API bool UnregisterOwnersAndAppend(const TArray<FName>& OwnerNamesToRemove, const FPathPermissionList& FiltersToAdd);
 
 	/** Get raw DenyList */
 	const TMap<FString, FPermissionListOwners>& GetDenyList() const { return DenyList; }
@@ -278,7 +278,7 @@ public:
 	FSimpleMulticastDelegate& OnFilterChanged() { return OnFilterChangedDelegate; }
 
 	/** Dumps the path permission list details into a multi-line string */
-	FString ToString() const;
+	CORE_API FString ToString() const;
 
 protected:
 
@@ -286,7 +286,7 @@ protected:
 	 * Checks if an item is of a valid format for this list
 	 * @return True if the item passes list type test.
 	 */
-	void VerifyItemMatchesListType(const FStringView Item) const;
+	CORE_API void VerifyItemMatchesListType(const FStringView Item) const;
 
 	/** List if items to filter out */
 	TMap<FString, FPermissionListOwners> DenyList;
diff --git a/Engine/Source/Runtime/Core/Public/Misc/NetworkVersion.h b/Engine/Source/Runtime/Core/Public/Misc/NetworkVersion.h
--- a/Engine/Source/Runtime/Core/Public/Misc/NetworkVersion.h
+++ b/Engine/Source/Runtime/Core/Public/Misc/NetworkVersion.h
@@ -48,46 +48,46 @@ enum class EEngineNetworkRuntimeFeatures : uint16
 };
 ENUM_CLASS_FLAGS(EEngineNetworkRuntimeFeatures);
 
-struct CORE_API FNetworkVersion
+struct FNetworkVersion
 {
 	/** Called in GetLocalNetworkVersion if bound */
 	DECLARE_DELEGATE_RetVal( uint32, FGetLocalNetworkVersionOverride );
-	static FGetLocalNetworkVersionOverride GetLocalNetworkVersionOverride;
+	static CORE_API FGetLocalNetworkVersionOverride GetLocalNetworkVersionOverride;
 
 	/** Called in IsNetworkCompatible if bound */
 	DECLARE_DELEGATE_RetVal_TwoParams( bool, FIsNetworkCompatibleOverride, uint32, uint32 );
-	static FIsNetworkCompatibleOverride IsNetworkCompatibleOverride;
+	static CORE_API FIsNetworkCompatibleOverride IsNetworkCompatibleOverride;
 
 	/** Called in GetReplayCompatibleChangelist if bound */
 	DECLARE_DELEGATE_RetVal(uint32, FGetReplayCompatibleChangeListOverride);
-	static FGetReplayCompatibleChangeListOverride GetReplayCompatibleChangeListOverride;
+	static CORE_API FGetReplayCompatibleChangeListOverride GetReplayCompatibleChangeListOverride;
 
-	static uint32 GetNetworkCompatibleChangelist();
-	static uint32 GetReplayCompatibleChangelist();
+	static CORE_API uint32 GetNetworkCompatibleChangelist();
+	static CORE_API uint32 GetReplayCompatibleChangelist();
 
 	UE_DEPRECATED(5.2, "Please use GetNetworkProtocolVersion instead.")
-	static uint32 GetEngineNetworkProtocolVersion();
+	static CORE_API uint32 GetEngineNetworkProtocolVersion();
 	UE_DEPRECATED(5.2, "Please use GetNetworkProtocolVersion instead.")
-	static uint32 GetGameNetworkProtocolVersion();
+	static CORE_API uint32 GetGameNetworkProtocolVersion();
 
 	UE_DEPRECATED(5.2, "Please use GetCompatibleNetworkProtocolVersion instead.")
-	static uint32 GetEngineCompatibleNetworkProtocolVersion();
+	static CORE_API uint32 GetEngineCompatibleNetworkProtocolVersion();
 	UE_DEPRECATED(5.2, "Please use GetCompatibleNetworkProtocolVersion instead.")
-	static uint32 GetGameCompatibleNetworkProtocolVersion();
+	static CORE_API uint32 GetGameCompatibleNetworkProtocolVersion();
 
-	static uint32 GetNetworkProtocolVersion(const FGuid& VersionGuid);
-	static uint32 GetCompatibleNetworkProtocolVersion(const FGuid& VersionGuid);
+	static CORE_API uint32 GetNetworkProtocolVersion(const FGuid& VersionGuid);
+	static CORE_API uint32 GetCompatibleNetworkProtocolVersion(const FGuid& VersionGuid);
 
-	static const FCustomVersionContainer& GetNetworkCustomVersions();
+	static CORE_API const FCustomVersionContainer& GetNetworkCustomVersions();
 
-	static void RegisterNetworkCustomVersion(const FGuid& VersionGuid, int32 Version, int32 CompatibleVersion, const FName& FriendlyName);
+	static CORE_API void RegisterNetworkCustomVersion(const FGuid& VersionGuid, int32 Version, int32 CompatibleVersion, const FName& FriendlyName);
 
 	/**
 	* Generates a version number, that by default, is based on a checksum of the engine version + project name + project version string
 	* Game/project code can completely override what this value returns through the GetLocalNetworkVersionOverride delegate
 	* If called with AllowOverrideDelegate=false, we will not call the game project override. (This allows projects to call base implementation in their project implementation)
 	*/
-	static uint32 GetLocalNetworkVersion( bool AllowOverrideDelegate=true );
+	static CORE_API uint32 GetLocalNetworkVersion( bool AllowOverrideDelegate=true );
 
 	/**
 	* Determine if a connection is compatible with this instance
@@ -98,12 +98,12 @@ struct CORE_API FNetworkVersion
 	*
 	* @return true if the two instances can communicate, false otherwise
 	*/
-	static bool IsNetworkCompatible( const uint32 LocalNetworkVersion, const uint32 RemoteNetworkVersion );
+	static CORE_API bool IsNetworkCompatible( const uint32 LocalNetworkVersion, const uint32 RemoteNetworkVersion );
 
 	/**
 	* Generates a special struct that contains information to send to replay server
 	*/
-	static FNetworkReplayVersion GetReplayVersion();
+	static CORE_API FNetworkReplayVersion GetReplayVersion();
 
 	/**
 	* Sets the project version used for networking. Needs to be a function to verify
@@ -112,27 +112,27 @@ struct CORE_API FNetworkVersion
 	* @param  InVersion
 	* @return void
 	*/
-	static void SetProjectVersion(const TCHAR* InVersion);
+	static CORE_API void SetProjectVersion(const TCHAR* InVersion);
 
 	/**
 	* Sets the game network protocol version used for networking and invalidate cached values
 	*/
-	static void SetGameNetworkProtocolVersion(uint32 GameNetworkProtocolVersion);
+	static CORE_API void SetGameNetworkProtocolVersion(uint32 GameNetworkProtocolVersion);
 
 	/**
 	* Sets the game compatible network protocol version used for networking and invalidate cached values
 	*/
-	static void SetGameCompatibleNetworkProtocolVersion(uint32 GameCompatibleNetworkProtocolVersion);
+	static CORE_API void SetGameCompatibleNetworkProtocolVersion(uint32 GameCompatibleNetworkProtocolVersion);
 
 	/**
 	 * Compares if the connection's runtime features are compatible with each other
 	 */
-	static bool AreNetworkRuntimeFeaturesCompatible(EEngineNetworkRuntimeFeatures LocalFeatures, EEngineNetworkRuntimeFeatures RemoteFeatures);
+	static CORE_API bool AreNetworkRuntimeFeaturesCompatible(EEngineNetworkRuntimeFeatures LocalFeatures, EEngineNetworkRuntimeFeatures RemoteFeatures);
 
 	/**
 	 * Build and return a string describing the status of the the network runtime features bitflag
 	 */
-	static void DescribeNetworkRuntimeFeaturesBitset(EEngineNetworkRuntimeFeatures FeaturesBitflag, FStringBuilderBase& OutVerboseDescription);
+	static CORE_API void DescribeNetworkRuntimeFeaturesBitset(EEngineNetworkRuntimeFeatures FeaturesBitflag, FStringBuilderBase& OutVerboseDescription);
 	
 	/**
 	* Returns the project version used by networking
@@ -151,21 +151,21 @@ protected:
 	/**
 	* Used to allow BP only projects to override network versions
 	*/
-	static FString& GetProjectVersion_Internal();
+	static CORE_API FString& GetProjectVersion_Internal();
 
-	static bool		bHasCachedNetworkChecksum;
-	static uint32	CachedNetworkChecksum;
+	static CORE_API bool		bHasCachedNetworkChecksum;
+	static CORE_API uint32	CachedNetworkChecksum;
 
-	static bool		bHasCachedReplayChecksum;
-	static uint32	CachedReplayChecksum;
+	static CORE_API bool		bHasCachedReplayChecksum;
+	static CORE_API uint32	CachedReplayChecksum;
 
 	UE_DEPRECATED(5.2, "Now storing this value in NetworkCustomVersions, do not use directly.")
-	static uint32	EngineNetworkProtocolVersion;
+	static CORE_API uint32	EngineNetworkProtocolVersion;
 	UE_DEPRECATED(5.2, "Now storing this value in NetworkCustomVersions, do not use directly.")
-	static uint32	GameNetworkProtocolVersion;
+	static CORE_API uint32	GameNetworkProtocolVersion;
 
 	UE_DEPRECATED(5.2, "Now storing this value in CompatibleNetworkCustomVersions, do not use directly.")
-	static uint32	EngineCompatibleNetworkProtocolVersion;
+	static CORE_API uint32	EngineCompatibleNetworkProtocolVersion;
 	UE_DEPRECATED(5.2, "Now storing this value in CompatibleNetworkCustomVersions, do not use directly.")
-	static uint32	GameCompatibleNetworkProtocolVersion;
+	static CORE_API uint32	GameCompatibleNetworkProtocolVersion;
 };
diff --git a/Engine/Source/Runtime/Core/Public/Misc/ObjectThumbnail.h b/Engine/Source/Runtime/Core/Public/Misc/ObjectThumbnail.h
--- a/Engine/Source/Runtime/Core/Public/Misc/ObjectThumbnail.h
+++ b/Engine/Source/Runtime/Core/Public/Misc/ObjectThumbnail.h
@@ -56,7 +56,7 @@ public:
 /**
  * Thumbnail image data for an object.
  */
-class CORE_API FObjectThumbnail
+class FObjectThumbnail
 {
 public:
 
@@ -76,13 +76,13 @@ public:
 private:
 
 	/** Static: Thumbnail compressor. */
-	static FThumbnailCompressionInterface* PNGThumbnailCompressor;
-	static FThumbnailCompressionInterface* JPEGThumbnailCompressor;
+	static CORE_API FThumbnailCompressionInterface* PNGThumbnailCompressor;
+	static CORE_API FThumbnailCompressionInterface* JPEGThumbnailCompressor;
 
 public:
 
 	/** Default constructor. */
-	FObjectThumbnail();
+	CORE_API FObjectThumbnail();
 
 	/** Returns the width of the thumbnail. */
 	int32 GetImageWidth() const
@@ -155,44 +155,44 @@ public:
 	}
 
 	/** Returns thumbnail compressor used on current compressed image data. */
-	FThumbnailCompressionInterface* GetCompressor() const;
+	CORE_API FThumbnailCompressionInterface* GetCompressor() const;
 
 	/** Returns thumbnail compressor that would be used on current uncompressed image data. */
-	FThumbnailCompressionInterface* ChooseNewCompressor() const;
+	CORE_API FThumbnailCompressionInterface* ChooseNewCompressor() const;
 
 	/** Returns uncompressed image data, decompressing it on demand if needed. */
-	const TArray< uint8 >& GetUncompressedImageData() const;
+	CORE_API const TArray< uint8 >& GetUncompressedImageData() const;
 
 	/** Serializers */
-	void Serialize(FArchive& Ar);
-	void Serialize(FStructuredArchive::FSlot Slot);
+	CORE_API void Serialize(FArchive& Ar);
+	CORE_API void Serialize(FStructuredArchive::FSlot Slot);
 	
 	/** Compress image data. */
-	void CompressImageData();
+	CORE_API void CompressImageData();
 
 	/** Decompress image data. */
-	void DecompressImageData();
+	CORE_API void DecompressImageData();
 
 	/**
 	 * Calculates the memory usage of this FObjectThumbnail.
 	 *
 	 * @param	Ar	the FArchiveCountMem (or similar) archive that will store the results of the memory usage calculation.
 	 */
-	void CountBytes( FArchive& Ar ) const;
+	CORE_API void CountBytes( FArchive& Ar ) const;
 
 	/**
 	 * Calculates the amount of memory used by the compressed bytes array.
 	 *
 	 * @param	Ar	the FArchiveCountMem (or similar) archive that will store the results of the memory usage calculation.
 	 */
-	void CountImageBytes_Compressed( FArchive& Ar ) const;
+	CORE_API void CountImageBytes_Compressed( FArchive& Ar ) const;
 
 	/**
 	 * Calculates the amount of memory used by the uncompressed bytes array.
 	 *
 	 * @param	Ar	the FArchiveCountMem (or similar) archive that will store the results of the memory usage calculation.
 	 */
-	void CountImageBytes_Uncompressed( FArchive& Ar ) const;
+	CORE_API void CountImageBytes_Uncompressed( FArchive& Ar ) const;
 
 	/** I/O operator */
 	friend FArchive& operator<<( FArchive& Ar, FObjectThumbnail& Thumb )
@@ -264,7 +264,7 @@ typedef TMap< FName, FObjectThumbnail > FThumbnailMap;
 
 
 /** Wraps an object's full name and thumbnail */
-struct CORE_API FObjectFullNameAndThumbnail
+struct FObjectFullNameAndThumbnail
 {
 	/** Full name of the object */
 	FName ObjectFullName;
@@ -294,7 +294,7 @@ struct CORE_API FObjectFullNameAndThumbnail
 	 *
 	 * @param	Ar	the FArchiveCountMem (or similar) archive that will store the results of the memory usage calculation.
 	 */
-	void CountBytes( FArchive& Ar ) const;
+	CORE_API void CountBytes( FArchive& Ar ) const;
 
 	/** I/O operator */
 	friend FArchive& operator<<( FArchive& Ar, FObjectFullNameAndThumbnail& NameThumbPair )
diff --git a/Engine/Source/Runtime/Core/Public/Misc/OutputDevice.h b/Engine/Source/Runtime/Core/Public/Misc/OutputDevice.h
--- a/Engine/Source/Runtime/Core/Public/Misc/OutputDevice.h
+++ b/Engine/Source/Runtime/Core/Public/Misc/OutputDevice.h
@@ -126,7 +126,7 @@ namespace ELogTimes
 class FName;
 
 // An output device.
-class CORE_API FOutputDevice
+class FOutputDevice
 {
 public:
 	FOutputDevice()
@@ -148,7 +148,7 @@ public:
 		Serialize( V, Verbosity, Category );
 	}
 
-	virtual void SerializeRecord(const UE::FLogRecord& Record);
+	CORE_API virtual void SerializeRecord(const UE::FLogRecord& Record);
 
 	virtual void Flush()
 	{
@@ -216,18 +216,18 @@ public:
 	}
 
 	// Simple text printing.
-	void Log( const TCHAR* S );
-	void Log( ELogVerbosity::Type Verbosity, const TCHAR* S );
-	void Log( const FName& Category, ELogVerbosity::Type Verbosity, const TCHAR* Str );
-	void Log( const FString& S );
-	void Log( const FText& S );
-	void Log( ELogVerbosity::Type Verbosity, const FString& S );
-	void Log( const FName& Category, ELogVerbosity::Type Verbosity, const FString& S );
+	CORE_API void Log( const TCHAR* S );
+	CORE_API void Log( ELogVerbosity::Type Verbosity, const TCHAR* S );
+	CORE_API void Log( const FName& Category, ELogVerbosity::Type Verbosity, const TCHAR* Str );
+	CORE_API void Log( const FString& S );
+	CORE_API void Log( const FText& S );
+	CORE_API void Log( ELogVerbosity::Type Verbosity, const FString& S );
+	CORE_API void Log( const FName& Category, ELogVerbosity::Type Verbosity, const FString& S );
 
 private:
-	void VARARGS LogfImpl(const TCHAR* Fmt, ...);
-	void VARARGS LogfImpl(ELogVerbosity::Type Verbosity, const TCHAR* Fmt, ...);
-	void VARARGS CategorizedLogfImpl(const FName& Category, ELogVerbosity::Type Verbosity, const TCHAR* Fmt, ...);
+	CORE_API void VARARGS LogfImpl(const TCHAR* Fmt, ...);
+	CORE_API void VARARGS LogfImpl(ELogVerbosity::Type Verbosity, const TCHAR* Fmt, ...);
+	CORE_API void VARARGS CategorizedLogfImpl(const FName& Category, ELogVerbosity::Type Verbosity, const TCHAR* Fmt, ...);
 
 public:
 	template <typename FmtType>
diff --git a/Engine/Source/Runtime/Core/Public/Misc/OutputDeviceAnsiError.h b/Engine/Source/Runtime/Core/Public/Misc/OutputDeviceAnsiError.h
--- a/Engine/Source/Runtime/Core/Public/Misc/OutputDeviceAnsiError.h
+++ b/Engine/Source/Runtime/Core/Public/Misc/OutputDeviceAnsiError.h
@@ -7,11 +7,11 @@
 #include "Misc/OutputDeviceError.h"
 #include "UObject/NameTypes.h"
 
-class CORE_API FOutputDeviceAnsiError : public FOutputDeviceError
+class FOutputDeviceAnsiError : public FOutputDeviceError
 {
 public:
 	/** Constructor, initializing member variables */
-	FOutputDeviceAnsiError();
+	CORE_API FOutputDeviceAnsiError();
 
 	/**
 	* Serializes the passed in data unless the current event is suppressed.
@@ -19,7 +19,7 @@ public:
 	* @param	Data	Text to log
 	* @param	Event	Event name used for suppression purposes
 	*/
-	virtual void Serialize(const TCHAR* Msg, ELogVerbosity::Type Verbosity, const class FName& Category) override;
+	CORE_API virtual void Serialize(const TCHAR* Msg, ELogVerbosity::Type Verbosity, const class FName& Category) override;
 
 	virtual bool CanBeUsedOnAnyThread() const override
 	{
@@ -30,7 +30,7 @@ public:
 	* Error handling function that is being called from within the system wide global
 	* error handler, e.g. using structured exception handling on the PC.
 	*/
-	void HandleError() override;
+	CORE_API void HandleError() override;
 
 private:
 
diff --git a/Engine/Source/Runtime/Core/Public/Misc/OutputDeviceArchiveWrapper.h b/Engine/Source/Runtime/Core/Public/Misc/OutputDeviceArchiveWrapper.h
--- a/Engine/Source/Runtime/Core/Public/Misc/OutputDeviceArchiveWrapper.h
+++ b/Engine/Source/Runtime/Core/Public/Misc/OutputDeviceArchiveWrapper.h
@@ -12,7 +12,7 @@ class FArchive;
 /**
 * Output device wrapping any kind of FArchive.  Note: Works in any build configuration.
 */
-class CORE_API FOutputDeviceArchiveWrapper : public FOutputDevice
+class FOutputDeviceArchiveWrapper : public FOutputDevice
 {
 public:
 	/**
@@ -27,8 +27,8 @@ public:
 	}
 
 	// FOutputDevice interface
-	virtual void Flush() override;
-	virtual void Serialize(const TCHAR* Data, ELogVerbosity::Type Verbosity, const class FName& Category) override;
+	CORE_API virtual void Flush() override;
+	CORE_API virtual void Serialize(const TCHAR* Data, ELogVerbosity::Type Verbosity, const class FName& Category) override;
 	// End of FOutputDevice interface
 
 private:
diff --git a/Engine/Source/Runtime/Core/Public/Misc/OutputDeviceConsole.h b/Engine/Source/Runtime/Core/Public/Misc/OutputDeviceConsole.h
--- a/Engine/Source/Runtime/Core/Public/Misc/OutputDeviceConsole.h
+++ b/Engine/Source/Runtime/Core/Public/Misc/OutputDeviceConsole.h
@@ -9,7 +9,7 @@
 /**
  * This class servers as the base class for console window output.
  */
-class CORE_API FOutputDeviceConsole
+class FOutputDeviceConsole
 	: public FOutputDevice
 {
 public:
diff --git a/Engine/Source/Runtime/Core/Public/Misc/OutputDeviceDebug.h b/Engine/Source/Runtime/Core/Public/Misc/OutputDeviceDebug.h
--- a/Engine/Source/Runtime/Core/Public/Misc/OutputDeviceDebug.h
+++ b/Engine/Source/Runtime/Core/Public/Misc/OutputDeviceDebug.h
@@ -9,7 +9,7 @@
 
 #include <atomic>
 
-class CORE_API FOutputDeviceDebug : public FOutputDevice
+class FOutputDeviceDebug : public FOutputDevice
 {
 public:
 	/**
@@ -18,18 +18,18 @@ public:
 	* @param	Data	Text to log
 	* @param	Event	Event name used for suppression purposes
 	*/
-	virtual void Serialize(const TCHAR* Data, ELogVerbosity::Type Verbosity, const FName& Category, double Time) override;
+	CORE_API virtual void Serialize(const TCHAR* Data, ELogVerbosity::Type Verbosity, const FName& Category, double Time) override;
 
-	virtual void Serialize(const TCHAR* Data, ELogVerbosity::Type Verbosity, const FName& Category) override;
+	CORE_API virtual void Serialize(const TCHAR* Data, ELogVerbosity::Type Verbosity, const FName& Category) override;
 
-	virtual void SerializeRecord(const UE::FLogRecord& Record) override;
+	CORE_API virtual void SerializeRecord(const UE::FLogRecord& Record) override;
 
 	virtual bool CanBeUsedOnAnyThread() const override
 	{
 		return true;
 	}
 
-	virtual bool CanBeUsedOnMultipleThreads() const override;
+	CORE_API virtual bool CanBeUsedOnMultipleThreads() const override;
 
 	virtual bool CanBeUsedOnPanicThread() const override
 	{
@@ -37,8 +37,8 @@ public:
 	}
 
 private:
-	void ConditionalTickAsync(double Time);
-	void TickAsync();
+	CORE_API void ConditionalTickAsync(double Time);
+	CORE_API void TickAsync();
 
 	std::atomic<double> LastTickTime = 0.0;
 	std::atomic<bool> bSerializeAsJson = false;
diff --git a/Engine/Source/Runtime/Core/Public/Misc/OutputDeviceError.h b/Engine/Source/Runtime/Core/Public/Misc/OutputDeviceError.h
--- a/Engine/Source/Runtime/Core/Public/Misc/OutputDeviceError.h
+++ b/Engine/Source/Runtime/Core/Public/Misc/OutputDeviceError.h
@@ -6,7 +6,7 @@
 #include "Misc/OutputDevice.h"
 
 // Error device.
-class CORE_API FOutputDeviceError : public FOutputDevice
+class FOutputDeviceError : public FOutputDevice
 {
 public:
 	virtual bool CanBeUsedOnPanicThread() const override
diff --git a/Engine/Source/Runtime/Core/Public/Misc/OutputDeviceFile.h b/Engine/Source/Runtime/Core/Public/Misc/OutputDeviceFile.h
--- a/Engine/Source/Runtime/Core/Public/Misc/OutputDeviceFile.h
+++ b/Engine/Source/Runtime/Core/Public/Misc/OutputDeviceFile.h
@@ -27,7 +27,7 @@ class FRunnableThread;
 * Provides a thread-safe serialization interface with a background thread doing the actual writes.
 * [] tags identify which thread owns a variable or function
 */
-class CORE_API FAsyncWriter : public FRunnable, public FSingleThreadRunnable, public FArchive
+class FAsyncWriter : public FRunnable, public FSingleThreadRunnable, public FArchive
 {
 	enum EConstants
 	{
@@ -60,13 +60,13 @@ class CORE_API FAsyncWriter : public FRunnable, public FSingleThreadRunnable, pu
 	double LastArchiveFlushTime;
 
 	/** [WRITER THREAD] Flushes the archive and reset the flush timer. */
-	void FlushArchiveAndResetTimer();
+	CORE_API void FlushArchiveAndResetTimer();
 
 	/** [WRITER THREAD] Serialize the contents of the ring buffer to disk */
-	void SerializeBufferToArchive();
+	CORE_API void SerializeBufferToArchive();
 
 	/** [CLIENT THREAD] Flush the memory buffer (doesn't force the archive to flush). Can only be used from inside of BufferPosCritical lock. */
-	void FlushBuffer();
+	CORE_API void FlushBuffer();
 
 public:
 
@@ -76,27 +76,27 @@ public:
 		Sequential, // Uses a global sequential number to append to the threadname: FAsyncWriter_1
 	};
 
-	FAsyncWriter(FArchive& InAr, FAsyncWriter::EThreadNameOption NameOption=FAsyncWriter::EThreadNameOption::FileName);
+	CORE_API FAsyncWriter(FArchive& InAr, FAsyncWriter::EThreadNameOption NameOption=FAsyncWriter::EThreadNameOption::FileName);
 
-	virtual ~FAsyncWriter();
+	CORE_API virtual ~FAsyncWriter();
 
 	/** [CLIENT THREAD] Serialize data to buffer that will later be saved to disk by the async thread */
-	virtual void Serialize(void* InData, int64 Length) override;
+	CORE_API virtual void Serialize(void* InData, int64 Length) override;
 
 	/** Flush all buffers to disk */
-	void Flush();
+	CORE_API void Flush();
 
 	//~ Begin FRunnable Interface.
-	virtual bool Init();
-	virtual uint32 Run();
-	virtual void Stop();
+	CORE_API virtual bool Init();
+	CORE_API virtual uint32 Run();
+	CORE_API virtual void Stop();
 	virtual FSingleThreadRunnable* GetSingleThreadInterface() override { return this; }
 	//~ End FRunnable Interface
 
 protected:
 	//~ Begin FSingleThreadRunnable Interface.
 	/** [CLIENT THREAD] A substitute for Run() for when threading is disabled. */
-	virtual void Tick() override;
+	CORE_API virtual void Tick() override;
 	//~ End FSingleThreadRunnable Interface
 
 };
@@ -110,7 +110,7 @@ enum class EByteOrderMark : int8
 /**
 * File output device (Note: Only works if ALLOW_LOG_FILE && !NO_LOGGING is true, otherwise Serialize does nothing).
 */
-class CORE_API FOutputDeviceFile : public FOutputDevice
+class FOutputDeviceFile : public FOutputDevice
 {
 public:
 	/**
@@ -123,16 +123,16 @@ public:
 	* @param bCreateWriterLazily If true, delay the creation of the file until something needs to be written, otherwise, open it immediatedly.
 	* @param FileOpenedCallback If bound, invoked when the output file is successfully opened, passing the actual filename.
 	*/
-	FOutputDeviceFile(const TCHAR* InFilename = nullptr, bool bDisableBackup = false, bool bAppendIfExists = false, bool bCreateWriterLazily = true, TFunction<void(const TCHAR*)> FileOpenedCallback = TFunction<void(const TCHAR*)>());
+	CORE_API FOutputDeviceFile(const TCHAR* InFilename = nullptr, bool bDisableBackup = false, bool bAppendIfExists = false, bool bCreateWriterLazily = true, TFunction<void(const TCHAR*)> FileOpenedCallback = TFunction<void(const TCHAR*)>());
 
 	/**
 	* Destructor to perform teardown
 	*
 	*/
-	~FOutputDeviceFile();
+	CORE_API ~FOutputDeviceFile();
 
 	/** Sets the filename that the output device writes to.  If the output device was already writing to a file, closes that file. */
-	void SetFilename(const TCHAR* InFilename);
+	CORE_API void SetFilename(const TCHAR* InFilename);
 
 	//~ Begin FOutputDevice Interface.
 	/**
@@ -140,16 +140,16 @@ public:
 	* as we have to call "delete" which cannot be done for static/ global
 	* objects.
 	*/
-	void TearDown() override;
+	CORE_API void TearDown() override;
 
 	/**
 	* Flush the write cache so the file isn't truncated in case we crash right
 	* after calling this function.
 	*/
-	void Flush() override;
+	CORE_API void Flush() override;
 
-	virtual void Serialize(const TCHAR* Data, ELogVerbosity::Type Verbosity, const class FName& Category, const double Time) override;
-	virtual void Serialize(const TCHAR* Data, ELogVerbosity::Type Verbosity, const class FName& Category) override;
+	CORE_API virtual void Serialize(const TCHAR* Data, ELogVerbosity::Type Verbosity, const class FName& Category, const double Time) override;
+	CORE_API virtual void Serialize(const TCHAR* Data, ELogVerbosity::Type Verbosity, const class FName& Category) override;
 	virtual bool CanBeUsedOnAnyThread() const override
 	{
 		return true;
@@ -161,18 +161,18 @@ public:
 	//~ End FOutputDevice Interface.
 
 	/** Creates a backup copy of a log file if it already exists */
-	static void CreateBackupCopy(const TCHAR* Filename);
+	static CORE_API void CreateBackupCopy(const TCHAR* Filename);
 
 	/** Checks if the filename represents a backup copy of a log file */
-	static bool IsBackupCopy(const TCHAR* Filename);
+	static CORE_API bool IsBackupCopy(const TCHAR* Filename);
 
 	/** Add a category name to our inclusion filter. As soon as one inclusion exists, all others will be ignored */
-	void IncludeCategory(const class FName& InCategoryName);
+	CORE_API void IncludeCategory(const class FName& InCategoryName);
 
 	/** Returns the filename associated with this output device */
 	const TCHAR* GetFilename() const { return Filename; }
 
-	bool IsOpened() const;
+	CORE_API bool IsOpened() const;
 
 private:
 
@@ -195,10 +195,10 @@ private:
 	/** If true, existing files will not be backed up */
 	bool		bDisableBackup;
 
-	void WriteRaw(const TCHAR* C);
+	CORE_API void WriteRaw(const TCHAR* C);
 
 	/** Creates the async writer and its archive. Returns true if successful.  */
-	bool CreateWriter(uint32 MaxAttempts = 32);
+	CORE_API bool CreateWriter(uint32 MaxAttempts = 32);
 
-	void WriteByteOrderMarkToArchive(EByteOrderMark ByteOrderMark);
-};
\ No newline at end of file
+	CORE_API void WriteByteOrderMarkToArchive(EByteOrderMark ByteOrderMark);
+};
diff --git a/Engine/Source/Runtime/Core/Public/Misc/OutputDeviceHelper.h b/Engine/Source/Runtime/Core/Public/Misc/OutputDeviceHelper.h
--- a/Engine/Source/Runtime/Core/Public/Misc/OutputDeviceHelper.h
+++ b/Engine/Source/Runtime/Core/Public/Misc/OutputDeviceHelper.h
@@ -16,11 +16,11 @@ class FName;
 class FString;
 
 /** Helper functions used by FOutputDevice derived classes **/
-struct CORE_API FOutputDeviceHelper
+struct FOutputDeviceHelper
 {
 	/** Append a formatted log line to the string builder. */
-	static void AppendFormatLogLine(FWideStringBuilderBase& Output, ELogVerbosity::Type Verbosity, const FName& Category, const TCHAR* Message = nullptr, ELogTimes::Type LogTime = ELogTimes::None, double Time = -1.0, int32* OutCategoryIndex = nullptr);
-	static void AppendFormatLogLine(FUtf8StringBuilderBase& Output, ELogVerbosity::Type Verbosity, const FName& Category, const TCHAR* Message = nullptr, ELogTimes::Type LogTime = ELogTimes::None, double Time = -1.0, int32* OutCategoryIndex = nullptr);
+	static CORE_API void AppendFormatLogLine(FWideStringBuilderBase& Output, ELogVerbosity::Type Verbosity, const FName& Category, const TCHAR* Message = nullptr, ELogTimes::Type LogTime = ELogTimes::None, double Time = -1.0, int32* OutCategoryIndex = nullptr);
+	static CORE_API void AppendFormatLogLine(FUtf8StringBuilderBase& Output, ELogVerbosity::Type Verbosity, const FName& Category, const TCHAR* Message = nullptr, ELogTimes::Type LogTime = ELogTimes::None, double Time = -1.0, int32* OutCategoryIndex = nullptr);
 
 	/**
 	 * Formats a log line with date, time, category and verbosity prefix
@@ -32,7 +32,7 @@ struct CORE_API FOutputDeviceHelper
 	 * @param OutCategoryIndex (if non-null) The index of the category within the return string is written here, or INDEX_NONE if the category is suppressed
 	 * @returns Formatted log line
 	 */
-	static FString FormatLogLine(ELogVerbosity::Type Verbosity, const FName& Category, const TCHAR* Message = nullptr, ELogTimes::Type LogTime = ELogTimes::None, double Time = -1.0, int32* OutCategoryIndex = nullptr);
+	static CORE_API FString FormatLogLine(ELogVerbosity::Type Verbosity, const FName& Category, const TCHAR* Message = nullptr, ELogTimes::Type LogTime = ELogTimes::None, double Time = -1.0, int32* OutCategoryIndex = nullptr);
 
 	/**
 	 * Formats, casts to ANSI char and serializes a message to archive. Optimized for small number of allocations and Serialize calls
@@ -44,5 +44,5 @@ struct CORE_API FOutputDeviceHelper
 	 * @param bSuppressEventTag True if the message date/time prefix should be suppressed
 	 * @param bAutoEmitLineTerminator True if the message should be automatically appended with a line terminator
 	 **/
-	static void FormatCastAndSerializeLine(FArchive& Output, const TCHAR* Message, ELogVerbosity::Type Verbosity, const FName& Category, double Time, bool bSuppressEventTag, bool bAutoEmitLineTerminator);
+	static CORE_API void FormatCastAndSerializeLine(FArchive& Output, const TCHAR* Message, ELogVerbosity::Type Verbosity, const FName& Category, double Time, bool bSuppressEventTag, bool bAutoEmitLineTerminator);
 };
