[Core]

diff --git a/Engine/Source/Runtime/Core/Public/Misc/OutputDeviceMemory.h b/Engine/Source/Runtime/Core/Public/Misc/OutputDeviceMemory.h
--- a/Engine/Source/Runtime/Core/Public/Misc/OutputDeviceMemory.h
+++ b/Engine/Source/Runtime/Core/Public/Misc/OutputDeviceMemory.h
@@ -13,7 +13,7 @@
 /**
  * Memory output device. Logs only into pre-allocated memory buffer.
  */
-class CORE_API FOutputDeviceMemory : public FOutputDevice
+class FOutputDeviceMemory : public FOutputDevice
 {
 public:
 	/** 
@@ -22,10 +22,10 @@ public:
 	 * @param InPreserveSize	Bytes of the rung buffer not to overwrite (startup info etc)
 	 * @param InBufferSize		Maximum size of the memory ring buffer
 	 */
-	FOutputDeviceMemory(int32 InPreserveSize = 256 * 1024, int32 InBufferSize = 2048 * 1024, bool bInSuppressEventTag = false);
+	CORE_API FOutputDeviceMemory(int32 InPreserveSize = 256 * 1024, int32 InBufferSize = 2048 * 1024, bool bInSuppressEventTag = false);
 
 	/** Dumps the contents of the buffer to an archive */
-	virtual void Dump(FArchive& Ar) override;
+	CORE_API virtual void Dump(FArchive& Ar) override;
 
 	//~ Begin FOutputDevice Interface.
 	/**
@@ -33,16 +33,16 @@ public:
 	 * as we have to call "delete" which cannot be done for static/ global
 	 * objects.
 	 */
-	virtual void TearDown() override;
+	CORE_API virtual void TearDown() override;
 
 	/**
 	 * Flush the write cache so the file isn't truncated in case we crash right
 	 * after calling this function.
 	 */
-	virtual void Flush() override;
+	CORE_API virtual void Flush() override;
 
-	virtual void Serialize( const TCHAR* Data, ELogVerbosity::Type Verbosity, const class FName& Category, const double Time ) override;
-	virtual void Serialize( const TCHAR* Data, ELogVerbosity::Type Verbosity, const class FName& Category) override;
+	CORE_API virtual void Serialize( const TCHAR* Data, ELogVerbosity::Type Verbosity, const class FName& Category, const double Time ) override;
+	CORE_API virtual void Serialize( const TCHAR* Data, ELogVerbosity::Type Verbosity, const class FName& Category) override;
 	virtual bool CanBeUsedOnAnyThread() const override
 	{
 		return true;
@@ -60,7 +60,7 @@ public:
 private:
 	
 	/** Serialize cast data to the actual memory buffer */
-	void SerializeToBuffer(ANSICHAR* Data, int32 Length);
+	CORE_API void SerializeToBuffer(ANSICHAR* Data, int32 Length);
 
 
 	class FOutputDeviceMemoryProxyArchive : public FArchive
diff --git a/Engine/Source/Runtime/Core/Public/Misc/OutputDeviceNull.h b/Engine/Source/Runtime/Core/Public/Misc/OutputDeviceNull.h
--- a/Engine/Source/Runtime/Core/Public/Misc/OutputDeviceNull.h
+++ b/Engine/Source/Runtime/Core/Public/Misc/OutputDeviceNull.h
@@ -7,7 +7,7 @@
 #include "Misc/OutputDevice.h"
 
 // Null output device.
-class CORE_API FOutputDeviceNull : public FOutputDevice
+class FOutputDeviceNull : public FOutputDevice
 {
 public:
 	/**
@@ -16,6 +16,6 @@ public:
 	* @param	Data	unused
 	* @param	Event	unused
 	*/
-	virtual void Serialize(const TCHAR* V, ELogVerbosity::Type Verbosity, const class FName& Category) override;
+	CORE_API virtual void Serialize(const TCHAR* V, ELogVerbosity::Type Verbosity, const class FName& Category) override;
 };
 
diff --git a/Engine/Source/Runtime/Core/Public/Misc/OutputDeviceRedirector.h b/Engine/Source/Runtime/Core/Public/Misc/OutputDeviceRedirector.h
--- a/Engine/Source/Runtime/Core/Public/Misc/OutputDeviceRedirector.h
+++ b/Engine/Source/Runtime/Core/Public/Misc/OutputDeviceRedirector.h
@@ -19,7 +19,7 @@ FOutputDeviceRedirector.
 namespace UE::Private { struct FOutputDeviceRedirectorState; }
 
 /** The type of lines buffered by secondary threads. */
-struct CORE_API FBufferedLine
+struct FBufferedLine
 {
 	TUniquePtr<TCHAR[]> Data;
 	const FLazyName Category;
@@ -27,7 +27,7 @@ struct CORE_API FBufferedLine
 	const ELogVerbosity::Type Verbosity;
 	bool bExternalAllocation;
 
-	explicit FBufferedLine(const TCHAR* InData, const FName& InCategory, ELogVerbosity::Type InVerbosity, const double InTime = -1);
+	CORE_API explicit FBufferedLine(const TCHAR* InData, const FName& InCategory, ELogVerbosity::Type InVerbosity, const double InTime = -1);
 };
 
 enum class EOutputDeviceRedirectorFlushOptions : uint32
@@ -48,28 +48,28 @@ ENUM_CLASS_FLAGS(EOutputDeviceRedirectorFlushOptions);
 /**
  * Class used for output redirection to allow logs to show in multiple output devices.
  */
-class CORE_API FOutputDeviceRedirector final : public FOutputDevice
+class FOutputDeviceRedirector final : public FOutputDevice
 {
 public:
 	/** Initialization constructor. */
-	FOutputDeviceRedirector();
+	CORE_API FOutputDeviceRedirector();
 
 	/** Get the GLog singleton. */
-	static FOutputDeviceRedirector* Get();
+	static CORE_API FOutputDeviceRedirector* Get();
 
 	/**
 	 * Adds an output device to the chain of redirections.
 	 *
 	 * @param OutputDevice   Output device to add.
 	 */
-	void AddOutputDevice(FOutputDevice* OutputDevice);
+	CORE_API void AddOutputDevice(FOutputDevice* OutputDevice);
 
 	/**
 	 * Removes an output device from the chain of redirections.
 	 *
 	 * @param OutputDevice   Output device to remove.
 	 */
-	void RemoveOutputDevice(FOutputDevice* OutputDevice);
+	CORE_API void RemoveOutputDevice(FOutputDevice* OutputDevice);
 
 	/**
 	 * Returns whether an output device is in the list of redirections.
@@ -77,23 +77,23 @@ public:
 	 * @param OutputDevice   Output device to check the list against.
 	 * @return true if messages are currently redirected to the the passed in output device, false otherwise.
 	 */
-	bool IsRedirectingTo(FOutputDevice* OutputDevice);
+	CORE_API bool IsRedirectingTo(FOutputDevice* OutputDevice);
 
 	/** Flushes lines buffered by secondary threads. */
-	void FlushThreadedLogs(EOutputDeviceRedirectorFlushOptions Options = EOutputDeviceRedirectorFlushOptions::None);
+	CORE_API void FlushThreadedLogs(EOutputDeviceRedirectorFlushOptions Options = EOutputDeviceRedirectorFlushOptions::None);
 
 	/**
 	 * Serializes the current backlog to the specified output device.
 	 * @param OutputDevice   Output device that will receive the current backlog.
 	 */
-	void SerializeBacklog(FOutputDevice* OutputDevice);
+	CORE_API void SerializeBacklog(FOutputDevice* OutputDevice);
 
 	/**
 	 * Enables or disables the backlog.
 	 *
 	 * @param bEnable   Starts saving a backlog if true, disables and discards any backlog if false.
 	 */
-	void EnableBacklog(bool bEnable);
+	CORE_API void EnableBacklog(bool bEnable);
 
 	/**
 	 * Sets the current thread to be the thread that redirects logs to buffered output devices.
@@ -102,7 +102,7 @@ public:
 	 * responsible for flushing buffered logs from secondary threads. Logs from secondary threads
 	 * will not be redirected unless the current thread periodically flushes threaded logs.
 	 */
-	void SetCurrentThreadAsPrimaryThread();
+	CORE_API void SetCurrentThreadAsPrimaryThread();
 
 	/**
 	 * Starts a dedicated primary thread that redirects logs to buffered output devices.
@@ -111,14 +111,14 @@ public:
 	 *
 	 * @return true if a dedicated primary logging thread is running, false otherwise.
 	 */
-	bool TryStartDedicatedPrimaryThread();
+	CORE_API bool TryStartDedicatedPrimaryThread();
 
 	/**
 	 * Serializes the log record via all current output devices.
 	 *
 	 * The format string pointed to by the record must remain valid indefinitely.
 	 */
-	void SerializeRecord(const UE::FLogRecord& Record) final;
+	CORE_API void SerializeRecord(const UE::FLogRecord& Record) final;
 
 	/**
 	 * Serializes the passed in data via all current output devices.
@@ -126,7 +126,7 @@ public:
 	 * @param Data   Text to log.
 	 * @param Event  Event name used for suppression purposes.
 	 */
-	void Serialize(const TCHAR* Data, ELogVerbosity::Type Verbosity, const FName& Category, const double Time) final;
+	CORE_API void Serialize(const TCHAR* Data, ELogVerbosity::Type Verbosity, const FName& Category, const double Time) final;
 
 	/**
 	 * Serializes the passed in data via all current output devices.
@@ -134,14 +134,14 @@ public:
 	 * @param Data   Text to log.
 	 * @param Event  Event name used for suppression purposes.
 	 */
-	void Serialize(const TCHAR* Data, ELogVerbosity::Type Verbosity, const FName& Category) final;
+	CORE_API void Serialize(const TCHAR* Data, ELogVerbosity::Type Verbosity, const FName& Category) final;
 	
 	/** Same as Serialize(). */
-	void RedirectLog(const FName& Category, ELogVerbosity::Type Verbosity, const TCHAR* Data);
-	void RedirectLog(const FLazyName& Category, ELogVerbosity::Type Verbosity, const TCHAR* Data);
+	CORE_API void RedirectLog(const FName& Category, ELogVerbosity::Type Verbosity, const TCHAR* Data);
+	CORE_API void RedirectLog(const FLazyName& Category, ELogVerbosity::Type Verbosity, const TCHAR* Data);
 
 	/** Passes on the flush request to all current output devices. */
-	void Flush() final;
+	CORE_API void Flush() final;
 
 	/**
 	 * Attempts to set the calling thread as the panic thread and enable panic mode.
@@ -152,19 +152,19 @@ public:
 	 * Flushes buffered logs to panic-safe output devices.
 	 * Flushes panic-safe output devices.
 	 */
-	void Panic();
+	CORE_API void Panic();
 
 	/**
 	 * Closes output device and cleans up.
 	 *
 	 * This can't happen in the destructor as we might have to call "delete" which cannot be done for static/global objects.
 	 */
-	void TearDown() final;
+	CORE_API void TearDown() final;
 
 	/**
 	 * Determine if the backlog is enabled.
 	 */
-	bool IsBacklogEnabled() const;
+	CORE_API bool IsBacklogEnabled() const;
 
 private:
 	TPimplPtr<UE::Private::FOutputDeviceRedirectorState> State;
diff --git a/Engine/Source/Runtime/Core/Public/Misc/Parse.h b/Engine/Source/Runtime/Core/Public/Misc/Parse.h
--- a/Engine/Source/Runtime/Core/Public/Misc/Parse.h
+++ b/Engine/Source/Runtime/Core/Public/Misc/Parse.h
@@ -17,20 +17,20 @@ class FText;
 /*-----------------------------------------------------------------------------
 	Parsing functions.
 -----------------------------------------------------------------------------*/
-struct CORE_API FParse
+struct FParse
 {
 	/**
 	 * Sees if Stream starts with the named command.  If it does,
 	 * skips through the command and blanks past it.  Returns true of match.
 	 * @param bParseMightTriggerExecution true: Caller guarantees this is only part of parsing and no execution happens without further parsing (good for "DumpConsoleCommands").
 	 */
-	static bool Command( const TCHAR** Stream, const TCHAR* Match, bool bParseMightTriggerExecution = true );
+	static CORE_API bool Command( const TCHAR** Stream, const TCHAR* Match, bool bParseMightTriggerExecution = true );
 	/** Parses a name. */
-	static bool Value( const TCHAR* Stream, const TCHAR* Match, FName& Name );
+	static CORE_API bool Value( const TCHAR* Stream, const TCHAR* Match, FName& Name );
 	/** Parses a uint32. */
-	static bool Value( const TCHAR* Stream, const TCHAR* Match, uint32& Value );
+	static CORE_API bool Value( const TCHAR* Stream, const TCHAR* Match, uint32& Value );
 	/** Parses a globally unique identifier. */
-	static bool Value( const TCHAR* Stream, const TCHAR* Match, struct FGuid& Guid );
+	static CORE_API bool Value( const TCHAR* Stream, const TCHAR* Match, struct FGuid& Guid );
 	/** Parses a string from a text string. 
 	 * @param Stream, the string you want to extract the value from.
 	 * @param Match, the identifier for the value in the stream.
@@ -42,21 +42,21 @@ struct CORE_API FParse
 				of the end of the value within Stream. This permits consuming of stream in a loop where Match may
 				occur multiple times.
 	*/
-	static bool Value( const TCHAR* Stream, const TCHAR* Match, TCHAR* Value, int32 MaxLen, bool bShouldStopOnSeparator=true, const TCHAR** OptStreamGotTo = nullptr);
+	static CORE_API bool Value( const TCHAR* Stream, const TCHAR* Match, TCHAR* Value, int32 MaxLen, bool bShouldStopOnSeparator=true, const TCHAR** OptStreamGotTo = nullptr);
 	/** Parses a byte. */
-	static bool Value( const TCHAR* Stream, const TCHAR* Match, uint8& Value );
+	static CORE_API bool Value( const TCHAR* Stream, const TCHAR* Match, uint8& Value );
 	/** Parses a signed byte. */
-	static bool Value( const TCHAR* Stream, const TCHAR* Match, int8& Value );
+	static CORE_API bool Value( const TCHAR* Stream, const TCHAR* Match, int8& Value );
 	/** Parses a uint16. */
-	static bool Value( const TCHAR* Stream, const TCHAR* Match, uint16& Value );
+	static CORE_API bool Value( const TCHAR* Stream, const TCHAR* Match, uint16& Value );
 	/** Parses a signed word. */
-	static bool Value( const TCHAR* Stream, const TCHAR* Match, int16& Value );
+	static CORE_API bool Value( const TCHAR* Stream, const TCHAR* Match, int16& Value );
 	/** Parses a floating-point value. */
-	static bool Value( const TCHAR* Stream, const TCHAR* Match, float& Value );
+	static CORE_API bool Value( const TCHAR* Stream, const TCHAR* Match, float& Value );
 	/** Parses a double precision floating-point value. */
-	static bool Value(const TCHAR* Stream, const TCHAR* Match, double& Value);
+	static CORE_API bool Value(const TCHAR* Stream, const TCHAR* Match, double& Value);
 	/** Parses a signed double word. */
-	static bool Value( const TCHAR* Stream, const TCHAR* Match, int32& Value );
+	static CORE_API bool Value( const TCHAR* Stream, const TCHAR* Match, int32& Value );
 	/** Parses a string. 
 	 * @param Stream, the string you want to extract the value from.
 	 * @param Match, the identifier for the value in the stream.
@@ -67,43 +67,43 @@ struct CORE_API FParse
 				of the end of the value within Stream. This permits consuming of stream in a loop where Match may
 				occur multiple times.
 	*/
-	static bool Value( const TCHAR* Stream, const TCHAR* Match, FString& Value, bool bShouldStopOnSeparator =true, const TCHAR** OptStreamGotTo = nullptr);
+	static CORE_API bool Value( const TCHAR* Stream, const TCHAR* Match, FString& Value, bool bShouldStopOnSeparator =true, const TCHAR** OptStreamGotTo = nullptr);
 	/** Parses an FText. */
-	static bool Value( const TCHAR* Stream, const TCHAR* Match, FText& Value, const TCHAR* Namespace = NULL );
+	static CORE_API bool Value( const TCHAR* Stream, const TCHAR* Match, FText& Value, const TCHAR* Namespace = NULL );
 	/** Parses a quadword. */
-	static bool Value( const TCHAR* Stream, const TCHAR* Match, uint64& Value );
+	static CORE_API bool Value( const TCHAR* Stream, const TCHAR* Match, uint64& Value );
 	/** Parses a signed quadword. */
-	static bool Value( const TCHAR* Stream, const TCHAR* Match, int64& Value );
+	static CORE_API bool Value( const TCHAR* Stream, const TCHAR* Match, int64& Value );
 	/** Parses a boolean value. */
-	static bool Bool( const TCHAR* Stream, const TCHAR* Match, bool& OnOff );
+	static CORE_API bool Bool( const TCHAR* Stream, const TCHAR* Match, bool& OnOff );
 	/** Get a line of Stream (everything up to, but not including, CR/LF. Returns 0 if ok, nonzero if at end of stream and returned 0-length string. */
-	static bool Line( const TCHAR** Stream, TCHAR* Result, int32 MaxLen, bool Exact= false );
+	static CORE_API bool Line( const TCHAR** Stream, TCHAR* Result, int32 MaxLen, bool Exact= false );
 	/** Get a line of Stream (everything up to, but not including, CR/LF. Returns 0 if ok, nonzero if at end of stream and returned 0-length string. */
-	static bool Line( const TCHAR** Stream, FString& Result, bool Exact = false );
+	static CORE_API bool Line( const TCHAR** Stream, FString& Result, bool Exact = false );
 	/** Get a line of Stream (everything up to, but not including, CR/LF. Returns 0 if ok, nonzero if at end of stream and returned 0-length string. */
-	static bool Line( const TCHAR** Stream, FStringView& Result, bool Exact = false );
+	static CORE_API bool Line( const TCHAR** Stream, FStringView& Result, bool Exact = false );
 	/** Get a line of Stream, with support for extending beyond that line with certain characters, e.g. {} and \
 	 * the out character array will not include the ignored endlines
 	 */
-	static bool LineExtended(const TCHAR** Stream, FString& Result, int32& LinesConsumed, bool Exact = 0);
-	static bool LineExtended(const TCHAR** Stream, FStringBuilderBase& Result, int32& LinesConsumed, bool Exact = 0);
+	static CORE_API bool LineExtended(const TCHAR** Stream, FString& Result, int32& LinesConsumed, bool Exact = 0);
+	static CORE_API bool LineExtended(const TCHAR** Stream, FStringBuilderBase& Result, int32& LinesConsumed, bool Exact = 0);
 	/** Grabs the next space-delimited string from the input stream. If quoted, gets entire quoted string. */
-	static bool Token( const TCHAR*& Str, TCHAR* Result, int32 MaxLen, bool UseEscape );
+	static CORE_API bool Token( const TCHAR*& Str, TCHAR* Result, int32 MaxLen, bool UseEscape );
 	/** Grabs the next space-delimited string from the input stream. If quoted, gets entire quoted string. */
-	static bool Token( const TCHAR*& Str, FString& Arg, bool UseEscape );
+	static CORE_API bool Token( const TCHAR*& Str, FString& Arg, bool UseEscape );
 	/** Grabs the next alpha-numeric space-delimited token from the input stream. */
-	static bool AlnumToken(const TCHAR*& Str, FString& Arg);
+	static CORE_API bool AlnumToken(const TCHAR*& Str, FString& Arg);
 	/** Grabs the next space-delimited string from the input stream. If quoted, gets entire quoted string. */
-	static FString Token( const TCHAR*& Str, bool UseEscape );
+	static CORE_API FString Token( const TCHAR*& Str, bool UseEscape );
 	/** Get next command.  Skips past comments and cr's. */
-	static void Next( const TCHAR** Stream );
+	static CORE_API void Next( const TCHAR** Stream );
 	/** Checks if a command-line parameter exists in the stream. */
-	static bool Param( const TCHAR* Stream, const TCHAR* Param );
+	static CORE_API bool Param( const TCHAR* Stream, const TCHAR* Param );
 	/** Parse an Text token. */
-	static bool Text( const TCHAR* Stream, FText& Value, const TCHAR* Namespace = nullptr );
+	static CORE_API bool Text( const TCHAR* Stream, FText& Value, const TCHAR* Namespace = nullptr );
 	/** Parse a quoted string token. */
-	static bool QuotedString( const TCHAR* Stream, FString& Value, int32* OutNumCharsRead = nullptr );
-	static bool QuotedString( const TCHAR* Stream, FStringBuilderBase& Value, int32* OutNumCharsRead = nullptr );
+	static CORE_API bool QuotedString( const TCHAR* Stream, FString& Value, int32* OutNumCharsRead = nullptr );
+	static CORE_API bool QuotedString( const TCHAR* Stream, FStringBuilderBase& Value, int32* OutNumCharsRead = nullptr );
 
 	//
 	// Parse a hex digit.
@@ -133,17 +133,17 @@ struct CORE_API FParse
 	}
 
 	/** Parses a hexadecimal string value. */
-	static uint32 HexNumber(const TCHAR* HexString);
-	static uint64 HexNumber64(const TCHAR* HexString);
+	static CORE_API uint32 HexNumber(const TCHAR* HexString);
+	static CORE_API uint64 HexNumber64(const TCHAR* HexString);
 
 	/** Parses a resolution in the form 1920x1080. */
-	static bool Resolution( const TCHAR* InResolution, uint32& OutX, uint32& OutY );
+	static CORE_API bool Resolution( const TCHAR* InResolution, uint32& OutX, uint32& OutY );
 
 	/** Parses a resolution in the form 1920x1080<f|w|wf>. Same as above, but also attempts to process a fullscreen/windowed flag from the end */
-	static bool Resolution(const TCHAR* InResolution, uint32& OutX, uint32& OutY, int32& OutWindowMode);
+	static CORE_API bool Resolution(const TCHAR* InResolution, uint32& OutX, uint32& OutY, int32& OutWindowMode);
 
 	/** Parses the scheme name from a URI */
-	static bool SchemeNameFromURI(const TCHAR* InURI, FString& OutSchemeName);
+	static CORE_API bool SchemeNameFromURI(const TCHAR* InURI, FString& OutSchemeName);
 
 
 	//
@@ -189,7 +189,7 @@ struct CORE_API FParse
 	*   [] = set of characters
 	*   () = treat enclosed expressions as 1 for purpose of other operators
 	*/
-	static FGrammarBasedParseResult GrammarBasedCLIParse(const TCHAR* Stream, TFunctionRef<void(FStringView, FStringView)> OnCommandCallback, EGrammarBasedParseFlags Flags = EGrammarBasedParseFlags::AllowQuotedCommands);
+	static CORE_API FGrammarBasedParseResult GrammarBasedCLIParse(const TCHAR* Stream, TFunctionRef<void(FStringView, FStringView)> OnCommandCallback, EGrammarBasedParseFlags Flags = EGrammarBasedParseFlags::AllowQuotedCommands);
 };
 
 ENUM_CLASS_FLAGS(FParse::EGrammarBasedParseFlags);
diff --git a/Engine/Source/Runtime/Core/Public/Misc/PathViews.h b/Engine/Source/Runtime/Core/Public/Misc/PathViews.h
--- a/Engine/Source/Runtime/Core/Public/Misc/PathViews.h
+++ b/Engine/Source/Runtime/Core/Public/Misc/PathViews.h
@@ -10,7 +10,7 @@ class FString;
 template <typename CharType> class TStringBuilderBase;
 template <typename FuncType> class TFunctionRef;
 
-class CORE_API FPathViews
+class FPathViews
 {
 public:
 	/**
@@ -24,7 +24,7 @@ public:
 	 *
 	 * @return The portion of the path after the last separator.
 	 */
-	static FStringView GetCleanFilename(const FStringView& InPath);
+	static CORE_API FStringView GetCleanFilename(const FStringView& InPath);
 
 	/**
 	 * Returns the portion of the path after the last separator and before the last dot.
@@ -37,7 +37,7 @@ public:
 	 *
 	 * @return The portion of the path after the last separator and before the last dot.
 	 */
-	static FStringView GetBaseFilename(const FStringView& InPath);
+	static CORE_API FStringView GetBaseFilename(const FStringView& InPath);
 
 	/**
 	 * Returns the portion of the path before the last dot.
@@ -50,12 +50,12 @@ public:
 	 *
 	 * @return The portion of the path before the last dot.
 	 */
-	static FStringView GetBaseFilenameWithPath(const FStringView& InPath);
+	static CORE_API FStringView GetBaseFilenameWithPath(const FStringView& InPath);
 
 	/** Returns the portion of the path before the last dot and optionally after the last separator. */
 	UE_DEPRECATED(4.25, "FPathViews::GetBaseFilename(InPath, bRemovePath) has been superseded by "
 		"FPathViews::GetBaseFilename(InPath) and FPathViews::GetBaseFilenameWithPath(InPath).")
-	static FStringView GetBaseFilename(const FStringView& InPath, bool bRemovePath);
+	static CORE_API FStringView GetBaseFilename(const FStringView& InPath, bool bRemovePath);
 
 	/**
 	 * Returns the portion of the path before the last separator.
@@ -68,7 +68,7 @@ public:
 	 *
 	 * @return The portion of the path before the last separator.
 	 */
-	static FStringView GetPath(const FStringView& InPath);
+	static CORE_API FStringView GetPath(const FStringView& InPath);
 
 	/**
 	 * Returns the portion of the path after the last dot following the last separator, optionally including the dot.
@@ -87,7 +87,7 @@ public:
 	 *
 	 * @return The portion of the path after the last dot following the last separator, optionally including the dot.
 	 */
-	static FStringView GetExtension(const FStringView& InPath, bool bIncludeDot=false);
+	static CORE_API FStringView GetExtension(const FStringView& InPath, bool bIncludeDot=false);
 
 	/**
 	 * Returns the last non-empty path component.
@@ -100,7 +100,7 @@ public:
 	 *
 	 * @return The last non-empty path component.
 	 */
-	static FStringView GetPathLeaf(const FStringView& InPath);
+	static CORE_API FStringView GetPathLeaf(const FStringView& InPath);
 
 	/**
 	 * Return whether the given relative or absolute path is a leaf path - has no separators.
@@ -115,7 +115,7 @@ public:
 	 * D:/A		-> false
 	 * //A		-> false
 	 */
-	static bool IsPathLeaf(FStringView InPath);
+	static CORE_API bool IsPathLeaf(FStringView InPath);
 
 	/**
 	 * Splits InPath into individual directory components, and calls ComponentVisitor on each.
@@ -125,7 +125,7 @@ public:
 	 * "A/B/C" -> {"A", "B", "C"}
 	 * "../../A/B/C.D" -> {"..", "..", "A", "B", "C.D" }
 	 */
-	static void IterateComponents(FStringView InPath, TFunctionRef<void(FStringView)> ComponentVisitor);
+	static CORE_API void IterateComponents(FStringView InPath, TFunctionRef<void(FStringView)> ComponentVisitor);
 
 	/**
 	 * Splits a path into three parts, any of which may be empty: the path, the clean name, and the extension.
@@ -144,7 +144,7 @@ public:
 	 * @param OutName [out] Receives the name portion of the input string.
 	 * @param OutExt  [out] Receives the extension portion of the input string, excluding the dot.
 	 */
-	static void Split(const FStringView& InPath, FStringView& OutPath, FStringView& OutName, FStringView& OutExt);
+	static CORE_API void Split(const FStringView& InPath, FStringView& OutPath, FStringView& OutName, FStringView& OutExt);
 
 	/**
 	 * Appends each suffix argument to the path in the builder and ensures that there is a separator between them.
@@ -186,7 +186,7 @@ public:
 	 * @return The new file path complete with the new extension unless InPath is not valid in which
 	 * case a copy of InPath will be returned instead.
 	 */
-	static FString ChangeExtension(const FStringView& InPath, const FStringView& InNewExtension);
+	static CORE_API FString ChangeExtension(const FStringView& InPath, const FStringView& InNewExtension);
 	
 	/**
 	 * Sets the file extension of a filename.
@@ -196,10 +196,10 @@ public:
 	 *
 	 * @return The new file path complete with the new extension.
 	 */
-	static FString SetExtension(const FStringView& InPath, const FStringView& InNewExtension);
+	static CORE_API FString SetExtension(const FStringView& InPath, const FStringView& InNewExtension);
 
 	/** Return whether the given character is a path-separator character (/ or \) */
-	static bool IsSeparator(TCHAR c);
+	static CORE_API bool IsSeparator(TCHAR c);
 
 	/**
 	 * Return true if the given paths are the same path (with exceptions noted below).
@@ -222,7 +222,7 @@ public:
 	 * (d:/root/Engine/..", "d:/root") -> false
 	 * (d:/root/Engine/./Content", "d:/root/Engine/Content") -> false
 	 */
-	static bool Equals(FStringView A, FStringView B);
+	static CORE_API bool Equals(FStringView A, FStringView B);
 
 	/**
 	 * Return true if the the first path is lexicographically less than the second path (with caveats noted below).
@@ -246,7 +246,7 @@ public:
 	 * (d:/root/Engine/..", "d:/root") -> false (A string is greater than any prefix of itself)
 	 * (d:/root/Engine/./Content", "d:/root/Engine/Content") -> false
 	 */
-	static bool Less(FStringView A, FStringView B);
+	static CORE_API bool Less(FStringView A, FStringView B);
 
 	/**
 	 * Check whether Parent is a parent path of Child and report the relative path if so.
@@ -272,7 +272,7 @@ public:
 	 *
 	 * @return True if and only if Child is a child path of Parent (or is equal to it).
 	 */
-	static bool TryMakeChildPathRelativeTo(FStringView Child, FStringView Parent, FStringView& OutRelPath);
+	static CORE_API bool TryMakeChildPathRelativeTo(FStringView Child, FStringView Parent, FStringView& OutRelPath);
 
 	/**
 	 * Return whether Parent is a parent path of (or is equal to) Child.
@@ -297,37 +297,37 @@ public:
 	 *
 	 * @return True if and only if Child is a child path of Parent (or is equal to it).
 	 */
-	static bool IsParentPathOf(FStringView Parent, FStringView Child);
+	static CORE_API bool IsParentPathOf(FStringView Parent, FStringView Child);
 
 	/**
 	 * Return whether the given path is a relativepath - does not start with a separator or volume:.
 	 * Returns true for empty paths.
 	 */
-	static bool IsRelativePath(FStringView InPath);
+	static CORE_API bool IsRelativePath(FStringView InPath);
 
 	/** Convert to absolute using process BaseDir(), normalize and append. FPaths::ConvertRelativePathToFull() equivalent. */
-	static void ToAbsolutePath(FStringView InPath, FStringBuilderBase& OutPath);
+	static CORE_API void ToAbsolutePath(FStringView InPath, FStringBuilderBase& OutPath);
 	
 	/** Convert to absolute using explicit BasePath, normalize and append. FPaths::ConvertRelativePathToFull() equivalent. */
-	static void ToAbsolutePath(FStringView BasePath, FStringView InPath, FStringBuilderBase& OutPath);
+	static CORE_API void ToAbsolutePath(FStringView BasePath, FStringView InPath, FStringBuilderBase& OutPath);
 	
 	/** Convert to absolute using process BaseDir() and normalize inlined. FPaths::ConvertRelativePathToFull() equivalent. */
-	static void ToAbsolutePathInline(FStringBuilderBase& InOutPath);
+	static CORE_API void ToAbsolutePathInline(FStringBuilderBase& InOutPath);
 	
 	/** Convert to absolute using explicit BasePath and normalize inlined. FPaths::ConvertRelativePathToFull() equivalent. */
-	static void ToAbsolutePathInline(FStringView BasePath, FStringBuilderBase& InOutPath);
+	static CORE_API void ToAbsolutePathInline(FStringView BasePath, FStringBuilderBase& InOutPath);
 
 	/** Convert \\ to / and do platform-specific normalization */
-	static void NormalizeFilename(FStringBuilderBase& InOutPath);
+	static CORE_API void NormalizeFilename(FStringBuilderBase& InOutPath);
 
 	/** Normalize and remove trailing slash unless the preceding character is '/' or ':' */
-	static void NormalizeDirectoryName(FStringBuilderBase& InOutPath);
+	static CORE_API void NormalizeDirectoryName(FStringBuilderBase& InOutPath);
 
 	/** Collapses redundant paths like "/./" and "SkipDir/..". FPaths::CollapseRelativeDirectories() equivalent. */
-	static bool CollapseRelativeDirectories(FStringBuilderBase& InOutPath);
+	static CORE_API bool CollapseRelativeDirectories(FStringBuilderBase& InOutPath);
 
 	/** Removes duplicate forward slashes, e.g. "a/b//c////f.e" -> "a/b/c/f.e" */
-	static void RemoveDuplicateSlashes(FStringBuilderBase& InOutPath);
+	static CORE_API void RemoveDuplicateSlashes(FStringBuilderBase& InOutPath);
 
 	/**
 	 * Split the given absolute or relative path into its topmost directory and the relative path from that directory.
@@ -339,7 +339,7 @@ public:
 	 * @param OutFirstComponent Receives the first directory element in the path, or InPath if it is a leaf path.
 	 * @param OutRemainder Receives the relative path from OutFirstComponent to InPath, or empty if InPath is a leaf path.
 	 */
-	static void SplitFirstComponent(FStringView InPath, FStringView& OutFirstComponent, FStringView& OutRemainder);
+	static CORE_API void SplitFirstComponent(FStringView InPath, FStringView& OutFirstComponent, FStringView& OutRemainder);
 
 	/**
 	 * If AppendPath is a relative path, append it as a relative path onto InOutPath.
@@ -347,7 +347,7 @@ public:
 	 * Handles presence or absence of terminating separator in BasePath.
 	 * Does not interpret . or ..; each occurrence of these in either path will remain in the combined InOutPath.
 	 */
-	static void AppendPath(FStringBuilderBase& InOutPath, FStringView AppendPath);
+	static CORE_API void AppendPath(FStringBuilderBase& InOutPath, FStringView AppendPath);
 
 	/**
 	 * Returns the name of the mount point in a path
@@ -355,5 +355,5 @@ public:
 	 * Example: "/Classes_A/Textures" returns "A" and sets bOutHadClassesPrefix=true if bInWithoutSlashes is true
 	 *           returns "/Classes_A" otherwise and sets bOutHadClassesPrefix=false
 	 */
-	static FStringView GetMountPointNameFromPath(const FStringView InPath, bool* bOutHadClassesPrefix = nullptr, bool bInWithoutSlashes = true);
+	static CORE_API FStringView GetMountPointNameFromPath(const FStringView InPath, bool* bOutHadClassesPrefix = nullptr, bool bInWithoutSlashes = true);
 };
diff --git a/Engine/Source/Runtime/Core/Public/Misc/Paths.h b/Engine/Source/Runtime/Core/Public/Misc/Paths.h
--- a/Engine/Source/Runtime/Core/Public/Misc/Paths.h
+++ b/Engine/Source/Runtime/Core/Public/Misc/Paths.h
@@ -32,19 +32,19 @@ namespace UE::Core::Private
 /**
  * Path helpers for retrieving game dir, engine dir, etc.
  */
-class CORE_API FPaths
+class FPaths
 {
 public:
 
 	/**
 	 * Should the "saved" directory structures be rooted in the user dir or relative to the "engine/game" 
 	 */
-	static bool ShouldSaveToUserDir();
+	static CORE_API bool ShouldSaveToUserDir();
 
 	/**
 	 * Returns the directory the application was launched from (useful for commandline utilities)
 	 */
-	static FString LaunchDir();
+	static CORE_API FString LaunchDir();
 	 
 	/** 
 	 * Returns the base directory of the "core" engine that can be shared across
@@ -53,21 +53,21 @@ public:
 	 *
 	 * @return engine directory
 	 */
-	static FString EngineDir();
+	static CORE_API FString EngineDir();
 
 	/**
 	* Returns the root directory for user-specific engine files. Always writable.
 	*
 	* @return root user directory
 	*/
-	static FString EngineUserDir();
+	static CORE_API FString EngineUserDir();
 
 	/**
 	* Returns the root directory for user-specific engine files which can be shared between versions. Always writable.
 	*
 	* @return root user directory
 	*/
-	static FString EngineVersionAgnosticUserDir();
+	static CORE_API FString EngineVersionAgnosticUserDir();
 
 	/** 
 	 * Returns the content directory of the "core" engine that can be shared across
@@ -75,113 +75,113 @@ public:
 	 *
 	 * @return engine content directory
 	 */
-	static FString EngineContentDir();
+	static CORE_API FString EngineContentDir();
 
 	/**
 	 * Returns the directory the root configuration files are located.
 	 *
 	 * @return root config directory
 	 */
-	static FString EngineConfigDir();
+	static CORE_API FString EngineConfigDir();
 
 	/**
 	 * Returns the Editor Settings directory of the engine
 	 *
 	 * @return Editor Settings directory.
 	 */
-	static FString EngineEditorSettingsDir();
+	static CORE_API FString EngineEditorSettingsDir();
 
 	/**
 	 * Returns the intermediate directory of the engine
 	 *
 	 * @return content directory
 	 */
-	static FString EngineIntermediateDir();
+	static CORE_API FString EngineIntermediateDir();
 
 	/**
 	 * Returns the saved directory of the engine
 	 *
 	 * @return Saved directory.
 	 */
-	static FString EngineSavedDir();
+	static CORE_API FString EngineSavedDir();
 
 	/**
 	 * Returns the plugins directory of the engine
 	 *
 	 * @return Plugins directory.
 	 */
-	static FString EnginePluginsDir();
+	static CORE_API FString EnginePluginsDir();
 
 	/**
 	 * Returns the directory for default Editor UI Layout files of the engine
 	 * @return Directory for default Editor UI Layout files.
 	 */
-	static FString EngineDefaultLayoutDir();
+	static CORE_API FString EngineDefaultLayoutDir();
 
 	/**
 	 * Returns the directory for project Editor UI Layout files of the engine
 	 * @return Directory for project Editor UI Layout files.
 	 */
-	static FString EngineProjectLayoutDir();
+	static CORE_API FString EngineProjectLayoutDir();
 
 	/**
 	 * Returns the directory for user-generated Editor UI Layout files of the engine
 	 * @return Directory for user-generated Editor UI Layout files.
 	 */
-	static FString EngineUserLayoutDir();
+	static CORE_API FString EngineUserLayoutDir();
 
 	/** 
 	* Returns the base directory enterprise directory.
 	*
 	* @return enterprise directory
 	*/
-	static FString EnterpriseDir();
+	static CORE_API FString EnterpriseDir();
 
 	/**
 	* Returns the enterprise plugins directory
 	*
 	* @return Plugins directory.
 	*/
-	static FString EnterprisePluginsDir();
+	static CORE_API FString EnterprisePluginsDir();
 
 	/**
 	* Returns the enterprise FeaturePack directory
 	*
 	* @return FeaturePack directory.
 	*/
-	static FString EnterpriseFeaturePackDir();
+	static CORE_API FString EnterpriseFeaturePackDir();
 
 	/**
 	 * Returns the directory where engine platform extensions reside
 	 *
 	 * @return engine platform extensions directory
 	 */
-	static FString EnginePlatformExtensionsDir();
+	static CORE_API FString EnginePlatformExtensionsDir();
 
 	/**
 	 * Returns the directory where the project's platform extensions reside
 	 *
 	 * @return project platform extensions directory
 	 */
-	static FString ProjectPlatformExtensionsDir();
+	static CORE_API FString ProjectPlatformExtensionsDir();
 
 	/**
 	 * Returns platform and restricted extensions that are present and valid (for platforms, it uses FDataDrivePlatformInfo to determine valid platforms, it doesn't just use what's present)
 	 *
 	 * @return BaseDir and usable extension directories under BaseDir (either Engine or Project)
 	 */
-	static TArray<FString> GetExtensionDirs(const FString& BaseDir, const FString& SubDir=FString());
+	static CORE_API TArray<FString> GetExtensionDirs(const FString& BaseDir, const FString& SubDir=FString());
 
 	/**
 	 * Returns the root directory of the engine directory tree
 	 *
 	 * @return Root directory.
 	 */
-	static FString RootDir();
+	static CORE_API FString RootDir();
 
 #if WITH_EDITOR
 	/** Returns the special path used when mounting FeaturePaks in editor */
-	static const TCHAR* GameFeatureRootPrefix();
+	static CORE_API const TCHAR* GameFeatureRootPrefix();
 #endif
 
 	/**
@@ -192,68 +192,68 @@ public:
 	 *
 	 * @return base directory
 	 */
-	static FString ProjectDir();
+	static CORE_API FString ProjectDir();
 
 	/**
 	* Returns the root directory for user-specific game files.
 	*
 	* @return game user directory
 	*/
-	static FString ProjectUserDir();
+	static CORE_API FString ProjectUserDir();
 
 	/**
 	 * Returns the content directory of the current game by looking at FApp::GetProjectName().
 	 *
 	 * @return content directory
 	 */
-	static FString ProjectContentDir();
+	static CORE_API FString ProjectContentDir();
 
 	/**
 	* Returns the directory the root configuration files are located.
 	*
 	* @return root config directory
 	*/
-	static FString ProjectConfigDir();
+	static CORE_API FString ProjectConfigDir();
 
 	/**
 	 * Returns the saved directory of the current game by looking at FApp::GetProjectName().
 	 *
 	 * @return saved directory
 	 */
-	static const FString& ProjectSavedDir();
+	static CORE_API const FString& ProjectSavedDir();
 
 	/**
 	 * Returns the intermediate directory of the current game by looking at FApp::GetProjectName().
 	 *
 	 * @return intermediate directory
 	 */
-	static FString ProjectIntermediateDir();
+	static CORE_API FString ProjectIntermediateDir();
 
-	static FString ShaderWorkingDir();
+	static CORE_API FString ShaderWorkingDir();
 
 	/**
 	 * Returns the plugins directory of the current game by looking at FApp::GetProjectName().
 	 *
 	 * @return plugins directory
 	 */
-	static FString ProjectPluginsDir();
+	static CORE_API FString ProjectPluginsDir();
 
 	/**
 	 * Returns the mods directory of the current project by looking at FApp::GetProjectName().
 	 *
 	 * @return mods directory
 	 */
-	static FString ProjectModsDir();
+	static CORE_API FString ProjectModsDir();
 
 	/*
 	* Returns true if a writable directory for downloaded data that persists across play sessions is available
 	*/
-	static bool HasProjectPersistentDownloadDir();
+	static CORE_API bool HasProjectPersistentDownloadDir();
 
 	/*
 	* Returns the writable directory for downloaded data that persists across play sessions.
 	*/
-	static FString ProjectPersistentDownloadDir();
+	static CORE_API FString ProjectPersistentDownloadDir();
 
 	/**
 	 * Returns the directory the engine uses to look for the source leaf ini files. This
@@ -261,56 +261,56 @@ public:
 	 *
 	 * @return source config directory
 	 */
-	static FString SourceConfigDir();
+	static CORE_API FString SourceConfigDir();
 
 	/**
 	 * Returns the directory the engine saves generated config files.
 	 *
 	 * @return config directory
 	 */
-	static FString GeneratedConfigDir();
+	static CORE_API FString GeneratedConfigDir();
 
 	/**
 	 * Returns the directory the engine stores sandbox output
 	 *
 	 * @return sandbox directory
 	 */
-	static FString SandboxesDir();
+	static CORE_API FString SandboxesDir();
 
 	/**
 	 * Returns the directory the engine uses to output profiling files.
 	 *
 	 * @return log directory
 	 */
-	static FString ProfilingDir();
+	static CORE_API FString ProfilingDir();
 
 	/**
 	 * Returns the directory the engine uses to output screenshot files.
 	 *
 	 * @return screenshot directory
 	 */
-	static FString ScreenShotDir();
+	static CORE_API FString ScreenShotDir();
 
 	/**
 	 * Returns the directory the engine uses to output BugIt files.
 	 *
 	 * @return screenshot directory
 	 */
-	static FString BugItDir();
+	static CORE_API FString BugItDir();
 
 	/**
 	 * Returns the directory the engine uses to output user requested video capture files.
 	 *
 	 * @return Video capture directory
 	 */
-	static FString VideoCaptureDir();
+	static CORE_API FString VideoCaptureDir();
 
 	/**
 	 * Returns the directory the engine uses to output user requested audio capture files.
 	 *
 	 * @return Audio capture directory
 	 */
-	static FString AudioCaptureDir();
+	static CORE_API FString AudioCaptureDir();
 	
 	/**
 	 * Returns the directory the engine uses to output logs. This currently can't 
@@ -319,88 +319,88 @@ public:
 	 *
 	 * @return log directory
 	 */
-	static FString ProjectLogDir();
+	static CORE_API FString ProjectLogDir();
 
 	/** Returns the directory for automation save files */
-	static FString AutomationDir();
+	static CORE_API FString AutomationDir();
 
 	/** Returns the directory for automation save files that are meant to be deleted every run */
-	static FString AutomationTransientDir();
+	static CORE_API FString AutomationTransientDir();
 
 	/** Returns the directory for results of automation tests. May be deleted every run. */
-	static FString AutomationReportsDir();
+	static CORE_API FString AutomationReportsDir();
 
 	/** Returns the directory for automation log files */
-	static FString AutomationLogDir();
+	static CORE_API FString AutomationLogDir();
 
 	/** Returns the directory for local files used in cloud emulation or support */
-	static FString CloudDir();
+	static CORE_API FString CloudDir();
 
 	/**
 	 * Returns the directory that contains subfolders for developer-specific content
 	 * Example: "../../../ProjectName/Content/Developers
 	 */
-	static FString GameDevelopersDir();
+	static CORE_API FString GameDevelopersDir();
 
 	/**
 	 * Returns the name of the subfolder for developer-specific content
 	 * Example: "Developers"
 	 */
-	static FStringView DevelopersFolderName();
+	static CORE_API FStringView DevelopersFolderName();
 
 	/** Returns The folder name for the developer-specific directory for the current user */
-	static FString GameUserDeveloperFolderName();
+	static CORE_API FString GameUserDeveloperFolderName();
 
 	/** Returns The directory that contains developer-specific content for the current user */
-	static FString GameUserDeveloperDir();
+	static CORE_API FString GameUserDeveloperDir();
 
 	/** Returns the directory for temp files used for diffing */
-	static FString DiffDir();
+	static CORE_API FString DiffDir();
 
 	/** 
 	 * Returns a list of engine-specific localization paths
 	 */
-	static const TArray<FString>& GetEngineLocalizationPaths();
+	static CORE_API const TArray<FString>& GetEngineLocalizationPaths();
 
 	/** 
 	 * Returns a list of editor-specific localization paths
 	 */
-	static const TArray<FString>& GetEditorLocalizationPaths();
+	static CORE_API const TArray<FString>& GetEditorLocalizationPaths();
 
 	/**
 	 * Returns a list of cooked editor-specific localization paths
 	 */
-	static const TArray<FString>& GetCookedEditorLocalizationPaths();
+	static CORE_API const TArray<FString>& GetCookedEditorLocalizationPaths();
 
 	/** 
 	 * Returns a list of property name localization paths
 	 */
-	static const TArray<FString>& GetPropertyNameLocalizationPaths();
+	static CORE_API const TArray<FString>& GetPropertyNameLocalizationPaths();
 
 	/** 
 	 * Returns a list of tool tip localization paths
 	 */
-	static const TArray<FString>& GetToolTipLocalizationPaths();
+	static CORE_API const TArray<FString>& GetToolTipLocalizationPaths();
 
 	/** 
 	 * Returns a list of game-specific localization paths
 	 */
-	static const TArray<FString>& GetGameLocalizationPaths();
+	static CORE_API const TArray<FString>& GetGameLocalizationPaths();
 
 	/**
 	 * Get the name of the platform-specific localization sub-folder
 	 */
-	static FString GetPlatformLocalizationFolderName();
+	static CORE_API FString GetPlatformLocalizationFolderName();
 
 	/** 
 	 * Returns a list of restricted/internal folder names (without any slashes) which may be tested against full paths to determine if a path is restricted or not.
 	 */
-	static const TArray<FString>& GetRestrictedFolderNames();
+	static CORE_API const TArray<FString>& GetRestrictedFolderNames();
 
 	/** 
 	 * Determines if supplied path uses a restricted/internal subdirectory.	Note that slashes are normalized and character case is ignored for the comparison.
 	 */
-	static bool IsRestrictedPath(const FString& InPath);
+	static CORE_API bool IsRestrictedPath(const FString& InPath);
 
 	/**
 	 * Returns the saved directory that is not game specific. This is usually the same as
@@ -408,37 +408,37 @@ public:
 	 *
 	 * @return saved directory
 	 */
-	static FString GameAgnosticSavedDir();
+	static CORE_API FString GameAgnosticSavedDir();
 
 	/** Returns the directory where engine source code files are kept */
-	static FString EngineSourceDir();
+	static CORE_API FString EngineSourceDir();
 
 	/** Returns the directory where game source code files are kept */
-	static FString GameSourceDir();
+	static CORE_API FString GameSourceDir();
 
 	/** Returns the directory where feature packs are kept */
-	static FString FeaturePackDir();
+	static CORE_API FString FeaturePackDir();
 
 	/**
 	 * Checks whether the path to the project file, if any, is set.
 	 *
 	 * @return true if the path is set, false otherwise.
 	 */
-	static bool IsProjectFilePathSet();
+	static CORE_API bool IsProjectFilePathSet();
 	
 	/**
 	 * Gets the path to the project file.
 	 *
 	 * @return Project file path.
 	 */
-	static FString GetProjectFilePath();
+	static CORE_API FString GetProjectFilePath();
 
 	/**
 	 * Sets the path to the project file.
 	 *
 	 * @param NewGameProjectFilePath - The project file path to set.
 	 */
-	static void SetProjectFilePath( const FString& NewGameProjectFilePath );
+	static CORE_API void SetProjectFilePath( const FString& NewGameProjectFilePath );
 
 	/**
 	 * Gets the extension for this filename.
@@ -447,55 +447,55 @@ public:
 	 *
 	 * @return	the extension of this filename, or an empty string if the filename doesn't have an extension.
 	 */
-	static FString GetExtension( const FString& InPath, bool bIncludeDot=false );
+	static CORE_API FString GetExtension( const FString& InPath, bool bIncludeDot=false );
 
 	// Returns the filename (with extension), minus any path information.
-	static FString GetCleanFilename(const FString& InPath);
+	static CORE_API FString GetCleanFilename(const FString& InPath);
 
 	// Returns the filename (with extension), minus any path information.
-	static FString GetCleanFilename(FString&& InPath);
+	static CORE_API FString GetCleanFilename(FString&& InPath);
 
 	// Returns the same thing as GetCleanFilename, but without the extension
-	static FString GetBaseFilename(const FString& InPath, bool bRemovePath=true );
+	static CORE_API FString GetBaseFilename(const FString& InPath, bool bRemovePath=true );
 
 	// Returns the same thing as GetCleanFilename, but without the extension
-	static FString GetBaseFilename(FString&& InPath, bool bRemovePath = true);
+	static CORE_API FString GetBaseFilename(FString&& InPath, bool bRemovePath = true);
 
 	// Returns the path in front of the filename
-	static FString GetPath(const FString& InPath);
+	static CORE_API FString GetPath(const FString& InPath);
 
 	// Returns the path in front of the filename
-	static FString GetPath(FString&& InPath);
+	static CORE_API FString GetPath(FString&& InPath);
 
 	// Returns the leaf in the path
-	static FString GetPathLeaf(const FString& InPath);
+	static CORE_API FString GetPathLeaf(const FString& InPath);
 
 	// Returns the leaf in the path
-	static FString GetPathLeaf(FString&& InPath);
+	static CORE_API FString GetPathLeaf(FString&& InPath);
 
 	/** Changes the extension of the given filename (does nothing if the file has no extension) */
-	static FString ChangeExtension(const FString& InPath, const FString& InNewExtension);
+	static CORE_API FString ChangeExtension(const FString& InPath, const FString& InNewExtension);
 
 	/** Sets the extension of the given filename (like ChangeExtension, but also applies the extension if the file doesn't have one) */
-	static FString SetExtension(const FString& InPath, const FString& InNewExtension);
+	static CORE_API FString SetExtension(const FString& InPath, const FString& InNewExtension);
 
 	/** Returns true if this file was found, false otherwise */
-	static bool FileExists(const FString& InPath);
+	static CORE_API bool FileExists(const FString& InPath);
 
 	/** Returns true if this directory was found, false otherwise */
-	static bool DirectoryExists(const FString& InPath);
+	static CORE_API bool DirectoryExists(const FString& InPath);
 
 	/** Returns true if this path represents a root drive or volume */
-	static bool IsDrive(const FString& InPath);
+	static CORE_API bool IsDrive(const FString& InPath);
 
 	/** Returns true if this path is relative to another path */
-	static bool IsRelative(const FString& InPath);
+	static CORE_API bool IsRelative(const FString& InPath);
 
 	/** Convert all / and \ to TEXT("/") */
-	static void NormalizeFilename(FString& InPath);
+	static CORE_API void NormalizeFilename(FString& InPath);
 
 	/** Converts the path casing to match the casing found on disk without resolving directory junctions */
-	static FString FindCorrectCase(const FString& Path);
+	static CORE_API FString FindCorrectCase(const FString& Path);
 
 	/**
 	 * Checks if two paths are the same.
@@ -505,13 +505,13 @@ public:
 	 *
 	 * @returns True if both paths are the same. False otherwise.
 	 */
-	static bool IsSamePath(const FString& PathA, const FString& PathB);
+	static CORE_API bool IsSamePath(const FString& PathA, const FString& PathB);
 
 	/** Determines if a path is under a given directory */
-	static bool IsUnderDirectory(const FString& InPath, const FString& InDirectory);
+	static CORE_API bool IsUnderDirectory(const FString& InPath, const FString& InDirectory);
 
 	/** Normalize all / and \ to TEXT("/") and remove any trailing TEXT("/") if the character before that is not a TEXT("/") or a colon */
-	static void NormalizeDirectoryName(FString& InPath);
+	static CORE_API void NormalizeDirectoryName(FString& InPath);
 
 	/**
 	 * Takes a fully pathed string and eliminates relative pathing (eg: annihilates ".." with the adjacent directory).
@@ -521,7 +521,7 @@ public:
 	 * and converts it to:
 	 *	BaseDirectory/SomeOtherDirectory/Filename.ext
 	 */
-	static bool CollapseRelativeDirectories(FString& InPath);
+	static CORE_API bool CollapseRelativeDirectories(FString& InPath);
 
 	/**
 	 * Removes duplicate slashes in paths.
@@ -531,10 +531,10 @@ public:
 	 * and converts it to:
 	 *	BaseDirectory/SomeDirectory/SomeOtherDirectory/Filename.ext
 	 */
-	static void RemoveDuplicateSlashes(FString& InPath);
+	static CORE_API void RemoveDuplicateSlashes(FString& InPath);
 
 	 /** Returns a copy of the given path on which duplicate slashes were removed (see the inplace version for more details). */
-	static FString RemoveDuplicateSlashes(const FString& InPath);
+	static CORE_API FString RemoveDuplicateSlashes(const FString& InPath);
 
 	/**
 	 * Make fully standard "Unreal" pathname:
@@ -543,12 +543,12 @@ public:
 	 *    - Collapses internal ..'s
 	 *    - Makes relative to Engine\Binaries\<Platform> (will ALWAYS start with ..\..\..)
 	 */
-	static FString CreateStandardFilename(const FString& InPath);
+	static CORE_API FString CreateStandardFilename(const FString& InPath);
 
-	static void MakeStandardFilename(FString& InPath);
+	static CORE_API void MakeStandardFilename(FString& InPath);
 
 	/** Takes an "Unreal" pathname and converts it to a platform filename. */
-	static void MakePlatformFilename(FString& InPath);
+	static CORE_API void MakePlatformFilename(FString& InPath);
 
 	/** 
 	 * Assuming both paths (or filenames) are relative to the same base dir, modifies InPath to be relative to InRelativeTo
@@ -557,51 +557,51 @@ public:
 	 * @param InRelativeTo Path to use as the new relative base
 	 * @returns true if InPath was changed to be relative
 	 */
-	static bool MakePathRelativeTo( FString& InPath, const TCHAR* InRelativeTo );
+	static CORE_API bool MakePathRelativeTo( FString& InPath, const TCHAR* InRelativeTo );
 
 	/**
 	 * Converts a relative path name to a fully qualified name relative to the process BaseDir().
 	 */
-	static FString ConvertRelativePathToFull(const FString& InPath);
+	static CORE_API FString ConvertRelativePathToFull(const FString& InPath);
 
 	/**
 	 * Converts a relative path name to a fully qualified name relative to the process BaseDir().
 	 */
-	static FString ConvertRelativePathToFull(FString&& InPath);
+	static CORE_API FString ConvertRelativePathToFull(FString&& InPath);
 
 	/**
 	 * Converts a relative path name to a fully qualified name relative to the specified BasePath.
 	 */
-	static FString ConvertRelativePathToFull(const FString& BasePath, const FString& InPath);
+	static CORE_API FString ConvertRelativePathToFull(const FString& BasePath, const FString& InPath);
 
 	/**
 	 * Converts a relative path name to a fully qualified name relative to the specified BasePath.
 	 */
-	static FString ConvertRelativePathToFull(const FString& BasePath, FString&& InPath);
+	static CORE_API FString ConvertRelativePathToFull(const FString& BasePath, FString&& InPath);
 
 	/**
 	 * Converts a relative path name to a fully qualified name relative to the specified BasePath.
 	 */
-	static FString ConvertRelativePathToFull(FString&& BasePath, const FString& InPath);
+	static CORE_API FString ConvertRelativePathToFull(FString&& BasePath, const FString& InPath);
 
 	/**
 	 * Converts a relative path name to a fully qualified name relative to the specified BasePath.
 	 */
-	static FString ConvertRelativePathToFull(FString&& BasePath, FString&& InPath);
+	static CORE_API FString ConvertRelativePathToFull(FString&& BasePath, FString&& InPath);
 
 	/**
 	 * Converts a normal path to a sandbox path (in Saved/Sandboxes).
 	 *
 	 * @param InSandboxName The name of the sandbox.
 	 */
-	static FString ConvertToSandboxPath( const FString& InPath, const TCHAR* InSandboxName );
+	static CORE_API FString ConvertToSandboxPath( const FString& InPath, const TCHAR* InSandboxName );
 
 	/**
 	 * Converts a sandbox (in Saved/Sandboxes) path to a normal path.
 	 *
 	 * @param InSandboxName The name of the sandbox.
 	 */
-	static FString ConvertFromSandboxPath( const FString& InPath, const TCHAR* InSandboxName );
+	static CORE_API FString ConvertFromSandboxPath( const FString& InPath, const TCHAR* InSandboxName );
 
 	/** 
 	 * Creates a temporary filename with the specified prefix.
@@ -610,18 +610,18 @@ public:
 	 * @param Prefix The file prefix.
 	 * @param Extension File extension ('.' required).
 	 */
-	static FString CreateTempFilename( const TCHAR* Path, const TCHAR* Prefix = TEXT(""), const TCHAR* Extension = TEXT(".tmp") );
+	static CORE_API FString CreateTempFilename( const TCHAR* Path, const TCHAR* Prefix = TEXT(""), const TCHAR* Extension = TEXT(".tmp") );
 
 	/**
 	* Returns a string containing all invalid characters as dictated by the operating system
 	*/
-	static FString GetInvalidFileSystemChars();
+	static CORE_API FString GetInvalidFileSystemChars();
 
 	/**
 	*	Returns a string that is safe to use as a filename because all items in
 	*	GetInvalidFileSystemChars() are removed
 	*/
-	static FString MakeValidFileName(const FString& InString, const TCHAR InReplacementChar = TEXT('\0'));
+	static CORE_API FString MakeValidFileName(const FString& InString, const TCHAR InReplacementChar = TEXT('\0'));
 
 	/** 
 	 * Validates that the parts that make up the path contain no invalid characters as dictated by the operating system
@@ -630,7 +630,7 @@ public:
 	 * @param InPath - path to validate
 	 * @param OutReason - optional parameter to fill with the failure reason
 	 */
-	static bool ValidatePath( const FString& InPath, FText* OutReason = nullptr );
+	static CORE_API bool ValidatePath( const FString& InPath, FText* OutReason = nullptr );
 
 	/**
 	 * Parses a fully qualified or relative filename into its components (filename, path, extension).
@@ -640,10 +640,10 @@ public:
 	 * @param	FilenamePart	[out] receives the value of the filename portion of the input string
 	 * @param	ExtensionPart	[out] receives the value of the extension portion of the input string
 	 */
-	static void Split( const FString& InPath, FString& PathPart, FString& FilenamePart, FString& ExtensionPart );
+	static CORE_API void Split( const FString& InPath, FString& PathPart, FString& FilenamePart, FString& ExtensionPart );
 
 	/** Gets the relative path to get from BaseDir to RootDirectory  */
-	static const FString& GetRelativePathToRoot();
+	static CORE_API const FString& GetRelativePathToRoot();
 
 	template <typename... PathTypes>
 	FORCEINLINE static FString Combine(PathTypes&&... InPaths)
@@ -654,11 +654,11 @@ public:
 	/**
 	 * Frees any memory retained by FPaths.
 	 */
-	static void TearDown();
+	static CORE_API void TearDown();
 
 protected:
 
-	static FString CombineInternal(const FStringView* Paths, int32 NumPaths);
+	static CORE_API FString CombineInternal(const FStringView* Paths, int32 NumPaths);
 
 private:
 	struct FStaticData;
@@ -671,8 +671,8 @@ private:
 	}
 
 	/** Returns, if any, the value of the -userdir command line argument. This can be used to sandbox artifacts to a desired location */
-	static const FString& CustomUserDirArgument();
+	static CORE_API const FString& CustomUserDirArgument();
 
 	/** Returns, if any, the value of the -shaderworkingdir command line argument. This can be used to sandbox shader working files to a desired location */
-	static const FString& CustomShaderDirArgument();
+	static CORE_API const FString& CustomShaderDirArgument();
 };
diff --git a/Engine/Source/Runtime/Core/Public/Misc/PreLoadFile.h b/Engine/Source/Runtime/Core/Public/Misc/PreLoadFile.h
--- a/Engine/Source/Runtime/Core/Public/Misc/PreLoadFile.h
+++ b/Engine/Source/Runtime/Core/Public/Misc/PreLoadFile.h
@@ -9,19 +9,19 @@
 
 class FEvent;
 
-struct CORE_API FPreLoadFile : public FDelayedAutoRegisterHelper
+struct FPreLoadFile : public FDelayedAutoRegisterHelper
 {
-	FPreLoadFile(const TCHAR* InPath);
+	CORE_API FPreLoadFile(const TCHAR* InPath);
 
-	void* TakeOwnershipOfLoadedData(int64* OutFileSize=nullptr);
-	static void* TakeOwnershipOfLoadedDataByPath(const TCHAR* Filename, int64* OutFileSize);
+	CORE_API void* TakeOwnershipOfLoadedData(int64* OutFileSize=nullptr);
+	static CORE_API void* TakeOwnershipOfLoadedDataByPath(const TCHAR* Filename, int64* OutFileSize);
 
 protected:
-	void KickOffRead();
+	CORE_API void KickOffRead();
 
 	bool bIsComplete;
 	bool bFailedToOpenInKickOff;
-	static bool bSystemNoLongerTakingRequests;
+	static CORE_API bool bSystemNoLongerTakingRequests;
 	void* Data;
 	int64 FileSize;
 	FString Path;
diff --git a/Engine/Source/Runtime/Core/Public/Misc/PreloadableFile.h b/Engine/Source/Runtime/Core/Public/Misc/PreloadableFile.h
--- a/Engine/Source/Runtime/Core/Public/Misc/PreloadableFile.h
+++ b/Engine/Source/Runtime/Core/Public/Misc/PreloadableFile.h
@@ -46,7 +46,7 @@ class IFileHandle;
  * This class is not threadsafe. The public interface can be used at the same time as internal asynchronous tasks are executing, but the
  * public interface can not be used from multiple threads at once.
  */
-class CORE_API FPreloadableArchive : public FArchive
+class FPreloadableArchive : public FArchive
 {
 public:
 	typedef TUniqueFunction<FArchive* ()> FCreateArchive;
@@ -70,59 +70,59 @@ public:
 		DefaultPageSize = 64*1024	// The default size of read requests made to the LowerLevel Archive in PreloadBytes mode when reading bytes into the In-Memory cache.
 	};
 
-	FPreloadableArchive(FStringView ArchiveName);
-	virtual ~FPreloadableArchive();
+	CORE_API FPreloadableArchive(FStringView ArchiveName);
+	CORE_API virtual ~FPreloadableArchive();
 
 	// Initialization
 	/** Set the PageSize used read requests made to the LowerLevel Archive in PreloadBytes mode when reading bytes into the In-Memory cache. Invalid to set after Initialization; PageSize must be constant during use. */
-	void SetPageSize(int64 PageSize);
+	CORE_API void SetPageSize(int64 PageSize);
 	/** Initialize the FPreloadableFile asynchronously, performing FileOpen operations on another thread. Use IsInitialized or WaitForInitialization to check progress. */
-	void InitializeAsync(FCreateArchive&& InCreateArchiveFunction, uint32 InFlags = Flags::None, int64 PrimeSize = DefaultPrimeSize);
-	void InitializeAsync(FCreateAsyncArchive&& InCreateAsyncArchiveFunction, uint32 InFlags = Flags::None, int64 PrimeSize = DefaultPrimeSize);
+	CORE_API void InitializeAsync(FCreateArchive&& InCreateArchiveFunction, uint32 InFlags = Flags::None, int64 PrimeSize = DefaultPrimeSize);
+	CORE_API void InitializeAsync(FCreateAsyncArchive&& InCreateAsyncArchiveFunction, uint32 InFlags = Flags::None, int64 PrimeSize = DefaultPrimeSize);
 	/** Return whether InitializeAsync has completed. If Close is called, state returns to false until the next call to InitializeAsync. */
-	bool IsInitialized() const;
+	CORE_API bool IsInitialized() const;
 	/** Wait for InitializeAsync to complete if it is running, otherwise return immediately. */
-	void WaitForInitialization() const;
+	CORE_API void WaitForInitialization() const;
 
 	// Preloading
 	/** When in PreloadBytes mode, if not already preloading, allocate if necessary the memory for the preloaded bytes and start the chain of asynchronous ReadRequests for the bytes. Returns whether preloading is now active. */
-	bool StartPreload();
+	CORE_API bool StartPreload();
 	/** Cancel any current asynchronous ReadRequests and wait for the asynchronous work to exit. */
-	void StopPreload();
+	CORE_API void StopPreload();
 	/** Return whether preloading is in progress. Value may not be up to date since asynchronous work might complete in a race condition. */
-	bool IsPreloading() const;
+	CORE_API bool IsPreloading() const;
 	/** When in PreloadBytes mode, allocate if necessary the memory for the preloaded bytes. Return whether the memory is now allocated. */
-	bool AllocateCache();
+	CORE_API bool AllocateCache();
 	/** Free all memory used by the cache or for preloading (calling StopPreload if necessary). */
-	void ReleaseCache();
+	CORE_API void ReleaseCache();
 	/** Return whether the cache is currently allocated. */
-	bool IsCacheAllocated() const;
+	CORE_API bool IsCacheAllocated() const;
 	/** Return the LowerLevel FArchive if it has been allocated. May return null, even if the FPreloadableFile is currently active. If return value is non-null, caller is responsible for deleting it. */
-	FArchive* DetachLowerLevel();
+	CORE_API FArchive* DetachLowerLevel();
 
 	// FArchive
-	virtual void Serialize(void* V, int64 Length) final;
-	virtual void Seek(int64 InPos) final;
-	virtual int64 Tell() final;
+	CORE_API virtual void Serialize(void* V, int64 Length) final;
+	CORE_API virtual void Seek(int64 InPos) final;
+	CORE_API virtual int64 Tell() final;
 	/** Return the size of the file, or -1 if the file does not exist. This is also the amount of memory that will be allocated by AllocateCache. */
-	virtual int64 TotalSize() final;
-	virtual bool Close() final;
-	virtual FString GetArchiveName() const final;
+	CORE_API virtual int64 TotalSize() final;
+	CORE_API virtual bool Close() final;
+	CORE_API virtual FString GetArchiveName() const final;
 
 protected:
 	/** Helper function for InitializeAsync, sets up the asynchronous call to InitializeInternal */
-	void InitializeInternalAsync(FCreateArchive&& InCreateArchiveFunction, FCreateAsyncArchive&& InCreateAsyncArchiveFunction, uint32 InFlags, int64 PrimeSize);
+	CORE_API void InitializeInternalAsync(FCreateArchive&& InCreateArchiveFunction, FCreateAsyncArchive&& InCreateAsyncArchiveFunction, uint32 InFlags, int64 PrimeSize);
 	/** Helper function for InitializeAsync, called from a TaskGraph thread. */
-	void InitializeInternal(FCreateArchive&& InCreateArchiveFunction, FCreateAsyncArchive&& InCreateAsyncArchiveFunction, uint32 Flags, int64 PrimeSize);
+	CORE_API void InitializeInternal(FCreateArchive&& InCreateArchiveFunction, FCreateAsyncArchive&& InCreateAsyncArchiveFunction, uint32 Flags, int64 PrimeSize);
 #if FPRELOADABLEFILE_TEST_ENABLED
-	void SerializeInternal(void* V, int64 Length);
+	CORE_API void SerializeInternal(void* V, int64 Length);
 #endif
-	void PausePreload();
-	void ResumePreload();
-	bool ResumePreloadNonRecursive();
-	void OnReadComplete(bool bCanceled, IAsyncReadRequest* ReadRequest);
-	void FreeRetiredRequests();
-	void SerializeSynchronously(void* V, int64 Length);
+	CORE_API void PausePreload();
+	CORE_API void ResumePreload();
+	CORE_API bool ResumePreloadNonRecursive();
+	CORE_API void OnReadComplete(bool bCanceled, IAsyncReadRequest* ReadRequest);
+	CORE_API void FreeRetiredRequests();
+	CORE_API void SerializeSynchronously(void* V, int64 Length);
 
 	FString ArchiveName;
 	/** The Offset into the file or preloaded bytes that will be used in the next call to Serialize. */
@@ -204,13 +204,13 @@ protected:
  * Activate PreloadHandle mode by passing Flags::PreloadHandle to InitializeAsync, optionally or'd with Flags::Prime.
  *
  */
-class CORE_API FPreloadableFile : public FPreloadableArchive
+class FPreloadableFile : public FPreloadableArchive
 {
 public:
-	FPreloadableFile(FStringView FileName);
+	CORE_API FPreloadableFile(FStringView FileName);
 
 	/** Initialize the FPreloadableFile asynchronously, performing FileOpen operations on another thread. Use IsInitialized or WaitForInitialization to check progress. */
-	void InitializeAsync(uint32 InFlags = Flags::None, int64 PrimeSize=DefaultPrimeSize);
+	CORE_API void InitializeAsync(uint32 InFlags = Flags::None, int64 PrimeSize=DefaultPrimeSize);
 
 	// Registration
 	/**
@@ -219,16 +219,16 @@ public:
 	 * Return whether the instance is currently registered. Returns true if the instance was already registered.
 	 * Registered files are referenced-counted, and the reference will not be dropped until (1) (TryTakeArchive or UnRegister is called) and (2) (PreloadBytes mode only) the archive returned from TryTakeArchive is deleted.
 	 */
-	static bool TryRegister(const TSharedPtr<FPreloadableFile>& PreloadableFile);
+	static CORE_API bool TryRegister(const TSharedPtr<FPreloadableFile>& PreloadableFile);
 	/**
 	 * Look up an FPreloadableFile instance registered for the given FileName, and return an FArchive from it.
 	 * If found, removes the registration so no future call to TryTakeArchive can sue the same FArchive.
 	 * If the instance is in PreloadHandle mode, the Lower-Level FArchive will be detached from the FPreloadableFile and returned using DetachLowerLevel.
 	 * If the instance is in PreloadBytes mode, a ProxyArchive will be returned that forwards call to the FPreloadableFile instance.
 	 */
-	static FArchive* TryTakeArchive(const TCHAR* FileName);
+	static CORE_API FArchive* TryTakeArchive(const TCHAR* FileName);
 	/** Remove the FPreloadableFile instance if it is registered for its FileName. Returns whether the instance was registered. */
-	static bool UnRegister(const TSharedPtr<FPreloadableFile>& PreloadableFile);
+	static CORE_API bool UnRegister(const TSharedPtr<FPreloadableFile>& PreloadableFile);
 
 private:
 	/** Map used for TryTakeArchive registration. */
diff --git a/Engine/Source/Runtime/Core/Public/Misc/QueuedThreadPool.h b/Engine/Source/Runtime/Core/Public/Misc/QueuedThreadPool.h
--- a/Engine/Source/Runtime/Core/Public/Misc/QueuedThreadPool.h
+++ b/Engine/Source/Runtime/Core/Public/Misc/QueuedThreadPool.h
@@ -48,35 +48,35 @@ inline const TCHAR* LexToString(EQueuedWorkPriority Priority)
  *
  *  This class is NOT thread-safe and must be properly protected.
  */
-class CORE_API FThreadPoolPriorityQueue
+class FThreadPoolPriorityQueue
 {
 public:
-	FThreadPoolPriorityQueue();
+	CORE_API FThreadPoolPriorityQueue();
 
 	/**
 	 * Enqueue a work item at specified priority
 	 */
-	void Enqueue(IQueuedWork* InQueuedWork, EQueuedWorkPriority InPriority = EQueuedWorkPriority::Normal);
+	CORE_API void Enqueue(IQueuedWork* InQueuedWork, EQueuedWorkPriority InPriority = EQueuedWorkPriority::Normal);
 
 	/**
 	 * Search and remove a queued work item from the list
 	 */
-	bool Retract(IQueuedWork* InQueuedWork);
+	CORE_API bool Retract(IQueuedWork* InQueuedWork);
 
 	/**
 	 * Get the next work item in priority order.
 	 */
-	IQueuedWork* Dequeue(EQueuedWorkPriority* OutDequeuedWorkPriority = nullptr);
+	CORE_API IQueuedWork* Dequeue(EQueuedWorkPriority* OutDequeuedWorkPriority = nullptr);
 
 	/**
 	 * Get the next work item in priority order without actually dequeuing.
 	 */
-	IQueuedWork* Peek(EQueuedWorkPriority* OutDequeuedWorkPriority = nullptr) const;
+	CORE_API IQueuedWork* Peek(EQueuedWorkPriority* OutDequeuedWorkPriority = nullptr) const;
 
 	/**
 	 * Empty the queue.
 	 */
-	void Reset();
+	CORE_API void Reset();
 
 	/**
 	 * Get the total number of queued items.
@@ -86,7 +86,7 @@ public:
 	/**
 	 * Sort Priority Bucket given Predicate
 	 */
-	void Sort(EQueuedWorkPriority InPriorityBucket, TFunctionRef<bool (const IQueuedWork* A, const IQueuedWork* B)> Predicate);
+	CORE_API void Sort(EQueuedWorkPriority InPriorityBucket, TFunctionRef<bool (const IQueuedWork* A, const IQueuedWork* B)> Predicate);
 private:
 	/** The first queue to extract a work item from to avoid scanning all priorities when unqueuing. */
 	int32 FirstNonEmptyQueueIndex = 0;
@@ -100,7 +100,7 @@ private:
  * This interface is used by all queued thread pools. It used as a callback by
  * FQueuedThreads and is used to queue asynchronous work for callers.
  */
-class CORE_API FQueuedThreadPool
+class FQueuedThreadPool
 {
 public:
 	/**
@@ -142,8 +142,8 @@ public:
 	virtual int32 GetNumThreads() const = 0;
 
 public:
-			FQueuedThreadPool();
-	virtual	~FQueuedThreadPool();
+			CORE_API FQueuedThreadPool();
+	CORE_API virtual	~FQueuedThreadPool();
 
 public:
 
@@ -152,14 +152,14 @@ public:
 	 *
 	 * @return A new thread pool.
 	 */
-	static FQueuedThreadPool* Allocate();
+	static CORE_API FQueuedThreadPool* Allocate();
 
 	/**
 	 *	Stack size for threads created for the thread pool. 
 	 *	Can be overridden by other projects.
 	 *	If 0 means to use the value passed in the Create method.
 	 */
-	static uint32 OverrideStackSize;
+	static CORE_API uint32 OverrideStackSize;
 };
 
 /**
diff --git a/Engine/Source/Runtime/Core/Public/Misc/QueuedThreadPoolWrapper.h b/Engine/Source/Runtime/Core/Public/Misc/QueuedThreadPoolWrapper.h
--- a/Engine/Source/Runtime/Core/Public/Misc/QueuedThreadPoolWrapper.h
+++ b/Engine/Source/Runtime/Core/Public/Misc/QueuedThreadPoolWrapper.h
@@ -37,7 +37,7 @@
   * another thread-pool a breeze and allowing more fine-grained control
   * over scheduling by effectively giving another set of priorities.
   */
-class CORE_API FQueuedThreadPoolWrapper : public FQueuedThreadPool
+class FQueuedThreadPoolWrapper : public FQueuedThreadPool
 {
 public:
 	/**
@@ -45,34 +45,34 @@ public:
 	 * InMaxConcurrency           Maximum number of concurrent tasks allowed, -1 will limit concurrency to number of threads available in the underlying thread pool.
 	 * InPriorityMapper           Thread-safe function used to map any priority from this Queue to the priority that should be used when scheduling the task on the underlying thread pool.
 	 */
-	FQueuedThreadPoolWrapper(FQueuedThreadPool* InWrappedQueuedThreadPool, int32 InMaxConcurrency = -1, TFunction<EQueuedWorkPriority(EQueuedWorkPriority)> InPriorityMapper = [](EQueuedWorkPriority InPriority) { return InPriority; });
-	~FQueuedThreadPoolWrapper();
+	CORE_API FQueuedThreadPoolWrapper(FQueuedThreadPool* InWrappedQueuedThreadPool, int32 InMaxConcurrency = -1, TFunction<EQueuedWorkPriority(EQueuedWorkPriority)> InPriorityMapper = [](EQueuedWorkPriority InPriority) { return InPriority; });
+	CORE_API ~FQueuedThreadPoolWrapper();
 
 	/**
 	 *  Queued task are not scheduled against the wrapped thread-pool until resumed
 	 */
-	void Pause();
+	CORE_API void Pause();
 
 	/**
 	 *  Resume a specified amount of queued work, or -1 to unpause.
 	 */
-	void Resume(int32 InNumQueuedWork = -1);
+	CORE_API void Resume(int32 InNumQueuedWork = -1);
 
 	/**
 	 *  Dynamically adjust the maximum number of concurrent tasks, -1 for unlimited.
 	 */
-	void SetMaxConcurrency(int32 MaxConcurrency = -1);
+	CORE_API void SetMaxConcurrency(int32 MaxConcurrency = -1);
 
-	void AddQueuedWork(IQueuedWork* InQueuedWork, EQueuedWorkPriority InPriority = EQueuedWorkPriority::Normal) override;
-	bool RetractQueuedWork(IQueuedWork* InQueuedWork) override;
-	int32 GetNumThreads() const override;
+	CORE_API void AddQueuedWork(IQueuedWork* InQueuedWork, EQueuedWorkPriority InPriority = EQueuedWorkPriority::Normal) override;
+	CORE_API bool RetractQueuedWork(IQueuedWork* InQueuedWork) override;
+	CORE_API int32 GetNumThreads() const override;
 	int32 GetCurrentConcurrency() const { return CurrentConcurrency.load(std::memory_order_relaxed); }
 
 protected:
-	class CORE_API FScheduledWork : public IQueuedWork, public IExecutionResource
+	class FScheduledWork : public IQueuedWork, public IExecutionResource
 	{
 	public:
-		FScheduledWork();
+		CORE_API FScheduledWork();
 
 		FScheduledWork(const FScheduledWork&) = delete;
 		FScheduledWork& operator=(const FScheduledWork&) = delete;
@@ -80,7 +80,7 @@ protected:
 		FScheduledWork(const FScheduledWork&&) = delete;
 		FScheduledWork& operator=(const FScheduledWork&&) = delete;
 
-		~FScheduledWork() override;
+		CORE_API ~FScheduledWork() override;
 		
 		const TCHAR * GetDebugName() const final
 		{
@@ -199,13 +199,13 @@ protected:
 	// Can be overriden to allocate a more specialized version if needed.
 	virtual FScheduledWork* AllocateScheduledWork() { return new FScheduledWork(); }
 private:
-	FScheduledWork* AllocateWork(IQueuedWork* InnerWork, EQueuedWorkPriority Priority);
-	bool CanSchedule(EQueuedWorkPriority Priority) const;
-	bool Create(uint32 InNumQueuedThreads, uint32 StackSize, EThreadPriority ThreadPriority, const TCHAR* Name) override;
-	void Destroy() override;
-	void Schedule(FScheduledWork* Work = nullptr);
-	void ReleaseWorkNoLock(FScheduledWork* Work);
-	bool TryRetractWorkNoLock(EQueuedWorkPriority InPriority);
+	CORE_API FScheduledWork* AllocateWork(IQueuedWork* InnerWork, EQueuedWorkPriority Priority);
+	CORE_API bool CanSchedule(EQueuedWorkPriority Priority) const;
+	CORE_API bool Create(uint32 InNumQueuedThreads, uint32 StackSize, EThreadPriority ThreadPriority, const TCHAR* Name) override;
+	CORE_API void Destroy() override;
+	CORE_API void Schedule(FScheduledWork* Work = nullptr);
+	CORE_API void ReleaseWorkNoLock(FScheduledWork* Work);
+	CORE_API bool TryRetractWorkNoLock(EQueuedWorkPriority InPriority);
 	
 	TFunction<EQueuedWorkPriority(EQueuedWorkPriority)> PriorityMapper;
 
@@ -224,7 +224,7 @@ private:
   * another thread-pool a breeze and allowing more fine-grained control
   * over scheduling by giving full control of task reordering.
   */
-class CORE_API FQueuedThreadPoolDynamicWrapper : public FQueuedThreadPoolWrapper
+class FQueuedThreadPoolDynamicWrapper : public FQueuedThreadPoolWrapper
 {
 public:
 	/**
@@ -255,7 +255,7 @@ public:
 
 /** ThreadPool wrapper implementation allowing to schedule thread-pool tasks on the task graph.
   */
-class CORE_API FQueuedThreadPoolTaskGraphWrapper final : public FQueuedThreadPool
+class FQueuedThreadPoolTaskGraphWrapper final : public FQueuedThreadPool
 {
 public:
 	/**
@@ -366,7 +366,7 @@ private:
 
 /** ThreadPool wrapper implementation allowing to schedule thread-pool tasks on the the low level backend which is also used by the taskgraph.
 */
-class CORE_API FQueuedLowLevelThreadPool final : public FQueuedThreadPool
+class FQueuedLowLevelThreadPool final : public FQueuedThreadPool
 {
 	/* Internal data of the scheduler used for cancellation */
 	struct FQueuedWorkInternalData : TConcurrentLinearObject<FQueuedWorkInternalData, FTaskGraphBlockAllocationTag>, IQueuedWorkInternalData
@@ -573,4 +573,4 @@ private:
 	std::atomic_uint TaskCount{0};
 	std::atomic_bool bIsExiting{false};
 	std::atomic_bool bIsPaused{false};
-};
\ No newline at end of file
+};
diff --git a/Engine/Source/Runtime/Core/Public/Misc/RuntimeErrors.h b/Engine/Source/Runtime/Core/Public/Misc/RuntimeErrors.h
--- a/Engine/Source/Runtime/Core/Public/Misc/RuntimeErrors.h
+++ b/Engine/Source/Runtime/Core/Public/Misc/RuntimeErrors.h
@@ -17,7 +17,7 @@ class FText;
 #define UE_RAISE_RUNTIME_ERRORS DO_CHECK
 #endif
 
-struct CORE_API FRuntimeErrors
+struct FRuntimeErrors
 {
 	/**
 	 * Prints out a runtime warning or error (typically by 'throwing' a BP script exception with the BP callstack)
@@ -29,13 +29,13 @@ struct CORE_API FRuntimeErrors
 	 *
 	 * @return false in all cases.
 	 */
-	static void LogRuntimeIssue(ELogVerbosity::Type Verbosity, const ANSICHAR* FileName, int32 LineNumber, const FText& Message);
+	static CORE_API void LogRuntimeIssue(ELogVerbosity::Type Verbosity, const ANSICHAR* FileName, int32 LineNumber, const FText& Message);
 
 	// The style of delegate called when a runtime issue occurs
 	DECLARE_DELEGATE_FourParams(FRuntimeErrorDelegate, ELogVerbosity::Type /*Verbosity*/, const ANSICHAR* /*File*/, int32 /*Line*/, const FText& /*Message*/);
 
 	// A delegate this is called when a runtime issue is raised
-	static FRuntimeErrorDelegate OnRuntimeIssueLogged;
+	static CORE_API FRuntimeErrorDelegate OnRuntimeIssueLogged;
 
 	/**
 	 * Raises a runtime error and returns false.
@@ -46,7 +46,7 @@ struct CORE_API FRuntimeErrors
 	 *
 	 * @return false in all cases.
 	 */
-	static bool LogRuntimeIssueReturningFalse(const ANSICHAR* Expr, const ANSICHAR* File, int32 Line);
+	static CORE_API bool LogRuntimeIssueReturningFalse(const ANSICHAR* Expr, const ANSICHAR* File, int32 Line);
 
 private:
 	FRuntimeErrors() {}
diff --git a/Engine/Source/Runtime/Core/Public/Misc/SecureHash.h b/Engine/Source/Runtime/Core/Public/Misc/SecureHash.h
--- a/Engine/Source/Runtime/Core/Public/Misc/SecureHash.h
+++ b/Engine/Source/Runtime/Core/Public/Misc/SecureHash.h
@@ -60,11 +60,11 @@ struct FMD5Hash;
 // CORE_API void appMD5Encode( uint8* output, uint32* input, int32 len );
 // CORE_API void appMD5Decode( uint32* output, uint8* input, int32 len );
 
-class CORE_API FMD5
+class FMD5
 {
 public:
-	FMD5();
-	~FMD5();
+	CORE_API FMD5();
+	CORE_API ~FMD5();
 
 	/**
 	 * MD5 block update operation.  Continues an MD5 message-digest operation,
@@ -73,7 +73,7 @@ public:
 	 * @param input		input data
 	 * @param inputLen	length of the input data in bytes
 	 **/
-	void Update(const uint8* input, uint64 inputLen);
+	CORE_API void Update(const uint8* input, uint64 inputLen);
 
 	/**
 	 * MD5 finalization. Ends an MD5 message-digest operation, writing the
@@ -82,7 +82,7 @@ public:
 	 *
 	 * @param digest	pointer to a buffer where the digest should be stored ( must have at least 16 bytes )
 	 **/
-	void Final(uint8* digest);
+	CORE_API void Final(uint8* digest);
 
 	/**
 	 * Helper to perform the very common case of hashing an ASCII string into a hex representation.
@@ -124,9 +124,9 @@ private:
 		uint8 buffer[64];
 	};
 
-	void Transform( uint32* state, const uint8* block );
-	void Encode( uint8* output, const uint32* input, int32 len );
-	void Decode( uint32* output, const uint8* input, int32 len );
+	CORE_API void Transform( uint32* state, const uint8* block );
+	CORE_API void Encode( uint8* output, const uint32* input, int32 len );
+	CORE_API void Decode( uint32* output, const uint8* input, int32 len );
 
 	FContext Context;
 };
@@ -224,7 +224,7 @@ inline FGuid MD5HashToGuid(const FMD5Hash& Hash)
 #define HASHES_SHA_DIVIDER "+++"
 
 /** Stores an SHA hash generated by FSHA1. */
-class CORE_API FSHAHash
+class FSHAHash
 {
 public:
 	alignas(uint32) uint8 Hash[20];
@@ -305,14 +305,14 @@ namespace Freeze
 
 DECLARE_INTRINSIC_TYPE_LAYOUT(FSHAHash);
 
-class CORE_API FSHA1
+class FSHA1
 {
 public:
 
 	enum {DigestSize=20};
 	// Constructor and Destructor
-	FSHA1();
-	~FSHA1();
+	CORE_API FSHA1();
+	CORE_API ~FSHA1();
 
 	uint32 m_state[5];
 	uint64 m_count;
@@ -321,16 +321,16 @@ public:
 	uint8  m_digest[20];
 	uint32 __reserved2[3];
 
-	void Reset();
+	CORE_API void Reset();
 
 	// Update the hash value
-	void Update(const uint8 *data, uint64 len);
+	CORE_API void Update(const uint8 *data, uint64 len);
 
 	// Update the hash value with string
-	void UpdateWithString(const TCHAR *data, uint32 len);
+	CORE_API void UpdateWithString(const TCHAR *data, uint32 len);
 
 	// Finalize hash and report
-	void Final();
+	CORE_API void Final();
 
 	// Finalize hash and return it
 	FSHAHash Finalize()
@@ -342,7 +342,7 @@ public:
 	}
 
 	// Report functions: as pre-formatted and raw data
-	void GetHash(uint8 *puDest);
+	CORE_API void GetHash(uint8 *puDest);
 
 	/**
 	 * Calculate the hash on a single block and return it
@@ -351,7 +351,7 @@ public:
 	 * @param DataSize Size of the Data block
 	 * @param OutHash Resulting hash value (20 byte buffer)
 	 */
-	static void HashBuffer(const void* Data, uint64 DataSize, uint8* OutHash);
+	static CORE_API void HashBuffer(const void* Data, uint64 DataSize, uint8* OutHash);
 
 	/**
 	 * Calculate the hash on a single block and return it
@@ -377,7 +377,7 @@ public:
 	 * @param DataSize	Size of the Data block
 	 * @param OutHash	Resulting hash value (20 byte buffer)
 	 */
-	static void HMACBuffer(const void* Key, uint32 KeySize, const void* Data, uint64 DataSize, uint8* OutHash);
+	static CORE_API void HMACBuffer(const void* Key, uint32 KeySize, const void* Data, uint64 DataSize, uint8* OutHash);
 
 	/**
 	 * Shared hashes.sha reading code (each platform gets a buffer to the data,
@@ -387,7 +387,7 @@ public:
 	 * @param BufferSize Size of Buffer
 	 * @param bDuplicateKeyMemory If Buffer is not always loaded, pass true so that the 20 byte hashes are duplicated 
 	 */
-	static void InitializeFileHashesFromBuffer(uint8* Buffer, uint64 BufferSize, bool bDuplicateKeyMemory=false);
+	static CORE_API void InitializeFileHashesFromBuffer(uint8* Buffer, uint64 BufferSize, bool bDuplicateKeyMemory=false);
 
 	/**
 	 * Gets the stored SHA hash from the platform, if it exists. This function
@@ -399,24 +399,24 @@ public:
 	 *
 	 * @return true if the hash was found, false otherwise
 	 */
-	static bool GetFileSHAHash(const TCHAR* Pathname, uint8 Hash[20], bool bIsFullPackageHash=true);
+	static CORE_API bool GetFileSHAHash(const TCHAR* Pathname, uint8 Hash[20], bool bIsFullPackageHash=true);
 
 private:
 	// Private SHA-1 transformation
-	void Transform(const uint8* buffer, uint64 len);
+	CORE_API void Transform(const uint8* buffer, uint64 len);
 
 	/** Global map of filename to hash value, filled out in InitializeFileHashesFromBuffer */
-	static TMap<FString, uint8*> FullFileSHAHashMap;
+	static CORE_API TMap<FString, uint8*> FullFileSHAHashMap;
 
 	/** Global map of filename to hash value, but for script-only SHA hashes */
-	static TMap<FString, uint8*> ScriptSHAHashMap;
+	static CORE_API TMap<FString, uint8*> ScriptSHAHashMap;
 };
 
 
 /**
  * Asynchronous SHA verification
  */
-class CORE_API FAsyncSHAVerify
+class FAsyncSHAVerify
 {
 protected:
 	/** Buffer to run the has on. This class can take ownership of the buffer is bShouldDeleteBuffer is true */
@@ -466,7 +466,7 @@ public:
 	/**
 	 * Performs the async hash verification
 	 */
-	void DoWork();
+	CORE_API void DoWork();
 
 	/**
 	 * Task API, return true to indicate that we can abandon
diff --git a/Engine/Source/Runtime/Core/Public/Misc/SingleThreadRunnable.h b/Engine/Source/Runtime/Core/Public/Misc/SingleThreadRunnable.h
--- a/Engine/Source/Runtime/Core/Public/Misc/SingleThreadRunnable.h
+++ b/Engine/Source/Runtime/Core/Public/Misc/SingleThreadRunnable.h
@@ -8,7 +8,7 @@
  * Interface for ticking runnables when there's only one thread available and 
  * multithreading is disabled.
  */
-class CORE_API FSingleThreadRunnable
+class FSingleThreadRunnable
 {
 public:
 
diff --git a/Engine/Source/Runtime/Core/Public/Misc/SlowTask.h b/Engine/Source/Runtime/Core/Public/Misc/SlowTask.h
--- a/Engine/Source/Runtime/Core/Public/Misc/SlowTask.h
+++ b/Engine/Source/Runtime/Core/Public/Misc/SlowTask.h
@@ -25,7 +25,7 @@ enum class ESlowTaskVisibility
 /**
  * Data type used to store information about a currently running slow task. Direct use is not advised, use FScopedSlowTask instead
  */
-struct CORE_API FSlowTask
+struct FSlowTask
 {
 	/** Default message to display to the user when not overridden by a frame */
 	FText DefaultMessage;
@@ -88,7 +88,7 @@ public:
 	/**
 	 * True if a slow task should be created. False if a slow task should not be created because another was created recently.
 	 */
-	static bool ShouldCreateThrottledSlowTask();
+	static CORE_API bool ShouldCreateThrottledSlowTask();
 
 	/**
 	 * Construct this scope from an amount of work to do, and a message to display
@@ -97,13 +97,13 @@ public:
 	 * @param		InDefaultMessage		A message to display to the user to describe the purpose of the scope
 	 * @param		bInVisible				When false, this scope will have no effect. Allows for proper scoped objects that are conditionally hidden.
 	 */
-	FSlowTask(float InAmountOfWork, const FText& InDefaultMessage = FText(), bool bInEnabled = true, FFeedbackContext& InContext = *GWarn);
+	CORE_API FSlowTask(float InAmountOfWork, const FText& InDefaultMessage = FText(), bool bInEnabled = true, FFeedbackContext& InContext = *GWarn);
 
 	/** Function that initializes the scope by adding it to its context's stack */
-	void Initialize();
+	CORE_API void Initialize();
 
 	/** Function that finishes any remaining work and removes itself from the global scope stack */
-	void Destroy();
+	CORE_API void Destroy();
 
 	/**
 	 * Creates a new dialog for this slow task after the given time threshold. If the task completes before this time, no dialog will be shown.
@@ -111,40 +111,40 @@ public:
 	 * @param		bShowCancelButton		Whether to show a cancel button on the dialog or not
 	 * @param		bAllowInPIE				Whether to allow this dialog in PIE. If false, this dialog will not appear during PIE sessions.
 	 */
-	void MakeDialogDelayed(float Threshold, bool bShowCancelButton = false, bool bAllowInPIE = false);
+	CORE_API void MakeDialogDelayed(float Threshold, bool bShowCancelButton = false, bool bAllowInPIE = false);
 
 	/**
 	 * Creates a new dialog for this slow task, if there is currently not one open
 	 * @param		bShowCancelButton		Whether to show a cancel button on the dialog or not
 	 * @param		bAllowInPIE				Whether to allow this dialog in PIE. If false, this dialog will not appear during PIE sessions.
 	 */
-	void MakeDialog(bool bShowCancelButton = false, bool bAllowInPIE = false);
+	CORE_API void MakeDialog(bool bShowCancelButton = false, bool bAllowInPIE = false);
 
 	/**
 	 * Indicate that we are to enter a frame that will take up the specified amount of work. Completes any previous frames (potentially contributing to parent scopes' progress).
 	 * @param		ExpectedWorkThisFrame	The amount of work that will happen between now and the next frame, as a numerator of TotalAmountOfWork.
 	 * @param		Text					Optional text to describe this frame's purpose.
 	 */
-	void EnterProgressFrame(float ExpectedWorkThisFrame = 1.f, const FText& Text = FText());
+	CORE_API void EnterProgressFrame(float ExpectedWorkThisFrame = 1.f, const FText& Text = FText());
 
 	/**
 	 * Let the UI refresh but doesn't advance progress. This should be called at regular intervals even when no progress is being made to keep the UI responsive.
 	 */
-	void TickProgress();
+	CORE_API void TickProgress();
 
 	/** 
 	* Whenever we encounter edge condition requiring to force a UI refresh
 	* (Like if another popup would show when the slow task finishes and you want it to show 100%)
 	*/
-	void ForceRefresh();
+	CORE_API void ForceRefresh();
 
 	/**
 	 * Get the frame message or default message if empty
 	 */
-	const FText& GetCurrentMessage() const;
+	CORE_API const FText& GetCurrentMessage() const;
 
 	/**
 	 * True if the user has requested that the slow task be canceled
 	 */
-	bool ShouldCancel() const;
+	CORE_API bool ShouldCancel() const;
 };
diff --git a/Engine/Source/Runtime/Core/Public/Misc/StringFormatArg.h b/Engine/Source/Runtime/Core/Public/Misc/StringFormatArg.h
--- a/Engine/Source/Runtime/Core/Public/Misc/StringFormatArg.h
+++ b/Engine/Source/Runtime/Core/Public/Misc/StringFormatArg.h
@@ -7,7 +7,7 @@
 #include "Containers/ContainersFwd.h"
 
 /** An argument supplied to FString::Format */
-struct CORE_API FStringFormatArg
+struct FStringFormatArg
 {
 	enum EType { Int, UInt, Double, String, StringLiteralANSI, StringLiteralWIDE, StringLiteralUCS2, StringLiteralUTF8 };
 
@@ -40,18 +40,18 @@ struct CORE_API FStringFormatArg
 	FStringFormatArg() = delete;
 	~FStringFormatArg() = default;
 
-	FStringFormatArg(const int32 Value);
-	FStringFormatArg(const uint32 Value);
-	FStringFormatArg(const int64 Value);
-	FStringFormatArg(const uint64 Value);
-	FStringFormatArg(const float Value);
-	FStringFormatArg(const double Value);
-	FStringFormatArg(FString Value);
-	FStringFormatArg(FStringView Value);
-	FStringFormatArg(const ANSICHAR* Value);
-	FStringFormatArg(const WIDECHAR* Value);
-	FStringFormatArg(const UCS2CHAR* Value);
-	FStringFormatArg(const UTF8CHAR* Value);
+	CORE_API FStringFormatArg(const int32 Value);
+	CORE_API FStringFormatArg(const uint32 Value);
+	CORE_API FStringFormatArg(const int64 Value);
+	CORE_API FStringFormatArg(const uint64 Value);
+	CORE_API FStringFormatArg(const float Value);
+	CORE_API FStringFormatArg(const double Value);
+	CORE_API FStringFormatArg(FString Value);
+	CORE_API FStringFormatArg(FStringView Value);
+	CORE_API FStringFormatArg(const ANSICHAR* Value);
+	CORE_API FStringFormatArg(const WIDECHAR* Value);
+	CORE_API FStringFormatArg(const UCS2CHAR* Value);
+	CORE_API FStringFormatArg(const UTF8CHAR* Value);
 
 	FStringFormatArg(const FStringFormatArg& Other)
 	{
@@ -62,7 +62,7 @@ struct CORE_API FStringFormatArg
 		*this = MoveTemp(Other);
 	}
 	
-	FStringFormatArg& operator=(const FStringFormatArg& Other);
-	FStringFormatArg& operator=(FStringFormatArg&& Other);
+	CORE_API FStringFormatArg& operator=(const FStringFormatArg& Other);
+	CORE_API FStringFormatArg& operator=(FStringFormatArg&& Other);
 
 };
diff --git a/Engine/Source/Runtime/Core/Public/Misc/TextFilterExpressionEvaluator.h b/Engine/Source/Runtime/Core/Public/Misc/TextFilterExpressionEvaluator.h
--- a/Engine/Source/Runtime/Core/Public/Misc/TextFilterExpressionEvaluator.h
+++ b/Engine/Source/Runtime/Core/Public/Misc/TextFilterExpressionEvaluator.h
@@ -194,45 +194,45 @@ DEFINE_EXPRESSION_NODE_TYPE(TextFilterExpressionParser::FFunction,	0x084E6214, 0
 DECLARE_DELEGATE_RetVal_OneParam(bool, FTokenFunctionHandler, const FTextFilterString&);
 
 /** Defines an expression evaluator that can be used to perform complex text filter queries */
-class CORE_API FTextFilterExpressionEvaluator
+class FTextFilterExpressionEvaluator
 {
 public:
 	/** Construction and assignment */
-	FTextFilterExpressionEvaluator(const ETextFilterExpressionEvaluatorMode InMode);
-	FTextFilterExpressionEvaluator(const FTextFilterExpressionEvaluator& Other);
-	FTextFilterExpressionEvaluator& operator=(const FTextFilterExpressionEvaluator& Other);
+	CORE_API FTextFilterExpressionEvaluator(const ETextFilterExpressionEvaluatorMode InMode);
+	CORE_API FTextFilterExpressionEvaluator(const FTextFilterExpressionEvaluator& Other);
+	CORE_API FTextFilterExpressionEvaluator& operator=(const FTextFilterExpressionEvaluator& Other);
 	virtual ~FTextFilterExpressionEvaluator() {}
 
 	/** Get the complexity of the current filter terms */
-	ETextFilterExpressionType GetFilterType() const;
+	CORE_API ETextFilterExpressionType GetFilterType() const;
 
 	/** Get the filter terms that we're currently using */
-	FText GetFilterText() const;
+	CORE_API FText GetFilterText() const;
 
 	/** Set the filter terms to be compiled for evaluation later. Returns true if the filter was actually changed */
-	bool SetFilterText(const FText& InFilterText);
+	CORE_API bool SetFilterText(const FText& InFilterText);
 
 	/** Get the last error returned from lexing or compiling the current filter text */
-	FText GetFilterErrorText() const;
+	CORE_API FText GetFilterErrorText() const;
 
 	/** Get the filter expression tokens created from the current filter text */
-	const TArray<FExpressionToken>& GetFilterExpressionTokens() const;
+	CORE_API const TArray<FExpressionToken>& GetFilterExpressionTokens() const;
 
 	/** Test our compiled filter using the given context */
-	bool TestTextFilter(const ITextFilterExpressionContext& InContext) const;
+	CORE_API bool TestTextFilter(const ITextFilterExpressionContext& InContext) const;
 
 	/** Helper function to add callbacks for function tokens */
-	void AddFunctionTokenCallback(FString InFunctionName, FTokenFunctionHandler InCallback);
+	CORE_API void AddFunctionTokenCallback(FString InFunctionName, FTokenFunctionHandler InCallback);
 
 protected:
 	/** Sets up grammar used for evaluation */
-	void SetupGrammar();
+	CORE_API void SetupGrammar();
 
 	/** Common function to construct the expression parser */
-	virtual void ConstructExpressionParser();
+	CORE_API virtual void ConstructExpressionParser();
 
 	/** Evaluate the given compiled result, and optionally populate OutErrorText with any error information */
-	virtual bool EvaluateCompiledExpression(const ExpressionParser::CompileResultType& InCompiledResult, const ITextFilterExpressionContext& InContext, FText* OutErrorText) const;
+	CORE_API virtual bool EvaluateCompiledExpression(const ExpressionParser::CompileResultType& InCompiledResult, const ITextFilterExpressionContext& InContext, FText* OutErrorText) const;
 
 	/** Defines whether or not the expression parser can evaluate complex expressions */
 	ETextFilterExpressionEvaluatorMode ExpressionEvaluatorMode;
diff --git a/Engine/Source/Runtime/Core/Public/Misc/TextFilterUtils.h b/Engine/Source/Runtime/Core/Public/Misc/TextFilterUtils.h
--- a/Engine/Source/Runtime/Core/Public/Misc/TextFilterUtils.h
+++ b/Engine/Source/Runtime/Core/Public/Misc/TextFilterUtils.h
@@ -34,42 +34,42 @@ enum {NAME_WITH_NUMBER_SIZE	= NAME_SIZE + 16};
  * String used by the text filter.
  * The given string will be stored as uppercase since filter text always performs case-insensitive string comparisons, so this will minimize ToUpper calls.
  */
-class CORE_API FTextFilterString
+class FTextFilterString
 {
 public:
 	/** Default constructor */
-	FTextFilterString();
+	CORE_API FTextFilterString();
 
 	/** Move and copy constructors */
-	FTextFilterString(const FTextFilterString& Other);
-	FTextFilterString(FTextFilterString&& Other);
+	CORE_API FTextFilterString(const FTextFilterString& Other);
+	CORE_API FTextFilterString(FTextFilterString&& Other);
 
 	/** Construct from a string */
-	FTextFilterString(const FString& InString);
-	FTextFilterString(FString&& InString);
-	FTextFilterString(const TCHAR* InString);
+	CORE_API FTextFilterString(const FString& InString);
+	CORE_API FTextFilterString(FString&& InString);
+	CORE_API FTextFilterString(const TCHAR* InString);
 
 	/** Construct from a name */
-	FTextFilterString(const FName& InName);
+	CORE_API FTextFilterString(const FName& InName);
 
 	/** Move and copy assignment */
-	FTextFilterString& operator=(const FTextFilterString& Other);
-	FTextFilterString& operator=(FTextFilterString&& Other);
+	CORE_API FTextFilterString& operator=(const FTextFilterString& Other);
+	CORE_API FTextFilterString& operator=(FTextFilterString&& Other);
 
 	/** Compare this string against the other, using the text comparison mode provided */
-	bool CompareText(const FTextFilterString& InOther, const ETextFilterTextComparisonMode InTextComparisonMode) const;
+	CORE_API bool CompareText(const FTextFilterString& InOther, const ETextFilterTextComparisonMode InTextComparisonMode) const;
 
 	/** Compare this string against the other FString, using the text comparison mode provided */
-	bool CompareFString(const FString& InOther, const ETextFilterTextComparisonMode InTextComparisonMode) const;
+	CORE_API bool CompareFString(const FString& InOther, const ETextFilterTextComparisonMode InTextComparisonMode) const;
 
 	/** Compare this string against the other FName, using the text comparison mode provided */
-	bool CompareName(const FName& InOther, const ETextFilterTextComparisonMode InTextComparisonMode) const;
+	CORE_API bool CompareName(const FName& InOther, const ETextFilterTextComparisonMode InTextComparisonMode) const;
 
 	/** Are the two given strings able to be compared numberically? */
-	bool CanCompareNumeric(const FTextFilterString& InOther) const;
+	CORE_API bool CanCompareNumeric(const FTextFilterString& InOther) const;
 
 	/** Compare this string against the other, converting them to numbers and using the comparison operator provided - you should have tested CanCompareNumeric first! */
-	bool CompareNumeric(const FTextFilterString& InOther, const ETextFilterComparisonOperation InComparisonOperation) const;
+	CORE_API bool CompareNumeric(const FTextFilterString& InOther, const ETextFilterComparisonOperation InComparisonOperation) const;
 
 	/** Get the internal uppercase string of this filter string */
 	FORCEINLINE const FString& AsString() const
@@ -91,7 +91,7 @@ public:
 
 private:
 	/** Inline convert our internal string to uppercase */
-	void UppercaseInternalString();
+	CORE_API void UppercaseInternalString();
 
 	/** The uppercase string to use for comparisons */
 	FString InternalString;
