[Core]

diff --git a/Engine/Source/Runtime/Core/Public/Misc/TrackedActivity.h b/Engine/Source/Runtime/Core/Public/Misc/TrackedActivity.h
--- a/Engine/Source/Runtime/Core/Public/Misc/TrackedActivity.h
+++ b/Engine/Source/Runtime/Core/Public/Misc/TrackedActivity.h
@@ -25,7 +25,7 @@ enum class ETrackedActivityLight : uint8
  * When new console is enabled tracked activities show at the bottom of the window under the log
  * Tracked Activities can be created/updated/destroyed on multiple threads
  */
-class CORE_API FTrackedActivity : public TSharedFromThis<FTrackedActivity>
+class FTrackedActivity : public TSharedFromThis<FTrackedActivity>
 {
 public:
 	/**
@@ -49,33 +49,33 @@ public:
 	* @param Type Decides where activity information show in console. Activity is to the left, Info to the right
 	* @param SortValue Decides in what order within type the activity will show in console. Lower value means earlier
 	*/
-	FTrackedActivity(const TCHAR* Name, const TCHAR* Status = TEXT(""), ELight Light = ELight::None, EType Type = EType::Activity, int32 SortValue = 100);
+	CORE_API FTrackedActivity(const TCHAR* Name, const TCHAR* Status = TEXT(""), ELight Light = ELight::None, EType Type = EType::Activity, int32 SortValue = 100);
 
 	/** Dtor */
-	~FTrackedActivity();
+	CORE_API ~FTrackedActivity();
 
 
 	/**
 	* Pushes new status on to tracked activity, will require a pop to get back to previous status
 	* ShowParent can be used to make sure Status of parent is visible in front of status of pushed scope.
 	*/
-	uint32	Push(const TCHAR* Status, bool bShowParent = false, ELight Light = ELight::Inherit);
-	void	Pop();
+	CORE_API uint32	Push(const TCHAR* Status, bool bShowParent = false, ELight Light = ELight::Inherit);
+	CORE_API void	Pop();
 
 	/** Updates status. If Index is ~0u entry at top of stack will be updated (which is the one showing) */
-	void	Update(const TCHAR* Status, uint32 Index = ~0u);
-	void	Update(const TCHAR* Status, ELight Light, uint32 Index = ~0u);
-	void	Update(ELight Light, uint32 Index = ~0u);
+	CORE_API void	Update(const TCHAR* Status, uint32 Index = ~0u);
+	CORE_API void	Update(const TCHAR* Status, ELight Light, uint32 Index = ~0u);
+	CORE_API void	Update(ELight Light, uint32 Index = ~0u);
 
 
 	/**
 	* Process Engine Activity. General status of the Engine. Used by Engine initialization etc.
 	* By default, status light will stay yellow until main threads hits Tick update, then it turns green.
 	*/
-	static FTrackedActivity& GetEngineActivity();
+	static CORE_API FTrackedActivity& GetEngineActivity();
 
 	/** I/O Activity.Shows current I / O operation.If plugin / game have their own I / O, scopes will need to be manually added */
-	static FTrackedActivity& GetIOActivity();
+	static CORE_API FTrackedActivity& GetIOActivity();
 
 
 
@@ -92,7 +92,7 @@ public:
 	/**
 	* Traverses all FTrackedActivities in the order they were added.
 	*/
-	static void TraverseActivities(const TFunction<void(const FInfo& Info)>& Func);
+	static CORE_API void TraverseActivities(const TFunction<void(const FInfo& Info)>& Func);
 
 
 
@@ -107,7 +107,7 @@ public:
 	/** Register listener that can track adds, removes and changes when they happen.
 	* Listener will be called from the same thread as the event happens, so make sure listener is threadsafe.
 	*/
-	static void RegisterEventListener(TUniqueFunction<void(EEvent Event, const FInfo& Info)>&& Func, uint32 MaxDepth = ~0u);
+	static CORE_API void RegisterEventListener(TUniqueFunction<void(EEvent Event, const FInfo& Info)>&& Func, uint32 MaxDepth = ~0u);
 
 
 private:
@@ -121,11 +121,11 @@ private:
 /**
 * RAII class that calls push in ctor and pop in dtor.
 */
-class CORE_API FTrackedActivityScope
+class FTrackedActivityScope
 {
 public:
-	FTrackedActivityScope(FTrackedActivity& Activity, const TCHAR* Status, bool bShowParent = false, FTrackedActivity::ELight Light = FTrackedActivity::ELight::Inherit);
-	~FTrackedActivityScope();
+	CORE_API FTrackedActivityScope(FTrackedActivity& Activity, const TCHAR* Status, bool bShowParent = false, FTrackedActivity::ELight Light = FTrackedActivity::ELight::Inherit);
+	CORE_API ~FTrackedActivityScope();
 private:
 	FTrackedActivity& Activity;
 };
@@ -142,4 +142,4 @@ private:
 		#define UE_ENABLE_TRACKED_IO 1
 		#define UE_SCOPED_IO_ACTIVITY(...) FTrackedActivityScope ANONYMOUS_VARIABLE(IOActivity_)(FTrackedActivity::GetIOActivity(), __VA_ARGS__);
 	#endif
-#endif
\ No newline at end of file
+#endif
diff --git a/Engine/Source/Runtime/Core/Public/Misc/UProjectInfo.h b/Engine/Source/Runtime/Core/Public/Misc/UProjectInfo.h
--- a/Engine/Source/Runtime/Core/Public/Misc/UProjectInfo.h
+++ b/Engine/Source/Runtime/Core/Public/Misc/UProjectInfo.h
@@ -10,29 +10,29 @@
 /**
  * Dictionary of all the non-foreign projects for an engine installation, found by parsing .uprojectdirs files for source directories.
  */
-class CORE_API FUProjectDictionary
+class FUProjectDictionary
 {
 public:
 	/** Scans the engine root directory for all the known projects. */
-	FUProjectDictionary(const FString& InRootDir);
+	CORE_API FUProjectDictionary(const FString& InRootDir);
 	
 	/** Refreshes the list of known projects */
-	void Refresh();
+	CORE_API void Refresh();
 
 	/** Determines whether a project is a foreign project or not. */
-	bool IsForeignProject(const FString& ProjectFileName) const;
+	CORE_API bool IsForeignProject(const FString& ProjectFileName) const;
 
 	/** Gets the project filename for the given game. Empty if not found. */
-	FString GetRelativeProjectPathForGame(const TCHAR* GameName, const FString& BaseDir) const;
+	CORE_API FString GetRelativeProjectPathForGame(const TCHAR* GameName, const FString& BaseDir) const;
 
 	/** Gets the project filename for the given game. Empty if not found. */
-	FString GetProjectPathForGame(const TCHAR* GameName) const;
+	CORE_API FString GetProjectPathForGame(const TCHAR* GameName) const;
 
 	/** Gets a list of all the known projects. */
-	TArray<FString> GetProjectPaths() const;
+	CORE_API TArray<FString> GetProjectPaths() const;
 
 	/** Gets the project dictionary for the active engine installation. */
-	static FUProjectDictionary& GetDefault();
+	static CORE_API FUProjectDictionary& GetDefault();
 
 private:
 	/** The root directory for this engine installation */
diff --git a/Engine/Source/Runtime/Core/Public/Misc/URLRequestFilter.h b/Engine/Source/Runtime/Core/Public/Misc/URLRequestFilter.h
--- a/Engine/Source/Runtime/Core/Public/Misc/URLRequestFilter.h
+++ b/Engine/Source/Runtime/Core/Public/Misc/URLRequestFilter.h
@@ -25,7 +25,7 @@ namespace UE::Core
  * To allow both an absolute domain and all of its subdomains, provide two entries in the array,
  * one with the period prefix and one without. For example, "epicgames.com" and ".epicgames.com".
  */
-class CORE_API FURLRequestFilter
+class FURLRequestFilter
 {
 public:
 	using FRequestMap = TMap<const FString, const TArray<FString>>;
@@ -43,18 +43,18 @@ public:
 	 * !AllowedDomains=ClearArray
 	 * +AllowedDomains=epicgames.com
 	 */
-	FURLRequestFilter(const TCHAR* ConfigSectionRootName, const FString& ConfigFileName);
+	CORE_API FURLRequestFilter(const TCHAR* ConfigSectionRootName, const FString& ConfigFileName);
 
 	/**
 	 * Initializes the allowlist based on an explicit scheme/domain map
 	 */
-	FURLRequestFilter(const FRequestMap& InAllowedRequests);
-	FURLRequestFilter(FRequestMap&& InAllowedRequests);
+	CORE_API FURLRequestFilter(const FRequestMap& InAllowedRequests);
+	CORE_API FURLRequestFilter(FRequestMap&& InAllowedRequests);
 
 	/**
 	 * Update the allow list by corresponding section in a config file
 	 */
-	void UpdateConfig(const TCHAR* ConfigSectionRootName, const FString& ConfigFileName);
+	CORE_API void UpdateConfig(const TCHAR* ConfigSectionRootName, const FString& ConfigFileName);
 
 	/**
 	 * Check if allow list is empty
@@ -64,10 +64,10 @@ public:
 	/**
 	 * Matches InURL against the configued list of allowed URL schemes and domains. Returns true if the URL is allowed.
 	 */
-	bool IsRequestAllowed(FStringView InURL) const;
+	CORE_API bool IsRequestAllowed(FStringView InURL) const;
 
 private:
 	FRequestMap AllowedRequests;
 };
 
-}
\ No newline at end of file
+}
diff --git a/Engine/Source/Runtime/Core/Public/Misc/WildcardString.h b/Engine/Source/Runtime/Core/Public/Misc/WildcardString.h
--- a/Engine/Source/Runtime/Core/Public/Misc/WildcardString.h
+++ b/Engine/Source/Runtime/Core/Public/Misc/WildcardString.h
@@ -106,11 +106,11 @@ public:
 protected:
 
 	/** Holds the string terminator character. */
-	static CORE_API const TCHAR EndOfString = TCHAR('\0');
+	static const TCHAR EndOfString = TCHAR('\0');
 
 	/** Holds the wild card that matches exactly one character (default is '?'). */
-	static CORE_API const TCHAR ExactWildcard = TCHAR('?');
+	static const TCHAR ExactWildcard = TCHAR('?');
 
 	/** Holds the wild card that matches a sequence of characters (default is '*'). */
-	static CORE_API const TCHAR SequenceWildcard = TCHAR('*');
+	static const TCHAR SequenceWildcard = TCHAR('*');
 };
diff --git a/Engine/Source/Runtime/Core/Public/Modules/BuildVersion.h b/Engine/Source/Runtime/Core/Public/Modules/BuildVersion.h
--- a/Engine/Source/Runtime/Core/Public/Modules/BuildVersion.h
+++ b/Engine/Source/Runtime/Core/Public/Modules/BuildVersion.h
@@ -11,7 +11,7 @@ class FJsonObject;
 /**
  * Stores the version information associated with a build
  */
-class CORE_API FBuildVersion
+class FBuildVersion
 {
 public:
 	/**
@@ -67,37 +67,37 @@ public:
 	/**
 	 * Default constructor. Initializes the structure to empty.
 	 */
-	FBuildVersion();
+	CORE_API FBuildVersion();
 
 	/// <summary>
 	/// Gets the compatible changelist if set, otherwise the default compatible changelist
 	/// </summary>
 	/// <returns>The compatible changelist</returns>
-	int GetEffectiveCompatibleChangelist() const;
+	CORE_API int GetEffectiveCompatibleChangelist() const;
 
 	/// <summary>
 	/// Get an engine version object for this build version
 	/// </summary>
 	/// <returns>New engine version object</returns>
-	FEngineVersion GetEngineVersion() const;
+	CORE_API FEngineVersion GetEngineVersion() const;
 
 	/// <summary>
 	/// Get a compatible engine version object for this build version
 	/// </summary>
 	/// <returns>New engine version object</returns>
-	FEngineVersion GetCompatibleEngineVersion() const;
+	CORE_API FEngineVersion GetCompatibleEngineVersion() const;
 
 	/// <summary>
 	/// Get the default path to the build.version file on disk
 	/// </summary>
 	/// <returns>Path to the Build.version file</returns>
-	static FString GetDefaultFileName();
+	static CORE_API FString GetDefaultFileName();
 
 	/// <summary>
 	/// Get the path to the version file for the current executable.
 	/// </summary>
 	/// <returns>Path to the target's version file</returns>
-	static FString GetFileNameForCurrentExecutable();
+	static CORE_API FString GetFileNameForCurrentExecutable();
 	
 	/**
 	 * Try to read a version file from disk
@@ -106,5 +106,5 @@ public:
 	 * @param OutVersion The version information
 	 * @return True if the version was read successfully, false otherwise
 	 */
-	static bool TryRead(const FString& FileName, FBuildVersion& OutVersion);
+	static CORE_API bool TryRead(const FString& FileName, FBuildVersion& OutVersion);
 };
diff --git a/Engine/Source/Runtime/Core/Public/Modules/ModuleManager.h b/Engine/Source/Runtime/Core/Public/Modules/ModuleManager.h
--- a/Engine/Source/Runtime/Core/Public/Modules/ModuleManager.h
+++ b/Engine/Source/Runtime/Core/Public/Modules/ModuleManager.h
@@ -139,7 +139,7 @@ namespace UE::Core::Private
  * The module manager is used to load and unload modules, as well as to keep track of all of the
  * modules that are currently loaded. You can access this singleton using FModuleManager::Get().
  */
-class CORE_API FModuleManager
+class FModuleManager
 	: private FSelfRegisteringExec
 {
 public:
@@ -147,7 +147,7 @@ public:
 	/**
 	 * Destructor.
 	 */
-	~FModuleManager();
+	CORE_API ~FModuleManager();
 
 
 	/**
@@ -155,10 +155,10 @@ public:
 	 *
 	 * @return The module manager instance.
 	 */
-	static FModuleManager& Get( );
+	static CORE_API FModuleManager& Get( );
 
 	/** Destroys singleton if it exists. Get() must not be called after Destroy(). */
-	static void TearDown();
+	static CORE_API void TearDown();
 
 
 	/**
@@ -167,17 +167,17 @@ public:
 	 * @param InModuleName The name of the module to abandon.  Should not include path, extension or platform/configuration info.  This is just the "module name" part of the module file name.
 	 * @see IsModuleLoaded, LoadModule, LoadModuleWithFailureReason, UnloadModule
 	 */
-	void AbandonModule( const FName InModuleName );
+	CORE_API void AbandonModule( const FName InModuleName );
 
 	/**
 	 * Adds a module to our list of modules, unless it's already known.
 	 *
 	 * @param InModuleName The base name of the module file.  Should not include path, extension or platform/configuration info.  This is just the "name" part of the module file name.  Names should be globally unique.
 	 */
-	void AddModule( const FName InModuleName );
+	CORE_API void AddModule( const FName InModuleName );
 
 #if !IS_MONOLITHIC
-	void RefreshModuleFilenameFromManifest(const FName InModuleName);
+	CORE_API void RefreshModuleFilenameFromManifest(const FName InModuleName);
 #endif	// !IS_MONOLITHIC
 
 	/**
@@ -187,7 +187,7 @@ public:
 	 * @return 	The module, or nullptr if the module is not loaded.
 	 * @see GetModuleChecked, GetModulePtr
 	 */
-	IModuleInterface* GetModule( const FName InModuleName );
+	CORE_API IModuleInterface* GetModule( const FName InModuleName );
 
 	/**
 	 * Checks whether the specified module is currently loaded.
@@ -198,7 +198,7 @@ public:
 	 * @return true if module is currently loaded, false otherwise.
 	 * @see AbandonModule, LoadModule, LoadModuleWithFailureReason, UnloadModule
 	 */
-	bool IsModuleLoaded( const FName InModuleName ) const;
+	CORE_API bool IsModuleLoaded( const FName InModuleName ) const;
 
 	/**
 	 * Loads the specified module.
@@ -208,7 +208,7 @@ public:
 	 * @return The loaded module, or nullptr if the load operation failed.
 	 * @see AbandonModule, IsModuleLoaded, LoadModuleChecked, LoadModulePtr, LoadModuleWithFailureReason, UnloadModule
 	 */
-	IModuleInterface* LoadModule( const FName InModuleName, ELoadModuleFlags InLoadModuleFlags = ELoadModuleFlags::None );
+	CORE_API IModuleInterface* LoadModule( const FName InModuleName, ELoadModuleFlags InLoadModuleFlags = ELoadModuleFlags::None );
 
 	/**
 	 * Loads the specified module, checking to ensure it exists.
@@ -217,7 +217,7 @@ public:
 	 * @return The loaded module, or nullptr if the load operation failed.
 	 * @see AbandonModule, IsModuleLoaded, LoadModuleChecked, LoadModulePtr, LoadModuleWithFailureReason, UnloadModule
 	 */
-	IModuleInterface& LoadModuleChecked( const FName InModuleName );
+	CORE_API IModuleInterface& LoadModuleChecked( const FName InModuleName );
 
 	/**
 	 * Loads a module in memory then calls PostLoad.
@@ -227,7 +227,7 @@ public:
 	 * @return true on success, false otherwise.
 	 * @see UnloadOrAbandonModuleWithCallback
 	 */
-	bool LoadModuleWithCallback( const FName InModuleName, FOutputDevice &Ar );
+	CORE_API bool LoadModuleWithCallback( const FName InModuleName, FOutputDevice &Ar );
 
 	/**
 	 * Loads the specified module and returns a result.
@@ -238,7 +238,7 @@ public:
 	 * @return The loaded module (null if the load operation failed).
 	 * @see AbandonModule, IsModuleLoaded, LoadModule, LoadModuleChecked, LoadModulePtr, UnloadModule
 	 */
-	IModuleInterface* LoadModuleWithFailureReason( const FName InModuleName, EModuleLoadResult& OutFailureReason, ELoadModuleFlags InLoadModuleFlags = ELoadModuleFlags::None);
+	CORE_API IModuleInterface* LoadModuleWithFailureReason( const FName InModuleName, EModuleLoadResult& OutFailureReason, ELoadModuleFlags InLoadModuleFlags = ELoadModuleFlags::None);
 
 	/**
 	 * Queries information about a specific module name.
@@ -248,7 +248,7 @@ public:
 	 * @return true if the module was found and the OutModuleStatus is valid, false otherwise.
 	 * @see QueryModules
 	 */
-	bool QueryModule( const FName InModuleName, FModuleStatus& OutModuleStatus ) const;
+	CORE_API bool QueryModule( const FName InModuleName, FModuleStatus& OutModuleStatus ) const;
 
 	/**
 	 * Queries information about all of the currently known modules.
@@ -256,7 +256,7 @@ public:
 	 * @param OutModuleStatuses Status of all modules.
 	 * @see QueryModule
 	 */
-	void QueryModules( TArray<FModuleStatus>& OutModuleStatuses ) const;
+	CORE_API void QueryModules( TArray<FModuleStatus>& OutModuleStatuses ) const;
 
 	/**
 	 * Unloads a specific module
@@ -268,7 +268,7 @@ public:
 	 * @return true if module was unloaded successfully, false otherwise.
 	 * @see AbandonModule, IsModuleLoaded, LoadModule, LoadModuleWithFailureReason
 	 */
-	bool UnloadModule( const FName InModuleName, bool bIsShutdown = false );
+	CORE_API bool UnloadModule( const FName InModuleName, bool bIsShutdown = false );
 
 	/**
 	 * Calls PreUnload then either unloads or abandons a module in memory, depending on whether the module supports unloading.
@@ -277,7 +277,7 @@ public:
 	 * @param Ar The archive to receive error messages, if any.
 	 * @see LoadModuleWithCallback
 	 */
-	void UnloadOrAbandonModuleWithCallback( const FName InModuleName, FOutputDevice &Ar);
+	CORE_API void UnloadOrAbandonModuleWithCallback( const FName InModuleName, FOutputDevice &Ar);
 
 	/**
 	 * Calls PreUnload then abandons a module in memory.
@@ -285,12 +285,12 @@ public:
 	 * @param InModuleName The name of the module to unload.
 	 * @see LoadModuleWithCallback
 	 */
-	void AbandonModuleWithCallback( const FName InModuleName );
+	CORE_API void AbandonModuleWithCallback( const FName InModuleName );
 
 	/**
 	 * Add any extra search paths that may be required
 	 */
-	void AddExtraBinarySearchPaths();
+	CORE_API void AddExtraBinarySearchPaths();
 
 	/**
 	  * Gets a module by name, checking to ensure it exists.
@@ -311,7 +311,7 @@ public:
 	}
 
 private:
-	static IModuleInterface* GetModulePtr_Internal(FName ModuleName);
+	static CORE_API IModuleInterface* GetModulePtr_Internal(FName ModuleName);
 
 public:
 
@@ -364,7 +364,7 @@ public:
 	 * @param WildcardWithoutExtension Filename part (no path, no extension, no build config info) to search for.
 	 * @param OutModules List of modules found.
 	 */
-	void FindModules( const TCHAR* WildcardWithoutExtension, TArray<FName>& OutModules ) const;
+	CORE_API void FindModules( const TCHAR* WildcardWithoutExtension, TArray<FName>& OutModules ) const;
 
 	/**
 	 * Determines if a module with the given name exists, regardless of whether it is currently loaded.
@@ -372,14 +372,14 @@ public:
 	 * @param ModuleName Name of the module to look for.
 	 * @return Whether the module exists.
 	 */
-	bool ModuleExists(const TCHAR* ModuleName) const;
+	CORE_API bool ModuleExists(const TCHAR* ModuleName) const;
 
 	/**
 	 * Gets the number of loaded modules.
 	 *
 	 * @return The number of modules.
 	 */
-	int32 GetModuleCount( ) const;
+	CORE_API int32 GetModuleCount( ) const;
 
 	/**
 	 * Unloads modules during the shutdown process. Modules are unloaded in reverse order to when their StartupModule() FINISHES.
@@ -390,7 +390,7 @@ public:
 	 *
 	 * This method is Usually called at various points while exiting an application.
 	 */
-	void UnloadModulesAtShutdown( );
+	CORE_API void UnloadModulesAtShutdown( );
 
 
 	/** Delegate that's used by the module manager to initialize a registered module that we statically linked with (monolithic only) */
@@ -411,27 +411,27 @@ public:
 	 * Called by the engine at startup to let the Module Manager know that it's now
 	 * safe to process new UObjects discovered by loading C++ modules.
 	 */
-	void StartProcessingNewlyLoadedObjects();
+	CORE_API void StartProcessingNewlyLoadedObjects();
 
 	/** Adds an engine binaries directory. */
-	void AddBinariesDirectory(const TCHAR *InDirectory, bool bIsGameDirectory);
+	CORE_API void AddBinariesDirectory(const TCHAR *InDirectory, bool bIsGameDirectory);
 
 	/** Will load the binary without doing initialization.
 	 *  Calling this will not make it possible to unload the module again
 	 */
-	void LoadModuleBinaryOnly(FName ModuleName);
+	CORE_API void LoadModuleBinaryOnly(FName ModuleName);
 
 	/**
 	 *	Set the game binaries directory
 	 *
 	 *	@param InDirectory The game binaries directory.
 	 */
-	void SetGameBinariesDirectory(const TCHAR* InDirectory);
+	CORE_API void SetGameBinariesDirectory(const TCHAR* InDirectory);
 
 	/**
 	*	Gets the game binaries directory
 	*/
-	FString GetGameBinariesDirectory() const;
+	CORE_API FString GetGameBinariesDirectory() const;
 
 #if !IS_MONOLITHIC
 	/**
@@ -440,7 +440,7 @@ public:
 	 * @param InModuleName The base name of the module file.
 	 * @return true if module exists and is up to date, false otherwise.
 	 */
-	bool IsModuleUpToDate( const FName InModuleName ) const;
+	CORE_API bool IsModuleUpToDate( const FName InModuleName ) const;
 #endif
 
 	/**
@@ -450,27 +450,27 @@ public:
 	 * @param ModuleName Name of the loaded module to check.
 	 * @return True if the module was found to contain UObjects, or false if it did not (or wasn't loaded.)
 	 */
-	bool DoesLoadedModuleHaveUObjects( const FName ModuleName ) const;
+	CORE_API bool DoesLoadedModuleHaveUObjects( const FName ModuleName ) const;
 
 	/**
 	 * Gets the build configuration for compiling modules, as required by UBT.
 	 *
 	 * @return	Configuration name for UBT.
 	 */
-	static const TCHAR* GetUBTConfiguration( );
+	static CORE_API const TCHAR* GetUBTConfiguration( );
 
 #if !IS_MONOLITHIC
 	/** Gets the filename for a module. The return value is a full path of a module known to the module manager. */
-	FString GetModuleFilename(FName ModuleName) const;
+	CORE_API FString GetModuleFilename(FName ModuleName) const;
 
 	/** Sets the filename for a module. The module is not reloaded immediately, but the new name will be used for subsequent unload/load events. */
-	void SetModuleFilename(FName ModuleName, const FString& Filename);
+	CORE_API void SetModuleFilename(FName ModuleName, const FString& Filename);
 
 	/** Determines if any non-default module instances are loaded (eg. hot reloaded modules) */
-	bool HasAnyOverridenModuleFilename() const;
+	CORE_API bool HasAnyOverridenModuleFilename() const;
 
 	/** Save the current module manager's state into a file for bootstrapping other processes. */
-	void SaveCurrentStateForBootstrap(const TCHAR* Filename);
+	CORE_API void SaveCurrentStateForBootstrap(const TCHAR* Filename);
 #endif
 
 	/**
@@ -513,7 +513,7 @@ public:
 
 protected:
 	// FSelfRegisteringExec interface.
-	virtual bool Exec_Dev( UWorld* Inworld, const TCHAR* Cmd, FOutputDevice& Ar ) override;
+	CORE_API virtual bool Exec_Dev( UWorld* Inworld, const TCHAR* Cmd, FOutputDevice& Ar ) override;
 
 private:
 	friend TOptional<FModuleManager>& UE::Core::Private::GetModuleManagerSingleton();
@@ -525,7 +525,7 @@ public:
 	 *
 	 * Use the static Get function to return the singleton instance.
 	 */
-	FModuleManager(FPrivateToken);
+	CORE_API FModuleManager(FPrivateToken);
 	FModuleManager(const FModuleManager&) = delete;
 	FModuleManager& operator=(const FModuleManager&) = delete;
 
@@ -586,20 +586,20 @@ public:
 	/**
 	 * Generates a unique file name for the specified module name by adding a random suffix and checking for file collisions.
 	 */
-	void MakeUniqueModuleFilename( const FName InModuleName, FString& UniqueSuffix, FString& UniqueModuleFileName ) const;
+	CORE_API void MakeUniqueModuleFilename( const FName InModuleName, FString& UniqueSuffix, FString& UniqueModuleFileName ) const;
 
-	void AddModuleToModulesList(const FName InModuleName, FModuleManager::ModuleInfoRef& ModuleInfo);
+	CORE_API void AddModuleToModulesList(const FName InModuleName, FModuleManager::ModuleInfoRef& ModuleInfo);
 
 	/** Clears module path cache */
-	void ResetModulePathsCache();
+	CORE_API void ResetModulePathsCache();
 
 	friend FArchive& operator<<( FArchive& Ar, FModuleManager& ModuleManager );
 private:
-	static void WarnIfItWasntSafeToLoadHere(const FName InModuleName);
+	static CORE_API void WarnIfItWasntSafeToLoadHere(const FName InModuleName);
 
 	/** Thread safe module finding routine. */
-	ModuleInfoPtr FindModule(FName InModuleName);
-	ModuleInfoRef FindModuleChecked(FName InModuleName);
+	CORE_API ModuleInfoPtr FindModule(FName InModuleName);
+	CORE_API ModuleInfoRef FindModuleChecked(FName InModuleName);
 
 	FORCEINLINE TSharedPtr<const FModuleInfo, ESPMode::ThreadSafe> FindModule(FName InModuleName) const
 	{
@@ -613,20 +613,20 @@ private:
 
 #if !IS_MONOLITHIC
 	/** Finds modules matching a given name wildcard. */
-	void FindModulePaths(const TCHAR *NamePattern, TMap<FName, FString> &OutModulePaths) const;
+	CORE_API void FindModulePaths(const TCHAR *NamePattern, TMap<FName, FString> &OutModulePaths) const;
 
 	/** Finds modules within a given directory. */
-	void FindModulePathsInDirectory(const FString &DirectoryName, bool bIsGameDirectory, TMap<FName, FString> &OutModulePaths) const;
+	CORE_API void FindModulePathsInDirectory(const FString &DirectoryName, bool bIsGameDirectory, TMap<FName, FString> &OutModulePaths) const;
 
 	/** Serialize a bootstrapping state into or from an archive. */
-	void SerializeStateForBootstrap_Impl(FArchive& Ar);
+	CORE_API void SerializeStateForBootstrap_Impl(FArchive& Ar);
 
 	/** Refreshes the filename of a new module from the manifest */
-	void RefreshModuleFilenameFromManifestImpl(const FName InModuleName, FModuleInfo& ModuleInfo);
+	CORE_API void RefreshModuleFilenameFromManifestImpl(const FName InModuleName, FModuleInfo& ModuleInfo);
 #endif
 
 	/** Adds pending module initializer registrations to the StaticallyLinkedModuleInitializers map. */
-	void ProcessPendingStaticallyLinkedModuleInitializers() const;
+	CORE_API void ProcessPendingStaticallyLinkedModuleInitializers() const;
 
 private:
 	/** Map of all modules.  Maps the case-insensitive module name to information about that module, loaded or not. */
diff --git a/Engine/Source/Runtime/Core/Public/Modules/ModuleManifest.h b/Engine/Source/Runtime/Core/Public/Modules/ModuleManifest.h
--- a/Engine/Source/Runtime/Core/Public/Modules/ModuleManifest.h
+++ b/Engine/Source/Runtime/Core/Public/Modules/ModuleManifest.h
@@ -10,7 +10,7 @@
 /**
  * Stores a record of a built target, with all metadata that other tools may need to know about the build.
  */
-class CORE_API FModuleManifest
+class FModuleManifest
 {
 public:
 	FString BuildId;
@@ -19,7 +19,7 @@ public:
 	/**
 	 * Default constructor 
 	 */
-	FModuleManifest();
+	CORE_API FModuleManifest();
 
 	/**
 	 * Gets the path to a version manifest for the given folder.
@@ -28,12 +28,12 @@ public:
 	 * @param bIsGameFolder		Whether the directory is a game folder of not. Used to adjust the name if the application is running in DebugGame.
 	 * @return The filename of the version manifest.
 	 */
-	static FString GetFileName(const FString& DirectoryName, bool bIsGameFolder);
+	static CORE_API FString GetFileName(const FString& DirectoryName, bool bIsGameFolder);
 
 	/**
 	 * Read a version manifest from disk.
 	 *
 	 * @param FileName		Filename to read from
 	 */
-	static bool TryRead(const FString& FileName, FModuleManifest& Manifest);
+	static CORE_API bool TryRead(const FString& FileName, FModuleManifest& Manifest);
 };
diff --git a/Engine/Source/Runtime/Core/Public/ProfilingDebugging/ABTesting.h b/Engine/Source/Runtime/Core/Public/ProfilingDebugging/ABTesting.h
--- a/Engine/Source/Runtime/Core/Public/ProfilingDebugging/ABTesting.h
+++ b/Engine/Source/Runtime/Core/Public/ProfilingDebugging/ABTesting.h
@@ -24,14 +24,14 @@
 template<typename Allocator > class TBitArray;
 
 #if ENABLE_ABTEST
-class CORE_API FABTest
+class FABTest
 {
 public:
 
-	FABTest();
+	CORE_API FABTest();
 	
 	//returns a command to execute, if any.
-	const TCHAR* TickAndGetCommand();
+	CORE_API const TCHAR* TickAndGetCommand();
 
 	bool IsActive()
 	{
@@ -48,9 +48,9 @@ public:
 		return CurrentTest == 0;
 	}
 
-	static FABTest& Get();
+	static CORE_API FABTest& Get();
 
-	static void ABTestCmdFunc(const TArray<FString>& Args);
+	static CORE_API void ABTestCmdFunc(const TArray<FString>& Args);
 
 	static bool StaticIsActive()
 	{
@@ -59,13 +59,13 @@ public:
 
 private:
 
-	void StartFrameLog();
-	void FrameLogTick(double Delta);
+	CORE_API void StartFrameLog();
+	CORE_API void FrameLogTick(double Delta);
 	
-	void Start(FString* InABTestCmds, bool bScopeTest);
-	void Stop();
+	CORE_API void Start(FString* InABTestCmds, bool bScopeTest);
+	CORE_API void Stop();
 
-	const TCHAR* SwitchTest(int32 Index);
+	CORE_API const TCHAR* SwitchTest(int32 Index);
 	
 	
 
@@ -108,7 +108,7 @@ private:
 
 };
 
-class CORE_API FScopedABTimer : public FDurationTimer
+class FScopedABTimer : public FDurationTimer
 {
 public:
 	explicit FScopedABTimer() 
@@ -132,7 +132,7 @@ private:
 #define SCOPED_ABTEST_DOFIRSTTEST() FABTest::Get().GetDoFirstScopeTest()
 #else
 
-class CORE_API FABTest
+class FABTest
 {
 public:
 	static bool StaticIsActive()
diff --git a/Engine/Source/Runtime/Core/Public/ProfilingDebugging/CallstackTrace.h b/Engine/Source/Runtime/Core/Public/ProfilingDebugging/CallstackTrace.h
--- a/Engine/Source/Runtime/Core/Public/ProfilingDebugging/CallstackTrace.h
+++ b/Engine/Source/Runtime/Core/Public/ProfilingDebugging/CallstackTrace.h
@@ -72,7 +72,7 @@ inline bool  CallstackTrace_ResolveFullCallStack()
 /*
  * Callstack Trace scope for override CallStack
  */
-class CORE_API FCallStackTraceLimitResolveScope
+class FCallStackTraceLimitResolveScope
 {
 public:
 	FORCENOINLINE FCallStackTraceLimitResolveScope()
diff --git a/Engine/Source/Runtime/Core/Public/ProfilingDebugging/CsvProfiler.h b/Engine/Source/Runtime/Core/Public/ProfilingDebugging/CsvProfiler.h
--- a/Engine/Source/Runtime/Core/Public/ProfilingDebugging/CsvProfiler.h
+++ b/Engine/Source/Runtime/Core/Public/ProfilingDebugging/CsvProfiler.h
@@ -424,7 +424,7 @@ public:
 
 	CORE_API void SetDeviceProfileName(FString InDeviceProfileName);
 
-	CORE_API FString GetOutputFilename() const { return OutputFilename; }
+	FString GetOutputFilename() const { return OutputFilename; }
 	CORE_API TMap<FString, FString> GetMetadataMapCopy();
 
 	CORE_API static bool IsWaitTrackingEnabledOnCurrentThread();
@@ -453,12 +453,12 @@ public:
 	DECLARE_MULTICAST_DELEGATE(FOnCSVProfileEndFrame);
 	FOnCSVProfileEndFrame& OnCSVProfileEndFrame() { return OnCSVProfileEndFrameDelegate; }
 
-	CORE_API void SetRenderThreadId(uint32 InRenderThreadId)
+	void SetRenderThreadId(uint32 InRenderThreadId)
 	{
 		RenderThreadId = InRenderThreadId;
 	}
 
-	CORE_API void SetRHIThreadId(uint32 InRHIThreadId)
+	void SetRHIThreadId(uint32 InRHIThreadId)
 	{
 		RHIThreadId = InRHIThreadId;
 	}
@@ -617,7 +617,7 @@ public:
 	const char * StatName;
 };
 
-struct CORE_API FCsvCategory
+struct FCsvCategory
 {
 	FCsvCategory() : Index(-1) {}
 	FCsvCategory(const TCHAR* CategoryString, bool bDefaultValue, bool bIsGlobal = false)
@@ -634,4 +634,4 @@ struct CORE_API FCsvCategory
 CSV_DECLARE_CATEGORY_MODULE_EXTERN(CORE_API, Exclusive);
 
 
-#endif //CSV_PROFILER
\ No newline at end of file
+#endif //CSV_PROFILER
diff --git a/Engine/Source/Runtime/Core/Public/ProfilingDebugging/ExternalProfiler.h b/Engine/Source/Runtime/Core/Public/ProfilingDebugging/ExternalProfiler.h
--- a/Engine/Source/Runtime/Core/Public/ProfilingDebugging/ExternalProfiler.h
+++ b/Engine/Source/Runtime/Core/Public/ProfilingDebugging/ExternalProfiler.h
@@ -20,13 +20,13 @@
  *
  * Interface to various external profiler API functions, dynamically linked
  */
-class CORE_API FExternalProfiler : public IModularFeature
+class FExternalProfiler : public IModularFeature
 {
 
 public:
 
 	/** Constructor */
-	FExternalProfiler();
+	CORE_API FExternalProfiler();
 
 	/** Empty virtual destructor. */
 	virtual ~FExternalProfiler()
@@ -34,10 +34,10 @@ public:
 	}
 
 	/** Pauses profiling. */
-	void PauseProfiler();
+	CORE_API void PauseProfiler();
 
 	/** Resumes profiling. */
-	void ResumeProfiler();
+	CORE_API void ResumeProfiler();
 
 	/**
 	 * Profiler interface.
@@ -59,7 +59,7 @@ public:
 	virtual const TCHAR* GetProfilerName() const = 0;
 
 	/** @return Returns the name to use for any profiler registered as a modular feature usable by this system */
-	static FName GetFeatureName();
+	static CORE_API FName GetFeatureName();
 
 	/** Starts a scoped event specific to the profiler. */
 	virtual void StartScopedEvent(const struct FColor& Color, const TCHAR* Text) {};
@@ -84,19 +84,19 @@ private:
 	friend class FScopedExternalProfilerBase;
 };
 
-class CORE_API FActiveExternalProfilerBase
+class FActiveExternalProfilerBase
 {
 public:	
 
 	static FExternalProfiler* GetActiveProfiler() { return ActiveProfiler;	};
 
-	static FExternalProfiler* InitActiveProfiler();
+	static CORE_API FExternalProfiler* InitActiveProfiler();
 private:
 	/** Static: True if we've tried to initialize a profiler already */
-	static bool bDidInitialize;
+	static CORE_API bool bDidInitialize;
 
 	/** Static: Active profiler instance that we're using */
-	static FExternalProfiler* ActiveProfiler;
+	static CORE_API FExternalProfiler* ActiveProfiler;
 };
 
 /**
@@ -214,4 +214,4 @@ public:
 #define SCOPE_PROFILER_INCLUDER(X)
 #define SCOPE_PROFILER_EXCLUDER(X)
 
-#endif	// !UE_EXTERNAL_PROFILING_ENABLED
\ No newline at end of file
+#endif	// !UE_EXTERNAL_PROFILING_ENABLED
diff --git a/Engine/Source/Runtime/Core/Public/ProfilingDebugging/Histogram.h b/Engine/Source/Runtime/Core/Public/ProfilingDebugging/Histogram.h
--- a/Engine/Source/Runtime/Core/Public/ProfilingDebugging/Histogram.h
+++ b/Engine/Source/Runtime/Core/Public/ProfilingDebugging/Histogram.h
@@ -30,16 +30,16 @@ DECLARE_LOG_CATEGORY_EXTERN(LogHistograms, Log, All);
  struct FHistogramBuilder;
 
  /** Fairly generic histogram for values that have natural lower bound and possibly no upper bound, e.g., frame time */
-struct CORE_API FHistogram
+struct FHistogram
 {
 	/** Inits histogram with linear, equally sized bins */
-	void InitLinear(double MinTime, double MaxTime, double BinSize);
+	CORE_API void InitLinear(double MinTime, double MaxTime, double BinSize);
 
 	/** Inits histogram to mimic our existing hitch buckets */
-	void InitHitchTracking();
+	CORE_API void InitHitchTracking();
 
 	/** Inits histogram with the specified bin boundaries, with the final bucket extending to infinity (e.g., passing in 0,5 creates a [0..5) bucket and a [5..infinity) bucket) */
-	void InitFromArray(TArrayView<const double> Thresholds);
+	CORE_API void InitFromArray(TArrayView<const double> Thresholds);
 
 	/** Inits histogram with the specified bin boundaries, with the final bucket extending to infinity (e.g., passing in 0,5 creates a [0..5) bucket and a [5..infinity) bucket) */
 	FORCEINLINE void InitFromArray(std::initializer_list<double> Thresholds)
@@ -48,7 +48,7 @@ struct CORE_API FHistogram
 	}
 
 	/** Resets measurements, without resetting the configured bins. */
-	void Reset();
+	CORE_API void Reset();
 
 	/** Adds an observed measurement. */
 	inline void AddMeasurement(double Value)
@@ -57,33 +57,33 @@ struct CORE_API FHistogram
 	}
 
 	/** Adds an observed measurement (with a different thresholding key than the measurement, e.g., when accumulating time spent in a chart keyed on framerate). */
-	void AddMeasurement(double ValueForBinning, double MeasurementValue);
+	CORE_API void AddMeasurement(double ValueForBinning, double MeasurementValue);
 
 	/** Prints histogram contents to the log. */
-	void DumpToLog(const FString& HistogramName);
+	CORE_API void DumpToLog(const FString& HistogramName);
 
 	/** 
 	 * Returns a string in a Json format: [{"Bin":"BinName","Count":Count,"Sum":Sum},...]. 
 	 * Bucket name is constructed by calling ConvertBinToLabel on the MinValue and UpperBound for each bin. 
 	 * Convert function is used to allow the bin range, which is stored as a double, to be printed prettily.
 	 */
-	FString DumpToJsonString(TFunctionRef<FString (double, double)> ConvertBinToLabel) const;
+	CORE_API FString DumpToJsonString(TFunctionRef<FString (double, double)> ConvertBinToLabel) const;
 
 	/** Same as DumpToJsonString but uses a DefaultConvertBinToLabel. */
-	FString DumpToJsonString() const;
+	CORE_API FString DumpToJsonString() const;
 
 	/** 
 	* Returns a string in a Json format: [{"BinName":{"Count":Count,"Sum":Sum}},...]. 
 	* Bucket name is constructed by calling ConvertBinToLabel on the MinValue and UpperBound for each bin. 
 	* Convert function is used to allow the bin range, which is stored as a double, to be printed prettily.
 	*/
-	FString DumpToJsonString2(TFunctionRef<FString (double, double)> ConvertBinToLabel) const;
+	CORE_API FString DumpToJsonString2(TFunctionRef<FString (double, double)> ConvertBinToLabel) const;
 
 	/** Same as DumpToJsonString2 but uses a DefaultConvertBinToLabel. */
-	FString DumpToJsonString2() const;
+	CORE_API FString DumpToJsonString2() const;
 
 	/** Default stringifier for bins for use with DumpToJsonString. Truncates to int and uses Plus as the suffix for the last bin. ie. [0.0, 3.75, 9.8] -> 0_3, 3_9, 9_Plus */
-	static FString DefaultConvertBinToLabel(double MinValue, double UpperBound);
+	static CORE_API FString DefaultConvertBinToLabel(double MinValue, double UpperBound);
 
 	/** Gets number of bins. */
 	inline int32 GetNumBins() const
diff --git a/Engine/Source/Runtime/Core/Public/ProfilingDebugging/InstanceCounter.h b/Engine/Source/Runtime/Core/Public/ProfilingDebugging/InstanceCounter.h
--- a/Engine/Source/Runtime/Core/Public/ProfilingDebugging/InstanceCounter.h
+++ b/Engine/Source/Runtime/Core/Public/ProfilingDebugging/InstanceCounter.h
@@ -27,7 +27,7 @@ class FOutputDevice;
 	};
 
 */
-class CORE_API FInstanceCountingObject
+class FInstanceCountingObject
 {
 public:
 
@@ -35,22 +35,22 @@ public:
 	 *	Constructor, though geneerally these objects should be created using the COUNT_INSTANCES_ 
 	 *	macro.
 	 */
-	FInstanceCountingObject(const TCHAR* InName, bool InLogConstruction=false);
+	CORE_API FInstanceCountingObject(const TCHAR* InName, bool InLogConstruction=false);
 
 	/**
 	 *	Copy-constructor for assigment
 	 */
-	FInstanceCountingObject(const FInstanceCountingObject& RHS);
+	CORE_API FInstanceCountingObject(const FInstanceCountingObject& RHS);
 
 	/**	Destructor */
-	virtual ~FInstanceCountingObject();
+	CORE_API virtual ~FInstanceCountingObject();
 
 	/**
 	 * Returns the count of instances with "Name". 
 	 *
 	 * @param: 	Name		
 	 * @return: int32		
-	*/static int32 GetInstanceCount(const TCHAR* Name);
+	*/static CORE_API int32 GetInstanceCount(const TCHAR* Name);
 	
 	/**
 	 * Dumps stats for all counted instances to the provided output device. This
@@ -58,15 +58,15 @@ public:
 	 *
 	 * @param: 	OutputDevice		
 	 * @return: void		
-	*/static void LogCounts(FOutputDevice& OutputDevice);
+	*/static CORE_API void LogCounts(FOutputDevice& OutputDevice);
 	
 protected:
 
 	/**	Increments stats for objects of this type */
-	void IncrementStats();
+	CORE_API void IncrementStats();
 
 	/**	Decrements stats for objects of this type */
-	void DecrementStats();
+	CORE_API void DecrementStats();
 
 protected:
 	
@@ -79,7 +79,7 @@ protected:
 	/**
 	 *	Vars used by our singleton
 	 */
-	struct CORE_API FGlobalVars
+	struct FGlobalVars
 	{
 		TMap<FName, int32>	InstanceCounts;
 		FCriticalSection	Mutex;
@@ -89,8 +89,8 @@ protected:
 	 *	Vars are stored as a pointer and initialized on demand due to avoid dependencies
 	 *	on global crot order
 	 */
-	static FGlobalVars* Globals;
-	static FGlobalVars& GetGlobals();
+	static CORE_API FGlobalVars* Globals;
+	static CORE_API FGlobalVars& GetGlobals();
 };
 
 
diff --git a/Engine/Source/Runtime/Core/Public/ProfilingDebugging/LoadTimeTracker.h b/Engine/Source/Runtime/Core/Public/ProfilingDebugging/LoadTimeTracker.h
--- a/Engine/Source/Runtime/Core/Public/ProfilingDebugging/LoadTimeTracker.h
+++ b/Engine/Source/Runtime/Core/Public/ProfilingDebugging/LoadTimeTracker.h
@@ -31,19 +31,19 @@
 #define ENABLE_LOADTIME_RAW_TIMINGS 0
 
 /** High level load time tracker utility (such as initial engine startup or game specific timings) */
-class CORE_API FLoadTimeTracker
+class FLoadTimeTracker
 {
 public:
-	static FLoadTimeTracker& Get();
+	static CORE_API FLoadTimeTracker& Get();
 
 	/** Adds a scoped time for a given label.  Records each instance individually */
-	void ReportScopeTime(double ScopeTime, const FName ScopeLabel);
+	CORE_API void ReportScopeTime(double ScopeTime, const FName ScopeLabel);
 
 	/** Gets/adds a scoped time for a given label and instance. Records each instance individually */
-	double& GetScopeTimeAccumulator(const FName& ScopeLabel, const FName& ScopeInstance);
+	CORE_API double& GetScopeTimeAccumulator(const FName& ScopeLabel, const FName& ScopeInstance);
 
 	/** Prints out total time and individual times */
-	void DumpHighLevelLoadTimes() const;
+	CORE_API void DumpHighLevelLoadTimes() const;
 
 	static void DumpHighLevelLoadTimesStatic()
 	{
@@ -55,31 +55,31 @@ public:
 		return TimeInfo;
 	}
 
-	void ResetHighLevelLoadTimes();
+	CORE_API void ResetHighLevelLoadTimes();
 
 	/** Prints out raw load times for individual timers */
-	void DumpRawLoadTimes() const;
+	CORE_API void DumpRawLoadTimes() const;
 
 	static void DumpRawLoadTimesStatic()
 	{
 		Get().DumpRawLoadTimes();
 	}
 
-	void ResetRawLoadTimes();
+	CORE_API void ResetRawLoadTimes();
 
 	static void ResetRawLoadTimesStatic()
 	{
 		Get().ResetRawLoadTimes();
 	}
 
-	void StartAccumulatedLoadTimes();
+	CORE_API void StartAccumulatedLoadTimes();
 
 	static void StartAccumulatedLoadTimesStatic()
 	{
 		Get().StartAccumulatedLoadTimes();
 	}
 
-	void StopAccumulatedLoadTimes();
+	CORE_API void StopAccumulatedLoadTimes();
 
 	static void StopAccumulatedLoadTimesStatic()
 	{
@@ -257,15 +257,15 @@ private:
 	/** We dont normally track accumulated load time info, only when this flag is true */
 	bool bAccumulating;
 private:
-	FLoadTimeTracker();
+	CORE_API FLoadTimeTracker();
 };
 
 /** Scoped helper class for tracking accumulated object times */
-struct CORE_API FScopedLoadTimeAccumulatorTimer : public FScopedDurationTimer
+struct FScopedLoadTimeAccumulatorTimer : public FScopedDurationTimer
 {
-	static double DummyTimer;
+	static CORE_API double DummyTimer;
 
-	FScopedLoadTimeAccumulatorTimer(const FName& InTimerName, const FName& InInstanceName);
+	CORE_API FScopedLoadTimeAccumulatorTimer(const FName& InTimerName, const FName& InInstanceName);
 };
 
 #if ENABLE_LOADTIME_TRACKING
diff --git a/Engine/Source/Runtime/Core/Public/ProfilingDebugging/PlatformFileTrace.h b/Engine/Source/Runtime/Core/Public/ProfilingDebugging/PlatformFileTrace.h
--- a/Engine/Source/Runtime/Core/Public/ProfilingDebugging/PlatformFileTrace.h
+++ b/Engine/Source/Runtime/Core/Public/ProfilingDebugging/PlatformFileTrace.h
@@ -14,22 +14,22 @@
 
 #if PLATFORMFILETRACE_ENABLED
 
-struct CORE_API FPlatformFileTrace
+struct FPlatformFileTrace
 {
-	static void BeginOpen(const TCHAR* Path);
-	static void EndOpen(uint64 FileHandle);
-	static void FailOpen(const TCHAR* Path);
-	static void BeginReOpen(uint64 OldFileHandle);
-	static void EndReOpen(uint64 NewFileHandle);
-	static void BeginClose(uint64 FileHandle);
-	static void EndClose(uint64 FileHandle);
-	static void FailClose(uint64 FileHandle);
-	static void BeginRead(uint64 ReadHandle, uint64 FileHandle, uint64 Offset, uint64 Size);
-	static void EndRead(uint64 ReadHandle, uint64 SizeRead);
-	static void BeginWrite(uint64 WriteHandle, uint64 FileHandle, uint64 Offset, uint64 Size);
-	static void EndWrite(uint64 WriteHandle, uint64 SizeWritten);
-
-	static uint32 GetOpenFileHandleCount();
+	static CORE_API void BeginOpen(const TCHAR* Path);
+	static CORE_API void EndOpen(uint64 FileHandle);
+	static CORE_API void FailOpen(const TCHAR* Path);
+	static CORE_API void BeginReOpen(uint64 OldFileHandle);
+	static CORE_API void EndReOpen(uint64 NewFileHandle);
+	static CORE_API void BeginClose(uint64 FileHandle);
+	static CORE_API void EndClose(uint64 FileHandle);
+	static CORE_API void FailClose(uint64 FileHandle);
+	static CORE_API void BeginRead(uint64 ReadHandle, uint64 FileHandle, uint64 Offset, uint64 Size);
+	static CORE_API void EndRead(uint64 ReadHandle, uint64 SizeRead);
+	static CORE_API void BeginWrite(uint64 WriteHandle, uint64 FileHandle, uint64 Offset, uint64 Size);
+	static CORE_API void EndWrite(uint64 WriteHandle, uint64 SizeWritten);
+
+	static CORE_API uint32 GetOpenFileHandleCount();
 };
 
 #define TRACE_PLATFORMFILE_BEGIN_OPEN(Path) \
@@ -83,4 +83,4 @@ struct CORE_API FPlatformFileTrace
 #define TRACE_PLATFORMFILE_BEGIN_WRITE(WriteHandle, FileHandle, Offset, Size)
 #define TRACE_PLATFORMFILE_END_WRITE(WriteHandle, SizeWritten)
 
-#endif
\ No newline at end of file
+#endif
diff --git a/Engine/Source/Runtime/Core/Public/ProfilingDebugging/ResourceSize.h b/Engine/Source/Runtime/Core/Public/ProfilingDebugging/ResourceSize.h
--- a/Engine/Source/Runtime/Core/Public/ProfilingDebugging/ResourceSize.h
+++ b/Engine/Source/Runtime/Core/Public/ProfilingDebugging/ResourceSize.h
@@ -27,89 +27,89 @@ namespace EResourceSizeMode
  * Struct used to count up the amount of memory used by a resource.
  * This is typically used for assets via UObject::GetResourceSizeEx(...).
  */
-struct CORE_API FResourceSizeEx
+struct FResourceSizeEx
 {
 public:
 
 	/**
 	 * Default constructor. 
 	 */
-	explicit FResourceSizeEx();
+	CORE_API explicit FResourceSizeEx();
 
 	/**
 	 * Construct using a given mode. 
 	 */
-	explicit FResourceSizeEx(const EResourceSizeMode::Type InResourceSizeMode);
+	CORE_API explicit FResourceSizeEx(const EResourceSizeMode::Type InResourceSizeMode);
 
 	/**
 	 * Construct from known sizes. 
 	 */
-	FResourceSizeEx(const EResourceSizeMode::Type InResourceSizeMode, const SIZE_T InDedicatedSystemMemoryBytes, const SIZE_T InDedicatedVideoMemoryBytes);
+	CORE_API FResourceSizeEx(const EResourceSizeMode::Type InResourceSizeMode, const SIZE_T InDedicatedSystemMemoryBytes, const SIZE_T InDedicatedVideoMemoryBytes);
 
 	/**
 	 * Construct from legacy unknown size.
 	 * Deliberately explicit to avoid accidental use.
 	 */
-	FResourceSizeEx(const EResourceSizeMode::Type InResourceSizeMode, const SIZE_T InUnknownMemoryBytes);
+	CORE_API FResourceSizeEx(const EResourceSizeMode::Type InResourceSizeMode, const SIZE_T InUnknownMemoryBytes);
 
-	void LogSummary(FOutputDevice& Ar) const;
+	CORE_API void LogSummary(FOutputDevice& Ar) const;
 
 	/**
 	 * Get the type of resource size held in this struct.
 	 */
-	EResourceSizeMode::Type GetResourceSizeMode() const;
+	CORE_API EResourceSizeMode::Type GetResourceSizeMode() const;
 
-	FResourceSizeEx& AddDedicatedSystemMemoryBytes(const FName& Tag, const SIZE_T InMemoryBytes);
+	CORE_API FResourceSizeEx& AddDedicatedSystemMemoryBytes(const FName& Tag, const SIZE_T InMemoryBytes);
 
 	/**
 	 * Add the given number of bytes to the dedicated system memory count.
 	 * @see DedicatedSystemMemoryBytes for a description of that memory type.
 	 */
-	FResourceSizeEx& AddDedicatedSystemMemoryBytes(const SIZE_T InMemoryBytes);
+	CORE_API FResourceSizeEx& AddDedicatedSystemMemoryBytes(const SIZE_T InMemoryBytes);
 
 	/**
 	 * Get the number of bytes allocated from dedicated system memory.
 	 * @see DedicatedSystemMemoryBytes for a description of that memory type.
 	 */
-	SIZE_T GetDedicatedSystemMemoryBytes() const;
+	CORE_API SIZE_T GetDedicatedSystemMemoryBytes() const;
 
-	FResourceSizeEx& AddDedicatedVideoMemoryBytes(const FName& Tag, const SIZE_T InMemoryBytes);
+	CORE_API FResourceSizeEx& AddDedicatedVideoMemoryBytes(const FName& Tag, const SIZE_T InMemoryBytes);
 
 	/**
 	 * Add the given number of bytes to the dedicated video memory count.
 	 * @see DedicatedVideoMemoryBytes for a description of that memory type.
 	 */
-	FResourceSizeEx& AddDedicatedVideoMemoryBytes(const SIZE_T InMemoryBytes);
+	CORE_API FResourceSizeEx& AddDedicatedVideoMemoryBytes(const SIZE_T InMemoryBytes);
 
 	/**
 	 * Get the number of bytes allocated from dedicated video memory.
 	 * @see DedicatedVideoMemoryBytes for a description of that memory type.
 	 */
-	SIZE_T GetDedicatedVideoMemoryBytes() const;
+	CORE_API SIZE_T GetDedicatedVideoMemoryBytes() const;
 
-	FResourceSizeEx& AddUnknownMemoryBytes(const FName& Tag, const SIZE_T InMemoryBytes);
+	CORE_API FResourceSizeEx& AddUnknownMemoryBytes(const FName& Tag, const SIZE_T InMemoryBytes);
 
 	/**
 	 * Add the given number of bytes to the unknown memory count.
 	 * @see UnknownMemoryBytes for a description of that memory type.
 	 */
-	FResourceSizeEx& AddUnknownMemoryBytes(const SIZE_T InMemoryBytes);
+	CORE_API FResourceSizeEx& AddUnknownMemoryBytes(const SIZE_T InMemoryBytes);
 
 	/**
 	 * Get the number of bytes allocated from unknown memory.
 	 * @see UnknownMemoryBytes for a description of that memory type.
 	 */
-	SIZE_T GetUnknownMemoryBytes() const;
+	CORE_API SIZE_T GetUnknownMemoryBytes() const;
 
 	/**
 	 * Get the total number of bytes allocated from any memory.
 	 */
-	SIZE_T GetTotalMemoryBytes() const;
+	CORE_API SIZE_T GetTotalMemoryBytes() const;
 
 	/**
 	 * Add another FResourceSizeEx to this one.
 	 */
-	FResourceSizeEx& operator+=(const FResourceSizeEx& InRHS);
+	CORE_API FResourceSizeEx& operator+=(const FResourceSizeEx& InRHS);
 
 private:
 	/**
diff --git a/Engine/Source/Runtime/Core/Public/ProfilingDebugging/ScopedDebugInfo.h b/Engine/Source/Runtime/Core/Public/ProfilingDebugging/ScopedDebugInfo.h
--- a/Engine/Source/Runtime/Core/Public/ProfilingDebugging/ScopedDebugInfo.h
+++ b/Engine/Source/Runtime/Core/Public/ProfilingDebugging/ScopedDebugInfo.h
@@ -10,7 +10,7 @@
 #include "Containers/UnrealString.h"
 
 /** Debug info that is stored on the stack and handled by stack unwinding. */
-class CORE_API FScopedDebugInfo
+class FScopedDebugInfo
 {
 public:
 
@@ -23,10 +23,10 @@ public:
 	FScopedDebugInfo* const NextOuterInfo;
 
 	/** Initialization constructor. */
-	FScopedDebugInfo(int32 InNumReplacedOuterCalls);
+	CORE_API FScopedDebugInfo(int32 InNumReplacedOuterCalls);
 
 	/** Destructor. */
-	virtual ~FScopedDebugInfo();
+	CORE_API virtual ~FScopedDebugInfo();
 
 	/** @return The "function name" to display on the call stack for this debug info. */
 	virtual FString GetFunctionName() const = 0;
@@ -38,6 +38,6 @@ public:
 	virtual int32 GetLineNumber() const = 0;
 
 	/** Accesses the list of debug infos on the stack in this thread. */
-	static FScopedDebugInfo* GetDebugInfoStack();
+	static CORE_API FScopedDebugInfo* GetDebugInfoStack();
 };
 
diff --git a/Engine/Source/Runtime/Core/Public/ProfilingDebugging/StallDetector.h b/Engine/Source/Runtime/Core/Public/ProfilingDebugging/StallDetector.h
--- a/Engine/Source/Runtime/Core/Public/ProfilingDebugging/StallDetector.h
+++ b/Engine/Source/Runtime/Core/Public/ProfilingDebugging/StallDetector.h
@@ -137,7 +137,7 @@ namespace UE
 	/**
 	 * This structure is allocated as a function-local static to persist stats across stall checks.
 	 */
-	class CORE_API FStallDetectorStats
+	class FStallDetectorStats
 	{
 		friend class FStallDetectorRunnable;
 
@@ -145,17 +145,17 @@ namespace UE
 		/**
 		* Construct a stats object for use with a stall detector
 		*/
-		FStallDetectorStats(const TCHAR* InName, const double InBudgetSeconds, const EStallDetectorReportingMode InReportingMode);
+		CORE_API FStallDetectorStats(const TCHAR* InName, const double InBudgetSeconds, const EStallDetectorReportingMode InReportingMode);
 
 		/**
 		* Destruct a stats object, just removes it from the instance tracking
 		*/
-		~FStallDetectorStats();
+		CORE_API ~FStallDetectorStats();
 
 		/**
 		* Stall detector scope has ended, record the interval
 		*/
-		void OnStallCompleted(double InOverageSeconds);
+		CORE_API void OnStallCompleted(double InOverageSeconds);
 
 		struct TabulatedResult
 		{
@@ -165,7 +165,7 @@ namespace UE
 			double OverageSeconds = 0.0;
 			double OverageRatio = 0.0;
 		};
-		static void TabulateStats(TArray<TabulatedResult>& TabulatedResults);
+		static CORE_API void TabulateStats(TArray<TabulatedResult>& TabulatedResults);
 
 		// The name to refer to this callsite in the codebase
 		const TCHAR* Name;
@@ -180,10 +180,10 @@ namespace UE
 		bool bReported;
 
 		// The total number of times all callsites have been triggered
-		static FCountersTrace::FCounterAtomicInt TotalTriggeredCount;
+		static CORE_API FCountersTrace::FCounterAtomicInt TotalTriggeredCount;
 
 		// The total number of reports that have been sent
-		static FCountersTrace::FCounterAtomicInt TotalReportedCount;
+		static CORE_API FCountersTrace::FCounterAtomicInt TotalReportedCount;
 
 	private:
 		// The number of times this callsite has been triggered
@@ -202,33 +202,33 @@ namespace UE
 	 * FStallTimer tracks remaining time from some budget, and access is protected from multiple threads.
 	 * Owned by FStallDetector
 	 **/
-	class CORE_API FStallTimer
+	class FStallTimer
 	{
 	public:
 		/**
 		* Construct the timer, you must call Reset() before the other methods
 		*/
-		FStallTimer();
+		CORE_API FStallTimer();
 
 		/**
 		* Initialize the timer with the current timestamp and remaining time to track
 		**/
-		void Reset(const double InSeconds, const double InRemainingSeconds);
+		CORE_API void Reset(const double InSeconds, const double InRemainingSeconds);
 
 		/**
 		* Perform a check against the specified timestamp and provide some useful details
 		**/
-		void Check(const double InSeconds, double& OutDeltaSeconds, double& OutOverageSeconds);
+		CORE_API void Check(const double InSeconds, double& OutDeltaSeconds, double& OutOverageSeconds);
 
 		/**
 		* Pause the timer
 		**/
-		void Pause(const double InSeconds);
+		CORE_API void Pause(const double InSeconds);
 
 		/**
 		* Resume the timer
 		**/
-		void Resume(const double InSeconds);
+		CORE_API void Resume(const double InSeconds);
 
 	private:
 		// To track pause/unpause calls
@@ -247,7 +247,7 @@ namespace UE
 	/**
 	* FStallDetector is meant to be constructed and destructed across a single scope to measure that specific timespan
 	**/
-	class CORE_API FStallDetector
+	class FStallDetector
 	{
 		friend class FStallDetectorRunnable;
 
@@ -255,29 +255,29 @@ namespace UE
 		/**
 		* Construct the stall detector linking it to it's stats object and take a timestamp
 		*/
-		FStallDetector(FStallDetectorStats& InStats);
+		CORE_API FStallDetector(FStallDetectorStats& InStats);
 
 		/**
 		* Destruct the stall detector and perform a check if appropriate (the default behavior)
 		*/
-		~FStallDetector();
+		CORE_API ~FStallDetector();
 
 		/**
 		* Perform a Check and then reset the timer to the current time
 		* This changes the mode of this class to defer the first check, and skip the destruction check
 		* This mode enables this class to used as a global or static, and CheckAndReset() callable periodically
 		**/
-		void CheckAndReset();
+		CORE_API void CheckAndReset();
 
 		/**
 		* Pause the timer for an expectantly long/slow operation we don't want counted against the stall timer
 		**/
-		void Pause(const double InSeconds);
+		CORE_API void Pause(const double InSeconds);
 
 		/**
 		* Resume the timer at the specified timestamp
 		**/
-		void Resume(const double InSeconds);
+		CORE_API void Resume(const double InSeconds);
 
 		/**
 		* The Parent stall detector on the calling thread's stack space
@@ -290,40 +290,40 @@ namespace UE
 		/**
 		* Sample the same clock that the stall detector uses internally
 		**/
-		static double Seconds();
+		static CORE_API double Seconds();
 
 		/**
 		* Initialize Stall Detector API resources
 		**/
-		static void Startup();
+		static CORE_API void Startup();
 
 		/**
 		* Cleanup Stall Detector API resources
 		**/
-		static void Shutdown();
+		static CORE_API void Shutdown();
 
 		/**
 		* Stall Detector API is running?
 		**/
-		static bool IsRunning();
+		static CORE_API bool IsRunning();
 
 		/**
 		 * Event that triggers when we detect that a stall in any instance is occuring
 		 * Note: It is possible for multiple instances to stall at the same time,
 		 * delegates should take this into account 
 		 */
-		static FOnStallDetected StallDetected;
+		static CORE_API FOnStallDetected StallDetected;
 
 		/**
 		 * Event that triggers when a stall has stopped
 		 * Note: It is possible for multiple instances to complete their stall at the same time,
 		 * delegates should take this into account
 		 */
-		static FOnStallCompleted StallCompleted;
+		static CORE_API FOnStallCompleted StallCompleted;
 
 	private:
 		/** Monotonically increasing counter to generate a unique ID per stall detector */
-		static TAtomic<uint64> UIDGenerator;
+		static CORE_API TAtomic<uint64> UIDGenerator;
 		
 		// The time provided, the deadline
 		FStallDetectorStats& Stats;
@@ -357,17 +357,17 @@ namespace UE
 		 * @param InCheckSeconds		The timestamp at which to compare the start time to
 		 * @return The amount the scope's budget has been exceeded by or 0.0
 		 **/
-		double Check(bool bFinalCheck, double InCheckSeconds);
+		CORE_API double Check(bool bFinalCheck, double InCheckSeconds);
 	};
 
 	/**
 	* FStallDetectorPause checks, pauses, and resumes the current thread's stack of stall detectors across it's lifetime
 	**/
-	class CORE_API FStallDetectorPause
+	class FStallDetectorPause
 	{
 	public:
-		FStallDetectorPause();
-		~FStallDetectorPause();
+		CORE_API FStallDetectorPause();
+		CORE_API ~FStallDetectorPause();
 
 	private:
 		bool bPaused;
@@ -379,4 +379,4 @@ namespace UE
 	static constexpr SIZE_T StallDetectorScopeMaxNameLength = StallDetectorStatNameBufferSize - 1 - UE_ARRAY_COUNT(StallDetectorCounterPrefix)
 															+ FMath::Max(UE_ARRAY_COUNT(StallDetectorTriggerCountSuffix), UE_ARRAY_COUNT(StallDetectorOverageTimeSuffix));
 } // namespace UE
-#endif // STALL_DETECTOR
\ No newline at end of file
+#endif // STALL_DETECTOR
diff --git a/Engine/Source/Runtime/Core/Public/ProfilingDebugging/TraceAuxiliary.h b/Engine/Source/Runtime/Core/Public/ProfilingDebugging/TraceAuxiliary.h
--- a/Engine/Source/Runtime/Core/Public/ProfilingDebugging/TraceAuxiliary.h
+++ b/Engine/Source/Runtime/Core/Public/ProfilingDebugging/TraceAuxiliary.h
@@ -12,7 +12,7 @@
 #include "Misc/Build.h"
 
 ////////////////////////////////////////////////////////////////////////////////
-class CORE_API FTraceAuxiliary
+class FTraceAuxiliary
 {
 public:
 	// In no logging configurations all log categories are of type FNoLoggingCategory, which has no relation with
@@ -81,88 +81,88 @@ public:
 	 * @param LogCategory Log channel to output messages to. Default set to 'Core'.
 	 * @return True when successfully starting the trace, false if the data connection could not be made.
 	 */
-	static bool Start(EConnectionType Type, const TCHAR* Target, const TCHAR* Channels = TEXT("default"), FOptions* Options = nullptr, const FLogCategoryAlias& LogCategory = LogCore);
+	static CORE_API bool Start(EConnectionType Type, const TCHAR* Target, const TCHAR* Channels = TEXT("default"), FOptions* Options = nullptr, const FLogCategoryAlias& LogCategory = LogCore);
 
 	/**
 	 * Stop tracing.
 	 * @return True if the trace was stopped, false if there was no data connection.
 	 */
-	static bool Stop();
+	static CORE_API bool Stop();
 
 	/**
 	 * Pause all tracing by disabling all active channels.
 	 */
-	static bool Pause();
+	static CORE_API bool Pause();
 
 	/**
 	 * Resume tracing by enabling all previously active channels.
 	 */
-	static bool Resume();
+	static CORE_API bool Resume();
 
 	/**
 	 * Write tailing memory state to a utrace file.
 	 * @param FilePath Path to the file to write the snapshot to. If it is null or empty a file path will be generated.
 	 */
-	static bool WriteSnapshot(const TCHAR* FilePath);
+	static CORE_API bool WriteSnapshot(const TCHAR* FilePath);
 	
 	/**
 	 * Write tailing memory state to a trace server.
 	 * @param FilePath Path to the file to write the snapshot to. If it is null or empty a file path will be generated.
 	 */
-	static bool SendSnapshot(const TCHAR* Host = nullptr, uint32 Port = 0);
+	static CORE_API bool SendSnapshot(const TCHAR* Host = nullptr, uint32 Port = 0);
 
 	/**
 	 * Initialize Trace systems.
 	 * @param CommandLine to use for initializing
 	 */
-	static void Initialize(const TCHAR* CommandLine);
+	static CORE_API void Initialize(const TCHAR* CommandLine);
 
 	/**
 	 * Initialize channels that use the config driven presets.
 	 * @param CommandLine to use for initializing
 	 */
-	static void InitializePresets(const TCHAR* CommandLine);
+	static CORE_API void InitializePresets(const TCHAR* CommandLine);
 
 	/**
 	 * Shut down Trace systems.
 	 */
-	static void Shutdown();
+	static CORE_API void Shutdown();
 
 	/**
 	 * Attempts to auto connect to an active trace server if an active session
 	 * of Unreal Insights Session Browser is running.
 	 */
-	static void TryAutoConnect();
+	static CORE_API void TryAutoConnect();
 
 	/**
 	 *  Enable previously selected channels. This method can be called multiple times
 	 *  as channels can be announced on module loading.
 	 */
-	static void EnableChannels();
+	static CORE_API void EnableChannels();
 
 	/**
 	* Disable channels to stop recording traces with them.
 	* @param Channels List of channels (or a preset) to disable. If null it will disable all active channels.
 	*/
-	static void DisableChannels(const TCHAR* Channels = nullptr);
+	static CORE_API void DisableChannels(const TCHAR* Channels = nullptr);
 
 	/**
 	 *  Returns the destination string that is currently being traced to.
 	 *  Contains either a file path or network address. Points to an empty string if tracing is disabled.
 	 */
 	UE_DEPRECATED(5.3, "Use GetTraceDestinationString instead.")
-	static const TCHAR*	GetTraceDestination();
+	static CORE_API const TCHAR*	GetTraceDestination();
 	
 	/**
 	 *  Returns the destination string that is currently being traced to.
 	 *  Contains either a file path or network address. Empty if tracing is disabled.
 	 */
-	static FString GetTraceDestinationString();
+	static CORE_API FString GetTraceDestinationString();
 	
 	/**
 	 *  Returns whether the trace system is currently connected to a trace sink (file or network)
 	 */
-	static bool	IsConnected();
+	static CORE_API bool	IsConnected();
 
 
 	/**
@@ -171,23 +171,23 @@ public:
 	 * @param OutTraceGuid If connected, the trace guid
 	 * @return True if connected
 	 */
-	static bool IsConnected(FGuid& OutSessionGuid, FGuid& OutTraceGuid);
+	static CORE_API bool IsConnected(FGuid& OutSessionGuid, FGuid& OutTraceGuid);
 
 	/**
 	*  Returns the current connection type.
 	*/
-	static EConnectionType GetConnectionType();
+	static CORE_API EConnectionType GetConnectionType();
 
 	/**
 	 *  Adds a comma separated list of currently active channels to the passed in StringBuilder
 	 */
-	static void	GetActiveChannelsString(FStringBuilderBase& String);
+	static CORE_API void	GetActiveChannelsString(FStringBuilderBase& String);
 
 	/**
 	 * Used when process is panicking. Stops all tracing immediately to avoid further allocations. Process is not
 	 * expected to continue after this call.
 	 */
-	 static void Panic();
+	 static CORE_API void Panic();
 
 	/**
 	 * Delegate that triggers when a connection is established. Gives subscribers a chance to trace events that appear
@@ -201,24 +201,24 @@ public:
 	 *		 
 	 * @param Callback Delegate to call on new connections.
 	 */
-	static FOnConnection OnConnection;
+	static CORE_API FOnConnection OnConnection;
 
 	/**
 	 * Delegate that triggers when a trace session is started.
 	 * The type of recording and the destination (filepath or network) is passed to the delegate.
 	 */
-	static FOnTraceStarted OnTraceStarted;
+	static CORE_API FOnTraceStarted OnTraceStarted;
 
 	/**
 	 * Delegate that triggers when a trace has finished recording. Useful if you need to collect all completed trace files in a session.
 	 * The type of recording and the destination (filepath or network) is passed to the delegate.
 	 */
-	static FOnTraceStopped OnTraceStopped;
+	static CORE_API FOnTraceStopped OnTraceStopped;
 
 	/**
 	 * Delegate that triggers when a snapshot has been saved.
 	 * The path to the snapshot file is passed to the delegate.
 	 */
-	static FOnSnapshotSaved OnSnapshotSaved;
+	static CORE_API FOnSnapshotSaved OnSnapshotSaved;
 	
 };
diff --git a/Engine/Source/Runtime/Core/Public/Serialization/Archive.h b/Engine/Source/Runtime/Core/Public/Serialization/Archive.h
--- a/Engine/Source/Runtime/Core/Public/Serialization/Archive.h
+++ b/Engine/Source/Runtime/Core/Public/Serialization/Archive.h
@@ -62,29 +62,29 @@ struct FArchiveCookContext;
 //       This will also make TSAN happy when multiple threads are hitting the same bitfield during concurrent saves.
 #define FArchive_Serialize_BitfieldBool(ARCHIVE, BITFIELD_BOOL) { bool TEMP_BITFIELD_BOOL = BITFIELD_BOOL; ARCHIVE << TEMP_BITFIELD_BOOL; if (BITFIELD_BOOL != TEMP_BITFIELD_BOOL) { BITFIELD_BOOL = TEMP_BITFIELD_BOOL; } }
 
-struct CORE_API FArchiveState
+struct FArchiveState
 {
 private:
 	// Only FArchive is allowed to instantiate this, by inheritance
 	friend class FArchive;
 
-	FArchiveState();
+	CORE_API FArchiveState();
 
 	/** Copy constructor. */
-	FArchiveState(const FArchiveState&);
+	CORE_API FArchiveState(const FArchiveState&);
 
 	/**
 	 * Copy assignment operator.
 	 *
 	 * @param ArchiveToCopy The archive to copy from.
 	 */
-	FArchiveState& operator=(const FArchiveState& ArchiveToCopy);
+	CORE_API FArchiveState& operator=(const FArchiveState& ArchiveToCopy);
 
-	virtual ~FArchiveState() = 0;
+	CORE_API virtual ~FArchiveState() = 0;
 
 protected:
-	static void LinkProxy(FArchiveState& Inner, FArchiveState& Proxy);
-	static void UnlinkProxy(FArchiveState& Inner, FArchiveState& Proxy);
+	static CORE_API void LinkProxy(FArchiveState& Inner, FArchiveState& Proxy);
+	static CORE_API void UnlinkProxy(FArchiveState& Inner, FArchiveState& Proxy);
 
 public:
 	/** 
@@ -98,22 +98,22 @@ public:
 	/** 
 	 * Modifies current archive state, can be used to override flags.
 	 */
-	void SetArchiveState(const FArchiveState& InState);
+	CORE_API void SetArchiveState(const FArchiveState& InState);
 
 	/**
 	 * Sets ArIsError to true. Also sets error in the proxy archiver if one is wrapping this.
 	 */
-	void SetError();
+	CORE_API void SetError();
 
 	/**
 	 * Sets ArIsError to false, this does not clear any CriticalErrors
 	 */
-	void ClearError();
+	CORE_API void ClearError();
 
 	/**
 	 * Sets the archiver IsCriticalError and IsError to true. Also sets CriticalError in the proxy archiver if one is wrapping this.
 	 */
-	void SetCriticalError();
+	CORE_API void SetCriticalError();
 
 	/**
 	 * Called to get the computed size from a size-detecting archive after it has finished serializing.
@@ -126,7 +126,7 @@ public:
 	 *
 	 * This is overridden for the specific Archive Types
 	 */
-	virtual FString GetArchiveName() const;
+	CORE_API virtual FString GetArchiveName() const;
 
 	/**
 	 * If this archive is a FLinkerLoad or FLinkerSave, returns a pointer to the ULinker portion.
@@ -232,10 +232,10 @@ public:
 	}
 
 	/** Returns the engine-global network protocol version for this archive. */
-	virtual uint32 EngineNetVer() const;
+	CORE_API virtual uint32 EngineNetVer() const;
 
 	/** Returns the game-specific network protocol version for this archive. */
-	virtual uint32 GameNetVer() const;
+	CORE_API virtual uint32 GameNetVer() const;
 
 	/**
 	 * Queries a custom version from the archive.  If the archive is being used to write, the custom version must have already been registered.
@@ -243,7 +243,7 @@ public:
 	 * @param Key The guid of the custom version to query.
 	 * @return The version number, or -1 if the custom tag isn't stored in the archive.
 	 */
-	int32 CustomVer(const struct FGuid& Key) const;
+	CORE_API int32 CustomVer(const struct FGuid& Key) const;
 
 	/** Returns true if this archive is for loading data. */
 	FORCEINLINE bool IsLoading() const
@@ -338,7 +338,7 @@ public:
 	 * Setting this to true should improve serialization performance by skipping the step entirely.
 	 * @param bShouldSkip Whether we should skip the checking the version registry for new version data if the version key is already set
 	 */
-	void SetShouldSkipUpdateCustomVersion(bool bShouldSkip);
+	CORE_API void SetShouldSkipUpdateCustomVersion(bool bShouldSkip);
 
 	/** Returns true if this Archive is does not update the custom version data for a given key if has already been set previously in the local custom versions container*/
 	FORCEINLINE bool ShouldSkipUpdateCustomVersion() const
@@ -502,17 +502,17 @@ public:
 	 *
 	 * @return The container of custom versions in the archive.
 	 */
-	virtual const FCustomVersionContainer& GetCustomVersions() const;
+	CORE_API virtual const FCustomVersionContainer& GetCustomVersions() const;
 
 	/**
 	 * Sets the custom version numbers for this archive.
 	 *
 	 * @param CustomVersionContainer - The container of custom versions to copy into the archive.
 	 */
-	virtual void SetCustomVersions(const FCustomVersionContainer& CustomVersionContainer);
+	CORE_API virtual void SetCustomVersions(const FCustomVersionContainer& CustomVersionContainer);
 
 	/** Resets the custom version numbers for this archive. */
-	virtual void ResetCustomVersions();
+	CORE_API virtual void ResetCustomVersions();
 
 	/**
 	 * Sets a specific custom version
@@ -521,7 +521,7 @@ public:
 	 * @param Version - The version number to set key to
 	 * @param FriendlyName - Friendly name corresponding to the key
 	 */
-	void SetCustomVersion(const struct FGuid& Key, int32 Version, FName FriendlyName);
+	CORE_API void SetCustomVersion(const struct FGuid& Key, int32 Version, FName FriendlyName);
 
 	/**
 	 * Toggle byte order swapping. This is needed in rare cases when we already know that the data
@@ -688,7 +688,7 @@ public:
 	 * Gets the chain of properties that are currently being serialized
 	 * @note This populates the array in stack order, so the 0th entry in the array is the top of the stack of properties
 	 */
-	void GetSerializedPropertyChain(TArray<class FProperty*>& OutProperties) const;
+	CORE_API void GetSerializedPropertyChain(TArray<class FProperty*>& OutProperties) const;
 
 	/**
 	 * Get the raw serialized property chain for this archive
@@ -702,11 +702,11 @@ public:
 	/**
 	 * Set the raw serialized property chain for this archive, optionally overriding the serialized property too (or null to use the head of the property chain)
 	 */
-	virtual void SetSerializedPropertyChain(const FArchiveSerializedPropertyChain* InSerializedPropertyChain, class FProperty* InSerializedPropertyOverride = nullptr);
+	CORE_API virtual void SetSerializedPropertyChain(const FArchiveSerializedPropertyChain* InSerializedPropertyChain, class FProperty* InSerializedPropertyOverride = nullptr);
 
 #if WITH_EDITORONLY_DATA
 	/** Returns true if the stack of currently serialized properties contains an editor-only property */
-	virtual bool IsEditorOnlyPropertyOnTheStack() const;
+	CORE_API virtual bool IsEditorOnlyPropertyOnTheStack() const;
 #endif
 
 	/** Sets the current UObject serialization context for this archive. */
@@ -720,17 +720,17 @@ public:
 	 * Set the localization namespace that this archive should use when serializing text properties.
 	 * This is typically the namespace used by the package being serialized (if serializing a package, or an object within a package).
 	 */
-	virtual void SetLocalizationNamespace(const FString& InLocalizationNamespace);
+	CORE_API virtual void SetLocalizationNamespace(const FString& InLocalizationNamespace);
 
 	/**
 	 * Get the localization namespace that this archive should use when serializing text properties.
 	 * This is typically the namespace used by the package being serialized (if serializing a package, or an object within a package).
 	 */
-	virtual FString GetLocalizationNamespace() const;
+	CORE_API virtual FString GetLocalizationNamespace() const;
 #endif // USE_STABLE_LOCALIZATION_KEYS
 
 	/** Resets all of the base archive members. */
-	virtual void Reset();
+	CORE_API virtual void Reset();
 
 public:
 	/* These are used for fastpath inline serializers  */
@@ -759,7 +759,7 @@ public:
 // These will be private in FArchive
 protected:
 	/** Copies all of the members except CustomVersionContainer */
-	void CopyTrivialFArchiveStatusMembers(const FArchiveState& ArchiveStatusToCopy);
+	CORE_API void CopyTrivialFArchiveStatusMembers(const FArchiveState& ArchiveStatusToCopy);
 
 	/** Whether this archive is for loading data. */
 	uint8 ArIsLoading : 1;
@@ -875,44 +875,44 @@ public:
 	 *
 	 * @param bInIsLoading  true if this archive is for loading, false otherwise.
 	 */
-	virtual void SetIsLoading(bool bInIsLoading);
+	CORE_API virtual void SetIsLoading(bool bInIsLoading);
 
 	/** 
 	 * Sets whether the archive is loading from a cooked package.
 	 * @param bInIsLoadingFromCookedPackage  true if this archive is loading from a cooked package, false otherwise
 	 */
-	virtual void SetIsLoadingFromCookedPackage(bool bInIsLoadingFromCookedPackage);
+	CORE_API virtual void SetIsLoadingFromCookedPackage(bool bInIsLoadingFromCookedPackage);
 
 	/**
 	 * Sets whether this archive is for saving data.
 	 *
 	 * @param bInIsSaving  true if this archive is for saving, false otherwise.
 	 */
-	virtual void SetIsSaving(bool bInIsSaving);
+	CORE_API virtual void SetIsSaving(bool bInIsSaving);
 
 	/**
 	 * Sets whether this archive is for transacting.
 	 *
 	 * @param bInIsTransacting  true if this archive is for transacting, false otherwise.
 	 */
-	virtual void SetIsTransacting(bool bInIsTransacting);
+	CORE_API virtual void SetIsTransacting(bool bInIsTransacting);
 
 	/**
 	 * Sets whether this archive is in text format.
 	 *
 	 * @param bInIsTextFormat  true if this archive is in text format, false otherwise.
 	 */
-	virtual void SetIsTextFormat(bool bInIsTextFormat);
+	CORE_API virtual void SetIsTextFormat(bool bInIsTextFormat);
 
 	/**
 	 * Sets whether this archive wants binary property serialization.
 	 *
 	 * @param bInWantBinaryPropertySerialization  true if this archive wants binary serialization, false otherwise.
 	 */
-	virtual void SetWantBinaryPropertySerialization(bool bInWantBinaryPropertySerialization);
+	CORE_API virtual void SetWantBinaryPropertySerialization(bool bInWantBinaryPropertySerialization);
 
 	/** Sets whether tagged property serialization should be replaced by faster unversioned serialization. This assumes writer and reader share the same property definitions. */
-	virtual void SetUseUnversionedPropertySerialization(bool bInUseUnversioned);
+	CORE_API virtual void SetUseUnversionedPropertySerialization(bool bInUseUnversioned);
 
 	/**
 	 * Sets whether this archive wants to force saving as Unicode.
@@ -920,14 +920,14 @@ public:
 	 *
 	 * @param bInForceUnicode  true if this archive wants to force saving as Unicode, false otherwise.
 	 */
-	virtual void SetForceUnicode(bool bInForceUnicode);
+	CORE_API virtual void SetForceUnicode(bool bInForceUnicode);
 
 	/**
 	 * Sets whether this archive is to persistent storage.
 	 *
 	 * @param bInIsPersistent  true if this archive is to persistent storage, false otherwise.
 	 */
-	virtual void SetIsPersistent(bool bInIsPersistent);
+	CORE_API virtual void SetIsPersistent(bool bInIsPersistent);
 
 	/**
 	 * Sets the archive version number. Used by the code that makes sure that FLinkerLoad's 
@@ -935,7 +935,7 @@ public:
 	 *
 	 * @param UEVer	new version number
 	 */
-	virtual void SetUEVer(FPackageFileVersion InVer);
+	CORE_API virtual void SetUEVer(FPackageFileVersion InVer);
 
 	UE_DEPRECATED(5.0, "Use SetUEVer instead which takes the version as a FPackageFileVersion. See the @FPackageFileVersion documentation for further details")
 	FORCEINLINE void SetUE4Ver(int32 InVer)
@@ -950,7 +950,7 @@ public:
 	 *
 	 * @param Ver	new version number
 	 */
-	virtual void SetLicenseeUEVer(int32 InVer);
+	CORE_API virtual void SetLicenseeUEVer(int32 InVer);
 
 	UE_DEPRECATED(5.0, "Use SetLicenseeUEVer instead")
 	FORCEINLINE void SetLicenseeUE4Ver(int32 InVer)
@@ -964,17 +964,17 @@ public:
 	 *
 	 * @param InVer	new version number
 	 */
-	virtual void SetEngineVer(const FEngineVersionBase& InVer);
+	CORE_API virtual void SetEngineVer(const FEngineVersionBase& InVer);
 
 	/**
 	 * Sets the archive engine network version.
 	 */
-	virtual void SetEngineNetVer(const uint32 InEngineNetVer);
+	CORE_API virtual void SetEngineNetVer(const uint32 InEngineNetVer);
 
 	/**
 	 * Sets the archive game network version.
 	 */
-	virtual void SetGameNetVer(const uint32 InGameNetVer);
+	CORE_API virtual void SetGameNetVer(const uint32 InGameNetVer);
 
 // These will be private in FArchive
 protected:
@@ -1023,10 +1023,10 @@ protected:
 	FString* LocalizationNamespacePtr;
 
 	/** See GetLocalizationNamespace */
-	FString GetBaseLocalizationNamespace() const;
+	CORE_API FString GetBaseLocalizationNamespace() const;
 
 	/** See SetLocalizationNamespace */
-	void SetBaseLocalizationNamespace(const FString& InLocalizationNamespace);
+	CORE_API void SetBaseLocalizationNamespace(const FString& InLocalizationNamespace);
 #endif // USE_STABLE_LOCALIZATION_KEYS
 
 	/**
@@ -1131,7 +1131,7 @@ private:
  * Base class for archives that can be used for loading, saving, and garbage
  * collecting in a byte order neutral way.
  */
-class CORE_API FArchive : private FArchiveState
+class FArchive : private FArchiveState
 {
 public:
 	FArchive() = default;
@@ -1165,7 +1165,7 @@ public:
 	 * @param Ar The archive to serialize from or to.
 	 * @param Value The value to serialize.
 	 */
-	virtual FArchive& operator<<(FText& Value);
+	CORE_API virtual FArchive& operator<<(FText& Value);
 
 	/**
 	 * Serializes an UObject value from or into this archive.
@@ -1244,7 +1244,7 @@ public:
 	 * @param Value The value to serialize.
 	 * @return This instance.
 	 */
-	virtual FArchive& operator<<(struct FLazyObjectPtr& Value);
+	CORE_API virtual FArchive& operator<<(struct FLazyObjectPtr& Value);
 
 	/**
 	 * Serializes a wrapped object pointer value from or into this archive.
@@ -1254,7 +1254,7 @@ public:
 	 * @param Value The value to serialize.
 	 * @return This instance.
 	 */
-	virtual FArchive& operator<<(struct FObjectPtr& Value);
+	CORE_API virtual FArchive& operator<<(struct FObjectPtr& Value);
 	
 	/**
 	 * Serializes asset pointer from or into this archive.
@@ -1264,7 +1264,7 @@ public:
 	 * @param Value The asset pointer to serialize.
 	 * @return This instance.
 	 */
-	virtual FArchive& operator<<(struct FSoftObjectPtr& Value);
+	CORE_API virtual FArchive& operator<<(struct FSoftObjectPtr& Value);
 
 	/**
 	 * Serializes soft object paths from or into this archive.
@@ -1272,7 +1272,7 @@ public:
 	 * @param Value Soft object path to serialize.
 	 * @return This instance.
 	 */
-	virtual FArchive& operator<<(struct FSoftObjectPath& Value);
+	CORE_API virtual FArchive& operator<<(struct FSoftObjectPath& Value);
 
 	/**
 	* Serializes FWeakObjectPtr value from or into this archive.
@@ -1282,7 +1282,7 @@ public:
 	* @param Value The value to serialize.
 	* @return This instance.
 	*/
-	virtual FArchive& operator<<(struct FWeakObjectPtr& Value);
+	CORE_API virtual FArchive& operator<<(struct FWeakObjectPtr& Value);
 
 	/** 
 	 * Inform the archive that a blueprint would like to force finalization, normally
@@ -1418,7 +1418,7 @@ public:
 	 */
 #if WITH_EDITOR
 protected:
-	virtual void SerializeBool( bool& D );
+	CORE_API virtual void SerializeBool( bool& D );
 public:
 	FORCEINLINE friend FArchive& operator<<(FArchive& Ar, bool& D)
 	{
@@ -1603,8 +1603,8 @@ public:
 	}
 
 	/** Packs int value into bytes of 7 bits with 8th bit for 'more' */
-	virtual void SerializeIntPacked(uint32& Value);
-	virtual void SerializeIntPacked64(uint64& Value);
+	CORE_API virtual void SerializeIntPacked(uint32& Value);
+	CORE_API virtual void SerializeIntPacked64(uint64& Value);
 
 	/** Tells the archive to attempt to preload the specified object so data can be loaded out of it. */
 	virtual void Preload(UObject* Object) { }
@@ -1744,7 +1744,7 @@ public:
 	 * @param	Flags	Flags to control what method to use for [de]compression and optionally control memory vs speed when compressing
 	 * @param	bTreatBufferAsFileReader true if V is actually an FArchive, which is used when saving to read data - helps to avoid single huge allocations of source data
 	 */
-	void SerializeCompressed(void* V, int64 Length, FName CompressionFormatCannotChange,
+	CORE_API void SerializeCompressed(void* V, int64 Length, FName CompressionFormatCannotChange,
 		ECompressionFlags Flags=COMPRESS_NoFlags, bool bTreatBufferAsFileReader=false);
 	
 	/**
@@ -1764,7 +1764,7 @@ public:
 	 * @param	bTreatBufferAsFileReader true if V is actually an FArchive, which is used when saving to read data - helps to avoid single huge allocations of source data
 	 * @param	OutPartialReadLength if not null, partial reads are allowed and the size is filled here
 	 */
-	void SerializeCompressedNew(void* V, int64 Length, FName CompressionFormatToEncode, FName CompressionFormatToDecodeOldV1Files,
+	CORE_API void SerializeCompressedNew(void* V, int64 Length, FName CompressionFormatToEncode, FName CompressionFormatToDecodeOldV1Files,
 		ECompressionFlags Flags=COMPRESS_NoFlags, bool bTreatBufferAsFileReader=false,
 		int64 * OutPartialReadLength=nullptr);
 		
@@ -1776,13 +1776,13 @@ public:
 	 * @param	V		Data pointer to serialize data from/ to
 	 * @param	Length	Length of source data if we're saving, unused otherwise
 	 */
-	void SerializeCompressedNew(void* V, int64 Length);
+	CORE_API void SerializeCompressedNew(void* V, int64 Length);
 
 
 	using FArchiveState::IsByteSwapping;
 
 	/** Used to do byte swapping on small items. This does not happen usually, so we don't want it inline. */
-	void ByteSwap(void* V, int32 Length);
+	CORE_API void ByteSwap(void* V, int32 Length);
 
 	/** Serialize data of Length bytes, taking into account byte swapping if needed. */
 	FORCEINLINE FArchive& ByteOrderSerialize(void* V, int32 Length)
@@ -1829,7 +1829,7 @@ public:
 	using FArchiveState::GetArchetypeFromLoader;
 
 private:
-	void VARARGS LogfImpl(const TCHAR* Fmt, ...);
+	CORE_API void VARARGS LogfImpl(const TCHAR* Fmt, ...);
 
 public:
 	// Logf implementation for convenience.
@@ -1855,7 +1855,7 @@ public:
 	 *
 	 * @param Guid The guid of the custom version.  This must have previously been registered with FCustomVersionRegistration.
 	 */
-	virtual void UsingCustomVersion(const struct FGuid& Guid);
+	CORE_API virtual void UsingCustomVersion(const struct FGuid& Guid);
 
 	using FArchiveState::CustomVer;
 
@@ -1932,7 +1932,7 @@ public:
 	 * @param InProperty			Pointer to the property that is currently being serialized
 	 * @param bIsEditorOnlyProperty True if the property is editor only (call FProperty::IsEditorOnlyProperty to work this out, as the archive can't since it can't access CoreUObject types)
 	 */
-	virtual void PushSerializedProperty(class FProperty* InProperty, const bool bIsEditorOnlyProperty);
+	CORE_API virtual void PushSerializedProperty(class FProperty* InProperty, const bool bIsEditorOnlyProperty);
 
 	/**
 	 * Pop a property that was previously being serialized off the property stack
@@ -1940,7 +1940,7 @@ public:
 	 * @param InProperty			Pointer to the property that was previously being serialized
 	 * @param bIsEditorOnlyProperty True if the property is editor only (call FProperty::IsEditorOnlyProperty to work this out, as the archive can't since it can't access CoreUObject types)
 	 */
-	virtual void PopSerializedProperty(class FProperty* InProperty, const bool bIsEditorOnlyProperty);
+	CORE_API virtual void PopSerializedProperty(class FProperty* InProperty, const bool bIsEditorOnlyProperty);
 
 #if WITH_EDITORONLY_DATA
 	using FArchiveState::IsEditorOnlyPropertyOnTheStack;
@@ -2030,10 +2030,10 @@ private:
 	}
 
 	// Not inlined to keep ByteOrderSerialize(), small and fast.
-	FArchive& SerializeByteOrderSwapped(void* V, int32 Length);
-	FArchive& SerializeByteOrderSwapped(uint16& Value);
-	FArchive& SerializeByteOrderSwapped(uint32& Value);
-	FArchive& SerializeByteOrderSwapped(uint64& Value);
+	CORE_API FArchive& SerializeByteOrderSwapped(void* V, int32 Length);
+	CORE_API FArchive& SerializeByteOrderSwapped(uint16& Value);
+	CORE_API FArchive& SerializeByteOrderSwapped(uint32& Value);
+	CORE_API FArchive& SerializeByteOrderSwapped(uint64& Value);
 #if PLATFORM_COMPILER_DISTINGUISHES_INT_AND_LONG
 	FArchive& SerializeByteOrderSwapped(unsigned long& Value)
 	{
@@ -2158,7 +2158,7 @@ public:
 	using FArchiveState::ArDebugSerializationFlags;
 	
 	/** Debug stack storage if you want to add data to the archive for usage further down the serialization stack this should be used in conjunction with the FScopeAddDebugData struct */
-	virtual void PushDebugDataString(const FName& DebugData);
+	CORE_API virtual void PushDebugDataString(const FName& DebugData);
 	virtual void PopDebugDataString() { }
 
 	class FScopeAddDebugData
diff --git a/Engine/Source/Runtime/Core/Public/Serialization/ArchiveCookData.h b/Engine/Source/Runtime/Core/Public/Serialization/ArchiveCookData.h
--- a/Engine/Source/Runtime/Core/Public/Serialization/ArchiveCookData.h
+++ b/Engine/Source/Runtime/Core/Public/Serialization/ArchiveCookData.h
@@ -9,7 +9,7 @@ struct FArchiveCookContext;
 
 // Holds archive data only relevant for cooking archives. Fill this out
 // as part of FSavePackageArgs.
-struct CORE_API FArchiveCookData
+struct FArchiveCookData
 {
 	const ITargetPlatform& TargetPlatform;
 	FArchiveCookContext& CookContext;
diff --git a/Engine/Source/Runtime/Core/Public/Serialization/ArchiveLoadCompressedProxy.h b/Engine/Source/Runtime/Core/Public/Serialization/ArchiveLoadCompressedProxy.h
--- a/Engine/Source/Runtime/Core/Public/Serialization/ArchiveLoadCompressedProxy.h
+++ b/Engine/Source/Runtime/Core/Public/Serialization/ArchiveLoadCompressedProxy.h
@@ -16,13 +16,13 @@
 /**
  * FArchive Proxy to transparently load compressed data from an array.
  */
-class CORE_API FArchiveLoadCompressedProxy : public FArchive
+class FArchiveLoadCompressedProxy : public FArchive
 {
 public:
-	FArchiveLoadCompressedProxy(const TArray<uint8>& InCompressedData, FName CompressionFormat, ECompressionFlags InCompressionFlags=COMPRESS_NoFlags);
+	CORE_API FArchiveLoadCompressedProxy(const TArray<uint8>& InCompressedData, FName CompressionFormat, ECompressionFlags InCompressionFlags=COMPRESS_NoFlags);
 
 	/** Destructor, freeing temporary memory. */
-	virtual ~FArchiveLoadCompressedProxy();
+	CORE_API virtual ~FArchiveLoadCompressedProxy();
 
 	/**
 	 * Serializes data from archive. This function is called recursively and determines where to serialize
@@ -31,18 +31,18 @@ public:
 	 * @param	Data	Pointer to serialize to
 	 * @param	Count	Number of bytes to read
 	 */
-	virtual void Serialize( void* Data, int64 Count );
+	CORE_API virtual void Serialize( void* Data, int64 Count );
 
 	/**
 	 * Seeks to the passed in position in the stream. This archive only supports forward seeking
 	 * and implements it by serializing data till it reaches the position.
 	 */
-	virtual void Seek( int64 InPos );
+	CORE_API virtual void Seek( int64 InPos );
 
 	/**
 	 * @return current position in uncompressed stream in bytes.
 	 */
-	virtual int64 Tell();
+	CORE_API virtual int64 Tell();
 
 private:
 	/**
diff --git a/Engine/Source/Runtime/Core/Public/Serialization/ArchiveSaveCompressedProxy.h b/Engine/Source/Runtime/Core/Public/Serialization/ArchiveSaveCompressedProxy.h
--- a/Engine/Source/Runtime/Core/Public/Serialization/ArchiveSaveCompressedProxy.h
+++ b/Engine/Source/Runtime/Core/Public/Serialization/ArchiveSaveCompressedProxy.h
@@ -16,7 +16,7 @@
 /**
  * FArchive Proxy to transparently write out compressed data to an array.
  */
-class CORE_API FArchiveSaveCompressedProxy : public FArchive
+class FArchiveSaveCompressedProxy : public FArchive
 {
 	/**
 	 * This enum and the following constructor is a workaround for VC compiler bug that prevents using attributes
@@ -26,7 +26,7 @@ class CORE_API FArchiveSaveCompressedProxy : public FArchive
 	{
 		Redirect
 	};
-	FArchiveSaveCompressedProxy(EVS2015Redirector, TArray<uint8>& InCompressedData, ECompressionFlags InCompressionFlags);
+	CORE_API FArchiveSaveCompressedProxy(EVS2015Redirector, TArray<uint8>& InCompressedData, ECompressionFlags InCompressionFlags);
 
 public:
 	/** 
@@ -40,15 +40,15 @@ public:
 		// Make sure to remove the EVS2015Redirector constructor when this constructor is removed
 		: FArchiveSaveCompressedProxy(EVS2015Redirector::Redirect, InCompressedData, InCompressionFlags)
 	{}
-	FArchiveSaveCompressedProxy( TArray<uint8>& InCompressedData, FName InCompressionFormat, ECompressionFlags InCompressionFlags=COMPRESS_None);
+	CORE_API FArchiveSaveCompressedProxy( TArray<uint8>& InCompressedData, FName InCompressionFormat, ECompressionFlags InCompressionFlags=COMPRESS_None);
 
 	/** Destructor, flushing array if needed. Also frees temporary memory. */
-	virtual ~FArchiveSaveCompressedProxy();
+	CORE_API virtual ~FArchiveSaveCompressedProxy();
 
 	/**
 	 * Flushes tmp data to array.
 	 */
-	virtual void Flush();
+	CORE_API virtual void Flush();
 
 	/**
 	 * Serializes data to archive. This function is called recursively and determines where to serialize
@@ -57,19 +57,19 @@ public:
 	 * @param	Data	Pointer to serialize to
 	 * @param	Count	Number of bytes to read
 	 */
-	virtual void Serialize( void* Data, int64 Count );
+	CORE_API virtual void Serialize( void* Data, int64 Count );
 
 	/**
 	 * Seeking is only implemented internally for writing out compressed data and asserts otherwise.
 	 * 
 	 * @param	InPos	Position to seek to
 	 */
-	virtual void Seek( int64 InPos );
+	CORE_API virtual void Seek( int64 InPos );
 
 	/**
 	 * @return current position in uncompressed stream in bytes.
 	 */
-	virtual int64 Tell();
+	CORE_API virtual int64 Tell();
 
 private:
 	/** Array to write compressed data to.					*/
diff --git a/Engine/Source/Runtime/Core/Public/Serialization/BitArchive.h b/Engine/Source/Runtime/Core/Public/Serialization/BitArchive.h
--- a/Engine/Source/Runtime/Core/Public/Serialization/BitArchive.h
+++ b/Engine/Source/Runtime/Core/Public/Serialization/BitArchive.h
@@ -9,13 +9,13 @@
 /**
 * Base class for serializing bitstreams.
 */
-class CORE_API FBitArchive : public FArchive
+class FBitArchive : public FArchive
 {
 public:
 	/**
 	 * Default Constructor
 	 */
-	FBitArchive();
+	CORE_API FBitArchive();
 
-	virtual void SerializeBitsWithOffset( void* Src, int32 SourceBit, int64 LengthBits ) PURE_VIRTUAL(FBitArchive::SerializeBitsWithOffset,);
-};
\ No newline at end of file
+	CORE_API virtual void SerializeBitsWithOffset( void* Src, int32 SourceBit, int64 LengthBits ) PURE_VIRTUAL(FBitArchive::SerializeBitsWithOffset,);
+};
diff --git a/Engine/Source/Runtime/Core/Public/Serialization/BitReader.h b/Engine/Source/Runtime/Core/Public/Serialization/BitReader.h
--- a/Engine/Source/Runtime/Core/Public/Serialization/BitReader.h
+++ b/Engine/Source/Runtime/Core/Public/Serialization/BitReader.h
@@ -21,24 +21,24 @@ CORE_API void appBitsCpy( uint8* Dest, int32 DestBit, uint8* Src, int32 SrcBit,
 //
 // Reads bitstreams.
 //
-struct CORE_API FBitReader : public FBitArchive
+struct FBitReader : public FBitArchive
 {
 	friend struct FBitReaderMark;
 
 public:
-	FBitReader( const uint8* Src = nullptr, int64 CountBits = 0 );
+	CORE_API FBitReader( const uint8* Src = nullptr, int64 CountBits = 0 );
 
 	FBitReader(const FBitReader&) = default;
     FBitReader& operator=(const FBitReader&) = default;
     FBitReader(FBitReader&&) = default;
     FBitReader& operator=(FBitReader&&) = default;
 
-	void SetData( FBitReader& Src, int64 CountBits );
-	void SetData( uint8* Src, int64 CountBits );
-	void SetData( TArray<uint8>&& Src, int64 CountBits );
+	CORE_API void SetData( FBitReader& Src, int64 CountBits );
+	CORE_API void SetData( uint8* Src, int64 CountBits );
+	CORE_API void SetData( TArray<uint8>&& Src, int64 CountBits );
 
 	/** Equivalent to SetData (reset position, copy from Src into internal buffer), but uses Reset not Empty to avoid a realloc if possible. */
-	void ResetData(FBitReader& Src, int64 CountBits);
+	CORE_API void ResetData(FBitReader& Src, int64 CountBits);
 
 // Disable false positive buffer overrun warning during pgoprofile linking step
 PRAGMA_DISABLE_BUFFER_OVERRUN_WARNING
@@ -72,7 +72,7 @@ PRAGMA_DISABLE_BUFFER_OVERRUN_WARNING
 			Pos += LengthBits;
 		}
 	}
-PRAGMA_ENABLE_BUFFER_OVERRUN_WARNING
+CORE_API PRAGMA_ENABLE_BUFFER_OVERRUN_WARNING
 
 	virtual void SerializeBitsWithOffset( void* Dest, int32 DestBit, int64 LengthBits ) override;
 
@@ -106,7 +106,7 @@ PRAGMA_ENABLE_BUFFER_OVERRUN_WARNING
 		}
 	}
 
-	virtual void SerializeIntPacked(uint32& Value) override;
+	CORE_API virtual void SerializeIntPacked(uint32& Value) override;
 
 	FORCEINLINE_DEBUGGABLE uint32 ReadInt(uint32 Max)
 	{
@@ -208,17 +208,17 @@ PRAGMA_ENABLE_BUFFER_OVERRUN_WARNING
 	 *
 	 * @param LengthBits	The number of bits being read at the time of overflow
 	 */
-	void SetOverflowed(int64 LengthBits);
+	CORE_API void SetOverflowed(int64 LengthBits);
 
 	/** Set the stream at the end */
 	void SetAtEnd() { Pos = Num; }
 
-	void AppendDataFromChecked( FBitReader& Src );
-	void AppendDataFromChecked( uint8* Src, uint32 NumBits );
-	void AppendTo( TArray<uint8> &Buffer );
+	CORE_API void AppendDataFromChecked( FBitReader& Src );
+	CORE_API void AppendDataFromChecked( uint8* Src, uint32 NumBits );
+	CORE_API void AppendTo( TArray<uint8> &Buffer );
 
 	/** Counts the in-memory bytes used by this object */
-	virtual void CountMemory(FArchive& Ar) const;
+	CORE_API virtual void CountMemory(FArchive& Ar) const;
 
 protected:
 
@@ -227,7 +227,7 @@ protected:
 	int64 Pos;
 
 	/** Copies version information used for network compatibility from Source to this archive */
-	void SetNetVersionsFromArchive(FArchive& Source);
+	CORE_API void SetNetVersionsFromArchive(FArchive& Source);
 
 private:
 
@@ -242,7 +242,7 @@ private:
 //
 // For pushing and popping FBitWriter positions.
 //
-struct CORE_API FBitReaderMark
+struct FBitReaderMark
 {
 public:
 
@@ -264,7 +264,7 @@ public:
 		Reader.Pos = Pos;
 	}
 
-	void Copy( FBitReader& Reader, TArray<uint8> &Buffer );
+	CORE_API void Copy( FBitReader& Reader, TArray<uint8> &Buffer );
 
 private:
 
diff --git a/Engine/Source/Runtime/Core/Public/Serialization/BitWriter.h b/Engine/Source/Runtime/Core/Public/Serialization/BitWriter.h
--- a/Engine/Source/Runtime/Core/Public/Serialization/BitWriter.h
+++ b/Engine/Source/Runtime/Core/Public/Serialization/BitWriter.h
@@ -18,27 +18,27 @@ class FArchive;
 //
 // Writes bitstreams.
 //
-struct CORE_API FBitWriter : public FBitArchive
+struct FBitWriter : public FBitArchive
 {
 	friend struct FBitWriterMark;
 
 public:
 	/** Default constructor. Zeros everything. */
-	FBitWriter(void);
+	CORE_API FBitWriter(void);
 
 	/**
 	 * Constructor using known size the buffer needs to be
 	 */
-	FBitWriter( int64 InMaxBits, bool AllowResize = false );
+	CORE_API FBitWriter( int64 InMaxBits, bool AllowResize = false );
 
 	FBitWriter(const FBitWriter&) = default;
     FBitWriter& operator=(const FBitWriter&) = default;
     FBitWriter(FBitWriter&&) = default;
     FBitWriter& operator=(FBitWriter&&) = default;
 
-	virtual void SerializeBits( void* Src, int64 LengthBits ) override;
+	CORE_API virtual void SerializeBits( void* Src, int64 LengthBits ) override;
 
-	virtual void SerializeBitsWithOffset( void* Src, int32 SourceBit, int64 LengthBits ) override;
+	CORE_API virtual void SerializeBitsWithOffset( void* Src, int32 SourceBit, int64 LengthBits ) override;
 
 	/**
 	 * Serializes a compressed integer - Value must be < Max
@@ -46,9 +46,9 @@ public:
 	 * @param Value		The value to serialize, must be < Max
 	 * @param Max		The maximum allowed value - good to aim for power-of-two
 	 */
-	virtual void SerializeInt(uint32& Value, uint32 Max) override;
+	CORE_API virtual void SerializeInt(uint32& Value, uint32 Max) override;
 
-	virtual void SerializeIntPacked(uint32& Value) override;
+	CORE_API virtual void SerializeIntPacked(uint32& Value) override;
 
 	/**
 	 * Serializes the specified Value, but does not bounds check against ValueMax;
@@ -57,10 +57,10 @@ public:
 	 * @param Value		The value to serialize
 	 * @param ValueMax	The maximum value to write; wraps Value if it exceeds this
 	 */
-	void WriteIntWrapped(uint32 Value, uint32 ValueMax);
+	CORE_API void WriteIntWrapped(uint32 Value, uint32 ValueMax);
 
-	void WriteBit( uint8 In );
-	virtual void Serialize( void* Src, int64 LengthBytes ) override;
+	CORE_API void WriteBit( uint8 In );
+	CORE_API virtual void Serialize( void* Src, int64 LengthBytes ) override;
 
 	/**
 	 * Returns a pointer to our internal buffer.
@@ -124,7 +124,7 @@ public:
 	 *
 	 * @param LengthBits	The number of bits being written at the time of overflow
 	 */
-	void SetOverflowed(int32 LengthBits);
+	CORE_API void SetOverflowed(int32 LengthBits);
 
 	/**
 	 * Sets whether or not this writer intentionally allows overflows (to avoid logspam)
@@ -169,7 +169,7 @@ public:
 	/**
 	 * Resets the bit writer back to its initial state
 	 */
-	void Reset() override;
+	CORE_API void Reset() override;
 
 	FORCEINLINE void WriteAlign()
 	{
@@ -177,7 +177,7 @@ public:
 	}
 
 	/** Counts the in-memory bytes used by this object */
-	virtual void CountMemory(FArchive& Ar) const;
+	CORE_API virtual void CountMemory(FArchive& Ar) const;
 
 private:
 	TArray<uint8> Buffer;
@@ -193,7 +193,7 @@ private:
 //
 // For pushing and popping FBitWriter positions.
 //
-struct CORE_API FBitWriterMark
+struct FBitWriterMark
 {
 public:
 
@@ -224,8 +224,8 @@ public:
 		Num = 0;		
 	}
 
-	void Pop( FBitWriter& Writer );
-	void Copy( FBitWriter& Writer, TArray<uint8> &Buffer );
+	CORE_API void Pop( FBitWriter& Writer );
+	CORE_API void Copy( FBitWriter& Writer, TArray<uint8> &Buffer );
 
 	/** Pops the BitWriter back to the start but doesn't clear what was written. */
 	FORCEINLINE_DEBUGGABLE void PopWithoutClear( FBitWriter& Writer )
