[Core]

diff --git a/Engine/Source/Runtime/Core/Public/Serialization/CustomVersion.h b/Engine/Source/Runtime/Core/Public/Serialization/CustomVersion.h
--- a/Engine/Source/Runtime/Core/Public/Serialization/CustomVersion.h
+++ b/Engine/Source/Runtime/Core/Public/Serialization/CustomVersion.h
@@ -39,7 +39,7 @@ typedef bool (*CustomVersionValidatorFunc)(const FCustomVersion& Version, const
 /**
  * Structure to hold unique custom key with its version.
  */
-struct CORE_API FCustomVersion
+struct FCustomVersion
 {
 	friend class FCustomVersionContainer;
 
@@ -87,7 +87,7 @@ struct CORE_API FCustomVersion
 	CORE_API friend void operator<<(FStructuredArchive::FSlot Slot, FCustomVersion& Version);
 
 	/** Gets the friendly name for error messages or whatever */
-	const FName GetFriendlyName() const;
+	CORE_API const FName GetFriendlyName() const;
 
 private:
 
@@ -110,7 +110,7 @@ ENUM_CLASS_FLAGS(ESetCustomVersionFlags)
 /**
  * Container for all available/serialized custom versions.
  */
-class CORE_API FCustomVersionContainer
+class FCustomVersionContainer
 {
 	friend struct FStaticCustomVersionRegistry;
 
@@ -127,7 +127,7 @@ public:
 	 * @param CustomKey Custom key for which to retrieve the version.
 	 * @return The FCustomVersion for the specified custom key, or nullptr if the key doesn't exist in the container.
 	 */
-	const FCustomVersion* GetVersion(FGuid CustomKey) const;
+	CORE_API const FCustomVersion* GetVersion(FGuid CustomKey) const;
 
 	/**
 	* Gets a custom version friendly name from the container.
@@ -135,7 +135,7 @@ public:
 	* @param CustomKey Custom key for which to retrieve the version.
 	* @return The friendly name for the specified custom key, or NAME_None if the key doesn't exist in the container.
 	*/
-	const FName GetFriendlyName(FGuid CustomKey) const;
+	CORE_API const FName GetFriendlyName(FGuid CustomKey) const;
 
 	/**
 	 * Sets a specific custom version in the container.
@@ -144,18 +144,18 @@ public:
 	 * @param Version The version number for the specified custom key
 	 * @param FriendlyName A friendly name to assign to this version
 	 */
-	void SetVersion(FGuid CustomKey, int32 Version, FName FriendlyName);
+	CORE_API void SetVersion(FGuid CustomKey, int32 Version, FName FriendlyName);
 
 	/**
 	 * Sets a specific custom version in the container. It queries the versions registry to get the version data for the provided key
 	 * @param CustomKey Custom key for which to retrieve the version.
 	 * @param Options Optional flags used to alter the behavior of this method
 	 */
-	void SetVersionUsingRegistry(FGuid CustomKey, ESetCustomVersionFlags Options = ESetCustomVersionFlags::None);
+	CORE_API void SetVersionUsingRegistry(FGuid CustomKey, ESetCustomVersionFlags Options = ESetCustomVersionFlags::None);
 
 	/** Serialization. */
-	void Serialize(FArchive& Ar, ECustomVersionSerializationFormat::Type Format = ECustomVersionSerializationFormat::Latest);
-	void Serialize(FStructuredArchive::FSlot Slot, ECustomVersionSerializationFormat::Type Format = ECustomVersionSerializationFormat::Latest);
+	CORE_API void Serialize(FArchive& Ar, ECustomVersionSerializationFormat::Type Format = ECustomVersionSerializationFormat::Latest);
+	CORE_API void Serialize(FStructuredArchive::FSlot Slot, ECustomVersionSerializationFormat::Type Format = ECustomVersionSerializationFormat::Latest);
 
 	/**
 	 * Gets a singleton with the registered versions.
@@ -163,20 +163,20 @@ public:
 	 * @return The registered version container.
 	 */
 	UE_DEPRECATED(4.24, "Use one of the thread-safe FCurrentCustomVersions methods instead")
-	static const FCustomVersionContainer& GetRegistered();
+	static CORE_API const FCustomVersionContainer& GetRegistered();
 
 	/**
 	 * Empties the custom version container.
 	 */
-	void Empty();
+	CORE_API void Empty();
 
 	/**
 	 * Sorts the custom version container by key.
 	 */
-	void SortByKey();
+	CORE_API void SortByKey();
 
 	/** Return a string representation of custom versions. Used for debug. */
-	FString ToString(const FString& Indent) const;
+	CORE_API FString ToString(const FString& Indent) const;
 
 private:
 
@@ -194,23 +194,23 @@ struct FCustomVersionDifference
 };
 
 /** Provides access to code-defined custom versions registered via FCustomVersionRegistration. */
-class CORE_API FCurrentCustomVersions
+class FCurrentCustomVersions
 {
 public:
 	/** Get a copy of all versions that has been statically registered so far in the module loading process. */
-	static FCustomVersionContainer GetAll();
+	static CORE_API FCustomVersionContainer GetAll();
 
 	/** Get a copy of a single statically registered version if it exists. */
-	static TOptional<FCustomVersion> Get(const FGuid& Guid);
+	static CORE_API TOptional<FCustomVersion> Get(const FGuid& Guid);
 
 	/** Compare a number of versions to current ones and return potential differences. */
-	static TArray<FCustomVersionDifference> Compare(const FCustomVersionArray& CompareVersions, const TCHAR* DebugContext);
+	static CORE_API TArray<FCustomVersionDifference> Compare(const FCustomVersionArray& CompareVersions, const TCHAR* DebugContext);
 
 private:
 	friend class FCustomVersionRegistration;
 
-	static void Register(const FGuid& Key, int32 Version, const TCHAR* FriendlyName, CustomVersionValidatorFunc ValidatorFunc);
-	static void Unregister(const FGuid& Key);
+	static CORE_API void Register(const FGuid& Key, int32 Version, const TCHAR* FriendlyName, CustomVersionValidatorFunc ValidatorFunc);
+	static CORE_API void Unregister(const FGuid& Key);
 };
 
 
diff --git a/Engine/Source/Runtime/Core/Public/Serialization/Formatters/BinaryArchiveFormatter.h b/Engine/Source/Runtime/Core/Public/Serialization/Formatters/BinaryArchiveFormatter.h
--- a/Engine/Source/Runtime/Core/Public/Serialization/Formatters/BinaryArchiveFormatter.h
+++ b/Engine/Source/Runtime/Core/Public/Serialization/Formatters/BinaryArchiveFormatter.h
@@ -18,66 +18,66 @@ struct FSoftObjectPath;
 struct FSoftObjectPtr;
 struct FWeakObjectPtr;
 
-class CORE_API FBinaryArchiveFormatter final : public FStructuredArchiveFormatter
+class FBinaryArchiveFormatter final : public FStructuredArchiveFormatter
 {
 public:
-	FBinaryArchiveFormatter(FArchive& InInner);
-	virtual ~FBinaryArchiveFormatter();
-
-	virtual bool HasDocumentTree() const override;
-	virtual FArchive& GetUnderlyingArchive() override;
-
-	virtual void EnterRecord() override;
-	virtual void LeaveRecord() override;
-	virtual void EnterField(FArchiveFieldName Name) override;
-	virtual void LeaveField() override;
-	virtual bool TryEnterField(FArchiveFieldName Name, bool bEnterIfWriting);
-
-	virtual void EnterArray(int32& NumElements) override;
-	virtual void LeaveArray() override;
-	virtual void EnterArrayElement() override;
-	virtual void LeaveArrayElement() override;
-
-	virtual void EnterStream() override;
-	virtual void LeaveStream() override;
-	virtual void EnterStreamElement() override;
-	virtual void LeaveStreamElement() override;
-
-	virtual void EnterMap(int32& NumElements) override;
-	virtual void LeaveMap() override;
-	virtual void EnterMapElement(FString& Name) override;
-	virtual void LeaveMapElement() override;
-
-	virtual void EnterAttributedValue() override;
-	virtual void EnterAttribute(FArchiveFieldName AttributeName) override;
-	virtual void EnterAttributedValueValue() override;
-	virtual void LeaveAttribute() override;
-	virtual void LeaveAttributedValue() override;
-	virtual bool TryEnterAttribute(FArchiveFieldName AttributeName, bool bEnterWhenWriting) override;
-	virtual bool TryEnterAttributedValueValue() override;
-
-	virtual void Serialize(uint8& Value) override;
-	virtual void Serialize(uint16& Value) override;
-	virtual void Serialize(uint32& Value) override;
-	virtual void Serialize(uint64& Value) override;
-	virtual void Serialize(int8& Value) override;
-	virtual void Serialize(int16& Value) override;
-	virtual void Serialize(int32& Value) override;
-	virtual void Serialize(int64& Value) override;
-	virtual void Serialize(float& Value) override;
-	virtual void Serialize(double& Value) override;
-	virtual void Serialize(bool& Value) override;
-	virtual void Serialize(FString& Value) override;
-	virtual void Serialize(FName& Value) override;
-	virtual void Serialize(UObject*& Value) override;
-	virtual void Serialize(FText& Value) override;
-	virtual void Serialize(FWeakObjectPtr& Value) override;
-	virtual void Serialize(FSoftObjectPtr& Value) override;
-	virtual void Serialize(FSoftObjectPath& Value) override;
-	virtual void Serialize(FLazyObjectPtr& Value) override;
-	virtual void Serialize(FObjectPtr& Value) override;
-	virtual void Serialize(TArray<uint8>& Value) override;
-	virtual void Serialize(void* Data, uint64 DataSize) override;
+	CORE_API FBinaryArchiveFormatter(FArchive& InInner);
+	CORE_API virtual ~FBinaryArchiveFormatter();
+
+	CORE_API virtual bool HasDocumentTree() const override;
+	CORE_API virtual FArchive& GetUnderlyingArchive() override;
+
+	CORE_API virtual void EnterRecord() override;
+	CORE_API virtual void LeaveRecord() override;
+	CORE_API virtual void EnterField(FArchiveFieldName Name) override;
+	CORE_API virtual void LeaveField() override;
+	CORE_API virtual bool TryEnterField(FArchiveFieldName Name, bool bEnterIfWriting);
+
+	CORE_API virtual void EnterArray(int32& NumElements) override;
+	CORE_API virtual void LeaveArray() override;
+	CORE_API virtual void EnterArrayElement() override;
+	CORE_API virtual void LeaveArrayElement() override;
+
+	CORE_API virtual void EnterStream() override;
+	CORE_API virtual void LeaveStream() override;
+	CORE_API virtual void EnterStreamElement() override;
+	CORE_API virtual void LeaveStreamElement() override;
+
+	CORE_API virtual void EnterMap(int32& NumElements) override;
+	CORE_API virtual void LeaveMap() override;
+	CORE_API virtual void EnterMapElement(FString& Name) override;
+	CORE_API virtual void LeaveMapElement() override;
+
+	CORE_API virtual void EnterAttributedValue() override;
+	CORE_API virtual void EnterAttribute(FArchiveFieldName AttributeName) override;
+	CORE_API virtual void EnterAttributedValueValue() override;
+	CORE_API virtual void LeaveAttribute() override;
+	CORE_API virtual void LeaveAttributedValue() override;
+	CORE_API virtual bool TryEnterAttribute(FArchiveFieldName AttributeName, bool bEnterWhenWriting) override;
+	CORE_API virtual bool TryEnterAttributedValueValue() override;
+
+	CORE_API virtual void Serialize(uint8& Value) override;
+	CORE_API virtual void Serialize(uint16& Value) override;
+	CORE_API virtual void Serialize(uint32& Value) override;
+	CORE_API virtual void Serialize(uint64& Value) override;
+	CORE_API virtual void Serialize(int8& Value) override;
+	CORE_API virtual void Serialize(int16& Value) override;
+	CORE_API virtual void Serialize(int32& Value) override;
+	CORE_API virtual void Serialize(int64& Value) override;
+	CORE_API virtual void Serialize(float& Value) override;
+	CORE_API virtual void Serialize(double& Value) override;
+	CORE_API virtual void Serialize(bool& Value) override;
+	CORE_API virtual void Serialize(FString& Value) override;
+	CORE_API virtual void Serialize(FName& Value) override;
+	CORE_API virtual void Serialize(UObject*& Value) override;
+	CORE_API virtual void Serialize(FText& Value) override;
+	CORE_API virtual void Serialize(FWeakObjectPtr& Value) override;
+	CORE_API virtual void Serialize(FSoftObjectPtr& Value) override;
+	CORE_API virtual void Serialize(FSoftObjectPath& Value) override;
+	CORE_API virtual void Serialize(FLazyObjectPtr& Value) override;
+	CORE_API virtual void Serialize(FObjectPtr& Value) override;
+	CORE_API virtual void Serialize(TArray<uint8>& Value) override;
+	CORE_API virtual void Serialize(void* Data, uint64 DataSize) override;
 
 private:
 
diff --git a/Engine/Source/Runtime/Core/Public/Serialization/LargeMemoryData.h b/Engine/Source/Runtime/Core/Public/Serialization/LargeMemoryData.h
--- a/Engine/Source/Runtime/Core/Public/Serialization/LargeMemoryData.h
+++ b/Engine/Source/Runtime/Core/Public/Serialization/LargeMemoryData.h
@@ -13,15 +13,15 @@ template <class T, int TPaddingForCacheContention> class TLockFreePointerListUno
 * Data storage for the large memory reader and writer.
 */
 
-class CORE_API FLargeMemoryData
+class FLargeMemoryData
 {
 public:
 
-	explicit FLargeMemoryData(const int64 PreAllocateBytes = 0);
-	~FLargeMemoryData();
+	CORE_API explicit FLargeMemoryData(const int64 PreAllocateBytes = 0);
+	CORE_API ~FLargeMemoryData();
 
 	/** Write data at the given offset. Returns true if the data was written. */
-	bool Write(void* InData, int64 InOffset, int64 InNum);
+	CORE_API bool Write(void* InData, int64 InOffset, int64 InNum);
 
 	/** Append data at the given offset. */
 	FORCEINLINE void Append(void* InData, int64 InNum)
@@ -30,7 +30,7 @@ public:
 	}
 
 	/** Read data at the given offset. Returns true if the data was read. */
-	bool Read(void* OutData, int64 InOffset, int64 InNum) const;
+	CORE_API bool Read(void* OutData, int64 InOffset, int64 InNum) const;
 
 	/** Gets the size of the data written. */
 	FORCEINLINE int64 GetSize() const
@@ -51,7 +51,7 @@ public:
 	}
 
 	/** Releases ownership of the written data. */
-	uint8* ReleaseOwnership();
+	CORE_API uint8* ReleaseOwnership();
 
 	/** Check whether data is allocated or if the ownership was released. */
 	bool HasData() const 
@@ -59,7 +59,7 @@ public:
 		return Data != nullptr;
 	}
 
-	void Reserve(int64 Size);
+	CORE_API void Reserve(int64 Size);
 
 private:
 
@@ -77,20 +77,20 @@ private:
 	int64 MaxBytes;
 
 	/** Resizes the data buffer to at least NumBytes with some slack */
-	void GrowBuffer();
+	CORE_API void GrowBuffer();
 };
 
 /**
 * Pooled storage of FLargeMemoryData instances, allowing allocation-free and lock-free access.
 */
-class CORE_API FPooledLargeMemoryData
+class FPooledLargeMemoryData
 {
 public:
-	FPooledLargeMemoryData();
-	~FPooledLargeMemoryData();
+	CORE_API FPooledLargeMemoryData();
+	CORE_API ~FPooledLargeMemoryData();
 	FLargeMemoryData& Get() { return *Data; }
 private:
 	FLargeMemoryData* Data;
-	static TLockFreePointerListUnordered<FLargeMemoryData, 0> FreeList;
-	static std::atomic<int32> FreeListLength;
-};
\ No newline at end of file
+	static CORE_API TLockFreePointerListUnordered<FLargeMemoryData, 0> FreeList;
+	static CORE_API std::atomic<int32> FreeListLength;
+};
diff --git a/Engine/Source/Runtime/Core/Public/Serialization/LargeMemoryReader.h b/Engine/Source/Runtime/Core/Public/Serialization/LargeMemoryReader.h
--- a/Engine/Source/Runtime/Core/Public/Serialization/LargeMemoryReader.h
+++ b/Engine/Source/Runtime/Core/Public/Serialization/LargeMemoryReader.h
@@ -21,18 +21,18 @@ ENUM_CLASS_FLAGS(ELargeMemoryReaderFlags);
 /**
 * Archive for reading a large amount of arbitrary data from memory
 */
-class CORE_API FLargeMemoryReader : public FMemoryArchive
+class FLargeMemoryReader : public FMemoryArchive
 {
 public:
 
-	FLargeMemoryReader(const uint8* InData, const int64 Num, ELargeMemoryReaderFlags InFlags = ELargeMemoryReaderFlags::None, const FName InArchiveName = NAME_None);
+	CORE_API FLargeMemoryReader(const uint8* InData, const int64 Num, ELargeMemoryReaderFlags InFlags = ELargeMemoryReaderFlags::None, const FName InArchiveName = NAME_None);
 
-	virtual void Serialize(void* OutData, int64 Num) override;
+	CORE_API virtual void Serialize(void* OutData, int64 Num) override;
 
 	/**
 	* Gets the total size of the data buffer
 	*/
-	virtual int64 TotalSize() override;
+	CORE_API virtual int64 TotalSize() override;
 
 	/**
 	* Returns the name of the Archive.  Useful for getting the name of the package a struct or object
@@ -40,12 +40,12 @@ public:
 	*
 	* This is overridden for the specific Archive Types
 	**/
-	virtual FString GetArchiveName() const override;
+	CORE_API virtual FString GetArchiveName() const override;
 
 	/**
 	* Destructor
 	*/
-	~FLargeMemoryReader();
+	CORE_API ~FLargeMemoryReader();
 
 private:
 
diff --git a/Engine/Source/Runtime/Core/Public/Serialization/LargeMemoryWriter.h b/Engine/Source/Runtime/Core/Public/Serialization/LargeMemoryWriter.h
--- a/Engine/Source/Runtime/Core/Public/Serialization/LargeMemoryWriter.h
+++ b/Engine/Source/Runtime/Core/Public/Serialization/LargeMemoryWriter.h
@@ -13,13 +13,13 @@
 /**
 * Archive for storing a large amount of arbitrary data to memory
 */
-class CORE_API FLargeMemoryWriter : public FMemoryArchive
+class FLargeMemoryWriter : public FMemoryArchive
 {
 public:
 	
-	FLargeMemoryWriter(const int64 PreAllocateBytes = 0, bool bIsPersistent = false, const TCHAR* InFilename = nullptr);
+	CORE_API FLargeMemoryWriter(const int64 PreAllocateBytes = 0, bool bIsPersistent = false, const TCHAR* InFilename = nullptr);
 
-	virtual void Serialize(void* InData, int64 Num) override;
+	CORE_API virtual void Serialize(void* InData, int64 Num) override;
 
 	/**
 	* Returns the name of the Archive.  Useful for getting the name of the package a struct or object
@@ -27,7 +27,7 @@ public:
 	*
 	* This is overridden for the specific Archive Types
 	**/
-	virtual FString GetArchiveName() const override;
+	CORE_API virtual FString GetArchiveName() const override;
 
 	/**
 	 * Gets the total size of the data written
@@ -40,7 +40,7 @@ public:
 	/**
 	 * Returns the written data. To release this archive's ownership of the data, call ReleaseOwnership()
 	 */
-	uint8* GetData() const;
+	CORE_API uint8* GetData() const;
 
 	/**
 	 * Returns a view on the written data
diff --git a/Engine/Source/Runtime/Core/Public/Serialization/MemoryImage.h b/Engine/Source/Runtime/Core/Public/Serialization/MemoryImage.h
--- a/Engine/Source/Runtime/Core/Public/Serialization/MemoryImage.h
+++ b/Engine/Source/Runtime/Core/Public/Serialization/MemoryImage.h
@@ -43,18 +43,18 @@ class FMemoryImageString;
 // This is also required for creating frozen memory images
 #define UE_MEMORYIMAGE_TRACK_TYPE_DEPENDENCIES (WITH_EDITORONLY_DATA)
 
-class CORE_API FPointerTableBase
+class FPointerTableBase
 {
 public:
 	virtual ~FPointerTableBase() {}
 	virtual int32 AddIndexedPointer(const FTypeLayoutDesc& TypeDesc, void* Ptr) = 0;
 	virtual void* GetIndexedPointer(const FTypeLayoutDesc& TypeDesc, uint32 i) const = 0;
 
-	virtual void SaveToArchive(FArchive& Ar, const FPlatformTypeLayoutParameters& LayoutParams, const void* FrozenObject) const;
-	virtual bool LoadFromArchive(FArchive& Ar, const FPlatformTypeLayoutParameters& LayoutParams, void* FrozenObject);
+	CORE_API virtual void SaveToArchive(FArchive& Ar, const FPlatformTypeLayoutParameters& LayoutParams, const void* FrozenObject) const;
+	CORE_API virtual bool LoadFromArchive(FArchive& Ar, const FPlatformTypeLayoutParameters& LayoutParams, void* FrozenObject);
 
 #if UE_MEMORYIMAGE_TRACK_TYPE_DEPENDENCIES
-	int32 AddTypeDependency(const FTypeLayoutDesc& TypeDesc);
+	CORE_API int32 AddTypeDependency(const FTypeLayoutDesc& TypeDesc);
 	inline const FTypeLayoutDesc* GetTypeDependency(int32 Index) const { return TypeDependencies[Index]; }
 private:
 	TArray<const FTypeLayoutDesc*> TypeDependencies;
@@ -210,7 +210,7 @@ struct FMemoryImageResult
 	CORE_API static FMemoryImageObject LoadFromArchive(FArchive& Ar, const FTypeLayoutDesc& TypeDesc, FPointerTableBase* PointerTable, FPlatformTypeLayoutParameters& OutLayoutParameters);
 };
 
-class CORE_API FMemoryImageSection : public FRefCountedObject
+class FMemoryImageSection : public FRefCountedObject
 {
 public:
 	struct FSectionPointer
@@ -260,14 +260,14 @@ public:
 	template<typename T>
 	uint32 WriteBytes(const T& Data) { return WriteBytes(&Data, sizeof(T)); }
 
-	FMemoryImageSection* WritePointer(const FTypeLayoutDesc& StaticTypeDesc, const FTypeLayoutDesc& DerivedTypeDesc, uint32* OutOffsetToBase = nullptr);
-	uint32 WriteRawPointerSizedBytes(uint64 PointerValue);
-	uint32 WriteVTable(const FTypeLayoutDesc& TypeDesc, const FTypeLayoutDesc& DerivedTypeDesc);
-	uint32 WriteFMemoryImageName(int32 NumBytes, const FName& Name);
-	uint32 WriteFScriptName(const FScriptName& Name);
-	uint32 Flatten(FMemoryImageResult& OutResult) const;
+	CORE_API FMemoryImageSection* WritePointer(const FTypeLayoutDesc& StaticTypeDesc, const FTypeLayoutDesc& DerivedTypeDesc, uint32* OutOffsetToBase = nullptr);
+	CORE_API uint32 WriteRawPointerSizedBytes(uint64 PointerValue);
+	CORE_API uint32 WriteVTable(const FTypeLayoutDesc& TypeDesc, const FTypeLayoutDesc& DerivedTypeDesc);
+	CORE_API uint32 WriteFMemoryImageName(int32 NumBytes, const FName& Name);
+	CORE_API uint32 WriteFScriptName(const FScriptName& Name);
+	CORE_API uint32 Flatten(FMemoryImageResult& OutResult) const;
 
-	void ComputeHash();
+	CORE_API void ComputeHash();
 
 	FMemoryImage* ParentImage;
 	TArray<uint8> Bytes;
@@ -279,7 +279,7 @@ public:
 	uint32 MaxAlignment;
 };
 
-class CORE_API FMemoryImage
+class FMemoryImage
 {
 public:
 	FMemoryImage()
@@ -304,7 +304,7 @@ public:
 	/** Merging duplicate sections will make the resulting memory image smaller.
 	 * This will only work for data that is expected to be read-only after freezing.  Merging sections will break any manual fix-ups applied to the frozen data
 	 */
-	void Flatten(FMemoryImageResult& OutResult, bool bMergeDuplicateSections = false);
+	CORE_API void Flatten(FMemoryImageResult& OutResult, bool bMergeDuplicateSections = false);
 
 	TArray<TRefCountPtr<FMemoryImageSection>> Sections;
 	FPointerTableBase* PointerTable;
diff --git a/Engine/Source/Runtime/Core/Public/Serialization/MemoryImageWriter.h b/Engine/Source/Runtime/Core/Public/Serialization/MemoryImageWriter.h
--- a/Engine/Source/Runtime/Core/Public/Serialization/MemoryImageWriter.h
+++ b/Engine/Source/Runtime/Core/Public/Serialization/MemoryImageWriter.h
@@ -10,39 +10,39 @@ class FMemoryImage;
 class FMemoryImageSection;
 class FPointerTableBase;
 
-class CORE_API FMemoryImageWriter
+class FMemoryImageWriter
 {
 public:
-	explicit FMemoryImageWriter(FMemoryImage& InImage);
-	explicit FMemoryImageWriter(FMemoryImageSection* InSection);
-	~FMemoryImageWriter();
+	CORE_API explicit FMemoryImageWriter(FMemoryImage& InImage);
+	CORE_API explicit FMemoryImageWriter(FMemoryImageSection* InSection);
+	CORE_API ~FMemoryImageWriter();
 
-	FMemoryImage& GetImage() const;
-	const FPlatformTypeLayoutParameters& GetHostLayoutParams() const;
-	const FPlatformTypeLayoutParameters& GetTargetLayoutParams() const;
-	FPointerTableBase& GetPointerTable() const;
-	const FPointerTableBase* TryGetPrevPointerTable() const;
+	CORE_API FMemoryImage& GetImage() const;
+	CORE_API const FPlatformTypeLayoutParameters& GetHostLayoutParams() const;
+	CORE_API const FPlatformTypeLayoutParameters& GetTargetLayoutParams() const;
+	CORE_API FPointerTableBase& GetPointerTable() const;
+	CORE_API const FPointerTableBase* TryGetPrevPointerTable() const;
 
 	inline bool Is32BitTarget() const { return GetTargetLayoutParams().Is32Bit(); }
 	inline bool Is64BitTarget() const { return !Is32BitTarget(); }
 
-	int32 AddTypeDependency(const FTypeLayoutDesc& TypeDesc);
-
-	void WriteObject(const void* Object, const FTypeLayoutDesc& TypeDesc);
-	void WriteObjectArray(const void* Object, const FTypeLayoutDesc& TypeDesc, uint32_t NumArray);
-	void WriteRootObject(const void* Object, const FTypeLayoutDesc& TypeDesc);
-
-	uint32 GetOffset() const;
-	uint32 WriteAlignment(uint32 Alignment);
-	void WritePaddingToSize(uint32 Offset);
-	uint32 WriteBytes(const void* Data, uint32 Size);
-	FMemoryImageWriter WritePointer(const FTypeLayoutDesc& StaticTypeDesc, const FTypeLayoutDesc& DerivedTypeDesc, uint32* OutOffsetToBase = nullptr);
-	FMemoryImageWriter WritePointer(const FTypeLayoutDesc& TypeDesc);
-	uint32 WriteNullPointer();
-	uint32 WriteRawPointerSizedBytes(uint64 PointerValue);
-	uint32 WriteVTable(const FTypeLayoutDesc& TypeDesc, const FTypeLayoutDesc& DerivedTypeDesc);
-	uint32 WriteFMemoryImageName(int32 NumBytes, const FName& Name);
-	uint32 WriteFScriptName(const FScriptName& Name);
+	CORE_API int32 AddTypeDependency(const FTypeLayoutDesc& TypeDesc);
+
+	CORE_API void WriteObject(const void* Object, const FTypeLayoutDesc& TypeDesc);
+	CORE_API void WriteObjectArray(const void* Object, const FTypeLayoutDesc& TypeDesc, uint32_t NumArray);
+	CORE_API void WriteRootObject(const void* Object, const FTypeLayoutDesc& TypeDesc);
+
+	CORE_API uint32 GetOffset() const;
+	CORE_API uint32 WriteAlignment(uint32 Alignment);
+	CORE_API void WritePaddingToSize(uint32 Offset);
+	CORE_API uint32 WriteBytes(const void* Data, uint32 Size);
+	CORE_API FMemoryImageWriter WritePointer(const FTypeLayoutDesc& StaticTypeDesc, const FTypeLayoutDesc& DerivedTypeDesc, uint32* OutOffsetToBase = nullptr);
+	CORE_API FMemoryImageWriter WritePointer(const FTypeLayoutDesc& TypeDesc);
+	CORE_API uint32 WriteNullPointer();
+	CORE_API uint32 WriteRawPointerSizedBytes(uint64 PointerValue);
+	CORE_API uint32 WriteVTable(const FTypeLayoutDesc& TypeDesc, const FTypeLayoutDesc& DerivedTypeDesc);
+	CORE_API uint32 WriteFMemoryImageName(int32 NumBytes, const FName& Name);
+	CORE_API uint32 WriteFScriptName(const FScriptName& Name);
 
 	template<typename T>
 	void WriteObject(const T& Object)
diff --git a/Engine/Source/Runtime/Core/Public/Serialization/StructuredArchive.h b/Engine/Source/Runtime/Core/Public/Serialization/StructuredArchive.h
--- a/Engine/Source/Runtime/Core/Public/Serialization/StructuredArchive.h
+++ b/Engine/Source/Runtime/Core/Public/Serialization/StructuredArchive.h
@@ -47,7 +47,7 @@ namespace UE::StructuredArchive::Private
  * Both reading and writing to the archive are *forward only* from an interface point of view. There is no point at which it is possible to 
  * require seeking.
  */
-class CORE_API FStructuredArchive
+class FStructuredArchive
 {
 	friend FStructuredArchiveSlot;
 	friend FStructuredArchiveRecord;
@@ -67,22 +67,22 @@ public:
 	 *
 	 * @param InFormatter Formatter for the archive data
 	 */
-	explicit FStructuredArchive(FArchiveFormatterType& InFormatter);
+	CORE_API explicit FStructuredArchive(FArchiveFormatterType& InFormatter);
 	
 	/**
 	 * Default destructor. Closes the archive.
 	 */
-	~FStructuredArchive();
+	CORE_API ~FStructuredArchive();
 
 	/**
 	 * Start writing to the archive, and gets an interface to the root slot.
 	 */
-	FStructuredArchiveSlot Open();
+	CORE_API FStructuredArchiveSlot Open();
 
 	/**
 	 * Flushes any remaining scope to the underlying formatter and closes the archive.
 	 */
-	void Close();
+	CORE_API void Close();
 
 	/**
 	 * Gets the serialization context from the underlying archive.
@@ -180,24 +180,24 @@ private:
 	/**
 	 * Enters the current slot for serializing a value. Asserts if the archive is not in a state about to write to an empty-slot.
 	 */
-	void EnterSlot(UE::StructuredArchive::Private::FSlotPosition Slot, bool bEnteringAttributedValue = false);
+	CORE_API void EnterSlot(UE::StructuredArchive::Private::FSlotPosition Slot, bool bEnteringAttributedValue = false);
 
 	/**
 	 * Enters the current slot, adding an element onto the stack. Asserts if the archive is not in a state about to write to an empty-slot.
 	 *
 	 * @return  The depth of the newly-entered slot.
 	 */
-	int32 EnterSlotAsType(UE::StructuredArchive::Private::FSlotPosition Slot, UE::StructuredArchive::Private::EElementType ElementType);
+	CORE_API int32 EnterSlotAsType(UE::StructuredArchive::Private::FSlotPosition Slot, UE::StructuredArchive::Private::EElementType ElementType);
 
 	/**
 	 * Leaves slot at the top of the current scope
 	 */
-	void LeaveSlot();
+	CORE_API void LeaveSlot();
 
 	/**
 	 * Switches to the scope for the given slot.
 	 */
-	void SetScope(UE::StructuredArchive::Private::FSlotPosition Slot);
+	CORE_API void SetScope(UE::StructuredArchive::Private::FSlotPosition Slot);
 #endif
 };
 
diff --git a/Engine/Source/Runtime/Core/Public/Serialization/StructuredArchiveAdapters.h b/Engine/Source/Runtime/Core/Public/Serialization/StructuredArchiveAdapters.h
--- a/Engine/Source/Runtime/Core/Public/Serialization/StructuredArchiveAdapters.h
+++ b/Engine/Source/Runtime/Core/Public/Serialization/StructuredArchiveAdapters.h
@@ -31,41 +31,41 @@ private:
 
 #if WITH_TEXT_ARCHIVE_SUPPORT
 
-class CORE_API FArchiveFromStructuredArchiveImpl : public FArchiveProxy
+class FArchiveFromStructuredArchiveImpl : public FArchiveProxy
 {
 	UE_NONCOPYABLE(FArchiveFromStructuredArchiveImpl)
 
 		struct FImpl;
 
 public:
-	explicit FArchiveFromStructuredArchiveImpl(FStructuredArchiveSlot Slot);
-	virtual ~FArchiveFromStructuredArchiveImpl();
+	CORE_API explicit FArchiveFromStructuredArchiveImpl(FStructuredArchiveSlot Slot);
+	CORE_API virtual ~FArchiveFromStructuredArchiveImpl();
 
-	virtual void Flush() override;
-	virtual bool Close() override;
+	CORE_API virtual void Flush() override;
+	CORE_API virtual bool Close() override;
 
-	virtual int64 Tell() override;
-	virtual int64 TotalSize() override;
-	virtual void Seek(int64 InPos) override;
-	virtual bool AtEnd() override;
+	CORE_API virtual int64 Tell() override;
+	CORE_API virtual int64 TotalSize() override;
+	CORE_API virtual void Seek(int64 InPos) override;
+	CORE_API virtual bool AtEnd() override;
 
 	using FArchive::operator<<; // For visibility of the overloads we don't override
 
 	//~ Begin FArchive Interface
-	virtual FArchive& operator<<(class FName& Value) override;
-	virtual FArchive& operator<<(class UObject*& Value) override;
-	virtual FArchive& operator<<(class FText& Value) override;
+	CORE_API virtual FArchive& operator<<(class FName& Value) override;
+	CORE_API virtual FArchive& operator<<(class UObject*& Value) override;
+	CORE_API virtual FArchive& operator<<(class FText& Value) override;
 	//~ End FArchive Interface
 
-	virtual void Serialize(void* V, int64 Length) override;
+	CORE_API virtual void Serialize(void* V, int64 Length) override;
 
-	virtual FArchive* GetCacheableArchive() override;
+	CORE_API virtual FArchive* GetCacheableArchive() override;
 
-	bool ContainsData() const;
+	CORE_API bool ContainsData() const;
 
 protected:
-	virtual bool Finalize(FStructuredArchiveRecord Record);
-	void OpenArchive();
+	CORE_API virtual bool Finalize(FStructuredArchiveRecord Record);
+	CORE_API void OpenArchive();
 
 private:
 	void Commit();
diff --git a/Engine/Source/Runtime/Core/Public/Serialization/StructuredArchiveChildReader.h b/Engine/Source/Runtime/Core/Public/Serialization/StructuredArchiveChildReader.h
--- a/Engine/Source/Runtime/Core/Public/Serialization/StructuredArchiveChildReader.h
+++ b/Engine/Source/Runtime/Core/Public/Serialization/StructuredArchiveChildReader.h
@@ -20,11 +20,11 @@ class FStructuredArchiveFormatter;
  * Allows serialization code to get an archive instance for the current location, so that it can return to it
  * later on after the master archive has potentially moved on into a different location in the file.
  */
-class CORE_API FStructuredArchiveChildReader
+class FStructuredArchiveChildReader
 {
 public:
-	explicit FStructuredArchiveChildReader(FStructuredArchiveSlot InSlot);
-	~FStructuredArchiveChildReader();
+	CORE_API explicit FStructuredArchiveChildReader(FStructuredArchiveSlot InSlot);
+	CORE_API ~FStructuredArchiveChildReader();
 
 	// Non-copyable
 	FStructuredArchiveChildReader(FStructuredArchiveChildReader&&) = delete;
diff --git a/Engine/Source/Runtime/Core/Public/Serialization/StructuredArchiveFormatter.h b/Engine/Source/Runtime/Core/Public/Serialization/StructuredArchiveFormatter.h
--- a/Engine/Source/Runtime/Core/Public/Serialization/StructuredArchiveFormatter.h
+++ b/Engine/Source/Runtime/Core/Public/Serialization/StructuredArchiveFormatter.h
@@ -56,10 +56,10 @@ enum class EArchiveValueType
  * a fully defined document tree, and allow querying additional properties that aren't available when reading from
  * a pure binary archive. These functions will assert if called on a binary archive.
  */
-class CORE_API FStructuredArchiveFormatter
+class FStructuredArchiveFormatter
 {
 public:
-	virtual ~FStructuredArchiveFormatter();
+	CORE_API virtual ~FStructuredArchiveFormatter();
 
 	virtual FArchive& GetUnderlyingArchive() = 0;
 	virtual FStructuredArchiveFormatter* CreateSubtreeReader() { return this; }
diff --git a/Engine/Source/Runtime/Core/Public/Serialization/StructuredArchiveSlotBase.h b/Engine/Source/Runtime/Core/Public/Serialization/StructuredArchiveSlotBase.h
--- a/Engine/Source/Runtime/Core/Public/Serialization/StructuredArchiveSlotBase.h
+++ b/Engine/Source/Runtime/Core/Public/Serialization/StructuredArchiveSlotBase.h
@@ -44,7 +44,7 @@ namespace UE::StructuredArchive::Private
 	};
 
 	// Represents a position of a slot within the hierarchy.
-	class CORE_API FSlotPosition
+	class FSlotPosition
 	{
 		friend class FStructuredArchive;
 
@@ -60,7 +60,7 @@ namespace UE::StructuredArchive::Private
 	};
 
 	// The base class of all slot types
-	class CORE_API FSlotBase
+	class FSlotBase
 #if WITH_TEXT_ARCHIVE_SUPPORT
 		: protected FSlotPosition
 #endif
@@ -86,8 +86,8 @@ namespace UE::StructuredArchive::Private
 		}
 #endif
 
-		FArchive& GetUnderlyingArchive() const;
-		const FArchiveState& GetArchiveState() const;
+		CORE_API FArchive& GetUnderlyingArchive() const;
+		CORE_API const FArchiveState& GetArchiveState() const;
 
 	protected:
 		FStructuredArchive& StructuredArchive;
diff --git a/Engine/Source/Runtime/Core/Public/Serialization/StructuredArchiveSlots.h b/Engine/Source/Runtime/Core/Public/Serialization/StructuredArchiveSlots.h
--- a/Engine/Source/Runtime/Core/Public/Serialization/StructuredArchiveSlots.h
+++ b/Engine/Source/Runtime/Core/Public/Serialization/StructuredArchiveSlots.h
@@ -42,37 +42,37 @@ namespace UE::StructuredArchive::Private
  * and can merely have a value serialized into it. That value may be a literal (eg. int, float) or compound object
  * (eg. object, array, map).
  */
-class CORE_API FStructuredArchiveSlot final : public UE::StructuredArchive::Private::FSlotBase
+class FStructuredArchiveSlot final : public UE::StructuredArchive::Private::FSlotBase
 {
 public:
-	FStructuredArchiveRecord EnterRecord();
-	FStructuredArchiveArray EnterArray(int32& Num);
-	FStructuredArchiveStream EnterStream();
-	FStructuredArchiveMap EnterMap(int32& Num);
-	FStructuredArchiveSlot EnterAttribute(FArchiveFieldName AttributeName);
-	TOptional<FStructuredArchiveSlot> TryEnterAttribute(FArchiveFieldName AttributeName, bool bEnterWhenWriting);
+	CORE_API FStructuredArchiveRecord EnterRecord();
+	CORE_API FStructuredArchiveArray EnterArray(int32& Num);
+	CORE_API FStructuredArchiveStream EnterStream();
+	CORE_API FStructuredArchiveMap EnterMap(int32& Num);
+	CORE_API FStructuredArchiveSlot EnterAttribute(FArchiveFieldName AttributeName);
+	CORE_API TOptional<FStructuredArchiveSlot> TryEnterAttribute(FArchiveFieldName AttributeName, bool bEnterWhenWriting);
 
 	// We don't support chaining writes to a single slot, so this returns void.
-	void operator << (uint8& Value);
-	void operator << (uint16& Value);
-	void operator << (uint32& Value);
-	void operator << (uint64& Value);
-	void operator << (int8& Value);
-	void operator << (int16& Value);
-	void operator << (int32& Value);
-	void operator << (int64& Value);
-	void operator << (float& Value);
-	void operator << (double& Value);
-	void operator << (bool& Value);
-	void operator << (FString& Value);
-	void operator << (FName& Value);
-	void operator << (UObject*& Value);
-	void operator << (FText& Value);
-	void operator << (FWeakObjectPtr& Value);
-	void operator << (FSoftObjectPtr& Value);
-	void operator << (FSoftObjectPath& Value);
-	void operator << (FLazyObjectPtr& Value);
-	void operator << (FObjectPtr& Value);
+	CORE_API void operator << (uint8& Value);
+	CORE_API void operator << (uint16& Value);
+	CORE_API void operator << (uint32& Value);
+	CORE_API void operator << (uint64& Value);
+	CORE_API void operator << (int8& Value);
+	CORE_API void operator << (int16& Value);
+	CORE_API void operator << (int32& Value);
+	CORE_API void operator << (int64& Value);
+	CORE_API void operator << (float& Value);
+	CORE_API void operator << (double& Value);
+	CORE_API void operator << (bool& Value);
+	CORE_API void operator << (FString& Value);
+	CORE_API void operator << (FName& Value);
+	CORE_API void operator << (UObject*& Value);
+	CORE_API void operator << (FText& Value);
+	CORE_API void operator << (FWeakObjectPtr& Value);
+	CORE_API void operator << (FSoftObjectPtr& Value);
+	CORE_API void operator << (FSoftObjectPath& Value);
+	CORE_API void operator << (FLazyObjectPtr& Value);
+	CORE_API void operator << (FObjectPtr& Value);
 
 	template <typename T>
 	FORCEINLINE void operator<<(TEnumAsByte<T>& Value)
@@ -110,9 +110,9 @@ public:
 		}
 	}
 
-	void Serialize(TArray<uint8>& Data);
-	void Serialize(void* Data, uint64 DataSize);
-	bool IsFilled() const;
+	CORE_API void Serialize(TArray<uint8>& Data);
+	CORE_API void Serialize(void* Data, uint64 DataSize);
+	CORE_API bool IsFilled() const;
 
 private:
 	friend FStructuredArchive;
@@ -130,16 +130,16 @@ private:
  * Represents a record in the structured archive. An object contains slots that are identified by FArchiveName,
  * which may be compiled out with binary-only archives.
  */
-class CORE_API FStructuredArchiveRecord final : public UE::StructuredArchive::Private::FSlotBase
+class FStructuredArchiveRecord final : public UE::StructuredArchive::Private::FSlotBase
 {
 public:
-	FStructuredArchiveSlot EnterField(FArchiveFieldName Name);
-	FStructuredArchiveRecord EnterRecord(FArchiveFieldName Name);
-	FStructuredArchiveArray EnterArray(FArchiveFieldName Name, int32& Num);
-	FStructuredArchiveStream EnterStream(FArchiveFieldName Name);
-	FStructuredArchiveMap EnterMap(FArchiveFieldName Name, int32& Num);
+	CORE_API FStructuredArchiveSlot EnterField(FArchiveFieldName Name);
+	CORE_API FStructuredArchiveRecord EnterRecord(FArchiveFieldName Name);
+	CORE_API FStructuredArchiveArray EnterArray(FArchiveFieldName Name, int32& Num);
+	CORE_API FStructuredArchiveStream EnterStream(FArchiveFieldName Name);
+	CORE_API FStructuredArchiveMap EnterMap(FArchiveFieldName Name, int32& Num);
 
-	TOptional<FStructuredArchiveSlot> TryEnterField(FArchiveFieldName Name, bool bEnterForSaving);
+	CORE_API TOptional<FStructuredArchiveSlot> TryEnterField(FArchiveFieldName Name, bool bEnterForSaving);
 
 	template<typename T> FORCEINLINE FStructuredArchiveRecord& operator<<(UE::StructuredArchive::Private::TNamedValue<T> Item)
 	{
@@ -158,10 +158,10 @@ private:
  * Represents an array in the structured archive. An object contains slots that are identified by a FArchiveFieldName,
  * which may be compiled out with binary-only archives.
  */
-class CORE_API FStructuredArchiveArray final : public UE::StructuredArchive::Private::FSlotBase
+class FStructuredArchiveArray final : public UE::StructuredArchive::Private::FSlotBase
 {
 public:
-	FStructuredArchiveSlot EnterElement();
+	CORE_API FStructuredArchiveSlot EnterElement();
 
 	template<typename T> FORCEINLINE FStructuredArchiveArray& operator<<(T& Item)
 	{
@@ -179,10 +179,10 @@ private:
 /**
  * Represents an unsized sequence of slots in the structured archive (similar to an array, but without a known size).
  */
-class CORE_API FStructuredArchiveStream final : public UE::StructuredArchive::Private::FSlotBase
+class FStructuredArchiveStream final : public UE::StructuredArchive::Private::FSlotBase
 {
 public:
-	FStructuredArchiveSlot EnterElement();
+	CORE_API FStructuredArchiveSlot EnterElement();
 
 	template<typename T> FORCEINLINE FStructuredArchiveStream& operator<<(T& Item)
 	{
@@ -201,10 +201,10 @@ private:
  * Represents a map in the structured archive. A map is similar to a record, but keys can be read back out from an archive.
  * (This is an important distinction for binary archives).
  */
-class CORE_API FStructuredArchiveMap final : public UE::StructuredArchive::Private::FSlotBase
+class FStructuredArchiveMap final : public UE::StructuredArchive::Private::FSlotBase
 {
 public:
-	FStructuredArchiveSlot EnterElement(FString& Name);
+	CORE_API FStructuredArchiveSlot EnterElement(FString& Name);
 
 private:
 	friend FStructuredArchive;
diff --git a/Engine/Source/Runtime/Core/Public/Stats/Stats2.h b/Engine/Source/Runtime/Core/Public/Stats/Stats2.h
--- a/Engine/Source/Runtime/Core/Public/Stats/Stats2.h
+++ b/Engine/Source/Runtime/Core/Public/Stats/Stats2.h
@@ -150,36 +150,36 @@ public:
 CORE_API bool DirectStatsCommand(const TCHAR* Cmd, bool bBlockForCompletion = false, FOutputDevice* Ar = nullptr);
 
 /** Helper struct that contains method available even when the stats are disabled. */
-struct CORE_API FStats
+struct FStats
 {
 	/** Delegate to fire every time we need to advance the stats for the rendering thread. */
 	DECLARE_DELEGATE_ThreeParams( FOnAdvanceRenderingThreadStats, bool /*bDiscardCallstack*/, int64 /*StatsFrame*/, int32 /*PrimaryDisableChangeTagStartFrame*/ );
 
 	/** Advances stats for the current frame. */
-	static void AdvanceFrame( bool bDiscardCallstack, const FOnAdvanceRenderingThreadStats& AdvanceRenderingThreadStatsDelegate = FOnAdvanceRenderingThreadStats() );
+	static CORE_API void AdvanceFrame( bool bDiscardCallstack, const FOnAdvanceRenderingThreadStats& AdvanceRenderingThreadStatsDelegate = FOnAdvanceRenderingThreadStats() );
 
 	/** Advances stats for commandlets, only valid if the command line has the proper token. @see HasStatsForCommandletsToken */
-	static void TickCommandletStats();
+	static CORE_API void TickCommandletStats();
 
 	/**
 	* @return true, if the command line has the LoadTimeStatsForCommandlet or LoadTimeFileForCommandlet token which enables stats in the commandlets.
 	* !!!CAUTION!!! You need to manually advance stats frame in order to maintain the data integrity and not to leak the memory.
 	*/
-	static bool EnabledForCommandlet();
+	static CORE_API bool EnabledForCommandlet();
 
 	/**
 	* @return true, if the command line has the LoadTimeStatsForCommandlet token which enables LoadTimeStats equivalent for commandlets.
 	* All collected stats will be dumped to the log file at the end of running the specified commandlet.
 	*/
-	static bool HasLoadTimeStatsForCommandletToken();
+	static CORE_API bool HasLoadTimeStatsForCommandletToken();
 
 	/**
 	* @return true, if the command line has the LoadTimeFileForCommandlet token which enables LoadTimeFile equivalent for commandlets.
 	*/
-	static bool HasLoadTimeFileForCommandletToken();
+	static CORE_API bool HasLoadTimeFileForCommandletToken();
 
 	/** Current game thread stats frame. */
-	static TAtomic<int32> GameThreadStatsFrame;
+	static CORE_API TAtomic<int32> GameThreadStatsFrame;
 };
 
 enum class EStatFlags : uint8
diff --git a/Engine/Source/Runtime/Core/Public/Stats/StatsData.h b/Engine/Source/Runtime/Core/Public/Stats/StatsData.h
--- a/Engine/Source/Runtime/Core/Public/Stats/StatsData.h
+++ b/Engine/Source/Runtime/Core/Public/Stats/StatsData.h
@@ -39,41 +39,41 @@ enum
 };
 
 /** Holds stats related constants. */
-struct CORE_API FStatConstants
+struct FStatConstants
 {
 	/** Special name for thread root. */
-	static const FName NAME_ThreadRoot;
+	static CORE_API const FName NAME_ThreadRoot;
 
 	/** This is a special group name used to store threads metadata. */
-	static const char* ThreadGroupName;
-	static const FName NAME_ThreadGroup;
+	static CORE_API const char* ThreadGroupName;
+	static CORE_API const FName NAME_ThreadGroup;
 
 	/** Stat raw name for seconds per cycle. */
-	static const FName RAW_SecondsPerCycle;
+	static CORE_API const FName RAW_SecondsPerCycle;
 
 	/** Special case category, when we want to Stat to appear at the root of the menu (leaving the category blank omits it from the menu entirely) */
-	static const FName NAME_NoCategory;
+	static CORE_API const FName NAME_NoCategory;
 
 	/** Extension used to save a stats file. */
-	static const FString StatsFileExtension;
+	static CORE_API const FString StatsFileExtension;
 
 	/** Extension used to save a raw stats file, may be changed to the same as a regular stats file. */
-	static const FString StatsFileRawExtension;
+	static CORE_API const FString StatsFileRawExtension;
 
 	/** Indicates that the item is a thread. */
-	static const FString ThreadNameMarker;
+	static CORE_API const FString ThreadNameMarker;
 
 	/** A raw name for the event wait with id. */
-	static const FName RAW_EventWaitWithId;
+	static CORE_API const FName RAW_EventWaitWithId;
 
 	/** A raw name for the event trigger with id. */
-	static const FName RAW_EventTriggerWithId;
+	static CORE_API const FName RAW_EventTriggerWithId;
 
 	/** A raw name for the stat marker. */
-	static const FName RAW_NamedMarker;
+	static CORE_API const FName RAW_NamedMarker;
 
 	/** A special meta data used to advance the frame. */
-	static const FStatNameAndInfo AdvanceFrame;
+	static CORE_API const FStatNameAndInfo AdvanceFrame;
 };
 
 /** Parse a typed value into the specified out parameter.
@@ -216,7 +216,7 @@ struct FStatCallCountComparer<FStatMessage>
 * An indirect array of stat packets.
 */
 //@todo can we just use TIndirectArray here?
-struct CORE_API FStatPacketArray
+struct FStatPacketArray
 {
 	TArray<FStatPacket*> Packets;
 
@@ -226,7 +226,7 @@ struct CORE_API FStatPacketArray
 	}
 
 	/** Deletes all stats packets. */
-	void Empty();
+	CORE_API void Empty();
 
 	void RemovePtrsButNoData()
 	{
@@ -237,7 +237,7 @@ struct CORE_API FStatPacketArray
 /**
 * A call stack of stat messages. Normally we don't store or transmit these; they are used for stats processing and visualization.
 */
-struct CORE_API FRawStatStackNode
+struct FRawStatStackNode
 {
 	/** The stat message corresponding to this node **/
 	FStatMessage Meta; // includes aggregated inclusive time and call counts packed into the int64
@@ -252,7 +252,7 @@ struct CORE_API FRawStatStackNode
 	}
 
 	/** Copy Constructor **/
-	explicit FRawStatStackNode(FRawStatStackNode const& Other);
+	CORE_API explicit FRawStatStackNode(FRawStatStackNode const& Other);
 
 	/** Constructor used to build a child from a stat messasge **/
 	explicit FRawStatStackNode(FStatMessage const& InMeta)
@@ -283,36 +283,36 @@ struct CORE_API FRawStatStackNode
 		DeleteAllChildrenNodes();
 	}
 
-	void MergeMax(FRawStatStackNode const& Other);
-	void MergeAdd(FRawStatStackNode const& Other);
-	void Divide(uint32 Div);
+	CORE_API void MergeMax(FRawStatStackNode const& Other);
+	CORE_API void MergeAdd(FRawStatStackNode const& Other);
+	CORE_API void Divide(uint32 Div);
 
 	/** Cull this tree, merging children below MinCycles long. */
-	void CullByCycles(int64 MinCycles);
+	CORE_API void CullByCycles(int64 MinCycles);
 
 	/** Cull this tree, merging children below NoCullLevels. */
-	void CullByDepth(int32 NoCullLevels);
+	CORE_API void CullByDepth(int32 NoCullLevels);
 
 	/** Adds name hiearchy. **/
-	void AddNameHierarchy(int32 CurrentPrefixDepth = 0);
+	CORE_API void AddNameHierarchy(int32 CurrentPrefixDepth = 0);
 
 	/** Adds self nodes. **/
-	void AddSelf();
+	CORE_API void AddSelf();
 
 	/** Print this tree to the log **/
-	void DebugPrint(TCHAR const* Filter = nullptr, int32 MaxDepth = MAX_int32, int32 Depth = 0) const;
+	CORE_API void DebugPrint(TCHAR const* Filter = nullptr, int32 MaxDepth = MAX_int32, int32 Depth = 0) const;
 
 	/** Print this tree to the log **/
-	void DebugPrintLeafFilter(TCHAR const* Filter) const;
+	CORE_API void DebugPrintLeafFilter(TCHAR const* Filter) const;
 
 	/** Print this tree to the log **/
-	void DebugPrintLeafFilterInner(TCHAR const* Filter, int32 Depth, TArray<FString>& Stack) const;
+	CORE_API void DebugPrintLeafFilterInner(TCHAR const* Filter, int32 Depth, TArray<FString>& Stack) const;
 
 	/** Condense this tree into a flat list using EStatOperation::ChildrenStart, EStatOperation::ChildrenEnd, EStatOperation::Leaf **/
-	void Encode(TArray<FStatMessage>& OutStats) const;
+	CORE_API void Encode(TArray<FStatMessage>& OutStats) const;
 
 	/** Sum the inclusive cycles of the children **/
-	int64 ChildCycles() const;
+	CORE_API int64 ChildCycles() const;
 
 	/** Sorts children using the specified comparer class. */
 	template< class TComparer >
@@ -327,8 +327,8 @@ struct CORE_API FRawStatStackNode
 	}
 
 	/** Custom new/delete; these nodes are always pooled, except the root which one normally puts on the stack. */
-	void* operator new(size_t Size);
-	void operator delete(void *RawMemory);
+	CORE_API void* operator new(size_t Size);
+	CORE_API void operator delete(void *RawMemory);
 
 protected:
 	void DeleteAllChildrenNodes()
@@ -493,7 +493,7 @@ struct IItemFilter
  * Tracks stat state and history
  * GetLocalState() is a singleton to the state for stats being collected in this executable.
  */
-class CORE_API FStatsThreadState
+class FStatsThreadState
 {
 	friend class FStatsThread;
 	friend struct FStatPacketArray;
@@ -506,20 +506,20 @@ class CORE_API FStatsThreadState
 	DECLARE_MULTICAST_DELEGATE_OneParam(FOnNewRawStatPacket, const FStatPacket*);
 
 	/** Internal method to scan the messages to update the current frame. */
-	void ScanForAdvance(const FStatMessagesArray& Data);
+	CORE_API void ScanForAdvance(const FStatMessagesArray& Data);
 
 	/** Internal method to scan the messages to update the current frame. */
-	void ScanForAdvance(FStatPacketArray& NewData);
+	CORE_API void ScanForAdvance(FStatPacketArray& NewData);
 
 public:
 	/** Internal method to update the internal metadata. **/
 	UE_DEPRECATED(5.2, "Please use the overload that takes TArray64.")
-	void ProcessMetaDataOnly(TArray<FStatMessage>& Data);
+	CORE_API void ProcessMetaDataOnly(TArray<FStatMessage>& Data);
 
-	void ProcessMetaDataOnly(TArray64<FStatMessage>& Data);
+	CORE_API void ProcessMetaDataOnly(TArray64<FStatMessage>& Data);
 
 	/** Toggles tracking the most memory expensive stats. */
-	void ToggleFindMemoryExtensiveStats();
+	CORE_API void ToggleFindMemoryExtensiveStats();
 
 	/** Resets stats for raw stats. */
 	void ResetStatsForRawStats()
@@ -530,38 +530,38 @@ public:
 
 private:
 	/** Internal method to accumulate any non-frame stats. **/
-	void ProcessNonFrameStats( FStatMessagesArray& Data, TSet<FName>* NonFrameStatsFound );
+	CORE_API void ProcessNonFrameStats( FStatMessagesArray& Data, TSet<FName>* NonFrameStatsFound );
 
 	/** Internal method to place the data into the history, discard and broadcast any new frames to anyone who cares. **/
-	void AddToHistoryAndEmpty( FStatPacketArray& NewData );
+	CORE_API void AddToHistoryAndEmpty( FStatPacketArray& NewData );
 
 	/** Does basic processing on the raw stats packets, discard and broadcast any new raw stats packets to anyone who cares. */
-	void ProcessRawStats(FStatPacketArray& NewData);
+	CORE_API void ProcessRawStats(FStatPacketArray& NewData);
 
 	/** Resets the state of the raw stats. */
-	void ResetRawStats();
+	CORE_API void ResetRawStats();
 
 	/** Reset the state of the regular stats. */
-	void ResetRegularStats();
+	CORE_API void ResetRegularStats();
 
 	/** Prepares fake FGameThreadStatsData to display the raw stats memory overhead. */
-	void UpdateStatMessagesMemoryUsage();
+	CORE_API void UpdateStatMessagesMemoryUsage();
 
 	/** Generates a list of most memory expensive stats and dump to the log. */
-	void FindAndDumpMemoryExtensiveStats( const FStatPacketArray& Frame );
+	CORE_API void FindAndDumpMemoryExtensiveStats( const FStatPacketArray& Frame );
 
 protected:
 	/** Called in response to SetLongName messages to update ShortNameToLongName and NotClearedEveryFrame **/
-	void FindOrAddMetaData(FStatMessage const& Item);
+	CORE_API void FindOrAddMetaData(FStatMessage const& Item);
 
 	/** Parse the raw frame history into a tree representation **/
-	void GetRawStackStats(int64 FrameNumber, FRawStatStackNode& Out, TArray<FStatMessage>* OutNonStackStats = nullptr) const;
+	CORE_API void GetRawStackStats(int64 FrameNumber, FRawStatStackNode& Out, TArray<FStatMessage>* OutNonStackStats = nullptr) const;
 
 	/** Parse the raw frame history into a flat-array-based tree representation **/
-	void Condense(int64 TargetFrame, TArray<FStatMessage>& OutStats) const;
+	CORE_API void Condense(int64 TargetFrame, TArray<FStatMessage>& OutStats) const;
 
 	/** Returns cycles spent in the specified thread. */
-	int64 GetFastThreadFrameTimeInternal(int64 TargetFrame, int32 ThreadID, EThreadType::Type Thread) const;
+	CORE_API int64 GetFastThreadFrameTimeInternal(int64 TargetFrame, int32 ThreadID, EThreadType::Type Thread) const;
 
 	/** Number of frames to keep in the history. **/
 	int32 HistoryFrames;
@@ -599,7 +599,7 @@ protected:
 public:
 
 	/** Constructor used by GetLocalState(), also used by the profiler to hold a previewing stats thread state. We don't keep many frames by default **/
-	FStatsThreadState(int32 InHistoryFrames = STAT_FRAME_SLOP + 10);
+	CORE_API FStatsThreadState(int32 InHistoryFrames = STAT_FRAME_SLOP + 10);
 
 	/** Delegate we fire every time we have a new complete frame of data. **/
 	mutable FOnNewFrameHistory NewFrameDelegate;
@@ -635,10 +635,10 @@ public:
 	TMap<int64, FStatPacketArray> History;
 
 	/** Return the oldest frame of data we have. **/
-	int64 GetOldestValidFrame() const;
+	CORE_API int64 GetOldestValidFrame() const;
 
 	/** Return the newest frame of data we have. **/
-	int64 GetLatestValidFrame() const;
+	CORE_API int64 GetLatestValidFrame() const;
 
 	/** Return true if this is a valid frame. **/
 	bool IsFrameValid(int64 Frame) const
@@ -647,14 +647,14 @@ public:
 	}
 
 	/** Used by the hitch detector. Does a very fast look a thread to decide if this is a hitch frame. **/
-	int64 GetFastThreadFrameTime(int64 TargetFrame, EThreadType::Type Thread) const;
-	int64 GetFastThreadFrameTime(int64 TargetFrame, uint32 ThreadID) const;
+	CORE_API int64 GetFastThreadFrameTime(int64 TargetFrame, EThreadType::Type Thread) const;
+	CORE_API int64 GetFastThreadFrameTime(int64 TargetFrame, uint32 ThreadID) const;
 
 	/** For the specified stat packet looks for the thread name. */
-	FName GetStatThreadName( const FStatPacket& Packet ) const;
+	CORE_API FName GetStatThreadName( const FStatPacket& Packet ) const;
 
 	/** Looks in the history for a condensed frame, builds it if it isn't there. **/
-	TArray<FStatMessage> const& GetCondensedHistory(int64 TargetFrame) const;
+	CORE_API TArray<FStatMessage> const& GetCondensedHistory(int64 TargetFrame) const;
 
 	/** Looks in the history for a stat packet array, the stat packet array must be in the history. */
 	const FStatPacketArray& GetStatPacketArray( int64 TargetFrame ) const
@@ -665,28 +665,28 @@ public:
 	}
 
 	/** Gets the old-skool flat grouped inclusive stats. These ignore recursion, merge threads, etc and so generally the condensed callstack is less confusing. **/
-	void GetInclusiveAggregateStackStats( const TArray<FStatMessage>& CondensedMessages, TArray<FStatMessage>& OutStats, IItemFilter* Filter = nullptr, bool bAddNonStackStats = true, TMap<FName, TArray<FStatMessage>>* OptionalOutThreadBreakdownMap = nullptr ) const;
+	CORE_API void GetInclusiveAggregateStackStats( const TArray<FStatMessage>& CondensedMessages, TArray<FStatMessage>& OutStats, IItemFilter* Filter = nullptr, bool bAddNonStackStats = true, TMap<FName, TArray<FStatMessage>>* OptionalOutThreadBreakdownMap = nullptr ) const;
 
 	/** Gets the old-skool flat grouped inclusive stats. These ignore recursion, merge threads, etc and so generally the condensed callstack is less confusing. **/
-	void GetInclusiveAggregateStackStats(int64 TargetFrame, TArray<FStatMessage>& OutStats, IItemFilter* Filter = nullptr, bool bAddNonStackStats = true, TMap<FName, TArray<FStatMessage>>* OptionalOutThreadBreakdownMap = nullptr) const;
+	CORE_API void GetInclusiveAggregateStackStats(int64 TargetFrame, TArray<FStatMessage>& OutStats, IItemFilter* Filter = nullptr, bool bAddNonStackStats = true, TMap<FName, TArray<FStatMessage>>* OptionalOutThreadBreakdownMap = nullptr) const;
 
 	/** Gets the old-skool flat grouped exclusive stats. These merge threads, etc and so generally the condensed callstack is less confusing. **/
-	void GetExclusiveAggregateStackStats( const TArray<FStatMessage>& CondensedMessages, TArray<FStatMessage>& OutStats, IItemFilter* Filter = nullptr, bool bAddNonStackStats = true ) const;
+	CORE_API void GetExclusiveAggregateStackStats( const TArray<FStatMessage>& CondensedMessages, TArray<FStatMessage>& OutStats, IItemFilter* Filter = nullptr, bool bAddNonStackStats = true ) const;
 
 	/** Gets the old-skool flat grouped exclusive stats. These merge threads, etc and so generally the condensed callstack is less confusing. **/
-	void GetExclusiveAggregateStackStats(int64 TargetFrame, TArray<FStatMessage>& OutStats, IItemFilter* Filter = nullptr, bool bAddNonStackStats = true) const;
+	CORE_API void GetExclusiveAggregateStackStats(int64 TargetFrame, TArray<FStatMessage>& OutStats, IItemFilter* Filter = nullptr, bool bAddNonStackStats = true) const;
 
 	/** Used to turn the condensed version of stack stats back into a tree for easier handling. **/
-	void UncondenseStackStats( const TArray<FStatMessage>& CondensedMessages, FRawStatStackNode& Root, IItemFilter* Filter = nullptr, TArray<FStatMessage>* OutNonStackStats = nullptr ) const;
+	CORE_API void UncondenseStackStats( const TArray<FStatMessage>& CondensedMessages, FRawStatStackNode& Root, IItemFilter* Filter = nullptr, TArray<FStatMessage>* OutNonStackStats = nullptr ) const;
 
 	/** Used to turn the condensed version of stack stats back into a tree for easier handling. **/
-	void UncondenseStackStats(int64 TargetFrame, FRawStatStackNode& Root, IItemFilter* Filter = nullptr, TArray<FStatMessage>* OutNonStackStats = nullptr) const;
+	CORE_API void UncondenseStackStats(int64 TargetFrame, FRawStatStackNode& Root, IItemFilter* Filter = nullptr, TArray<FStatMessage>* OutNonStackStats = nullptr) const;
 
 	/** Adds missing stats to the group so it doesn't jitter. **/
-	void AddMissingStats(TArray<FStatMessage>& Dest, TSet<FName> const& EnabledItems) const;
+	CORE_API void AddMissingStats(TArray<FStatMessage>& Dest, TSet<FName> const& EnabledItems) const;
 
 	/** Singleton to get the stats being collected by this executable. Can be only accessed from the stats thread. **/
-	static FStatsThreadState& GetLocalState();
+	static CORE_API FStatsThreadState& GetLocalState();
 };
 
 //@todo split header
@@ -694,22 +694,22 @@ public:
 /**
 * Set of utility functions for dealing with stats
 */
-struct CORE_API FStatsUtils
+struct FStatsUtils
 {
 	/** Divides a stat by an integer. Used to finish an average. **/
-	static void DivideStat(FStatMessage& Dest, uint32 Div);
+	static CORE_API void DivideStat(FStatMessage& Dest, uint32 Div);
 
 	/** Merges two stat message arrays, adding corresponding number values **/
-	static void AddMergeStatArray(TArray<FStatMessage>& Dest, TArray<FStatMessage> const& Item);
+	static CORE_API void AddMergeStatArray(TArray<FStatMessage>& Dest, TArray<FStatMessage> const& Item);
 
 	/** Merges two stat message arrays, corresponding number values are set to the max **/
-	static void MaxMergeStatArray(TArray<FStatMessage>& Dest, TArray<FStatMessage> const& Item);
+	static CORE_API void MaxMergeStatArray(TArray<FStatMessage>& Dest, TArray<FStatMessage> const& Item);
 
 	/** Divides a stat array by an integer. Used to finish an average. **/
-	static void DivideStatArray(TArray<FStatMessage>& Dest, uint32 Div);
+	static CORE_API void DivideStatArray(TArray<FStatMessage>& Dest, uint32 Div);
 
 	/** Apply Op or Item.Op to Item and Dest, soring the result in dest. **/
-	static void AccumulateStat(FStatMessage& Dest, FStatMessage const& Item, EStatOperation::Type Op = EStatOperation::Invalid, bool bAllowNameMismatch = false);
+	static CORE_API void AccumulateStat(FStatMessage& Dest, FStatMessage const& Item, EStatOperation::Type Op = EStatOperation::Invalid, bool bAllowNameMismatch = false);
 
 	/** Adds a non-stack stats */
 	static void AddNonStackStats( const FName LongName, const FStatMessage& Item, const EStatOperation::Type Op, TMap<FName, FStatMessage>& out_NonStackStats )
@@ -735,10 +735,10 @@ struct CORE_API FStatsUtils
 	}
 
 	/** Spews a stat message, not all messages are supported, this is used for reports, so it focuses on numeric output, not say the operation. **/
-	static FString DebugPrint(FStatMessage const& Item);
+	static CORE_API FString DebugPrint(FStatMessage const& Item);
 
 	/** Gets friendly names from a stat message, not all messages are supported, this is used for reports, **/
-	static void GetNameAndGroup(FStatMessage const& Item, FString& OutName, FString& OutGroup);
+	static CORE_API void GetNameAndGroup(FStatMessage const& Item, FString& OutName, FString& OutGroup);
 
 	/** Subtract a CycleScopeStart from a CycleScopeEnd to create a IsPackedCCAndDuration with the call count and inclusive cycles. **/
 	static FStatMessage ComputeCall(FStatMessage const& ScopeStart, FStatMessage const& ScopeEnd)
@@ -787,8 +787,8 @@ struct CORE_API FStatsUtils
 	}
 
 	/** Internal use, converts arbitrary string to and from an escaped notation for storage in an FName. **/
-	static void ToEscapedString(FStringView Input, FStringBuilderBase& Output);
-	static FString FromEscapedString(const TCHAR* Escaped);
+	static CORE_API void ToEscapedString(FStringView Input, FStringBuilderBase& Output);
+	static CORE_API FString FromEscapedString(const TCHAR* Escaped);
 	
 	static FString BuildUniqueThreadName( uint32 InThreadID )
 	{
@@ -819,19 +819,19 @@ struct CORE_API FStatsUtils
 /**
 * Contains helpers functions to manage complex stat messages.
 */
-struct CORE_API FComplexStatUtils
+struct FComplexStatUtils
 {
 	/** Accumulates a stat message into a complex stat message. */
-	static void AddAndMinMax( FComplexStatMessage& Dest, const FStatMessage& Item, EComplexStatField::Type SumIndex, EComplexStatField::Type MaxIndex, EComplexStatField::Type MinIndex);
+	static CORE_API void AddAndMinMax( FComplexStatMessage& Dest, const FStatMessage& Item, EComplexStatField::Type SumIndex, EComplexStatField::Type MaxIndex, EComplexStatField::Type MinIndex);
 
 	/** Divides a complex stat message by the specified value. */
-	static void DivideStat( FComplexStatMessage& Dest, uint32 Div, EComplexStatField::Type SumIndex, EComplexStatField::Type DestIndex );
+	static CORE_API void DivideStat( FComplexStatMessage& Dest, uint32 Div, EComplexStatField::Type SumIndex, EComplexStatField::Type DestIndex );
 
 	/** Merges two complex stat message arrays. */
-	static void MergeAddAndMinMaxArray( TArray<FComplexStatMessage>& Dest, const TArray<FStatMessage>& Source, EComplexStatField::Type SumIndex, EComplexStatField::Type MaxIndex, EComplexStatField::Type MinIndex);
+	static CORE_API void MergeAddAndMinMaxArray( TArray<FComplexStatMessage>& Dest, const TArray<FStatMessage>& Source, EComplexStatField::Type SumIndex, EComplexStatField::Type MaxIndex, EComplexStatField::Type MinIndex);
 
 	/** Divides a complex stat array by the specified value. */
-	static void DiviveStatArray( TArray<FComplexStatMessage>& Dest, uint32 Div, EComplexStatField::Type SumIndex, EComplexStatField::Type DestIndex );
+	static CORE_API void DiviveStatArray( TArray<FComplexStatMessage>& Dest, uint32 Div, EComplexStatField::Type SumIndex, EComplexStatField::Type DestIndex );
 };
 
 //@todo split header
diff --git a/Engine/Source/Runtime/Core/Public/Stats/StatsFile.h b/Engine/Source/Runtime/Core/Public/Stats/StatsFile.h
--- a/Engine/Source/Runtime/Core/Public/Stats/StatsFile.h
+++ b/Engine/Source/Runtime/Core/Public/Stats/StatsFile.h
@@ -110,7 +110,7 @@ struct EStatsFileConstants
 -----------------------------------------------------------------------------*/
 
 /** Helper struct used to operate on the compressed data. */
-struct CORE_API FCompressedStatsData
+struct FCompressedStatsData
 {
 	/**
 	 * Initialization constructor 
@@ -359,14 +359,14 @@ struct FStatsFrameInfo
 -----------------------------------------------------------------------------*/
 
 /** Struct used to send a stream of stat messages. */
-struct CORE_API FStatsWriteStream
+struct FStatsWriteStream
 {
 protected:
 	/** Writes metadata messages into the stream. */
-	void WriteMetadata( FArchive& Ar );
+	CORE_API void WriteMetadata( FArchive& Ar );
 
 	/** Writes condensed messages into the stream. */
-	void WriteCondensedMessages( FArchive& Ar, int64 TargetFrame );
+	CORE_API void WriteCondensedMessages( FArchive& Ar, int64 TargetFrame );
 
 	/** Sends an FName, and the string it represents if we have not sent that string before. **/
 	FORCEINLINE_STATS void WriteFName( FArchive& Ar, FStatNameAndInfo NameAndInfo )
@@ -438,7 +438,7 @@ protected:
 -----------------------------------------------------------------------------*/
 
 /** Interface for writing stats data. Can be only used in the stats thread. */
-struct CORE_API IStatsWriteFile : public FStatsWriteStream
+struct IStatsWriteFile : public FStatsWriteStream
 {
 	friend class FAsyncStatsWrite;
 
@@ -476,7 +476,7 @@ protected:
 
 protected:
 	/** Default constructor. */
-	IStatsWriteFile();
+	CORE_API IStatsWriteFile();
 
 public:
 	/** Destructor. */
@@ -484,15 +484,15 @@ public:
 	{}
 
 	/** Creates a file writer and registers for the data delegate. */
-	void Start( const FString& InFilename );
+	CORE_API void Start( const FString& InFilename );
 
 	/** Finalizes writing the stats data and unregisters the data delegate. */
-	void Stop();
+	CORE_API void Stop();
 
 	/**
 	 * @return the stats file metadata description like the current size and the duration of the stats session.
 	 */
-	FText GetFileMetaDesc() const;
+	CORE_API FText GetFileMetaDesc() const;
 
 protected:
 	/** Sets the data delegate used to receive a stats data. */
@@ -508,21 +508,21 @@ protected:
 	}
 
 	/** Writes magic value, dummy header and initial metadata. */
-	void WriteHeader();
+	CORE_API void WriteHeader();
 
 protected:
 	/**	Finalizes writing to the file. */
-	void Finalize();
+	CORE_API void Finalize();
 
 	/** Waits on the async send task to complete. */
-	void WaitTask();
+	CORE_API void WaitTask();
 
 	/** Sends the data to the file via async task. */
-	void SendTask();
+	CORE_API void SendTask();
 };
 
 /** Helper struct used to write regular stats to the file. */
-struct CORE_API FStatsWriteFile : public IStatsWriteFile
+struct FStatsWriteFile : public IStatsWriteFile
 {
 	friend class FAsyncStatsWrite;
 
@@ -538,19 +538,19 @@ public:
 	}
 
 protected:
-	virtual void SetDataDelegate( bool bSet ) override;
+	CORE_API virtual void SetDataDelegate( bool bSet ) override;
 
-	virtual void FinalizeSavingData( int64 FrameFileOffset ) override;
+	CORE_API virtual void FinalizeSavingData( int64 FrameFileOffset ) override;
 
 	/**
 	 *	Grabs a frame from the local FStatsThreadState and adds it to the output.
 	 *	Called from the stats thread, but the data is saved using the the FAsyncStatsWrite. 
 	 */
-	void WriteFrame( int64 TargetFrame );
+	CORE_API void WriteFrame( int64 TargetFrame );
 };
 
 /** Helper struct used to write raw stats to the file. */
-struct CORE_API FRawStatsWriteFile : public IStatsWriteFile
+struct FRawStatsWriteFile : public IStatsWriteFile
 {
 	bool bWrittenOffsetToData;
 
@@ -563,12 +563,12 @@ public:
 	}
 
 protected:
-	virtual void SetDataDelegate( bool bSet ) override;
+	CORE_API virtual void SetDataDelegate( bool bSet ) override;
 
-	void WriteRawStatPacket( const FStatPacket* StatPacket );
+	CORE_API void WriteRawStatPacket( const FStatPacket* StatPacket );
 
 	/** Write a stat packed into the specified archive. */
-	void WriteStatPacket( FArchive& Ar, FStatPacket& StatPacket );
+	CORE_API void WriteStatPacket( FArchive& Ar, FStatPacket& StatPacket );
 };
 
 /*-----------------------------------------------------------------------------
@@ -576,7 +576,7 @@ protected:
 -----------------------------------------------------------------------------*/
 
 /** Tracks stat state and history for loaded stats. */
-class CORE_API FStatsLoadedState : public FStatsThreadState
+class FStatsLoadedState : public FStatsThreadState
 {
 	friend struct FStatsReadFile;
 
@@ -588,24 +588,24 @@ public:
 	{}
 
 	/** Method to update the internal metadata. Metadata messages are removed from the array. */
-	void ProcessMetaDataAndLeaveDataOnly( TArray<FStatMessage>& CondensedMessages );
+	CORE_API void ProcessMetaDataAndLeaveDataOnly( TArray<FStatMessage>& CondensedMessages );
 
 	/**
 	* Method to place the data into the history,
 	* maintains the history based on the requested number of frames to keep in the history.
 	* The condensed messages are emplaced in the condensed history.
 	*/
-	void AddFrameFromCondensedMessages( TArray<FStatMessage>& CondensedMessages );
+	CORE_API void AddFrameFromCondensedMessages( TArray<FStatMessage>& CondensedMessages );
 
 	/** Return the oldest frame of data we have. **/
-	int64 GetOldestValidFrame() const;
+	CORE_API int64 GetOldestValidFrame() const;
 
 	/** Return the newest frame of data we have. **/
-	int64 GetLatestValidFrame() const;
+	CORE_API int64 GetLatestValidFrame() const;
 
 protected:
 	/** Internal method to scan the messages to find the current game/render thread frame. */
-	void AdvanceFrameForLoad( TArray<FStatMessage>& CondensedMessages );
+	CORE_API void AdvanceFrameForLoad( TArray<FStatMessage>& CondensedMessages );
 
 	/** Largest frame seen. **/
 	int64 MaxFrameSeen;
@@ -972,23 +972,23 @@ struct FStatsReader
 };
 
 /** Struct used to read from uestats/uestatsraw files, initializes all metadata and starts a process of reading the file asynchronously. */
-struct CORE_API FStatsReadFile
+struct FStatsReadFile
 {
 	friend class FAsyncStatsFile;
 	friend struct FStatsReader<FStatsReadFile>;
 
 	/** Number of seconds between updating the current stage. */
-	static const double NumSecondsBetweenUpdates;
+	static CORE_API const double NumSecondsBetweenUpdates;
 
 public:
 	/** Reads and processes the file on the current thread. This is a blocking operation. */
-	void ReadAndProcessSynchronously();
+	CORE_API void ReadAndProcessSynchronously();
 
 	/** Reads and processes the file using the async tasks on the pool thread. The read data is sent to the game thread using the task graph. This is a non-blocking operation. */
-	void ReadAndProcessAsynchronously();
+	CORE_API void ReadAndProcessAsynchronously();
 
 	/** Sets the number of frame to be stored for future use. */
-	void SetHistoryFrames( int32 InHistoryFrames );
+	CORE_API void SetHistoryFrames( int32 InHistoryFrames );
 
 	/**
 	 * @return number of frames in the file.
@@ -1002,40 +1002,40 @@ public:
 
 protected:
 	/** Initialization constructor. */
-	FStatsReadFile( const TCHAR* InFilename, bool bInRawStatsFile );
+	CORE_API FStatsReadFile( const TCHAR* InFilename, bool bInRawStatsFile );
 
 public:
 	/** Destructor. */
-	virtual ~FStatsReadFile();
+	CORE_API virtual ~FStatsReadFile();
 
 protected:
 	/**
 	 * Prepares file to be loaded, makes sanity checks, reads and initializes metadata
 	 * @return true, if the process was completed successfully
 	 */
-	bool PrepareLoading();
+	CORE_API bool PrepareLoading();
 
 	/** Reads stats from the file. */
-	void ReadStats();
+	CORE_API void ReadStats();
 
 	/** Reads raw stats. */
-	void ReadRawStats();
+	CORE_API void ReadRawStats();
 
 	/** Reads regular stats. */
-	void ReadRegularStats();
+	CORE_API void ReadRegularStats();
 
 	/** Called before started processing. */
-	virtual void PreProcessStats();
+	CORE_API virtual void PreProcessStats();
 
 	/** Processes combined history using the internal functionality and provided overloaded Process*Operation methods. */
-	void ProcessStats();
+	CORE_API void ProcessStats();
 
 	/** Called every each frame has been read from the file. */
 	virtual void ReadStatsFrame( const TArray<FStatMessage>& CondensedMessages, const int64 Frame )
 	{}
 
 	/** Called after finished processing combined history. */
-	virtual void PostProcessStats();
+	CORE_API virtual void PostProcessStats();
 
 	/** Processes special message for advancing the stats frame from the game thread. */
 	virtual void ProcessAdvanceFrameEventGameThreadOperation( const FStatMessage& Message, const FStackState& StackState )
@@ -1176,15 +1176,15 @@ protected:
 	/**
 	 * Updates read stage progress periodically, does debug logging if enabled.
 	 */
-	void UpdateReadStageProgress();
+	CORE_API void UpdateReadStageProgress();
 
 	/** Dumps combined history stats. Only for raw stats. */
-	void UpdateCombinedHistoryStats();
+	CORE_API void UpdateCombinedHistoryStats();
 
 	/**
 	 * Updates process stage progress periodically, does debug logging if enabled.
 	 */
-	void UpdateProcessStageProgress( const int32 CurrentStatMessageIndex, const int32 FrameIndex, const int32 PacketIndex );
+	CORE_API void UpdateProcessStageProgress( const int32 CurrentStatMessageIndex, const int32 FrameIndex, const int32 PacketIndex );
 
 protected:
 	/** Current state of the stats. Mostly for metadata. */
diff --git a/Engine/Source/Runtime/Core/Public/Stats/StatsHierarchical.h b/Engine/Source/Runtime/Core/Public/Stats/StatsHierarchical.h
--- a/Engine/Source/Runtime/Core/Public/Stats/StatsHierarchical.h
+++ b/Engine/Source/Runtime/Core/Public/Stats/StatsHierarchical.h
@@ -23,18 +23,18 @@ class FStatsHierarchical;
  *	a description including its name, path (invocation path / call-stack) and its duration in 
  *  system cycles.
  */
-class CORE_API FStatsTreeElement
+class FStatsTreeElement
 {
 public:
 
 	// Default constructor
-	FStatsTreeElement();
+	CORE_API FStatsTreeElement();
 
 	// Returns the name of this element
-	FName GetFName() const;
+	CORE_API FName GetFName() const;
 
 	// Returns the name of this element as a FString
-	FString GetName() const;
+	CORE_API FString GetName() const;
 
 	/* 
 		Returns the invocation path of this element.
@@ -42,35 +42,35 @@ public:
 		profiling scopes.
 		For example: main.MyClass::Method.Algo::Sort
 	*/
-	FString GetPath() const;
+	CORE_API FString GetPath() const;
 
 	// Returns the number of invocations collected into this element
 	// @param bInclusive If true the invocations will contain the invocations used by child elements
-	uint32 Num(bool bInclusive = false) const;
+	CORE_API uint32 Num(bool bInclusive = false) const;
 
 	// Returns the total number of cycles recorded
 	// @param bInclusive If true the cycles will contain the cycles used by child elements
-	uint32 TotalCycles(bool bInclusive = true) const;
+	CORE_API uint32 TotalCycles(bool bInclusive = true) const;
 
 	// Returns the number of maximum cycles for this element (and children)
-	uint32 MaxCycles(bool bInclusive = true) const;
+	CORE_API uint32 MaxCycles(bool bInclusive = true) const;
 
 	// Returns the total number of seconds recorded
 	// @param bInclusive If true the time will contain the cycles used by child elements
-	double TotalSeconds(bool bInclusive = true) const;
+	CORE_API double TotalSeconds(bool bInclusive = true) const;
 
 	// Returns the average number of seconds recorded (total / num)
 	// @param bInclusive If true the time will contain the cycles used by child elements
-	double AverageSeconds(bool bInclusive = true) const;
+	CORE_API double AverageSeconds(bool bInclusive = true) const;
 
 	// Returns the contribution between 0.0 and 1.0 within the parent element.
 	// 1.0 means that 100% of the time of the parent element is spent in this child.
 	// @param bAgainstMaximum If true the ratio is expressed against the largest time in the tree
 	// @param bInclusive If true the time will contain the cycles used by child elements
-	double Contribution(bool bAgainstMaximum = false, bool bInclusive = true) const;
+	CORE_API double Contribution(bool bAgainstMaximum = false, bool bInclusive = true) const;
 
 	// Returns all child elements
-	const TArray<TSharedPtr<FStatsTreeElement>>& GetChildren() const;
+	CORE_API const TArray<TSharedPtr<FStatsTreeElement>>& GetChildren() const;
 
 protected:
 
@@ -90,9 +90,9 @@ protected:
 	TArray<TSharedPtr<FStatsTreeElement>> Children;
 
 	// returns the child or nullptr based on a given path
-	FStatsTreeElement* FindChild(const FString& InPath);
+	CORE_API FStatsTreeElement* FindChild(const FString& InPath);
 
-	void UpdatePostMeasurement(double InCyclesPerTimerToRemove = 0);
+	CORE_API void UpdatePostMeasurement(double InCyclesPerTimerToRemove = 0);
 
 	friend class FStatsHierarchical;
 	friend class FStatsHierarchicalClient;
@@ -133,7 +133,7 @@ protected:
  * Note: This system is not thread-safe. You want to call it from a single
  * thread only.
  */
-class CORE_API FStatsHierarchical
+class FStatsHierarchical
 {
 public:
 
@@ -158,26 +158,26 @@ public:
 	};
 
 	// Enabled measurements / profiling
-	static void BeginMeasurements();
+	static CORE_API void BeginMeasurements();
 
 	// Returns true if measurements are enabled
-	static bool IsEnabled();
+	static CORE_API bool IsEnabled();
 
 	/**
 	 * Ends measurements / profiling and returns the compacted profiling tree.
 	 * @param MeasurementsToMerge The baseline for the measurement.
 	 * @param bAddUntrackedElements If true adds an element for untracked time (time not profiled) for each node in the tree.
 	 */
-	static FStatsTreeElement EndMeasurements(FStatsTreeElement MeasurementsToMerge = FStatsTreeElement(), bool bAddUntrackedElements = true);
+	static CORE_API FStatsTreeElement EndMeasurements(FStatsTreeElement MeasurementsToMerge = FStatsTreeElement(), bool bAddUntrackedElements = true);
 
 	// Returns the last recorded profiling tree
-	static FStatsTreeElement GetLastMeasurements();
+	static CORE_API FStatsTreeElement GetLastMeasurements();
 
 	// Prints the results into a provided log
-	static void DumpMeasurements(FMessageLog& Log, bool bSortByDuration = true);
+	static CORE_API void DumpMeasurements(FMessageLog& Log, bool bSortByDuration = true);
 
 	// Returns the name to use for untracked time
-	static FName GetUntrackedTimeName();
+	static CORE_API FName GetUntrackedTimeName();
 
 private:
 
@@ -208,13 +208,13 @@ private:
 	 * Begins a single measurement given a label.
 	 * @param Label The label for this measurement (has to be != nullptr)
 	 */
-	static void BeginMeasurement(const ANSICHAR * Label);
+	static CORE_API void BeginMeasurement(const ANSICHAR * Label);
 
 	// Ends the last measurement
-	static void EndMeasurement();
+	static CORE_API void EndMeasurement();
 
-	static bool bEnabled;
-	static TArray<FHierarchicalStatEntry> Entries;
+	static CORE_API bool bEnabled;
+	static CORE_API TArray<FHierarchicalStatEntry> Entries;
 
 	friend struct FStatsHierarchical::FScope;
-};
\ No newline at end of file
+};
diff --git a/Engine/Source/Runtime/Core/Public/Stats/StatsMisc.h b/Engine/Source/Runtime/Core/Public/Stats/StatsMisc.h
--- a/Engine/Source/Runtime/Core/Public/Stats/StatsMisc.h
+++ b/Engine/Source/Runtime/Core/Public/Stats/StatsMisc.h
@@ -185,12 +185,12 @@ protected:
  */
 struct FScopeLogTime : public FConditionalScopeLogTime
 {
-	CORE_API FScopeLogTime(const WIDECHAR* InName, FTotalTimeAndCount* InCumulative = nullptr, EScopeLogTimeUnits InUnits = ScopeLog_Milliseconds)
+	FScopeLogTime(const WIDECHAR* InName, FTotalTimeAndCount* InCumulative = nullptr, EScopeLogTimeUnits InUnits = ScopeLog_Milliseconds)
 		: FConditionalScopeLogTime(true, InName, InCumulative, InUnits)
 	{
 	}
 
-	CORE_API FScopeLogTime(const ANSICHAR* InName, FTotalTimeAndCount* InCumulative = nullptr, EScopeLogTimeUnits InUnits = ScopeLog_Milliseconds)
+	FScopeLogTime(const ANSICHAR* InName, FTotalTimeAndCount* InCumulative = nullptr, EScopeLogTimeUnits InUnits = ScopeLog_Milliseconds)
 		: FConditionalScopeLogTime(true, InName, InCumulative, InUnits)
 	{
 	}
diff --git a/Engine/Source/Runtime/Core/Public/Templates/RefCounting.h b/Engine/Source/Runtime/Core/Public/Templates/RefCounting.h
--- a/Engine/Source/Runtime/Core/Public/Templates/RefCounting.h
+++ b/Engine/Source/Runtime/Core/Public/Templates/RefCounting.h
@@ -75,7 +75,7 @@ private:
  * the reference count.
  *
  */
-class CORE_API FRefCountedObject
+class FRefCountedObject
 {
 public:
 	FRefCountedObject(): NumRefs(0) {}
@@ -106,7 +106,7 @@ private:
 /**
  * Like FRefCountedObject, but internal ref count is thread safe
  */
-class CORE_API FThreadSafeRefCountedObject
+class FThreadSafeRefCountedObject
 {
 public:
 	FThreadSafeRefCountedObject() : NumRefs(0) {}
diff --git a/Engine/Source/Runtime/Core/Public/Unix/UnixErrorOutputDevice.h b/Engine/Source/Runtime/Core/Public/Unix/UnixErrorOutputDevice.h
--- a/Engine/Source/Runtime/Core/Public/Unix/UnixErrorOutputDevice.h
+++ b/Engine/Source/Runtime/Core/Public/Unix/UnixErrorOutputDevice.h
@@ -7,12 +7,12 @@
 #include "Misc/OutputDeviceError.h"
 #include "UObject/NameTypes.h"
 
-class CORE_API FUnixErrorOutputDevice : public FOutputDeviceError
+class FUnixErrorOutputDevice : public FOutputDeviceError
 {
 public:
 
 	/** Constructor, initializing member variables */
-	FUnixErrorOutputDevice();
+	CORE_API FUnixErrorOutputDevice();
 
 	/**
 	 * Serializes the passed in data unless the current event is suppressed.
@@ -20,13 +20,13 @@ public:
 	 * @param	Data	Text to log
 	 * @param	Event	Event name used for suppression purposes
 	 */
-	virtual void Serialize(const TCHAR* Msg, ELogVerbosity::Type Verbosity, const class FName& Category) override;
+	CORE_API virtual void Serialize(const TCHAR* Msg, ELogVerbosity::Type Verbosity, const class FName& Category) override;
 
 	/**
 	 * Error handling function that is being called from within the system wide global
 	 * error handler, e.g. using structured exception handling on the PC.
 	 */
-	void HandleError() override;
+	CORE_API void HandleError() override;
 
 private:
 
@@ -36,5 +36,5 @@ protected:
 	/**
 	 * Callback to allow FUnixApplicationErrorOutputDevice to restore the UI.
 	 */
-	virtual void HandleErrorRestoreUI();
+	CORE_API virtual void HandleErrorRestoreUI();
 };
diff --git a/Engine/Source/Runtime/Core/Public/Unix/UnixPlatformCrashContext.h b/Engine/Source/Runtime/Core/Public/Unix/UnixPlatformCrashContext.h
--- a/Engine/Source/Runtime/Core/Public/Unix/UnixPlatformCrashContext.h
+++ b/Engine/Source/Runtime/Core/Public/Unix/UnixPlatformCrashContext.h
@@ -28,7 +28,7 @@ struct ThreadStackUserData
 	TAtomic<bool> bDone;
 };
 
-struct CORE_API FUnixCrashContext : public FGenericCrashContext
+struct FUnixCrashContext : public FGenericCrashContext
 {
 	/** Signal number */
 	int32 Signal;
@@ -52,7 +52,7 @@ struct CORE_API FUnixCrashContext : public FGenericCrashContext
 	char MinidumpCallstackInfo[16384];
 
 	/** Fake siginfo used when handling ensure(), etc */
-	static __thread siginfo_t	FakeSiginfoForDiagnostics;
+	static CORE_API __thread siginfo_t	FakeSiginfoForDiagnostics;
 
 	/** The PC of the first function used when handling a crash. Used to figure out the number of frames to ignore */
 	uint64* FirstCrashHandlerFrame = nullptr;
@@ -73,7 +73,7 @@ struct CORE_API FUnixCrashContext : public FGenericCrashContext
 		MinidumpCallstackInfo[ 0 ] = 0;
 	}
 
-	~FUnixCrashContext();
+	CORE_API ~FUnixCrashContext();
 
 	/**
 	 * Inits the crash context from data provided by a signal handler.
@@ -82,34 +82,34 @@ struct CORE_API FUnixCrashContext : public FGenericCrashContext
 	 * @param InInfo additional info (e.g. address we tried to read, etc)
 	 * @param InContext thread context
 	 */
-	void InitFromSignal(int32 InSignal, siginfo_t* InInfo, void* InContext);
+	CORE_API void InitFromSignal(int32 InSignal, siginfo_t* InInfo, void* InContext);
 
 	/**
 	 * Inits the crash context from some diagnostic handler (ensure, stall, etc...)
 	 *
 	 * @param InAddress address where the event happened, for optional callstack trimming
 	 */
-	void InitFromDiagnostics(const void* InAddress = nullptr);
+	CORE_API void InitFromDiagnostics(const void* InAddress = nullptr);
 
 	/**
 	 * Populates crash context stack trace and a few related fields for the calling thread
 	 *
 	 */
-	void CaptureStackTrace(void* ErrorProgramCounter);
+	CORE_API void CaptureStackTrace(void* ErrorProgramCounter);
 
 	/**
 	 * Populates crash context stack trace and a few related fields for another thread
 	 *
 	 * @param ThreadId The thread to capture the stack trace from
 	 */
-	void CaptureThreadStackTrace(uint32_t ThreadId);
+	CORE_API void CaptureThreadStackTrace(uint32_t ThreadId);
 
 	/**
 	 * Generates a new crash report containing information needed for the crash reporter and launches it; may not return.
 	 *
 	 * @return If the crash type is not continuable, the function will not return
 	 */
-	void GenerateCrashInfoAndLaunchReporter() const;
+	CORE_API void GenerateCrashInfoAndLaunchReporter() const;
 
 	/**
 	 * Sets whether this crash represents a non-crash event like an ensure
@@ -119,11 +119,11 @@ struct CORE_API FUnixCrashContext : public FGenericCrashContext
 	/**
 	 * Sets the FirstCrashHandlerFrame only if it has not been set before
 	 */
-	void SetFirstCrashHandlerFrame(uint64* ProgramCounter);
+	CORE_API void SetFirstCrashHandlerFrame(uint64* ProgramCounter);
 
-	virtual void GetPortableCallStack(const uint64* StackFrames, int32 NumStackFrames, TArray<FCrashStackFrame>& OutCallStack) const override;
+	CORE_API virtual void GetPortableCallStack(const uint64* StackFrames, int32 NumStackFrames, TArray<FCrashStackFrame>& OutCallStack) const override;
 
-	void AddPlatformSpecificProperties() const;
+	CORE_API void AddPlatformSpecificProperties() const;
 
 protected:
 	/**
@@ -131,7 +131,7 @@ protected:
 	 *
 	 * @param DiagnosticsPath Path to put the file to
 	 */
-	void GenerateReport(const FString & DiagnosticsPath) const;
+	CORE_API void GenerateReport(const FString & DiagnosticsPath) const;
 };
 
 typedef FUnixCrashContext FPlatformCrashContext;
diff --git a/Engine/Source/Runtime/Core/Public/Unix/UnixPlatformFile.h b/Engine/Source/Runtime/Core/Public/Unix/UnixPlatformFile.h
--- a/Engine/Source/Runtime/Core/Public/Unix/UnixPlatformFile.h
+++ b/Engine/Source/Runtime/Core/Public/Unix/UnixPlatformFile.h
@@ -17,52 +17,52 @@ template <typename FuncType> class TFunctionRef;
 /**
  * Unix File I/O implementation
 **/
-class CORE_API FUnixPlatformFile : public IPhysicalPlatformFile
+class FUnixPlatformFile : public IPhysicalPlatformFile
 {
 protected:
-	virtual FString NormalizeFilename(const TCHAR* Filename, bool bIsForWriting);
-	virtual FString NormalizeDirectory(const TCHAR* Directory, bool bIsForWriting);
+	CORE_API virtual FString NormalizeFilename(const TCHAR* Filename, bool bIsForWriting);
+	CORE_API virtual FString NormalizeDirectory(const TCHAR* Directory, bool bIsForWriting);
 public:
 	//~ For visibility of overloads we don't override
 	using IPhysicalPlatformFile::IterateDirectory;
 	using IPhysicalPlatformFile::IterateDirectoryStat;
 
-	virtual bool FileExists(const TCHAR* Filename) override;
-	virtual int64 FileSize(const TCHAR* Filename) override;
-	virtual bool DeleteFile(const TCHAR* Filename) override;
-	virtual bool IsReadOnly(const TCHAR* Filename) override;
-	virtual bool MoveFile(const TCHAR* To, const TCHAR* From) override;
-	virtual bool SetReadOnly(const TCHAR* Filename, bool bNewReadOnlyValue) override;
+	CORE_API virtual bool FileExists(const TCHAR* Filename) override;
+	CORE_API virtual int64 FileSize(const TCHAR* Filename) override;
+	CORE_API virtual bool DeleteFile(const TCHAR* Filename) override;
+	CORE_API virtual bool IsReadOnly(const TCHAR* Filename) override;
+	CORE_API virtual bool MoveFile(const TCHAR* To, const TCHAR* From) override;
+	CORE_API virtual bool SetReadOnly(const TCHAR* Filename, bool bNewReadOnlyValue) override;
 
 
-	virtual FDateTime GetTimeStamp(const TCHAR* Filename) override;
+	CORE_API virtual FDateTime GetTimeStamp(const TCHAR* Filename) override;
 
-	virtual void SetTimeStamp(const TCHAR* Filename, const FDateTime DateTime) override;
+	CORE_API virtual void SetTimeStamp(const TCHAR* Filename, const FDateTime DateTime) override;
 
-	virtual FDateTime GetAccessTimeStamp(const TCHAR* Filename) override;
-	virtual FString GetFilenameOnDisk(const TCHAR* Filename) override;
+	CORE_API virtual FDateTime GetAccessTimeStamp(const TCHAR* Filename) override;
+	CORE_API virtual FString GetFilenameOnDisk(const TCHAR* Filename) override;
 
-	virtual ESymlinkResult IsSymlink(const TCHAR* Filename) override;
+	CORE_API virtual ESymlinkResult IsSymlink(const TCHAR* Filename) override;
 
-	virtual IFileHandle* OpenRead(const TCHAR* Filename, bool bAllowWrite = false) override;
-	virtual IFileHandle* OpenWrite(const TCHAR* Filename, bool bAppend = false, bool bAllowRead = false) override;
-	virtual IMappedFileHandle* OpenMapped(const TCHAR* Filename) override;
+	CORE_API virtual IFileHandle* OpenRead(const TCHAR* Filename, bool bAllowWrite = false) override;
+	CORE_API virtual IFileHandle* OpenWrite(const TCHAR* Filename, bool bAppend = false, bool bAllowRead = false) override;
+	CORE_API virtual IMappedFileHandle* OpenMapped(const TCHAR* Filename) override;
 
-	virtual bool DirectoryExists(const TCHAR* Directory) override;
-	virtual bool CreateDirectory(const TCHAR* Directory) override;
-	virtual bool DeleteDirectory(const TCHAR* Directory) override;
+	CORE_API virtual bool DirectoryExists(const TCHAR* Directory) override;
+	CORE_API virtual bool CreateDirectory(const TCHAR* Directory) override;
+	CORE_API virtual bool DeleteDirectory(const TCHAR* Directory) override;
 
-	virtual FFileStatData GetStatData(const TCHAR* FilenameOrDirectory) override;
+	CORE_API virtual FFileStatData GetStatData(const TCHAR* FilenameOrDirectory) override;
 
-	bool CreateDirectoriesFromPath(const TCHAR* Path);
+	CORE_API bool CreateDirectoriesFromPath(const TCHAR* Path);
 
-	virtual bool IterateDirectory(const TCHAR* Directory, FDirectoryVisitor& Visitor) override;
-	virtual bool IterateDirectoryStat(const TCHAR* Directory, FDirectoryStatVisitor& Visitor) override;
+	CORE_API virtual bool IterateDirectory(const TCHAR* Directory, FDirectoryVisitor& Visitor) override;
+	CORE_API virtual bool IterateDirectoryStat(const TCHAR* Directory, FDirectoryStatVisitor& Visitor) override;
 
-	virtual bool CopyFile(const TCHAR* To, const TCHAR* From, EPlatformFileRead ReadFlags = EPlatformFileRead::None, EPlatformFileWrite WriteFlags = EPlatformFileWrite::None) override;
+	CORE_API virtual bool CopyFile(const TCHAR* To, const TCHAR* From, EPlatformFileRead ReadFlags = EPlatformFileRead::None, EPlatformFileWrite WriteFlags = EPlatformFileWrite::None) override;
 
 protected:
-	bool IterateDirectoryCommon(const TCHAR* Directory, const TFunctionRef<bool(struct dirent*)>& Visitor);
+	CORE_API bool IterateDirectoryCommon(const TCHAR* Directory, const TFunctionRef<bool(struct dirent*)>& Visitor);
 
 	/** We're logging an error message. */
 	bool bLoggingError = false;
diff --git a/Engine/Source/Runtime/Core/Public/Unix/UnixPlatformMemory.h b/Engine/Source/Runtime/Core/Public/Unix/UnixPlatformMemory.h
--- a/Engine/Source/Runtime/Core/Public/Unix/UnixPlatformMemory.h
+++ b/Engine/Source/Runtime/Core/Public/Unix/UnixPlatformMemory.h
@@ -45,7 +45,7 @@ struct FExtendedPlatformMemoryStats
 /**
 * Unix implementation of the memory OS functions
 **/
-struct CORE_API FUnixPlatformMemory : public FGenericPlatformMemory
+struct FUnixPlatformMemory : public FGenericPlatformMemory
 {
 	/**
 	 * Unix representation of a shared memory region
@@ -74,14 +74,14 @@ struct CORE_API FUnixPlatformMemory : public FGenericPlatformMemory
 	};
 
 	//~ Begin FGenericPlatformMemory Interface
-	static void Init();
-	static class FMalloc* BaseAllocator();
-	static FPlatformMemoryStats GetStats();
-	static FExtendedPlatformMemoryStats GetExtendedStats();
-	static const FPlatformMemoryConstants& GetConstants();
-	static bool PageProtect(void* const Ptr, const SIZE_T Size, const bool bCanRead, const bool bCanWrite);
-	static void* BinnedAllocFromOS(SIZE_T Size);
-	static void BinnedFreeToOS(void* Ptr, SIZE_T Size);
+	static CORE_API void Init();
+	static CORE_API class FMalloc* BaseAllocator();
+	static CORE_API FPlatformMemoryStats GetStats();
+	static CORE_API FExtendedPlatformMemoryStats GetExtendedStats();
+	static CORE_API const FPlatformMemoryConstants& GetConstants();
+	static CORE_API bool PageProtect(void* const Ptr, const SIZE_T Size, const bool bCanRead, const bool bCanWrite);
+	static CORE_API void* BinnedAllocFromOS(SIZE_T Size);
+	static CORE_API void BinnedFreeToOS(void* Ptr, SIZE_T Size);
 
 	class FPlatformVirtualMemoryBlock : public FBasicVirtualMemoryBlock
 	{
@@ -132,13 +132,13 @@ struct CORE_API FUnixPlatformMemory : public FGenericPlatformMemory
 		static size_t GetVirtualSizeAlignment();
 	};
 
-	static FSharedMemoryRegion * MapNamedSharedMemoryRegion(const FString& InName, bool bCreate, uint32 AccessMode, SIZE_T Size);
-	static bool UnmapNamedSharedMemoryRegion(FSharedMemoryRegion * MemoryRegion);
-	static bool GetLLMAllocFunctions(void*(*&OutAllocFunction)(size_t), void(*&OutFreeFunction)(void*, size_t), int32& OutAlignment);
-	[[noreturn]] static void OnOutOfMemory(uint64 Size, uint32 Alignment);
+	static CORE_API FSharedMemoryRegion * MapNamedSharedMemoryRegion(const FString& InName, bool bCreate, uint32 AccessMode, SIZE_T Size);
+	static CORE_API bool UnmapNamedSharedMemoryRegion(FSharedMemoryRegion * MemoryRegion);
+	static CORE_API bool GetLLMAllocFunctions(void*(*&OutAllocFunction)(size_t), void(*&OutFreeFunction)(void*, size_t), int32& OutAlignment);
+	[[noreturn]] static CORE_API void OnOutOfMemory(uint64 Size, uint32 Alignment);
 	//~ End FGenericPlatformMemory Interface
 
-	static bool HasForkPageProtectorEnabled();
+	static CORE_API bool HasForkPageProtectorEnabled();
 };
 
 typedef FUnixPlatformMemory FPlatformMemory;
