[Core]

diff --git a/Engine/Source/Runtime/Core/Public/Unix/UnixPlatformMisc.h b/Engine/Source/Runtime/Core/Public/Unix/UnixPlatformMisc.h
--- a/Engine/Source/Runtime/Core/Public/Unix/UnixPlatformMisc.h
+++ b/Engine/Source/Runtime/Core/Public/Unix/UnixPlatformMisc.h
@@ -19,66 +19,66 @@ struct FGenericCrashContext;
 /**
  * Unix implementation of the misc OS functions
  */
-struct CORE_API FUnixPlatformMisc : public FGenericPlatformMisc
+struct FUnixPlatformMisc : public FGenericPlatformMisc
 {
-	static void PlatformPreInit();
-	static void PlatformInit();
-	static void PlatformTearDown();
-	static void SetGracefulTerminationHandler();
-	static void SetCrashHandler(void (* CrashHandler)(const FGenericCrashContext& Context));
-	static int32 GetMaxPathLength();
+	static CORE_API void PlatformPreInit();
+	static CORE_API void PlatformInit();
+	static CORE_API void PlatformTearDown();
+	static CORE_API void SetGracefulTerminationHandler();
+	static CORE_API void SetCrashHandler(void (* CrashHandler)(const FGenericCrashContext& Context));
+	static CORE_API int32 GetMaxPathLength();
 
 	UE_DEPRECATED(4.21, "void FPlatformMisc::GetEnvironmentVariable(Name, Result, Length) is deprecated. Use FString FPlatformMisc::GetEnvironmentVariable(Name) instead.")
-	static void GetEnvironmentVariable(const TCHAR* VariableName, TCHAR* Result, int32 ResultLength);
+	static CORE_API void GetEnvironmentVariable(const TCHAR* VariableName, TCHAR* Result, int32 ResultLength);
 
-	static FString GetEnvironmentVariable(const TCHAR* VariableName);
-	static void SetEnvironmentVar(const TCHAR* VariableName, const TCHAR* Value);
-	static TArray<uint8> GetMacAddress();
-	static bool IsRunningOnBattery();
+	static CORE_API FString GetEnvironmentVariable(const TCHAR* VariableName);
+	static CORE_API void SetEnvironmentVar(const TCHAR* VariableName, const TCHAR* Value);
+	static CORE_API TArray<uint8> GetMacAddress();
+	static CORE_API bool IsRunningOnBattery();
 
 #if !UE_BUILD_SHIPPING
-	static bool IsDebuggerPresent();
+	static CORE_API bool IsDebuggerPresent();
 #endif // !UE_BUILD_SHIPPING
 
-	static void LowLevelOutputDebugString(const TCHAR *Message);
+	static CORE_API void LowLevelOutputDebugString(const TCHAR *Message);
 
-	static void RequestExit(bool Force, const TCHAR* CallSite = nullptr);
-	static void RequestExitWithStatus(bool Force, uint8 ReturnCode, const TCHAR* CallSite = nullptr);
-	static const TCHAR* GetSystemErrorMessage(TCHAR* OutBuffer, int32 BufferCount, int32 Error);
+	static CORE_API void RequestExit(bool Force, const TCHAR* CallSite = nullptr);
+	static CORE_API void RequestExitWithStatus(bool Force, uint8 ReturnCode, const TCHAR* CallSite = nullptr);
+	static CORE_API const TCHAR* GetSystemErrorMessage(TCHAR* OutBuffer, int32 BufferCount, int32 Error);
 
-	static void NormalizePath(FString& InPath);
-	static void NormalizePath(FStringBuilderBase& InPath);
+	static CORE_API void NormalizePath(FString& InPath);
+	static CORE_API void NormalizePath(FStringBuilderBase& InPath);
 
 	static const TCHAR* GetPathVarDelimiter()
 	{
 		return TEXT(":");
 	}
 
-	static EAppReturnType::Type MessageBoxExt(EAppMsgType::Type MsgType, const TCHAR* Text, const TCHAR* Caption);
+	static CORE_API EAppReturnType::Type MessageBoxExt(EAppMsgType::Type MsgType, const TCHAR* Text, const TCHAR* Caption);
 
 	FORCEINLINE static void MemoryBarrier()
 	{
 		__sync_synchronize();
 	}
 
-	static int32 NumberOfCores();
-	static int32 NumberOfCoresIncludingHyperthreads();
-	static FString GetOperatingSystemId();
-	static bool GetDiskTotalAndFreeSpace(const FString& InPath, uint64& TotalNumberOfBytes, uint64& NumberOfFreeBytes);
-	static bool GetPageFaultStats(FPageFaultStats& OutStats, EPageFaultFlags Flags=EPageFaultFlags::All);
-	static bool GetBlockingIOStats(FProcessIOStats& OutStats, EInputOutputFlags Flags=EInputOutputFlags::All);
-	static bool GetContextSwitchStats(FContextSwitchStats& OutStats, EContextSwitchFlags Flags=EContextSwitchFlags::All);
+	static CORE_API int32 NumberOfCores();
+	static CORE_API int32 NumberOfCoresIncludingHyperthreads();
+	static CORE_API FString GetOperatingSystemId();
+	static CORE_API bool GetDiskTotalAndFreeSpace(const FString& InPath, uint64& TotalNumberOfBytes, uint64& NumberOfFreeBytes);
+	static CORE_API bool GetPageFaultStats(FPageFaultStats& OutStats, EPageFaultFlags Flags=EPageFaultFlags::All);
+	static CORE_API bool GetBlockingIOStats(FProcessIOStats& OutStats, EInputOutputFlags Flags=EInputOutputFlags::All);
+	static CORE_API bool GetContextSwitchStats(FContextSwitchStats& OutStats, EContextSwitchFlags Flags=EContextSwitchFlags::All);
 
 	/**
 	 * Determines the shader format for the platform
 	 *
 	 * @return	Returns the shader format to be used by that platform
 	 */
-	static const TCHAR* GetNullRHIShaderFormat();
+	static CORE_API const TCHAR* GetNullRHIShaderFormat();
 
-	static bool HasCPUIDInstruction();
-	static FString GetCPUVendor();
-	static FString GetCPUBrand();
+	static CORE_API bool HasCPUIDInstruction();
+	static CORE_API FString GetCPUVendor();
+	static CORE_API FString GetCPUBrand();
 
 	/**
 	 * Uses cpuid instruction to get the vendor string
@@ -94,22 +94,22 @@ struct CORE_API FUnixPlatformMisc : public FGenericPlatformMisc
 	 *			Bits 20-27	Extended family
 	 *			Bits 28-31	Reserved
 	 */
-	static uint32 GetCPUInfo();
+	static CORE_API uint32 GetCPUInfo();
 
-	static FString GetPrimaryGPUBrand();
+	static CORE_API FString GetPrimaryGPUBrand();
 
-	static bool HasNonoptionalCPUFeatures();
-	static bool NeedsNonoptionalCPUFeaturesCheck();
+	static CORE_API bool HasNonoptionalCPUFeatures();
+	static CORE_API bool NeedsNonoptionalCPUFeaturesCheck();
 
 	/**
 	 * Ungrabs input (useful before breaking into debugging)
 	 */
-	static void UngrabAllInput();
+	static CORE_API void UngrabAllInput();
 
 	/**
 	 * Returns whether the program has been started remotely (e.g. over SSH)
 	 */
-	static bool HasBeenStartedRemotely();
+	static CORE_API bool HasBeenStartedRemotely();
 
 	/**
 	 * Determines if return code has been overriden and returns it.
@@ -118,30 +118,30 @@ struct CORE_API FUnixPlatformMisc : public FGenericPlatformMisc
 	 *
 	 * @return true if the error code has been overriden, false if not
 	 */
-	static bool HasOverriddenReturnCode(uint8 * OverriddenReturnCodeToUsePtr);
-	static void GetOSVersions(FString& out_OSVersionLabel, FString& out_OSSubVersionLabel);
-	static FString GetOSVersion();
-	static FString GetLoginId();
+	static CORE_API bool HasOverriddenReturnCode(uint8 * OverriddenReturnCodeToUsePtr);
+	static CORE_API void GetOSVersions(FString& out_OSVersionLabel, FString& out_OSSubVersionLabel);
+	static CORE_API FString GetOSVersion();
+	static CORE_API FString GetLoginId();
 
-	static void CreateGuid(FGuid& Result);
+	static CORE_API void CreateGuid(FGuid& Result);
 
-	static IPlatformChunkInstall* GetPlatformChunkInstall();
+	static CORE_API IPlatformChunkInstall* GetPlatformChunkInstall();
 
-	static bool SetStoredValues(const FString& InStoreId, const FString& InSectionName, const TMap<FString, FString>& InKeyValues);
+	static CORE_API bool SetStoredValues(const FString& InStoreId, const FString& InSectionName, const TMap<FString, FString>& InKeyValues);
 
-	static int32 NumberOfWorkerThreadsToSpawn();
+	static CORE_API int32 NumberOfWorkerThreadsToSpawn();
 
 #if STATS || ENABLE_STATNAMEDEVENTS
-	static void BeginNamedEventFrame();
-	static void BeginNamedEvent(const struct FColor& Color, const TCHAR* Text);
-	static void BeginNamedEvent(const struct FColor& Color, const ANSICHAR* Text);
-	static void EndNamedEvent();
-	static void CustomNamedStat(const TCHAR* Text, float Value, const TCHAR* Graph, const TCHAR* Unit);
-	static void CustomNamedStat(const ANSICHAR* Text, float Value, const ANSICHAR* Graph, const ANSICHAR* Unit);
+	static CORE_API void BeginNamedEventFrame();
+	static CORE_API void BeginNamedEvent(const struct FColor& Color, const TCHAR* Text);
+	static CORE_API void BeginNamedEvent(const struct FColor& Color, const ANSICHAR* Text);
+	static CORE_API void EndNamedEvent();
+	static CORE_API void CustomNamedStat(const TCHAR* Text, float Value, const TCHAR* Graph, const TCHAR* Unit);
+	static CORE_API void CustomNamedStat(const ANSICHAR* Text, float Value, const ANSICHAR* Graph, const ANSICHAR* Unit);
 #endif
 
 	/* Explicitly call this function to setup syscall filters based on a file passed in from the command line
 	 * -allowsyscallfilterfile=PATH_TO_FILE
 	 */
-	static bool SetupSyscallFilters();
+	static CORE_API bool SetupSyscallFilters();
 };
diff --git a/Engine/Source/Runtime/Core/Public/Unix/UnixPlatformOutputDevices.h b/Engine/Source/Runtime/Core/Public/Unix/UnixPlatformOutputDevices.h
--- a/Engine/Source/Runtime/Core/Public/Unix/UnixPlatformOutputDevices.h
+++ b/Engine/Source/Runtime/Core/Public/Unix/UnixPlatformOutputDevices.h
@@ -13,12 +13,12 @@
 class FOutputDevice;
 class FOutputDeviceError;
 
-struct CORE_API FUnixOutputDevices : public FGenericPlatformOutputDevices
+struct FUnixOutputDevices : public FGenericPlatformOutputDevices
 {
-	static void SetupOutputDevices();
-	static FString GetAbsoluteLogFilename();
-	static FOutputDevice* GetEventLog();
-	static FOutputDeviceError* GetError();
+	static CORE_API void SetupOutputDevices();
+	static CORE_API FString GetAbsoluteLogFilename();
+	static CORE_API FOutputDevice* GetEventLog();
+	static CORE_API FOutputDeviceError* GetError();
 };
 
 typedef FUnixOutputDevices FPlatformOutputDevices;
diff --git a/Engine/Source/Runtime/Core/Public/Unix/UnixPlatformProcess.h b/Engine/Source/Runtime/Core/Public/Unix/UnixPlatformProcess.h
--- a/Engine/Source/Runtime/Core/Public/Unix/UnixPlatformProcess.h
+++ b/Engine/Source/Runtime/Core/Public/Unix/UnixPlatformProcess.h
@@ -205,33 +205,33 @@ protected:
 /**
  * Unix implementation of the Process OS functions
  */
-struct CORE_API FUnixPlatformProcess : public FGenericPlatformProcess
+struct FUnixPlatformProcess : public FGenericPlatformProcess
 {
 	struct FProcEnumInfo;
 
 	/**
 	 * Process enumerator.
 	 */
-	class CORE_API FProcEnumerator
+	class FProcEnumerator
 	{
 	public:
 		// Constructor
-		FProcEnumerator();
+		CORE_API FProcEnumerator();
 		FProcEnumerator(const FProcEnumerator&) = delete;
 		FProcEnumerator& operator=(const FProcEnumerator&) = delete;
 
 		// Destructor
-		~FProcEnumerator();
+		CORE_API ~FProcEnumerator();
 
 		// Gets current process enumerator info.
-		FProcEnumInfo GetCurrent() const;
+		CORE_API FProcEnumInfo GetCurrent() const;
 		
 		/**
 		 * Moves current to the next process.
 		 *
 		 * @returns True if succeeded. False otherwise.
 		 */
-		bool MoveNext();
+		CORE_API bool MoveNext();
 	private:
 		// Private implementation data.
 		struct FProcEnumData* Data;
@@ -240,21 +240,21 @@ struct CORE_API FUnixPlatformProcess : public FGenericPlatformProcess
 	/**
 	 * Process enumeration info structure.
 	 */
-	struct CORE_API FProcEnumInfo
+	struct FProcEnumInfo
 	{
 		friend FUnixPlatformProcess::FProcEnumerator::FProcEnumerator();
 
 		// Gets process PID.
-		uint32 GetPID() const;
+		CORE_API uint32 GetPID() const;
 
 		// Gets parent process PID.
-		uint32 GetParentPID() const;
+		CORE_API uint32 GetParentPID() const;
 
 		// Gets process name. I.e. exec name.
-		FString GetName() const;
+		CORE_API FString GetName() const;
 
 		// Gets process full image path. I.e. full path of the exec file.
-		FString GetFullPath() const;
+		CORE_API FString GetFullPath() const;
 
 	private:
 		// Private constructor.
@@ -264,67 +264,67 @@ struct CORE_API FUnixPlatformProcess : public FGenericPlatformProcess
 		uint32 PID;
 	};
 
-	static void* GetDllHandle( const TCHAR* Filename );
-	static void FreeDllHandle( void* DllHandle );
-	static void* GetDllExport( void* DllHandle, const TCHAR* ProcName );
-	static const TCHAR* ComputerName();
-	static const TCHAR* UserName(bool bOnlyAlphaNumeric = true);
-	static const TCHAR* UserTempDir();
-	static const TCHAR* UserDir();
-	static const TCHAR* UserSettingsDir();
-	static const TCHAR* ApplicationSettingsDir();
-	static FString GetApplicationSettingsDir(const ApplicationSettingsContext& Settings);
-	static void SetCurrentWorkingDirectoryToBaseDir();
-	static FString GetCurrentWorkingDirectory();
-	static FString GenerateApplicationPath(const FString& AppName, EBuildConfiguration BuildConfiguration);
-	static FString GetApplicationName( uint32 ProcessId );
-	static bool SetProcessLimits(EProcessResource::Type Resource, uint64 Limit);
-	static const TCHAR* ExecutablePath();
-	static const TCHAR* ExecutableName(bool bRemoveExtension = true);
-	static const TCHAR* GetModulePrefix();
-	static const TCHAR* GetModuleExtension();
-	static void ClosePipe( void* ReadPipe, void* WritePipe );
-	static bool CreatePipe(void*& ReadPipe, void*& WritePipe, bool bWritePipeLocal = false);
-	static FString ReadPipe( void* ReadPipe );
-	static bool ReadPipeToArray(void* ReadPipe, TArray<uint8> & Output);
-	static bool WritePipe(void* WritePipe, const FString& Message, FString* OutWritten = nullptr);
-	static bool WritePipe(void* WritePipe, const uint8* Data, const int32 DataLength, int32* OutDataLength = nullptr);
-	static class FRunnableThread* CreateRunnableThread();
-	static const FString GetModulesDirectory();
-	static bool CanLaunchURL(const TCHAR* URL);
-	static void LaunchURL(const TCHAR* URL, const TCHAR* Parms, FString* Error);
-	static FProcHandle CreateProc(const TCHAR* URL, const TCHAR* Parms, bool bLaunchDetached, bool bLaunchHidden, bool bLaunchReallyHidden, uint32* OutProcessID, int32 PriorityModifier, const TCHAR* OptionalWorkingDirectory, void* PipeWriteChild, void* PipeReadChild = nullptr);
-	static FProcHandle CreateProc(const TCHAR* URL, const TCHAR* Parms, bool bLaunchDetached, bool bLaunchHidden, bool bLaunchReallyHidden, uint32* OutProcessID, int32 PriorityModifier, const TCHAR* OptionalWorkingDirectory, void* PipeWriteChild, void* PipeReadChild , void* PipeStdErrChild);
-	static FProcHandle OpenProcess(uint32 ProcessID);
-	static bool IsProcRunning( FProcHandle & ProcessHandle );
-	static void WaitForProc( FProcHandle & ProcessHandle );
-	static void CloseProc( FProcHandle & ProcessHandle );
-	static void TerminateProc( FProcHandle & ProcessHandle, bool KillTree = false );
-	static EWaitAndForkResult WaitAndFork();
-	static uint32 GetCurrentProcessId();
-	static uint32 GetCurrentCoreNumber();
-	static bool GetProcReturnCode( FProcHandle & ProcHandle, int32* ReturnCode );
-	static bool Daemonize();
-	static bool IsApplicationRunning( uint32 ProcessId );
-	static bool IsApplicationRunning( const TCHAR* ProcName );
-	static bool ExecProcess(const TCHAR* URL, const TCHAR* Params, int32* OutReturnCode, FString* OutStdOut, FString* OutStdErr, const TCHAR* OptionalWorkingDirectory = NULL, bool bShouldEndWithParentProcess = false);
-	static void ExploreFolder( const TCHAR* FilePath );
-	static bool LaunchFileInDefaultExternalApplication( const TCHAR* FileName, const TCHAR* Parms = NULL, ELaunchVerb::Type Verb = ELaunchVerb::Open, bool bPromptToOpenOnFailure = true);
-	static bool IsFirstInstance();
-	static void OnChildEndFramePostFork();
-	static int32 TranslateThreadPriority(EThreadPriority Priority);
-	static void SetThreadNiceValue( uint32_t ThreadId, int32 NiceValue );
-	static void SetThreadPriority( EThreadPriority NewPriority );
+	static CORE_API void* GetDllHandle( const TCHAR* Filename );
+	static CORE_API void FreeDllHandle( void* DllHandle );
+	static CORE_API void* GetDllExport( void* DllHandle, const TCHAR* ProcName );
+	static CORE_API const TCHAR* ComputerName();
+	static CORE_API const TCHAR* UserName(bool bOnlyAlphaNumeric = true);
+	static CORE_API const TCHAR* UserTempDir();
+	static CORE_API const TCHAR* UserDir();
+	static CORE_API const TCHAR* UserSettingsDir();
+	static CORE_API const TCHAR* ApplicationSettingsDir();
+	static CORE_API FString GetApplicationSettingsDir(const ApplicationSettingsContext& Settings);
+	static CORE_API void SetCurrentWorkingDirectoryToBaseDir();
+	static CORE_API FString GetCurrentWorkingDirectory();
+	static CORE_API FString GenerateApplicationPath(const FString& AppName, EBuildConfiguration BuildConfiguration);
+	static CORE_API FString GetApplicationName( uint32 ProcessId );
+	static CORE_API bool SetProcessLimits(EProcessResource::Type Resource, uint64 Limit);
+	static CORE_API const TCHAR* ExecutablePath();
+	static CORE_API const TCHAR* ExecutableName(bool bRemoveExtension = true);
+	static CORE_API const TCHAR* GetModulePrefix();
+	static CORE_API const TCHAR* GetModuleExtension();
+	static CORE_API void ClosePipe( void* ReadPipe, void* WritePipe );
+	static CORE_API bool CreatePipe(void*& ReadPipe, void*& WritePipe, bool bWritePipeLocal = false);
+	static CORE_API FString ReadPipe( void* ReadPipe );
+	static CORE_API bool ReadPipeToArray(void* ReadPipe, TArray<uint8> & Output);
+	static CORE_API bool WritePipe(void* WritePipe, const FString& Message, FString* OutWritten = nullptr);
+	static CORE_API bool WritePipe(void* WritePipe, const uint8* Data, const int32 DataLength, int32* OutDataLength = nullptr);
+	static CORE_API class FRunnableThread* CreateRunnableThread();
+	static CORE_API const FString GetModulesDirectory();
+	static CORE_API bool CanLaunchURL(const TCHAR* URL);
+	static CORE_API void LaunchURL(const TCHAR* URL, const TCHAR* Parms, FString* Error);
+	static CORE_API FProcHandle CreateProc(const TCHAR* URL, const TCHAR* Parms, bool bLaunchDetached, bool bLaunchHidden, bool bLaunchReallyHidden, uint32* OutProcessID, int32 PriorityModifier, const TCHAR* OptionalWorkingDirectory, void* PipeWriteChild, void* PipeReadChild = nullptr);
+	static CORE_API FProcHandle CreateProc(const TCHAR* URL, const TCHAR* Parms, bool bLaunchDetached, bool bLaunchHidden, bool bLaunchReallyHidden, uint32* OutProcessID, int32 PriorityModifier, const TCHAR* OptionalWorkingDirectory, void* PipeWriteChild, void* PipeReadChild , void* PipeStdErrChild);
+	static CORE_API FProcHandle OpenProcess(uint32 ProcessID);
+	static CORE_API bool IsProcRunning( FProcHandle & ProcessHandle );
+	static CORE_API void WaitForProc( FProcHandle & ProcessHandle );
+	static CORE_API void CloseProc( FProcHandle & ProcessHandle );
+	static CORE_API void TerminateProc( FProcHandle & ProcessHandle, bool KillTree = false );
+	static CORE_API EWaitAndForkResult WaitAndFork();
+	static CORE_API uint32 GetCurrentProcessId();
+	static CORE_API uint32 GetCurrentCoreNumber();
+	static CORE_API bool GetProcReturnCode( FProcHandle & ProcHandle, int32* ReturnCode );
+	static CORE_API bool Daemonize();
+	static CORE_API bool IsApplicationRunning( uint32 ProcessId );
+	static CORE_API bool IsApplicationRunning( const TCHAR* ProcName );
+	static CORE_API bool ExecProcess(const TCHAR* URL, const TCHAR* Params, int32* OutReturnCode, FString* OutStdOut, FString* OutStdErr, const TCHAR* OptionalWorkingDirectory = NULL, bool bShouldEndWithParentProcess = false);
+	static CORE_API void ExploreFolder( const TCHAR* FilePath );
+	static CORE_API bool LaunchFileInDefaultExternalApplication( const TCHAR* FileName, const TCHAR* Parms = NULL, ELaunchVerb::Type Verb = ELaunchVerb::Open, bool bPromptToOpenOnFailure = true);
+	static CORE_API bool IsFirstInstance();
+	static CORE_API void OnChildEndFramePostFork();
+	static CORE_API int32 TranslateThreadPriority(EThreadPriority Priority);
+	static CORE_API void SetThreadNiceValue( uint32_t ThreadId, int32 NiceValue );
+	static CORE_API void SetThreadPriority( EThreadPriority NewPriority );
 
 	/**
 	 * @brief Releases locks that we held for IsFirstInstance check
 	 */
-	static void CeaseBeingFirstInstance();
+	static CORE_API void CeaseBeingFirstInstance();
 
 	/**
 	 * @brief Returns user home directory (i.e. $HOME).
 	 *
 	 * Like other directory functions, cannot return nullptr!
 	 */
-	static const TCHAR* UserHomeDir();
+	static CORE_API const TCHAR* UserHomeDir();
 };
diff --git a/Engine/Source/Runtime/Core/Public/Unix/UnixPlatformStackWalk.h b/Engine/Source/Runtime/Core/Public/Unix/UnixPlatformStackWalk.h
--- a/Engine/Source/Runtime/Core/Public/Unix/UnixPlatformStackWalk.h
+++ b/Engine/Source/Runtime/Core/Public/Unix/UnixPlatformStackWalk.h
@@ -11,19 +11,19 @@
 
 struct FGenericCrashContext;
 
-struct CORE_API FUnixPlatformStackWalk : public FGenericPlatformStackWalk
+struct FUnixPlatformStackWalk : public FGenericPlatformStackWalk
 {
-	static void ProgramCounterToSymbolInfo( uint64 ProgramCounter, FProgramCounterSymbolInfo& out_SymbolInfo );
-	static bool ProgramCounterToHumanReadableString( int32 CurrentCallDepth, uint64 ProgramCounter, ANSICHAR* HumanReadableString, SIZE_T HumanReadableStringSize, FGenericCrashContext* Context = nullptr );
-	static uint32 CaptureStackBackTrace( uint64* BackTrace, uint32 MaxDepth, void* Context = nullptr );
-	static void StackWalkAndDump( ANSICHAR* HumanReadableString, SIZE_T HumanReadableStringSize, int32 IgnoreCount, void* Context = nullptr );
-	static void StackWalkAndDump( ANSICHAR* HumanReadableString, SIZE_T HumanReadableStringSize, void* ProgramCounter, void* Context = nullptr );
-	static void StackWalkAndDumpEx(ANSICHAR* HumanReadableString, SIZE_T HumanReadableStringSize, int32 IgnoreCount, uint32 Flags, void* Context = nullptr);
-	static void StackWalkAndDumpEx(ANSICHAR* HumanReadableString, SIZE_T HumanReadableStringSize, void* ProgramCounter, uint32 Flags, void* Context = nullptr);
-	static uint32 CaptureThreadStackBackTrace(uint64 ThreadId, uint64* BackTrace, uint32 MaxDepth, void* Context = nullptr);
-	static void ThreadStackWalkAndDump(ANSICHAR* HumanReadableString, SIZE_T HumanReadableStringSize, int32 IgnoreCount, uint32 ThreadId);
-	static int32 GetProcessModuleCount();
-	static int32 GetProcessModuleSignatures(FStackWalkModuleInfo *ModuleSignatures, const int32 ModuleSignaturesSize);
+	static CORE_API void ProgramCounterToSymbolInfo( uint64 ProgramCounter, FProgramCounterSymbolInfo& out_SymbolInfo );
+	static CORE_API bool ProgramCounterToHumanReadableString( int32 CurrentCallDepth, uint64 ProgramCounter, ANSICHAR* HumanReadableString, SIZE_T HumanReadableStringSize, FGenericCrashContext* Context = nullptr );
+	static CORE_API uint32 CaptureStackBackTrace( uint64* BackTrace, uint32 MaxDepth, void* Context = nullptr );
+	static CORE_API void StackWalkAndDump( ANSICHAR* HumanReadableString, SIZE_T HumanReadableStringSize, int32 IgnoreCount, void* Context = nullptr );
+	static CORE_API void StackWalkAndDump( ANSICHAR* HumanReadableString, SIZE_T HumanReadableStringSize, void* ProgramCounter, void* Context = nullptr );
+	static CORE_API void StackWalkAndDumpEx(ANSICHAR* HumanReadableString, SIZE_T HumanReadableStringSize, int32 IgnoreCount, uint32 Flags, void* Context = nullptr);
+	static CORE_API void StackWalkAndDumpEx(ANSICHAR* HumanReadableString, SIZE_T HumanReadableStringSize, void* ProgramCounter, uint32 Flags, void* Context = nullptr);
+	static CORE_API uint32 CaptureThreadStackBackTrace(uint64 ThreadId, uint64* BackTrace, uint32 MaxDepth, void* Context = nullptr);
+	static CORE_API void ThreadStackWalkAndDump(ANSICHAR* HumanReadableString, SIZE_T HumanReadableStringSize, int32 IgnoreCount, uint32 ThreadId);
+	static CORE_API int32 GetProcessModuleCount();
+	static CORE_API int32 GetProcessModuleSignatures(FStackWalkModuleInfo *ModuleSignatures, const int32 ModuleSignaturesSize);
 };
 
 typedef FUnixPlatformStackWalk FPlatformStackWalk;
diff --git a/Engine/Source/Runtime/Core/Public/Unix/UnixPlatformTLS.h b/Engine/Source/Runtime/Core/Public/Unix/UnixPlatformTLS.h
--- a/Engine/Source/Runtime/Core/Public/Unix/UnixPlatformTLS.h
+++ b/Engine/Source/Runtime/Core/Public/Unix/UnixPlatformTLS.h
@@ -18,7 +18,7 @@
 /**
  * Unix implementation of the TLS OS functions
  */
-struct CORE_API FUnixTLS : public FGenericPlatformTLS
+struct FUnixTLS : public FGenericPlatformTLS
 {
 	/**
 	 * Returns the currently executing thread's id
@@ -133,10 +133,10 @@ private:
 
 #if IS_MONOLITHIC
 	/** per-thread static variable - cannot have that in a non-monolithic build because it eats into a precious OS limit. */
-	static __thread uint32 ThreadIdTLS;
+	static CORE_API __thread uint32 ThreadIdTLS;
 #else
 	/** TLS key to store TID */
-	static uint32 ThreadIdTLSKey;
+	static CORE_API uint32 ThreadIdTLSKey;
 #endif
 };
 
diff --git a/Engine/Source/Runtime/Core/Public/Unix/UnixPlatformTime.h b/Engine/Source/Runtime/Core/Public/Unix/UnixPlatformTime.h
--- a/Engine/Source/Runtime/Core/Public/Unix/UnixPlatformTime.h
+++ b/Engine/Source/Runtime/Core/Public/Unix/UnixPlatformTime.h
@@ -14,9 +14,9 @@
 /**
  * Unix implementation of the Time OS functions
  */
-struct CORE_API FUnixTime : public FGenericPlatformTime
+struct FUnixTime : public FGenericPlatformTime
 {
-	static double InitTiming();
+	static CORE_API double InitTiming();
 
 	static FORCEINLINE double Seconds()
 	{
@@ -49,37 +49,37 @@ struct CORE_API FUnixTime : public FGenericPlatformTime
 		return static_cast<uint64>(static_cast<uint64>(ts.tv_sec) * (uint64)1e7 + static_cast<uint64>(ts.tv_nsec) / 100ULL);
 	}
 
-	static bool UpdateCPUTime(float DeltaSeconds);
-	static bool UpdateThreadCPUTime(float = 0.0);
-	static void AutoUpdateGameThreadCPUTime(double UpdateInterval);
+	static CORE_API bool UpdateCPUTime(float DeltaSeconds);
+	static CORE_API bool UpdateThreadCPUTime(float = 0.0);
+	static CORE_API void AutoUpdateGameThreadCPUTime(double UpdateInterval);
 
-	static FCPUTime GetCPUTime();
-	static FCPUTime GetThreadCPUTime();
-	static double GetLastIntervalThreadCPUTimeInSeconds();
+	static CORE_API FCPUTime GetCPUTime();
+	static CORE_API FCPUTime GetThreadCPUTime();
+	static CORE_API double GetLastIntervalThreadCPUTimeInSeconds();
 
 	/**
 	 * Calibration log to be printed at later time
 	 */
-	static void PrintCalibrationLog();
+	static CORE_API void PrintCalibrationLog();
 
 private:
 
 	/** Clock source to use */
-	static int ClockSource;
+	static CORE_API int ClockSource;
 
 	/** Log information about calibrating the clock. */
-	static char CalibrationLog[4096];
+	static CORE_API char CalibrationLog[4096];
 
 	/**
 	 * Benchmarks clock_gettime(), possibly switches to something else is too slow.
 	 * Unix-specific.
 	 */
-	static int CalibrateAndSelectClock();
+	static CORE_API int CalibrateAndSelectClock();
 
 	/**
 	 * Returns number of time we can call the clock per second.
 	 */
-	static uint64 CallsPerSecondBenchmark(clockid_t BenchClockId, const char * BenchClockIdName);
+	static CORE_API uint64 CallsPerSecondBenchmark(clockid_t BenchClockId, const char * BenchClockIdName);
 };
 
-typedef FUnixTime FPlatformTime;
\ No newline at end of file
+typedef FUnixTime FPlatformTime;
diff --git a/Engine/Source/Runtime/Core/Public/Unix/UnixSignalHeartBeat.h b/Engine/Source/Runtime/Core/Public/Unix/UnixSignalHeartBeat.h
--- a/Engine/Source/Runtime/Core/Public/Unix/UnixSignalHeartBeat.h
+++ b/Engine/Source/Runtime/Core/Public/Unix/UnixSignalHeartBeat.h
@@ -6,37 +6,37 @@
 
 #include <time.h> // IWYU pragma: export
 
-class CORE_API FUnixSignalGameHitchHeartBeat
+class FUnixSignalGameHitchHeartBeat
 {
 public:
-	static FUnixSignalGameHitchHeartBeat* Singleton;
+	static CORE_API FUnixSignalGameHitchHeartBeat* Singleton;
 
 	/** Gets the heartbeat singleton */
-	static FUnixSignalGameHitchHeartBeat& Get();
-	static FUnixSignalGameHitchHeartBeat* GetNoInit();
+	static CORE_API FUnixSignalGameHitchHeartBeat& Get();
+	static CORE_API FUnixSignalGameHitchHeartBeat* GetNoInit();
 
 	/**
 	* Called at the start of a frame to register the time we are looking to detect a hitch
 	*/
-	void FrameStart(bool bSkipThisFrame = false);
+	CORE_API void FrameStart(bool bSkipThisFrame = false);
 
-	double GetFrameStartTime();
-	double GetCurrentTime();
+	CORE_API double GetFrameStartTime();
+	CORE_API double GetCurrentTime();
 
 	/**
 	* Suspend heartbeat hitch detection. Must call ResumeHeartBeat later to resume.
 	*/
-	void SuspendHeartBeat();
+	CORE_API void SuspendHeartBeat();
 
 	/**
 	* Resume heartbeat hitch detection. Call only after first calling SuspendHeartBeat.
 	*/
-	void ResumeHeartBeat();
+	CORE_API void ResumeHeartBeat();
 
-	void Restart();
-	void Stop();
+	CORE_API void Restart();
+	CORE_API void Stop();
 
-	void PostFork();
+	CORE_API void PostFork();
 
 private:
 	FUnixSignalGameHitchHeartBeat();
diff --git a/Engine/Source/Runtime/Core/Public/Windows/WindowsCriticalSection.h b/Engine/Source/Runtime/Core/Public/Windows/WindowsCriticalSection.h
--- a/Engine/Source/Runtime/Core/Public/Windows/WindowsCriticalSection.h
+++ b/Engine/Source/Runtime/Core/Public/Windows/WindowsCriticalSection.h
@@ -74,24 +74,24 @@ private:
 };
 
 /** System-Wide Critical Section for windows using mutex */
-class CORE_API FWindowsSystemWideCriticalSection
+class FWindowsSystemWideCriticalSection
 {
 public:
 	/** Construct a named, system-wide critical section and attempt to get access/ownership of it */
-	explicit FWindowsSystemWideCriticalSection(const class FString& InName, FTimespan InTimeout = FTimespan::Zero());
+	CORE_API explicit FWindowsSystemWideCriticalSection(const class FString& InName, FTimespan InTimeout = FTimespan::Zero());
 
 	/** Destructor releases system-wide critical section if it is currently owned */
-	~FWindowsSystemWideCriticalSection();
+	CORE_API ~FWindowsSystemWideCriticalSection();
 
 	/**
 	 * Does the calling thread have ownership of the system-wide critical section?
 	 *
 	 * @return True if obtained. WARNING: Returns true for an owned but previously abandoned locks so shared resources can be in undetermined states. You must handle shared data robustly.
 	 */
-	bool IsValid() const;
+	CORE_API bool IsValid() const;
 
 	/** Releases system-wide critical section if it is currently owned */
-	void Release();
+	CORE_API void Release();
 
 private:
 	FWindowsSystemWideCriticalSection(const FWindowsSystemWideCriticalSection&);
diff --git a/Engine/Source/Runtime/Core/Public/Windows/WindowsErrorOutputDevice.h b/Engine/Source/Runtime/Core/Public/Windows/WindowsErrorOutputDevice.h
--- a/Engine/Source/Runtime/Core/Public/Windows/WindowsErrorOutputDevice.h
+++ b/Engine/Source/Runtime/Core/Public/Windows/WindowsErrorOutputDevice.h
@@ -6,11 +6,11 @@
 #include "Misc/OutputDeviceError.h"
 #include "Misc/OutputDeviceConsole.h"
 
-class CORE_API FWindowsErrorOutputDevice : public FOutputDeviceError
+class FWindowsErrorOutputDevice : public FOutputDeviceError
 {
 public:
 	/** Constructor, initializing member variables */
-	FWindowsErrorOutputDevice();
+	CORE_API FWindowsErrorOutputDevice();
 
 	/**
 	 * Serializes the passed in data unless the current event is suppressed.
@@ -18,18 +18,18 @@ public:
 	 * @param	Data	Text to log
 	 * @param	Event	Event name used for suppression purposes
 	 */
-	virtual void Serialize( const TCHAR* Msg, ELogVerbosity::Type Verbosity, const class FName& Category ) override;
+	CORE_API virtual void Serialize( const TCHAR* Msg, ELogVerbosity::Type Verbosity, const class FName& Category ) override;
 
 	/**
 	 * Error handling function that is being called from within the system wide global
 	 * error handler, e.g. using structured exception handling on the PC.
 	 */
-	virtual void HandleError() override;
+	CORE_API virtual void HandleError() override;
 
 protected:
 	/**
 	 * Callback to allow FWindowsApplicationErrorOutputDevice to restore the UI.
 	 */
-	virtual void HandleErrorRestoreUI();
+	CORE_API virtual void HandleErrorRestoreUI();
 };
 
diff --git a/Engine/Source/Runtime/Core/Public/Windows/WindowsEventLogOutputDevice.h b/Engine/Source/Runtime/Core/Public/Windows/WindowsEventLogOutputDevice.h
--- a/Engine/Source/Runtime/Core/Public/Windows/WindowsEventLogOutputDevice.h
+++ b/Engine/Source/Runtime/Core/Public/Windows/WindowsEventLogOutputDevice.h
@@ -11,7 +11,7 @@ class FName;
 /**
  * Output device that writes to Windows Event Log
  */
-class CORE_API FWindowsEventLogOutputDevice
+class FWindowsEventLogOutputDevice
 	: public FOutputDevice
 {
 	/** Handle to the event log object */
@@ -21,18 +21,18 @@ public:
 	/**
 	 * Constructor, initializing member variables
 	 */
-	FWindowsEventLogOutputDevice();
+	CORE_API FWindowsEventLogOutputDevice();
 
 	/** Destructor that cleans up any remaining resources */
-	virtual ~FWindowsEventLogOutputDevice();
+	CORE_API virtual ~FWindowsEventLogOutputDevice();
 
-	virtual void Serialize(const TCHAR* Buffer, ELogVerbosity::Type Verbosity, const class FName& Category) override;
+	CORE_API virtual void Serialize(const TCHAR* Buffer, ELogVerbosity::Type Verbosity, const class FName& Category) override;
 	
 	/** Does nothing */
-	virtual void Flush(void);
+	CORE_API virtual void Flush(void);
 
 	/**
 	 * Closes any event log handles that are open
 	 */
-	virtual void TearDown(void);
+	CORE_API virtual void TearDown(void);
 };
diff --git a/Engine/Source/Runtime/Core/Public/Windows/WindowsPlatformAtomics.h b/Engine/Source/Runtime/Core/Public/Windows/WindowsPlatformAtomics.h
--- a/Engine/Source/Runtime/Core/Public/Windows/WindowsPlatformAtomics.h
+++ b/Engine/Source/Runtime/Core/Public/Windows/WindowsPlatformAtomics.h
@@ -10,7 +10,7 @@
 /**
  * Windows implementation of the Atomics OS functions
  */
-struct CORE_API FWindowsPlatformAtomics
+struct FWindowsPlatformAtomics
 	: public FGenericPlatformAtomics
 {
 	static_assert(sizeof(int8)  == sizeof(char)      && alignof(int8)  == alignof(char),      "int8 must be compatible with char");
@@ -444,7 +444,7 @@ protected:
 	 *
 	 * @param InFormat - The string format string.
 	 */
-	static void HandleAtomicsFailure( const TCHAR* InFormat, ... );
+	static CORE_API void HandleAtomicsFailure( const TCHAR* InFormat, ... );
 };
 
 
diff --git a/Engine/Source/Runtime/Core/Public/Windows/WindowsPlatformCrashContext.h b/Engine/Source/Runtime/Core/Public/Windows/WindowsPlatformCrashContext.h
--- a/Engine/Source/Runtime/Core/Public/Windows/WindowsPlatformCrashContext.h
+++ b/Engine/Source/Runtime/Core/Public/Windows/WindowsPlatformCrashContext.h
@@ -6,9 +6,9 @@
 #include "Microsoft/MicrosoftPlatformCrashContext.h"
 
 
-struct CORE_API FWindowsPlatformCrashContext : public FMicrosoftPlatformCrashContext
+struct FWindowsPlatformCrashContext : public FMicrosoftPlatformCrashContext
 {
-	static const TCHAR* const UEGPUAftermathMinidumpName;
+	static CORE_API const TCHAR* const UEGPUAftermathMinidumpName;
 	
 	FWindowsPlatformCrashContext(ECrashContextType InType, const TCHAR* InErrorMessage)
 		: FMicrosoftPlatformCrashContext(InType, InErrorMessage)
@@ -16,8 +16,8 @@ struct CORE_API FWindowsPlatformCrashContext : public FMicrosoftPlatformCrashCon
 	}
 
 	
-	virtual void AddPlatformSpecificProperties() const override;
-	virtual void CopyPlatformSpecificFiles(const TCHAR* OutputDirectory, void* Context) override;
+	CORE_API virtual void AddPlatformSpecificProperties() const override;
+	CORE_API virtual void CopyPlatformSpecificFiles(const TCHAR* OutputDirectory, void* Context) override;
 };
 
 typedef FWindowsPlatformCrashContext FPlatformCrashContext;
diff --git a/Engine/Source/Runtime/Core/Public/Windows/WindowsPlatformMemory.h b/Engine/Source/Runtime/Core/Public/Windows/WindowsPlatformMemory.h
--- a/Engine/Source/Runtime/Core/Public/Windows/WindowsPlatformMemory.h
+++ b/Engine/Source/Runtime/Core/Public/Windows/WindowsPlatformMemory.h
@@ -40,7 +40,7 @@ struct FPlatformMemoryStats
 /**
 * Windows implementation of the memory OS functions
 **/
-struct CORE_API FWindowsPlatformMemory
+struct FWindowsPlatformMemory
 	: public FGenericPlatformMemory
 {
 	enum EMemoryCounterRegion
@@ -78,7 +78,7 @@ struct CORE_API FWindowsPlatformMemory
 	};
 
 	//~ Begin FGenericPlatformMemory Interface
-	static void Init();
+	static CORE_API void Init();
 	static uint32 GetBackMemoryPoolSize()
 	{
 		/**
@@ -91,14 +91,14 @@ struct CORE_API FWindowsPlatformMemory
 		return 32 * 1024 * 1024;
 	}
 
-	static class FMalloc* BaseAllocator();
-	static FPlatformMemoryStats GetStats();
-	static void GetStatsForMallocProfiler( FGenericMemoryStats& out_Stats );
-	static const FPlatformMemoryConstants& GetConstants();
-	static bool PageProtect(void* const Ptr, const SIZE_T Size, const bool bCanRead, const bool bCanWrite);
-	static void* BinnedAllocFromOS( SIZE_T Size );
-	static void BinnedFreeToOS( void* Ptr, SIZE_T Size );
-	static void MiMallocInit();
+	static CORE_API class FMalloc* BaseAllocator();
+	static CORE_API FPlatformMemoryStats GetStats();
+	static CORE_API void GetStatsForMallocProfiler( FGenericMemoryStats& out_Stats );
+	static CORE_API const FPlatformMemoryConstants& GetConstants();
+	static CORE_API bool PageProtect(void* const Ptr, const SIZE_T Size, const bool bCanRead, const bool bCanWrite);
+	static CORE_API void* BinnedAllocFromOS( SIZE_T Size );
+	static CORE_API void BinnedFreeToOS( void* Ptr, SIZE_T Size );
+	static CORE_API void MiMallocInit();
 
 	class FPlatformVirtualMemoryBlock : public FBasicVirtualMemoryBlock
 	{
@@ -149,13 +149,13 @@ struct CORE_API FWindowsPlatformMemory
 		static size_t GetVirtualSizeAlignment();
 	};
 
-	static FSharedMemoryRegion* MapNamedSharedMemoryRegion(const FString& InName, bool bCreate, uint32 AccessMode, SIZE_T Size, const void* pSecurityAttributes = nullptr);
-	static bool UnmapNamedSharedMemoryRegion(FSharedMemoryRegion * MemoryRegion);
-	static bool GetLLMAllocFunctions(void*(*&OutAllocFunction)(size_t), void(*&OutFreeFunction)(void*, size_t), int32& OutAlignment);
+	static CORE_API FSharedMemoryRegion* MapNamedSharedMemoryRegion(const FString& InName, bool bCreate, uint32 AccessMode, SIZE_T Size, const void* pSecurityAttributes = nullptr);
+	static CORE_API bool UnmapNamedSharedMemoryRegion(FSharedMemoryRegion * MemoryRegion);
+	static CORE_API bool GetLLMAllocFunctions(void*(*&OutAllocFunction)(size_t), void(*&OutFreeFunction)(void*, size_t), int32& OutAlignment);
 protected:
 	friend struct FGenericStatsUpdater;
 
-	static void InternalUpdateStats( const FPlatformMemoryStats& MemoryStats );
+	static CORE_API void InternalUpdateStats( const FPlatformMemoryStats& MemoryStats );
 	//~ End FGenericPlatformMemory Interface
 };
 
diff --git a/Engine/Source/Runtime/Core/Public/Windows/WindowsPlatformMisc.h b/Engine/Source/Runtime/Core/Public/Windows/WindowsPlatformMisc.h
--- a/Engine/Source/Runtime/Core/Public/Windows/WindowsPlatformMisc.h
+++ b/Engine/Source/Runtime/Core/Public/Windows/WindowsPlatformMisc.h
@@ -13,7 +13,7 @@ struct FGuid;
 class IPlatformChunkInstall;
 
 /** Helper struct used to get the string version of the Windows version. */
-struct CORE_API FWindowsOSVersionHelper
+struct FWindowsOSVersionHelper
 {
 	enum ErrorCodes
 	{
@@ -24,7 +24,7 @@ struct CORE_API FWindowsOSVersionHelper
 		ERROR_GETWINDOWSGT62VERSIONS_FAILED = 8,
 	};
 
-	static int32 GetOSVersions( FString& OutOSVersion, FString& OutOSSubVersion );
+	static CORE_API int32 GetOSVersions( FString& OutOSVersion, FString& OutOSSubVersion );
 };
 
 /**
@@ -41,39 +41,39 @@ enum class ECOMModel : uint8
 /**
 * Windows implementation of the misc OS functions
 **/
-struct CORE_API FWindowsPlatformMisc
+struct FWindowsPlatformMisc
 	: public FGenericPlatformMisc
 {
-	static void PlatformPreInit();
-	static void PlatformInit();
-	static void PlatformTearDown();
-	static void SetGracefulTerminationHandler();
-	static void CallGracefulTerminationHandler();
-	static ECrashHandlingType GetCrashHandlingType();
-	static ECrashHandlingType SetCrashHandlingType(ECrashHandlingType);
-	static int32 GetMaxPathLength();
+	static CORE_API void PlatformPreInit();
+	static CORE_API void PlatformInit();
+	static CORE_API void PlatformTearDown();
+	static CORE_API void SetGracefulTerminationHandler();
+	static CORE_API void CallGracefulTerminationHandler();
+	static CORE_API ECrashHandlingType GetCrashHandlingType();
+	static CORE_API ECrashHandlingType SetCrashHandlingType(ECrashHandlingType);
+	static CORE_API int32 GetMaxPathLength();
 
 	UE_DEPRECATED(4.21, "void FPlatformMisc::GetEnvironmentVariable(Name, Result, Length) is deprecated. Use FString FPlatformMisc::GetEnvironmentVariable(Name) instead.")
-	static void GetEnvironmentVariable(const TCHAR* VariableName, TCHAR* Result, int32 ResultLength);
+	static CORE_API void GetEnvironmentVariable(const TCHAR* VariableName, TCHAR* Result, int32 ResultLength);
 
-	static FString GetEnvironmentVariable(const TCHAR* VariableName);
-	static void SetEnvironmentVar(const TCHAR* VariableName, const TCHAR* Value);
+	static CORE_API FString GetEnvironmentVariable(const TCHAR* VariableName);
+	static CORE_API void SetEnvironmentVar(const TCHAR* VariableName, const TCHAR* Value);
 
-	static TArray<uint8> GetMacAddress();
-	static void SubmitErrorReport( const TCHAR* InErrorHist, EErrorReportMode::Type InMode );
+	static CORE_API TArray<uint8> GetMacAddress();
+	static CORE_API void SubmitErrorReport( const TCHAR* InErrorHist, EErrorReportMode::Type InMode );
 
 #if !UE_BUILD_SHIPPING
-	static bool IsDebuggerPresent();
-	static EProcessDiagnosticFlags GetProcessDiagnostics();
+	static CORE_API bool IsDebuggerPresent();
+	static CORE_API EProcessDiagnosticFlags GetProcessDiagnostics();
 #endif
 
 #if STATS || ENABLE_STATNAMEDEVENTS
-	static void BeginNamedEventFrame();
-	static void BeginNamedEvent(const struct FColor& Color, const TCHAR* Text);
-	static void BeginNamedEvent(const struct FColor& Color, const ANSICHAR* Text);
-	static void EndNamedEvent();
-	static void CustomNamedStat(const TCHAR* Text, float Value, const TCHAR* Graph, const TCHAR* Unit);
-	static void CustomNamedStat(const ANSICHAR* Text, float Value, const ANSICHAR* Graph, const ANSICHAR* Unit);
+	static CORE_API void BeginNamedEventFrame();
+	static CORE_API void BeginNamedEvent(const struct FColor& Color, const TCHAR* Text);
+	static CORE_API void BeginNamedEvent(const struct FColor& Color, const ANSICHAR* Text);
+	static CORE_API void EndNamedEvent();
+	static CORE_API void CustomNamedStat(const TCHAR* Text, float Value, const TCHAR* Graph, const TCHAR* Unit);
+	static CORE_API void CustomNamedStat(const ANSICHAR* Text, float Value, const ANSICHAR* Graph, const ANSICHAR* Unit);
 #endif
 
 	FORCEINLINE static void MemoryBarrier() 
@@ -85,12 +85,12 @@ struct CORE_API FWindowsPlatformMisc
 #endif
 	}
 
-	static bool IsRemoteSession();
+	static CORE_API bool IsRemoteSession();
 
-	static void SetUTF8Output();
-	static void LocalPrint(const TCHAR *Message);
+	static CORE_API void SetUTF8Output();
+	static CORE_API void LocalPrint(const TCHAR *Message);
 
-	static bool IsLowLevelOutputDebugStringStructured();
+	static CORE_API bool IsLowLevelOutputDebugStringStructured();
 
 	static bool IsLocalPrintThreadSafe()
 	{ 
@@ -100,34 +100,34 @@ struct CORE_API FWindowsPlatformMisc
 		return IsDebuggerPresent();
 	}
 	
-	static void RequestExitWithStatus(bool Force, uint8 ReturnCode, const TCHAR* CallSite = nullptr);
-	static void RequestExit(bool Force, const TCHAR* CallSite = nullptr);
-	static const TCHAR* GetSystemErrorMessage(TCHAR* OutBuffer, int32 BufferCount, int32 Error);
-	static void CreateGuid(struct FGuid& Result);
-	static EAppReturnType::Type MessageBoxExt( EAppMsgType::Type MsgType, const TCHAR* Text, const TCHAR* Caption );
-	static bool CommandLineCommands();
-	static bool Is64bitOperatingSystem();
-	static bool IsValidAbsolutePathFormat(const FString& Path);
-	static int32 NumberOfCores();
-	static const FProcessorGroupDesc& GetProcessorGroupDesc();
-	static int32 NumberOfCoresIncludingHyperthreads();
-	static int32 NumberOfWorkerThreadsToSpawn();
-
-	static const TCHAR* GetPlatformFeaturesModuleName();
-
-	static FString GetDefaultLanguage();
-	static FString GetDefaultLocale();
-
-	static uint32 GetLastError();
-	static void SetLastError(uint32 ErrorCode);
-	static void RaiseException( uint32 ExceptionCode );
-	static bool SetStoredValue(const FString& InStoreId, const FString& InSectionName, const FString& InKeyName, const FString& InValue);
-	static bool GetStoredValue(const FString& InStoreId, const FString& InSectionName, const FString& InKeyName, FString& OutValue);
-	static bool DeleteStoredValue(const FString& InStoreId, const FString& InSectionName, const FString& InKeyName);
-	static bool DeleteStoredSection(const FString& InStoreId, const FString& InSectionName);
-
-	static bool CoInitialize(ECOMModel Model = ECOMModel::Singlethreaded);
-	static void CoUninitialize();
+	static CORE_API void RequestExitWithStatus(bool Force, uint8 ReturnCode, const TCHAR* CallSite = nullptr);
+	static CORE_API void RequestExit(bool Force, const TCHAR* CallSite = nullptr);
+	static CORE_API const TCHAR* GetSystemErrorMessage(TCHAR* OutBuffer, int32 BufferCount, int32 Error);
+	static CORE_API void CreateGuid(struct FGuid& Result);
+	static CORE_API EAppReturnType::Type MessageBoxExt( EAppMsgType::Type MsgType, const TCHAR* Text, const TCHAR* Caption );
+	static CORE_API bool CommandLineCommands();
+	static CORE_API bool Is64bitOperatingSystem();
+	static CORE_API bool IsValidAbsolutePathFormat(const FString& Path);
+	static CORE_API int32 NumberOfCores();
+	static CORE_API const FProcessorGroupDesc& GetProcessorGroupDesc();
+	static CORE_API int32 NumberOfCoresIncludingHyperthreads();
+	static CORE_API int32 NumberOfWorkerThreadsToSpawn();
+
+	static CORE_API const TCHAR* GetPlatformFeaturesModuleName();
+
+	static CORE_API FString GetDefaultLanguage();
+	static CORE_API FString GetDefaultLocale();
+
+	static CORE_API uint32 GetLastError();
+	static CORE_API void SetLastError(uint32 ErrorCode);
+	static CORE_API void RaiseException( uint32 ExceptionCode );
+	static CORE_API bool SetStoredValue(const FString& InStoreId, const FString& InSectionName, const FString& InKeyName, const FString& InValue);
+	static CORE_API bool GetStoredValue(const FString& InStoreId, const FString& InSectionName, const FString& InKeyName, FString& OutValue);
+	static CORE_API bool DeleteStoredValue(const FString& InStoreId, const FString& InSectionName, const FString& InKeyName);
+	static CORE_API bool DeleteStoredSection(const FString& InStoreId, const FString& InSectionName);
+
+	static CORE_API bool CoInitialize(ECOMModel Model = ECOMModel::Singlethreaded);
+	static CORE_API void CoUninitialize();
 
 	/**
 	 * Has the OS execute a command and path pair (such as launch a browser)
@@ -138,7 +138,7 @@ struct CORE_API FWindowsPlatformMisc
 	 *
 	 * @return whether the command was successful or not
 	 */
-	static bool OsExecute(const TCHAR* CommandType, const TCHAR* Command, const TCHAR* CommandLine = NULL);
+	static CORE_API bool OsExecute(const TCHAR* CommandType, const TCHAR* Command, const TCHAR* CommandLine = NULL);
 
 	/**
 	 * Attempts to get the handle to a top-level window of the specified process.
@@ -149,7 +149,7 @@ struct CORE_API FWindowsPlatformMisc
 	 * @param ProcessId The identifier of the process to get the window for.
 	 * @return Window handle, or 0 if not found.
 	 */
-	static Windows::HWND GetTopLevelWindowHandle(uint32 ProcessId);
+	static CORE_API Windows::HWND GetTopLevelWindowHandle(uint32 ProcessId);
 
 	/** 
 	 * Determines if we are running on the Windows version or newer
@@ -159,10 +159,10 @@ struct CORE_API FWindowsPlatformMisc
 	 *
 	 * @return	Returns true if the current Windows version if equal or newer than MajorVersion
 	 */
-	static bool VerifyWindowsVersion(uint32 MajorVersion, uint32 MinorVersion, uint32 BuildNumber = 0);
+	static CORE_API bool VerifyWindowsVersion(uint32 MajorVersion, uint32 MinorVersion, uint32 BuildNumber = 0);
 
 #if !UE_BUILD_SHIPPING
-	static void PromptForRemoteDebugging(bool bIsEnsure);
+	static CORE_API void PromptForRemoteDebugging(bool bIsEnsure);
 #endif	//#if !UE_BUILD_SHIPPING
 
 
@@ -171,24 +171,24 @@ struct CORE_API FWindowsPlatformMisc
 	 *
 	 * @return	Returns true if cpuid is supported
 	 */
-	static bool HasCPUIDInstruction();
+	static CORE_API bool HasCPUIDInstruction();
 
 	/**
 	 * Determines if AVX2 instruction set is supported on this platform
 	 *
 	 * @return	Returns true if instruction-set is supported
 	 */
-	static bool HasAVX2InstructionSupport();
-
-	static FString GetCPUVendor();
-	static FString GetCPUBrand();
-	static FString GetPrimaryGPUBrand();
-	static struct FGPUDriverInfo GetGPUDriverInfo(const FString& DeviceDescription, bool bVerbose = true);
-	static void GetOSVersions( FString& out_OSVersionLabel, FString& out_OSSubVersionLabel );
-	static FString GetOSVersion();
-	static bool GetDiskTotalAndFreeSpace( const FString& InPath, uint64& TotalNumberOfBytes, uint64& NumberOfFreeBytes );
-	static bool GetPageFaultStats(FPageFaultStats& OutStats, EPageFaultFlags Flags=EPageFaultFlags::All);
-	static bool GetBlockingIOStats(FProcessIOStats& OutStats, EInputOutputFlags Flags=EInputOutputFlags::All);
+	static CORE_API bool HasAVX2InstructionSupport();
+
+	static CORE_API FString GetCPUVendor();
+	static CORE_API FString GetCPUBrand();
+	static CORE_API FString GetPrimaryGPUBrand();
+	static CORE_API struct FGPUDriverInfo GetGPUDriverInfo(const FString& DeviceDescription, bool bVerbose = true);
+	static CORE_API void GetOSVersions( FString& out_OSVersionLabel, FString& out_OSSubVersionLabel );
+	static CORE_API FString GetOSVersion();
+	static CORE_API bool GetDiskTotalAndFreeSpace( const FString& InPath, uint64& TotalNumberOfBytes, uint64& NumberOfFreeBytes );
+	static CORE_API bool GetPageFaultStats(FPageFaultStats& OutStats, EPageFaultFlags Flags=EPageFaultFlags::All);
+	static CORE_API bool GetBlockingIOStats(FProcessIOStats& OutStats, EInputOutputFlags Flags=EInputOutputFlags::All);
 
 	/**
 	 * Uses cpuid instruction to get the vendor string
@@ -204,14 +204,14 @@ struct CORE_API FWindowsPlatformMisc
 	 *			Bits 20-27	Extended family
 	 *			Bits 28-31	Reserved
 	 */
-	static uint32 GetCPUInfo();
+	static CORE_API uint32 GetCPUInfo();
 
 	/** @return whether this cpu supports certain required instructions or not */
-	static bool HasNonoptionalCPUFeatures();
+	static CORE_API bool HasNonoptionalCPUFeatures();
 	/** @return whether to check for specific CPU compatibility or not */
-	static bool NeedsNonoptionalCPUFeaturesCheck();
+	static CORE_API bool NeedsNonoptionalCPUFeaturesCheck();
 	/** @return whether this cpu has timed pause instruction support or not */
-	static bool HasTimedPauseCPUFeature();
+	static CORE_API bool HasTimedPauseCPUFeature();
 
 	/** 
 	 * Provides a simpler interface for fetching and cleanup of registry value queries
@@ -223,7 +223,7 @@ struct CORE_API FWindowsPlatformMisc
 	 *
 	 * @return	true, if it successfully found the Value
 	 */
-	static bool QueryRegKey( const Windows::HKEY InKey, const TCHAR* InSubKey, const TCHAR* InValueName, FString& OutData );
+	static CORE_API bool QueryRegKey( const Windows::HKEY InKey, const TCHAR* InSubKey, const TCHAR* InValueName, FString& OutData );
 
 	/**
 	 * Gets Visual Studio common tools path.
@@ -233,17 +233,17 @@ struct CORE_API FWindowsPlatformMisc
 	 *
 	 * @return Returns if succeeded.
 	 */
-	static bool GetVSComnTools(int32 Version, FString& OutData);
+	static CORE_API bool GetVSComnTools(int32 Version, FString& OutData);
 
 	UE_DEPRECATED(5.2, "Please use PLATFORM_CACHE_LINE_SIZE instead, runtime query of cache line size not supported")
-	static int32 GetCacheLineSize();
+	static CORE_API int32 GetCacheLineSize();
 	/**
 	* @return Windows path separator.
 	*/
-	static const TCHAR* GetDefaultPathSeparator();
+	static CORE_API const TCHAR* GetDefaultPathSeparator();
 
 	/** @return Get the name of the platform specific file manager (Explorer) */
-	static FText GetFileManagerName();
+	static CORE_API FText GetFileManagerName();
 
 	/**
 	* Returns whether WiFi connection is currently active
@@ -257,7 +257,7 @@ struct CORE_API FWindowsPlatformMisc
 	/**
 	 * Returns whether the platform is running on battery power or not.
 	 */
-	static bool IsRunningOnBattery();
+	static CORE_API bool IsRunningOnBattery();
 
 	FORCEINLINE static void ChooseHDRDeviceAndColorGamut(uint32 DeviceId, uint32 DisplayNitLevel, EDisplayOutputFormat& OutputDevice, EDisplayColorGamut& ColorGamut)
 	{
@@ -281,17 +281,17 @@ struct CORE_API FWindowsPlatformMisc
 	/**
 	 * Gets a globally unique ID the represents a particular operating system install.
 	 */
-	static FString GetOperatingSystemId();
+	static CORE_API FString GetOperatingSystemId();
 
-	static EConvertibleLaptopMode GetConvertibleLaptopMode();
+	static CORE_API EConvertibleLaptopMode GetConvertibleLaptopMode();
 
-	static IPlatformChunkInstall* GetPlatformChunkInstall();
+	static CORE_API IPlatformChunkInstall* GetPlatformChunkInstall();
 
-	static void PumpMessagesOutsideMainLoop();
+	static CORE_API void PumpMessagesOutsideMainLoop();
 
-	static uint64 GetFileVersion(const FString &FileName);
+	static CORE_API uint64 GetFileVersion(const FString &FileName);
 
-	static int32 GetMaxRefreshRate();
+	static CORE_API int32 GetMaxRefreshRate();
 };
 
 
diff --git a/Engine/Source/Runtime/Core/Public/Windows/WindowsPlatformNamedPipe.h b/Engine/Source/Runtime/Core/Public/Windows/WindowsPlatformNamedPipe.h
--- a/Engine/Source/Runtime/Core/Public/Windows/WindowsPlatformNamedPipe.h
+++ b/Engine/Source/Runtime/Core/Public/Windows/WindowsPlatformNamedPipe.h
@@ -13,13 +13,13 @@ class FString;
 
 
 // Windows wrapper for named pipe communications
-class CORE_API FWindowsPlatformNamedPipe
+class FWindowsPlatformNamedPipe
 	: public FGenericPlatformNamedPipe
 {
 public:
 
-	FWindowsPlatformNamedPipe();
-	virtual ~FWindowsPlatformNamedPipe();
+	CORE_API FWindowsPlatformNamedPipe();
+	CORE_API virtual ~FWindowsPlatformNamedPipe();
 
 	FWindowsPlatformNamedPipe(const FWindowsPlatformNamedPipe&) = delete;
 	FWindowsPlatformNamedPipe& operator=(const FWindowsPlatformNamedPipe&) = delete;
@@ -28,16 +28,16 @@ public:
 
 	// FGenericPlatformNamedPipe overrides
 
-	virtual bool Create(const FString& PipeName, bool bAsServer, bool bAsync) override;
-	virtual bool Destroy() override;
-	virtual bool OpenConnection() override;
-	virtual bool BlockForAsyncIO() override;
-	virtual bool IsReadyForRW() const override;
-	virtual bool UpdateAsyncStatus() override;
-	virtual bool WriteBytes(int32 NumBytes, const void* Data) override;
-	virtual bool ReadBytes(int32 NumBytes, void* OutData) override;
-	virtual bool IsCreated() const override;
-	virtual bool HasFailed() const override;
+	CORE_API virtual bool Create(const FString& PipeName, bool bAsServer, bool bAsync) override;
+	CORE_API virtual bool Destroy() override;
+	CORE_API virtual bool OpenConnection() override;
+	CORE_API virtual bool BlockForAsyncIO() override;
+	CORE_API virtual bool IsReadyForRW() const override;
+	CORE_API virtual bool UpdateAsyncStatus() override;
+	CORE_API virtual bool WriteBytes(int32 NumBytes, const void* Data) override;
+	CORE_API virtual bool ReadBytes(int32 NumBytes, void* OutData) override;
+	CORE_API virtual bool IsCreated() const override;
+	CORE_API virtual bool HasFailed() const override;
 
 private:
 
@@ -60,7 +60,7 @@ private:
 
 	EState State;
 
-	bool UpdateAsyncStatusAfterRW();
+	CORE_API bool UpdateAsyncStatusAfterRW();
 };
 
 
diff --git a/Engine/Source/Runtime/Core/Public/Windows/WindowsPlatformOutputDevices.h b/Engine/Source/Runtime/Core/Public/Windows/WindowsPlatformOutputDevices.h
--- a/Engine/Source/Runtime/Core/Public/Windows/WindowsPlatformOutputDevices.h
+++ b/Engine/Source/Runtime/Core/Public/Windows/WindowsPlatformOutputDevices.h
@@ -10,11 +10,11 @@ class FOutputDeviceConsole;
 class FOutputDeviceError;
 class FFeedbackContext;
 
-struct CORE_API FWindowsPlatformOutputDevices
+struct FWindowsPlatformOutputDevices
 	: public FGenericPlatformOutputDevices
 {
-	static FOutputDevice*			GetEventLog();
-	static FOutputDeviceError*      GetError();
+	static CORE_API FOutputDevice*			GetEventLog();
+	static CORE_API FOutputDeviceError*      GetError();
 };
 
 
diff --git a/Engine/Source/Runtime/Core/Public/Windows/WindowsPlatformProcess.h b/Engine/Source/Runtime/Core/Public/Windows/WindowsPlatformProcess.h
--- a/Engine/Source/Runtime/Core/Public/Windows/WindowsPlatformProcess.h
+++ b/Engine/Source/Runtime/Core/Public/Windows/WindowsPlatformProcess.h
@@ -28,7 +28,7 @@ public:
 /**
 * Windows implementation of the Process OS functions.
 **/
-struct CORE_API FWindowsPlatformProcess
+struct FWindowsPlatformProcess
 	: public FGenericPlatformProcess
 {
 	/**
@@ -63,26 +63,26 @@ struct CORE_API FWindowsPlatformProcess
 	/**
 	 * Process enumerator.
 	 */
-	class CORE_API FProcEnumerator
+	class FProcEnumerator
 	{
 	public:
 		// Constructor
-		FProcEnumerator();
+		CORE_API FProcEnumerator();
 		FProcEnumerator(const FProcEnumerator&) = delete;
 		FProcEnumerator& operator=(const FProcEnumerator&) = delete;
 
 		// Destructor
-		~FProcEnumerator();
+		CORE_API ~FProcEnumerator();
 
 		// Gets current process enumerator info.
-		FProcEnumInfo GetCurrent() const;
+		CORE_API FProcEnumInfo GetCurrent() const;
 
 		/**
 		 * Moves current to the next process.
 		 *
 		 * @returns True if succeeded. False otherwise.
 		 */
-		bool MoveNext();
+		CORE_API bool MoveNext();
 	private:
 		// Process info structure for current process.
 		Windows::PROCESSENTRY32* CurrentEntry;
@@ -94,24 +94,24 @@ struct CORE_API FWindowsPlatformProcess
 	/**
 	 * Process enumeration info structure.
 	 */
-	struct CORE_API FProcEnumInfo
+	struct FProcEnumInfo
 	{
 		friend FProcEnumInfo FProcEnumerator::GetCurrent() const;
 	public:
 		// Destructor
-		~FProcEnumInfo();
+		CORE_API ~FProcEnumInfo();
 
 		// Gets process PID.
-		uint32 GetPID() const;
+		CORE_API uint32 GetPID() const;
 
 		// Gets parent process PID.
-		uint32 GetParentPID() const;
+		CORE_API uint32 GetParentPID() const;
 
 		// Gets process name. I.e. exec name.
-		FString GetName() const;
+		CORE_API FString GetName() const;
 
 		// Gets process full image path. I.e. full path of the exec file.
-		FString GetFullPath() const;
+		CORE_API FString GetFullPath() const;
 
 	private:
 		// Private constructor.
@@ -125,85 +125,85 @@ public:
 
 	// FGenericPlatformProcess interface
 
-	static void* GetDllHandle( const TCHAR* Filename );
-	static void FreeDllHandle( void* DllHandle );
-	static void* GetDllExport( void* DllHandle, const TCHAR* ProcName );
-	static void AddDllDirectory(const TCHAR* Directory);
-	static void GetDllDirectories(TArray<FString>& OutDllDirectories);
-	static void PushDllDirectory(const TCHAR* Directory);
-	static void PopDllDirectory(const TCHAR* Directory);
-	static uint32 GetCurrentProcessId();
-	static uint32 GetCurrentCoreNumber();
-	static void SetThreadAffinityMask( uint64 AffinityMask );
-	static void SetThreadName( const TCHAR* ThreadName );
-	static const TCHAR* BaseDir();
-	static const TCHAR* UserDir();
-	static const TCHAR* UserTempDir();
-	static const TCHAR* UserSettingsDir();
-	static const TCHAR* ApplicationSettingsDir();
-	static FString GetApplicationSettingsDir(const ApplicationSettingsContext& Settings);
-	static const TCHAR* ComputerName();
-	static const TCHAR* UserName(bool bOnlyAlphaNumeric = true);
-	static void SetCurrentWorkingDirectoryToBaseDir();
-	static FString GetCurrentWorkingDirectory();
-	static const FString ShaderWorkingDir();
-	static const TCHAR* ExecutablePath();
-	static const TCHAR* ExecutableName(bool bRemoveExtension = true);
-	static FString GenerateApplicationPath( const FString& AppName, EBuildConfiguration BuildConfiguration);
-	static const TCHAR* GetModuleExtension();
-	static const TCHAR* GetBinariesSubdirectory();
-	static const FString GetModulesDirectory();
-	static bool CanLaunchURL(const TCHAR* URL);
-	static void LaunchURL( const TCHAR* URL, const TCHAR* Parms, FString* Error );
-	static FProcHandle CreateProc( const TCHAR* URL, const TCHAR* Parms, bool bLaunchDetached, bool bLaunchHidden, bool bLaunchReallyHidden, uint32* OutProcessID, int32 PriorityModifier, const TCHAR* OptionalWorkingDirectory, void* PipeWriteChild, void* PipeReadChild = nullptr);
-	static FProcHandle CreateProc( const TCHAR* URL, const TCHAR* Parms, bool bLaunchDetached, bool bLaunchHidden, bool bLaunchReallyHidden, uint32* OutProcessID, int32 PriorityModifier, const TCHAR* OptionalWorkingDirectory, void* PipeWriteChild, void* PipeReadChild, void* PipeStdErrChild);
-	static bool SetProcPriority(FProcHandle & InProcHandle, int32 PriorityModifier);
-	static FProcHandle OpenProcess(uint32 ProcessID);
-	static bool IsProcRunning( FProcHandle & ProcessHandle );
-	static void WaitForProc( FProcHandle & ProcessHandle );
-	static void CloseProc( FProcHandle & ProcessHandle );
-	static void TerminateProc( FProcHandle & ProcessHandle, bool KillTree = false );
-	static void TerminateProcTreeWithPredicate(
+	static CORE_API void* GetDllHandle( const TCHAR* Filename );
+	static CORE_API void FreeDllHandle( void* DllHandle );
+	static CORE_API void* GetDllExport( void* DllHandle, const TCHAR* ProcName );
+	static CORE_API void AddDllDirectory(const TCHAR* Directory);
+	static CORE_API void GetDllDirectories(TArray<FString>& OutDllDirectories);
+	static CORE_API void PushDllDirectory(const TCHAR* Directory);
+	static CORE_API void PopDllDirectory(const TCHAR* Directory);
+	static CORE_API uint32 GetCurrentProcessId();
+	static CORE_API uint32 GetCurrentCoreNumber();
+	static CORE_API void SetThreadAffinityMask( uint64 AffinityMask );
+	static CORE_API void SetThreadName( const TCHAR* ThreadName );
+	static CORE_API const TCHAR* BaseDir();
+	static CORE_API const TCHAR* UserDir();
+	static CORE_API const TCHAR* UserTempDir();
+	static CORE_API const TCHAR* UserSettingsDir();
+	static CORE_API const TCHAR* ApplicationSettingsDir();
+	static CORE_API FString GetApplicationSettingsDir(const ApplicationSettingsContext& Settings);
+	static CORE_API const TCHAR* ComputerName();
+	static CORE_API const TCHAR* UserName(bool bOnlyAlphaNumeric = true);
+	static CORE_API void SetCurrentWorkingDirectoryToBaseDir();
+	static CORE_API FString GetCurrentWorkingDirectory();
+	static CORE_API const FString ShaderWorkingDir();
+	static CORE_API const TCHAR* ExecutablePath();
+	static CORE_API const TCHAR* ExecutableName(bool bRemoveExtension = true);
+	static CORE_API FString GenerateApplicationPath( const FString& AppName, EBuildConfiguration BuildConfiguration);
+	static CORE_API const TCHAR* GetModuleExtension();
+	static CORE_API const TCHAR* GetBinariesSubdirectory();
+	static CORE_API const FString GetModulesDirectory();
+	static CORE_API bool CanLaunchURL(const TCHAR* URL);
+	static CORE_API void LaunchURL( const TCHAR* URL, const TCHAR* Parms, FString* Error );
+	static CORE_API FProcHandle CreateProc( const TCHAR* URL, const TCHAR* Parms, bool bLaunchDetached, bool bLaunchHidden, bool bLaunchReallyHidden, uint32* OutProcessID, int32 PriorityModifier, const TCHAR* OptionalWorkingDirectory, void* PipeWriteChild, void* PipeReadChild = nullptr);
+	static CORE_API FProcHandle CreateProc( const TCHAR* URL, const TCHAR* Parms, bool bLaunchDetached, bool bLaunchHidden, bool bLaunchReallyHidden, uint32* OutProcessID, int32 PriorityModifier, const TCHAR* OptionalWorkingDirectory, void* PipeWriteChild, void* PipeReadChild, void* PipeStdErrChild);
+	static CORE_API bool SetProcPriority(FProcHandle & InProcHandle, int32 PriorityModifier);
+	static CORE_API FProcHandle OpenProcess(uint32 ProcessID);
+	static CORE_API bool IsProcRunning( FProcHandle & ProcessHandle );
+	static CORE_API void WaitForProc( FProcHandle & ProcessHandle );
+	static CORE_API void CloseProc( FProcHandle & ProcessHandle );
+	static CORE_API void TerminateProc( FProcHandle & ProcessHandle, bool KillTree = false );
+	static CORE_API void TerminateProcTreeWithPredicate(
 			FProcHandle& ProcessHandle,
 			TFunctionRef<bool(uint32 ProcessId, const TCHAR* ApplicationName)> Predicate);
-	static bool GetProcReturnCode( FProcHandle & ProcHandle, int32* ReturnCode );
-	static bool GetApplicationMemoryUsage(uint32 ProcessId, SIZE_T* OutMemoryUsage);
-	static bool GetPerFrameProcessorUsage(uint32 ProcessId, float& ProcessUsageFraction, float& IdleUsageFraction);
-	static bool IsApplicationRunning( uint32 ProcessId );
-	static bool IsApplicationRunning( const TCHAR* ProcName );
-	static FString GetApplicationName( uint32 ProcessId );	
-	static bool ExecProcess(const TCHAR* URL, const TCHAR* Params, int32* OutReturnCode, FString* OutStdOut, FString* OutStdErr, const TCHAR* OptionalWorkingDirectory = NULL, bool bShouldEndWithParentProcess  = false);
-	static bool ExecElevatedProcess(const TCHAR* URL, const TCHAR* Params, int32* OutReturnCode);
-	static FProcHandle CreateElevatedProcess(const TCHAR* URL, const TCHAR* Params);
-	static bool LaunchFileInDefaultExternalApplication( const TCHAR* FileName, const TCHAR* Parms = NULL, ELaunchVerb::Type Verb = ELaunchVerb::Open, bool bPromptToOpenOnFailure = true );
-	static void ExploreFolder( const TCHAR* FilePath );
-	static bool ResolveNetworkPath( FString InUNCPath, FString& OutPath ); 
-	static void Sleep(float Seconds);
-	static void SleepNoStats(float Seconds);
-	[[noreturn]] static void SleepInfinite();
-	static void YieldThread();
+	static CORE_API bool GetProcReturnCode( FProcHandle & ProcHandle, int32* ReturnCode );
+	static CORE_API bool GetApplicationMemoryUsage(uint32 ProcessId, SIZE_T* OutMemoryUsage);
+	static CORE_API bool GetPerFrameProcessorUsage(uint32 ProcessId, float& ProcessUsageFraction, float& IdleUsageFraction);
+	static CORE_API bool IsApplicationRunning( uint32 ProcessId );
+	static CORE_API bool IsApplicationRunning( const TCHAR* ProcName );
+	static CORE_API FString GetApplicationName( uint32 ProcessId );	
+	static CORE_API bool ExecProcess(const TCHAR* URL, const TCHAR* Params, int32* OutReturnCode, FString* OutStdOut, FString* OutStdErr, const TCHAR* OptionalWorkingDirectory = NULL, bool bShouldEndWithParentProcess  = false);
+	static CORE_API bool ExecElevatedProcess(const TCHAR* URL, const TCHAR* Params, int32* OutReturnCode);
+	static CORE_API FProcHandle CreateElevatedProcess(const TCHAR* URL, const TCHAR* Params);
+	static CORE_API bool LaunchFileInDefaultExternalApplication( const TCHAR* FileName, const TCHAR* Parms = NULL, ELaunchVerb::Type Verb = ELaunchVerb::Open, bool bPromptToOpenOnFailure = true );
+	static CORE_API void ExploreFolder( const TCHAR* FilePath );
+	static CORE_API bool ResolveNetworkPath( FString InUNCPath, FString& OutPath ); 
+	static CORE_API void Sleep(float Seconds);
+	static CORE_API void SleepNoStats(float Seconds);
+	[[noreturn]] static CORE_API void SleepInfinite();
+	static CORE_API void YieldThread();
 	UE_DEPRECATED(5.0, "Please use GetSynchEventFromPool to create a new event, and ReturnSynchEventToPool to release the event.")
-	static class FEvent* CreateSynchEvent(bool bIsManualReset = false);
-	static class FRunnableThread* CreateRunnableThread();
-	static void ClosePipe( void* ReadPipe, void* WritePipe );
-	static bool CreatePipe(void*& ReadPipe, void*& WritePipe, bool bWritePipeLocal = false);
-	static FString ReadPipe( void* ReadPipe );
-	static bool ReadPipeToArray(void* ReadPipe, TArray<uint8> & Output);
-	static bool WritePipe(void* WritePipe, const FString& Message, FString* OutWritten = nullptr);
-	static bool WritePipe(void* WritePipe, const uint8* Data, const int32 DataLength, int32* OutDataLength = nullptr);
-	static FSemaphore* NewInterprocessSynchObject(const FString& Name, bool bCreate, uint32 MaxLocks = 1);
-	static FSemaphore* NewInterprocessSynchObject(const TCHAR* Name, bool bCreate, uint32 MaxLocks = 1);
-	static bool DeleteInterprocessSynchObject(FSemaphore * Object);
-	static bool Daemonize();
-	static void SetupGameThread();
-	static void SetupAudioThread();
-	static void TeardownAudioThread();
-	static bool IsFirstInstance();
+	static CORE_API class FEvent* CreateSynchEvent(bool bIsManualReset = false);
+	static CORE_API class FRunnableThread* CreateRunnableThread();
+	static CORE_API void ClosePipe( void* ReadPipe, void* WritePipe );
+	static CORE_API bool CreatePipe(void*& ReadPipe, void*& WritePipe, bool bWritePipeLocal = false);
+	static CORE_API FString ReadPipe( void* ReadPipe );
+	static CORE_API bool ReadPipeToArray(void* ReadPipe, TArray<uint8> & Output);
+	static CORE_API bool WritePipe(void* WritePipe, const FString& Message, FString* OutWritten = nullptr);
+	static CORE_API bool WritePipe(void* WritePipe, const uint8* Data, const int32 DataLength, int32* OutDataLength = nullptr);
+	static CORE_API FSemaphore* NewInterprocessSynchObject(const FString& Name, bool bCreate, uint32 MaxLocks = 1);
+	static CORE_API FSemaphore* NewInterprocessSynchObject(const TCHAR* Name, bool bCreate, uint32 MaxLocks = 1);
+	static CORE_API bool DeleteInterprocessSynchObject(FSemaphore * Object);
+	static CORE_API bool Daemonize();
+	static CORE_API void SetupGameThread();
+	static CORE_API void SetupAudioThread();
+	static CORE_API void TeardownAudioThread();
+	static CORE_API bool IsFirstInstance();
 
 	/**
 	 * @brief Releases locks that we held for IsFirstInstance check
 	 */
-	static void CeaseBeingFirstInstance();
+	static CORE_API void CeaseBeingFirstInstance();
 
 protected:
 
@@ -214,7 +214,7 @@ protected:
 	 * @param InPipes The pipes to read from.
 	 * @param PipeCount The number of pipes.
 	 */
-	static void ReadFromPipes(FString* OutStrings[], Windows::HANDLE InPipes[], int32 PipeCount);
+	static CORE_API void ReadFromPipes(FString* OutStrings[], Windows::HANDLE InPipes[], int32 PipeCount);
 
 private:
 
diff --git a/Engine/Source/Runtime/Core/Public/Windows/WindowsPlatformStackWalk.h b/Engine/Source/Runtime/Core/Public/Windows/WindowsPlatformStackWalk.h
--- a/Engine/Source/Runtime/Core/Public/Windows/WindowsPlatformStackWalk.h
+++ b/Engine/Source/Runtime/Core/Public/Windows/WindowsPlatformStackWalk.h
@@ -8,26 +8,26 @@
 /**
  * Windows implementation of the stack walking.
  **/
-struct CORE_API FWindowsPlatformStackWalk
+struct FWindowsPlatformStackWalk
 	: public FMicrosoftPlatformStackWalk
 {
-	static bool InitStackWalking();
-	static bool InitStackWalkingForProcess(const FProcHandle& Process);
+	static CORE_API bool InitStackWalking();
+	static CORE_API bool InitStackWalkingForProcess(const FProcHandle& Process);
 	
-	static TArray<FProgramCounterSymbolInfo> GetStack(int32 IgnoreCount, int32 MaxDepth = 100, void* Context = nullptr);
+	static CORE_API TArray<FProgramCounterSymbolInfo> GetStack(int32 IgnoreCount, int32 MaxDepth = 100, void* Context = nullptr);
 
-	static void ProgramCounterToSymbolInfo( uint64 ProgramCounter, FProgramCounterSymbolInfo& out_SymbolInfo );
-	static void ProgramCounterToSymbolInfoEx( uint64 ProgramCounter, FProgramCounterSymbolInfoEx& out_SymbolInfo );
-	FORCENOINLINE static uint32 CaptureStackBackTrace( uint64* BackTrace, uint32 MaxDepth, void* Context = nullptr );  // FORCENOINLINE so it can be counted during StackTrace
-	static uint32 CaptureThreadStackBackTrace( uint64 ThreadId, uint64* BackTrace, uint32 MaxDepth, void* Context = nullptr);
-	FORCENOINLINE static void StackWalkAndDump( ANSICHAR* HumanReadableString, SIZE_T HumanReadableStringSize, int32 IgnoreCount, void* Context = nullptr );  // FORCENOINLINE so it can be counted during StackTrace
-	static void StackWalkAndDump( ANSICHAR* HumanReadableString, SIZE_T HumanReadableStringSize, void* ProgramCounter, void* Context = nullptr );
-	static void ThreadStackWalkAndDump(ANSICHAR* HumanReadableString, SIZE_T HumanReadableStringSize, int32 IgnoreCount, uint32 ThreadId);
+	static CORE_API void ProgramCounterToSymbolInfo( uint64 ProgramCounter, FProgramCounterSymbolInfo& out_SymbolInfo );
+	static CORE_API void ProgramCounterToSymbolInfoEx( uint64 ProgramCounter, FProgramCounterSymbolInfoEx& out_SymbolInfo );
+	CORE_API FORCENOINLINE static uint32 CaptureStackBackTrace( uint64* BackTrace, uint32 MaxDepth, void* Context = nullptr );  // FORCENOINLINE so it can be counted during StackTrace
+	static CORE_API uint32 CaptureThreadStackBackTrace( uint64 ThreadId, uint64* BackTrace, uint32 MaxDepth, void* Context = nullptr);
+	CORE_API FORCENOINLINE static void StackWalkAndDump( ANSICHAR* HumanReadableString, SIZE_T HumanReadableStringSize, int32 IgnoreCount, void* Context = nullptr );  // FORCENOINLINE so it can be counted during StackTrace
+	static CORE_API void StackWalkAndDump( ANSICHAR* HumanReadableString, SIZE_T HumanReadableStringSize, void* ProgramCounter, void* Context = nullptr );
+	static CORE_API void ThreadStackWalkAndDump(ANSICHAR* HumanReadableString, SIZE_T HumanReadableStringSize, int32 IgnoreCount, uint32 ThreadId);
 
-	static int32 GetProcessModuleCount();
-	static int32 GetProcessModuleSignatures(FStackWalkModuleInfo *ModuleSignatures, const int32 ModuleSignaturesSize);
+	static CORE_API int32 GetProcessModuleCount();
+	static CORE_API int32 GetProcessModuleSignatures(FStackWalkModuleInfo *ModuleSignatures, const int32 ModuleSignaturesSize);
 
-	static void RegisterOnModulesChanged();
+	static CORE_API void RegisterOnModulesChanged();
 
 	/**
 	 * Upload localy built symbols to network symbol storage.
@@ -37,15 +37,15 @@ struct CORE_API FWindowsPlatformStackWalk
 	 *   In this case all game .dll-files are compiled locally.
 	 *   For post-mortem debug programmers need .dll and .pdb files from designers.
 	 */
-	static bool UploadLocalSymbols();
+	static CORE_API bool UploadLocalSymbols();
 
 	/**
 	 * Get downstream storage with downloaded from remote symbol storage files.
 	 */
-	static FString GetDownstreamStorage();
+	static CORE_API FString GetDownstreamStorage();
 
-	static void* MakeThreadContextWrapper(void* Context, void* ThreadHandle);
-	static void ReleaseThreadContextWrapper(void* ThreadContext);
+	static CORE_API void* MakeThreadContextWrapper(void* Context, void* ThreadHandle);
+	static CORE_API void ReleaseThreadContextWrapper(void* ThreadContext);
 
 	/**
 	 * Returns the source file pathname, line and column where the specified function is defined.
@@ -61,10 +61,10 @@ struct CORE_API FWindowsPlatformStackWalk
 	 * @param OutColumnNumber The offset on the line at which the function is defined in the source file.
 	 * @return True if the the function location is found, false otherwise.
 	 */
-	static bool GetFunctionDefinitionLocation(const FString& FunctionSymbolName, const FString& FunctionModuleName, FString& OutPathname, uint32& OutLineNumber, uint32& OutColumnNumber);
+	static CORE_API bool GetFunctionDefinitionLocation(const FString& FunctionSymbolName, const FString& FunctionModuleName, FString& OutPathname, uint32& OutLineNumber, uint32& OutColumnNumber);
 
 protected:
-	static void CaptureStackTraceByProcess(uint64* OutBacktrace, uint32 MaxDepth, void* InContext, void* InThreadHandle, uint32* OutDepth, bool bExternalProcess);
+	static CORE_API void CaptureStackTraceByProcess(uint64* OutBacktrace, uint32 MaxDepth, void* InContext, void* InThreadHandle, uint32* OutDepth, bool bExternalProcess);
 private:
 	static bool InitStackWalkingInternal(void* Process, bool bForceReinitOnProcessMismatch);
 };
diff --git a/Engine/Source/Runtime/Core/Public/Windows/WindowsPlatformTLS.h b/Engine/Source/Runtime/Core/Public/Windows/WindowsPlatformTLS.h
--- a/Engine/Source/Runtime/Core/Public/Windows/WindowsPlatformTLS.h
+++ b/Engine/Source/Runtime/Core/Public/Windows/WindowsPlatformTLS.h
@@ -10,7 +10,7 @@
 /**
  * Windows implementation of the TLS OS functions.
  */
-struct CORE_API FWindowsPlatformTLS
+struct FWindowsPlatformTLS
 	: public FGenericPlatformTLS
 {
 	/**
diff --git a/Engine/Source/Runtime/Core/Public/Windows/WindowsPlatformTime.h b/Engine/Source/Runtime/Core/Public/Windows/WindowsPlatformTime.h
--- a/Engine/Source/Runtime/Core/Public/Windows/WindowsPlatformTime.h
+++ b/Engine/Source/Runtime/Core/Public/Windows/WindowsPlatformTime.h
@@ -15,10 +15,10 @@
  *
  * https://udn.epicgames.com/lists/showpost.php?id=46794&list=unprog3
  */
-struct CORE_API FWindowsPlatformTime
+struct FWindowsPlatformTime
 	: public FGenericPlatformTime
 {
-	static double InitTiming();
+	static CORE_API double InitTiming();
 
 	static FORCEINLINE double Seconds()
 	{
@@ -44,20 +44,20 @@ struct CORE_API FWindowsPlatformTime
 	}
 
 
-	static void SystemTime( int32& Year, int32& Month, int32& DayOfWeek, int32& Day, int32& Hour, int32& Min, int32& Sec, int32& MSec );
-	static void UtcTime( int32& Year, int32& Month, int32& DayOfWeek, int32& Day, int32& Hour, int32& Min, int32& Sec, int32& MSec );
+	static CORE_API void SystemTime( int32& Year, int32& Month, int32& DayOfWeek, int32& Day, int32& Hour, int32& Min, int32& Sec, int32& MSec );
+	static CORE_API void UtcTime( int32& Year, int32& Month, int32& DayOfWeek, int32& Day, int32& Hour, int32& Min, int32& Sec, int32& MSec );
 
-	static bool UpdateCPUTime( float DeltaTime );
-	static bool UpdateThreadCPUTime(float = 0.0);
-	static void AutoUpdateGameThreadCPUTime(double UpdateInterval);
-	static FCPUTime GetCPUTime();
-	static FCPUTime GetThreadCPUTime();
-	static double GetLastIntervalThreadCPUTimeInSeconds();
+	static CORE_API bool UpdateCPUTime( float DeltaTime );
+	static CORE_API bool UpdateThreadCPUTime(float = 0.0);
+	static CORE_API void AutoUpdateGameThreadCPUTime(double UpdateInterval);
+	static CORE_API FCPUTime GetCPUTime();
+	static CORE_API FCPUTime GetThreadCPUTime();
+	static CORE_API double GetLastIntervalThreadCPUTimeInSeconds();
 
 protected:
 
 	/** Percentage CPU utilization for the last interval relative to one core. */
-	static float CPUTimePctRelative;
+	static CORE_API float CPUTimePctRelative;
 };
 
 
diff --git a/Engine/Source/Runtime/CoreUObject/Public/Serialization/TestUndeclaredScriptStructObjectReferences.h b/Engine/Source/Runtime/CoreUObject/Public/Serialization/TestUndeclaredScriptStructObjectReferences.h
--- a/Engine/Source/Runtime/CoreUObject/Public/Serialization/TestUndeclaredScriptStructObjectReferences.h
+++ b/Engine/Source/Runtime/CoreUObject/Public/Serialization/TestUndeclaredScriptStructObjectReferences.h
@@ -6,7 +6,7 @@
 #include "UObject/SoftObjectPtr.h"
 #include "UObject/WeakObjectPtr.h"
 
-class CORE_API FArchive;
+class FArchive;
  
 // Helper struct to test if struct serializer object reference declaration tests work properly
 struct FTestUndeclaredScriptStructObjectReferencesTest
diff --git a/Engine/Source/Runtime/CoreUObject/Public/UObject/TopLevelAssetPath.h b/Engine/Source/Runtime/CoreUObject/Public/UObject/TopLevelAssetPath.h
--- a/Engine/Source/Runtime/CoreUObject/Public/UObject/TopLevelAssetPath.h
+++ b/Engine/Source/Runtime/CoreUObject/Public/UObject/TopLevelAssetPath.h
@@ -17,7 +17,7 @@
 #include <type_traits>
 
 struct FPropertyTag;
-class CORE_API FString;
+class FString;
 class UObject;
 
 /* 
