Add listener to TraceLog messages

Implement a listener for TraceLog messages and log any events.

diff --git a/Engine/Source/Runtime/Core/Private/ProfilingDebugging/TraceAuxiliary.cpp b/Engine/Source/Runtime/Core/Private/ProfilingDebugging/TraceAuxiliary.cpp
--- a/Engine/Source/Runtime/Core/Private/ProfilingDebugging/TraceAuxiliary.cpp
+++ b/Engine/Source/Runtime/Core/Private/ProfilingDebugging/TraceAuxiliary.cpp
@@ -74,6 +74,16 @@ DECLARE_MEMORY_STAT(TEXT("Important event cache used"), STAT_TraceCacheUsed, STA
 DECLARE_MEMORY_STAT(TEXT("Important event cache waste"), STAT_TraceCacheWaste, STATGROUP_Trace);
 DECLARE_MEMORY_STAT(TEXT("Sent"), STAT_TraceSent, STATGROUP_Trace);
 
+////////////////////////////////////////////////////////////////////////////////
+enum class ETraceConnectType
+{
+	Network,
+	File
+};
+
+////////////////////////////////////////////////////////////////////////////////
+DEFINE_LOG_CATEGORY_STATIC(LogTrace, Log, All);
+
 ////////////////////////////////////////////////////////////////////////////////
 class FTraceAuxiliaryImpl
 {
@@ -293,7 +303,7 @@ bool FTraceAuxiliaryImpl::Connect(FTraceAuxiliary::EConnectionType Type, const T
 			}
 			else
 			{
-				UE_LOG_REF(LogCategory, Error, TEXT("Trace failed to connect (trace server: %s)!"), Parameter ? Parameter : TEXT(""));
+				UE_LOG_REF(LogCategory, Error, TEXT("Trace failed to connect (trace host: %s)!"), Parameter ? Parameter : TEXT(""));
 			}
 
 		}
@@ -747,6 +757,21 @@ void OnConnectionCallback()
 	FTraceAuxiliary::OnConnection.Broadcast();
 }
 
+////////////////////////////////////////////////////////////////////////////////
+void OnMessageCallback(const UE::Trace::FMessageEvent& Message)
+{
+	const auto TypeStr = StringCast<TCHAR>(Message.TypeStr);
+	if (Message.Description != nullptr)
+	{
+		const auto Description = StringCast<TCHAR>(Message.Description);
+		UE_LOG(LogTrace, Error, TEXT("%s %s"), TypeStr.Get(), Description.Get());
+	}
+	else
+	{
+		UE_LOG(LogTrace, Error, TEXT("%s"), TypeStr.Get());
+	}
+}
+
 ////////////////////////////////////////////////////////////////////////////////
 static void SetupInitFromConfig(UE::Trace::FInitializeDesc& OutDesc)
 {
@@ -1201,7 +1226,7 @@ static void LaunchUnrealTraceInternal(const TCHAR* CommandLine)
 #if !USING_THREAD_SANITISER
 	if (GUnrealTraceLaunched.load(std::memory_order_relaxed))
 	{
-		UE_LOG(LogCore, Log, TEXT("UnrealTraceServer: Trace store already started"));
+		UE_LOG(LogTrace, Log, TEXT("UnrealTraceServer: Trace store already started"));
 		return;
 	}
 
@@ -1216,7 +1241,7 @@ static void LaunchUnrealTraceInternal(const TCHAR* CommandLine)
 
 	if (access(*BinPath, F_OK) < 0)
 	{
-		UE_LOG(LogCore, Display, TEXT("UnrealTraceServer: Binary not found (%s)"), ANSI_TO_TCHAR(*BinPath));
+		UE_LOG(LogTrace, Display, TEXT("UnrealTraceServer: Binary not found (%s)"), ANSI_TO_TCHAR(*BinPath));
 		return;
 	}
 
@@ -1227,7 +1252,7 @@ static void LaunchUnrealTraceInternal(const TCHAR* CommandLine)
 	pid_t UtsPid = fork();
 	if (UtsPid < 0)
 	{
-		UE_LOG(LogCore, Display, TEXT("UnrealTraceServer: Unable to fork (errno: %d)"), errno);
+		UE_LOG(LogTrace, Display, TEXT("UnrealTraceServer: Unable to fork (errno: %d)"), errno);
 		return;
 	}
 	else if (UtsPid == 0)
@@ -1244,7 +1269,7 @@ static void LaunchUnrealTraceInternal(const TCHAR* CommandLine)
 		int32 WaitRet = waitpid(UtsPid, &WaitStatus, 0);
 		if (WaitRet < 0)
 		{
-			UE_LOG(LogCore, Display, TEXT("UnrealTraceServer: waitpid() error; (errno: %d)"), errno);
+			UE_LOG(LogTrace, Display, TEXT("UnrealTraceServer: waitpid() error; (errno: %d)"), errno);
 			return;
 		}
 	}
@@ -1253,11 +1278,11 @@ static void LaunchUnrealTraceInternal(const TCHAR* CommandLine)
 	int32 UtsRet = WEXITSTATUS(WaitStatus);
 	if (UtsRet)
 	{
-		UE_LOG(LogCore, Display, TEXT("UnrealTraceServer: Trace store returned an error (0x%08x)"), UtsRet);
+		UE_LOG(LogTrace, Display, TEXT("UnrealTraceServer: Trace store returned an error (0x%08x)"), UtsRet);
 	}
 	else
 	{
-		UE_LOG(LogCore, Log, TEXT("UnrealTraceServer: Trace store launch successful"));
+		UE_LOG(LogTrace, Log, TEXT("UnrealTraceServer: Trace store launch successful"));
 		GUnrealTraceLaunched.fetch_add(1, std::memory_order_relaxed);
 	}
 #endif // #if !USING_THREAD_SANITISER
@@ -1328,7 +1353,7 @@ static bool StartFromCommandlineArguments(const TCHAR* CommandLine, bool& bOutSt
 		Type = FTraceAuxiliary::EConnectionType::File;
 		if (Parameter.IsEmpty())
 		{
-			UE_LOG(LogCore, Warning, TEXT("Empty parameter to 'tracefile' argument. Using default filename."));
+			UE_LOG(LogTrace, Warning, TEXT("Empty parameter to 'tracefile' argument. Using default filename."));
 			Target = nullptr;
 		}
 		else
@@ -1457,7 +1482,7 @@ bool FTraceAuxiliary::Resume()
 bool FTraceAuxiliary::WriteSnapshot(const TCHAR* InFilePath)
 {
 #if UE_TRACE_ENABLED
-	return GTraceAuxiliary.WriteSnapshot(InFilePath, LogCore);
+	return GTraceAuxiliary.WriteSnapshot(InFilePath, LogTrace);
 #else
 	return true;
 #endif
@@ -1467,7 +1492,7 @@ bool FTraceAuxiliary::WriteSnapshot(const TCHAR* InFilePath)
 bool FTraceAuxiliary::SendSnapshot(const TCHAR* Host, uint32 Port)
 {
 #if UE_TRACE_ENABLED
-	return GTraceAuxiliary.SendSnapshot(Host, Port, LogCore);
+	return GTraceAuxiliary.SendSnapshot(Host, Port, LogTrace);
 #else
 	return true;
 #endif
@@ -1479,7 +1504,7 @@ static void TraceAuxiliaryAddPostForkCallback(const TCHAR* CommandLine)
 {
 	if (GTraceAutoStart)
 	{
-		UE_LOG(LogCore, Display, TEXT("Trace not started in parent because forking is expected. Use -NoFakeForking to trace parent."));
+		UE_LOG(LogTrace, Display, TEXT("Trace not started in parent because forking is expected. Use -NoFakeForking to trace parent."));
 	}
 
 	checkf(!GOnPostForkHandle.IsValid(), TEXT("TraceAuxiliaryAddPostForkCallback should only be called once."));
@@ -1516,14 +1541,17 @@ void FTraceAuxiliary::Initialize(const TCHAR* CommandLine)
 #endif
 
 #if UE_TRACE_ENABLED
-	UE_LOG(LogCore, Log, TEXT("Initializing trace..."));
+	UE_LOG(LogTrace, Log, TEXT("Initializing trace..."));
+
+	// Setup message callback so we get feedback from TraceLog
+	UE::Trace::SetMessageCallback(&OnMessageCallback);
 
 	FParse::Bool(CommandLine, TEXT("-traceautostart="), GTraceAutoStart);
-	UE_LOG(LogCore, Verbose, TEXT("Trace auto start = %d."), GTraceAutoStart);
+	UE_LOG(LogTrace, Verbose, TEXT("Trace auto start = %d."), GTraceAutoStart);
 
 	if (GTraceAuxiliary.IsParentProcessAndPreFork())
 	{
-		UE_LOG(LogCore, Log, TEXT("Trace initialization skipped for parent process (pre fork)."));
+		UE_LOG(LogTrace, Log, TEXT("Trace initialization skipped for parent process (pre fork)."));
 
 		GTraceAuxiliary.DisableChannels(nullptr);
 
@@ -1656,7 +1684,7 @@ void FTraceAuxiliary::Initialize(const TCHAR* CommandLine)
 
 	UE::Trace::ThreadRegister(TEXT("GameThread"), FPlatformTLS::GetCurrentThreadId(), -1);
 
-	UE_LOG(LogCore, Log, TEXT("Finished trace initialization."));
+	UE_LOG(LogTrace, Log, TEXT("Finished trace initialization."));
 #endif //UE_TRACE_ENABLED
 }
 
@@ -1776,7 +1804,7 @@ void FTraceAuxiliary::TryAutoConnect()
 		HANDLE KnownEvent = ::OpenEvent(EVENT_ALL_ACCESS, false, TEXT("Local\\UnrealInsightsBrowser"));
 		if (KnownEvent != nullptr)
 		{
-			UE_LOG(LogCore, Display, TEXT("Unreal Insights instance detected, auto-connecting to local trace server..."));
+			UE_LOG(LogTrace, Display, TEXT("Unreal Insights instance detected, auto-connecting to local trace server..."));
 			Start(EConnectionType::Network, TEXT("127.0.0.1"), GTraceAuxiliary.HasCommandlineChannels() ? nullptr : TEXT("default"), nullptr);
 			::CloseHandle(KnownEvent);
 		}
