Re enabling for grammared parse for CLI filtering.

Fixed issue where parsing would stop if an un-key'd was encountered.
Added comparison test between new method and old.

diff --git a/Engine/Source/Runtime/Core/Private/Misc/CommandLine.cpp b/Engine/Source/Runtime/Core/Private/Misc/CommandLine.cpp
--- a/Engine/Source/Runtime/Core/Private/Misc/CommandLine.cpp
+++ b/Engine/Source/Runtime/Core/Private/Misc/CommandLine.cpp
@@ -108,7 +108,7 @@ bool FCommandLine::IsCommandLineLoggingFiltered()
 #endif
 }
 
-bool FCommandLine::FilterMove(TCHAR* OutLine, int32 MaxLen, const TCHAR* InLine, const TArrayView<FString>& AllowedList)
+bool FCommandLine::FilterCLIUsingGrammarBasedParser(TCHAR* OutLine, int32 MaxLen, const TCHAR* InLine, const TArrayView<FString>& AllowedList)
 {
 	if (MaxLen == 0)
 	{
@@ -225,17 +225,13 @@ void FCommandLine::ApplyCommandLineAllowList()
 		FCommandLine::Parse(FilterForLoggingList, FilterArgsForLogging, Ignored);
 	}
 	// Process the original command line
-	TArray<FString> OriginalList = FilterCommandLine(OriginalCmdLine);
-	BuildCommandLineAllowList(OriginalCmdLine, UE_ARRAY_COUNT(OriginalCmdLine), OriginalList);
+	FilterCLIUsingGrammarBasedParser(OriginalCmdLine, UE_ARRAY_COUNT(OriginalCmdLine), OriginalCmdLine, ApprovedArgs);
 	// Process the current command line
-	TArray<FString> CmdList = FilterCommandLine(CmdLine);
-	BuildCommandLineAllowList(CmdLine, UE_ARRAY_COUNT(CmdLine), CmdList);
+	FilterCLIUsingGrammarBasedParser(CmdLine, UE_ARRAY_COUNT(CmdLine), CmdLine, ApprovedArgs);
 	// Process the command line for logging purposes
-	TArray<FString> LoggingCmdList = FilterCommandLineForLogging(LoggingCmdLine);
-	BuildCommandLineAllowList(LoggingCmdLine, UE_ARRAY_COUNT(LoggingCmdLine), LoggingCmdList);
+	FilterCLIUsingGrammarBasedParser(LoggingCmdLine, UE_ARRAY_COUNT(LoggingCmdLine), LoggingCmdLine, FilterArgsForLogging);
 	// Process the original command line for logging purposes
-	TArray<FString> LoggingOriginalCmdList = FilterCommandLineForLogging(LoggingOriginalCmdLine);
-	BuildCommandLineAllowList(LoggingOriginalCmdLine, UE_ARRAY_COUNT(LoggingOriginalCmdLine), LoggingOriginalCmdList);
+	FilterCLIUsingGrammarBasedParser(LoggingOriginalCmdLine, UE_ARRAY_COUNT(LoggingOriginalCmdLine), LoggingOriginalCmdLine, FilterArgsForLogging);
 }
 
 TArray<FString> FCommandLine::FilterCommandLine(TCHAR* CommandLine)
diff --git a/Engine/Source/Runtime/Core/Private/Misc/Parse.cpp b/Engine/Source/Runtime/Core/Private/Misc/Parse.cpp
--- a/Engine/Source/Runtime/Core/Private/Misc/Parse.cpp
+++ b/Engine/Source/Runtime/Core/Private/Misc/Parse.cpp
@@ -1514,7 +1514,20 @@ struct FGrammarBasedParser
 		}
 		else
 		{
-			OnCmd(Item, FStringView{});
+			if (Item.Len())
+			{
+				OnCmd(Item, FStringView{});
+			}
+			else
+			{
+				// If there is no Key then we will try consuming a value, if we can parse one
+				FStringView ItemValue = ParseValue();
+				if (HasError())
+				{
+					return {};
+				}
+				OnCmd(FStringView{}, ItemValue);
+			}
 		}
 		return FStringView{ Start,  UE_PTRDIFF_TO_INT32(Cursor - Start) };
 	}
diff --git a/Engine/Source/Runtime/Core/Public/Misc/CommandLine.h b/Engine/Source/Runtime/Core/Public/Misc/CommandLine.h
--- a/Engine/Source/Runtime/Core/Public/Misc/CommandLine.h
+++ b/Engine/Source/Runtime/Core/Public/Misc/CommandLine.h
@@ -110,7 +110,8 @@ struct FCommandLine
 	CORE_API static FString BuildFromArgV(const ANSICHAR* Prefix, int32 ArgC, ANSICHAR* ArgV[], const ANSICHAR* Suffix);
 
 	/**
-	* Filter Move parses CLI style arguments in the InLine, filters for commands or keys specified in the AllowedList
+	* FilterCLIUsingGrammarBasedParser parses CLI style arguments.
+	* filters for commands or keys specified in the AllowedList
 	* and writes the to the OutLine.
 	* OutLine and InLine may be point to the same buffer.
 	* 
@@ -121,7 +122,7 @@ struct FCommandLine
 	* 
 	* Returns true if Outline was large enough to hold the filtered string, Otherwise returns false.
 	*/
-	static bool FilterMove(TCHAR* OutLine, int32 MaxLen, const TCHAR* InLine, const TArrayView<FString>& AllowedList);
+	static bool FilterCLIUsingGrammarBasedParser(TCHAR* OutLine, int32 MaxLen, const TCHAR* InLine, const TArrayView<FString>& AllowedList);
 private:
 #if UE_COMMAND_LINE_USES_ALLOW_LIST
 	/** Filters both the original and current command line list for approved only args */
diff --git a/Engine/Source/Runtime/Core/Tests/Misc/CommandLineTests.cpp b/Engine/Source/Runtime/Core/Tests/Misc/CommandLineTests.cpp
--- a/Engine/Source/Runtime/Core/Tests/Misc/CommandLineTests.cpp
+++ b/Engine/Source/Runtime/Core/Tests/Misc/CommandLineTests.cpp
@@ -10,7 +10,7 @@
 
 #include <catch2/generators/catch_generators.hpp>
 
-TEST_CASE("CommandLine::Filtering::FilterMove", "[Smoke]")
+TEST_CASE("CommandLine::Filtering::FilterCLIInplace", "[Smoke]")
 {
 	TArray<FString> AllowedList{ { "cmd_a", "-cmd_b" } };
 
@@ -31,7 +31,7 @@ TEST_CASE("CommandLine::Filtering::FilterMove", "[Smoke]")
 		}));
 		TCHAR Result[256]{};
 
-		CHECK(FCommandLine::FilterMove(Result, UE_ARRAY_COUNT(Result), Input, AllowedList));
+		CHECK(FCommandLine::FilterCLIUsingGrammarBasedParser(Result, UE_ARRAY_COUNT(Result), Input, AllowedList));
 		CHECK(FStringView{ Result } == FStringView{ Expected });
 	}
 
@@ -45,7 +45,7 @@ TEST_CASE("CommandLine::Filtering::FilterMove", "[Smoke]")
 			}));
 		TCHAR Result[256]{};
 
-		CHECK(FCommandLine::FilterMove(Result, UE_ARRAY_COUNT(Result), Input, AllowedList));
+		CHECK(FCommandLine::FilterCLIUsingGrammarBasedParser(Result, UE_ARRAY_COUNT(Result), Input, AllowedList));
 		CHECK(FStringView{ Result } == FStringView{ Expected });
 	}
 
@@ -71,7 +71,7 @@ TEST_CASE("CommandLine::Filtering::FilterMove", "[Smoke]")
 		TCHAR SourceAndResult[256]{};
 		
 		FCString::Strcpy(SourceAndResult, Input);
-		CHECK(FCommandLine::FilterMove(SourceAndResult, UE_ARRAY_COUNT(SourceAndResult), SourceAndResult, AllowedList));
+		CHECK(FCommandLine::FilterCLIUsingGrammarBasedParser(SourceAndResult, UE_ARRAY_COUNT(SourceAndResult), SourceAndResult, AllowedList));
 		CHECK(FStringView{ SourceAndResult } == FStringView{ Expected });
 	}
 
@@ -80,7 +80,7 @@ TEST_CASE("CommandLine::Filtering::FilterMove", "[Smoke]")
 		const TCHAR* Input = TEXT("-cmd_a --cmd_b");
 		TCHAR Result[256] = TEXT("Not Empty");
 
-		CHECK(FCommandLine::FilterMove(Result, UE_ARRAY_COUNT(Result), Input, {}));
+		CHECK(FCommandLine::FilterCLIUsingGrammarBasedParser(Result, UE_ARRAY_COUNT(Result), Input, {}));
 		CHECK(Result[0] == TCHAR(0));
 	}
 
@@ -89,7 +89,7 @@ TEST_CASE("CommandLine::Filtering::FilterMove", "[Smoke]")
 		const TCHAR* Input = TEXT("-cmd_a --cmd_b");
 		TCHAR Result[5]{};
 
-		CHECK(false == FCommandLine::FilterMove(Result, UE_ARRAY_COUNT(Result), Input, AllowedList));
+		CHECK(false == FCommandLine::FilterCLIUsingGrammarBasedParser(Result, UE_ARRAY_COUNT(Result), Input, AllowedList));
 		CHECK(Result[0] == TCHAR(0));
 	}
 
@@ -126,10 +126,105 @@ TEST_CASE("CommandLine::Filtering::FilterMove", "[Smoke]")
 		TCHAR SourceAndResult[256]{};
 
 		FCString::Strcpy(SourceAndResult, Input);
-		CHECK(FCommandLine::FilterMove(SourceAndResult, UE_ARRAY_COUNT(SourceAndResult), SourceAndResult, ApprovedArgs));
+		CHECK(FCommandLine::FilterCLIUsingGrammarBasedParser(SourceAndResult, UE_ARRAY_COUNT(SourceAndResult), SourceAndResult, ApprovedArgs));
 		CHECK(FStringView{ SourceAndResult } == FStringView{ Expected });
 	}
+}
+
+
+// This code was lifted from CommandLine.cpp
+// SO we can prove the new Grammar based parse does an equivalent job.
+// It can be removed along with the version in CommandLine.cpp once we are happy with FilterCLIUsingGrammarBasedParser 
+namespace OldFilterMethod
+{
+	TArray<FString> ApprovedArgs;
+
+	TArray<FString> FilterCommandLine(const TCHAR* CommandLine)
+	{
+		TArray<FString> Ignored;
+		TArray<FString> ParsedList;
+		// Parse the command line list
+		FCommandLine::Parse(CommandLine, ParsedList, Ignored);
+		// Remove any that are not in our approved list
+		for (int32 Index = 0; Index < ParsedList.Num(); Index++)
+		{
+			bool bFound = false;
+			for (auto ApprovedArg : ApprovedArgs)
+			{
+				if (ParsedList[Index].StartsWith(ApprovedArg))
+				{
+					bFound = true;
+					break;
+				}
+			}
+			if (!bFound)
+			{
+				ParsedList.RemoveAt(Index);
+				Index--;
+		}
+	}
+		return ParsedList;
+}
+
+	void BuildCommandLineAllowList(TCHAR* CommandLine, uint32 ArrayCount, const TArray<FString>& FilteredArgs)
+	{
+		check(ArrayCount > 0);
+		// Zero the whole string
+		FMemory::Memzero(CommandLine, sizeof(TCHAR) * ArrayCount);
+
+		uint32 StartIndex = 0;
+		for (auto Arg : FilteredArgs)
+		{
+			if ((StartIndex + Arg.Len() + 2) < ArrayCount)
+			{
+				if (StartIndex != 0)
+				{
+					CommandLine[StartIndex++] = TEXT(' ');
+				}
+				CommandLine[StartIndex++] = TEXT('-');
+				FCString::Strncpy(&CommandLine[StartIndex], *Arg, ArrayCount - StartIndex);
+				StartIndex += Arg.Len();
+			}
+		}
+	}
 
+	void FilterCLI(TCHAR* OutLine, int32 MaxLen, const TCHAR* InLine, const TArrayView<FString>& AllowedList)
+	{
+		ApprovedArgs = AllowedList;
+		TArray<FString> OriginalList = FilterCommandLine(InLine);
+		BuildCommandLineAllowList(OutLine, MaxLen, OriginalList);
+	}
+}
+
+
+TEST_CASE("CommandLine::Filtering::FilterCLI_Comparison", "[Smoke]")
+{
+	SECTION("Sudo real example")
+	{
+		TArray<FString> ApprovedArgs;
+
+		{
+			const TCHAR* OverrideList = TEXT("-cmd_a -cmd_b -cmd_c");
+
+			TArray<FString> Ignored;
+			FCommandLine::Parse(OverrideList, ApprovedArgs, Ignored);
+		}
+
+		TCHAR* Original = TEXT("\"un keyed value should be filtered\" -cmd_a=\"c:/some/path\" -to_be_filtered=be_gone -cmd_b -not_allowed -cmd_c=7890sfd9wjd8jf984mx");
+
+		TCHAR NewMethod[1024];
+		TCHAR OldMethod[1024];
+
+		FCString::Strcpy(NewMethod, UE_ARRAY_COUNT(NewMethod), Original);
+		FCString::Strcpy(OldMethod, UE_ARRAY_COUNT(NewMethod), Original);
+
+		FCommandLine::FilterCLIUsingGrammarBasedParser(NewMethod, UE_ARRAY_COUNT(NewMethod), NewMethod, ApprovedArgs);
+		OldFilterMethod::FilterCLI(OldMethod, UE_ARRAY_COUNT(OldMethod), OldMethod, ApprovedArgs);
+
+		CHECK(NewMethod[0] != TCHAR('\0')); // not empty
+		CHECK(FStringView{ NewMethod } == FStringView{ OldMethod }); // same result both methods
+		CHECK(FStringView{ NewMethod } == FStringView{ TEXT("-cmd_a=\"c:/some/path\" -cmd_b -cmd_c=7890sfd9wjd8jf984mx") }); // this should be the result
+	}
 }
 #endif
 
