Added const to test code to fix build issue not found using preflight

diff --git a/Engine/Source/Runtime/Core/Tests/Misc/CommandLineTests.cpp b/Engine/Source/Runtime/Core/Tests/Misc/CommandLineTests.cpp
--- a/Engine/Source/Runtime/Core/Tests/Misc/CommandLineTests.cpp
+++ b/Engine/Source/Runtime/Core/Tests/Misc/CommandLineTests.cpp
@@ -210,7 +210,7 @@ TEST_CASE("CommandLine::Filtering::FilterCLI_Comparison", "[Smoke]")
 			FCommandLine::Parse(OverrideList, ApprovedArgs, Ignored);
 		}
 
-		TCHAR* Original = TEXT("\"un keyed value should be filtered\" -cmd_a=\"c:/some/path\" -to_be_filtered=be_gone -cmd_b -not_allowed -cmd_c=7890sfd9wjd8jf984mx");
+		const TCHAR* Original = TEXT("\"un keyed value should be filtered\" -cmd_a=\"c:/some/path\" -to_be_filtered=be_gone -cmd_b -not_allowed -cmd_c=7890sfd9wjd8jf984mx");
 
 		TCHAR NewMethod[1024];
 		TCHAR OldMethod[1024];
