- Move FApp::GetInstanceId into .cpp file to fix unique vs shared with dll export change

#rnx

diff --git a/Engine/Source/Runtime/Core/Private/Misc/App.cpp b/Engine/Source/Runtime/Core/Private/Misc/App.cpp
--- a/Engine/Source/Runtime/Core/Private/Misc/App.cpp
+++ b/Engine/Source/Runtime/Core/Private/Misc/App.cpp
@@ -459,3 +459,10 @@ FString FApp::GetZenStoreProjectId()
 #endif
 	return FString();
 }
+
+
+FGuid FApp::GetInstanceId()
+{
+	static FGuid InstanceId = FGuid::NewGuid();
+	return InstanceId;
+}
diff --git a/Engine/Source/Runtime/Core/Public/Misc/App.h b/Engine/Source/Runtime/Core/Public/Misc/App.h
--- a/Engine/Source/Runtime/Core/Public/Misc/App.h
+++ b/Engine/Source/Runtime/Core/Public/Misc/App.h
@@ -249,11 +249,7 @@ public:
 	 * @return Instance identifier, or an invalid GUID if there is no local instance.
 	 * @see GetSessionId
 	 */
-	FORCEINLINE static FGuid GetInstanceId()
-	{
-		static FGuid InstanceId = FGuid::NewGuid();
-		return InstanceId;
-	}
+	static CORE_API FGuid GetInstanceId();
 
 	/**
 	 * Gets the name of this application instance.
