Fix some ANSI to TCHAR issues in AutoRTFM Asserts

diff --git a/Engine/Source/Runtime/Core/Private/AutoRTFM/Utils.h b/Engine/Source/Runtime/Core/Private/AutoRTFM/Utils.h
--- a/Engine/Source/Runtime/Core/Private/AutoRTFM/Utils.h
+++ b/Engine/Source/Runtime/Core/Private/AutoRTFM/Utils.h
@@ -13,7 +13,7 @@
 
 DECLARE_LOG_CATEGORY_EXTERN(LogAutoRTFM, Display, All)
 
-#define ASSERT(exp) UE_CLOG(UNLIKELY(!(exp)), LogAutoRTFM, Fatal, TEXT("%s:%d:%s: assertion %s failed."), __FILE__, __LINE__, __PRETTY_FUNCTION__, #exp)
+#define ASSERT(exp) UE_CLOG(UNLIKELY(!(exp)), LogAutoRTFM, Fatal, TEXT("%s:%d:%s: assertion %s failed."), ANSI_TO_TCHAR(__FILE__), __LINE__, ANSI_TO_TCHAR(__PRETTY_FUNCTION__), ANSI_TO_TCHAR(#exp))
 
 namespace AutoRTFM
 {
