Fix for configurations where trace is disabled.

diff --git a/Engine/Source/Runtime/Core/Private/ProfilingDebugging/TraceAuxiliary.cpp b/Engine/Source/Runtime/Core/Private/ProfilingDebugging/TraceAuxiliary.cpp
--- a/Engine/Source/Runtime/Core/Private/ProfilingDebugging/TraceAuxiliary.cpp
+++ b/Engine/Source/Runtime/Core/Private/ProfilingDebugging/TraceAuxiliary.cpp
@@ -1223,7 +1223,7 @@ static void LaunchUnrealTraceInternal(const TCHAR* CommandLine)
 static void LaunchUnrealTraceInternal(const TCHAR* CommandLine)
 {
 // TSAN doesn't like fork(), so disable this for now.
-#if !USING_THREAD_SANITISER
+#if !USING_THREAD_SANITISER && UE_TRACE_ENABLED
 	if (GUnrealTraceLaunched.load(std::memory_order_relaxed))
 	{
 		UE_LOG(LogTrace, Log, TEXT("UnrealTraceServer: Trace store already started"));
