[Core]
* Readded dll export to TThreadSingleton::GetTlsSlot in an attempt to fix an assert happening on macos

diff --git a/Engine/Source/Runtime/Core/Public/HAL/ThreadSingleton.h b/Engine/Source/Runtime/Core/Public/HAL/ThreadSingleton.h
--- a/Engine/Source/Runtime/Core/Public/HAL/ThreadSingleton.h
+++ b/Engine/Source/Runtime/Core/Public/HAL/ThreadSingleton.h
@@ -44,7 +44,7 @@ class TThreadSingleton : public FTlsAutoCleanup
 	/**
 	 * @return TLS slot that holds a TThreadSingleton.
 	 */
-	static uint32& GetTlsSlot()
+	CORE_API static uint32& GetTlsSlot()
 	{
 		static uint32 TlsSlot = FPlatformTLS::InvalidTlsSlot;
 		return TlsSlot;
