Added Algo::Includes, based on std::includes.
Added unit tests for Algo::Includes

diff --git a/Engine/Source/Runtime/Core/Private/Containers/Algo/AlgosTest.cpp b/Engine/Source/Runtime/Core/Private/Containers/Algo/AlgosTest.cpp
--- a/Engine/Source/Runtime/Core/Private/Containers/Algo/AlgosTest.cpp
+++ b/Engine/Source/Runtime/Core/Private/Containers/Algo/AlgosTest.cpp
@@ -9,6 +9,7 @@
 #include "Algo/Copy.h"
 #include "Algo/Heapify.h"
 #include "Algo/HeapSort.h"
+#include "Algo/Includes.h"
 #include "Algo/IntroSort.h"
 #include "Algo/IsHeap.h"
 #include "Algo/IsSorted.h"
@@ -566,6 +567,119 @@ public:
 		}
 	}
 
+	void TestIncludes()
+	{
+		// Fixed arrays with elements of fundamental types - test Algo::Includes overloads
+		{
+			constexpr int32 FixedArrayA[] = { 1, 2, 3, 4, 5 };
+			constexpr int32 FixedArrayB[] = { 1, 3, 4 };
+			constexpr int32 FixedArrayC[] = { 5, 4, 3, 2, 1 };
+			constexpr int32 FixedArrayD[] = { 4, 3, 1 };
+			const TArrayView<int32> Empty;
+
+			// Test case 1: A contains A as a subsequence.
+			check(Algo::Includes(FixedArrayA, FixedArrayA));
+			// Test case 2: A contains B as a subsequence.
+			check(Algo::Includes(FixedArrayA, FixedArrayB));
+			// Test case 3: A contains Empty as a subsequence, because an empty set is always considered a subset of a non-empty set.
+			check(Algo::Includes(FixedArrayA, Empty));
+			// Test case 4: Empty contains Empty as a subsequence, because an empty set is always considered a subset of an empty set.
+			check(Algo::Includes(Empty, Empty));
+			// Test case 5: Empty doesn't contain A as a subsequence, because an empty set can't contain any elements of a non-empty set.
+			check(!Algo::Includes(Empty, FixedArrayA));
+			// Test case 6: B doesn't contain A as a subsequence, because A contains elements B doesn't have.
+			check(!Algo::Includes(FixedArrayB, FixedArrayA));
+			// Test case 7: C doesn't contain B as a subsequence, because C isn't ordered according to the default comparison predicate (operator<).
+			check(!Algo::Includes(FixedArrayC, FixedArrayB));
+			// Test case 8: A doesn't contain D as a subsequence, because D isn't ordered according to the default comparison predicate (operator<).
+			check(!Algo::Includes(FixedArrayA, FixedArrayD));
+
+			// Test case 9: C contains C as a subsequence with TGreater<> as the comparison predicate.
+			check(Algo::Includes(FixedArrayC, FixedArrayC, TGreater<>()));
+			// Test case 10: C contains D as a subsequence with TGreater<> as the comparison predicate.
+			check(Algo::Includes(FixedArrayC, FixedArrayD, TGreater<>()));
+			// Test case 11: C contains Empty as a subsequence, because an empty set is always considered a subset of a non-empty set.
+			check(Algo::Includes(FixedArrayC, Empty, TGreater<>()));
+			// Test case 12: Empty contains Empty as a subsequence, because an empty set is always considered a subset of an empty set.
+			check(Algo::Includes(Empty, Empty, TGreater<>()));
+			// Test case 13: Empty doesn't contain C as a subsequence, because an empty set can't contain any elements of a non-empty set.
+			check(!Algo::Includes(Empty, FixedArrayC, TGreater<>()));
+			// Test case 14: D doesn't contain C as a subsequence, because C contains elements D doesn't have.
+			check(!Algo::Includes(FixedArrayD, FixedArrayC, TGreater<>()));
+			// Test case 15: A doesn't contain D as a subsequence, because A isn't ordered according to the custom comparison predicate (TGreater<>()).
+			check(!Algo::Includes(FixedArrayA, FixedArrayD, TGreater<>()));
+			// Test case 16: C doesn't contain B as a subsequence, because B isn't ordered according to the custom comparison predicate (TGreater<>()).
+			check(!Algo::Includes(FixedArrayC, FixedArrayB, TGreater<>()));
+		}
+
+		// Dynamic arrays with elements of compound types - test Algo::IncludesBy
+		{
+			TArray<FTestData> DynamicArrayA =
+			{
+				FTestData(TEXT("1"), 1),
+				FTestData(TEXT("2"), 2),
+				FTestData(TEXT("3"), 3),
+				FTestData(TEXT("4"), 4),
+				FTestData(TEXT("5"), 5)
+			};
+			TArray<FTestData> DynamicArrayB =
+			{
+				FTestData(TEXT("1"), 1),
+				FTestData(TEXT("3"), 3),
+				FTestData(TEXT("4"), 4)
+			};
+			TArray<FTestData> DynamicArrayC =
+			{
+				FTestData(TEXT("5"), 5),
+				FTestData(TEXT("4"), 4),
+				FTestData(TEXT("3"), 3),
+				FTestData(TEXT("2"), 2),
+				FTestData(TEXT("1"), 1)
+			};
+			TArray<FTestData> DynamicArrayD =
+			{
+				FTestData(TEXT("4"), 4),
+				FTestData(TEXT("3"), 3),
+				FTestData(TEXT("1"), 1)
+			};
+			const TArrayView<FTestData> Empty;
+
+			// Test case 1: A contains A as a subsequence.
+			check(Algo::IncludesBy(DynamicArrayA, DynamicArrayA, &FTestData::Name));
+			// Test case 2: A contains B as a subsequence.
+			check(Algo::IncludesBy(DynamicArrayA, DynamicArrayB, &FTestData::Name));
+			// Test case 3: A contains Empty as a subsequence, because an empty set is always considered a subset of a non-empty set.
+			check(Algo::IncludesBy(DynamicArrayA, Empty, &FTestData::Name));
+			// Test case 4: Empty contains Empty as a subsequence, because an empty set is always considered a subset of an empty set.
+			check(Algo::IncludesBy(Empty, Empty, &FTestData::Name));
+			// Test case 5: Empty doesn't contain A as a subsequence, because an empty set can't contain any elements of a non-empty set.
+			check(!Algo::IncludesBy(Empty, DynamicArrayA, &FTestData::Name));
+			// Test case 6: B doesn't contain A as a subsequence, because A contains elements B doesn't have.
+			check(!Algo::IncludesBy(DynamicArrayB, DynamicArrayA, &FTestData::Name));
+			// Test case 7: C doesn't contain B as a subsequence, because C isn't ordered according to the default comparison predicate (operator<).
+			check(!Algo::IncludesBy(DynamicArrayC, DynamicArrayB, &FTestData::Name));
+			// Test case 8: A doesn't contain D as a subsequence, because D isn't ordered according to the default comparison predicate (operator<).
+			check(!Algo::IncludesBy(DynamicArrayA, DynamicArrayD, &FTestData::Name));
+
+			// Test case 9: C contains C as a subsequence with TGreater<> as the comparison predicate.
+			check(Algo::IncludesBy(DynamicArrayC, DynamicArrayC, &FTestData::Name, TGreater<>()));
+			// Test case 10: C contains D as a subsequence with TGreater<> as the comparison predicate.
+			check(Algo::IncludesBy(DynamicArrayC, DynamicArrayD, &FTestData::Name, TGreater<>()));
+			// Test case 11: C contains Empty as a subsequence, because an empty set is always considered a subset of a non-empty set.
+			check(Algo::IncludesBy(DynamicArrayC, Empty, &FTestData::Name, TGreater<>()));
+			// Test case 12: Empty contains Empty as a subsequence, because an empty set is always considered a subset of an empty set.
+			check(Algo::IncludesBy(Empty, Empty, &FTestData::Name, TGreater<>()));
+			// Test case 13: Empty doesn't contain C as a subsequence, because an empty set can't contain any elements of a non-empty set.
+			check(!Algo::IncludesBy(Empty, DynamicArrayC, &FTestData::Name, TGreater<>()));
+			// Test case 14: D doesn't contain C as a subsequence, because C contains elements D doesn't have.
+			check(!Algo::IncludesBy(DynamicArrayD, DynamicArrayC, &FTestData::Name, TGreater<>()));
+			// Test case 15: A doesn't contain D as a subsequence, because A isn't ordered according to the custom comparison predicate (TGreater<>()).
+			check(!Algo::IncludesBy(DynamicArrayA, DynamicArrayD, &FTestData::Name, TGreater<>()));
+			// Test case 16: C doesn't contain B as a subsequence, because B isn't ordered according to the custom comparison predicate (TGreater<>()).
+			check(!Algo::IncludesBy(DynamicArrayC, DynamicArrayB, &FTestData::Name, TGreater<>()));
+		}
+	}
+
 private:
 	TArray<int> TestData;
 	TArray<int> TestData2;
@@ -588,6 +702,7 @@ bool FAlgosTest::RunTest(const FString& Parameters)
 	TestSort();
 	TestEditDistance();
 	TestEditDistanceArray();
+	TestIncludes();
 	Cleanup();
 
 	return true;
diff --git a/Engine/Source/Runtime/Core/Public/Algo/Includes.h b/Engine/Source/Runtime/Core/Public/Algo/Includes.h
new file mode 100644
--- /dev/null
+++ b/Engine/Source/Runtime/Core/Public/Algo/Includes.h
@@ -0,0 +1,108 @@
+// Copyright Epic Games, Inc. All Rights Reserved.
+
+#pragma once
+
+#include "CoreTypes.h"
+#include "Templates/Less.h"
+#include "Templates/IdentityFunctor.h"
+#include "Templates/Invoke.h"
+#include "Templates/UnrealTemplate.h"
+
+namespace Algo::Private
+{
+	template <typename DataTypeA, typename SizeTypeA, typename DataTypeB, typename SizeTypeB, typename ProjectionType, typename SortPredicateType>
+	constexpr bool Includes(const DataTypeA* DataA, SizeTypeA NumA, const DataTypeB* DataB, SizeTypeB NumB, ProjectionType Projection, SortPredicateType SortPredicate)
+	{
+		SizeTypeA IndexA = 0;
+		SizeTypeB IndexB = 0;
+
+		for (; IndexB < NumB; ++IndexA)
+		{
+			if (IndexA >= NumA)
+			{
+				return false;
+			}
+
+			const auto& RefA = Invoke(Projection, DataA[IndexA]);
+			const auto& RefB = Invoke(Projection, DataB[IndexB]);
+
+			if (Invoke(SortPredicate, RefB, RefA))
+			{
+				return false;
+			}
+
+			if (!Invoke(SortPredicate, RefA, RefB))
+			{
+				++IndexB;
+			}
+		}
+		return true;
+	}
+}
+
+namespace Algo
+{
+	/**
+	 * Checks if one sorted contiguous container is a subsequence of another sorted contiguous container by comparing pairs of elements.
+	 * The subsequence does not need to be contiguous. Uses operator< to compare pairs of elements.
+	 *
+	 * @param  RangeA        Container of elements to search through. Must be already sorted by operator<.
+	 * @param  RangeB        Container of elements to search for. Must be already sorted by operator<.
+	 *
+	 * @return True if RangeB is a subsequence of RangeA, false otherwise.
+	 */
+	template <typename RangeTypeA, typename RangeTypeB>
+	constexpr bool Includes(RangeTypeA&& RangeA, RangeTypeB&& RangeB)
+	{
+		return Private::Includes(GetData(RangeA), GetNum(RangeA), GetData(RangeB), GetNum(RangeB), FIdentityFunctor(), TLess<>());
+	}
+	
+	/**
+	 * Checks if one sorted contiguous container is a subsequence of another sorted contiguous container by comparing pairs of elements using a custom predicate.
+	 * The subsequence does not need to be contiguous.
+	 *
+	 * @param  RangeA        Container of elements to search through. Must be already sorted by SortPredicate.
+	 * @param  RangeB        Container of elements to search for. Must be already sorted by SortPredicate.
+	 * @param  SortPredicate A binary predicate object used to specify if one element should precede another.
+	 *
+	 * @return True if RangeB is a subsequence of RangeA according to the provided predicate, false otherwise.
+	 */
+	template <typename RangeTypeA, typename RangeTypeB, typename SortPredicateType>
+	constexpr bool Includes(RangeTypeA&& RangeA, RangeTypeB&& RangeB, SortPredicateType SortPredicate)
+	{
+		return Private::Includes(GetData(RangeA), GetNum(RangeA), GetData(RangeB), GetNum(RangeB), FIdentityFunctor(), MoveTemp(SortPredicate));
+	}
+	
+	/**
+	 * Checks if one sorted contiguous container is a subsequence of another sorted contiguous container by comparing pairs of projected elements.
+	 * The subsequence does not need to be contiguous. Uses operator< to compare pairs of projected elements.
+	 *
+	 * @param  RangeA        Container of elements to search through. Must be already sorted by operator<.
+	 * @param  RangeB        Container of elements to search for. Must be already sorted by operator<.
+	 * @param  Projection    Projection to apply to the elements before comparing them.
+	 *
+	 * @return True if RangeB is a subsequence of RangeA, based on the comparison of projected elements, false otherwise.
+	 */
+	template <typename RangeTypeA, typename RangeTypeB, typename ProjectionType>
+	constexpr bool IncludesBy(RangeTypeA&& RangeA, RangeTypeB&& RangeB, ProjectionType Projection)
+	{
+		return Private::Includes(GetData(RangeA), GetNum(RangeA), GetData(RangeB), GetNum(RangeB), MoveTemp(Projection), TLess<>());
+	}
+	
+	/**
+	 * Checks if one sorted contiguous container is a subsequence of another sorted contiguous container by comparing pairs of projected elements using a custom predicate.
+	 * The subsequence does not need to be contiguous.
+	 *
+	 * @param  RangeA        Container of elements to search through. Must be already sorted by SortPredicate.
+	 * @param  RangeB        Container of elements to search for. Must be already sorted by SortPredicate.
+	 * @param  Projection    Projection to apply to the elements before comparing them.
+	 * @param  SortPredicate A binary predicate object, applied to the projection, used to specify if one element should precede another.
+	 *
+	 * @return True if RangeB is a subsequence of RangeA according to the comparison of projected elements using the provided predicate, false otherwise.
+	 */
+	template <typename RangeTypeA, typename RangeTypeB, typename ProjectionType, typename SortPredicateType>
+	constexpr bool IncludesBy(RangeTypeA&& RangeA, RangeTypeB&& RangeB, ProjectionType Projection, SortPredicateType SortPredicate)
+	{
+		return Private::Includes(GetData(RangeA), GetNum(RangeA), GetData(RangeB), GetNum(RangeB), MoveTemp(Projection), MoveTemp(SortPredicate));
+	}
+}
