Made some MacPlatformProcess methods related to retrieving directories common for iOS

diff --git a/Engine/Source/Runtime/Core/Private/Apple/ApplePlatformProcess.cpp b/Engine/Source/Runtime/Core/Private/Apple/ApplePlatformProcess.cpp
new file mode 100644
--- /dev/null
+++ b/Engine/Source/Runtime/Core/Private/Apple/ApplePlatformProcess.cpp
@@ -0,0 +1,77 @@
+// Copyright Epic Games, Inc. All Rights Reserved.
+
+#include "Apple/ApplePlatformProcess.h"
+
+#include "HAL/PlatformMisc.h"
+
+const TCHAR* FApplePlatformProcess::UserDir()
+{
+	static TCHAR Result[FPlatformMisc::GetMaxPathLength()] = TEXT("");
+	if (!Result[0])
+	{
+		SCOPED_AUTORELEASE_POOL;
+		NSString *DocumentsFolder = [NSSearchPathForDirectoriesInDomains(NSDocumentDirectory, NSUserDomainMask, YES) objectAtIndex: 0];
+		FPlatformString::CFStringToTCHAR((CFStringRef)DocumentsFolder, Result);
+		FCString::Strcat(Result, TEXT("/"));
+	}
+	return Result;
+}
+
+const TCHAR* FApplePlatformProcess::UserTempDir()
+{
+	static FString TempDir;
+	if (!TempDir.Len())
+	{
+		TempDir = NSTemporaryDirectory();
+	}
+	return *TempDir;
+}
+
+const TCHAR* FApplePlatformProcess::UserSettingsDir()
+{
+	return ApplicationSettingsDir();
+}
+
+const TCHAR* FApplePlatformProcess::UserHomeDir()
+{
+	static TCHAR Result[FPlatformMisc::GetMaxPathLength()] = TEXT("");
+	if (!Result[0])
+	{
+		SCOPED_AUTORELEASE_POOL;
+		FPlatformString::CFStringToTCHAR((CFStringRef)NSHomeDirectory(), Result);
+	}
+	return Result;
+}
+
+const TCHAR* FApplePlatformProcess::ApplicationSettingsDir()
+{
+	static TCHAR Result[FPlatformMisc::GetMaxPathLength()] = TEXT("");
+	if (!Result[0])
+	{
+		SCOPED_AUTORELEASE_POOL;
+		NSString *ApplicationSupportFolder = [NSSearchPathForDirectoriesInDomains(NSApplicationSupportDirectory, NSUserDomainMask, YES) objectAtIndex: 0];
+		FPlatformString::CFStringToTCHAR((CFStringRef)ApplicationSupportFolder, Result);
+		// @todo rocket this folder should be based on your company name, not just be hard coded to /Epic/
+		FCString::Strcat(Result, TEXT("/Epic/"));
+	}
+	return Result;
+}
+
+FString FApplePlatformProcess::GetApplicationSettingsDir(const ApplicationSettingsContext& Settings)
+{
+	TCHAR Result[FPlatformMisc::GetMaxPathLength()] = TEXT("");
+	SCOPED_AUTORELEASE_POOL;
+	NSString* ApplicationSupportFolder = [NSSearchPathForDirectoriesInDomains(NSApplicationSupportDirectory, NSUserDomainMask, YES) objectAtIndex:0];
+	FPlatformString::CFStringToTCHAR((CFStringRef) ApplicationSupportFolder, Result);
+	if (Settings.bIsEpic)
+	{
+		// @todo rocket this folder should be based on your company name, not just be hard coded to /Epic/
+		FCString::Strcat(Result, TEXT("/Epic/"));
+	}
+	else
+	{
+		FCString::Strcat(Result, TEXT("/"));
+	}
+	return FString(Result);
+}
+
diff --git a/Engine/Source/Runtime/Core/Private/Mac/MacPlatformProcess.cpp b/Engine/Source/Runtime/Core/Private/Mac/MacPlatformProcess.cpp
--- a/Engine/Source/Runtime/Core/Private/Mac/MacPlatformProcess.cpp
+++ b/Engine/Source/Runtime/Core/Private/Mac/MacPlatformProcess.cpp
@@ -28,7 +28,6 @@ namespace PlatformProcessLimits
 {
 	enum
 	{
-		MaxUserHomeDirLength = MAC_MAX_PATH + 1,
 		MaxArgvParameters	 = 256
 	};
 };
@@ -1050,34 +1049,6 @@ const TCHAR* FMacPlatformProcess::BaseDir()
 	return Result;
 }
 
-const TCHAR* FMacPlatformProcess::UserDir()
-{
-	static TCHAR Result[MAC_MAX_PATH] = TEXT("");
-	if (!Result[0])
-	{
-		SCOPED_AUTORELEASE_POOL;
-		NSString *DocumentsFolder = [NSSearchPathForDirectoriesInDomains(NSDocumentDirectory, NSUserDomainMask, YES) objectAtIndex: 0];
-		FPlatformString::CFStringToTCHAR((CFStringRef)DocumentsFolder, Result);
-		FCString::Strcat(Result, TEXT("/"));
-	}
-	return Result;
-}
-
-const TCHAR* FMacPlatformProcess::UserTempDir()
-{
-	static FString MacUserTempDir;
-	if (!MacUserTempDir.Len())
-	{
-		MacUserTempDir = NSTemporaryDirectory();
-	}
-	return *MacUserTempDir;
-}
-
-const TCHAR* FMacPlatformProcess::UserSettingsDir()
-{
-	return ApplicationSettingsDir();
-}
-
 static TCHAR* UserLibrarySubDirectory()
 {
 	static TCHAR Result[MAC_MAX_PATH] = TEXT("");
@@ -1128,49 +1099,6 @@ const TCHAR* FMacPlatformProcess::UserLogsDir()
 	return Result;
 }
 
-const TCHAR* FMacPlatformProcess::UserHomeDir()
-{
-	static TCHAR Result[MAC_MAX_PATH] = TEXT("");
-	if (!Result[0])
-	{
-		SCOPED_AUTORELEASE_POOL;
-		FPlatformString::CFStringToTCHAR((CFStringRef)NSHomeDirectory(), Result);
-	}
-	return Result;
-}
-
-const TCHAR* FMacPlatformProcess::ApplicationSettingsDir()
-{
-	static TCHAR Result[MAC_MAX_PATH] = TEXT("");
-	if (!Result[0])
-	{
-		SCOPED_AUTORELEASE_POOL;
-		NSString *ApplicationSupportFolder = [NSSearchPathForDirectoriesInDomains(NSApplicationSupportDirectory, NSUserDomainMask, YES) objectAtIndex: 0];
-		FPlatformString::CFStringToTCHAR((CFStringRef)ApplicationSupportFolder, Result);
-		// @todo rocket this folder should be based on your company name, not just be hard coded to /Epic/
-		FCString::Strcat(Result, TEXT("/Epic/"));
-	}
-	return Result;
-}
-
-FString FMacPlatformProcess::GetApplicationSettingsDir(const ApplicationSettingsContext& Settings)
-{
-	TCHAR Result[MAC_MAX_PATH] = TEXT("");
-	SCOPED_AUTORELEASE_POOL;
-	NSString* ApplicationSupportFolder = [NSSearchPathForDirectoriesInDomains(NSApplicationSupportDirectory, NSUserDomainMask, YES) objectAtIndex:0];
-	FPlatformString::CFStringToTCHAR((CFStringRef) ApplicationSupportFolder, Result);
-	if (Settings.bIsEpic)
-	{
-		// @todo rocket this folder should be based on your company name, not just be hard coded to /Epic/
-		FCString::Strcat(Result, TEXT("/Epic/"));
-	}
-	else
-	{
-		FCString::Strcat(Result, TEXT("/"));
-	}
-	return FString(Result);
-}
-
 const TCHAR* FMacPlatformProcess::ComputerName()
 {
 	static TCHAR Result[256]=TEXT("");
diff --git a/Engine/Source/Runtime/Core/Public/Apple/ApplePlatformProcess.h b/Engine/Source/Runtime/Core/Public/Apple/ApplePlatformProcess.h
new file mode 100644
--- /dev/null
+++ b/Engine/Source/Runtime/Core/Public/Apple/ApplePlatformProcess.h
@@ -0,0 +1,20 @@
+// Copyright Epic Games, Inc. All Rights Reserved.
+
+#pragma once
+
+#include "GenericPlatform/GenericPlatformProcess.h"
+
+/**
+ * Apple common  implementation of the Process OS functions
+ **/
+struct CORE_API FApplePlatformProcess : public FGenericPlatformProcess
+{
+	static const TCHAR* UserDir();
+	static const TCHAR* UserTempDir();
+	static const TCHAR* UserSettingsDir();
+	static const TCHAR* ApplicationSettingsDir();
+	static FString GetApplicationSettingsDir(const ApplicationSettingsContext& Settings);
+
+	// Apple specific
+	static const TCHAR* UserHomeDir();
+};
diff --git a/Engine/Source/Runtime/Core/Public/IOS/IOSPlatformMisc.h b/Engine/Source/Runtime/Core/Public/IOS/IOSPlatformMisc.h
--- a/Engine/Source/Runtime/Core/Public/IOS/IOSPlatformMisc.h
+++ b/Engine/Source/Runtime/Core/Public/IOS/IOSPlatformMisc.h
@@ -22,7 +22,7 @@ struct CORE_API FIOSPlatformMisc : public FApplePlatformMisc
     static void PlatformHandleSplashScreen(bool ShowSplashScreen = false);
 	static const TCHAR* GetPlatformFeaturesModuleName();
 
-	FORCEINLINE static int32 GetMaxPathLength()
+	FORCEINLINE static constexpr int32 GetMaxPathLength()
 	{
 		return IOS_MAX_PATH;
 	}
diff --git a/Engine/Source/Runtime/Core/Public/IOS/IOSPlatformProcess.h b/Engine/Source/Runtime/Core/Public/IOS/IOSPlatformProcess.h
--- a/Engine/Source/Runtime/Core/Public/IOS/IOSPlatformProcess.h
+++ b/Engine/Source/Runtime/Core/Public/IOS/IOSPlatformProcess.h
@@ -5,7 +5,7 @@
 ==============================================================================================*/
 
 #pragma once
-#include "GenericPlatform/GenericPlatformProcess.h"
+#include "Apple/ApplePlatformProcess.h"
 #import "IOSAsyncTask.h"
 
 /** Dummy process handle for platforms that use generic implementation. */
@@ -26,7 +26,7 @@ public:
 /**
 * iOS implementation of the Process OS functions
 **/
-struct CORE_API FIOSPlatformProcess : public FGenericPlatformProcess
+struct CORE_API FIOSPlatformProcess : public FApplePlatformProcess
 {
 	static const TCHAR* ComputerName();
 	static const TCHAR* BaseDir();
diff --git a/Engine/Source/Runtime/Core/Public/Mac/MacPlatformMisc.h b/Engine/Source/Runtime/Core/Public/Mac/MacPlatformMisc.h
--- a/Engine/Source/Runtime/Core/Public/Mac/MacPlatformMisc.h
+++ b/Engine/Source/Runtime/Core/Public/Mac/MacPlatformMisc.h
@@ -19,7 +19,7 @@ struct CORE_API FMacPlatformMisc : public FApplePlatformMisc
 	static void PlatformTearDown();
 	static void SetEnvironmentVar(const TCHAR* VariableName, const TCHAR* Value);
 
-	FORCEINLINE static int32 GetMaxPathLength()
+	FORCEINLINE static constexpr int32 GetMaxPathLength()
 	{
 		return MAC_MAX_PATH;
 	}
diff --git a/Engine/Source/Runtime/Core/Public/Mac/MacPlatformProcess.h b/Engine/Source/Runtime/Core/Public/Mac/MacPlatformProcess.h
--- a/Engine/Source/Runtime/Core/Public/Mac/MacPlatformProcess.h
+++ b/Engine/Source/Runtime/Core/Public/Mac/MacPlatformProcess.h
@@ -6,7 +6,7 @@
 
 #pragma once
 
-#include "GenericPlatform/GenericPlatformProcess.h"
+#include "Apple/ApplePlatformProcess.h"
 #include <sys/sysctl.h>
 
 /** Wrapper around Unix pid_t. Should not be copyable as changes in the process state won't be properly propagated to all copies. */
@@ -153,7 +153,7 @@ struct FProcHandle
 /**
  * Mac implementation of the Process OS functions
  **/
-struct CORE_API FMacPlatformProcess : public FGenericPlatformProcess
+struct CORE_API FMacPlatformProcess : public FApplePlatformProcess
 {
 	struct FProcEnumInfo;
 
@@ -225,11 +225,6 @@ struct CORE_API FMacPlatformProcess : public FGenericPlatformProcess
 	static uint32 GetCurrentProcessId();
 	static uint32 GetCurrentCoreNumber();
 	static const TCHAR* BaseDir();
-	static const TCHAR* UserDir();
-	static const TCHAR* UserTempDir();
-	static const TCHAR* UserSettingsDir();
-	static const TCHAR* ApplicationSettingsDir();
-	static FString GetApplicationSettingsDir(const ApplicationSettingsContext& Settings);
 	static const TCHAR* ComputerName();
 	static const TCHAR* UserName(bool bOnlyAlphaNumeric = true);
 	static void SetCurrentWorkingDirectoryToBaseDir();
@@ -270,7 +265,6 @@ struct CORE_API FMacPlatformProcess : public FGenericPlatformProcess
 	// Mac specific
 	static const TCHAR* UserPreferencesDir();
 	static const TCHAR* UserLogsDir();
-	static const TCHAR* UserHomeDir();
 
 private:
 
