cosmetic changes to FMRSWRecursiveAccessDetector

diff --git a/Engine/Source/Runtime/Core/Public/Misc/MTAccessDetector.h b/Engine/Source/Runtime/Core/Public/Misc/MTAccessDetector.h
--- a/Engine/Source/Runtime/Core/Public/Misc/MTAccessDetector.h
+++ b/Engine/Source/Runtime/Core/Public/Misc/MTAccessDetector.h
@@ -310,6 +310,11 @@ private:
 		{
 			return WriterThreadId - 1;
 		}
+
+		void SetWriterThreadId(uint32 ThreadId)
+		{
+			WriterThreadId = ThreadId + 1; // to shift from 0 TID that is considered "invalid"
+		}
 	};
 
 	// all atomic ops are relaxed to preserve the original memory order, as the detector is compiled out in non-dev builds
@@ -357,6 +362,11 @@ private:
 
 	static FString GetThreadCallstack(uint32 ThreadId)
 	{
+		if (ThreadId == uint32(-1))
+		{
+			return TEXT("[none]");
+		}
+
 		const SIZE_T StackTraceSize = 65536;
 		ANSICHAR StackTrace[StackTraceSize] = { 0 };
 		FPlatformStackWalk::ThreadStackWalkAndDump(StackTrace, StackTraceSize, 0, ThreadId);
@@ -409,7 +419,7 @@ private:
 
 	void RemoveReaderFromTls() const
 	{
-		uint32 ReaderIndex = GetReadersTls().IndexOfByPredicate([this](FReaderNum ReaderNum) { return ReaderNum.Reader == this; });
+		int32 ReaderIndex = GetReadersTls().IndexOfByPredicate([this](FReaderNum ReaderNum) { return ReaderNum.Reader == this; });
 		checkfSlow(ReaderIndex != INDEX_NONE,
 			TEXT("Invalid usage of the race detector! No matching AcquireReadAccess(): %u readers, %u writers on thread %u"),
 			LoadState().ReaderNum, LoadState().WriterNum, LoadState().GetWriterThreadId());
@@ -459,7 +469,7 @@ public:
 				else
 				{
 					++ExpectedState.WriterNum;
-					ExpectedState.WriterThreadId = FPlatformTLS::GetCurrentThreadId() + 1; // to shift from 0 TID that is considered "invalid"
+					ExpectedState.SetWriterThreadId(FPlatformTLS::GetCurrentThreadId());
 				}
 
 				GetDestructionSentinelStackTls().RemoveAtSwap(DestructionSentinelIndex);
@@ -482,7 +492,7 @@ public:
 		CheckOtherThreadWriters(PrevState);
 
 		// register the reader in TLS
-		uint32 ReaderIndex = PrevState.ReaderNum == 0 ? 
+		int32 ReaderIndex = PrevState.ReaderNum == 0 ? 
 			INDEX_NONE : 
 			GetReadersTls().IndexOfByPredicate([this](FReaderNum ReaderNum) { return ReaderNum.Reader == this; });
 		if (ReaderIndex == INDEX_NONE)
@@ -525,7 +535,7 @@ public:
 		GetDestructionSentinelStackTls().RemoveAtSwap(GetDestructionSentinelStackTls().Num() - 1);
 	}
 
-	FORCEINLINE_DEBUGGABLE void AcquireWriteAccess()
+	FORCEINLINE void AcquireWriteAccess()
 	{
 		FState LocalState = LoadState();
 
@@ -533,10 +543,12 @@ public:
 		{	// check that all readers are on the current thread
 			int32 ReaderIndex = GetReadersTls().IndexOfByPredicate([this](FReaderNum ReaderNum) { return ReaderNum.Reader == this; });
 			ensureMsgf(ReaderIndex != INDEX_NONE, TEXT("Race detector is not trivially copyable while this delegate is copied trivially. Consider changing this delegate to use `FNotThreadSafeNotCheckedDelegateUserPolicy`"));
+			
 			if (ReaderIndex == INDEX_NONE)
 			{
 				return; // to avoid asserting on the next line inside TArray
 			}
+
 			ensureMsgf(GetReadersTls()[ReaderIndex].Num == LocalState.ReaderNum,
 				TEXT("Data race detected: %d reader(s) on another thread(s) while acquiring write access.\nCurrent thread %u callstack:\n%s"),
 				LocalState.ReaderNum - GetReadersTls()[ReaderIndex].Num,
