FMRSWRecursiveAccessDetector: improved logging the current thread callstack to avoid suspending it

diff --git a/Engine/Source/Runtime/Core/Public/Misc/MTAccessDetector.h b/Engine/Source/Runtime/Core/Public/Misc/MTAccessDetector.h
--- a/Engine/Source/Runtime/Core/Public/Misc/MTAccessDetector.h
+++ b/Engine/Source/Runtime/Core/Public/Misc/MTAccessDetector.h
@@ -355,11 +355,19 @@ private:
 			ensureMsgf(InState.GetWriterThreadId() == CurrentThreadId,
 				TEXT("Data race detected! Acquiring read access on thread %u concurrently with %d writers on thread %u:\nCurrent thread %u callstack:\n%s\nWriter thread %u callstack:\n%s"),
 				CurrentThreadId, InState.WriterNum, InState.GetWriterThreadId(), 
-				CurrentThreadId, *GetThreadCallstack(CurrentThreadId),
+				CurrentThreadId, *GetCurrentThreadCallstack(),
 				InState.GetWriterThreadId(), *GetThreadCallstack(InState.GetWriterThreadId()));
 		}
 	}
 
+	static FString GetCurrentThreadCallstack()
+	{
+		const SIZE_T StackTraceSize = 65536;
+		ANSICHAR StackTrace[StackTraceSize] = { 0 };
+		FPlatformStackWalk::StackWalkAndDump(StackTrace, StackTraceSize, 0);
+		return StackTrace;
+	}
+
 	static FString GetThreadCallstack(uint32 ThreadId)
 	{
 		if (ThreadId == uint32(-1))
@@ -481,7 +489,7 @@ public:
 			TEXT("Race detector destroyed while being accessed on another thread: %d readers, %d writers on thread %u:\nCurrent thread %u callstack:\n%s\nWriter thread %u callstack:\n%s"),
 			LoadState().ReaderNum - ExpectedState.ReaderNum, 
 			LoadState().WriterNum - ExpectedState.WriterNum, LoadState().GetWriterThreadId(), 
-			FPlatformTLS::GetCurrentThreadId(), *GetThreadCallstack(FPlatformTLS::GetCurrentThreadId()),
+			FPlatformTLS::GetCurrentThreadId(), *GetCurrentThreadCallstack(),
 			LoadState().GetWriterThreadId(), *GetThreadCallstack(LoadState().GetWriterThreadId()));
 	}
 
@@ -551,8 +559,7 @@ public:
 
 			ensureMsgf(GetReadersTls()[ReaderIndex].Num == LocalState.ReaderNum,
 				TEXT("Data race detected: %d reader(s) on another thread(s) while acquiring write access.\nCurrent thread %u callstack:\n%s"),
-				LocalState.ReaderNum - GetReadersTls()[ReaderIndex].Num,
-				FPlatformTLS::GetCurrentThreadId(), *GetThreadCallstack(FPlatformTLS::GetCurrentThreadId()));
+				LocalState.ReaderNum - GetReadersTls()[ReaderIndex].Num, FPlatformTLS::GetCurrentThreadId(), *GetCurrentThreadCallstack());
 		}
 
 		uint32 CurrentThreadId = FPlatformTLS::GetCurrentThreadId();
@@ -561,7 +568,7 @@ public:
 			ensureMsgf(LocalState.GetWriterThreadId() == CurrentThreadId,
 				TEXT("Data race detected: acquiring write access on thread %u concurrently with %d writers on thread %u:\nCurrent thread %u callstack:\n%s\nWriter thread %u callstack:\n%s"), 
 				CurrentThreadId, LocalState.WriterNum, LocalState.GetWriterThreadId(), 
-				CurrentThreadId, *GetThreadCallstack(CurrentThreadId),
+				CurrentThreadId, *GetCurrentThreadCallstack(),
 				LocalState.GetWriterThreadId(), *GetThreadCallstack(LocalState.GetWriterThreadId()));
 		}
 
@@ -574,7 +581,7 @@ public:
 			LocalState.ReaderNum, PrevState.ReaderNum,
 			LocalState.WriterNum, PrevState.WriterNum,
 			LocalState.GetWriterThreadId(), PrevState.GetWriterThreadId(), 
-			CurrentThreadId, *GetThreadCallstack(CurrentThreadId),
+			CurrentThreadId, *GetCurrentThreadCallstack(),
 			PrevState.GetWriterThreadId(), *GetThreadCallstack(PrevState.GetWriterThreadId()));
 	}
 
@@ -601,7 +608,7 @@ public:
 			LocalState.ReaderNum, PrevState.ReaderNum,
 			LocalState.WriterNum, PrevState.WriterNum,
 			LocalState.GetWriterThreadId(), PrevState.GetWriterThreadId(), 
-			FPlatformTLS::GetCurrentThreadId(), *GetThreadCallstack(FPlatformTLS::GetCurrentThreadId()),
+			FPlatformTLS::GetCurrentThreadId(), *GetCurrentThreadCallstack(),
 			PrevState.GetWriterThreadId(), *GetThreadCallstack(PrevState.GetWriterThreadId()));
 	}
 
