Fix for root node's children not being properly unset during TOctree2::ApplyOffset.

diff --git a/Engine/Source/Runtime/Core/Public/Math/GenericOctree.h b/Engine/Source/Runtime/Core/Public/Math/GenericOctree.h
--- a/Engine/Source/Runtime/Core/Public/Math/GenericOctree.h
+++ b/Engine/Source/Runtime/Core/Public/Math/GenericOctree.h
@@ -592,7 +592,7 @@ private:
 
 	void CollapseNodesInternal(FNodeIndex CurrentNodeIndex, ElementArrayType& CollapsedNodeElements)
 	{
-		CollapsedNodeElements.Append(MoveTemp(TreeElements[CurrentNodeIndex]));	
+		CollapsedNodeElements.Append(MoveTemp(TreeElements[CurrentNodeIndex]));
 		TreeElements[CurrentNodeIndex].Reset();
 
 		if (!TreeNodes[CurrentNodeIndex].IsLeaf())
@@ -602,6 +602,10 @@ private:
 			{
 				CollapseNodesInternal(ChildStartIndex + i, CollapsedNodeElements);
 			}
+
+			// Mark the node as a leaf.
+			TreeNodes[CurrentNodeIndex].ChildNodes = INDEX_NONE;
+
 			FreeEightNodes(ChildStartIndex);
 		}
 	}
@@ -868,9 +872,6 @@ public:
 					// Update the external element id for the element that's being collapsed.
 					SetElementId(TreeElements[CollapseNodeIndex][ElementIndex], FOctreeElementId2(CollapseNodeIndex, ElementIndex));
 				}
-
-				// Mark the node as a leaf.
-				TreeNodes[CollapseNodeIndex].ChildNodes = INDEX_NONE;
 			}
 		}
 	}
diff --git a/Engine/Source/Runtime/Core/Tests/Math/GenericOctreeTest.cpp b/Engine/Source/Runtime/Core/Tests/Math/GenericOctreeTest.cpp
new file mode 100644
--- /dev/null
+++ b/Engine/Source/Runtime/Core/Tests/Math/GenericOctreeTest.cpp
@@ -0,0 +1,67 @@
+// Copyright Epic Games, Inc. All Rights Reserved.
+
+#include "CoreTypes.h"
+#include "Math/GenericOctree.h"
+#include "Misc/AutomationTest.h"
+#include "Tests/TestHarnessAdapter.h"
+
+#if WITH_TESTS
+
+namespace UE::MathTest
+{
+	struct FTestOctreeElement
+	{
+		FTestOctreeElement() = default;
+
+		explicit FTestOctreeElement(int InValue, const FBoxCenterAndExtent& InBounds)
+			: Value(InValue)
+			, Bounds(InBounds)
+		{
+		}
+
+		int32 Value = 0;
+		FBoxCenterAndExtent Bounds = {};
+	};
+
+	struct TTestOctreeSemantics
+	{
+		enum { MaxElementsPerLeaf = 2 };
+		enum { MinInclusiveElementsPerNode = 7 };
+		enum { MaxNodeDepth = 12 };
+
+		typedef TInlineAllocator<MaxElementsPerLeaf> ElementAllocator;
+
+		static FBoxCenterAndExtent GetBoundingBox(const FTestOctreeElement& InPoint)
+		{
+			return InPoint.Bounds;
+		}
+
+		static const bool AreElementsEqual(const FTestOctreeElement& Lhs, const FTestOctreeElement& Rhs)
+		{
+			return Lhs.Value == Rhs.Value;
+		}
+
+		static void ApplyOffset(const FTestOctreeElement& Element, const FVector& Offset)
+		{
+		}
+
+		static void SetElementId(const FTestOctreeElement& Element, FOctreeElementId2 OctreeElementID)
+		{
+		}
+	};
+
+	TEST_CASE_NAMED(FGenericOctreeTest, "System::Core::Math::GenericOctree::ApplyOffset", "[ApplicationContextMask][EngineFilter]")
+	{
+		// Ensure that no assert fires when applying an offset
+		TOctree2<FTestOctreeElement, TTestOctreeSemantics> Octree;
+		for (int32 Count = 20; Count != 0; --Count)
+		{
+			FVector Center((float)Count, (float)Count, (float)Count);
+			FVector Extent = Center + FVector(1., 1., 1.);
+			Octree.AddElement(FTestOctreeElement(5, FBoxCenterAndExtent(Center, Extent)));
+		}
+		Octree.ApplyOffset(FVector{});
+	}
+}
+
+#endif // WITH_TESTS
