FMRSWRecursiveAccessDetector: improved race detection messaging for trivial relocations, added a unit-test

diff --git a/Engine/Source/Runtime/Core/Private/Misc/MTAccessDetector.cpp b/Engine/Source/Runtime/Core/Private/Misc/MTAccessDetector.cpp
--- a/Engine/Source/Runtime/Core/Private/Misc/MTAccessDetector.cpp
+++ b/Engine/Source/Runtime/Core/Private/Misc/MTAccessDetector.cpp
@@ -459,6 +459,50 @@ bool FMRSWRecursiveAccessDetectorTest::RunTest(const FString& Parameters)
 #pragma warning(pop)
 }
 
+IMPLEMENT_SIMPLE_AUTOMATION_TEST(FMRSWRecursiveAccessDetectorTrivialRelocationTest, "System.Core.MRSWRecursiveAccessDetector.TrivialRelocation", EAutomationTestFlags::ApplicationContextMask | EAutomationTestFlags::EngineFilter);
+
+bool FMRSWRecursiveAccessDetectorTrivialRelocationTest::RunTest(const FString& Parameters)
+{
+#pragma warning(push)
+#pragma warning(disable:6001) // "Using uninitialized memory" warning
+
+#if 0
+	{	// trivial relocation by shrinking a TArray while read-accessed is not supported by the race detector. asserts, must provide a useful message
+		TArray<FMRSWRecursiveAccessDetector> Arr;
+
+		Arr.AddDefaulted();
+		Arr.AddDefaulted();
+
+		Arr[1].AcquireReadAccess(); // acquire read access on the old instance
+
+		Arr.RemoveAtSwap(0); // remove the unused first element
+
+		Arr.Shrink(); // trigger trivial relocation
+
+		Arr[0].ReleaseReadAccess(); // release read access on the new broken instance
+	}
+#endif
+
+	{	// trivial relocation by shrinking a TArray while write-accessed is supported by the race detector
+		TArray<FMRSWRecursiveAccessDetector> Arr;
+
+		Arr.AddDefaulted();
+		Arr.AddDefaulted();
+
+		Arr[1].AcquireWriteAccess(); // acquire access on the old instance
+
+		Arr.RemoveAtSwap(0); // remove the unused first element
+
+		Arr.Shrink(); // trigger trivial relocation
+
+		Arr[0].ReleaseWriteAccess(); // release access on the new broken instance
+	}
+
+	return true;
+
+#pragma warning(pop)
+}
+
 template<uint64 Num, typename AccessDetectorType>
 void RWAccessDetectorPerfTest()
 {
diff --git a/Engine/Source/Runtime/Core/Public/Misc/MTAccessDetector.h b/Engine/Source/Runtime/Core/Public/Misc/MTAccessDetector.h
--- a/Engine/Source/Runtime/Core/Public/Misc/MTAccessDetector.h
+++ b/Engine/Source/Runtime/Core/Public/Misc/MTAccessDetector.h
@@ -428,9 +428,17 @@ private:
 	void RemoveReaderFromTls() const
 	{
 		int32 ReaderIndex = GetReadersTls().IndexOfByPredicate([this](FReaderNum ReaderNum) { return ReaderNum.Reader == this; });
-		checkfSlow(ReaderIndex != INDEX_NONE,
-			TEXT("Invalid usage of the race detector! No matching AcquireReadAccess(): %u readers, %u writers on thread %u"),
-			LoadState().ReaderNum, LoadState().WriterNum, LoadState().GetWriterThreadId());
+
+		ensureMsgf(ReaderIndex != INDEX_NONE,
+			TEXT("Either invalid usage of the access detector (no matching AcquireReadAccess()) or the access detector was trivially relocated (this instance is not registered): %u readers, %u writers on thread %u:\nCurrent thread %u callstack:\n%s"),
+			LoadState().ReaderNum, LoadState().WriterNum, LoadState().GetWriterThreadId(), 
+			FPlatformTLS::GetCurrentThreadId(), *GetCurrentThreadCallstack());
+
+		if (ReaderIndex == INDEX_NONE)
+		{
+			return; // to avoid asserting on the next line inside TArray
+		}
+
 		uint32 ReaderNum = --GetReadersTls()[ReaderIndex].Num;
 		if (ReaderNum == 0)
 		{
@@ -550,8 +558,11 @@ public:
 		if (LocalState.ReaderNum >= 1)
 		{	// check that all readers are on the current thread
 			int32 ReaderIndex = GetReadersTls().IndexOfByPredicate([this](FReaderNum ReaderNum) { return ReaderNum.Reader == this; });
-			ensureMsgf(ReaderIndex != INDEX_NONE, TEXT("Race detector is not trivially copyable while this delegate is copied trivially. Consider changing this delegate to use `FNotThreadSafeNotCheckedDelegateUserPolicy`"));
-			
+
+			ensureMsgf(ReaderIndex != INDEX_NONE, 
+				TEXT("Either a race detected (%u reader(s) on another thread(s) while acquiring write access on the current thread) or the access detector was trivially relocated:\nCurrent thread %u callstack:\n%s"), 
+				LocalState.ReaderNum, FPlatformTLS::GetCurrentThreadId(), *GetCurrentThreadCallstack());
+
 			if (ReaderIndex == INDEX_NONE)
 			{
 				return; // to avoid asserting on the next line inside TArray
