Fix missing "

diff --git a/Engine/Source/Runtime/Core/Public/UObject/FortniteNCBranchObjectVersion.h b/Engine/Source/Runtime/Core/Public/UObject/FortniteNCBranchObjectVersion.h
--- a/Engine/Source/Runtime/Core/Public/UObject/FortniteNCBranchObjectVersion.h
+++ b/Engine/Source/Runtime/Core/Public/UObject/FortniteNCBranchObjectVersion.h
@@ -9,5 +9,5 @@
 
 #include "UObject/FortniteSeasonBranchObjectVersion.h"
 
-UE_DEPRECATED(5.3, "FFortniteNCBranchObjectVersion has been repurposed, use FFortniteSeasonBranchObjectVersion instead)
+UE_DEPRECATED(5.3, "FFortniteNCBranchObjectVersion has been repurposed, use FFortniteSeasonBranchObjectVersion instead")
 typedef FFortniteSeasonBranchObjectVersion FFortniteNCBranchObjectVersion;
