Add FindAfter function to TBitArray, to allow a user to provide a start index when searching for a value.
Added unit tests for TBitArray's Find and FindAfter functions.

diff --git a/Engine/Source/Runtime/Core/Public/Containers/BitArray.h b/Engine/Source/Runtime/Core/Public/Containers/BitArray.h
--- a/Engine/Source/Runtime/Core/Public/Containers/BitArray.h
+++ b/Engine/Source/Runtime/Core/Public/Containers/BitArray.h
@@ -1105,6 +1105,63 @@ public:
 		return INDEX_NONE;
 	}
 
+private:
+
+	int32 FindAfterImpl(bool bValue, int32 StartIndex) const
+	{
+		checkSlow((StartIndex >= 0) & (StartIndex <= NumBits));
+
+		// Produce a mask for the first iteration
+		uint32 Mask = ~0u << (StartIndex % FBitSet::BitsPerWord);
+
+		// Iterate over the array until we see a word with a matching bit
+		const uint32 Test = bValue ? 0u : ~0u;
+
+		const uint32* RESTRICT DwordArray = GetData();
+		const int32 LocalNumBits = NumBits;
+		const int32 DwordCount = FBitSet::CalculateNumWords(LocalNumBits);
+		int32 DwordIndex = FMath::DivideAndRoundDown(StartIndex, NumBitsPerDWORD);
+		while (DwordIndex < DwordCount && (DwordArray[DwordIndex] & Mask) == (Test & Mask))
+		{
+			++DwordIndex;
+			Mask = ~0u;
+		}
+
+		if (DwordIndex < DwordCount)
+		{
+			// If we're looking for a false, then we flip the bits - then we only need to find the first one bit
+			const uint32 Bits = (bValue ? DwordArray[DwordIndex] : ~DwordArray[DwordIndex]) & Mask;
+			UE_ASSUME(Bits != 0);
+			const int32 LowestBitIndex = FMath::CountTrailingZeros(Bits) + (DwordIndex << NumBitsPerDWORDLogTwo);
+			if (LowestBitIndex < LocalNumBits)
+			{
+				return LowestBitIndex;
+			}
+		}
+
+		return INDEX_NONE;
+	}
+
+public:
+
+	/**
+	 * Finds the first occurrence of the specified value (true/false) in the array, starting from the given bit index, and returns the bit index.
+	 * If the specified value is not found after the given index, INDEX_NONE is returned.
+	 *
+	 * @param  bValue      The value (true/false) to search for.
+	 * @param  StartIndex  The index to start the search from.
+	 *
+	 * @pre StartIndex is expected to be in the inclusive range [0, Num()].
+	 *
+	 * @return The index of the first occurrence of the specified value (true/false) after StartIndex, or INDEX_NONE if not found.
+	 */
+	template <typename IndexType>
+	FORCEINLINE int32 FindAfter(bool bValue, IndexType StartIndex) const
+	{
+		static_assert(!std::is_same_v<IndexType, bool>, "TBitArray::FindAfter: unexpected bool passed as the StartIndex argument");
+		return FindAfterImpl(bValue, StartIndex);
+	}
+
 	/**
 	* Finds the last true/false bit in the array, and returns the bit index.
 	* If there is none, INDEX_NONE is returned.
diff --git a/Engine/Source/Runtime/Core/Tests/Containers/BitArrayTests.cpp b/Engine/Source/Runtime/Core/Tests/Containers/BitArrayTests.cpp
--- a/Engine/Source/Runtime/Core/Tests/Containers/BitArrayTests.cpp
+++ b/Engine/Source/Runtime/Core/Tests/Containers/BitArrayTests.cpp
@@ -1521,7 +1521,6 @@ public:
 
 	// TODO: GetAllocatedSize
 	// TODO: CountBytes
-	// TODO: Find
 	// TODO: FindLast
 	// TODO: Contains
 	// TODO: FindAndSetFirstZeroBit
@@ -2022,6 +2021,75 @@ TEST_CASE_NAMED(FBitArrayBitwiseXORTest, "System::Core::Containers::BitArray::Bi
 	}
 }
 
+TEST_CASE_NAMED(FBitArrayFindTest, "System::Core::Containers::BitArray::Find", "[ApplicationContextMask][SmokeFilter]")
+{
+	{
+		TBitArray<> ArrEmpty;
+		CHECK_MESSAGE(TEXT("Find true in an empty array"), ArrEmpty.Find(true) == INDEX_NONE);
+		CHECK_MESSAGE(TEXT("Find false in an empty array"), ArrEmpty.Find(false) == INDEX_NONE);
+	}
+
+	{
+		TBitArray<> ArrTrue2(true, NumBitsPerDWORD);
+		CHECK_MESSAGE(TEXT("Find true in an array containing 32 true"), ArrTrue2.Find(true) == 0);
+		CHECK_MESSAGE(TEXT("Find false in an array containing 32 true"), ArrTrue2.Find(false) == INDEX_NONE);
+	}
+
+	{
+		TBitArray<> ArrFalse2(false, NumBitsPerDWORD);
+		CHECK_MESSAGE(TEXT("Find false in an array containing 32 false"), ArrFalse2.Find(false) == 0);
+		CHECK_MESSAGE(TEXT("Find true in an array containing 32 false"), ArrFalse2.Find(true) == INDEX_NONE);
+	}
+
+	{
+		TBitArray<> ArrFalseThenTrue(false, NumBitsPerDWORD);
+		ArrFalseThenTrue.Add(true);
+		CHECK_MESSAGE(TEXT("Find true in an array containing multiple false and then a single true"), ArrFalseThenTrue.Find(true) == NumBitsPerDWORD);
+	}
+
+	{
+		TBitArray<> ArrTrueThenFalse(true, NumBitsPerDWORD);
+		ArrTrueThenFalse.Add(false);
+		CHECK_MESSAGE(TEXT("Find true in an array containing multiple true and then a single false"), ArrTrueThenFalse.Find(false) == NumBitsPerDWORD);
+	}
+}
+
+TEST_CASE_NAMED(FBitArrayFindAfterTest, "System::Core::Containers::BitArray::FindAfter", "[ApplicationContextMask][SmokeFilter]")
+{
+	{
+		TBitArray<> ArrTrue2(true, NumBitsPerDWORD);
+		CHECK_MESSAGE(TEXT("Find true in an array full of true, starting at index 0"), ArrTrue2.FindAfter(true, 0) == 0);
+		CHECK_MESSAGE(TEXT("Find true in an array full of true, starting at index 1"), ArrTrue2.FindAfter(true, 1) == 1);
+		CHECK_MESSAGE(TEXT("Find true in an array full of true, starting at index 31"), ArrTrue2.FindAfter(true, NumBitsPerDWORD - 1) == NumBitsPerDWORD - 1);
+		CHECK_MESSAGE(TEXT("Find true in an array full of true, starting at index 32"), ArrTrue2.FindAfter(true, NumBitsPerDWORD) == INDEX_NONE);
+		CHECK_MESSAGE(TEXT("Find false in an array full of true, starting at index 1"), ArrTrue2.FindAfter(false, 1) == INDEX_NONE);
+		CHECK_MESSAGE(TEXT("Find false in an array full of true, starting at index 32"), ArrTrue2.FindAfter(false, NumBitsPerDWORD) == INDEX_NONE);
+	}
+
+	{
+		TBitArray<> ArrFalse2(false, NumBitsPerDWORD);
+		CHECK_MESSAGE(TEXT("Find false in an array full of false, starting at index 0"), ArrFalse2.FindAfter(false, 0) == 0);
+		CHECK_MESSAGE(TEXT("Find false in an array full of false, starting at index 1"), ArrFalse2.FindAfter(false, 1) == 1);
+		CHECK_MESSAGE(TEXT("Find false in an array full of false, starting at index 31"), ArrFalse2.FindAfter(false, NumBitsPerDWORD - 1) == NumBitsPerDWORD - 1);
+		CHECK_MESSAGE(TEXT("Find false in an array full of false, starting at index 32"), ArrFalse2.FindAfter(false, NumBitsPerDWORD) == INDEX_NONE);
+		CHECK_MESSAGE(TEXT("Find true in an array full of false, starting at index 1"), ArrFalse2.FindAfter(true, 1) == INDEX_NONE);
+		CHECK_MESSAGE(TEXT("Find false in an array full of false, starting at index 32"), ArrFalse2.FindAfter(true, NumBitsPerDWORD) == INDEX_NONE);
+	}
+
+	{
+		TBitArray<> ArrFalseThenTrue(false, NumBitsPerDWORD);
+		ArrFalseThenTrue.Add(true, NumBitsPerDWORD);
+		CHECK_MESSAGE(TEXT("Find true in an array containing 32 false and then 32 true, starting after all of the false"), ArrFalseThenTrue.FindAfter(true, NumBitsPerDWORD) == NumBitsPerDWORD);
+		CHECK_MESSAGE(TEXT("Find false in an array containing 32 false and then 32 true, starting after all of the false"), ArrFalseThenTrue.FindAfter(false, NumBitsPerDWORD) == INDEX_NONE);
+	}
+
+	{
+		TBitArray<> ArrTrueThenFalse(true, NumBitsPerDWORD);
+		ArrTrueThenFalse.Add(false, NumBitsPerDWORD);
+		CHECK_MESSAGE(TEXT("Find true in an array containing 32 true and then 32 false, starting after all of the true"), ArrTrueThenFalse.FindAfter(true, NumBitsPerDWORD) == INDEX_NONE);
+		CHECK_MESSAGE(TEXT("Find false in an array containing 32 true and then 32 false, starting after all of the true"), ArrTrueThenFalse.FindAfter(false, NumBitsPerDWORD) == NumBitsPerDWORD);
+	}
+}
 
 class FBitArrayMemoryTest
 {
