made FModuleManager::ModulesChangedEvent thread-safe to allow subscription from a worker thread

diff --git a/Engine/Source/Runtime/Core/Public/Modules/ModuleManager.h b/Engine/Source/Runtime/Core/Public/Modules/ModuleManager.h
--- a/Engine/Source/Runtime/Core/Public/Modules/ModuleManager.h
+++ b/Engine/Source/Runtime/Core/Public/Modules/ModuleManager.h
@@ -474,14 +474,15 @@ public:
 #endif
 
 	/**
-	 * Gets an event delegate that is executed when the set of known modules changed, i.e. upon module load or unload.
+	 * Gets a multicast delegate that is executed when the set of known modules changed, i.e. upon module load or unload.
+	 * The delegate is thread-safe to allow subscribing from other than the game thread but is always broadcasted from the game thread.
 	 *
 	 * The first parameter is the name of the module that changed.
 	 * The second parameter is the reason for the change.
 	 *
-	 * @return The event delegate.
+	 * @return The multicast delegate.
 	 */
-	DECLARE_EVENT_TwoParams(FModuleManager, FModulesChangedEvent, FName, EModuleChangeReason);
+	using FModulesChangedEvent = TTSMulticastDelegate<void(FName ModuleName, EModuleChangeReason ChangeReason)>;
 	FModulesChangedEvent& OnModulesChanged( )
 	{
 		return ModulesChangedEvent;
