Mark FGuid constructors as constexpr so they can be used in declarations such as "static inline constexpr FGuid MyGuid"

diff --git a/Engine/Source/Runtime/Core/Public/Misc/Guid.h b/Engine/Source/Runtime/Core/Public/Misc/Guid.h
--- a/Engine/Source/Runtime/Core/Public/Misc/Guid.h
+++ b/Engine/Source/Runtime/Core/Public/Misc/Guid.h
@@ -109,7 +109,7 @@ struct FGuid
 public:
 
 	/** Default constructor. */
-	FGuid()
+	constexpr FGuid()
 		: A(0)
 		, B(0)
 		, C(0)
@@ -124,7 +124,7 @@ public:
 	 * @param InC The third component.
 	 * @param InD The fourth component.
 	 */
-	explicit FGuid(uint32 InA, uint32 InB, uint32 InC, uint32 InD)
+	explicit constexpr FGuid(uint32 InA, uint32 InB, uint32 InC, uint32 InD)
 		: A(InA), B(InB), C(InC), D(InD)
 	{ }
 
