Replaced use of sa_handler with sa_sigaction
* If SA_SIGINFO is specified in sa_flags, then sa_sigaction (instead of sa_handler) specifies the signal-handling function
Added sigaction failure check

diff --git a/Engine/Source/Runtime/Core/Private/Unix/UnixSignalHeartBeat.cpp b/Engine/Source/Runtime/Core/Private/Unix/UnixSignalHeartBeat.cpp
--- a/Engine/Source/Runtime/Core/Private/Unix/UnixSignalHeartBeat.cpp
+++ b/Engine/Source/Runtime/Core/Private/Unix/UnixSignalHeartBeat.cpp
@@ -53,7 +53,7 @@ namespace
 	// 1ms for the lowest amount allowed for hitch detection. Anything less we wont try to detect hitches
 	double MinimalHitchThreashold = 0.001;
 
-	void SignalHitchHandler(int Signal)
+	void SignalHitchHandler(int Signal, siginfo_t* info, void* context)
 	{
 #if USE_HITCH_DETECTION
 		UE_LOG(LogUnixHeartBeat, Verbose, TEXT("SignalHitchHandler"));
@@ -98,9 +98,14 @@ void FUnixSignalGameHitchHeartBeat::Init()
 	struct sigaction SigAction;
 	FMemory::Memzero(SigAction);
 	SigAction.sa_flags = SA_SIGINFO;
-	SigAction.sa_handler = SignalHitchHandler;
+	SigAction.sa_sigaction = SignalHitchHandler;
 
-	sigaction(HEART_BEAT_SIGNAL, &SigAction, nullptr);
+	const bool bActionCreated = sigaction(HEART_BEAT_SIGNAL, &SigAction, nullptr) == 0;
+	if (!bActionCreated)
+	{
+		const int Errno = errno;
+		UE_LOG(LogUnixHeartBeat, Warning, TEXT("Failed to sigaction() errno=%d (%s)"), Errno, UTF8_TO_TCHAR(strerror(Errno)));
+	}
 
 	struct sigevent SignalEvent;
 	FMemory::Memzero(SignalEvent);
@@ -110,7 +115,7 @@ void FUnixSignalGameHitchHeartBeat::Init()
 	bTimerCreated = timer_create(CLOCK_REALTIME, &SignalEvent, &TimerId) == 0;
 	if (!bTimerCreated)
 	{
-		int Errno = errno;
+		const int Errno = errno;
 		UE_LOG(LogUnixHeartBeat, Warning, TEXT("Failed to timer_create() errno=%d (%s)"), Errno, UTF8_TO_TCHAR(strerror(Errno)));
 	}
 
