Add UE::ComparisonHelpers::CompareNaturalOrder, use it in scene outliner and changelist window

diff --git a/Engine/Source/Editor/SceneOutliner/Public/SortHelper.h b/Engine/Source/Editor/SceneOutliner/Public/SortHelper.h
--- a/Engine/Source/Editor/SceneOutliner/Public/SortHelper.h
+++ b/Engine/Source/Editor/SceneOutliner/Public/SortHelper.h
@@ -5,6 +5,7 @@
 #include "CoreMinimal.h"
 #include "Widgets/Views/SHeaderRow.h"
 #include "ISceneOutlinerTreeItem.h"
+#include "Misc/ComparisonUtility.h"
 
 namespace SceneOutliner
 {
@@ -28,108 +29,17 @@ namespace SceneOutliner
 		}
 		FORCEINLINE bool operator<(const FNumericStringWrapper& Other) const
 		{
-			return CompareNumeric(String, Other.String) < 0;
+			return UE::ComparisonUtility::CompareNaturalOrder(String, Other.String) < 0;
 		}
 		FORCEINLINE bool operator>(const FNumericStringWrapper& Other) const
 		{
-			return CompareNumeric(String, Other.String) > 0;
+			return UE::ComparisonUtility::CompareNaturalOrder(String, Other.String) > 0;
 		}
 		FORCEINLINE bool operator==(const FNumericStringWrapper& Other) const
 		{
 			return String == Other.String;
 		}
-
 	private:
-		static bool BothAscii(TCHAR C1, TCHAR C2)
-		{
-			return (((uint32)C1 | (uint32)C2) & 0xffffff80) == 0;
-		}
-
-		static bool BothNumbers(TCHAR C1, TCHAR C2)
-		{
-			return C1 >= '0' && C1 <= '9' && C2 >= '0' && C2 <= '9';
-		}
-
-		static int32 CompareNumeric(const FString& A, const FString& B)
-		{
-			const TCHAR* String1 = *A;
-			const TCHAR* String2 = *B;
-
-			while (true)
-			{
-				TCHAR C1 = *String1;
-				TCHAR C2 = *String2;
-
-				// Ignore underscores
-				if (TChar<TCHAR>::IsUnderscore(C1))
-				{
-					String1++;
-					continue;
-				}
-
-				// Ignore underscores
-				if (TChar<TCHAR>::IsUnderscore(C2))
-				{
-					String2++;
-					continue;
-				}
-
-				// Sort numerically when numbers are found 
-				if (BothNumbers(C1, C2))
-				{
-					int64 IntA = 0;
-					do
-					{
-						IntA *= 10;
-						IntA += TChar<TCHAR>::ConvertCharDigitToInt(C1);
-
-						String1++;
-						C1 = *String1;
-					} while (C1 >= '0' && C1 <= '9');
-
-					int64 IntB = 0;
-					do
-					{
-						IntB *= 10;
-						IntB += TChar<TCHAR>::ConvertCharDigitToInt(C2);
-
-						String2++;
-						C2 = *String2;
-					} while (C2 >= '0' && C2 <= '9');
-
-					if (IntA != IntB)
-					{
-						return IntA < IntB ? -1 : 1;
-					}
-
-					continue;
-				}
-				else if (C1 == C2)
-				{
-					// Reached the end of the string
-					if (C1 == 0)
-					{
-						// Strings compared equal, return shortest first
-						return A.Len() == B.Len() ? 0 : A.Len() < B.Len() ? -1 : 1;
-					}
-				}
-				else if (BothAscii(C1, C2))
-				{
-					if (int32 Diff = TChar<TCHAR>::ToUnsigned(TChar<TCHAR>::ToLower(C1)) - TChar<TCHAR>::ToUnsigned(TChar<TCHAR>::ToLower(C2)))
-					{
-						return Diff;
-					}
-				}
-				else
-				{
-					return TChar<TCHAR>::ToUnsigned(C1) - TChar<TCHAR>::ToUnsigned(C2);
-				}
-
-				String1++;
-				String2++;
-			}
-		}
-
 		FNumericStringWrapper(const FNumericStringWrapper&);
 		FNumericStringWrapper& operator=(const FNumericStringWrapper&);
 	};
diff --git a/Engine/Source/Editor/SourceControlWindows/Private/SSourceControlChangelists.cpp b/Engine/Source/Editor/SourceControlWindows/Private/SSourceControlChangelists.cpp
--- a/Engine/Source/Editor/SourceControlWindows/Private/SSourceControlChangelists.cpp
+++ b/Engine/Source/Editor/SourceControlWindows/Private/SSourceControlChangelists.cpp
@@ -51,6 +51,7 @@
 #include "Framework/Notifications/NotificationManager.h"
 #include "Framework/Docking/TabManager.h"
 #include "Framework/MultiBox/MultiBoxBuilder.h"
+#include "Misc/ComparisonUtility.h"
 #include "ProfilingDebugging/CpuProfilerTrace.h"
 
 #define LOCTEXT_NAMESPACE "SourceControlChangelist"
@@ -3573,7 +3574,7 @@ void SSourceControlChangelistsWidget::SortFileView()
 
 	auto CompareNames = [](const IFileViewTreeItem* Lhs, const IFileViewTreeItem* Rhs)
 	{
-		return FCString::Stricmp(*Lhs->GetName(), *Rhs->GetName());
+		return UE::ComparisonUtility::CompareNaturalOrder(*Lhs->GetName(), *Rhs->GetName());
 	};
 
 	auto ComparePaths = [](const IFileViewTreeItem* Lhs, const IFileViewTreeItem* Rhs)
diff --git a/Engine/Source/Runtime/Core/Private/Misc/ComparisonUtility.cpp b/Engine/Source/Runtime/Core/Private/Misc/ComparisonUtility.cpp
--- a/Engine/Source/Runtime/Core/Private/Misc/ComparisonUtility.cpp
+++ b/Engine/Source/Runtime/Core/Private/Misc/ComparisonUtility.cpp
@@ -6,7 +6,7 @@
 #include "CoreTypes.h"
 #include "Misc/CString.h"
 
-namespace UE { namespace ComparisonUtility {
+namespace UE::ComparisonUtility {
 
 int32 CompareWithNumericSuffix(FName A, FName B)
 {
@@ -17,24 +17,14 @@ int32 CompareWithNumericSuffix(FStringView A, FStringView B)
 {
 	auto SplitNumericSuffix = [](FStringView Full, FStringView& OutPrefix, int64& OutSuffix)
 	{
-		auto IsDigit = [](const TCHAR C)
-		{
-			return C >= TEXT('0') && C <= TEXT('9');
-		};
-
-		auto CharToDigit = [](const TCHAR C) -> int32
-		{
-			return C - '0';
-		};
-
 		int32 NumDigits = 0;
 		OutSuffix = 0;
 		{
 			int32 Magnitude = 1;
 			const TCHAR* FullStart = Full.GetData();
-			for (const TCHAR* C = FullStart + Full.Len() - 1; C >= FullStart && IsDigit(*C); --C)
+			for (const TCHAR* C = FullStart + Full.Len() - 1; C >= FullStart && FChar::IsDigit(*C); --C)
 			{
-				OutSuffix += CharToDigit(*C) * Magnitude;
+				OutSuffix += FChar::ConvertCharDigitToInt(*C) * Magnitude;
 				Magnitude *= 10;
 				++NumDigits;
 			}
@@ -86,4 +76,96 @@ int32 CompareWithNumericSuffix(FStringView A, FStringView B)
 	return 0;
 }
 
-} } // namespace UE::ComparisonUtility
+int32 CompareNaturalOrder(FStringView A, FStringView B)
+{
+	auto BothAscii = [](TCHAR C1, TCHAR C2) -> bool
+	{
+		return ((static_cast<uint32>(C1) | static_cast<uint32>(C2)) & 0xffffff80) == 0;
+	};
+
+	auto ConsumeNumber = [](const TCHAR*& Str) -> int64
+	{
+		bool bOverflow = false;
+		uint64 Number = 0;
+		TCHAR C = *Str;
+		do
+		{
+			Number *= 10;
+			Number += FChar::ConvertCharDigitToInt(C);
+
+			bOverflow = bOverflow || Number > static_cast<uint64>(MAX_int64);
+			Str++;
+			C = *Str;
+		} while (FChar::IsDigit(C));
+
+		if (bOverflow)
+		{
+			// Too large to process as a number
+			Number = INDEX_NONE;
+		}
+
+		return Number;
+	};
+	
+	const TCHAR* StringA = A.GetData();
+	const TCHAR* StringB = B.GetData();
+ 
+	while (true)
+	{
+		TCHAR CharA = *StringA;
+		TCHAR CharB = *StringB;
+ 
+		// Ignore underscores
+		if (FChar::IsUnderscore(CharA))
+		{
+			StringA++;
+			continue;
+		}
+ 
+		// Ignore underscores
+		if (FChar::IsUnderscore(CharB))
+		{
+			StringB++;
+			continue;
+		}
+ 
+		// Sort numerically when numbers are found 
+		if (FChar::IsDigit(CharA) && FChar::IsDigit(CharB))
+		{
+			const int64 IntA = ConsumeNumber(StringA);
+			const int64 IntB = ConsumeNumber(StringB);
+ 
+			if (IntA != IntB)
+			{
+				return IntA < IntB ? -1 : 1;
+			}
+ 
+			continue;
+		}
+		else if (CharA == CharB)
+		{
+			// Reached the end of the string
+			if (CharA == TCHAR('\0'))
+			{
+				// Strings compared equal, return shortest first
+				return A.Len() == B.Len() ? 0 : A.Len() < B.Len() ? -1 : 1;
+			}
+		}
+		else if (BothAscii(CharA, CharB))
+		{
+			if (int32 Diff = FChar::ToUnsigned(FChar::ToLower(CharA)) - FChar::ToUnsigned(FChar::ToLower(CharB)))
+			{
+				return Diff;
+			}
+		}
+		else
+		{
+			return FChar::ToUnsigned(CharA) - FChar::ToUnsigned(CharB);
+		}
+ 
+		StringA++;
+		StringB++;
+	}
+}
+	
+} // namespace UE::ComparisonUtility
diff --git a/Engine/Source/Runtime/Core/Public/Misc/ComparisonUtility.h b/Engine/Source/Runtime/Core/Public/Misc/ComparisonUtility.h
--- a/Engine/Source/Runtime/Core/Public/Misc/ComparisonUtility.h
+++ b/Engine/Source/Runtime/Core/Public/Misc/ComparisonUtility.h
@@ -7,7 +7,7 @@
 #include "HAL/Platform.h"
 #include "UObject/NameTypes.h"
 
-namespace UE { namespace ComparisonUtility {
+namespace UE::ComparisonUtility {
 
 /** Compare the two names, correctly ordering any numeric suffixes they may have */
 CORE_API int32 CompareWithNumericSuffix(FName A, FName B);
@@ -15,4 +15,7 @@ CORE_API int32 CompareWithNumericSuffix(FName A, FName B);
 /** Compare the two strings, correctly ordering any numeric suffixes they may have */
 CORE_API int32 CompareWithNumericSuffix(FStringView A, FStringView B);
 
-} } // namespace UE::ComparisonUtility
+/** Compare the two strings, correctly ordering any numeric components they may have */
+CORE_API int32 CompareNaturalOrder(FStringView A, FStringView B);
+
+} // namespace UE::ComparisonUtility
