Change Oodle encoder scratch allocation to accomodate other codecs.

Currently sized for Mermaid VeryFast only, but we use Kraken Fast in certain places,
and bumping it up all the way to Kraken Normal is less than a 2x increase in memory
use (~1.2MB/scratch buffer to ~2.2MB in the editor; not used in game)

diff --git a/Engine/Source/Runtime/Core/Private/Compression/OodleDataCompression.cpp b/Engine/Source/Runtime/Core/Private/Compression/OodleDataCompression.cpp
--- a/Engine/Source/Runtime/Core/Private/Compression/OodleDataCompression.cpp
+++ b/Engine/Source/Runtime/Core/Private/Compression/OodleDataCompression.cpp
@@ -289,9 +289,10 @@ struct OodleScratchBuffers
 		else
 		{
 			// tools, like UnrealPak or DDC commandlets
-
-			OO_SINTa EncoderMemorySizeNeeded = OodleLZ_GetCompressScratchMemBound(OodleLZ_Compressor_Mermaid, OodleLZ_CompressionLevel_VeryFast, 256*1024, nullptr);
-			// DecoderMemorySizeNeeded is ~ 450000 , EncoderMemorySizeNeeded is ~ 1200000
+			// allocate scratch space for Kraken Normal, that should get us all the usual cases.
+			// (Mermaid and lower compression levels need less.)
+			OO_SINTa EncoderMemorySizeNeeded = OodleLZ_GetCompressScratchMemBound(OodleLZ_Compressor_Kraken, OodleLZ_CompressionLevel_Normal, 256*1024, nullptr);
+			// DecoderMemorySizeNeeded is ~ 450k , EncoderMemorySizeNeeded is ~ 2160k
 			OodleScratchMemorySize = DecoderMemorySizeNeeded > EncoderMemorySizeNeeded ? DecoderMemorySizeNeeded : EncoderMemorySizeNeeded;
 
 			// allow one scratch buffer per logical core
