Fix for trace but no logging configurations

diff --git a/Engine/Source/Runtime/Core/Private/ProfilingDebugging/TraceAuxiliary.cpp b/Engine/Source/Runtime/Core/Private/ProfilingDebugging/TraceAuxiliary.cpp
--- a/Engine/Source/Runtime/Core/Private/ProfilingDebugging/TraceAuxiliary.cpp
+++ b/Engine/Source/Runtime/Core/Private/ProfilingDebugging/TraceAuxiliary.cpp
@@ -82,7 +82,11 @@ enum class ETraceConnectType
 };
 
 ////////////////////////////////////////////////////////////////////////////////
+#if NO_LOGGING
+FTraceAuxiliary::FLogCategoryAlias LogTrace;
+#else
 DEFINE_LOG_CATEGORY_STATIC(LogTrace, Log, All);
+#endif
 
 ////////////////////////////////////////////////////////////////////////////////
 class FTraceAuxiliaryImpl
