CSV Profiler: insights support for -csvNamedEvents
Also add a CSV_  prefix to csv named events so they can be distinguished from ordinary named events

diff --git a/Engine/Source/Runtime/Core/Private/ProfilingDebugging/CsvProfiler.cpp b/Engine/Source/Runtime/Core/Private/ProfilingDebugging/CsvProfiler.cpp
--- a/Engine/Source/Runtime/Core/Private/ProfilingDebugging/CsvProfiler.cpp
+++ b/Engine/Source/Runtime/Core/Private/ProfilingDebugging/CsvProfiler.cpp
@@ -26,6 +26,7 @@
 #include "Misc/App.h"
 #include "HAL/Runnable.h"
 #include "Misc/EngineVersion.h"
+#include "ProfilingDebugging/CpuProfilerTrace.h"
 #include "Stats/Stats.h"
 #include "Stats/Stats2.h"
 #include "HAL/LowLevelMemTracker.h"
@@ -279,14 +280,41 @@ static FCsvPersistentCustomStats GCsvPersistentCustomStats;
 
 
 #if CSV_PROFILER_SUPPORT_NAMED_EVENTS
-  #if PLATFORM_IMPLEMENTS_BeginNamedEventStatic
-    #define CSV_PROFILER_BeginNamedEvent(Color,Text) FPlatformMisc::BeginNamedEventStatic(Color, Text)
-  #else
-    #define CSV_PROFILER_BeginNamedEvent(Color,Text) FPlatformMisc::BeginNamedEvent(Color, Text)
-  #endif
-
 bool GCsvProfilerNamedEventsExclusive = false;
 bool GCsvProfilerNamedEventsTiming = false;
+
+void CsvBeginNamedEvent(FColor Color, const char* NamedEventName)
+{
+#if CPUPROFILERTRACE_ENABLED 
+	if (UE_TRACE_CHANNELEXPR_IS_ENABLED(CpuChannel))
+	{
+		FCpuProfilerTrace::OutputBeginDynamicEvent(NamedEventName);
+	}
+	else
+#endif
+	{
+#if PLATFORM_IMPLEMENTS_BeginNamedEventStatic
+		FPlatformMisc::BeginNamedEventStatic(Color, NamedEventName);
+#else
+		FPlatformMisc::BeginNamedEvent(Color, NamedEventName);
+#endif
+	}
+}
+
+
+void CsvEndNamedEvent()
+{
+#if CPUPROFILERTRACE_ENABLED
+	if (UE_TRACE_CHANNELEXPR_IS_ENABLED(CpuChannel))
+	{
+		FCpuProfilerTrace::OutputEndEvent();
+	}
+	else
+#endif
+	{
+		FPlatformMisc::EndNamedEvent();
+	}
+}
 #endif //CSV_PROFILER_SUPPORT_NAMED_EVENTS
 
 
@@ -2838,7 +2866,7 @@ void FCsvProfiler::BeginFrame()
 			else
 			{
 				UE_LOG(LogCsvProfiler, Display, TEXT("Capture Starting"));
-				
+
 				// signal external profiler that we are capturing
 				OnCSVProfileStartDelegate.Broadcast();
 
@@ -3312,7 +3340,7 @@ void FCsvProfiler::SetDeviceProfileName(FString InDeviceProfileName)
 }
 
 /** Push/pop events */
-void FCsvProfiler::BeginStat(const char * StatName, uint32 CategoryIndex)
+void FCsvProfiler::BeginStat(const char * StatName, uint32 CategoryIndex, const char * NamedEventName)
 {
 #if RECORD_TIMESTAMPS
 	if (GCsvProfilerIsCapturing && GCsvCategoriesEnabled[CategoryIndex])
@@ -3320,7 +3348,7 @@ void FCsvProfiler::BeginStat(const char * StatName, uint32 CategoryIndex)
 #if CSV_PROFILER_SUPPORT_NAMED_EVENTS
 		if (UNLIKELY(GCsvProfilerNamedEventsTiming))
 		{
-			CSV_PROFILER_BeginNamedEvent(FColor(255, 128, 255), StatName);
+			CsvBeginNamedEvent(FColor(255, 128, 255), NamedEventName ? NamedEventName : StatName);
 		}
 #endif
 		FCsvProfilerThreadData::Get().AddTimestampBegin(StatName, CategoryIndex);
@@ -3347,7 +3375,7 @@ void FCsvProfiler::EndStat(const char * StatName, uint32 CategoryIndex)
 #if CSV_PROFILER_SUPPORT_NAMED_EVENTS
 		if (UNLIKELY(GCsvProfilerNamedEventsTiming))
 		{
-			FPlatformMisc::EndNamedEvent();
+			CsvEndNamedEvent();
 		}
 #endif
 	}
@@ -3364,7 +3392,7 @@ void FCsvProfiler::EndStat(const FName& StatName, uint32 CategoryIndex)
 #endif
 }
 
-void FCsvProfiler::BeginExclusiveStat(const char * StatName)
+void FCsvProfiler::BeginExclusiveStat(const char * StatName, const char * NamedEventName)
 {
 #if RECORD_TIMESTAMPS
 	if (GCsvProfilerIsCapturing && GCsvCategoriesEnabled[CSV_CATEGORY_INDEX(Exclusive)])
@@ -3372,7 +3400,7 @@ void FCsvProfiler::BeginExclusiveStat(const char * StatName)
 #if CSV_PROFILER_SUPPORT_NAMED_EVENTS
 		if (UNLIKELY(GCsvProfilerNamedEventsExclusive))
 		{
-			CSV_PROFILER_BeginNamedEvent(FColor(255, 128, 128), StatName);
+			CsvBeginNamedEvent(FColor(255, 128, 128), NamedEventName ? NamedEventName : StatName);
 		}
 #endif
 		FCsvProfilerThreadData::Get().AddTimestampExclusiveBegin(StatName);
@@ -3390,7 +3418,7 @@ void FCsvProfiler::EndExclusiveStat(const char * StatName)
 #if CSV_PROFILER_SUPPORT_NAMED_EVENTS
 		if (UNLIKELY(GCsvProfilerNamedEventsExclusive))
 		{
-			FPlatformMisc::EndNamedEvent();
+			CsvEndNamedEvent();
 		}
 #endif
 	}
@@ -3443,11 +3471,11 @@ void FCsvProfiler::BeginWait()
 			{
 				if ( FThreadIdleStats::Get().IsCriticalPath() )
 				{
-					CSV_PROFILER_BeginNamedEvent(FColor(192, 96, 96), "CsvEventWait");
+					CsvBeginNamedEvent(FColor(192, 96, 96), "CsvEventWait");
 				}
 				else
 				{
-					CSV_PROFILER_BeginNamedEvent(FColor(255, 128, 128), "CsvEventWait (Non-CP)");
+					CsvBeginNamedEvent(FColor(255, 128, 128), "CsvEventWait (Non-CP)");
 				}
 			}
 #endif
@@ -3469,7 +3497,7 @@ void FCsvProfiler::EndWait()
 #if CSV_PROFILER_SUPPORT_NAMED_EVENTS
 			if (UNLIKELY(GCsvProfilerNamedEventsExclusive))
 			{
-				FPlatformMisc::EndNamedEvent();
+				CsvEndNamedEvent();
 			}
 #endif
 		}
diff --git a/Engine/Source/Runtime/Core/Public/ProfilingDebugging/CsvProfiler.h b/Engine/Source/Runtime/Core/Public/ProfilingDebugging/CsvProfiler.h
--- a/Engine/Source/Runtime/Core/Public/ProfilingDebugging/CsvProfiler.h
+++ b/Engine/Source/Runtime/Core/Public/ProfilingDebugging/CsvProfiler.h
@@ -56,16 +56,16 @@ struct FCsvDeclaredStat;
 // Inline stats (no up front definition)
 #define CSV_SCOPED_TIMING_STAT(Category,StatName) \
 	TRACE_CSV_PROFILER_INLINE_STAT(#StatName, CSV_CATEGORY_INDEX(Category)); \
-	FScopedCsvStat _ScopedCsvStat_ ## StatName (#StatName, CSV_CATEGORY_INDEX(Category));
+	FScopedCsvStat _ScopedCsvStat_ ## StatName (#StatName, CSV_CATEGORY_INDEX(Category), "CSV_"#StatName);
 #define CSV_SCOPED_TIMING_STAT_GLOBAL(StatName) \
 	TRACE_CSV_PROFILER_INLINE_STAT(#StatName, CSV_CATEGORY_INDEX_GLOBAL); \
-	FScopedCsvStat _ScopedCsvStat_ ## StatName (#StatName, CSV_CATEGORY_INDEX_GLOBAL);
+	FScopedCsvStat _ScopedCsvStat_ ## StatName (#StatName, CSV_CATEGORY_INDEX_GLOBAL, "CSV_"#StatName);
 #define CSV_SCOPED_TIMING_STAT_EXCLUSIVE(StatName) \
 	TRACE_CSV_PROFILER_INLINE_STAT_EXCLUSIVE(#StatName); \
-	FScopedCsvStatExclusive _ScopedCsvStatExclusive_ ## StatName (#StatName);
+	FScopedCsvStatExclusive _ScopedCsvStatExclusive_ ## StatName (#StatName, "CSV_"#StatName);
 #define CSV_SCOPED_TIMING_STAT_EXCLUSIVE_CONDITIONAL(StatName,Condition) \
 	TRACE_CSV_PROFILER_INLINE_STAT_EXCLUSIVE(#StatName); \
-	FScopedCsvStatExclusiveConditional _ScopedCsvStatExclusive_ ## StatName (#StatName,Condition);
+	FScopedCsvStatExclusiveConditional _ScopedCsvStatExclusive_ ## StatName (#StatName,Condition, "CSV_"#StatName);
 
 #define CSV_SCOPED_WAIT(WaitTime)							FScopedCsvWaitConditional _ScopedCsvWait(WaitTime>0 && FCsvProfiler::IsWaitTrackingEnabledOnCurrentThread());
 #define CSV_SCOPED_WAIT_CONDITIONAL(Condition)				FScopedCsvWaitConditional _ScopedCsvWait(Condition);
@@ -344,12 +344,12 @@ public:
 
 	/** Begin static interface (used by macros)*/
 	/** Push/pop events */
-	CORE_API static void BeginStat(const char* StatName, uint32 CategoryIndex);
+	CORE_API static void BeginStat(const char* StatName, uint32 CategoryIndex, const char* NamedEventName = nullptr);
 	CORE_API static void BeginStat(const FName& StatName, uint32 CategoryIndex);
 	CORE_API static void EndStat(const char* StatName, uint32 CategoryIndex);
 	CORE_API static void EndStat(const FName& StatName, uint32 CategoryIndex);
 
-	CORE_API static void BeginExclusiveStat(const char * StatName);
+	CORE_API static void BeginExclusiveStat(const char * StatName, const char* NamedEventName = nullptr);
 	CORE_API static void EndExclusiveStat(const char * StatName);
 
 	CORE_API static void RecordCustomStat(const char * StatName, uint32 CategoryIndex, float Value, const ECsvCustomStatOp CustomStatOp);
@@ -523,11 +523,11 @@ private:
 class FScopedCsvStat
 {
 public:
-	FScopedCsvStat(const char * InStatName, uint32 InCategoryIndex)
+	FScopedCsvStat(const char * InStatName, uint32 InCategoryIndex, const char * InNamedEventName = nullptr)
 		: StatName(InStatName)
 		, CategoryIndex(InCategoryIndex)
 	{
-		FCsvProfiler::BeginStat(StatName, CategoryIndex);
+		FCsvProfiler::BeginStat(StatName, CategoryIndex, InNamedEventName);
 	}
 
 	~FScopedCsvStat()
@@ -541,10 +541,10 @@ public:
 class FScopedCsvStatExclusive 
 {
 public:
-	FScopedCsvStatExclusive(const char * InStatName)
+	FScopedCsvStatExclusive(const char * InStatName, const char* InNamedEventName = nullptr)
 		: StatName(InStatName)
 	{
-		FCsvProfiler::BeginExclusiveStat(StatName);
+		FCsvProfiler::BeginExclusiveStat(StatName, InNamedEventName);
 	}
 
 	~FScopedCsvStatExclusive()
@@ -557,13 +557,13 @@ public:
 class FScopedCsvStatExclusiveConditional
 {
 public:
-	FScopedCsvStatExclusiveConditional(const char * InStatName, bool bInCondition)
+	FScopedCsvStatExclusiveConditional(const char * InStatName, bool bInCondition, const char* InNamedEventName = nullptr)
 		: StatName(InStatName)
 		, bCondition(bInCondition)
 	{
 		if (bCondition)
 		{
-			FCsvProfiler::BeginExclusiveStat(StatName);
+			FCsvProfiler::BeginExclusiveStat(StatName, InNamedEventName);
 		}
 	}
 
