Rename TBitArray's new FromAfter function to FindFrom. This name is more explicit about the search range being inclusive of the start index.

diff --git a/Engine/Source/Runtime/Core/Public/Containers/BitArray.h b/Engine/Source/Runtime/Core/Public/Containers/BitArray.h
--- a/Engine/Source/Runtime/Core/Public/Containers/BitArray.h
+++ b/Engine/Source/Runtime/Core/Public/Containers/BitArray.h
@@ -1107,7 +1107,7 @@ public:
 
 private:
 
-	int32 FindAfterImpl(bool bValue, int32 StartIndex) const
+	int32 FindFromImpl(bool bValue, int32 StartIndex) const
 	{
 		checkSlow((StartIndex >= 0) & (StartIndex <= NumBits));
 
@@ -1146,20 +1146,20 @@ public:
 
 	/**
 	 * Finds the first occurrence of the specified value (true/false) in the array, starting from the given bit index, and returns the bit index.
-	 * If the specified value is not found after the given index, INDEX_NONE is returned.
+	 * If the specified value is not found from the given index, INDEX_NONE is returned.
 	 *
 	 * @param  bValue      The value (true/false) to search for.
 	 * @param  StartIndex  The index to start the search from.
 	 *
 	 * @pre StartIndex is expected to be in the inclusive range [0, Num()].
 	 *
-	 * @return The index of the first occurrence of the specified value (true/false) after StartIndex, or INDEX_NONE if not found.
+	 * @return The index of the first occurrence of the specified value (true/false) from StartIndex, or INDEX_NONE if not found.
 	 */
 	template <typename IndexType>
-	FORCEINLINE int32 FindAfter(bool bValue, IndexType StartIndex) const
+	FORCEINLINE int32 FindFrom(bool bValue, IndexType StartIndex) const
 	{
-		static_assert(!std::is_same_v<IndexType, bool>, "TBitArray::FindAfter: unexpected bool passed as the StartIndex argument");
-		return FindAfterImpl(bValue, StartIndex);
+		static_assert(!std::is_same_v<IndexType, bool>, "TBitArray::FindFrom: unexpected bool passed as the StartIndex argument");
+		return FindFromImpl(bValue, StartIndex);
 	}
 
 	/**
diff --git a/Engine/Source/Runtime/Core/Tests/Containers/BitArrayTests.cpp b/Engine/Source/Runtime/Core/Tests/Containers/BitArrayTests.cpp
--- a/Engine/Source/Runtime/Core/Tests/Containers/BitArrayTests.cpp
+++ b/Engine/Source/Runtime/Core/Tests/Containers/BitArrayTests.cpp
@@ -2054,40 +2054,40 @@ TEST_CASE_NAMED(FBitArrayFindTest, "System::Core::Containers::BitArray::Find", "
 	}
 }
 
-TEST_CASE_NAMED(FBitArrayFindAfterTest, "System::Core::Containers::BitArray::FindAfter", "[ApplicationContextMask][SmokeFilter]")
+TEST_CASE_NAMED(FBitArrayFindFromTest, "System::Core::Containers::BitArray::FindFrom", "[ApplicationContextMask][SmokeFilter]")
 {
 	{
 		TBitArray<> ArrTrue2(true, NumBitsPerDWORD);
-		CHECK_MESSAGE(TEXT("Find true in an array full of true, starting at index 0"), ArrTrue2.FindAfter(true, 0) == 0);
-		CHECK_MESSAGE(TEXT("Find true in an array full of true, starting at index 1"), ArrTrue2.FindAfter(true, 1) == 1);
-		CHECK_MESSAGE(TEXT("Find true in an array full of true, starting at index 31"), ArrTrue2.FindAfter(true, NumBitsPerDWORD - 1) == NumBitsPerDWORD - 1);
-		CHECK_MESSAGE(TEXT("Find true in an array full of true, starting at index 32"), ArrTrue2.FindAfter(true, NumBitsPerDWORD) == INDEX_NONE);
-		CHECK_MESSAGE(TEXT("Find false in an array full of true, starting at index 1"), ArrTrue2.FindAfter(false, 1) == INDEX_NONE);
-		CHECK_MESSAGE(TEXT("Find false in an array full of true, starting at index 32"), ArrTrue2.FindAfter(false, NumBitsPerDWORD) == INDEX_NONE);
+		CHECK_MESSAGE(TEXT("Find true in an array full of true, starting at index 0"), ArrTrue2.FindFrom(true, 0) == 0);
+		CHECK_MESSAGE(TEXT("Find true in an array full of true, starting at index 1"), ArrTrue2.FindFrom(true, 1) == 1);
+		CHECK_MESSAGE(TEXT("Find true in an array full of true, starting at index 31"), ArrTrue2.FindFrom(true, NumBitsPerDWORD - 1) == NumBitsPerDWORD - 1);
+		CHECK_MESSAGE(TEXT("Find true in an array full of true, starting at index 32"), ArrTrue2.FindFrom(true, NumBitsPerDWORD) == INDEX_NONE);
+		CHECK_MESSAGE(TEXT("Find false in an array full of true, starting at index 1"), ArrTrue2.FindFrom(false, 1) == INDEX_NONE);
+		CHECK_MESSAGE(TEXT("Find false in an array full of true, starting at index 32"), ArrTrue2.FindFrom(false, NumBitsPerDWORD) == INDEX_NONE);
 	}
 
 	{
 		TBitArray<> ArrFalse2(false, NumBitsPerDWORD);
-		CHECK_MESSAGE(TEXT("Find false in an array full of false, starting at index 0"), ArrFalse2.FindAfter(false, 0) == 0);
-		CHECK_MESSAGE(TEXT("Find false in an array full of false, starting at index 1"), ArrFalse2.FindAfter(false, 1) == 1);
-		CHECK_MESSAGE(TEXT("Find false in an array full of false, starting at index 31"), ArrFalse2.FindAfter(false, NumBitsPerDWORD - 1) == NumBitsPerDWORD - 1);
-		CHECK_MESSAGE(TEXT("Find false in an array full of false, starting at index 32"), ArrFalse2.FindAfter(false, NumBitsPerDWORD) == INDEX_NONE);
-		CHECK_MESSAGE(TEXT("Find true in an array full of false, starting at index 1"), ArrFalse2.FindAfter(true, 1) == INDEX_NONE);
-		CHECK_MESSAGE(TEXT("Find false in an array full of false, starting at index 32"), ArrFalse2.FindAfter(true, NumBitsPerDWORD) == INDEX_NONE);
+		CHECK_MESSAGE(TEXT("Find false in an array full of false, starting at index 0"), ArrFalse2.FindFrom(false, 0) == 0);
+		CHECK_MESSAGE(TEXT("Find false in an array full of false, starting at index 1"), ArrFalse2.FindFrom(false, 1) == 1);
+		CHECK_MESSAGE(TEXT("Find false in an array full of false, starting at index 31"), ArrFalse2.FindFrom(false, NumBitsPerDWORD - 1) == NumBitsPerDWORD - 1);
+		CHECK_MESSAGE(TEXT("Find false in an array full of false, starting at index 32"), ArrFalse2.FindFrom(false, NumBitsPerDWORD) == INDEX_NONE);
+		CHECK_MESSAGE(TEXT("Find true in an array full of false, starting at index 1"), ArrFalse2.FindFrom(true, 1) == INDEX_NONE);
+		CHECK_MESSAGE(TEXT("Find false in an array full of false, starting at index 32"), ArrFalse2.FindFrom(true, NumBitsPerDWORD) == INDEX_NONE);
 	}
 
 	{
 		TBitArray<> ArrFalseThenTrue(false, NumBitsPerDWORD);
 		ArrFalseThenTrue.Add(true, NumBitsPerDWORD);
-		CHECK_MESSAGE(TEXT("Find true in an array containing 32 false and then 32 true, starting after all of the false"), ArrFalseThenTrue.FindAfter(true, NumBitsPerDWORD) == NumBitsPerDWORD);
-		CHECK_MESSAGE(TEXT("Find false in an array containing 32 false and then 32 true, starting after all of the false"), ArrFalseThenTrue.FindAfter(false, NumBitsPerDWORD) == INDEX_NONE);
+		CHECK_MESSAGE(TEXT("Find true in an array containing 32 false and then 32 true, starting after all of the false"), ArrFalseThenTrue.FindFrom(true, NumBitsPerDWORD) == NumBitsPerDWORD);
+		CHECK_MESSAGE(TEXT("Find false in an array containing 32 false and then 32 true, starting after all of the false"), ArrFalseThenTrue.FindFrom(false, NumBitsPerDWORD) == INDEX_NONE);
 	}
 
 	{
 		TBitArray<> ArrTrueThenFalse(true, NumBitsPerDWORD);
 		ArrTrueThenFalse.Add(false, NumBitsPerDWORD);
-		CHECK_MESSAGE(TEXT("Find true in an array containing 32 true and then 32 false, starting after all of the true"), ArrTrueThenFalse.FindAfter(true, NumBitsPerDWORD) == INDEX_NONE);
-		CHECK_MESSAGE(TEXT("Find false in an array containing 32 true and then 32 false, starting after all of the true"), ArrTrueThenFalse.FindAfter(false, NumBitsPerDWORD) == NumBitsPerDWORD);
+		CHECK_MESSAGE(TEXT("Find true in an array containing 32 true and then 32 false, starting after all of the true"), ArrTrueThenFalse.FindFrom(true, NumBitsPerDWORD) == INDEX_NONE);
+		CHECK_MESSAGE(TEXT("Find false in an array containing 32 true and then 32 false, starting after all of the true"), ArrTrueThenFalse.FindFrom(false, NumBitsPerDWORD) == NumBitsPerDWORD);
 	}
 }
 
