Add a [[clang::noautortfm]] attribute to allow marking of functions that should not be cloned for AutoRTFM.

diff --git a/Engine/Source/Programs/AutoRTFMTests/Private/SharedPointerTests.cpp b/Engine/Source/Programs/AutoRTFMTests/Private/SharedPointerTests.cpp
new file mode 100644
--- /dev/null
+++ b/Engine/Source/Programs/AutoRTFMTests/Private/SharedPointerTests.cpp
@@ -0,0 +1,188 @@
+// Copyright Epic Games, Inc. All Rights Reserved.
+
+#pragma autortfm
+
+#include "Catch2Includes.h"
+#include <AutoRTFM/AutoRTFM.h>
+#include "Templates/SharedPointer.h"
+
+TEST_CASE("SharedPointer.PreviouslyAllocated")
+{
+	TSharedPtr<int, ESPMode::ThreadSafe> Foo(new int(42));
+
+	AutoRTFM::Commit([&]
+		{
+			// Make a copy to bump the reference count.
+			TSharedPtr<int, ESPMode::ThreadSafe> Copy = Foo;
+
+			*Copy = 13;
+		});
+
+	REQUIRE(13 == *Foo);
+	REQUIRE(1 == Foo.GetSharedReferenceCount());
+}
+
+TEST_CASE("SharedPointer.AbortWithPreviouslyAllocated")
+{
+	TSharedPtr<int, ESPMode::ThreadSafe> Foo(new int(42));
+
+	REQUIRE(AutoRTFM::ETransactionResult::AbortedByRequest == AutoRTFM::Transact([&]
+		{
+			// Make a copy to bump the reference count.
+			TSharedPtr<int, ESPMode::ThreadSafe> Copy = Foo;
+
+			*Copy = 13;
+
+			AutoRTFM::AbortTransaction();
+		}));
+
+	REQUIRE(42 == *Foo);
+	REQUIRE(1 == Foo.GetSharedReferenceCount());
+}
+
+TEST_CASE("SharedPointer.NewlyAllocated")
+{
+	int Copy = 42;
+
+	AutoRTFM::Commit([&]
+		{
+			TSharedPtr<int, ESPMode::ThreadSafe> Foo(new int(13));
+
+			Copy = *Foo;
+		});
+
+	REQUIRE(13 == Copy);
+}
+
+TEST_CASE("SharedPointer.AbortWithNewlyAllocated")
+{
+	int Result = 42;
+
+	REQUIRE(AutoRTFM::ETransactionResult::AbortedByRequest == AutoRTFM::Transact([&]
+		{
+			TSharedPtr<int, ESPMode::ThreadSafe> Foo(new int(13));
+			TSharedPtr<int, ESPMode::ThreadSafe> Copy = Foo;
+			Result = *Copy;
+			AutoRTFM::AbortTransaction();
+		}));
+
+	REQUIRE(42 == Result);
+}
+
+TEST_CASE("SharedPointer.NestedTransactionWithPreviouslyAllocated")
+{
+	TSharedPtr<int, ESPMode::ThreadSafe> Foo(new int(42));
+
+	AutoRTFM::Commit([&]
+		{
+			AutoRTFM::Commit([&]
+				{
+					// Make a copy to bump the reference count.
+					TSharedPtr<int, ESPMode::ThreadSafe> Copy = Foo;
+
+					*Copy = 13;
+				});
+		});
+
+	REQUIRE(13 == *Foo);
+	REQUIRE(1 == Foo.GetSharedReferenceCount());
+}
+
+TEST_CASE("SharedPointer.AbortNestedTransactionWithPreviouslyAllocated")
+{
+	TSharedPtr<int, ESPMode::ThreadSafe> Foo(new int(42));
+
+	AutoRTFM::Commit([&]
+		{
+			REQUIRE(AutoRTFM::ETransactionResult::AbortedByRequest == AutoRTFM::Transact([&]
+				{
+					// Make a copy to bump the reference count.
+					TSharedPtr<int, ESPMode::ThreadSafe> Copy = Foo;
+
+					*Copy = 13;
+
+					AutoRTFM::AbortTransaction();
+				}));
+		});
+
+	REQUIRE(42 == *Foo);
+	REQUIRE(1 == Foo.GetSharedReferenceCount());
+}
+
+TEST_CASE("SharedPointer.NestedTransactionWithNewlyAllocated")
+{
+	int Result = 42;
+
+	AutoRTFM::Commit([&]
+		{
+			AutoRTFM::Commit([&]
+				{
+					TSharedPtr<int, ESPMode::ThreadSafe> Foo(new int(13));
+					TSharedPtr<int, ESPMode::ThreadSafe> Copy = Foo;
+					Result = *Copy;
+				});
+		});
+
+	REQUIRE(13 == Result);
+}
+
+TEST_CASE("SharedPointer.AbortNestedTransactionWithNewlyAllocated")
+{
+	int Result = 42;
+
+	AutoRTFM::Commit([&]
+		{
+			REQUIRE(AutoRTFM::ETransactionResult::AbortedByRequest == AutoRTFM::Transact([&]
+				{
+					TSharedPtr<int, ESPMode::ThreadSafe> Foo(new int(13));
+					TSharedPtr<int, ESPMode::ThreadSafe> Copy = Foo;
+					Result = *Copy;
+					AutoRTFM::AbortTransaction();
+				}));
+		});
+
+	REQUIRE(42 == Result);
+}
+
+template<typename T> FORCENOINLINE void* MakeMemoryForT()
+{
+	return malloc(sizeof(T));
+}
+
+TEST_CASE("SharedPointer.NestedTransactionWithPlacementNewlyAllocated")
+{
+	int Result = 42;
+
+	AutoRTFM::Commit([&]
+		{
+			AutoRTFM::Commit([&]
+				{
+					void* const Memory = MakeMemoryForT<TSharedPtr<int, ESPMode::ThreadSafe>>();
+					TSharedPtr<int, ESPMode::ThreadSafe>* const Foo = new (Memory) TSharedPtr<int, ESPMode::ThreadSafe>(new int(13));
+					TSharedPtr<int, ESPMode::ThreadSafe> Copy = *Foo;
+					Result = *Copy;
+					free(Memory);
+				});
+		});
+
+	REQUIRE(13 == Result);
+}
+
+TEST_CASE("SharedPointer.AbortNestedTransactionWithPlacementNewlyAllocated")
+{
+	int Result = 42;
+
+	AutoRTFM::Commit([&]
+		{
+			REQUIRE(AutoRTFM::ETransactionResult::AbortedByRequest == AutoRTFM::Transact([&]
+				{
+					void* const Memory = MakeMemoryForT<TSharedPtr<int, ESPMode::ThreadSafe>>();
+					TSharedPtr<int, ESPMode::ThreadSafe>* const Foo = new (Memory) TSharedPtr<int, ESPMode::ThreadSafe>(new int(13));
+					TSharedPtr<int, ESPMode::ThreadSafe> Copy = *Foo;
+					Result = *Copy;
+					AutoRTFM::AbortTransaction();
+				}));
+		});
+
+	REQUIRE(42 == Result);
+}
diff --git a/Engine/Source/Runtime/Core/Private/AutoRTFM/API.cpp b/Engine/Source/Runtime/Core/Private/AutoRTFM/API.cpp
--- a/Engine/Source/Runtime/Core/Private/AutoRTFM/API.cpp
+++ b/Engine/Source/Runtime/Core/Private/AutoRTFM/API.cpp
@@ -46,7 +46,7 @@ UE_AUTORTFM_FORCEINLINE autortfm_result TransactThenOpenImpl(void (*Work)(void*
 		}));
 }
 
-FORCENOINLINE extern "C" bool autortfm_is_transactional()
+extern "C" UE_AUTORTFM_NOAUTORTFM bool autortfm_is_transactional()
 {
 	if (GAutoRTFMRuntimeEnabled)
 	{
@@ -56,13 +56,13 @@ FORCENOINLINE extern "C" bool autortfm_is_transactional()
 	return false;
 }
 
-FORCENOINLINE extern "C" bool autortfm_is_closed()
+extern "C" UE_AUTORTFM_NOAUTORTFM bool autortfm_is_closed()
 {
     return false;
 }
 
 // First Part - the API exposed outside transactions.
-FORCENOINLINE extern "C" autortfm_result autortfm_transact(void (*Work)(void* Arg), void* Arg)
+extern "C" UE_AUTORTFM_NOAUTORTFM autortfm_result autortfm_transact(void (*Work)(void* Arg), void* Arg)
 {
 	if (GAutoRTFMRuntimeEnabled)
 	{
@@ -73,48 +73,48 @@ FORCENOINLINE extern "C" autortfm_result autortfm_transact(void (*Work)(void* Ar
 	return autortfm_committed;
 }
 
-FORCENOINLINE extern "C" autortfm_result autortfm_transact_then_open(void (*Work)(void* Arg), void* Arg)
+extern "C" UE_AUTORTFM_NOAUTORTFM autortfm_result autortfm_transact_then_open(void (*Work)(void* Arg), void* Arg)
 {
     return TransactThenOpenImpl(Work, Arg);
 }
 
-FORCENOINLINE extern "C" void autortfm_commit(void (*Work)(void* Arg), void* Arg)
+extern "C" UE_AUTORTFM_NOAUTORTFM void autortfm_commit(void (*Work)(void* Arg), void* Arg)
 {
     autortfm_result Result = autortfm_transact(Work, Arg);
 	UE_CLOG(Result != autortfm_committed, LogAutoRTFM, Fatal, TEXT("Unexpected transaction result: %u."), Result);
 }
 
-FORCENOINLINE extern "C" void autortfm_abort()
+extern "C" UE_AUTORTFM_NOAUTORTFM void autortfm_abort()
 {
 	UE_CLOG(!FContext::IsTransactional(), LogAutoRTFM, Fatal, TEXT("The function `autortfm_abort` was called from outside a transaction."));
 	FContext::Get()->AbortByRequestAndThrow();
 }
 
-FORCENOINLINE extern "C" bool autortfm_start_transaction()
+extern "C" UE_AUTORTFM_NOAUTORTFM bool autortfm_start_transaction()
 {
 	UE_CLOG(!FContext::IsTransactional(), LogAutoRTFM, Fatal, TEXT("The function `autortfm_start_transaction` was called from outside a transact."));
 	return FContext::Get()->StartTransaction();
 }
 
-FORCENOINLINE extern "C" autortfm_result autortfm_commit_transaction()
+extern "C" UE_AUTORTFM_NOAUTORTFM autortfm_result autortfm_commit_transaction()
 {
 	UE_CLOG(!FContext::IsTransactional(), LogAutoRTFM, Fatal, TEXT("The function `autortfm_commit_transaction` was called from outside a transact."));
 	return static_cast<autortfm_result>(FContext::Get()->CommitTransaction());
 }
 
-FORCENOINLINE extern "C" autortfm_result autortfm_abort_transaction()
+extern "C" UE_AUTORTFM_NOAUTORTFM autortfm_result autortfm_abort_transaction()
 {
 	UE_CLOG(!FContext::IsTransactional(), LogAutoRTFM, Fatal, TEXT("The function `autortfm_abort_transaction` was called from outside a transact."));
 	return static_cast<autortfm_result>(FContext::Get()->AbortTransaction(false));
 }
 
-FORCENOINLINE extern "C" void autortfm_clear_transaction_status()
+extern "C" UE_AUTORTFM_NOAUTORTFM void autortfm_clear_transaction_status()
 {
 	ASSERT(FContext::Get()->IsAborting());
 	FContext::Get()->ClearTransactionStatus();
 }
 
-FORCENOINLINE extern "C" bool autortfm_is_aborting()
+extern "C" UE_AUTORTFM_NOAUTORTFM bool autortfm_is_aborting()
 {
 	if (GAutoRTFMRuntimeEnabled)
 	{
@@ -124,27 +124,27 @@ FORCENOINLINE extern "C" bool autortfm_is_aborting()
 	return false;
 }
 
-FORCENOINLINE extern "C" bool autortfm_current_nest_throw()
+extern "C" UE_AUTORTFM_NOAUTORTFM bool autortfm_current_nest_throw()
 {
 	FContext::Get()->Throw();
 	return true;
 }
 
-FORCENOINLINE extern "C" void autortfm_abort_if_transactional()
+extern "C" UE_AUTORTFM_NOAUTORTFM void autortfm_abort_if_transactional()
 {
 	UE_CLOG(FContext::IsTransactional(), LogAutoRTFM, Fatal, TEXT("The function `autortfm_abort_if_transactional` was called from an open inside a transaction."));
 }
 
-FORCENOINLINE extern "C" void autortfm_abort_if_closed()
+extern "C" UE_AUTORTFM_NOAUTORTFM void autortfm_abort_if_closed()
 {
 }
 
-FORCENOINLINE extern "C" void autortfm_open(void (*Work)(void* Arg), void* Arg)
+extern "C" UE_AUTORTFM_NOAUTORTFM void autortfm_open(void (*Work)(void* Arg), void* Arg)
 {
 	Work(Arg);
 }
 
-FORCENOINLINE extern "C" autortfm_status autortfm_close(void (*Work)(void* Arg), void* Arg)
+extern "C" UE_AUTORTFM_NOAUTORTFM autortfm_status autortfm_close(void (*Work)(void* Arg), void* Arg)
 {
 	autortfm_status Result = autortfm_status_ontrack;
 
@@ -167,42 +167,42 @@ FORCENOINLINE extern "C" autortfm_status autortfm_close(void (*Work)(void* Arg),
 	return Result;
 }
 
-FORCENOINLINE extern "C" void autortfm_record_open_write(void* Ptr, size_t Size)
+extern "C" UE_AUTORTFM_NOAUTORTFM void autortfm_record_open_write(void* Ptr, size_t Size)
 {
     FContext::Get()->CheckOpenRecordWrite(Ptr);
 	FContext::Get()->RecordWrite(Ptr, Size);
 }
 
-FORCENOINLINE extern "C" void autortfm_register_open_function(void* OriginalFunction, void* NewFunction)
+extern "C" UE_AUTORTFM_NOAUTORTFM void autortfm_register_open_function(void* OriginalFunction, void* NewFunction)
 {
 	UE_LOG(LogAutoRTFM, Verbose, TEXT("Registering open %p->%p"), OriginalFunction, NewFunction);
     FunctionMapAdd(OriginalFunction, NewFunction);
 }
 
-void OpenCommit(TFunction<void()>&& Work)
+UE_AUTORTFM_NOAUTORTFM void OpenCommit(TFunction<void()>&& Work)
 {
     Work();
 }
 
-void OpenAbort(TFunction<void()>&& Work)
+UE_AUTORTFM_NOAUTORTFM void OpenAbort(TFunction<void()>&& Work)
 {
 }
 
-FORCENOINLINE extern "C" void autortfm_open_commit(void (*Work)(void* Arg), void* Arg)
+extern "C" UE_AUTORTFM_NOAUTORTFM void autortfm_open_commit(void (*Work)(void* Arg), void* Arg)
 {
     Work(Arg);
 }
 
-FORCENOINLINE extern "C" void autortfm_open_abort(void (*Work)(void* arg), void* Arg)
+extern "C" UE_AUTORTFM_NOAUTORTFM void autortfm_open_abort(void (*Work)(void* arg), void* Arg)
 {
 }
 
-FORCENOINLINE extern "C" void* autortfm_did_allocate(void* Ptr, size_t Size)
+extern "C" UE_AUTORTFM_NOAUTORTFM void* autortfm_did_allocate(void* Ptr, size_t Size)
 {
     return Ptr;
 }
 
-FORCENOINLINE extern "C" void autortfm_check_consistency_assuming_no_races()
+extern "C" UE_AUTORTFM_NOAUTORTFM void autortfm_check_consistency_assuming_no_races()
 {
     if (FContext::IsTransactional())
     {
@@ -210,7 +210,7 @@ FORCENOINLINE extern "C" void autortfm_check_consistency_assuming_no_races()
     }
 }
 
-FORCENOINLINE extern "C" void autortfm_check_abi(void* const Ptr, const size_t Size)
+extern "C" UE_AUTORTFM_NOAUTORTFM void autortfm_check_abi(void* const Ptr, const size_t Size)
 {
     struct FConstants final
     {
diff --git a/Engine/Source/Runtime/Core/Public/AutoRTFM/AutoRTFM.h b/Engine/Source/Runtime/Core/Public/AutoRTFM/AutoRTFM.h
--- a/Engine/Source/Runtime/Core/Public/AutoRTFM/AutoRTFM.h
+++ b/Engine/Source/Runtime/Core/Public/AutoRTFM/AutoRTFM.h
@@ -14,6 +14,12 @@
 	#define UE_PRAGMA_AUTORTFM
 #endif
 
+#if UE_AUTORTFM
+#define UE_AUTORTFM_NOAUTORTFM [[clang::noautortfm, clang::noinline]]
+#else
+#define UE_AUTORTFM_NOAUTORTFM
+#endif
+
 #if defined(_MSC_VER)
 #define UE_AUTORTFM_FORCEINLINE __forceinline
 #elif defined(__clang__)
