LLM SizeParams migrated to remaining memory tracking functions. This will allow commands such as DumpLLM to report on the snapshot memory delta.
Deprecating older public LLM functions in favor of newer more controlled reporting.

diff --git a/Engine/Plugins/MemoryUsageQueries/Source/MemoryUsageQueries/Private/MemoryUsageInfoProvider.cpp b/Engine/Plugins/MemoryUsageQueries/Source/MemoryUsageQueries/Private/MemoryUsageInfoProvider.cpp
--- a/Engine/Plugins/MemoryUsageQueries/Source/MemoryUsageQueries/Private/MemoryUsageInfoProvider.cpp
+++ b/Engine/Plugins/MemoryUsageQueries/Source/MemoryUsageQueries/Private/MemoryUsageInfoProvider.cpp
@@ -7,7 +7,7 @@ uint64 FMemoryUsageInfoProviderLLM::GetAssetMemoryUsage(FName Asset, FOutputDevi
 #if ENABLE_LOW_LEVEL_MEM_TRACKER
 	if (FLowLevelMemTracker::IsEnabled())
 	{
-		return FLowLevelMemTracker::Get().GetTagAmountForTracker(ELLMTracker::Default, Asset, ELLMTagSet::Assets, false);
+		return FLowLevelMemTracker::Get().GetTagAmountForTracker(ELLMTracker::Default, Asset, ELLMTagSet::Assets, UE::LLM::ESizeParams::Default);
 	}
 #endif
 	ErrorOutput->Logf(TEXT("MemoryUsageInfoProvider Error: LLM is disabled. Please run with -LLM"));
@@ -22,7 +22,7 @@ uint64 FMemoryUsageInfoProviderLLM::GetAssetsMemoryUsage(const TSet<FName>& Asse
 		uint64 TotalSize = 0;
 		for (const FName& Asset : Assets)
 		{
-			TotalSize += FLowLevelMemTracker::Get().GetTagAmountForTracker(ELLMTracker::Default, Asset, ELLMTagSet::Assets, false);
+			TotalSize += FLowLevelMemTracker::Get().GetTagAmountForTracker(ELLMTracker::Default, Asset, ELLMTagSet::Assets, UE::LLM::ESizeParams::Default);
 		}
 		return TotalSize;
 	}
@@ -39,7 +39,7 @@ uint64 FMemoryUsageInfoProviderLLM::GetAssetsMemoryUsageWithSize(const TSet<FNam
 		uint64 TotalSize = 0;
 		for (const FName& Asset : Assets)
 		{
-			int64 AssetSize = FLowLevelMemTracker::Get().GetTagAmountForTracker(ELLMTracker::Default, Asset, ELLMTagSet::Assets, false);
+			int64 AssetSize = FLowLevelMemTracker::Get().GetTagAmountForTracker(ELLMTracker::Default, Asset, ELLMTagSet::Assets, UE::LLM::ESizeParams::Default);
 			OutSizes.Add(Asset, AssetSize);
 			TotalSize += AssetSize;
 		}
diff --git a/Engine/Source/Runtime/Core/Private/HAL/LowLevelMemTracker.cpp b/Engine/Source/Runtime/Core/Private/HAL/LowLevelMemTracker.cpp
--- a/Engine/Source/Runtime/Core/Private/HAL/LowLevelMemTracker.cpp
+++ b/Engine/Source/Runtime/Core/Private/HAL/LowLevelMemTracker.cpp
@@ -80,19 +80,21 @@ FAutoConsoleCommand DumpLLM(
 	FConsoleCommandWithWorldArgsAndOutputDeviceDelegate::CreateLambda(
 		[](const TArray<FString>& Args, UWorld* InWorld, FOutputDevice& Ar)
 	{
-		bool bCSV = false;
-		for (const FString& Arg : Args)
+		FString Command = FString::Join(Args, TEXT(","));
+
+		bool bCSV = FParse::Param(*Command, TEXT("CSV"));
+		bool bSnapshot = FParse::Param(*Command, TEXT("SNAPSHOT"));
+
+		UE::LLM::ESizeParams SizeParams(UE::LLM::ESizeParams::Default);
+		if (bSnapshot)
 		{
-			if (FParse::Param(*Arg, TEXT("CSV")))
-			{
-				bCSV = true;
-			}
+			EnumAddFlags(SizeParams, UE::LLM::ESizeParams::RelativeToSnapshot);
 		}
 		FLowLevelMemTracker::Get().DumpToLog(
-			bCSV ? FLowLevelMemTracker::EDumpFormat::CSV : FLowLevelMemTracker::EDumpFormat::PlainText, &Ar);
+			bCSV ? FLowLevelMemTracker::EDumpFormat::CSV : FLowLevelMemTracker::EDumpFormat::PlainText, &Ar, SizeParams);
 	}));
 
-void FLowLevelMemTracker::DumpToLog(EDumpFormat DumpFormat, FOutputDevice* OutputDevice)
+void FLowLevelMemTracker::DumpToLog(EDumpFormat DumpFormat, FOutputDevice* OutputDevice, UE::LLM::ESizeParams SizeParams)
 {
 	if (!IsEnabled())
 	{
@@ -119,11 +121,16 @@ void FLowLevelMemTracker::DumpToLog(EDumpFormat DumpFormat, FOutputDevice* Outpu
 			FString Line;
 			int64 Size;
 		};
+
+		UE::LLM::ESizeParams CurrentSizeParams = SizeParams;
+		UE::LLM::ESizeParams PeakSizeParams = CurrentSizeParams;
+		EnumAddFlags(PeakSizeParams, UE::LLM::ESizeParams::ReportPeak);
+
 		TArray<FTagLine> TagLines;
 		for (const UE::LLMPrivate::FTagData* TagData : GetTrackedTags(TrackerType))
 		{
-			int64 CurrentAmount = GetTagAmountForTracker(TrackerType, TagData, false /* bPeakAmount */);
-			int64 PeakAmount = GetTagAmountForTracker(TrackerType, TagData, true /* bPeakAmount */);
+			int64 CurrentAmount = GetTagAmountForTracker(TrackerType, TagData, SizeParams);
+			int64 PeakAmount = GetTagAmountForTracker(TrackerType, TagData, PeakSizeParams);
 			FString TagName = GetTagUniqueName(TagData).ToString();
 			FString TagLine;
 			if (DumpFormat == EDumpFormat::CSV)
@@ -307,7 +314,7 @@ public:
 
 	void Publish(FLowLevelMemTracker& LLMRef, const FTrackerTagSizeMap& TagSizes,
 		const FTagData* OverrideTrackedTotalTagData, const FTagData* OverrideUntaggedTagData, int64 TrackedTotal,
-		ESizeParams SizeParams);
+		UE::LLM::ESizeParams SizeParams);
 
 	void OnPreFork();
 
@@ -320,7 +327,7 @@ private:
 	void WriteHeader(const FTagData* OverrideTrackedTotalTagData, const FTagData* OverrideUntaggedTagData);
 	void AddRow(FLowLevelMemTracker& LLMRef, const FTrackerTagSizeMap& TagSizes,
 		const FTagData* OverrideTrackedTotalName, const FTagData* OverrideUntaggedName, int64 TrackedTotal,
-		ESizeParams SizeParams);
+		UE::LLM::ESizeParams SizeParams);
 
 	FConstTagDataArray Columns;
 	TFastPointerLLMSet<const FTagData*> ExistingColumns;
@@ -339,7 +346,7 @@ public:
 	void Clear();
 	void Publish(FLowLevelMemTracker& LLMRef, const FTrackerTagSizeMap& TagSizes,
 		const FTagData* OverrideTrackedTotalName, const FTagData* OverrideUntaggedName, int64 TrackedTotal,
-		ESizeParams SizeParams);
+		UE::LLM::ESizeParams SizeParams);
 	static void TraceGenericTags(FLowLevelMemTracker& LLMRef);
 
 private:
@@ -359,7 +366,7 @@ public:
 	void SetTracker(ELLMTracker InTracker);
 	void Publish(FLowLevelMemTracker& LLMRef, const FTrackerTagSizeMap& TagSizes,
 		const FTagData* OverrideTrackedTotalName, const FTagData* OverrideUntaggedName, int64 TrackedTotal,
-		ESizeParams SizeParams);
+		UE::LLM::ESizeParams SizeParams);
 protected:
 	void RecordTagToCsv(int32 CsvCategoryIndex, const FTagData* TagData, int64 TagSize);
 #if LLM_CSV_PROFILER_WRITER_ENABLED 
@@ -441,10 +448,10 @@ public:
 
 	void Clear();
 
-	void PublishStats(ESizeParams SizeParam);
-	void PublishCsv(ESizeParams SizeParam);
-	void PublishTrace(ESizeParams SizeParam);
-	void PublishCsvProfiler(ESizeParams SizeParam);
+	void PublishStats(UE::LLM::ESizeParams SizeParam);
+	void PublishCsv(UE::LLM::ESizeParams SizeParam);
+	void PublishTrace(UE::LLM::ESizeParams SizeParam);
+	void PublishCsvProfiler(UE::LLM::ESizeParams SizeParam);
 
 	void OnPreFork();
 
@@ -546,7 +553,7 @@ public:
 	void CaptureTagSnapshot();
 	void ClearTagSnapshot();
 
-	int64 GetTagAmount(const FTagData* TagData, bool bPeakAmount = false) const;
+	int64 GetTagAmount(const FTagData* TagData, UE::LLM::ESizeParams SizeParams = UE::LLM::ESizeParams::Default) const;
 	void SetTagAmountExternal(const FTagData* TagData, int64 Amount, bool bAddToTotal);
 	void SetTagAmountInUpdate(const FTagData* TagData, int64 Amount, bool bAddToTotal);
 	const FTagData* GetActiveTagData(ELLMTagSet TagSet = ELLMTagSet::None);
@@ -559,9 +566,9 @@ public:
 
 	int64 GetAllocTypeAmount(ELLMAllocType AllocType);
 
-	int64 GetTrackedTotal(ESizeParams SizeParams = ESizeParams::Default) const
+	int64 GetTrackedTotal(UE::LLM::ESizeParams SizeParams = UE::LLM::ESizeParams::Default) const
 	{
-		if (EnumHasAnyFlags(SizeParams, ESizeParams::RelativeToSnapshot))
+		if (EnumHasAnyFlags(SizeParams, UE::LLM::ESizeParams::RelativeToSnapshot))
 		{
 			return FMath::Clamp<int64>(TrackedTotal - TrackedTotalInSnapshot, 0, INT64_MAX);
 		}
@@ -1100,15 +1107,15 @@ void FLowLevelMemTracker::PublishDataPerFrame(const TCHAR* LogName)
 
 		bool bTrackPeaks = CVarLLMTrackPeaks.GetValueOnAnyThread() != 0;
 
-		ESizeParams SizeParams(ESizeParams::Default);
+		UE::LLM::ESizeParams SizeParams(UE::LLM::ESizeParams::Default);
 		if (bTrackPeaks)
 		{
-			EnumAddFlags(SizeParams, ESizeParams::ReportPeak);
+			EnumAddFlags(SizeParams, UE::LLM::ESizeParams::ReportPeak);
 		}
 
 		if (bCapturedSizeSnapshot)
 		{
-			EnumAddFlags(SizeParams, ESizeParams::RelativeToSnapshot);
+			EnumAddFlags(SizeParams, UE::LLM::ESizeParams::RelativeToSnapshot);
 		}
 
 #if !LLM_ENABLED_TRACK_PEAK_MEMORY
@@ -2129,7 +2136,8 @@ FName FLowLevelMemTracker::GetTagUniqueName(const UE::LLMPrivate::FTagData* TagD
 	return TagData->GetName();
 }
 
-int64 FLowLevelMemTracker::GetTagAmountForTracker(ELLMTracker Tracker, ELLMTag Tag, bool bPeakAmount/* = false*/)
+// Deprecated in 5.3
+int64 FLowLevelMemTracker::GetTagAmountForTracker(ELLMTracker Tracker, ELLMTag Tag, bool bPeakAmount)
 {
 	if (bIsDisabled)
 	{
@@ -2138,11 +2146,12 @@ int64 FLowLevelMemTracker::GetTagAmountForTracker(ELLMTracker Tracker, ELLMTag T
 
 	BootstrapInitialise();
 
-	return GetTagAmountForTracker(Tracker, FindTagData(Tag), bPeakAmount);
+	return GetTagAmountForTracker(Tracker, FindTagData(Tag), bPeakAmount ? UE::LLM::ESizeParams::ReportPeak : UE::LLM::ESizeParams::Default);
 }
 
+// Deprecated in 5.3
 int64 FLowLevelMemTracker::GetTagAmountForTracker(ELLMTracker Tracker, const UE::LLMPrivate::FTagData* TagData,
-	bool bPeakAmount/* = false*/)
+	bool bPeakAmount)
 {
 	using namespace UE::LLMPrivate;
 
@@ -2159,11 +2168,64 @@ int64 FLowLevelMemTracker::GetTagAmountForTracker(ELLMTracker Tracker, const UE:
 	}
 
 	FScopeLock UpdateScopeLock(&UpdateLock); // uses of TagSizes are guarded by the UpdateLock
-	return GetTracker(Tracker)->GetTagAmount(TagData, bPeakAmount);
+	return GetTracker(Tracker)->GetTagAmount(TagData, bPeakAmount ? UE::LLM::ESizeParams::ReportPeak : UE::LLM::ESizeParams::Default);
 }
 
+// Deprecated in 5.3
 int64 FLowLevelMemTracker::GetTagAmountForTracker(ELLMTracker Tracker, FName Tag, ELLMTagSet TagSet,
-	bool bPeakAmount /* = false */)
+	bool bPeakAmount)
+{
+	if (bIsDisabled)
+	{
+		return 0;
+	}
+
+	BootstrapInitialise();
+
+	UE::LLM::ESizeParams SizeParams(UE::LLM::ESizeParams::Default);
+	if (bPeakAmount)
+	{
+		EnumAddFlags(SizeParams, UE::LLM::ESizeParams::ReportPeak);
+	}
+
+	return GetTagAmountForTracker(Tracker, FindTagData(Tag, TagSet, UE::LLMPrivate::ETagReferenceSource::FunctionAPI), bPeakAmount ? UE::LLM::ESizeParams::ReportPeak : UE::LLM::ESizeParams::Default);
+}
+
+int64 FLowLevelMemTracker::GetTagAmountForTracker(ELLMTracker Tracker, const UE::LLMPrivate::FTagData* TagData,
+	UE::LLM::ESizeParams SizeParams)
+{
+	using namespace UE::LLMPrivate;
+
+	if (bIsDisabled)
+	{
+		return 0;
+	}
+
+	BootstrapInitialise();
+
+	if (TagData == nullptr)
+	{
+		return 0;
+	}
+
+	FScopeLock UpdateScopeLock(&UpdateLock); // uses of TagSizes are guarded by the UpdateLock
+	return GetTracker(Tracker)->GetTagAmount(TagData, SizeParams);
+}
+
+int64 FLowLevelMemTracker::GetTagAmountForTracker(ELLMTracker Tracker, FName Tag, ELLMTagSet TagSet,
+	UE::LLM::ESizeParams SizeParams)
+{
+	if (bIsDisabled)
+	{
+		return 0;
+	}
+
+	BootstrapInitialise();
+
+	return GetTagAmountForTracker(Tracker, FindTagData(Tag, TagSet, UE::LLMPrivate::ETagReferenceSource::FunctionAPI), SizeParams);
+}
+
+int64 FLowLevelMemTracker::GetTagAmountForTracker(ELLMTracker Tracker, ELLMTag Tag, UE::LLM::ESizeParams SizeParams)
 {
 	if (bIsDisabled)
 	{
@@ -2172,7 +2234,7 @@ int64 FLowLevelMemTracker::GetTagAmountForTracker(ELLMTracker Tracker, FName Tag
 
 	BootstrapInitialise();
 
-	return GetTagAmountForTracker(Tracker, FindTagData(Tag, TagSet, UE::LLMPrivate::ETagReferenceSource::FunctionAPI), bPeakAmount);
+	return GetTagAmountForTracker(Tracker, FindTagData(Tag), SizeParams);
 }
 
 void FLowLevelMemTracker::SetTagAmountForTracker(ELLMTracker Tracker, ELLMTag Tag, int64 Amount, bool bAddToTotal)
@@ -3819,7 +3881,7 @@ void FLLMTracker::UpdateThreads()
 	}
 }
 
-void FLLMTracker::PublishStats(ESizeParams SizeParams)
+void FLLMTracker::PublishStats(UE::LLM::ESizeParams SizeParams)
 {
 	if (OverrideTrackedTotalTagData)
 	{
@@ -3857,20 +3919,20 @@ void FLLMTracker::PublishStats(ESizeParams SizeParams)
 	}
 }
 
-void FLLMTracker::PublishCsv(ESizeParams SizeParams)
+void FLLMTracker::PublishCsv(UE::LLM::ESizeParams SizeParams)
 {
 	CsvWriter.Publish(LLMRef, TagSizes, OverrideTrackedTotalTagData, OverrideUntaggedTagData, GetTrackedTotal(SizeParams), SizeParams);
 }
 
-void FLLMTracker::PublishTrace(ESizeParams SizeParams)
+void FLLMTracker::PublishTrace(UE::LLM::ESizeParams SizeParams)
 {
 	// Trace tracker does not support the handling of Snapshot currently
-	EnumRemoveFlags(SizeParams, ESizeParams::RelativeToSnapshot);
+	EnumRemoveFlags(SizeParams, UE::LLM::ESizeParams::RelativeToSnapshot);
 
 	TraceWriter.Publish(LLMRef, TagSizes, OverrideTrackedTotalTagData, OverrideUntaggedTagData, GetTrackedTotal(SizeParams), SizeParams);
 }
 
-void FLLMTracker::PublishCsvProfiler(ESizeParams SizeParams)
+void FLLMTracker::PublishCsvProfiler(UE::LLM::ESizeParams SizeParams)
 {
 	CsvProfilerWriter.Publish(LLMRef, TagSizes, OverrideTrackedTotalTagData, OverrideUntaggedTagData, GetTrackedTotal(SizeParams), SizeParams);
 }
@@ -3948,7 +4010,7 @@ void FLLMTracker::GetTagsNamesWithAmount(TMap<FName, uint64>& OutTagsNamesWithAm
 		const FTagData* TagData = It.Key;
 		if (TagData->GetTagSet() == TagSet)
 		{
-			OutTagsNamesWithAmount.Add(TagData->GetName(), It.Value.GetSize(ESizeParams::Default));
+			OutTagsNamesWithAmount.Add(TagData->GetName(), It.Value.GetSize(UE::LLM::ESizeParams::Default));
 		}
 	}
 }
@@ -4015,12 +4077,12 @@ bool FLLMTracker::FindTagsForPtr(void* InPtr, TArray<const FTagData *, TInlineAl
 	return true;
 }
 
-int64 FLLMTracker::GetTagAmount(const FTagData* TagData, bool bPeakAmount/* = false*/) const
+int64 FLLMTracker::GetTagAmount(const FTagData* TagData, UE::LLM::ESizeParams SizeParams) const
 {
 	const FTrackerTagSizeData* AllocationData = TagSizes.Find(TagData);
 	if (AllocationData)
 	{
-		return AllocationData->GetSize((bPeakAmount ? ESizeParams::ReportPeak : ESizeParams::Default));
+		return AllocationData->GetSize(SizeParams);
 	}
 	else
 	{
@@ -4340,7 +4402,7 @@ void FLLMCsvWriter::SetTracker(ELLMTracker InTracker)
 
 void FLLMCsvWriter::Publish(FLowLevelMemTracker& LLMRef, const FTrackerTagSizeMap& TagSizes,
 	const FTagData* OverrideTrackedTotalTagData, const FTagData* OverrideUntaggedTagData, int64 TrackedTotal,
-	ESizeParams SizeParams)
+	UE::LLM::ESizeParams SizeParams)
 {
 	const double Now = FPlatformTime::Seconds();
 
@@ -4485,7 +4547,7 @@ void FLLMCsvWriter::WriteHeader(const FTagData* OverrideTrackedTotalTagData, con
 
 void FLLMCsvWriter::AddRow(FLowLevelMemTracker& LLMRef, const FTrackerTagSizeMap& TagSizes,
 	const FTagData* OverrideTrackedTotalTagData, const FTagData* OverrideUntaggedTagData, int64 TrackedTotal,
-	ESizeParams SizeParams)
+	UE::LLM::ESizeParams SizeParams)
 {
 	TStringBuilder<256> TextBuffer;
 	auto WriteValue = [this, &TextBuffer](int64 Value)
@@ -4565,7 +4627,7 @@ const void* FLLMTraceWriter::GetTagId(const FTagData* TagData)
 
 void FLLMTraceWriter::Publish(FLowLevelMemTracker& LLMRef, const FTrackerTagSizeMap& TagSizes,
 	const FTagData* OverrideTrackedTotalTagData, const FTagData* OverrideUntaggedTagData, int64 TrackedTotal,
-	ESizeParams SizeParams)
+	UE::LLM::ESizeParams SizeParams)
 {
 	if (!UE_TRACE_CHANNELEXPR_IS_ENABLED(MemTagChannel))
 	{
@@ -4710,7 +4772,7 @@ void FLLMCsvProfilerWriter::SetTracker(ELLMTracker InTracker)
 
 void FLLMCsvProfilerWriter::Publish(FLowLevelMemTracker& LLMRef, const FTrackerTagSizeMap& TagSizes,
 	const FTagData* OverrideTrackedTotalTagData, const FTagData* OverrideUntaggedTagData, int64 TrackedTotal,
-	ESizeParams SizeParams)
+	UE::LLM::ESizeParams SizeParams)
 {
 #if LLM_CSV_PROFILER_WRITER_ENABLED
 	int32 CsvCategoryIndex = (Tracker == ELLMTracker::Platform) ?
diff --git a/Engine/Source/Runtime/Core/Private/HAL/LowLevelMemTrackerPrivate.h b/Engine/Source/Runtime/Core/Private/HAL/LowLevelMemTrackerPrivate.h
--- a/Engine/Source/Runtime/Core/Private/HAL/LowLevelMemTrackerPrivate.h
+++ b/Engine/Source/Runtime/Core/Private/HAL/LowLevelMemTrackerPrivate.h
@@ -173,15 +173,6 @@ namespace UE
 {
 namespace LLMPrivate
 {
-	enum class ESizeParams : uint8
-	{
-		Default = 0,
-		ReportCurrent = 0,
-		ReportPeak = 1,
-		RelativeToSnapshot = 2
-	};
-	ENUM_CLASS_FLAGS(ESizeParams);
-
 	/** Size information stored on the tracker for a tag; includes amounts aggregated from threadstates and from external api users */
 	struct FTrackerTagSizeData
 	{
@@ -194,18 +185,18 @@ namespace LLMPrivate
 		bool bExternalValid = false;
 		bool bExternalAddToTotal = false;
 
-		int64 GetSize(ESizeParams SizeParams) const
+		int64 GetSize(UE::LLM::ESizeParams SizeParams) const
 		{
 			int64 CurrentSize = Size;
 
 #if LLM_ENABLED_TRACK_PEAK_MEMORY
-			if (EnumHasAnyFlags(SizeParams, ESizeParams::ReportPeak))
+			if (EnumHasAnyFlags(SizeParams, UE::LLM::ESizeParams::ReportPeak))
 			{
 				CurrentSize = PeakSize;
 			}
 #endif
 			// Note, this will also subtract the snapshotted size from PeakSize if that flag is enabled
-			if (EnumHasAnyFlags(SizeParams, ESizeParams::RelativeToSnapshot))
+			if (EnumHasAnyFlags(SizeParams, UE::LLM::ESizeParams::RelativeToSnapshot))
 			{
 				CurrentSize = FMath::Clamp<int64>(CurrentSize - SizeInSnapshot, 0, INT64_MAX);
 			}
diff --git a/Engine/Source/Runtime/Core/Public/HAL/LowLevelMemTracker.h b/Engine/Source/Runtime/Core/Public/HAL/LowLevelMemTracker.h
--- a/Engine/Source/Runtime/Core/Public/HAL/LowLevelMemTracker.h
+++ b/Engine/Source/Runtime/Core/Public/HAL/LowLevelMemTracker.h
@@ -127,6 +127,24 @@ enum class ELLMTagSet : uint8
 	Max, // note: see FLowLevelMemTracker::ShouldReduceThreads and IsAssetTagForAssets if you add any asset-style tagsets
 };
 
+/*
+ * Size parameter flags used when requesting the size of the tracked tag data.
+ */
+namespace UE::LLM
+{
+
+enum class ESizeParams : uint8
+{
+	Default = 0,
+	ReportCurrent = 0,
+	ReportPeak = 1,
+	RelativeToSnapshot = 2
+};
+
+ENUM_CLASS_FLAGS(ESizeParams);
+
+} // UE::LLM
+
 // Do not add to these macros. Please use the LLM_DECLARE_TAG family of macros below to create new tags.
 #define LLM_ENUM_GENERIC_TAGS(macro) \
 	macro(Untagged,								"Untagged",						NAME_None,													NAME_None,										-1)\
@@ -634,12 +652,23 @@ public:
 	CORE_API FName GetTagUniqueName(const UE::LLMPrivate::FTagData* TagData) const;
 
 	/** Get the amount of memory for an ELLMTag from the given tracker. */
-	CORE_API int64 GetTagAmountForTracker(ELLMTracker Tracker, ELLMTag Tag, bool bPeakAmount = false);
+	CORE_API int64 GetTagAmountForTracker(ELLMTracker Tracker, ELLMTag Tag, UE::LLM::ESizeParams SizeParams = UE::LLM::ESizeParams::Default);
+
+	/** Get the amount of memory for a FTagData from the given tracker. */
+	CORE_API int64 GetTagAmountForTracker(ELLMTracker Tracker, const UE::LLMPrivate::FTagData* TagData, UE::LLM::ESizeParams SizeParams = UE::LLM::ESizeParams::Default);
+
+    CORE_API int64 GetTagAmountForTracker(ELLMTracker Tracker, FName Tag, ELLMTagSet TagSet, UE::LLM::ESizeParams SizeParams = UE::LLM::ESizeParams::Default);
+
+	/** Get the amount of memory for an ELLMTag from the given tracker. */
+	UE_DEPRECATED(5.3, "Use version that takes ESizeParams instead")
+	CORE_API int64 GetTagAmountForTracker(ELLMTracker Tracker, ELLMTag Tag, bool bPeakAmount);
 
 	/** Get the amount of memory for a FTagData from the given tracker. */
-	CORE_API int64 GetTagAmountForTracker(ELLMTracker Tracker, const UE::LLMPrivate::FTagData* TagData, bool bPeakAmount = false);
+	UE_DEPRECATED(5.3, "Use version that takes ESizeParams instead")
+	CORE_API int64 GetTagAmountForTracker(ELLMTracker Tracker, const UE::LLMPrivate::FTagData* TagData, bool bPeakAmount);
 
-    CORE_API int64 GetTagAmountForTracker(ELLMTracker Tracker, FName Tag, ELLMTagSet TagSet, bool bPeakAmount = false);
+	UE_DEPRECATED(5.3, "Use version that takes ESizeParams instead")
+	CORE_API int64 GetTagAmountForTracker(ELLMTracker Tracker, FName Tag, ELLMTagSet TagSet, bool bPeakAmount);
 
 	/** Set the amount of memory for an ELLMTag for a given tracker and optionally update the total tracked memory. */
 	CORE_API void SetTagAmountForTracker(ELLMTracker Tracker, ELLMTag Tag, int64 Amount, bool bAddToTotal);
@@ -658,7 +687,7 @@ public:
 		PlainText,
 		CSV,
 	};
-	CORE_API void DumpToLog(EDumpFormat DumpFormat = EDumpFormat::PlainText, FOutputDevice* OutputDevice = nullptr);
+	CORE_API void DumpToLog(EDumpFormat DumpFormat = EDumpFormat::PlainText, FOutputDevice* OutputDevice = nullptr, UE::LLM::ESizeParams SizeParams = UE::LLM::ESizeParams::Default);
 
 	CORE_API void OnPreFork();
 
