Fixed some random intellisense compile issues.

diff --git a/Engine/Source/Runtime/Core/Public/Misc/CoreMiscDefines.h b/Engine/Source/Runtime/Core/Public/Misc/CoreMiscDefines.h
--- a/Engine/Source/Runtime/Core/Public/Misc/CoreMiscDefines.h
+++ b/Engine/Source/Runtime/Core/Public/Misc/CoreMiscDefines.h
@@ -267,7 +267,12 @@ enum EInPlace {InPlace};
  * @param VERSION The release number in which the feature was marked deprecated.
  * @param MESSAGE A message containing upgrade notes.
  */
+
+#if defined (__INTELLISENSE__)
+#define UE_DEPRECATED(Version, Message)
+#else
 #define UE_DEPRECATED(Version, Message) [[deprecated(Message " Please update your code to the new API before upgrading to the next release, otherwise your project will no longer compile.")]]
+#endif
 
 #ifndef UE_DEPRECATED_FORGAME
 	#define UE_DEPRECATED_FORGAME(...)
diff --git a/Engine/Source/Runtime/Engine/Classes/Engine/Texture.h b/Engine/Source/Runtime/Engine/Classes/Engine/Texture.h
--- a/Engine/Source/Runtime/Engine/Classes/Engine/Texture.h
+++ b/Engine/Source/Runtime/Engine/Classes/Engine/Texture.h
@@ -46,6 +46,7 @@
 namespace FOodleDataCompression {enum class ECompressor : uint8; enum class ECompressionLevel : int8; }
 
 class FTextureReference;
+class FTextureResource;
 class ITargetPlatform;
 class UAssetUserData;
 struct FPropertyChangedEvent;
diff --git a/Engine/Source/Runtime/Engine/Classes/Engine/Texture2D.h b/Engine/Source/Runtime/Engine/Classes/Engine/Texture2D.h
--- a/Engine/Source/Runtime/Engine/Classes/Engine/Texture2D.h
+++ b/Engine/Source/Runtime/Engine/Classes/Engine/Texture2D.h
@@ -17,6 +17,7 @@
 
 class FTexture2DResourceMem;
 class FTexture2DResource;
+struct FTexture2DMipMap;
 struct FUpdateTextureRegion2D;
 
 UCLASS(hidecategories=Object, MinimalAPI, BlueprintType)
