Enable C4855 for Win64

#rnx

diff --git a/Engine/Source/Runtime/Core/Public/Windows/WindowsPlatformCompilerSetup.h b/Engine/Source/Runtime/Core/Public/Windows/WindowsPlatformCompilerSetup.h
--- a/Engine/Source/Runtime/Core/Public/Windows/WindowsPlatformCompilerSetup.h
+++ b/Engine/Source/Runtime/Core/Public/Windows/WindowsPlatformCompilerSetup.h
@@ -33,7 +33,7 @@ static_assert(_MSC_VER >= 1929, "Visual Studio 2019 v16.11 or greater is require
 #pragma warning (error:           4502 4503      4505 4506      4508 4509      4511 4512 4513 4514 4515 4516 4517 4518 4519      4521 4522 4523           4526                4530 4531 4532 4533 4534 4535 4536 4537 4538      4540 4541 4542 4543 4544 4545 4546                4550 4551 4552 4553 4554      4556 4557 4558 4559      4561 4562      4564 4565 4566      4568 4569 4570      4572 4573      4575 4576      4578      4580 4581           4584 4585 4586 4587 4588 4589      4591      4593 4594 4595 4596 4597 4598     )
 #pragma warning (error: 4600      4602 4603 4604      4606           4609 4610 4611 4612 4613      4615 4616      4618      4620 4621 4622 4623 4624 4625 4626 4627 4628 4629 4630 4631 4632 4633 4634 4635 4636 4637 4638 4639 4640 4641 4642      4644 4645 4646      4648 4649 4650      4652 4653 4654 4655 4656 4657 4658 4659      4661 4662                     4667      4669 4670 4671 4672 4673 4674      4676 4677 4678 4679 4680 4681 4682 4683 4684 4685 4686 4687 4688 4689 4690 4691      4693 4694 4695 4696      4698     )
 #pragma warning (error: 4700      4702 4703           4706           4709 4710 4711           4714 4715 4716 4717 4718 4719 4720 4721 4722 4723 4724 4725 4726 4727 4728 4729      4731 4732 4733 4734                     4739 4740 4741 4742 4743 4744 4745      4747      4749 4750 4751 4752 4753 4754 4755 4756 4757                4761           4764                               4771 4772           4775 4776 4777 4778                                    4786      4788 4789           4792 4793 4794                4798 4799)
-#pragma warning (error: 4800 4801      4803 4804 4805 4806 4807 4808 4809 4810 4811 4812 4813      4815 4816 4817                4821 4822 4823                4827      4829                     4834 4835                4839 4840 4841 4842 4843 4844 4845 4846 4847 4848 4849                                                                                      4867      4869           4872                                    4880 4881 4882 4883                                                                                )
+#pragma warning (error: 4800 4801      4803 4804 4805 4806 4807 4808 4809 4810 4811 4812 4813      4815 4816 4817                4821 4822 4823                4827      4829                     4834 4835                4839 4840 4841 4842 4843 4844 4845 4846 4847 4848 4849                          4855                                                        4867      4869           4872                                    4880 4881 4882 4883                                                                                )
 #pragma warning (error: 4900                     4905 4906                4910      4912 4913           4916 4917 4918      4920 4921                4925 4926 4927 4928 4929 4930 4931 4932      4934 4935 4936 4937 4938 4939                     4944 4945 4946 4947 4948 4949 4950 4951 4952 4953 4954 4955 4956 4957 4958 4959 4960 4961      4963 4964 4965 4966                4970 4971 4972 4973 4974                               4981                4985                4989 4990 4991 4992           4995      4997 4998 4999)
 #pragma warning (error:                                                                                                                                                                                               5038                                                                                                                                                                                                                                                                                                                 )
 
