Add millisecond precision to setting unix timestamps on FDateTime.

diff --git a/Engine/Source/Runtime/Core/Public/Misc/DateTime.h b/Engine/Source/Runtime/Core/Public/Misc/DateTime.h
--- a/Engine/Source/Runtime/Core/Public/Misc/DateTime.h
+++ b/Engine/Source/Runtime/Core/Public/Misc/DateTime.h
@@ -622,6 +622,18 @@ public:
 		return FDateTime(1970, 1, 1) + FTimespan(UnixTime * ETimespan::TicksPerSecond);
 	}
 
+	/**
+	 * Returns the date from Unix time (seconds from midnight 1970-01-01)
+	 *
+	 * @param UnixTime Unix time (seconds from midnight 1970-01-01)
+	 * @return Gregorian date and time.
+	 * @see ToUnixTimestampDecimal
+	 */
+	static FDateTime FromUnixTimestampDecimal(double UnixTime)
+	{
+		return FDateTime(1970, 1, 1) + FTimespan(UnixTime * ETimespan::TicksPerSecond);
+	}
+
 	/**
 	 * Checks whether the given year is a leap year.
 	 *
