Adding explicit cast to FDateTime::FromUnixTimestampDecimal

diff --git a/Engine/Source/Runtime/Core/Public/Misc/DateTime.h b/Engine/Source/Runtime/Core/Public/Misc/DateTime.h
--- a/Engine/Source/Runtime/Core/Public/Misc/DateTime.h
+++ b/Engine/Source/Runtime/Core/Public/Misc/DateTime.h
@@ -631,7 +631,7 @@ public:
 	 */
 	static FDateTime FromUnixTimestampDecimal(double UnixTime)
 	{
-		return FDateTime(1970, 1, 1) + FTimespan(UnixTime * ETimespan::TicksPerSecond);
+		return FDateTime(1970, 1, 1) + FTimespan(int64(UnixTime * ETimespan::TicksPerSecond));
 	}
 
 	/**
