Added checks to ensure that users cannot emplace intrusive unset values into TOptionals.

diff --git a/Engine/Source/Runtime/Core/Public/Misc/Optional.h b/Engine/Source/Runtime/Core/Public/Misc/Optional.h
--- a/Engine/Source/Runtime/Core/Public/Misc/Optional.h
+++ b/Engine/Source/Runtime/Core/Public/Misc/Optional.h
@@ -98,6 +98,11 @@ public:
 		{
 			Value.bIsSet = true;
 		}
+		else
+		{
+			// Ensure that a user doesn't emplace an unset state into the optional
+			checkf(IsSet(), TEXT("TOptional::TOptional(EInPlace, ...) - optionals should not be unset by emplacement"));
+		}
 	}
 	
 	/** Construct an OptionalType with an invalid value. */
@@ -275,6 +280,11 @@ public:
 		{
 			Value.bIsSet = true;
 		}
+		else
+		{
+			// Ensure that a user doesn't emplace an unset state into the optional
+			checkf(IsSet(), TEXT("TOptional::Emplace(...) - optionals should not be unset by an emplacement"));
+		}
 
 		return *Result;
 	}
