Add an overload to the FindRef method of TMap and TSortedMap that take a fallback default value

diff --git a/Engine/Source/Runtime/Core/Public/Containers/Map.h b/Engine/Source/Runtime/Core/Public/Containers/Map.h
--- a/Engine/Source/Runtime/Core/Public/Containers/Map.h
+++ b/Engine/Source/Runtime/Core/Public/Containers/Map.h
@@ -688,6 +688,23 @@ public:
 		return ValueType();
 	}
 
+	/**
+	 * Find the value associated with a specified key.
+	 *
+	 * @param Key The key to search for.
+	 * @param DefaultValue The fallback value if the key is not found.
+	 * @return The value associated with the specified key, or DefaultValue if the key isn't contained in this map.
+	 */
+	FORCEINLINE ValueType FindRef(KeyConstPointerType Key, ValueType DefaultValue) const
+	{
+		if (const auto* Pair = Pairs.Find(Key))
+		{
+			return Pair->Value;
+		}
+
+		return DefaultValue;
+	}
+
 	/**
 	 * Check if map contains the specified key.
 	 *
diff --git a/Engine/Source/Runtime/Core/Public/Containers/SortedMap.h b/Engine/Source/Runtime/Core/Public/Containers/SortedMap.h
--- a/Engine/Source/Runtime/Core/Public/Containers/SortedMap.h
+++ b/Engine/Source/Runtime/Core/Public/Containers/SortedMap.h
@@ -321,6 +321,23 @@ public:
 		return ValueType();
 	}
 
+	/**
+	 * Returns the value associated with a specified key.
+	 *
+	 * @param Key The key to search for.
+	 * @param DefaultValue The fallback value if the key is not found.
+	 * @return The value associated with the specified key, or DefaultValue if the key isn't contained in this map.
+	 */
+	FORCEINLINE ValueType FindRef(KeyConstPointerType Key, ValueType DefaultValue) const
+	{
+		if (const ValueType* Value = Find(Key))
+		{
+			return *Value;
+		}
+
+		return DefaultValue;
+	}
+
 	/**
 	 * Checks if map contains the specified key.
 	 *
