LLM AssetTags: Fix the missing stats for AssetTags; they were accidentally disabled by CL 23599453.
#rnx

diff --git a/Engine/Source/Runtime/Core/Private/HAL/LowLevelMemTracker.cpp b/Engine/Source/Runtime/Core/Private/HAL/LowLevelMemTracker.cpp
--- a/Engine/Source/Runtime/Core/Private/HAL/LowLevelMemTracker.cpp
+++ b/Engine/Source/Runtime/Core/Private/HAL/LowLevelMemTracker.cpp
@@ -2376,6 +2376,12 @@ const UE::LLMPrivate::FTagData* FLowLevelMemTracker::FindOrAddTagData(ELLMTag En
 
 const UE::LLMPrivate::FTagData* FLowLevelMemTracker::FindOrAddTagData(FName TagName, ELLMTagSet TagSet, bool bIsStatTag,
 	UE::LLMPrivate::ETagReferenceSource ReferenceSource)
+{
+	return FindOrAddTagData(TagName, TagSet, bIsStatTag ? TagName : NAME_None, ReferenceSource);
+}
+
+const UE::LLMPrivate::FTagData* FLowLevelMemTracker::FindOrAddTagData(FName TagName, ELLMTagSet TagSet, FName StatName,
+	UE::LLMPrivate::ETagReferenceSource ReferenceSource)
 {
 	using namespace UE::LLMPrivate;
 
@@ -2396,12 +2402,12 @@ const UE::LLMPrivate::FTagData* FLowLevelMemTracker::FindOrAddTagData(FName TagN
 		FinishInitialise();
 		// Reeneter this function so that we retry the find above; note we avoid infinite recursion because
 		// bFullyInitialised is now true.
-		return FindOrAddTagData(TagName, TagSet, bIsStatTag, ReferenceSource);
+		return FindOrAddTagData(TagName, TagSet, StatName, ReferenceSource);
 	}
 	LLMCheckf(!bIsBootstrapping, TEXT("LLM Error: Invalid use of FName tag when initialising tags."));
 
 	// Add the new Tag
-	FName StatName = bIsStatTag ? TagName : NAME_None;
+	bool bIsStatTag = !StatName.IsNone() && StatName == TagName;
 	FTagData* TagData = &RegisterTagData(TagName, NAME_None, NAME_None, StatName, NAME_None, false,
 		ELLMTag::CustomName, bIsStatTag, ReferenceSource, TagSet);
 	{
@@ -2512,7 +2518,7 @@ void FLLMScope::Init(const UE::LLMPrivate::FTagData* TagData, bool bInIsStatTag,
 {
 	FLowLevelMemTracker& LLMRef = FLowLevelMemTracker::Get();
 	// We have to check bIsDisabled again after calling Get, because the constructor is called from Get, and will set
-	// bIsDisabled=false if the platform doesn't support it.
+	// bIsDisabled=true if the platform doesn't support it.
 	if (FLowLevelMemTracker::bIsDisabled)
 	{
 		bEnabled = false;
@@ -2538,6 +2544,91 @@ void FLLMScope::Destruct()
 	LLMRef.GetTracker(Tracker)->PopTag(TagSet);
 }
 
+void FLLMScopeDynamic::Init(ELLMTracker InTracker, ELLMTagSet InTagSet)
+{
+	FLowLevelMemTracker& LLMRef = FLowLevelMemTracker::Get();
+	// We have to check bIsDisabled again after calling Get, because the constructor is called from Get, and will set
+	// bIsDisabled=true if the platform doesn't support it.
+	if (FLowLevelMemTracker::bIsDisabled)
+	{
+		bEnabled = false;
+		return;
+	}
+	LLMRef.BootstrapInitialise();
+	if (!LLMRef.IsTagSetActive(InTagSet))
+	{
+		bEnabled = false;
+		return;
+	}
+	bEnabled = true;
+	TagData = nullptr;
+	Tracker = InTracker;
+	TagSet = InTagSet;
+}
+
+bool FLLMScopeDynamic::TryFindTag(FName TagName)
+{
+	FLowLevelMemTracker& LLMRef = FLowLevelMemTracker::Get();
+	TagData = LLMRef.FindTagData(TagName, TagSet, UE::LLMPrivate::ETagReferenceSource::Scope);
+	return TagData != nullptr;
+}
+
+bool FLLMScopeDynamic::TryAddTagAndActivate(FName TagName, const ILLMDynamicTagConstructor& Constructor)
+{
+	FLowLevelMemTracker& LLMRef = FLowLevelMemTracker::Get();
+	FName StatFullName = NAME_None;
+#if LLM_ENABLED_STAT_TAGS && STATS
+	FString StatConstructorName = Constructor.GetStatName();
+	if (!StatConstructorName.IsEmpty())
+	{
+		if (Constructor.NeedsStatConstruction())
+		{
+			switch (TagSet)
+			{
+			case ELLMTagSet::None:
+				StatFullName = FDynamicStats::CreateMemoryStatId<FStatGroup_STATGROUP_LLMFULL>(
+					StatConstructorName).GetName();
+				break;
+			case ELLMTagSet::Assets:
+				StatFullName = FDynamicStats::CreateMemoryStatId<FStatGroup_STATGROUP_LLMAssets>(
+					StatConstructorName).GetName();
+				break;
+			case ELLMTagSet::AssetClasses:
+				StatFullName = FDynamicStats::CreateMemoryStatId<FStatGroup_STATGROUP_LLMAssets>(
+					StatConstructorName).GetName();
+				break;
+			default:
+				checkNoEntry();
+				break;
+			}
+		}
+		else
+		{
+			StatFullName = FName(*StatConstructorName);
+		}
+	}
+#endif
+	TagData = LLMRef.FindOrAddTagData(TagName, TagSet, StatFullName, UE::LLMPrivate::ETagReferenceSource::Scope);
+	LLMRef.GetTracker(Tracker)->PushTag(TagData, TagSet);
+	return true;
+}
+
+void FLLMScopeDynamic::Activate()
+{
+	FLowLevelMemTracker& LLMRef = FLowLevelMemTracker::Get();
+	LLMRef.GetTracker(Tracker)->PushTag(TagData, TagSet);
+}
+
+void FLLMScopeDynamic::Destruct()
+{
+	if (TagData)
+	{
+		FLowLevelMemTracker& LLMRef = FLowLevelMemTracker::Get();
+		LLMRef.GetTracker(Tracker)->PopTag(TagSet);
+	}
+}
+
+
 FLLMPauseScope::FLLMPauseScope(FName TagName, bool bIsStatTag, uint64 Amount, ELLMTracker TrackerToPause,
 	ELLMAllocType InAllocType)
 {
@@ -3202,7 +3293,7 @@ FTagData::FTagData(FName InName, ELLMTagSet InTagSet, FName InDisplayName, FName
 	: Name(InName), DisplayName(InDisplayName), ParentName(InParentName), StatName(InStatName)
 	, SummaryStatName(InSummaryStatName), EnumTag(InEnumTag), ReferenceSource(InReferenceSource)
 	, TagSet(InTagSet), bIsFinishConstructed(false), bParentIsName(true), bHasEnumTag(bInHasEnumTag)
-	, bIsReportable(InTagSet == ELLMTagSet::None)
+	, bIsReportable(true)
 {
 }
 
@@ -3338,6 +3429,11 @@ int32 FTagData::GetIndex() const
 }
 
 bool FTagData::IsReportable() const
+{
+	return bIsReportable && TagSet == ELLMTagSet::None;
+}
+
+bool FTagData::IsStatsReportable() const
 {
 	return bIsReportable;
 }
@@ -3902,7 +3998,7 @@ void FLLMTracker::PublishStats(UE::LLM::ESizeParams SizeParams)
 	for (const TPair<const FTagData*, FTrackerTagSizeData>& It : TagSizes)
 	{
 		const FTagData* TagData = It.Key;
-		if (!TagData->IsReportable())
+		if (!TagData->IsStatsReportable())
 		{
 			continue;
 		}
diff --git a/Engine/Source/Runtime/Core/Private/HAL/LowLevelMemTrackerPrivate.h b/Engine/Source/Runtime/Core/Private/HAL/LowLevelMemTrackerPrivate.h
--- a/Engine/Source/Runtime/Core/Private/HAL/LowLevelMemTrackerPrivate.h
+++ b/Engine/Source/Runtime/Core/Private/HAL/LowLevelMemTrackerPrivate.h
@@ -85,8 +85,10 @@ namespace LLMPrivate
 	{
 	public:
 
-		FTagData(FName InName, ELLMTagSet InTagSet, FName InDisplayName, FName InParentName, FName InStatName, FName InSummaryStatName, bool bInHasEnumTag, ELLMTag InEnumTag, ETagReferenceSource InReferenceSource);
-		FTagData(FName InName, ELLMTagSet InTagSet, FName InDisplayName, const FTagData* InParent, FName InStatName, FName InSummaryStatName, bool bInHasEnumTag, ELLMTag InEnumTag, ETagReferenceSource InReferenceSource);
+		FTagData(FName InName, ELLMTagSet InTagSet, FName InDisplayName, FName InParentName, FName InStatName,
+			FName InSummaryStatName, bool bInHasEnumTag, ELLMTag InEnumTag, ETagReferenceSource InReferenceSource);
+		FTagData(FName InName, ELLMTagSet InTagSet, FName InDisplayName, const FTagData* InParent, FName InStatName,
+			FName InSummaryStatName, bool bInHasEnumTag, ELLMTag InEnumTag, ETagReferenceSource InReferenceSource);
 		~FTagData();
 
 		bool IsParentConstructed() const;
@@ -107,10 +109,12 @@ namespace LLMPrivate
 		ETagReferenceSource GetReferenceSource() const;
 		int32 GetIndex() const;
 		bool IsReportable() const;
+		bool IsStatsReportable() const;
 
 		void SetParent(const FTagData* InParent);
 		void SetIndex(int32 InIndex);
 		void SetIsReportable(bool bInReportable);
+		void SetIsStatsReportable(bool bInReportable);
 		void SetFinishConstructed();
 
 		// These functions are normally invalid - these properties should be immutable - but are called for EnumTags during bootstrapping
diff --git a/Engine/Source/Runtime/Core/Private/ProfilingDebugging/TagTrace.cpp b/Engine/Source/Runtime/Core/Private/ProfilingDebugging/TagTrace.cpp
--- a/Engine/Source/Runtime/Core/Private/ProfilingDebugging/TagTrace.cpp
+++ b/Engine/Source/Runtime/Core/Private/ProfilingDebugging/TagTrace.cpp
@@ -37,6 +37,10 @@ UE_TRACE_EVENT_END()
 thread_local int32 GActiveTag;
 
 ////////////////////////////////////////////////////////////////////////////////
+FMemScope::FMemScope()
+{
+}
+
 FMemScope::FMemScope(int32 InTag, bool bShouldActivate /*= true*/)
 {
 	if (UE_TRACE_CHANNELEXPR_IS_ENABLED(MemAllocChannel) & bShouldActivate)
@@ -56,6 +60,11 @@ FMemScope::FMemScope(ELLMTag InTag, bool bShouldActivate /*= true*/)
 
 ////////////////////////////////////////////////////////////////////////////////
 FMemScope::FMemScope(const FName& InName, bool bShouldActivate /*= true*/)
+{
+	SetTagAndActivate(InName, bShouldActivate);
+}
+
+void FMemScope::SetTagAndActivate(const FName& InName, bool bShouldActivate /*= true*/)
 {
 	if (UE_TRACE_CHANNELEXPR_IS_ENABLED(MemAllocChannel) & bShouldActivate)
 	{
diff --git a/Engine/Source/Runtime/Core/Public/HAL/LowLevelMemStats.h b/Engine/Source/Runtime/Core/Public/HAL/LowLevelMemStats.h
--- a/Engine/Source/Runtime/Core/Public/HAL/LowLevelMemStats.h
+++ b/Engine/Source/Runtime/Core/Public/HAL/LowLevelMemStats.h
@@ -64,7 +64,10 @@
 #define LLM_SCOPED_SINGLE_STAT_TAG(Stat) DECLARE_LLM_MEMORY_STAT(TEXT(#Stat), Stat, STATGROUP_LLMFULL); LLM_SCOPED_TAG_WITH_STAT(Stat, ELLMTracker::Default);
 #define LLM_SCOPED_SINGLE_STAT_TAG_IN_SET(Stat, Set) DECLARE_LLM_MEMORY_STAT(TEXT(#Stat), Stat, STATGROUP_LLMFULL); LLM_SCOPED_TAG_WITH_STAT_IN_SET(Stat, Set, ELLMTracker::Default);
 #define LLM_SCOPED_PAUSE_TRACKING_WITH_STAT_AND_AMOUNT(Stat, Amount, Tracker) FLLMPauseScope SCOPE_NAME(GET_STATFNAME(Stat), true /* bIsStatTag */, Amount, Tracker, ELLMAllocType::None);
-#define LLM_SCOPED_TAG_WITH_OBJECT_IN_SET(Object, Set) FLLMScope SCOPE_NAME(FLowLevelMemTracker::Get().IsTagSetActive(Set) ? Object->GetOutermost()->GetFName() : NAME_None, (Set == ELLMTagSet::None) /* bIsStatTag */, Set, ELLMTracker::Default);
+#define LLM_SCOPE_DYNAMIC_STAT_OBJECTPATH(Object, Set) LLM_SCOPE_DYNAMIC(FName(*Object->GetPathName()), \
+					ELLMTracker::Default, Set, FLLMDynamicTagConstructorStatString(Object->GetPathName()))
+#define LLM_SCOPE_DYNAMIC_STAT_OBJECTPATH_FNAME(ObjectPath, Set) LLM_SCOPE_DYNAMIC(ObjectPath, \
+					ELLMTracker::Default, Set, FLLMDynamicTagConstructorStatString(ObjectPath.ToString()))
 
 // special stat pushing to update threads after each asset when tracking assets
 // Currently this is unused, but we may use it for optimizations later
@@ -85,8 +88,11 @@
 #define LLM_SCOPED_SINGLE_STAT_TAG(...)
 #define LLM_SCOPED_SINGLE_STAT_TAG_IN_SET(...)
 #define LLM_SCOPED_PAUSE_TRACKING_WITH_STAT_AND_AMOUNT(...)
-#define LLM_SCOPED_TAG_WITH_OBJECT_IN_SET(...)
+#define LLM_SCOPE_DYNAMIC_STAT_OBJECTPATH(...)
+#define LLM_SCOPE_DYNAMIC_STAT_OBJECTPATH_FNAME(...)
 #define LLM_PUSH_STATS_FOR_ASSET_TAGS()
 #endif
 
+#define LLM_SCOPED_TAG_WITH_OBJECT_IN_SET(Object, Set) UE_DEPRECATED_MACRO(5.3, "Use LLM_SCOPE_DYNAMIC_STAT_OBJECTPATH instead") LLM_SCOPE_DYNAMIC_STAT_OBJECTPATH(Object, Set);
+
 
diff --git a/Engine/Source/Runtime/Core/Public/HAL/LowLevelMemTracker.h b/Engine/Source/Runtime/Core/Public/HAL/LowLevelMemTracker.h
--- a/Engine/Source/Runtime/Core/Public/HAL/LowLevelMemTracker.h
+++ b/Engine/Source/Runtime/Core/Public/HAL/LowLevelMemTracker.h
@@ -330,6 +330,29 @@ extern FName LLMGetTagStat(ELLMTag Tag);
  */
 #define LLM_SCOPE(Tag)												FLLMScope SCOPE_NAME(Tag, false /* bIsStatTag */, ELLMTagSet::None, ELLMTracker::Default);\
 																	UE_MEMSCOPE(Tag) 
+#define LLM_SCOPE_DYNAMIC(UniqueName, Tracker, TagSet, Constructor) \
+	FLLMScopeDynamic SCOPE_NAME(Tracker, TagSet); \
+	UE_MEMSCOPE_UNINITIALIZED(__LINE__); \
+	do \
+	{ \
+		if (SCOPE_NAME.IsEnabled()) \
+		{ \
+			FName UniqueNameEvaluated(UniqueName); \
+			if (!UniqueNameEvaluated.IsNone()) \
+			{ \
+				if (SCOPE_NAME.TryFindTag(UniqueNameEvaluated)) \
+				{ \
+					SCOPE_NAME.Activate(); \
+				} \
+				else \
+				{ \
+					SCOPE_NAME.TryAddTagAndActivate(UniqueNameEvaluated, Constructor); \
+				} \
+				UE_MEMSCOPE_ACTIVATE(__LINE__, UniqueNameEvaluated); \
+			} \
+		} \
+	} while (false) /* do/while is added to require a semicolon */
+
 #define LLM_TAGSET_SCOPE(Tag, TagSet)								FLLMScope SCOPE_NAME(Tag, false /* bIsStaTag */, TagSet, ELLMTracker::Default);
 #define LLM_SCOPE_BYNAME(Tag) 										static FName PREPROCESSOR_JOIN(LLMScope_Name,__LINE__)(Tag);\
 																	FLLMScope SCOPE_NAME(PREPROCESSOR_JOIN(LLMScope_Name,__LINE__), false /* bIsStatTag */, ELLMTagSet::None, ELLMTracker::Default);\
@@ -737,11 +760,13 @@ private:
 
 	CORE_API const UE::LLMPrivate::FTagData* FindOrAddTagData(ELLMTag EnumTag, UE::LLMPrivate::ETagReferenceSource ReferenceSource = UE::LLMPrivate::ETagReferenceSource::FunctionAPI);
 	CORE_API const UE::LLMPrivate::FTagData* FindOrAddTagData(FName Name, ELLMTagSet TagSet, bool bIsStatData=false, UE::LLMPrivate::ETagReferenceSource ReferenceSource = UE::LLMPrivate::ETagReferenceSource::FunctionAPI);
+	CORE_API const UE::LLMPrivate::FTagData* FindOrAddTagData(FName Name, ELLMTagSet TagSet, FName StatName, UE::LLMPrivate::ETagReferenceSource ReferenceSource = UE::LLMPrivate::ETagReferenceSource::FunctionAPI);
 	CORE_API const UE::LLMPrivate::FTagData* FindTagData(ELLMTag EnumTag, UE::LLMPrivate::ETagReferenceSource ReferenceSource = UE::LLMPrivate::ETagReferenceSource::FunctionAPI);
 	CORE_API const UE::LLMPrivate::FTagData* FindTagData(FName Name, ELLMTagSet TagSet, UE::LLMPrivate::ETagReferenceSource ReferenceSource = UE::LLMPrivate::ETagReferenceSource::FunctionAPI);
 
 	friend class FLLMPauseScope;
 	friend class FLLMScope;
+	friend class FLLMScopeDynamic;
 	friend class FLLMScopeFromPtr;
 	friend class UE::LLMPrivate::FLLMCsvWriter;
 	friend class UE::LLMPrivate::FLLMTracker;
@@ -829,8 +854,76 @@ protected:
 	CORE_API void Init(const UE::LLMPrivate::FTagData* TagData, bool bIsStatTag, ELLMTagSet InTagSet, ELLMTracker InTracker, bool bOverride = true);
 	CORE_API void Destruct();
 
-	ELLMTracker Tracker{};
-	bool bEnabled = false;
+	// All fields other than bEnabled are not initialized by constructor unless bEnabled is true
+	ELLMTracker Tracker;
+	bool bEnabled = false;  // Needs to be a uint8 rather than a bitfield to minimize ALU operations
+	ELLMTagSet TagSet;
+};
+
+/**
+ * Provides arguments for TagConstruction that are passed through a LLM_SCOPE_DYNAMIC macro and should
+ * only be evaluated on the first LLM_SCOPE_DYNAMIC call.
+ */
+class ILLMDynamicTagConstructor
+{
+public:
+	virtual ~ILLMDynamicTagConstructor() = default;
+	virtual FString GetStatName() const
+	{
+		return FString();
+	}
+	virtual bool NeedsStatConstruction() const
+	{
+		return true;
+	}
+};
+
+/** ILLMDynamicTagConstructor that provides a string to pass to CreateMemoryStatId. */
+class FLLMDynamicTagConstructorStatString : public ILLMDynamicTagConstructor
+{
+public:
+	FLLMDynamicTagConstructorStatString(FString InStatName) : StatName(MoveTemp(InStatName)) {}
+	virtual FString GetStatName() const override { return StatName; }
+	FString StatName;
+};
+
+/**
+ * LLM scope for tracking memory for a dynamically-created tag. Like some normal tags, dynamically-created tags
+ * are created by Name, rather than a compile-time constant like ELLMTag or LLM_DECLARE_TAG, and they are
+ * created on demand by the first scope that uses them.
+ * Unlike normal tags, dynamically-created tags also provide the ability to create stats when first called.
+ */
+class FLLMScopeDynamic
+{
+public:
+	FLLMScopeDynamic(ELLMTracker InTracker, ELLMTagSet InTagSet)
+	{
+		if (!FLowLevelMemTracker::bIsDisabled)
+		{
+			Init(InTracker, InTagSet);
+		}
+	}
+	bool IsEnabled() const { return bEnabled; }
+	CORE_API bool TryFindTag(FName UniqueName); // Set name separately so the calculation of name can be skipped if TagSet is disabled
+	CORE_API bool TryAddTagAndActivate(FName UniqueName, const ILLMDynamicTagConstructor& Constructor);
+	CORE_API void Activate();
+
+	~FLLMScopeDynamic()
+	{
+		if (bEnabled)
+		{
+			Destruct();
+		}
+	}
+
+protected:
+	CORE_API void Init(ELLMTracker InTracker, ELLMTagSet InTagSet);
+	CORE_API void Destruct();
+
+	// All fields other than bEnabled are not initialized by constructor unless bEnabled is true
+	const UE::LLMPrivate::FTagData* TagData;
+	ELLMTracker Tracker;
+	bool bEnabled = false; // Needs to be a uint8 rather than a bitfield to minimize ALU operations
 	ELLMTagSet TagSet;
 };
 
@@ -913,6 +1006,7 @@ struct FLLMTagSetAllocationFilter
 #define LLM(...)
 #define LLM_IF_ENABLED(...)
 #define LLM_SCOPE(...)
+#define LLM_SCOPE_DYNAMIC(...)
 #define LLM_TAGSET_SCOPE(...)
 #define LLM_SCOPE_BYNAME(...)
 #define LLM_SCOPE_BYTAG(...)
diff --git a/Engine/Source/Runtime/Core/Public/ProfilingDebugging/TagTrace.h b/Engine/Source/Runtime/Core/Public/ProfilingDebugging/TagTrace.h
--- a/Engine/Source/Runtime/Core/Public/ProfilingDebugging/TagTrace.h
+++ b/Engine/Source/Runtime/Core/Public/ProfilingDebugging/TagTrace.h
@@ -43,11 +43,15 @@ CORE_API int32	MemoryTrace_GetActiveTag();
 class FMemScope
 {
 public:
+	CORE_API FMemScope(); // Used with SetTagAndActivate
 	CORE_API FMemScope(int32 InTag, bool bShouldActivate = true);
 	CORE_API FMemScope(ELLMTag InTag, bool bShouldActivate = true);
 	CORE_API FMemScope(const class FName& InName, bool bShouldActivate = true);
 	CORE_API FMemScope(const UE::LLMPrivate::FTagData* TagData, bool bShouldActivate = true);
 	CORE_API ~FMemScope();
+
+	CORE_API void SetTagAndActivate(const class FName& InName, bool bShouldActivate = true);
+
 private:
 	void ActivateScope(int32 InTag);
 	UE::Trace::Private::FScopedLogScope Inner;
@@ -84,9 +88,11 @@ private:
 inline constexpr int32 TRACE_TAG = 257;
 
 ////////////////////////////////////////////////////////////////////////////////
-#define UE_MEMSCOPE(InTag)				FMemScope PREPROCESSOR_JOIN(MemScope,__LINE__)(InTag);
-#define UE_MEMSCOPE_PTR(InPtr)			FMemScopePtr PREPROCESSOR_JOIN(MemPtrScope,__LINE__)((uint64)InPtr);
-#define UE_MEMSCOPE_DEFAULT(InTag)		FDefaultMemScope PREPROCESSOR_JOIN(MemScope,__LINE__)(InTag);
+#define UE_MEMSCOPE(InTag)					FMemScope PREPROCESSOR_JOIN(MemScope,__LINE__)(InTag);
+#define UE_MEMSCOPE_PTR(InPtr)				FMemScopePtr PREPROCESSOR_JOIN(MemPtrScope,__LINE__)((uint64)InPtr);
+#define UE_MEMSCOPE_DEFAULT(InTag)			FDefaultMemScope PREPROCESSOR_JOIN(MemScope,__LINE__)(InTag);
+#define UE_MEMSCOPE_UNINITIALIZED(Line)		FMemScope PREPROCESSOR_JOIN(MemScope,Line);
+#define UE_MEMSCOPE_ACTIVATE(Line, InTag)	PREPROCESSOR_JOIN(MemScope,Line).SetTagAndActivate(InTag);
 
 #else // UE_MEMORY_TAGS_TRACE_ENABLED
 
diff --git a/Engine/Source/Runtime/CoreUObject/Private/Serialization/AsyncLoading.cpp b/Engine/Source/Runtime/CoreUObject/Private/Serialization/AsyncLoading.cpp
--- a/Engine/Source/Runtime/CoreUObject/Private/Serialization/AsyncLoading.cpp
+++ b/Engine/Source/Runtime/CoreUObject/Private/Serialization/AsyncLoading.cpp
@@ -3162,8 +3162,8 @@ void FAsyncPackage::EventDrivenCreateExport(int32 LocalExportIndex)
 
 	const UClass* Class = CastEventDrivenIndexToObject<UClass>(Export.ClassIndex, false);
 	LLM_SCOPE(ELLMTag::AsyncLoading);
-	LLM_SCOPED_TAG_WITH_OBJECT_IN_SET(GetLinkerRoot(), ELLMTagSet::Assets);
-	LLM_SCOPED_TAG_WITH_OBJECT_IN_SET(Class, ELLMTagSet::AssetClasses);
+	LLM_SCOPE_DYNAMIC_STAT_OBJECTPATH(GetLinkerRoot(), ELLMTagSet::Assets);
+	LLM_SCOPE_DYNAMIC_STAT_OBJECTPATH(Class, ELLMTagSet::AssetClasses);
    	UE_TRACE_METADATA_SCOPE_ASSET_FNAME(Export.ObjectName, Class->GetFName(), GetLinkerRoot()->GetFName());
 
 	// Check whether we already loaded the object and if not whether the context flags allow loading it.
@@ -3449,8 +3449,8 @@ void FAsyncPackage::EventDrivenSerializeExport(int32 LocalExportIndex)
 	FObjectExport& Export = Linker->ExportMap[LocalExportIndex];
 
 	const UClass* Class = CastEventDrivenIndexToObject<UClass>(Export.ClassIndex, false);
-	LLM_SCOPED_TAG_WITH_OBJECT_IN_SET(GetLinkerRoot(), ELLMTagSet::Assets);
-	LLM_SCOPED_TAG_WITH_OBJECT_IN_SET(Class, ELLMTagSet::AssetClasses);
+	LLM_SCOPE_DYNAMIC_STAT_OBJECTPATH(GetLinkerRoot(), ELLMTagSet::Assets);
+	LLM_SCOPE_DYNAMIC_STAT_OBJECTPATH(Class, ELLMTagSet::AssetClasses);
 
 	UObject* Object = Export.Object;
 	if (Object && Object->HasAnyFlags(RF_NeedLoad))
diff --git a/Engine/Source/Runtime/CoreUObject/Private/Serialization/AsyncLoading2.cpp b/Engine/Source/Runtime/CoreUObject/Private/Serialization/AsyncLoading2.cpp
--- a/Engine/Source/Runtime/CoreUObject/Private/Serialization/AsyncLoading2.cpp
+++ b/Engine/Source/Runtime/CoreUObject/Private/Serialization/AsyncLoading2.cpp
@@ -6134,8 +6134,8 @@ void FAsyncPackage2::EventDrivenCreateExport(const FAsyncPackageHeaderData& Head
 		return;
 	}
 
-	LLM_SCOPED_TAG_WITH_OBJECT_IN_SET(GetLinkerRoot(), ELLMTagSet::Assets);
-	LLM_SCOPED_TAG_WITH_OBJECT_IN_SET(LoadClass, ELLMTagSet::AssetClasses);
+	LLM_SCOPE_DYNAMIC_STAT_OBJECTPATH(GetLinkerRoot(), ELLMTagSet::Assets);
+	LLM_SCOPE_DYNAMIC_STAT_OBJECTPATH(LoadClass, ELLMTagSet::AssetClasses);
     UE_TRACE_METADATA_SCOPE_ASSET_FNAME(ObjectName, LoadClass->GetFName(), GetLinkerRoot()->GetFName());
 
 	// Try to find existing object first as we cannot in-place replace objects, could have been created by other export in this package
@@ -6336,8 +6336,8 @@ bool FAsyncPackage2::EventDrivenSerializeExport(const FAsyncPackageHeaderData& H
 
 	const UClass* LoadClass = Export.ClassIndex.IsNull() ? UClass::StaticClass() : CastEventDrivenIndexToObject<UClass>(Header, Export.ClassIndex, true);
 	UE_TRACE_METADATA_SCOPE_ASSET(Object, LoadClass);
-	LLM_SCOPED_TAG_WITH_OBJECT_IN_SET(Object, ELLMTagSet::Assets);
-	LLM_SCOPED_TAG_WITH_OBJECT_IN_SET(LoadClass, ELLMTagSet::AssetClasses);
+	LLM_SCOPE_DYNAMIC_STAT_OBJECTPATH(Object->GetOutermost(), ELLMTagSet::Assets);
+	LLM_SCOPE_DYNAMIC_STAT_OBJECTPATH(LoadClass, ELLMTagSet::AssetClasses);
 
 	// cache archetype
 	// prevents GetArchetype from hitting the expensive GetArchetypeFromRequiredInfoImpl
diff --git a/Engine/Source/Runtime/CoreUObject/Private/UObject/Obj.cpp b/Engine/Source/Runtime/CoreUObject/Private/UObject/Obj.cpp
--- a/Engine/Source/Runtime/CoreUObject/Private/UObject/Obj.cpp
+++ b/Engine/Source/Runtime/CoreUObject/Private/UObject/Obj.cpp
@@ -1261,8 +1261,8 @@ void UObject::ConditionalPostLoad()
 #endif
 				UPackage* Package = GetPackage();
 				UE_SCOPED_COOK_STAT(Package->GetFName(), EPackageEventStatType::LoadPackage);
-				LLM_SCOPED_TAG_WITH_OBJECT_IN_SET(Package, ELLMTagSet::Assets);
-				LLM_SCOPED_TAG_WITH_OBJECT_IN_SET(GetClass(), ELLMTagSet::AssetClasses);
+				LLM_SCOPE_DYNAMIC_STAT_OBJECTPATH(Package, ELLMTagSet::Assets);
+				LLM_SCOPE_DYNAMIC_STAT_OBJECTPATH(GetClass(), ELLMTagSet::AssetClasses);
 				UE_TRACE_METADATA_SCOPE_ASSET(this, GetClass());
 				TRACE_LOADTIME_POSTLOAD_OBJECT_SCOPE(this);
 				
diff --git a/Engine/Source/Runtime/CoreUObject/Private/UObject/UObjectGlobals.cpp b/Engine/Source/Runtime/CoreUObject/Private/UObject/UObjectGlobals.cpp
--- a/Engine/Source/Runtime/CoreUObject/Private/UObject/UObjectGlobals.cpp
+++ b/Engine/Source/Runtime/CoreUObject/Private/UObject/UObjectGlobals.cpp
@@ -1971,11 +1971,7 @@ UPackage* LoadPackage(UPackage* InOuter, const FPackagePath& PackagePath, uint32
 	FScopeCycleCounter CycleCounter( StatId );
 #endif // STATS
 
-	// since we are faking the object name, this is basically a duplicate of LLM_SCOPED_TAG_WITH_OBJECT_IN_SET
-	LLM_SCOPED_TAG_WITH_STAT_NAME_IN_SET(FLowLevelMemTracker::Get().IsTagSetActive(ELLMTagSet::Assets) ?
-											PackagePath.GetPackageFName() :
-											NAME_None,
-										 ELLMTagSet::Assets, ELLMTracker::Default);
+	LLM_SCOPE_DYNAMIC_STAT_OBJECTPATH_FNAME(PackagePath.GetPackageFName(), ELLMTagSet::Assets);
 	TRACE_LOADTIME_REQUEST_GROUP_SCOPE(TEXT("SyncLoad - %s"), *PackagePath.GetDebugName());
 	
 	// if this is a supported asset, it should be loaded fully rather than just for diffing
diff --git a/Engine/Source/Runtime/Engine/Private/Components/ActorComponent.cpp b/Engine/Source/Runtime/Engine/Private/Components/ActorComponent.cpp
--- a/Engine/Source/Runtime/Engine/Private/Components/ActorComponent.cpp
+++ b/Engine/Source/Runtime/Engine/Private/Components/ActorComponent.cpp
@@ -1588,7 +1588,7 @@ void UActorComponent::OnDestroyPhysicsState()
 void UActorComponent::CreatePhysicsState(bool bAllowDeferral)
 {
 	LLM_SCOPE(ELLMTag::Chaos);
-	LLM_SCOPED_TAG_WITH_OBJECT_IN_SET(GetOutermost(), ELLMTagSet::Assets);
+	LLM_SCOPE_DYNAMIC_STAT_OBJECTPATH(GetOutermost(), ELLMTagSet::Assets);
 
 	SCOPE_CYCLE_COUNTER(STAT_ComponentCreatePhysicsState);
 
@@ -1674,7 +1674,7 @@ void UActorComponent::ExecuteRegisterEvents(FRegisterComponentContext* Context)
 	{
 		SCOPE_CYCLE_COUNTER(STAT_ComponentCreateRenderState);
 		LLM_SCOPE(ELLMTag::SceneRender);
-		LLM_SCOPED_TAG_WITH_OBJECT_IN_SET(GetOutermost(), ELLMTagSet::Assets);
+		LLM_SCOPE_DYNAMIC_STAT_OBJECTPATH(GetOutermost(), ELLMTagSet::Assets);
 		CreateRenderState_Concurrent(Context);
 		checkf(bRenderStateCreated, TEXT("Failed to route CreateRenderState_Concurrent (%s)"), *GetFullName());
 	}
@@ -1784,7 +1784,7 @@ void UActorComponent::RemoveTickPrerequisiteComponent(UActorComponent* Prerequis
 void UActorComponent::DoDeferredRenderUpdates_Concurrent()
 {
 	LLM_SCOPE(ELLMTag::SceneRender);
-	LLM_SCOPED_TAG_WITH_OBJECT_IN_SET(GetOutermost(), ELLMTagSet::Assets);
+	LLM_SCOPE_DYNAMIC_STAT_OBJECTPATH(GetOutermost(), ELLMTagSet::Assets);
 
 	checkf(!IsUnreachable(), TEXT("%s"), *GetFullName());
 	checkf(!IsTemplate(), TEXT("%s"), *GetFullName());
diff --git a/Engine/Source/Runtime/Engine/Private/Components/SkinnedMeshComponent.cpp b/Engine/Source/Runtime/Engine/Private/Components/SkinnedMeshComponent.cpp
--- a/Engine/Source/Runtime/Engine/Private/Components/SkinnedMeshComponent.cpp
+++ b/Engine/Source/Runtime/Engine/Private/Components/SkinnedMeshComponent.cpp
@@ -911,7 +911,7 @@ void USkinnedMeshComponent::CreateRenderState_Concurrent(FRegisterComponentConte
 
 	if( GetSkinnedAsset() )
 	{
-		LLM_SCOPED_TAG_WITH_OBJECT_IN_SET(GetSkinnedAsset()->GetOutermost(), ELLMTagSet::Assets);
+		LLM_SCOPE_DYNAMIC_STAT_OBJECTPATH(GetSkinnedAsset()->GetOutermost(), ELLMTagSet::Assets);
 
 		// Attempting to track down UE-45505, where it looks as if somehow a skeletal mesh component's mesh has only been partially loaded, causing a mismatch in the LOD arrays
 		checkf(!GetSkinnedAsset()->HasAnyFlags(RF_NeedLoad | RF_NeedPostLoad | RF_NeedPostLoadSubobjects | RF_WillBeLoaded), TEXT("Attempting to create render state for a skeletal mesh that is is not fully loaded. Mesh: %s"), *GetSkinnedAsset()->GetName());
@@ -2269,7 +2269,7 @@ bool USkinnedMeshComponent::AllocateTransformData()
 	// Allocate transforms if not present.
 	if (GetSkinnedAsset() != nullptr && LeaderPoseComponent == nullptr )
 	{
-		LLM_SCOPED_TAG_WITH_OBJECT_IN_SET(GetSkinnedAsset()->GetOutermost(), ELLMTagSet::Assets);
+		LLM_SCOPE_DYNAMIC_STAT_OBJECTPATH(GetSkinnedAsset()->GetOutermost(), ELLMTagSet::Assets);
 
 		const int32 NumBones = GetSkinnedAsset()->GetRefSkeleton().GetNum();
 
diff --git a/Engine/Source/Runtime/Engine/Private/Texture2D.cpp b/Engine/Source/Runtime/Engine/Private/Texture2D.cpp
--- a/Engine/Source/Runtime/Engine/Private/Texture2D.cpp
+++ b/Engine/Source/Runtime/Engine/Private/Texture2D.cpp
@@ -1326,7 +1326,7 @@ FVirtualTexture2DResource::~FVirtualTexture2DResource()
 void FVirtualTexture2DResource::InitRHI(FRHICommandListBase&)
 {
 	TRACE_CPUPROFILER_EVENT_SCOPE(FVirtualTexture2DResource::InitRHI);
-	LLM_SCOPED_TAG_WITH_STAT_NAME_IN_SET(FLowLevelMemTracker::Get().IsTagSetActive(ELLMTagSet::Assets) ? PackageName : NAME_None, ELLMTagSet::Assets, ELLMTracker::Default);
+	LLM_SCOPE_DYNAMIC_STAT_OBJECTPATH_FNAME(PackageName, ELLMTagSet::Assets);
 
 	uint32 MaxAnisotropy = 0;
 	if (VirtualTextureScalability::IsAnisotropicFilteringEnabled())
diff --git a/Engine/Source/Runtime/Engine/Private/TextureRenderTarget2D.cpp b/Engine/Source/Runtime/Engine/Private/TextureRenderTarget2D.cpp
--- a/Engine/Source/Runtime/Engine/Private/TextureRenderTarget2D.cpp
+++ b/Engine/Source/Runtime/Engine/Private/TextureRenderTarget2D.cpp
@@ -624,7 +624,7 @@ ETextureCreateFlags FTextureRenderTarget2DResource::GetCreateFlags()
  */
 void FTextureRenderTarget2DResource::InitRHI(FRHICommandListBase& RHICmdList)
 {
-	LLM_SCOPED_TAG_WITH_OBJECT_IN_SET(Owner->GetOutermost(), ELLMTagSet::Assets);
+	LLM_SCOPE_DYNAMIC_STAT_OBJECTPATH(Owner->GetOutermost(), ELLMTagSet::Assets);
 
 	if( TargetSizeX > 0 && TargetSizeY > 0 )
 	{
@@ -708,7 +708,7 @@ void FTextureRenderTarget2DResource::ReleaseRHI()
  */
 void FTextureRenderTarget2DResource::UpdateDeferredResource( FRHICommandListImmediate& RHICmdList, bool bClearRenderTarget/*=true*/ )
 {
-	LLM_SCOPED_TAG_WITH_OBJECT_IN_SET(Owner->GetOutermost(), ELLMTagSet::Assets);
+	LLM_SCOPE_DYNAMIC_STAT_OBJECTPATH(Owner->GetOutermost(), ELLMTagSet::Assets);
 
 	SCOPED_DRAW_EVENT(RHICmdList, GPUResourceUpdate)
 	RemoveFromDeferredUpdateList();
diff --git a/Engine/Source/Runtime/Engine/Private/TextureRenderTarget2DArray.cpp b/Engine/Source/Runtime/Engine/Private/TextureRenderTarget2DArray.cpp
--- a/Engine/Source/Runtime/Engine/Private/TextureRenderTarget2DArray.cpp
+++ b/Engine/Source/Runtime/Engine/Private/TextureRenderTarget2DArray.cpp
@@ -223,7 +223,7 @@ UTexture2DArray* UTextureRenderTarget2DArray::ConstructTexture2DArray(UObject* O
  */
 void FTextureRenderTarget2DArrayResource::InitRHI(FRHICommandListBase& RHICmdList)
 {
-	LLM_SCOPED_TAG_WITH_OBJECT_IN_SET(Owner->GetOutermost(), ELLMTagSet::Assets);
+	LLM_SCOPE_DYNAMIC_STAT_OBJECTPATH(Owner->GetOutermost(), ELLMTagSet::Assets);
 
 	if((Owner->SizeX > 0) && (Owner->SizeY > 0) && (Owner->Slices > 0))
 	{
@@ -305,7 +305,7 @@ void FTextureRenderTarget2DArrayResource::ReleaseRHI()
  */
 void FTextureRenderTarget2DArrayResource::UpdateDeferredResource(FRHICommandListImmediate& RHICmdList, bool bClearRenderTarget/*=true*/)
 {
-	LLM_SCOPED_TAG_WITH_OBJECT_IN_SET(Owner->GetOutermost(), ELLMTagSet::Assets);
+	LLM_SCOPE_DYNAMIC_STAT_OBJECTPATH(Owner->GetOutermost(), ELLMTagSet::Assets);
 
 	RemoveFromDeferredUpdateList();
 
diff --git a/Engine/Source/Runtime/Engine/Private/TextureRenderTargetCube.cpp b/Engine/Source/Runtime/Engine/Private/TextureRenderTargetCube.cpp
--- a/Engine/Source/Runtime/Engine/Private/TextureRenderTargetCube.cpp
+++ b/Engine/Source/Runtime/Engine/Private/TextureRenderTargetCube.cpp
@@ -252,7 +252,7 @@ UTextureCube* UTextureRenderTargetCube::ConstructTextureCube(
  */
 void FTextureRenderTargetCubeResource::InitRHI(FRHICommandListBase& RHICmdList)
 {
-	LLM_SCOPED_TAG_WITH_OBJECT_IN_SET(Owner->GetOutermost(), ELLMTagSet::Assets);
+	LLM_SCOPE_DYNAMIC_STAT_OBJECTPATH(Owner->GetOutermost(), ELLMTagSet::Assets);
 
 	if(Owner->SizeX > 0)
 	{
@@ -344,7 +344,7 @@ void FTextureRenderTargetCubeResource::ReleaseRHI()
  */
 void FTextureRenderTargetCubeResource::UpdateDeferredResource(FRHICommandListImmediate& RHICmdList, bool bClearRenderTarget/*=true*/)
 {
-	LLM_SCOPED_TAG_WITH_OBJECT_IN_SET(Owner->GetOutermost(), ELLMTagSet::Assets);
+	LLM_SCOPE_DYNAMIC_STAT_OBJECTPATH(Owner->GetOutermost(), ELLMTagSet::Assets);
 
 	RemoveFromDeferredUpdateList();
 
diff --git a/Engine/Source/Runtime/Engine/Private/TextureRenderTargetVolume.cpp b/Engine/Source/Runtime/Engine/Private/TextureRenderTargetVolume.cpp
--- a/Engine/Source/Runtime/Engine/Private/TextureRenderTargetVolume.cpp
+++ b/Engine/Source/Runtime/Engine/Private/TextureRenderTargetVolume.cpp
@@ -229,7 +229,7 @@ UVolumeTexture* UTextureRenderTargetVolume::ConstructTextureVolume(UObject* ObjO
  */
 void FTextureRenderTargetVolumeResource::InitRHI(FRHICommandListBase& RHICmdList)
 {
-	LLM_SCOPED_TAG_WITH_OBJECT_IN_SET(Owner->GetOutermost(), ELLMTagSet::Assets);
+	LLM_SCOPE_DYNAMIC_STAT_OBJECTPATH(Owner->GetOutermost(), ELLMTagSet::Assets);
 
 	if((Owner->SizeX > 0) && (Owner->SizeY > 0) && (Owner->SizeZ > 0))
 	{
@@ -310,7 +310,7 @@ void FTextureRenderTargetVolumeResource::ReleaseRHI()
  */
 void FTextureRenderTargetVolumeResource::UpdateDeferredResource(FRHICommandListImmediate& RHICmdList, bool bClearRenderTarget/*=true*/)
 {
-	LLM_SCOPED_TAG_WITH_OBJECT_IN_SET(Owner->GetOutermost(), ELLMTagSet::Assets);
+	LLM_SCOPE_DYNAMIC_STAT_OBJECTPATH(Owner->GetOutermost(), ELLMTagSet::Assets);
 
 	RemoveFromDeferredUpdateList();
 
diff --git a/Engine/Source/Runtime/Renderer/Private/RendererScene.cpp b/Engine/Source/Runtime/Renderer/Private/RendererScene.cpp
--- a/Engine/Source/Runtime/Renderer/Private/RendererScene.cpp
+++ b/Engine/Source/Runtime/Renderer/Private/RendererScene.cpp
@@ -1946,7 +1946,7 @@ void FScene::BatchAddPrimitives(TArrayView<UPrimitiveComponent*> InPrimitives)
 		return;
 	}
 #endif
-	LLM_SCOPED_TAG_WITH_OBJECT_IN_SET(InPrimitives[0]->GetOutermost(), ELLMTagSet::Assets);
+	LLM_SCOPE_DYNAMIC_STAT_OBJECTPATH(InPrimitives[0]->GetOutermost(), ELLMTagSet::Assets);
 
 	SCOPE_CYCLE_COUNTER(STAT_AddScenePrimitiveGT);
 	SCOPED_NAMED_EVENT(FScene_AddPrimitive, FColor::Green);
