Use TArray bulk serialization for all arithmetic types.

diff --git a/Engine/Source/Runtime/Core/Public/Containers/Array.h b/Engine/Source/Runtime/Core/Public/Containers/Array.h
--- a/Engine/Source/Runtime/Core/Public/Containers/Array.h
+++ b/Engine/Source/Runtime/Core/Public/Containers/Array.h
@@ -40,10 +40,11 @@
 	#define TARRAY_RANGED_FOR_CHECKS 1
 #endif
 
-template <typename T> struct TCanBulkSerialize { enum { Value = false }; };
-template<> struct TCanBulkSerialize<unsigned int> { enum { Value = true }; };
-template<> struct TCanBulkSerialize<unsigned short> { enum { Value = true }; };
-template<> struct TCanBulkSerialize<int> { enum { Value = true }; };
+template <typename T>
+struct TCanBulkSerialize
+{
+	enum { Value = std::is_arithmetic_v<T> };
+};
 
 // Forward declarations
 
