Added PREPROCESSOR_VA_ARG_COUNT(...) and PREPROCESSOR_APPEND_VA_ARG_COUNT(Prefix, ...) macros, which evaluate to the number of arguments passed to it and a token made by appending a prefix to the count respectively.

diff --git a/Engine/Source/Runtime/Core/Public/HAL/PreprocessorHelpers.h b/Engine/Source/Runtime/Core/Public/HAL/PreprocessorHelpers.h
--- a/Engine/Source/Runtime/Core/Public/HAL/PreprocessorHelpers.h
+++ b/Engine/Source/Runtime/Core/Public/HAL/PreprocessorHelpers.h
@@ -22,6 +22,26 @@
 // Expands to the parameter list of the macro - used to pass a *potentially* comma-separated identifier to another macro as a single parameter
 #define PREPROCESSOR_COMMA_SEPARATED(first, ...) first, ##__VA_ARGS__
 
+// Expands to a number which is the count of variadic arguments passed to it.
+#define PREPROCESSOR_VA_ARG_COUNT(...) PREPROCESSOR_APPEND_VA_ARG_COUNT(, ##__VA_ARGS__)
+
+// Expands to a token of Prefix##<count>, where <count> is the number of variadic arguments.
+//
+// Example:
+//   PREPROCESSOR_APPEND_VA_ARG_COUNT(SOME_MACRO_)          => SOME_MACRO_0
+//   PREPROCESSOR_APPEND_VA_ARG_COUNT(SOME_MACRO_, a, b, c) => SOME_MACRO_3
+#if !defined(_MSVC_TRADITIONAL) || !_MSVC_TRADITIONAL
+	#define PREPROCESSOR_APPEND_VA_ARG_COUNT(Prefix, ...) PREPROCESSOR_APPEND_VA_ARG_COUNT_INNER(Prefix, ##__VA_ARGS__, 26, 25, 24, 23, 22, 21, 20, 19, 18, 17, 16, 15, 14, 13, 12, 11, 10, 9, 8, 7, 6, 5, 4, 3, 2, 1, 0)
+#else
+	#define PREPROCESSOR_APPEND_VA_ARG_COUNT(Prefix, ...) PREPROCESSOR_APPEND_VA_ARG_COUNT_INVOKE(PREPROCESSOR_APPEND_VA_ARG_COUNT_INNER, (Prefix, ##__VA_ARGS__, 26, 25, 24, 23, 22, 21, 20, 19, 18, 17, 16, 15, 14, 13, 12, 11, 10, 9, 8, 7, 6, 5, 4, 3, 2, 1, 0))
+
+	// MSVC's traditional preprocessor doesn't handle the zero-argument case correctly, so we use a workaround.
+	// The workaround uses token pasting of Macro##ArgsInParens, which the conformant preprocessor doesn't like and emits C5103.
+	#define PREPROCESSOR_APPEND_VA_ARG_COUNT_INVOKE(Macro, ArgsInParens) PREPROCESSOR_APPEND_VA_ARG_COUNT_EXPAND(Macro##ArgsInParens)
+	#define PREPROCESSOR_APPEND_VA_ARG_COUNT_EXPAND(Arg) Arg
+#endif
+#define PREPROCESSOR_APPEND_VA_ARG_COUNT_INNER(Prefix,a,b,c,d,e,f,g,h,i,j,k,l,m,n,o,p,q,r,s,t,u,v,w,x,y,z,count,...) Prefix##count
+
 // Expands to nothing - used as a placeholder
 #define PREPROCESSOR_NOTHING
 
diff --git a/Engine/Source/Runtime/Core/Tests/HAL/PreprocessorHelperTests.cpp b/Engine/Source/Runtime/Core/Tests/HAL/PreprocessorHelperTests.cpp
new file mode 100644
--- /dev/null
+++ b/Engine/Source/Runtime/Core/Tests/HAL/PreprocessorHelperTests.cpp
@@ -0,0 +1,51 @@
+// Copyright Epic Games, Inc. All Rights Reserved.
+
+#include "HAL/PreprocessorHelpers.h"
+
+#include "Tests/TestHarnessAdapter.h"
+
+#if WITH_TESTS
+
+TEST_CASE_NAMED(FPreprocessorHelperArgCountTest, "System::Core::HAL::PREPROCESSOR_VA_ARG_COUNT", "[ApplicationContextMask][EngineFilter]")
+{
+	STATIC_CHECK(PREPROCESSOR_VA_ARG_COUNT() == 0);
+	STATIC_CHECK(PREPROCESSOR_VA_ARG_COUNT(a) == 1);
+	STATIC_CHECK(PREPROCESSOR_VA_ARG_COUNT(a, b) == 2);
+	STATIC_CHECK(PREPROCESSOR_VA_ARG_COUNT(a, b, c) == 3);
+	STATIC_CHECK(PREPROCESSOR_VA_ARG_COUNT(a, b, c, d) == 4);
+	STATIC_CHECK(PREPROCESSOR_VA_ARG_COUNT(a, b, c, d, e) == 5);
+	STATIC_CHECK(PREPROCESSOR_VA_ARG_COUNT(a, b, c, d, e, f) == 6);
+	STATIC_CHECK(PREPROCESSOR_VA_ARG_COUNT(a, b, c, d, e, f, g) == 7);
+	STATIC_CHECK(PREPROCESSOR_VA_ARG_COUNT(a, b, c, d, e, f, g, h) == 8);
+	STATIC_CHECK(PREPROCESSOR_VA_ARG_COUNT(a, b, c, d, e, f, g, h, i) == 9);
+	STATIC_CHECK(PREPROCESSOR_VA_ARG_COUNT(a, b, c, d, e, f, g, h, i, j) == 10);
+}
+
+TEST_CASE_NAMED(FPreprocessorHelperArgAppendCountTest, "System::Core::HAL::PREPROCESSOR_APPEND_VA_ARG_COUNT", "[ApplicationContextMask][EngineFilter]")
+{
+	#define TESTPREFIX_0 22
+	#define TESTPREFIX_1 32
+	#define TESTPREFIX_2 41
+	#define TESTPREFIX_3 52
+	#define TESTPREFIX_4 61
+	#define TESTPREFIX_5 70
+	#define TESTPREFIX_6 77
+	#define TESTPREFIX_7 82
+	#define TESTPREFIX_8 89
+	#define TESTPREFIX_9 90
+	#define TESTPREFIX_10 96
+
+	STATIC_CHECK(PREPROCESSOR_APPEND_VA_ARG_COUNT(TESTPREFIX_) == TESTPREFIX_0);
+	STATIC_CHECK(PREPROCESSOR_APPEND_VA_ARG_COUNT(TESTPREFIX_, a) == TESTPREFIX_1);
+	STATIC_CHECK(PREPROCESSOR_APPEND_VA_ARG_COUNT(TESTPREFIX_, a, b) == TESTPREFIX_2);
+	STATIC_CHECK(PREPROCESSOR_APPEND_VA_ARG_COUNT(TESTPREFIX_, a, b, c) == TESTPREFIX_3);
+	STATIC_CHECK(PREPROCESSOR_APPEND_VA_ARG_COUNT(TESTPREFIX_, a, b, c, d) == TESTPREFIX_4);
+	STATIC_CHECK(PREPROCESSOR_APPEND_VA_ARG_COUNT(TESTPREFIX_, a, b, c, d, e) == TESTPREFIX_5);
+	STATIC_CHECK(PREPROCESSOR_APPEND_VA_ARG_COUNT(TESTPREFIX_, a, b, c, d, e, f) == TESTPREFIX_6);
+	STATIC_CHECK(PREPROCESSOR_APPEND_VA_ARG_COUNT(TESTPREFIX_, a, b, c, d, e, f, g) == TESTPREFIX_7);
+	STATIC_CHECK(PREPROCESSOR_APPEND_VA_ARG_COUNT(TESTPREFIX_, a, b, c, d, e, f, g, h) == TESTPREFIX_8);
+	STATIC_CHECK(PREPROCESSOR_APPEND_VA_ARG_COUNT(TESTPREFIX_, a, b, c, d, e, f, g, h, i) == TESTPREFIX_9);
+	STATIC_CHECK(PREPROCESSOR_APPEND_VA_ARG_COUNT(TESTPREFIX_, a, b, c, d, e, f, g, h, i, j) == TESTPREFIX_10);
+}
+
+#endif //WITH_TESTS
