Adding flags to CSV metadata for LTO and PG_Instrumentation so we can show them in PRS

diff --git a/Engine/Source/Runtime/Core/Private/ProfilingDebugging/CsvProfiler.cpp b/Engine/Source/Runtime/Core/Private/ProfilingDebugging/CsvProfiler.cpp
--- a/Engine/Source/Runtime/Core/Private/ProfilingDebugging/CsvProfiler.cpp
+++ b/Engine/Source/Runtime/Core/Private/ProfilingDebugging/CsvProfiler.cpp
@@ -2789,6 +2789,8 @@ FCsvProfiler::FCsvProfiler()
 	SetMetadataInternal(TEXT("OS"), *OSString);
 	SetMetadataInternal(TEXT("CPU"), *FPlatformMisc::GetDeviceMakeAndModel());
 	SetMetadataInternal(TEXT("PGOEnabled"), FPlatformMisc::IsPGOEnabled() ? TEXT("1") : TEXT("0"));
+	SetMetadataInternal(TEXT("PGOProfilingEnabled"), PLATFORM_COMPILER_OPTIMIZATION_PG_PROFILING ? TEXT("1") : TEXT("0"));//True if Profile Guided Optimisation Instrumentation is enabled 
+	SetMetadataInternal(TEXT("LTOEnabled"), PLATFORM_COMPILER_OPTIMIZATION_LTCG ? TEXT("1") : TEXT("0"));//True if Link Time Optimisation is enabled 
 	SetMetadataInternal(TEXT("ASan"), USING_ADDRESS_SANITISER ? TEXT("1") : TEXT("0"));
 
 #if !UE_BUILD_SHIPPING
