adjusted LLM tracing for tasks

diff --git a/Engine/Source/Runtime/Core/Private/Async/TaskGraph.cpp b/Engine/Source/Runtime/Core/Private/Async/TaskGraph.cpp
--- a/Engine/Source/Runtime/Core/Private/Async/TaskGraph.cpp
+++ b/Engine/Source/Runtime/Core/Private/Async/TaskGraph.cpp
@@ -1951,7 +1951,7 @@ public:
 private:
 	void QueueTask(class FBaseGraphTask* Task, bool bWakeUpWorker, ENamedThreads::Type InThreadToExecuteOn, ENamedThreads::Type InCurrentThreadIfKnown) override
 	{
-		LLM_SCOPE_BYNAME(TEXT("Tasks/Scheduler/QueueTask"));
+		LLM_SCOPE_BYNAME(TEXT("Tasks/QueueTask"));
 
 		if (ENamedThreads::GetThreadIndex(InThreadToExecuteOn) == ENamedThreads::AnyThread)
 		{
@@ -2081,6 +2081,7 @@ private:
 
 	void WaitUntilTasksComplete(const FGraphEventArray& Tasks, ENamedThreads::Type CurrentThreadIfKnown = ENamedThreads::AnyThread) final override
 	{
+		LLM_SCOPE_BYNAME(TEXT("Tasks/WaitUntilTasksComplete"));
 		TaskTrace::FWaitingScope WaitingScope(GetTraceIds(Tasks));
 		TRACE_CPUPROFILER_EVENT_SCOPE(WaitUntilTasksComplete);
 
diff --git a/Engine/Source/Runtime/Core/Public/Async/TaskGraphInterfaces.h b/Engine/Source/Runtime/Core/Public/Async/TaskGraphInterfaces.h
--- a/Engine/Source/Runtime/Core/Public/Async/TaskGraphInterfaces.h
+++ b/Engine/Source/Runtime/Core/Public/Async/TaskGraphInterfaces.h
@@ -508,6 +508,8 @@ public:
 			return;
 		}
 
+		LLM_SCOPE_BYNAME(TEXT("Tasks/FGraphEvent/DontCompleteUntil"));
+
 		if (GetExtendedPriority() == UE::Tasks::EExtendedTaskPriority::TaskEvent)
 		{	// TaskEvent can't have nested tasks, add it as a prerequisite instead
 			AddPrerequisites(*NestedTask);
diff --git a/Engine/Source/Runtime/Core/Public/Tasks/Task.h b/Engine/Source/Runtime/Core/Public/Tasks/Task.h
--- a/Engine/Source/Runtime/Core/Public/Tasks/Task.h
+++ b/Engine/Source/Runtime/Core/Public/Tasks/Task.h
@@ -80,6 +80,7 @@ namespace UE::Tasks
 			// @return true if the task is completed
 			bool Wait(FTimespan Timeout = FTimespan::MaxValue()) const
 			{
+				LLM_SCOPE_BYNAME(TEXT("Tasks/Wait"));
 				return !IsValid() || Pimpl->Wait(FTimeout{ Timeout });
 			}
 
@@ -88,6 +89,7 @@ namespace UE::Tasks
 			// @return true if the task is completed
 			bool BusyWait(FTimespan Timeout = FTimespan::MaxValue()) const
 			{
+				LLM_SCOPE_BYNAME(TEXT("Tasks/BusyWait"));
 				return !IsValid() || Pimpl->BusyWait(FTimeout{ Timeout });
 			}
 
@@ -97,6 +99,7 @@ namespace UE::Tasks
 			template<typename ConditionType>
 			bool BusyWait(ConditionType&& Condition) const
 			{
+				LLM_SCOPE_BYNAME(TEXT("Tasks/BusyWait"));
 				return !IsValid() || Pimpl->BusyWait(Forward<ConditionType>(Condition));
 			}
 
@@ -115,6 +118,8 @@ namespace UE::Tasks
 			{
 				check(!IsValid());
 
+				LLM_SCOPE_BYNAME(TEXT("Tasks/Launch"));
+
 				using FExecutableTask = Private::TExecutableTask<std::decay_t<TaskBodyType>>;
 				FExecutableTask* Task = FExecutableTask::Create(DebugName, Forward<TaskBodyType>(TaskBody), Priority, ExtendedPriority);
 				// this must happen before launching, to support an ability to access the task itself from inside it
@@ -140,6 +145,8 @@ namespace UE::Tasks
 			{
 				check(!IsValid());
 
+				LLM_SCOPE_BYNAME(TEXT("Tasks/Launch"));
+
 				using FExecutableTask = Private::TExecutableTask<std::decay_t<TaskBodyType>>;
 				FExecutableTask* Task = FExecutableTask::Create(DebugName, Forward<TaskBodyType>(TaskBody), Priority, ExtendedPriority);
 				Task->AddPrerequisites(Forward<PrerequisitesCollectionType>(Prerequisites));
@@ -228,6 +235,7 @@ namespace UE::Tasks
 		template<typename PrerequisitesType>
 		void AddPrerequisites(const PrerequisitesType& Prerequisites)
 		{
+			LLM_SCOPE_BYNAME(TEXT("Tasks/FTaskEvent/AddPrerequisites"));
 			Pimpl->AddPrerequisites(Prerequisites);
 		}
 
@@ -235,6 +243,8 @@ namespace UE::Tasks
 		{
 			if (!IsCompleted()) // event can be triggered multiple times
 			{
+				LLM_SCOPE_BYNAME(TEXT("Tasks/FTaskEvent/Trigger"));
+
 				// An event is not "in the system" until it's triggered, and should be kept alive only by external references. Once it's triggered it's in the system 
 				// and can outlive external references, so we need to keep it alive by holding an internal reference. It will be released when the event is signalled
 				Pimpl->AddRef();
@@ -318,6 +328,7 @@ namespace UE::Tasks
 	template<typename TaskCollectionType>
 	bool Wait(const TaskCollectionType& Tasks, FTimespan InTimeout/* = FTimespan::MaxValue()*/)
 	{
+		LLM_SCOPE_BYNAME(TEXT("Tasks/Wait"));
 		TaskTrace::FWaitingScope WaitingScope(Private::GetTraceIds(Tasks));
 		TRACE_CPUPROFILER_EVENT_SCOPE(Tasks::Wait);
 
@@ -348,6 +359,7 @@ namespace UE::Tasks
 	template<typename TaskCollectionType>
 	bool BusyWait(const TaskCollectionType& Tasks, FTimespan InTimeout/* = FTimespan::MaxValue()*/)
 	{
+		LLM_SCOPE_BYNAME(TEXT("Tasks/BusyWait"));
 		TaskTrace::FWaitingScope WaitingScope(Private::GetTraceIds(Tasks));
 		TRACE_CPUPROFILER_EVENT_SCOPE(Tasks::BusyWait);
 
@@ -406,6 +418,7 @@ namespace UE::Tasks
 	template<typename TaskType>
 	void AddNested(const TaskType& Nested)
 	{
+		LLM_SCOPE_BYNAME(TEXT("Tasks/AddNested"));
 		Private::FTaskBase* Parent = Private::GetCurrentTask();
 		check(Parent != nullptr);
 		Parent->AddNested(*Nested.Pimpl);
diff --git a/Engine/Source/Runtime/Core/Public/Tasks/TaskPrivate.h b/Engine/Source/Runtime/Core/Public/Tasks/TaskPrivate.h
--- a/Engine/Source/Runtime/Core/Public/Tasks/TaskPrivate.h
+++ b/Engine/Source/Runtime/Core/Public/Tasks/TaskPrivate.h
@@ -216,8 +216,6 @@ namespace UE::Tasks
 					return false;
 				}
 
-				LLM_SCOPE_BYNAME(TEXT("Tasks/FTaskBase/AddPrerequisites"));
-
 				Prerequisite.AddRef(); // keep it alive until this task's execution
 				Prerequisites.Push(&Prerequisite); // release memory order
 				return true;
@@ -274,7 +272,6 @@ namespace UE::Tasks
 
 					if (Prerequisite->AddSubsequent(*this)) // acq_rel memory order
 					{
-						LLM_SCOPE_BYNAME(TEXT("Tasks/FTaskBase/AddPrerequisites"));
 						Prerequisite->AddRef(); // keep it alive until this task's execution
 						Prerequisites.Push(Prerequisite); // release memory order
 					}
@@ -293,7 +290,6 @@ namespace UE::Tasks
 			// returns false if the task is already completed and the subsequent wasn't added
 			bool AddSubsequent(FTaskBase& Subsequent)
 			{
-				LLM_SCOPE_BYNAME(TEXT("Tasks/FTaskBase/AddSubsequent"));
 				TaskTrace::SubsequentAdded(GetTraceId(), Subsequent.GetTraceId()); // doesn't matter if we suceeded below, we need to record task dependency
 				return Subsequents.PushIfNotClosed(&Subsequent);
 			}
@@ -351,7 +347,6 @@ namespace UE::Tasks
 
 				if (Nested.AddSubsequent(*this)) // "release" memory order
 				{
-					LLM_SCOPE_BYNAME(TEXT("Tasks/FTaskBase/AddNested"));
 					Nested.AddRef(); // keep it alive as we store it in `Prerequisites` and we can need it to try to retract it. it's released on closing the task
 					Prerequisites.Push(&Nested);
 				}
@@ -521,7 +516,6 @@ namespace UE::Tasks
 							FTaskBase* PrevPipedTask = TryPushIntoPipe();
 							if (PrevPipedTask != nullptr) // the pipe is blocked
 							{
-								LLM_SCOPE_BYNAME(TEXT("Tasks/FTaskBase/AddPrerequisites"));
 								// the prev task in pipe's chain becomes this task's prerequisite, to enabled piped task retraction.
 								// its ref count already accounted for this ref. the ref will be released when the prereq is not needed anymore
 								Prerequisites.Push(PrevPipedTask);
@@ -735,7 +729,6 @@ namespace UE::Tasks
 			// a helper that deduces the template argument
 			static TExecutableTask* Create(const TCHAR* InDebugName, TaskBodyType&& TaskBody, ETaskPriority InPriority, EExtendedTaskPriority InExtendedPriority)
 			{
-				LLM_SCOPE_BYNAME(TEXT("Tasks/TExecutableTask/Create"));
 				return new TExecutableTask(InDebugName, MoveTemp(TaskBody), InPriority, InExtendedPriority);
 			}
 
@@ -761,7 +754,7 @@ namespace UE::Tasks
 		public:
 			static FTaskEventBase* Create(const TCHAR* DebugName)
 			{
-				LLM_SCOPE_BYNAME(TEXT("Tasks/FTaskEventBase/Create"));
+				LLM_SCOPE_BYNAME(TEXT("Tasks/FTaskEvent/Create"));
 				return new FTaskEventBase(DebugName);
 			}
 
