Make csv named events switchable at runtime via cvars csv.NamedEventsExclusive and csv.NamedEventsTiming

diff --git a/Engine/Source/Runtime/Core/Private/ProfilingDebugging/CsvProfiler.cpp b/Engine/Source/Runtime/Core/Private/ProfilingDebugging/CsvProfiler.cpp
--- a/Engine/Source/Runtime/Core/Private/ProfilingDebugging/CsvProfiler.cpp
+++ b/Engine/Source/Runtime/Core/Private/ProfilingDebugging/CsvProfiler.cpp
@@ -162,6 +162,7 @@ static bool GCsvTrackWaitsOnRenderThread = true;
 static FAutoConsoleVariableRef CVarTrackWaitsAllThreads(TEXT("csv.trackWaitsAllThreads"), GCsvTrackWaitsOnAllThreads, TEXT("Determines whether to track waits on all threads. Note that this incurs a lot of overhead"), ECVF_Default);
 static FAutoConsoleVariableRef CVarTrackWaitsGT(TEXT("csv.trackWaitsGT"), GCsvTrackWaitsOnGameThread, TEXT("Determines whether to track game thread waits. Note that this incurs overhead"), ECVF_Default);
 static FAutoConsoleVariableRef CVarTrackWaitsRT(TEXT("csv.trackWaitsRT"), GCsvTrackWaitsOnRenderThread, TEXT("Determines whether to track render thread waits. Note that this incurs overhead"), ECVF_Default);
+
 //
 // Categories
 //
@@ -283,6 +284,9 @@ static FCsvPersistentCustomStats GCsvPersistentCustomStats;
 bool GCsvProfilerNamedEventsExclusive = false;
 bool GCsvProfilerNamedEventsTiming = false;
 
+static FAutoConsoleVariableRef CVarNamedEventsExclusive(TEXT("csv.NamedEventsExclusive"), GCsvProfilerNamedEventsExclusive, TEXT("Determines whether to emit named events for exclusive stats"), ECVF_Default);
+static FAutoConsoleVariableRef CVarNamedEventsTiming(TEXT("csv.NamedEventsTiming"), GCsvProfilerNamedEventsTiming, TEXT("Determines whether to emit named events for non-exclusive timing stats"), ECVF_Default);
+
 void CsvBeginNamedEvent(FColor Color, const char* NamedEventName)
 {
 #if CPUPROFILERTRACE_ENABLED 
