Rework our record write in AutoRTFM to make it 30% faster.

diff --git a/Engine/Source/Runtime/Core/Private/AutoRTFM/CompilerInterface.cpp b/Engine/Source/Runtime/Core/Private/AutoRTFM/CompilerInterface.cpp
--- a/Engine/Source/Runtime/Core/Private/AutoRTFM/CompilerInterface.cpp
+++ b/Engine/Source/Runtime/Core/Private/AutoRTFM/CompilerInterface.cpp
@@ -40,6 +40,18 @@ extern "C" void autortfm_record_write(FContext* Context, void* Ptr, size_t Size)
 	Context->RecordWrite(Ptr, Size);
 }
 
+extern "C" void autortfm_record_write_8(FContext* Context, void* Ptr)
+{
+	// check for writes to null here so we end up crashing in the user
+	// code rather than in the autortfm runtime.
+	if (UNLIKELY(nullptr == Ptr))
+	{
+		return;
+	}
+
+	Context->RecordWrite<8>(Ptr);
+}
+
 extern "C" void* autortfm_lookup_function(FContext* Context, void* OriginalFunction, const char* Where)
 {
     return FunctionMapLookup(OriginalFunction, Context, Where);
diff --git a/Engine/Source/Runtime/Core/Private/AutoRTFM/Context.h b/Engine/Source/Runtime/Core/Private/AutoRTFM/Context.h
--- a/Engine/Source/Runtime/Core/Private/AutoRTFM/Context.h
+++ b/Engine/Source/Runtime/Core/Private/AutoRTFM/Context.h
@@ -39,6 +39,7 @@ public:
 
     // Record that a write is about to occur at the given LogicalAddress of Size bytes.
     void RecordWrite(void* LogicalAddress, size_t Size);
+    template<unsigned SIZE> void RecordWrite(void* LogicalAddress);
 
     void DidAllocate(void* LogicalAddress, size_t Size);
 
diff --git a/Engine/Source/Runtime/Core/Private/AutoRTFM/ContextInlines.h b/Engine/Source/Runtime/Core/Private/AutoRTFM/ContextInlines.h
--- a/Engine/Source/Runtime/Core/Private/AutoRTFM/ContextInlines.h
+++ b/Engine/Source/Runtime/Core/Private/AutoRTFM/ContextInlines.h
@@ -24,6 +24,11 @@ UE_AUTORTFM_FORCEINLINE void FContext::RecordWrite(void* LogicalAddress, size_t
     CurrentTransaction->RecordWrite(LogicalAddress, Size);
 }
 
+template<unsigned SIZE> UE_AUTORTFM_FORCEINLINE void FContext::RecordWrite(void* LogicalAddress)
+{
+    CurrentTransaction->RecordWrite<SIZE>(LogicalAddress);
+}
+
 UE_AUTORTFM_FORCEINLINE void FContext::DidAllocate(void* LogicalAddress, size_t Size)
 {
     CurrentTransaction->DidAllocate(LogicalAddress, Size);
diff --git a/Engine/Source/Runtime/Core/Private/AutoRTFM/TaggedPtr.h b/Engine/Source/Runtime/Core/Private/AutoRTFM/TaggedPtr.h
--- a/Engine/Source/Runtime/Core/Private/AutoRTFM/TaggedPtr.h
+++ b/Engine/Source/Runtime/Core/Private/AutoRTFM/TaggedPtr.h
@@ -35,20 +35,20 @@ namespace AutoRTFM
 		// Assumes little-endian, so top tag is the last 2-bytes in memory.
 		static constexpr size_t OffsetOfTopTag = PointerBits / 8;
 
-		TTaggedPtr()
+		UE_AUTORTFM_FORCEINLINE TTaggedPtr()
 		{
 			Reset();
 		}
 
-		explicit TTaggedPtr(T* const Pointer) : Payload(reinterpret_cast<uintptr_t>(Pointer))
+		UE_AUTORTFM_FORCEINLINE explicit TTaggedPtr(T* const Pointer) : Payload(reinterpret_cast<uintptr_t>(Pointer))
 		{
 			ASSERT(0 == (Payload & ~PointerMask));
 		}
 
-		explicit TTaggedPtr(const TTaggedPtr& Other) : Payload(Other.Payload) {}
-		explicit TTaggedPtr(const TTaggedPtr&& Other) : Payload(Other.Payload) {}
+		UE_AUTORTFM_FORCEINLINE explicit TTaggedPtr(const TTaggedPtr& Other) : Payload(Other.Payload) {}
+		UE_AUTORTFM_FORCEINLINE explicit TTaggedPtr(const TTaggedPtr&& Other) : Payload(Other.Payload) {}
 
-		TTaggedPtr& operator=(T* const Pointer)
+		UE_AUTORTFM_FORCEINLINE TTaggedPtr& operator=(T* const Pointer)
 		{
 			const uintptr_t NewPayload = reinterpret_cast<uintptr_t>(Pointer);
 
@@ -58,49 +58,49 @@ namespace AutoRTFM
 			return *this;
 		}
 
-		TTaggedPtr& operator=(TTaggedPtr& Other)
+		UE_AUTORTFM_FORCEINLINE TTaggedPtr& operator=(TTaggedPtr& Other)
 		{
 			Payload = Other.Payload;
 			return *this;
 		}
 
-		TTaggedPtr& operator=(TTaggedPtr&& Other)
+		UE_AUTORTFM_FORCEINLINE TTaggedPtr& operator=(TTaggedPtr&& Other)
 		{
 			Payload = Other.Payload;
 			return *this;
 		}
 
-		T& operator*() const
+		UE_AUTORTFM_FORCEINLINE T& operator*() const
 	    {
 	    	return *reinterpret_cast<T*>(Payload & PointerMask);
 	    }
 
-	    T* operator->() const
+	    UE_AUTORTFM_FORCEINLINE T* operator->() const
 	    {
 	    	return reinterpret_cast<T*>(Payload & PointerMask);
 	    }
 
-	    T* Get() const
+	    UE_AUTORTFM_FORCEINLINE T* Get() const
 	    {
 	    	return reinterpret_cast<T*>(Payload & PointerMask);
 	    }
 
-	    void SetTopTag(const TopTagT Tag)
+	    UE_AUTORTFM_FORCEINLINE void SetTopTag(const TopTagT Tag)
 	    {
 	    	*reinterpret_cast<TopTagT*>(reinterpret_cast<uint8_t*>(&Payload) + OffsetOfTopTag) = MoveTemp(Tag);
 	    }
 
-	    TopTagT GetTopTag() const
+	    UE_AUTORTFM_FORCEINLINE TopTagT GetTopTag() const
 	    {
 	    	return *reinterpret_cast<const TopTagT*>(reinterpret_cast<const uint8_t*>(&Payload) + OffsetOfTopTag);
 	    }
 
-	    void Reset()
+	    UE_AUTORTFM_FORCEINLINE void Reset()
 	    {
 	    	Payload = 0;
 	    }
 
-	    uintptr_t GetPayload() const
+	    UE_AUTORTFM_FORCEINLINE uintptr_t GetPayload() const
 	    {
 	    	return Payload;
 	    }
@@ -135,20 +135,20 @@ namespace AutoRTFM
 		// Assumes little-endian, so top tag is the last 2-bytes in memory.
 		static constexpr size_t OffsetOfTopTag = PointerBits / 8;
 
-		TTaggedPtr()
+		UE_AUTORTFM_FORCEINLINE TTaggedPtr()
 		{
 			Reset();
 		}
 
-		TTaggedPtr(void* const Pointer) : Payload(reinterpret_cast<uintptr_t>(Pointer))
+		UE_AUTORTFM_FORCEINLINE TTaggedPtr(void* const Pointer) : Payload(reinterpret_cast<uintptr_t>(Pointer))
 		{
 			ASSERT(0 == (Payload & ~PointerMask));
 		}
 
-		TTaggedPtr(const TTaggedPtr& Other) : Payload(Other.Payload) {}
-		TTaggedPtr(const TTaggedPtr&& Other) : Payload(Other.Payload) {}
+		UE_AUTORTFM_FORCEINLINE TTaggedPtr(const TTaggedPtr& Other) : Payload(Other.Payload) {}
+		UE_AUTORTFM_FORCEINLINE TTaggedPtr(const TTaggedPtr&& Other) : Payload(Other.Payload) {}
 
-		TTaggedPtr& operator=(void* const Pointer)
+		UE_AUTORTFM_FORCEINLINE TTaggedPtr& operator=(void* const Pointer)
 		{
 			const uintptr_t NewPayload = reinterpret_cast<uintptr_t>(Pointer);
 
@@ -158,44 +158,44 @@ namespace AutoRTFM
 			return *this;
 		}
 
-		TTaggedPtr& operator=(TTaggedPtr& Other)
+		UE_AUTORTFM_FORCEINLINE TTaggedPtr& operator=(TTaggedPtr& Other)
 		{
 			Payload = Other.Payload;
 			return *this;
 		}
 
-		TTaggedPtr& operator=(TTaggedPtr&& Other)
+		UE_AUTORTFM_FORCEINLINE TTaggedPtr& operator=(TTaggedPtr&& Other)
 		{
 			Payload = Other.Payload;
 			return *this;
 		}
 
-		bool operator==(const TTaggedPtr& Other) const
+		UE_AUTORTFM_FORCEINLINE bool operator==(const TTaggedPtr& Other) const
 		{
 			return Payload == Other.Payload;
 		}
 
-		void* Get() const
+		UE_AUTORTFM_FORCEINLINE void* Get() const
 	    {
 	    	return reinterpret_cast<void*>(Payload & PointerMask);
 	    }
 
-	    void SetTopTag(const TopTagT Tag)
+	    UE_AUTORTFM_FORCEINLINE void SetTopTag(const TopTagT Tag)
 	    {
 	    	*reinterpret_cast<TopTagT*>(reinterpret_cast<uint8_t*>(&Payload) + OffsetOfTopTag) = Tag;
 	    }
 
-	    TopTagT GetTopTag() const
+	    UE_AUTORTFM_FORCEINLINE TopTagT GetTopTag() const
 	    {
 	    	return *reinterpret_cast<const TopTagT*>(reinterpret_cast<const uint8_t*>(&Payload) + OffsetOfTopTag);
 	    }
 
-	    void Reset()
+	    UE_AUTORTFM_FORCEINLINE void Reset()
 	    {
 	    	Payload = 0;
 	    }
 
-	    uintptr_t GetPayload() const
+	    UE_AUTORTFM_FORCEINLINE uintptr_t GetPayload() const
 	    {
 	    	return Payload;
 	    }
diff --git a/Engine/Source/Runtime/Core/Private/AutoRTFM/Transaction.h b/Engine/Source/Runtime/Core/Private/AutoRTFM/Transaction.h
--- a/Engine/Source/Runtime/Core/Private/AutoRTFM/Transaction.h
+++ b/Engine/Source/Runtime/Core/Private/AutoRTFM/Transaction.h
@@ -70,6 +70,7 @@ public:
     // Record that a write is about to occur at the given LogicalAddress of Size bytes.
     void RecordWrite(void* LogicalAddress, size_t Size);
     void RecordWriteMaxPageSized(void* LogicalAddress, size_t Size);
+    template<unsigned SIZE> void RecordWrite(void* LogicalAddress);
 
     void DidAllocate(void* LogicalAddress, size_t Size);
 
diff --git a/Engine/Source/Runtime/Core/Private/AutoRTFM/TransactionInlines.h b/Engine/Source/Runtime/Core/Private/AutoRTFM/TransactionInlines.h
--- a/Engine/Source/Runtime/Core/Private/AutoRTFM/TransactionInlines.h
+++ b/Engine/Source/Runtime/Core/Private/AutoRTFM/TransactionInlines.h
@@ -16,7 +16,7 @@ UE_AUTORTFM_FORCEINLINE void FTransaction::RecordWriteMaxPageSized(void* Logical
     WriteLog.Push(FWriteLogEntry(LogicalAddress, Size, CopyAddress));
 }
 
-inline void FTransaction::RecordWrite(void* LogicalAddress, size_t Size)
+UE_AUTORTFM_FORCEINLINE void FTransaction::RecordWrite(void* LogicalAddress, size_t Size)
 {
     // If we are recording a stack address that is relative to our current
     // transactions stack location, we do not need to record the data in the
@@ -55,7 +55,35 @@ inline void FTransaction::RecordWrite(void* LogicalAddress, size_t Size)
     RecordWriteMaxPageSized(Address + I, Size - I);
 }
 
-inline void FTransaction::DidAllocate(void* LogicalAddress, size_t Size)
+template<unsigned SIZE> UE_AUTORTFM_FORCEINLINE void FTransaction::RecordWrite(void* LogicalAddress)
+{
+    static_assert(SIZE <= FWriteLogBumpAllocator::MaxSize);
+
+    // If we are recording a stack address that is relative to our current
+    // transactions stack location, we do not need to record the data in the
+    // write log because if that transaction aborted, that memory will cease to
+    // be meaningful anyway!
+    if (Context->IsInnerTransactionStack(LogicalAddress))
+    {
+        Stats.Collect<EStatsKind::HitSetSkippedBecauseOfStackLocalMemory>();
+        return;
+    }
+
+    FMemoryLocation Key(LogicalAddress);
+    Key.SetTopTag(static_cast<uint16_t>(SIZE));
+
+    if (!HitSet.Insert(Key))
+    {
+        Stats.Collect<EStatsKind::HitSetHit>();
+        return;
+    }
+
+    Stats.Collect<EStatsKind::HitSetMiss>();
+
+    RecordWriteMaxPageSized(LogicalAddress, SIZE);
+}
+
+UE_AUTORTFM_FORCEINLINE void FTransaction::DidAllocate(void* LogicalAddress, size_t Size)
 {
     if ((0 < Size) && (Size <= FWriteLogBumpAllocator::MaxSize))
     {
@@ -67,12 +95,12 @@ inline void FTransaction::DidAllocate(void* LogicalAddress, size_t Size)
     }
 }
 
-inline void FTransaction::DeferUntilCommit(TFunction<void()>&& Callback)
+UE_AUTORTFM_FORCEINLINE void FTransaction::DeferUntilCommit(TFunction<void()>&& Callback)
 {
     CommitTasks.Add(MoveTemp(Callback));
 }
 
-inline void FTransaction::DeferUntilAbort(TFunction<void()>&& Callback)
+UE_AUTORTFM_FORCEINLINE void FTransaction::DeferUntilAbort(TFunction<void()>&& Callback)
 {
     AbortTasks.Add(MoveTemp(Callback));
 }
diff --git a/Engine/Source/Runtime/Core/Private/AutoRTFM/Utils.cpp b/Engine/Source/Runtime/Core/Private/AutoRTFM/Utils.cpp
--- a/Engine/Source/Runtime/Core/Private/AutoRTFM/Utils.cpp
+++ b/Engine/Source/Runtime/Core/Private/AutoRTFM/Utils.cpp
@@ -23,6 +23,17 @@ DEFINE_LOG_CATEGORY(LogAutoRTFM)
 namespace AutoRTFM
 {
 
+void PrettyAbort(const char* const File, const unsigned Line, const char* const Function, const char* const Expression)
+{
+    UE_LOG(LogAutoRTFM, Fatal, TEXT("%s:%d:%s: assertion %s failed."), ANSI_TO_TCHAR(File), Line, ANSI_TO_TCHAR(Function), ANSI_TO_TCHAR(Expression));
+
+#if PLATFORM_WINDOWS
+	__assume(false);
+#else
+	__builtin_unreachable();
+#endif // PLATFORM_WINDOWS
+}
+
 FString GetFunctionDescription(void* FunctionPtr)
 {
 #if PLATFORM_WINDOWS
diff --git a/Engine/Source/Runtime/Core/Private/AutoRTFM/Utils.h b/Engine/Source/Runtime/Core/Private/AutoRTFM/Utils.h
--- a/Engine/Source/Runtime/Core/Private/AutoRTFM/Utils.h
+++ b/Engine/Source/Runtime/Core/Private/AutoRTFM/Utils.h
@@ -13,11 +13,11 @@
 
 DECLARE_LOG_CATEGORY_EXTERN(LogAutoRTFM, Display, All)
 
-#define ASSERT(exp) UE_CLOG(UNLIKELY(!(exp)), LogAutoRTFM, Fatal, TEXT("%s:%d:%s: assertion %s failed."), ANSI_TO_TCHAR(__FILE__), __LINE__, ANSI_TO_TCHAR(__PRETTY_FUNCTION__), ANSI_TO_TCHAR(#exp))
-
 namespace AutoRTFM
 {
 
+[[noreturn]] void PrettyAbort(const char* const File, const unsigned Line, const char* const Function, const char* const Expression);
+
 [[noreturn]] inline void Unreachable()
 {
 	UE_LOG(LogAutoRTFM, Fatal, TEXT("Unreachable encountered!"));
@@ -44,3 +44,5 @@ template<size_t A, size_t B> struct PrettyStaticAssert final
 };
 
 } // namespace AutoRTFM
+
+#define ASSERT(exp) do { if (UNLIKELY(!(exp))) PrettyAbort(__FILE__, __LINE__, __PRETTY_FUNCTION__, #exp); } while (false)
diff --git a/Engine/Source/Runtime/Core/Private/AutoRTFM/WriteLog.h b/Engine/Source/Runtime/Core/Private/AutoRTFM/WriteLog.h
--- a/Engine/Source/Runtime/Core/Private/AutoRTFM/WriteLog.h
+++ b/Engine/Source/Runtime/Core/Private/AutoRTFM/WriteLog.h
@@ -15,11 +15,11 @@ namespace AutoRTFM
         FMemoryLocation OriginalAndSize;
         void* Copy;
 
-        FWriteLogEntry() = default;
-        FWriteLogEntry(FWriteLogEntry&) = default;
-        FWriteLogEntry& operator=(FWriteLogEntry&) = default;
+        UE_AUTORTFM_FORCEINLINE FWriteLogEntry() = default;
+        UE_AUTORTFM_FORCEINLINE FWriteLogEntry(FWriteLogEntry&) = default;
+        UE_AUTORTFM_FORCEINLINE FWriteLogEntry& operator=(FWriteLogEntry&) = default;
 
-        explicit FWriteLogEntry(void* Original, size_t Size, void* Copy) :
+        UE_AUTORTFM_FORCEINLINE explicit FWriteLogEntry(void* Original, size_t Size, void* Copy) :
             OriginalAndSize(Original), Copy(Copy)
         {
             OriginalAndSize.SetTopTag(static_cast<uint16_t>(Size));
