Make AutoRTFM macros handle unparenthesized commas in their argument (e.g. in a template argument list).

diff --git a/Engine/Source/Runtime/Core/Public/AutoRTFM/AutoRTFM.h b/Engine/Source/Runtime/Core/Public/AutoRTFM/AutoRTFM.h
--- a/Engine/Source/Runtime/Core/Public/AutoRTFM/AutoRTFM.h
+++ b/Engine/Source/Runtime/Core/Public/AutoRTFM/AutoRTFM.h
@@ -541,15 +541,15 @@ struct FRegisterOpenFunction
 
 // Macro-based variants so we completely compile away when not in use, even in debug builds
 #if UE_AUTORTFM
-#define UE_AUTORTFM_OPEN(x) AutoRTFM::Open([&]() { x })
-#define UE_AUTORTFM_OPENABORT(x) AutoRTFM::OpenAbort([=]() { x })
-#define UE_AUTORTFM_OPENCOMMIT(x) AutoRTFM::OpenCommit([=]() { x })
-#define UE_AUTORTFM_TRANSACT(x) AutoRTFM::Transact([&]() { x })
+#define UE_AUTORTFM_OPEN(...) AutoRTFM::Open([&]() { __VA_ARGS__ })
+#define UE_AUTORTFM_OPENABORT(...) AutoRTFM::OpenAbort([=]() { __VA_ARGS__ })
+#define UE_AUTORTFM_OPENCOMMIT(...) AutoRTFM::OpenCommit([=]() { __VA_ARGS__ })
+#define UE_AUTORTFM_TRANSACT(...) AutoRTFM::Transact([&]() { __VA_ARGS__ })
 #else
-#define UE_AUTORTFM_OPEN(x) do { x } while (false)
-#define UE_AUTORTFM_OPENABORT(x) do { /* do nothing */ } while (false)
-#define UE_AUTORTFM_OPENCOMMIT(x) do { x } while (false)
-#define UE_AUTORTFM_TRANSACT(x) do { x } while (false)
+#define UE_AUTORTFM_OPEN(...) do { __VA_ARGS__ } while (false)
+#define UE_AUTORTFM_OPENABORT(...) do { /* do nothing */ } while (false)
+#define UE_AUTORTFM_OPENCOMMIT(...) do { __VA_ARGS__ } while (false)
+#define UE_AUTORTFM_TRANSACT(...) do { __VA_ARGS__ } while (false)
 #endif
 
 #define UE_AUTORTFM_CONCAT_IMPL(A, B) A ## B
