Added a check to make sure ELLMTag::FMalloc isn't used with the Default LLM tracker.

  - This tag was being used previously in a spot where it shouldn't.
  - Since it's easy to make this mistake, a check was added to FLLMScope to indicate improper usage.
  - Also added a comment where ELLMTag::FMalloc is defined, explaining it's special usage with ELLMAllocType::FMalloc and ELLMTag::FMallocUnused.

#tests Built many client targets and Win64 Editor.

diff --git a/Engine/Source/Runtime/Core/Private/HAL/LowLevelMemTracker.cpp b/Engine/Source/Runtime/Core/Private/HAL/LowLevelMemTracker.cpp
--- a/Engine/Source/Runtime/Core/Private/HAL/LowLevelMemTracker.cpp
+++ b/Engine/Source/Runtime/Core/Private/HAL/LowLevelMemTracker.cpp
@@ -2461,6 +2461,9 @@ const UE::LLMPrivate::FTagData* FLowLevelMemTracker::FindTagData(FName TagName,
 void FLLMScope::Init(ELLMTag TagEnum, bool bInIsStatTag, ELLMTagSet InTagSet, ELLMTracker InTracker, bool bOverride)
 {
 	LLMCheck(!bInIsStatTag && InTagSet == ELLMTagSet::None);
+	// ELLMTag::FMalloc is a special tag expected to only be used by the Platform Tracker (see header where defined).
+	LLMCheck((TagEnum != ELLMTag::FMalloc) || (ELLMTracker::Platform == InTracker));
+
 	FLowLevelMemTracker& LLMRef = FLowLevelMemTracker::Get();
 	// We have to check bIsDisabled again after calling Get, because the constructor is called from Get, and will set 
 	// bIsDisabled=false if the platform doesn't support it.
diff --git a/Engine/Source/Runtime/Core/Public/HAL/LowLevelMemTracker.h b/Engine/Source/Runtime/Core/Public/HAL/LowLevelMemTracker.h
--- a/Engine/Source/Runtime/Core/Public/HAL/LowLevelMemTracker.h
+++ b/Engine/Source/Runtime/Core/Public/HAL/LowLevelMemTracker.h
@@ -161,6 +161,7 @@ ENUM_CLASS_FLAGS(ESizeParams);
 	macro(PlatformUntracked,					"Untracked",					GET_STATFNAME(STAT_PlatformUntrackedLLM),					NAME_None,										-1)\
 	macro(PlatformOverhead,						"LLMOverhead",					GET_STATFNAME(STAT_PlatformOverheadLLM),					NAME_None,										-1)\
 	macro(PlatformOSAvailable,					"OSAvailable",					GET_STATFNAME(STAT_PlatformOSAvailableLLM),					NAME_None,										-1)\
+	/*FMalloc is a special tag that is reserved for the Platform Tracker only. It's used with ELLMAllocType::FMalloc to calculate ELLMTag::FMallocUnused. */								   \
 	macro(FMalloc,								"FMalloc",						GET_STATFNAME(STAT_FMallocLLM),								NAME_None,										-1)\
 	macro(FMallocUnused,						"FMallocUnused",				GET_STATFNAME(STAT_FMallocUnusedLLM),						GET_STATFNAME(STAT_EngineSummaryLLM),			-1)\
 	macro(ThreadStack,							"ThreadStack",					GET_STATFNAME(STAT_ThreadStackLLM),							GET_STATFNAME(STAT_EngineSummaryLLM),			-1)\
