Fixed FLowLevelMemTracker::InitialiseProgramSize not setting the size.

  - Code needed to set the ProgramSize if it wasn't yet set.
  - Looks like a code cleanup flipped the check by mistake.

#tests Built many client targets and Win64 Editor.

diff --git a/Engine/Source/Runtime/Core/Private/HAL/LowLevelMemTracker.cpp b/Engine/Source/Runtime/Core/Private/HAL/LowLevelMemTracker.cpp
--- a/Engine/Source/Runtime/Core/Private/HAL/LowLevelMemTracker.cpp
+++ b/Engine/Source/Runtime/Core/Private/HAL/LowLevelMemTracker.cpp
@@ -1164,7 +1164,7 @@ void FLowLevelMemTracker::PublishDataPerFrame(const TCHAR* LogName)
 
 void FLowLevelMemTracker::InitialiseProgramSize()
 {
-	if (ProgramSize)
+	if (!ProgramSize)
 	{
 		FPlatformMemoryStats Stats = FPlatformMemory::GetStats();
 		ProgramSize = Stats.TotalPhysical - Stats.AvailablePhysical;
