Expose PGO collection in PlatformMisc to allow start/stop it.

diff --git a/Engine/Source/Runtime/Core/Private/GenericPlatform/GenericPlatformMisc.cpp b/Engine/Source/Runtime/Core/Private/GenericPlatform/GenericPlatformMisc.cpp
--- a/Engine/Source/Runtime/Core/Private/GenericPlatform/GenericPlatformMisc.cpp
+++ b/Engine/Source/Runtime/Core/Private/GenericPlatform/GenericPlatformMisc.cpp
@@ -2057,7 +2057,20 @@ int32 FGenericPlatformMisc::GetPakchunkIndexFromPakFile(const FString& InFilenam
 
 bool FGenericPlatformMisc::IsPGOEnabled()
 {
-	return PLATFORM_COMPILER_OPTIMIZATION_PG;
+	return PLATFORM_COMPILER_OPTIMIZATION_PG != 0;
+}
+
+bool FGenericPlatformMisc::IsPGICapableBinary()
+{
+	return PLATFORM_COMPILER_OPTIMIZATION_PG_PROFILING != 0;
+}
+
+
+bool FGenericPlatformMisc::IsPGIActive()
+{
+	// by default, assume it enabled from the start in PGI binaries (as is the usual behavior). If a platform provides a way
+	// to enable/disable PG data collection runtime, it can override this.
+	return FPlatformMisc::IsPGICapableBinary();
 }
 
 int FGenericPlatformMisc::GetMobilePropagateAlphaSetting()
diff --git a/Engine/Source/Runtime/Core/Private/Misc/App.cpp b/Engine/Source/Runtime/Core/Private/Misc/App.cpp
--- a/Engine/Source/Runtime/Core/Private/Misc/App.cpp
+++ b/Engine/Source/Runtime/Core/Private/Misc/App.cpp
@@ -428,7 +428,7 @@ void FApp::PrintStartupLogMessages()
 	UE_LOG(LogInit, Log, TEXT("This binary is optimized with LTO: %s, PGO: %s, instrumented for PGO data collection: %s"),
 		PLATFORM_COMPILER_OPTIMIZATION_LTCG ? TEXT("yes") : TEXT("no"),
 		FPlatformMisc::IsPGOEnabled() ? TEXT("yes") : TEXT("no"),
-		PLATFORM_COMPILER_OPTIMIZATION_PG_PROFILING ? TEXT("yes") : TEXT("no")
+		FPlatformMisc::IsPGICapableBinary() ? TEXT("yes") : TEXT("no")
 	);
 
 	FDevVersionRegistration::DumpVersionsToLog();
diff --git a/Engine/Source/Runtime/Core/Private/Unix/UnixPlatformMisc.cpp b/Engine/Source/Runtime/Core/Private/Unix/UnixPlatformMisc.cpp
--- a/Engine/Source/Runtime/Core/Private/Unix/UnixPlatformMisc.cpp
+++ b/Engine/Source/Runtime/Core/Private/Unix/UnixPlatformMisc.cpp
@@ -2017,3 +2017,61 @@ bool FUnixPlatformMisc::SetupSyscallFilters()
 
 	return true;
 }
+
+#if ENABLE_PGO_PROFILE
+// presence of this symbol prevents automatic PGI initialization
+int CORE_API __llvm_profile_runtime = 0;
+
+namespace UnixPlatformMisc
+{
+	bool GPGICollectionUnderway = false;
+}
+
+extern "C"
+{
+	void __llvm_profile_initialize_file(void);
+	int __llvm_profile_write_file(void);
+	void __llvm_profile_reset_counters(void);
+};
+
+bool FUnixPlatformMisc::StartNewPGODataCollection(const FString& AbsoluteFileName)
+{ 
+	const TCHAR* ProfFileEnvVar = TEXT("LLVM_PROFILE_FILE");
+	if (FPlatformMisc::IsPGIActive())
+	{
+		UE_LOG(LogHAL, Error, TEXT("Profiling data collection is already under way! (file being written is '%s'"), *FPlatformMisc::GetEnvironmentVariable(ProfFileEnvVar));
+		return false;
+	}
+
+	UnixPlatformMisc::GPGICollectionUnderway = true;
+	FPlatformMisc::SetEnvironmentVar(ProfFileEnvVar, *AbsoluteFileName);
+	UE_LOG(LogHAL, Log, TEXT("Starting PGI data collection to file '%s'"), *AbsoluteFileName);
+	__llvm_profile_reset_counters();
+	__llvm_profile_initialize_file();
+	return true;
+};
+
+bool FUnixPlatformMisc::IsPGIActive()
+{
+	return UnixPlatformMisc::GPGICollectionUnderway;
+}
+
+
+bool FUnixPlatformMisc::StopPGOCollectionAndCloseFile()
+{
+	if (!FPlatformMisc::IsPGIActive())
+	{
+		UE_LOG(LogHAL, Warning, TEXT("Cannot stop PGO data collection, it was not started."));
+		return false;
+	}
+
+	UE_LOG(LogHAL, Log, TEXT("Stopping PGO data collection."));
+
+	if (__llvm_profile_write_file() != 0)
+	{
+		UE_LOG(LogHAL, Error, TEXT("Error writing out PGO file."));
+	}
+	UnixPlatformMisc::GPGICollectionUnderway = false;
+	return true;
+}
+#endif
diff --git a/Engine/Source/Runtime/Core/Public/GenericPlatform/GenericPlatformMisc.h b/Engine/Source/Runtime/Core/Public/GenericPlatform/GenericPlatformMisc.h
--- a/Engine/Source/Runtime/Core/Public/GenericPlatform/GenericPlatformMisc.h
+++ b/Engine/Source/Runtime/Core/Public/GenericPlatform/GenericPlatformMisc.h
@@ -1877,6 +1877,35 @@ public:
 	 */
 	static CORE_API bool IsPGOEnabled();
 
+	/**
+	 * Returns true if this binary has been instrumented for collection of data needed for profile-guided optimization (PGO)
+	 * This value is not expected to change during the run-time
+	 */
+	static CORE_API bool IsPGICapableBinary();
+
+	/**
+	 * On platforms that support run-time PGI on/off, starts a new collection with a given file name.
+	 *
+	 * @param AbsoluteFileName File to put the collection into
+	 * @return true if everything was successful
+	 */
+	static CORE_API bool StartNewPGOCollection(const FString& AbsoluteFileName) { return false; }
+
+	/**
+	 * Returns true if PGI (profile-guided instrumentation, needed for PGO) is currently ongoing.
+	 * Note that on some platforms can start and stop profile data collection.
+	 */
+	static CORE_API bool IsPGIActive();
+
+	/**
+	 * On platforms that support run-time PGI on/off, stops the current collection and writes out the file (or completes its write).
+	 * The filename must have been given to the StartNewPGICollection() at the start of collection (it is needed in advance since the
+	 * collection process may stream to it).
+	 *
+	 * @return true if everything was successful
+	 */
+	static CORE_API bool StopPGOCollectionAndCloseFile() { return false; }
+
 	static CORE_API void ShowConsoleWindow();
 
 	/*
diff --git a/Engine/Source/Runtime/Core/Public/Unix/UnixPlatformMisc.h b/Engine/Source/Runtime/Core/Public/Unix/UnixPlatformMisc.h
--- a/Engine/Source/Runtime/Core/Public/Unix/UnixPlatformMisc.h
+++ b/Engine/Source/Runtime/Core/Public/Unix/UnixPlatformMisc.h
@@ -111,6 +111,12 @@ struct FUnixPlatformMisc : public FGenericPlatformMisc
 	 */
 	static CORE_API bool HasBeenStartedRemotely();
 
+#if ENABLE_PGO_PROFILE
+	static bool StartNewPGOCollection(const FString& AbsoluteFileName);
+	static bool IsPGIActive();
+	static bool StopPGOCollectionAndCloseFile();
+#endif
+
 	/**
 	 * Determines if return code has been overriden and returns it.
 	 *
