Rework command line parsing for memory trace

On some platforms and configurations all command line arguments are sent as one string rather than one string per argument. We need to parse each argument in the vector for sub arguments to be sure to catch any memory trace channel.

diff --git a/Engine/Source/Runtime/Core/Private/ProfilingDebugging/MemoryTrace.cpp b/Engine/Source/Runtime/Core/Private/ProfilingDebugging/MemoryTrace.cpp
--- a/Engine/Source/Runtime/Core/Private/ProfilingDebugging/MemoryTrace.cpp
+++ b/Engine/Source/Runtime/Core/Private/ProfilingDebugging/MemoryTrace.cpp
@@ -285,24 +285,39 @@ static bool MemoryTrace_ShouldEnable(int32 ArgC, const ArgCharType* const* ArgV)
 	{
 		const ArgCharType* Arg = ArgV[ArgIndex];
 		TStringView<ArgCharType> ArgView(Arg);
-		if (ArgView.SubStr(0, 7).Equals("-trace=", ESearchCase::IgnoreCase))
+
+		// In certain scenarios (for example when launching a packaged title on some platforms)
+		// the multiple arguments are squashed into one single argument. Due to this we also need
+		// to parse each argument as if it contains multiple sub arguments.
+		int32 ArgStart;
+		while (ArgView.FindChar('-', ArgStart))
 		{
-			const ArgCharType* Start = Arg + 7;
-			const ArgCharType* End = Arg + TCString<ArgCharType>::Strlen(Arg);
+			ArgView.RightChopInline(ArgStart);
+
+			int32 ArgEnd = INDEX_NONE;
+			ArgView.RightChop(1).FindChar('-', ArgEnd);
+
+			TStringView<ArgCharType> SubArg = ArgView.SubStr(0, ArgEnd == INDEX_NONE ? ArgView.Len() : ArgEnd);
 
-			for (const ArgCharType* c = Start; c < End + 1; ++c)
+			if (SubArg.SubStr(0, 7).Equals("-trace=", ESearchCase::IgnoreCase))
 			{
-				if (*c == '\0' || *c == ',')
+				TStringView<ArgCharType> ChannelsOrPresets = SubArg.SubStr(7, SubArg.Len() - 7);
+
+				int32 CommaPos = -1;
+				while (!ChannelsOrPresets.IsEmpty())
 				{
-					TStringView<ArgCharType> View(Start, uint32(c - Start));
-					if (View.Equals("memalloc", ESearchCase::IgnoreCase) || View.Equals("memory", ESearchCase::IgnoreCase))
+					ChannelsOrPresets.FindChar(',', CommaPos);
+					TStringView<ArgCharType> Channel = ChannelsOrPresets.SubStr(0, CommaPos == INDEX_NONE ? ChannelsOrPresets.Len() : CommaPos);
+					if (Channel.Equals("memalloc", ESearchCase::IgnoreCase) || Channel.Equals("memory", ESearchCase::IgnoreCase))
 					{
 						return true;
 					}
-
-					Start = c + 1;
+					ChannelsOrPresets.RightChopInline(Channel.Len() + 1);
 				}
 			}
+
+			//Pass the '-' and find the next
+			ArgView.RightChopInline(ArgEnd);
 		}
 	}
 
