Fixed stack offset in AddErrorIfFalse to correctly report file & line that the error was emitted on

diff --git a/Engine/Source/Runtime/Core/Private/Misc/AutomationTest.cpp b/Engine/Source/Runtime/Core/Private/Misc/AutomationTest.cpp
--- a/Engine/Source/Runtime/Core/Private/Misc/AutomationTest.cpp
+++ b/Engine/Source/Runtime/Core/Private/Misc/AutomationTest.cpp
@@ -1207,7 +1207,7 @@ bool FAutomationTestBase::AddErrorIfFalse(bool bCondition, const FString& InErro
 {
 	if (!bCondition)
 	{
-		AddError(InError, StackOffset);
+		AddError(InError, StackOffset + 1);
 	}
 	return bCondition;
 }
