Fix a typo in a function name.

- Header is already fixed.

diff --git a/Engine/Source/Runtime/Core/Private/Unix/UnixPlatformMisc.cpp b/Engine/Source/Runtime/Core/Private/Unix/UnixPlatformMisc.cpp
--- a/Engine/Source/Runtime/Core/Private/Unix/UnixPlatformMisc.cpp
+++ b/Engine/Source/Runtime/Core/Private/Unix/UnixPlatformMisc.cpp
@@ -2034,7 +2034,7 @@ extern "C"
 	void __llvm_profile_reset_counters(void);
 };
 
-bool FUnixPlatformMisc::StartNewPGODataCollection(const FString& AbsoluteFileName)
+bool FUnixPlatformMisc::StartNewPGOCollection(const FString& AbsoluteFileName)
 { 
 	const TCHAR* ProfFileEnvVar = TEXT("LLVM_PROFILE_FILE");
 	if (FPlatformMisc::IsPGIActive())
