FWindowsPlatformStackWalk::ThreadStackWalkAndDump: sync with other threads before suspending them as they can be in the process of suspending us

diff --git a/Engine/Source/Runtime/Core/Private/Windows/WindowsPlatformStackWalk.cpp b/Engine/Source/Runtime/Core/Private/Windows/WindowsPlatformStackWalk.cpp
--- a/Engine/Source/Runtime/Core/Private/Windows/WindowsPlatformStackWalk.cpp
+++ b/Engine/Source/Runtime/Core/Private/Windows/WindowsPlatformStackWalk.cpp
@@ -13,9 +13,11 @@
 #include "Misc/Paths.h"
 #include "Misc/CommandLine.h"
 #include "HAL/PlatformProcess.h"
+#include "HAL/CriticalSection.h"
 #include "CoreGlobals.h"
 #include "Misc/ConfigCacheIni.h"
 #include "Misc/OutputDeviceRedirector.h"
+#include "Misc/ScopeLock.h"
 
 #include "Windows/WindowsHWrapper.h"
 #include "Windows/AllowWindowsPlatformTypes.h"
@@ -173,6 +175,10 @@ void FWindowsPlatformStackWalk::ThreadStackWalkAndDump(ANSICHAR* HumanReadableSt
 	HANDLE ThreadHandle = OpenThread(THREAD_GET_CONTEXT | THREAD_SET_CONTEXT | THREAD_TERMINATE | THREAD_SUSPEND_RESUME, false, ThreadId);
 	if (ThreadHandle)
 	{
+		// sync with other threads that may try to suspend us as we try to suspend them
+		static FCriticalSection Mutex;
+		FScopeLock Lock(&Mutex);
+
 		// Suspend the thread before grabbing its context (possible fix for incomplete callstacks)
 		SuspendThread(ThreadHandle);
 		// Give task scheduler some time to actually suspend the thread
