FMRSWRecursiveAccessDetector: do not try to suspend the current thread

diff --git a/Engine/Source/Runtime/Core/Public/Misc/MTAccessDetector.h b/Engine/Source/Runtime/Core/Public/Misc/MTAccessDetector.h
--- a/Engine/Source/Runtime/Core/Public/Misc/MTAccessDetector.h
+++ b/Engine/Source/Runtime/Core/Public/Misc/MTAccessDetector.h
@@ -375,6 +375,12 @@ private:
 			return TEXT("[none]");
 		}
 
+		// to get the callstack of another thread we suspend it first. it's a bad idea to suspend the current thread
+		if (FPlatformTLS::GetCurrentThreadId() == ThreadId)
+		{
+			return GetCurrentThreadCallstack();
+		}
+
 		const SIZE_T StackTraceSize = 65536;
 		ANSICHAR StackTrace[StackTraceSize] = { 0 };
 		FPlatformStackWalk::ThreadStackWalkAndDump(StackTrace, StackTraceSize, 0, ThreadId);
