a workaround to avoid "deadlock by busy-waiting" in ShaderCodeArchive: the existing IAsyncReadRequest::WaitCompletion() didn't guarantee completion, only that the request is done and the result is ready. It was used in IMemoryReadStream::EnsureReadNonBlocking() to ensure the following read will be non-blocking, and this can fail and actually block on reading -> busy-waiting -> deadlock. The new method EnsureCompletion() was added instead of "fixing" existing WaitCompletion() because other use cases don't need completion guarantee and the new method is theoretically slower

diff --git a/Engine/Source/Runtime/Core/Private/FileCache/FileCache.cpp b/Engine/Source/Runtime/Core/Private/FileCache/FileCache.cpp
--- a/Engine/Source/Runtime/Core/Private/FileCache/FileCache.cpp
+++ b/Engine/Source/Runtime/Core/Private/FileCache/FileCache.cpp
@@ -558,7 +558,7 @@ public:
 		if (Request)
 		{
 			check(!Memory);
-			Request->WaitCompletion();
+			Request->EnsureCompletion();
 		}
 	}
 
diff --git a/Engine/Source/Runtime/Core/Private/GenericPlatform/GenericPlatformFile.cpp b/Engine/Source/Runtime/Core/Private/GenericPlatform/GenericPlatformFile.cpp
--- a/Engine/Source/Runtime/Core/Private/GenericPlatform/GenericPlatformFile.cpp
+++ b/Engine/Source/Runtime/Core/Private/GenericPlatform/GenericPlatformFile.cpp
@@ -96,6 +96,17 @@ public:
 			}
 		}
 	}
+
+	virtual void EnsureCompletion() override
+	{
+		if (Task)
+		{
+			Task->EnsureCompletion();
+			delete Task;
+			Task = nullptr;
+		}
+	}
+
 	virtual void CancelImpl() override
 	{
 		if (Task)
diff --git a/Engine/Source/Runtime/Core/Public/Async/AsyncFileHandle.h b/Engine/Source/Runtime/Core/Public/Async/AsyncFileHandle.h
--- a/Engine/Source/Runtime/Core/Public/Async/AsyncFileHandle.h
+++ b/Engine/Source/Runtime/Core/Public/Async/AsyncFileHandle.h
@@ -96,6 +96,15 @@ public:
 		return PollCompletion();
 	}
 
+	/**
+	* Waits for the request to complete, with an additional guarantee that the second consecutive call won't ever block, which is not a case for
+	* WaitCompletion().
+	*/
+	virtual void EnsureCompletion()
+	{
+		unimplemented();
+	}
+
 	/** Cancel the request. This is a non-blocking async call and so does not ensure completion! **/
 	FORCEINLINE void Cancel()
 	{
