Cached the tracker type (ELLMTracker) in FLLMTracker, instead of passing it to all the member functions.

  - The FLLMTracker is initialized with the ELLMTracker, so there is no reason not to cache it.
  - This change makes the code easier to read because there isn't a redundant ELLMTracker being passed around.
  - Also it's useful when debugging the code too see which tracker you are in.

#tests Built most Win64 targets. Ran editor. Did a local ReplayRun Did a Horde Minimal Code Only Build with most targets and LLM.

diff --git a/Engine/Source/Runtime/Core/Private/HAL/LowLevelMemTracker.cpp b/Engine/Source/Runtime/Core/Private/HAL/LowLevelMemTracker.cpp
--- a/Engine/Source/Runtime/Core/Private/HAL/LowLevelMemTracker.cpp
+++ b/Engine/Source/Runtime/Core/Private/HAL/LowLevelMemTracker.cpp
@@ -418,30 +418,27 @@ public:
 	FLLMTracker(FLowLevelMemTracker& InLLM);
 	~FLLMTracker();
 
-	void Initialise(ELLMTracker Tracker, FLLMAllocator* InAllocator);
+	void Initialise(ELLMTracker InTracker, FLLMAllocator* InAllocator);
 
 	void PushTag(ELLMTag EnumTag, ELLMTagSet TagSet);
 	void PushTag(FName Tag, bool bInIsStatTag, ELLMTagSet TagSet);
 	void PushTag(const FTagData* TagData, ELLMTagSet TagSet);
 	void PopTag(ELLMTagSet TagSet);
-	void TrackAllocation(const void* Ptr, int64 Size, FName DefaultTag, ELLMTracker Tracker, ELLMAllocType AllocType,
-		bool bTrackInMemPro);
-	void TrackAllocation(const void* Ptr, int64 Size, ELLMTag DefaultTag, ELLMTracker Tracker, ELLMAllocType AllocType,
-		bool bTrackInMemPro);
-	void TrackFree(const void* Ptr, ELLMTracker Tracker, ELLMAllocType AllocType, bool bTrackInMemPro);
-	void TrackMemoryOfActiveTag(int64 Amount, FName DefaultTag, ELLMTracker Tracker, ELLMAllocType AllocType);
-	void TrackMemoryOfActiveTag(int64 Amount, ELLMTag DefaultTag, ELLMTracker Tracker, ELLMAllocType AllocType);
-	void OnAllocMoved(const void* Dest, const void* Source, ELLMTracker Tracker, ELLMAllocType AllocType);
-
-	void TrackMemory(ELLMTag EnumTag, ELLMTracker Tracker, int64 Amount, ELLMAllocType AllocType);
-	void TrackMemory(FName TagName, ELLMTagSet TagSet, ELLMTracker Tracker, int64 Amount, ELLMAllocType AllocType);
-	void TrackMemory(const FTagData* TagData, ELLMTracker Tracker, int64 Amount, ELLMAllocType AllocType);
+	void TrackAllocation(const void* Ptr, int64 Size, FName DefaultTag, ELLMAllocType AllocType, bool bTrackInMemPro);
+	void TrackAllocation(const void* Ptr, int64 Size, ELLMTag DefaultTag, ELLMAllocType AllocType, bool bTrackInMemPro);
+	void TrackFree(const void* Ptr, ELLMAllocType AllocType, bool bTrackInMemPro);
+	void TrackMemoryOfActiveTag(int64 Amount, FName DefaultTag, ELLMAllocType AllocType);
+	void TrackMemoryOfActiveTag(int64 Amount, ELLMTag DefaultTag, ELLMAllocType AllocType);
+	void OnAllocMoved(const void* Dest, const void* Source, ELLMAllocType AllocType);
+
+	void TrackMemory(ELLMTag EnumTag, int64 Amount, ELLMAllocType AllocType);
+	void TrackMemory(FName TagName, ELLMTagSet TagSet, int64 Amount, ELLMAllocType AllocType);
+	void TrackMemory(const FTagData* TagData, int64 Amount, ELLMAllocType AllocType);
 
 	// This will pause/unpause tracking, and also manually increment a given tag
-	void PauseAndTrackMemory(FName TagName, ELLMTagSet TagSet, ELLMTracker Tracker, bool bInIsStatTag, int64 Amount,
-		ELLMAllocType AllocType);
-	void PauseAndTrackMemory(ELLMTag EnumTag, ELLMTracker Tracker, int64 Amount, ELLMAllocType AllocType);
-	void PauseAndTrackMemory(const FTagData* TagData, ELLMTracker Tracker, int64 Amount, ELLMAllocType AllocType);
+	void PauseAndTrackMemory(FName TagName, ELLMTagSet TagSet, bool bInIsStatTag, int64 Amount, ELLMAllocType AllocType);
+	void PauseAndTrackMemory(ELLMTag EnumTag, int64 Amount, ELLMAllocType AllocType);
+	void PauseAndTrackMemory(const FTagData* TagData, int64 Amount, ELLMAllocType AllocType);
 	void Pause(ELLMAllocType AllocType);
 	void Unpause(ELLMAllocType AllocType);
 	bool IsPaused(ELLMAllocType AllocType);
@@ -578,13 +575,13 @@ public:
 protected:
 	FLLMThreadState* GetOrCreateState();
 	FLLMThreadState* GetState();
-	void TrackAllocation(const void* Ptr, int64 Size, const FTagData* ActiveTagData, ELLMTracker Tracker,
-		ELLMAllocType AllocType, FLLMThreadState* State, bool bTrackInMemPro);
-	void TrackMemoryOfActiveTag(int64 Amount, const FTagData* TagData, ELLMTracker Tracker,
-		ELLMAllocType AllocType, FLLMThreadState* State);
+	void TrackAllocation(const void* Ptr, int64 Size, const FTagData* ActiveTagData, ELLMAllocType AllocType, FLLMThreadState* State, bool bTrackInMemPro);
+	void TrackMemoryOfActiveTag(int64 Amount, const FTagData* TagData, ELLMAllocType AllocType, FLLMThreadState* State);
 
 	FLowLevelMemTracker& LLMRef;
 
+	ELLMTracker Tracker;
+
 	uint32 TlsSlot;
 
 	TArray<FLLMThreadState*, FDefaultLLMAllocator> ThreadStates;
@@ -1169,10 +1166,8 @@ void FLowLevelMemTracker::InitialiseProgramSize()
 		FPlatformMemoryStats Stats = FPlatformMemory::GetStats();
 		ProgramSize = Stats.TotalPhysical - Stats.AvailablePhysical;
 
-		GetTracker(ELLMTracker::Platform)->TrackMemory(ELLMTag::ProgramSizePlatform, ELLMTracker::Platform,
-			 ProgramSize, ELLMAllocType::System);
-		GetTracker(ELLMTracker::Default)->TrackMemory(ELLMTag::ProgramSize, ELLMTracker::Default,
-			 ProgramSize, ELLMAllocType::System);
+		GetTracker(ELLMTracker::Platform)->TrackMemory(ELLMTag::ProgramSizePlatform, ProgramSize, ELLMAllocType::System);
+		GetTracker(ELLMTracker::Default)->TrackMemory(ELLMTag::ProgramSize, ProgramSize, ELLMAllocType::System);
 	}
 }
 
@@ -1188,10 +1183,8 @@ void FLowLevelMemTracker::SetProgramSize(uint64 InProgramSize)
 
 	ProgramSize = static_cast<int64>(InProgramSize);
 
-	GetTracker(ELLMTracker::Platform)->TrackMemory(ELLMTag::ProgramSizePlatform, ELLMTracker::Platform,
-		ProgramSizeDiff, ELLMAllocType::System);
-	GetTracker(ELLMTracker::Default)->TrackMemory(ELLMTag::ProgramSize, ELLMTracker::Default,
-		ProgramSizeDiff, ELLMAllocType::System);
+	GetTracker(ELLMTracker::Platform)->TrackMemory(ELLMTag::ProgramSizePlatform, ProgramSizeDiff, ELLMAllocType::System);
+	GetTracker(ELLMTracker::Default)->TrackMemory(ELLMTag::ProgramSize, ProgramSizeDiff, ELLMAllocType::System);
 }
 
 void FLowLevelMemTracker::ProcessCommandLine(const TCHAR* CmdLine)
@@ -1317,8 +1310,7 @@ void FLowLevelMemTracker::OnLowLevelAlloc(ELLMTracker Tracker, const void* Ptr,
 	}
 	BootstrapInitialise();
 
-	GetTracker(Tracker)->TrackAllocation(Ptr, static_cast<int64>(Size), DefaultTag, Tracker,
-		AllocType, bTrackInMemPro);
+	GetTracker(Tracker)->TrackAllocation(Ptr, static_cast<int64>(Size), DefaultTag, AllocType, bTrackInMemPro);
 }
 
 void FLowLevelMemTracker::OnLowLevelAlloc(ELLMTracker Tracker, const void* Ptr, uint64 Size, FName DefaultTag,
@@ -1330,8 +1322,7 @@ void FLowLevelMemTracker::OnLowLevelAlloc(ELLMTracker Tracker, const void* Ptr,
 	}
 	BootstrapInitialise();
 
-	GetTracker(Tracker)->TrackAllocation(Ptr, static_cast<int64>(Size), DefaultTag, Tracker,
-		AllocType, bTrackInMemPro);
+	GetTracker(Tracker)->TrackAllocation(Ptr, static_cast<int64>(Size), DefaultTag, AllocType, bTrackInMemPro);
 }
 
 void FLowLevelMemTracker::OnLowLevelFree(ELLMTracker Tracker, const void* Ptr,
@@ -1345,7 +1336,7 @@ void FLowLevelMemTracker::OnLowLevelFree(ELLMTracker Tracker, const void* Ptr,
 
 	if (Ptr != nullptr)
 	{
-		GetTracker(Tracker)->TrackFree(Ptr, Tracker, AllocType, bTrackInMemPro);
+		GetTracker(Tracker)->TrackFree(Ptr, AllocType, bTrackInMemPro);
 	}
 }
 
@@ -1357,7 +1348,7 @@ void FLowLevelMemTracker::OnLowLevelChangeInMemoryUse(ELLMTracker Tracker, int64
 		return;
 	}
 	BootstrapInitialise();
-	GetTracker(Tracker)->TrackMemoryOfActiveTag(DeltaMemory, DefaultTag, Tracker, AllocType);
+	GetTracker(Tracker)->TrackMemoryOfActiveTag(DeltaMemory, DefaultTag, AllocType);
 }
 
 void FLowLevelMemTracker::OnLowLevelChangeInMemoryUse(ELLMTracker Tracker, int64 DeltaMemory, FName DefaultTag,
@@ -1368,7 +1359,7 @@ void FLowLevelMemTracker::OnLowLevelChangeInMemoryUse(ELLMTracker Tracker, int64
 		return;
 	}
 	BootstrapInitialise();
-	GetTracker(Tracker)->TrackMemoryOfActiveTag(DeltaMemory, DefaultTag, Tracker, AllocType);
+	GetTracker(Tracker)->TrackMemoryOfActiveTag(DeltaMemory, DefaultTag, AllocType);
 }
 
 void FLowLevelMemTracker::OnLowLevelAllocMoved(ELLMTracker Tracker, const void* Dest, const void* Source,
@@ -1381,7 +1372,7 @@ void FLowLevelMemTracker::OnLowLevelAllocMoved(ELLMTracker Tracker, const void*
 	BootstrapInitialise();
 
 	//update the allocation map
-	GetTracker(Tracker)->OnAllocMoved(Dest, Source, Tracker, AllocType);
+	GetTracker(Tracker)->OnAllocMoved(Dest, Source, AllocType);
 }
 
 UE::LLMPrivate::FLLMTracker* FLowLevelMemTracker::GetTracker(ELLMTracker Tracker)
@@ -2691,11 +2682,11 @@ void FLLMPauseScope::Init(FName TagName, ELLMTag EnumTag, bool bIsEnumTag, bool
 			{
 				if (bIsEnumTag)
 				{
-					LLMRef.GetTracker(Tracker)->PauseAndTrackMemory(EnumTag, Tracker, static_cast<int64>(Amount), InAllocType);
+					LLMRef.GetTracker(Tracker)->PauseAndTrackMemory(EnumTag, static_cast<int64>(Amount), InAllocType);
 				}
 				else
 				{
-					LLMRef.GetTracker(Tracker)->PauseAndTrackMemory(TagName, ELLMTagSet::None, Tracker, bIsStatTag, 
+					LLMRef.GetTracker(Tracker)->PauseAndTrackMemory(TagName, ELLMTagSet::None, bIsStatTag, 
 						static_cast<int64>(Amount), InAllocType);
 				}
 			}
@@ -3503,6 +3494,7 @@ bool FTagData::IsUsedAsDisplayParent() const
 
 FLLMTracker::FLLMTracker(FLowLevelMemTracker& InLLM)
 	: LLMRef(InLLM)
+	, Tracker(ELLMTracker::Max)
 	, TrackedTotal(0)
 	, OverrideUntaggedTagData(nullptr)
 	, OverrideTrackedTotalTagData(nullptr)
@@ -3524,12 +3516,14 @@ FLLMTracker::~FLLMTracker()
 }
 
 void FLLMTracker::Initialise(
-	ELLMTracker Tracker,
-	FLLMAllocator* InAllocator)
+	ELLMTracker InTracker,
+	FLLMAllocator* InAllocator) 
 {
-	CsvWriter.SetTracker(Tracker);
-	TraceWriter.SetTracker(Tracker);
-	CsvProfilerWriter.SetTracker(Tracker);
+	check(InTracker != ELLMTracker::Max);
+	Tracker = InTracker;
+	CsvWriter.SetTracker(InTracker);
+	TraceWriter.SetTracker(InTracker);
+	CsvProfilerWriter.SetTracker(InTracker);
 
 	AllocationMap.SetAllocator(InAllocator);
 }
@@ -3592,8 +3586,7 @@ void FLLMTracker::PopTag(ELLMTagSet TagSet)
 	State->PopTag(TagSet);
 }
 
-void FLLMTracker::TrackAllocation(const void* Ptr, int64 Size, ELLMTag DefaultTag, ELLMTracker Tracker,
-	ELLMAllocType AllocType, bool bTrackInMemPro)
+void FLLMTracker::TrackAllocation(const void* Ptr, int64 Size, ELLMTag DefaultTag, ELLMAllocType AllocType, bool bTrackInMemPro)
 {
 	FLLMThreadState* State = GetOrCreateState();
 	const FTagData* TagData = State->GetTopTag(ELLMTagSet::None);
@@ -3601,11 +3594,10 @@ void FLLMTracker::TrackAllocation(const void* Ptr, int64 Size, ELLMTag DefaultTa
 	{
 		TagData = LLMRef.FindOrAddTagData(DefaultTag);
 	}
-	TrackAllocation(Ptr, Size, TagData, Tracker, AllocType, State, bTrackInMemPro);
+	TrackAllocation(Ptr, Size, TagData, AllocType, State, bTrackInMemPro);
 }
 
-void FLLMTracker::TrackAllocation(const void* Ptr, int64 Size, FName DefaultTag, ELLMTracker Tracker,
-	ELLMAllocType AllocType, bool bTrackInMemPro)
+void FLLMTracker::TrackAllocation(const void* Ptr, int64 Size, FName DefaultTag, ELLMAllocType AllocType, bool bTrackInMemPro)
 {
 	FLLMThreadState* State = GetOrCreateState();
 	const FTagData* TagData = State->GetTopTag(ELLMTagSet::None);
@@ -3613,11 +3605,10 @@ void FLLMTracker::TrackAllocation(const void* Ptr, int64 Size, FName DefaultTag,
 	{
 		TagData = LLMRef.FindOrAddTagData(DefaultTag, ELLMTagSet::None);
 	}
-	TrackAllocation(Ptr, Size, TagData, Tracker, AllocType, State, bTrackInMemPro);
+	TrackAllocation(Ptr, Size, TagData, AllocType, State, bTrackInMemPro);
 }
 
-void FLLMTracker::TrackAllocation(const void* Ptr, int64 Size, const FTagData* ActiveTagData, ELLMTracker Tracker,
-	ELLMAllocType AllocType, FLLMThreadState* State, bool bTrackInMemPro)
+void FLLMTracker::TrackAllocation(const void* Ptr, int64 Size, const FTagData* ActiveTagData, ELLMAllocType AllocType, FLLMThreadState* State, bool bTrackInMemPro)
 {
 	if (IsPaused(AllocType))
 	{
@@ -3672,7 +3663,7 @@ void FLLMTracker::TrackAllocation(const void* Ptr, int64 Size, const FTagData* A
 	}
 }
 
-void FLLMTracker::TrackFree(const void* Ptr, ELLMTracker Tracker, ELLMAllocType AllocType, bool bTrackInMemPro)
+void FLLMTracker::TrackFree(const void* Ptr, ELLMAllocType AllocType, bool bTrackInMemPro)
 {
 	// look up the pointer in the tracking map
 	FLLMAllocMap::Values Values;
@@ -3711,7 +3702,7 @@ void FLLMTracker::TrackFree(const void* Ptr, ELLMTracker Tracker, ELLMAllocType
 	State->TrackFree(Ptr, Size, Tracker, AllocType, TagData, AssetTagData, AssetClassTagData, bTrackInMemPro);
 }
 
-void FLLMTracker::OnAllocMoved(const void* Dest, const void* Source, ELLMTracker Tracker, ELLMAllocType AllocType)
+void FLLMTracker::OnAllocMoved(const void* Dest, const void* Source, ELLMAllocType AllocType)
 {
 	FLLMAllocMap::Values Values;
 	{
@@ -3740,7 +3731,7 @@ void FLLMTracker::OnAllocMoved(const void* Dest, const void* Source, ELLMTracker
 	State->TrackMoved(Dest, Source, Size, Tracker, TagData);
 }
 
-void FLLMTracker::TrackMemoryOfActiveTag(int64 Amount, FName DefaultTag, ELLMTracker Tracker, ELLMAllocType AllocType)
+void FLLMTracker::TrackMemoryOfActiveTag(int64 Amount, FName DefaultTag, ELLMAllocType AllocType)
 {
 	FLLMThreadState* State = GetOrCreateState();
 	const FTagData* TagData = State->GetTopTag(ELLMTagSet::None);
@@ -3748,10 +3739,10 @@ void FLLMTracker::TrackMemoryOfActiveTag(int64 Amount, FName DefaultTag, ELLMTra
 	{
 		TagData = LLMRef.FindOrAddTagData(DefaultTag, ELLMTagSet::None);
 	}
-	TrackMemoryOfActiveTag(Amount, TagData, Tracker, AllocType, State);
+	TrackMemoryOfActiveTag(Amount, TagData, AllocType, State);
 }
 
-void FLLMTracker::TrackMemoryOfActiveTag(int64 Amount, ELLMTag DefaultTag, ELLMTracker Tracker, ELLMAllocType AllocType)
+void FLLMTracker::TrackMemoryOfActiveTag(int64 Amount, ELLMTag DefaultTag, ELLMAllocType AllocType)
 {
 	FLLMThreadState* State = GetOrCreateState();
 	const FTagData* TagData = State->GetTopTag(ELLMTagSet::None);
@@ -3759,11 +3750,10 @@ void FLLMTracker::TrackMemoryOfActiveTag(int64 Amount, ELLMTag DefaultTag, ELLMT
 	{
 		TagData = LLMRef.FindOrAddTagData(DefaultTag);
 	}
-	TrackMemoryOfActiveTag(Amount, TagData, Tracker, AllocType, State);
+	TrackMemoryOfActiveTag(Amount, TagData, AllocType, State);
 }
 
-void FLLMTracker::TrackMemoryOfActiveTag(int64 Amount, const FTagData* TagData, ELLMTracker Tracker,
-	ELLMAllocType AllocType, FLLMThreadState* State)
+void FLLMTracker::TrackMemoryOfActiveTag(int64 Amount, const FTagData* TagData, ELLMAllocType AllocType, FLLMThreadState* State)
 {
 	if (IsPaused(AllocType))
 	{
@@ -3783,38 +3773,37 @@ void FLLMTracker::TrackMemoryOfActiveTag(int64 Amount, const FTagData* TagData,
 	State->TrackMemory(Amount, Tracker, AllocType, TagData);
 }
 
-void FLLMTracker::TrackMemory(ELLMTag Tag, ELLMTracker Tracker, int64 Amount, ELLMAllocType AllocType)
+void FLLMTracker::TrackMemory(ELLMTag Tag, int64 Amount, ELLMAllocType AllocType)
 {
-	TrackMemory(LLMRef.FindOrAddTagData(Tag), Tracker, Amount, AllocType);
+	TrackMemory(LLMRef.FindOrAddTagData(Tag), Amount, AllocType);
 }
 
-void FLLMTracker::TrackMemory(FName Tag, ELLMTagSet TagSet, ELLMTracker Tracker, int64 Amount, ELLMAllocType AllocType)
+void FLLMTracker::TrackMemory(FName Tag, ELLMTagSet TagSet, int64 Amount, ELLMAllocType AllocType)
 {
-	TrackMemory(LLMRef.FindOrAddTagData(Tag, TagSet, false, UE::LLMPrivate::ETagReferenceSource::FunctionAPI), Tracker, Amount, AllocType);
+	TrackMemory(LLMRef.FindOrAddTagData(Tag, TagSet, false, UE::LLMPrivate::ETagReferenceSource::FunctionAPI), Amount, AllocType);
 }
 
-void FLLMTracker::TrackMemory(const FTagData* TagData, ELLMTracker Tracker, int64 Amount, ELLMAllocType AllocType)
+void FLLMTracker::TrackMemory(const FTagData* TagData, int64 Amount, ELLMAllocType AllocType)
 {
 	FLLMThreadState* State = GetOrCreateState();
 	FPlatformAtomics::InterlockedAdd(&TrackedTotal, Amount);
 	State->TrackMemory(Amount, Tracker, AllocType, TagData);
 }
 
-void FLLMTracker::PauseAndTrackMemory(FName TagName, ELLMTagSet TagSet, ELLMTracker Tracker, bool bInIsStatTag, 
-	int64 Amount, ELLMAllocType AllocType)
+void FLLMTracker::PauseAndTrackMemory(FName TagName, ELLMTagSet TagSet, bool bInIsStatTag, int64 Amount, ELLMAllocType AllocType)
 {
 	const FTagData* TagData = LLMRef.FindOrAddTagData(TagName, TagSet, bInIsStatTag, UE::LLMPrivate::ETagReferenceSource::FunctionAPI);
-	PauseAndTrackMemory(TagData, Tracker, Amount, AllocType);
+	PauseAndTrackMemory(TagData, Amount, AllocType);
 }
 
-void FLLMTracker::PauseAndTrackMemory(ELLMTag EnumTag, ELLMTracker Tracker, int64 Amount, ELLMAllocType AllocType)
+void FLLMTracker::PauseAndTrackMemory(ELLMTag EnumTag, int64 Amount, ELLMAllocType AllocType)
 {
 	const FTagData* TagData = LLMRef.FindOrAddTagData(EnumTag);
-	PauseAndTrackMemory(TagData, Tracker, Amount, AllocType);
+	PauseAndTrackMemory(TagData, Amount, AllocType);
 }
 
 // This will pause/unpause tracking, and also manually increment a given tag.
-void FLLMTracker::PauseAndTrackMemory(const FTagData* TagData, ELLMTracker Tracker, int64 Amount, ELLMAllocType AllocType)
+void FLLMTracker::PauseAndTrackMemory(const FTagData* TagData, int64 Amount, ELLMAllocType AllocType)
 {
 	FLLMThreadState* State = GetOrCreateState();
 	FPlatformAtomics::InterlockedAdd(&TrackedTotal, Amount);
