I/O - initialize container flags to None

#rnx

diff --git a/Engine/Source/Runtime/Core/Public/IO/IoDispatcher.h b/Engine/Source/Runtime/Core/Public/IO/IoDispatcher.h
--- a/Engine/Source/Runtime/Core/Public/IO/IoDispatcher.h
+++ b/Engine/Source/Runtime/Core/Public/IO/IoDispatcher.h
@@ -517,7 +517,7 @@ struct FIoStoreWriterResult
 	uint64 ModifiedChunksCount = 0;
 	uint64 ModifiedChunksSize = 0;
 	FName CompressionMethod = NAME_None;
-	EIoContainerFlags ContainerFlags;
+	EIoContainerFlags ContainerFlags = EIoContainerFlags::None;
 };
 
 struct FIoWriteOptions
