Disabled tpause when compiling ARM64EC because it isn't defined.

diff --git a/Engine/Source/Runtime/Core/Private/Windows/WindowsPlatformMisc.cpp b/Engine/Source/Runtime/Core/Private/Windows/WindowsPlatformMisc.cpp
--- a/Engine/Source/Runtime/Core/Private/Windows/WindowsPlatformMisc.cpp
+++ b/Engine/Source/Runtime/Core/Private/Windows/WindowsPlatformMisc.cpp
@@ -2641,7 +2641,8 @@ private:
 	 */
 	static bool CheckForTimedPauseInstruction()
 	{
-#if PLATFORM_SEH_EXCEPTIONS_DISABLED
+		// _tpause isn't defined for Arm64EC
+#if PLATFORM_SEH_EXCEPTIONS_DISABLED || defined(_M_ARM64EC)
 		return false;
 #else
 		bool bSupportsTpause = false;
