Fix tracing of FNames for numbered names when outline numbers is enabled.

When UE_FNAME_OUTLINE_NUMBER is enabled the number is part of the name entry so we can format the name before tracing it.

diff --git a/Engine/Source/Runtime/Core/Private/UObject/UnrealNames.cpp b/Engine/Source/Runtime/Core/Private/UObject/UnrealNames.cpp
--- a/Engine/Source/Runtime/Core/Private/UObject/UnrealNames.cpp
+++ b/Engine/Source/Runtime/Core/Private/UObject/UnrealNames.cpp
@@ -1559,6 +1559,7 @@ public:
 
 #if UE_TRACE_ENABLED
 	UE::Trace::FEventRef32 Trace(const FNameEntryId& EntryId);
+	UE::Trace::FEventRef32 TraceDefinition(const FNameEntryId& EntryId) const;
 	void RetraceAll() const;
 #endif
 
@@ -1868,28 +1869,49 @@ UE::Trace::FEventRef32 FNamePool::Trace(const FNameEntryId& EntryId)
 	const uint32 DefinitionId = EntryId.ToUnstableInt();
 	if (!TracedNames[Handle.Block].SetBitAtomic(Handle.Offset))
 	{
-		const FNameEntry& DisplayEntry = Resolve(Handle);
+		return TraceDefinition(EntryId);
+	}
+	return UE::Trace::MakeEventRef(DefinitionId, UE_TRACE_GET_DEFINITION_TYPE_ID(Strings, FName));
+}
 
-		if (DisplayEntry.IsWide())
+UE::Trace::FEventRef32 FNamePool::TraceDefinition(const FNameEntryId& EntryId) const
+{
+	const uint32 DefinitionId = EntryId.ToUnstableInt();
+	FNameEntry* DisplayEntry = &Resolve(EntryId);
+	uint32 Number = 0;
+ 
+#if UE_FNAME_OUTLINE_NUMBER
+	if (DisplayEntry->IsNumbered())
+	{
+		Number = DisplayEntry->GetNumber();
+		DisplayEntry = DisplayEntry->IsNumbered() ? &Resolve(DisplayEntry->GetNumberedName().Id) : DisplayEntry;
+	}
+#endif
+ 
+	if (DisplayEntry->IsWide())
+	{
+		TStringBuilder<NAME_SIZE> NameString;
+		DisplayEntry->AppendNameToString(NameString);
+		if (Number)
 		{
-			TCHAR Chars[NAME_SIZE];
-			DisplayEntry.GetUnterminatedName(Chars, NAME_SIZE);
-			const uint32 Size = DisplayEntry.GetNameLength();
-			const auto Ref = UE_TRACE_LOG_DEFINITION(Strings, FName, DefinitionId, true)
-				 << FName.DisplayWide(Chars, Size);
-			return Ref;
+			NameString << TEXT("_") << NAME_INTERNAL_TO_EXTERNAL(Number);
 		}
-		else
+		const auto Ref = UE_TRACE_LOG_DEFINITION(Strings, FName, DefinitionId, true)
+			<< FName.DisplayWide(NameString.GetData(), NameString.Len());
+		return Ref;
+	}
+	else
+	{
+		TAnsiStringBuilder<NAME_SIZE> NameString;
+		DisplayEntry->AppendAnsiNameToString(NameString);
+		if (Number)
 		{
-			ANSICHAR Chars[NAME_SIZE];
-			DisplayEntry.GetAnsiName(Chars);
-			const uint32 Size = DisplayEntry.GetNameLength();
-			const auto Ref = UE_TRACE_LOG_DEFINITION(Strings, FName, DefinitionId, true)
-				 << FName.DisplayAnsi(Chars, Size);
-			return Ref;
+			NameString << "_" << NAME_INTERNAL_TO_EXTERNAL(Number);
 		}
+		const auto Ref = UE_TRACE_LOG_DEFINITION(Strings, FName, DefinitionId, true)
+			<< FName.DisplayAnsi(NameString.GetData(), NameString.Len());
+		return Ref;
 	}
-	return UE::Trace::MakeEventRef(DefinitionId, UE_TRACE_GET_DEFINITION_TYPE_ID(Strings, FName));
 }
 
 void FNamePool::RetraceAll() const
@@ -1902,25 +1924,7 @@ void FNamePool::RetraceAll() const
 		{
 			const FNameEntryHandle Handle(Block, Index);
 			const FNameEntryId EntryId = Handle;
-			const uint32 DefinitionId = EntryId.ToUnstableInt(); 
-			const FNameEntry& DisplayEntry = Resolve(Handle);
-
-			if (DisplayEntry.IsWide())
-			{
-				TCHAR Chars[NAME_SIZE];
-				DisplayEntry.GetUnterminatedName(Chars, NAME_SIZE);
-				const uint32 Size = DisplayEntry.GetNameLength();
-				UE_TRACE_LOG_DEFINITION(Strings, FName, DefinitionId, true)
-					<< FName.DisplayWide(Chars, Size);
-			}
-			else
-			{
-				ANSICHAR Chars[NAME_SIZE];
-				DisplayEntry.GetAnsiName(Chars);
-				const uint32 Size = DisplayEntry.GetNameLength();
-				UE_TRACE_LOG_DEFINITION(Strings, FName, DefinitionId, true)
-					<< FName.DisplayAnsi(Chars, Size);
-			}
+			TraceDefinition(EntryId);
 		});
 	}
 }
diff --git a/Engine/Source/Runtime/Core/Public/UObject/NameTypes.h b/Engine/Source/Runtime/Core/Public/UObject/NameTypes.h
--- a/Engine/Source/Runtime/Core/Public/UObject/NameTypes.h
+++ b/Engine/Source/Runtime/Core/Public/UObject/NameTypes.h
@@ -301,6 +301,7 @@ private:
 	friend struct FNameHelper;
 	friend class FNameEntryAllocator;
 	friend class FNamePoolShardBase;
+	friend class FNamePool;
 
 	static void Encode(ANSICHAR* Name, uint32 Len);
 	static void Encode(WIDECHAR* Name, uint32 Len);
