TBitArray:
- Find and FindAndSetFirstZeroBit updated to use FindFromImpl internally, as they perform the same search logic.
- FindAndSetLastZeroBit updated to use FindLast internally, as it performs the same search logic.
- Unit tests added for TestFindLast, TestFindAndSetFirstZeroBit and TestFindAndSetLastZeroBit

diff --git a/Engine/Source/Runtime/Core/Public/Containers/BitArray.h b/Engine/Source/Runtime/Core/Public/Containers/BitArray.h
--- a/Engine/Source/Runtime/Core/Public/Containers/BitArray.h
+++ b/Engine/Source/Runtime/Core/Public/Containers/BitArray.h
@@ -1073,36 +1073,16 @@ public:
 	}
 
 	/**
-	 * Finds the first true/false bit in the array, and returns the bit index.
-	 * If there is none, INDEX_NONE is returned.
+	 * Finds the first occurrence of the specified value (true/false) in the array, and returns the bit index.
+	 * If the specified value is not found, INDEX_NONE is returned.
+	 *
+	 * @param  bValue  The value (true/false) to search for.
+	 *
+	 * @return The index of the first occurrence of the specified value (true/false), or INDEX_NONE if not found.
 	 */
 	int32 Find(bool bValue) const
 	{
-		// Iterate over the array until we see a word with a matching bit
-		const uint32 Test = bValue ? 0u : (uint32)-1;
-
-		const uint32* RESTRICT DwordArray = GetData();
-		const int32 LocalNumBits = NumBits;
-		const int32 DwordCount = FBitSet::CalculateNumWords(LocalNumBits);
-		int32 DwordIndex = 0;
-		while (DwordIndex < DwordCount && DwordArray[DwordIndex] == Test)
-		{
-			++DwordIndex;
-		}
-
-		if (DwordIndex < DwordCount)
-		{
-			// If we're looking for a false, then we flip the bits - then we only need to find the first one bit
-			const uint32 Bits = bValue ? (DwordArray[DwordIndex]) : ~(DwordArray[DwordIndex]);
-			UE_ASSUME(Bits != 0);
-			const int32 LowestBitIndex = FMath::CountTrailingZeros(Bits) + (DwordIndex << NumBitsPerDWORDLogTwo);
-			if (LowestBitIndex < LocalNumBits)
-			{
-				return LowestBitIndex;
-			}
-		}
-
-		return INDEX_NONE;
+		return FindFromImpl(bValue, 0);
 	}
 
 private:
@@ -1163,9 +1143,13 @@ public:
 	}
 
 	/**
-	* Finds the last true/false bit in the array, and returns the bit index.
-	* If there is none, INDEX_NONE is returned.
-	*/
+	 * Finds the last occurrence of the specified value (true/false) in the array, and returns the bit index.
+	 * If the specified value is not found, INDEX_NONE is returned.
+	 *
+	 * @param  bValue  The value (true/false) to search for.
+	 *
+	 * @return The index of the last occurrence of the specified value (true/false), or INDEX_NONE if not found.
+	 */
 	int32 FindLast(bool bValue) const 
 	{
 		const int32 LocalNumBits = NumBits;
@@ -1201,87 +1185,54 @@ public:
 		return Result;
 	}
 
+	/**
+	 * Checks if the array contains the specified value (true/false).
+	 *
+	 * @param  bValue  The value (true/false) to check for.
+	 *
+	 * @return true if the array contains the specified value, false otherwise.
+	 */
 	FORCEINLINE bool Contains(bool bValue) const
 	{
 		return Find(bValue) != INDEX_NONE;
 	}
 
 	/**
-	 * Finds the first zero bit in the array, sets it to true, and returns the bit index.
-	 * If there is none, INDEX_NONE is returned.
+	 * Finds the first occurrence of a zero bit in the array and sets it to one, returning the bit index.
+	 * If there are no zero bits in the array, INDEX_NONE is returned.
+	 *
+	 * @param  StartIndex  The index to start the search from. Defaults to 0.
+	 *
+	 * @return The index of the first occurrence of a zero bit that was successfully set to one, or INDEX_NONE if not found.
 	 */
-	int32 FindAndSetFirstZeroBit(int32 ConservativeStartIndex = 0)
+	int32 FindAndSetFirstZeroBit(int32 StartIndex = 0)
 	{
-		// Iterate over the array until we see a word with a zero bit.
-		uint32* RESTRICT DwordArray = GetData();
-		const int32 LocalNumBits = NumBits;
-		const int32 DwordCount = FBitSet::CalculateNumWords(LocalNumBits);
-		int32 DwordIndex = FMath::DivideAndRoundDown(ConservativeStartIndex, NumBitsPerDWORD);
-		while (DwordIndex < DwordCount && DwordArray[DwordIndex] == (uint32)-1)
+		const int32 FirstZeroBitIndex = FindFromImpl(false, StartIndex);
+		if (FirstZeroBitIndex != INDEX_NONE)
 		{
-			++DwordIndex;
+			(*this)[FirstZeroBitIndex] = true;
+			CheckInvariants();
 		}
-
-		if (DwordIndex < DwordCount)
-		{
-			// Flip the bits, then we only need to find the first one bit -- easy.
-			const uint32 Bits = ~(DwordArray[DwordIndex]);
-			UE_ASSUME(Bits != 0);
-			const uint32 LowestBit = (Bits) & (-(int32)Bits);
-			const int32 LowestBitIndex = FMath::CountTrailingZeros(Bits) + (DwordIndex << NumBitsPerDWORDLogTwo);
-			if (LowestBitIndex < LocalNumBits)
-			{
-				DwordArray[DwordIndex] |= LowestBit;
-				CheckInvariants();
-				return LowestBitIndex;
-			}
-		}
-
-		return INDEX_NONE;
+		return FirstZeroBitIndex;
 	}
 
 	/**
-	 * Finds the last zero bit in the array, sets it to true, and returns the bit index.
-	 * If there is none, INDEX_NONE is returned.
+	 * Finds the last occurrence of a zero bit in the array and sets it to one, returning the bit index.
+	 * If there are no zero bits in the array, INDEX_NONE is returned.
+	 *
+	 * @return The index of the last occurrence of a zero bit that was successfully set to one, or INDEX_NONE if not found.
 	 */
 	int32 FindAndSetLastZeroBit()
 	{
-		const int32 LocalNumBits = NumBits;
-
-		// Get the correct mask for the last word
-		uint32 Mask = GetLastWordMask();
-
-		// Iterate over the array until we see a word with a zero bit.
-		uint32 DwordIndex = FBitSet::CalculateNumWords(LocalNumBits);
-		uint32* RESTRICT DwordArray = GetData();
-		for (;;)
+		const int32 LastZeroBitIndex = FindLast(false);
+		if (LastZeroBitIndex != INDEX_NONE)
 		{
-			if (DwordIndex == 0)
-			{
-				return INDEX_NONE;
-			}
-			--DwordIndex;
-			if ((DwordArray[DwordIndex] & Mask) != Mask)
-			{
-				break;
-			}
-			Mask = ~0u;
+			(*this)[LastZeroBitIndex] = true;
+			CheckInvariants();
 		}
-
-		// Flip the bits, then we only need to find the first one bit -- easy.
-		const uint32 Bits = ~DwordArray[DwordIndex] & Mask;
-		UE_ASSUME(Bits != 0);
-
-		uint32 BitIndex = (NumBitsPerDWORD - 1) - FMath::CountLeadingZeros(Bits);
-		DwordArray[DwordIndex] |= 1u << BitIndex;
-
-		CheckInvariants();
-
-		int32 Result = BitIndex + (DwordIndex << NumBitsPerDWORDLogTwo);
-		return Result;
+		return LastZeroBitIndex;
 	}
 
-
 	/**
 	 * Return the bitwise AND of two bit arrays. The resulting bit array will be sized according to InFlags.
 	 */
diff --git a/Engine/Source/Runtime/Core/Tests/Containers/BitArrayTests.cpp b/Engine/Source/Runtime/Core/Tests/Containers/BitArrayTests.cpp
--- a/Engine/Source/Runtime/Core/Tests/Containers/BitArrayTests.cpp
+++ b/Engine/Source/Runtime/Core/Tests/Containers/BitArrayTests.cpp
@@ -2030,15 +2030,15 @@ TEST_CASE_NAMED(FBitArrayFindTest, "System::Core::Containers::BitArray::Find", "
 	}
 
 	{
-		TBitArray<> ArrTrue2(true, NumBitsPerDWORD);
-		CHECK_MESSAGE(TEXT("Find true in an array containing 32 true"), ArrTrue2.Find(true) == 0);
-		CHECK_MESSAGE(TEXT("Find false in an array containing 32 true"), ArrTrue2.Find(false) == INDEX_NONE);
+		TBitArray<> ArrTrue(true, NumBitsPerDWORD);
+		CHECK_MESSAGE(TEXT("Find true in an array containing 32 true"), ArrTrue.Find(true) == 0);
+		CHECK_MESSAGE(TEXT("Find false in an array containing 32 true"), ArrTrue.Find(false) == INDEX_NONE);
 	}
 
 	{
-		TBitArray<> ArrFalse2(false, NumBitsPerDWORD);
-		CHECK_MESSAGE(TEXT("Find false in an array containing 32 false"), ArrFalse2.Find(false) == 0);
-		CHECK_MESSAGE(TEXT("Find true in an array containing 32 false"), ArrFalse2.Find(true) == INDEX_NONE);
+		TBitArray<> ArrFalse(false, NumBitsPerDWORD);
+		CHECK_MESSAGE(TEXT("Find false in an array containing 32 false"), ArrFalse.Find(false) == 0);
+		CHECK_MESSAGE(TEXT("Find true in an array containing 32 false"), ArrFalse.Find(true) == INDEX_NONE);
 	}
 
 	{
@@ -2057,23 +2057,23 @@ TEST_CASE_NAMED(FBitArrayFindTest, "System::Core::Containers::BitArray::Find", "
 TEST_CASE_NAMED(FBitArrayFindFromTest, "System::Core::Containers::BitArray::FindFrom", "[ApplicationContextMask][SmokeFilter]")
 {
 	{
-		TBitArray<> ArrTrue2(true, NumBitsPerDWORD);
-		CHECK_MESSAGE(TEXT("Find true in an array full of true, starting at index 0"), ArrTrue2.FindFrom(true, 0) == 0);
-		CHECK_MESSAGE(TEXT("Find true in an array full of true, starting at index 1"), ArrTrue2.FindFrom(true, 1) == 1);
-		CHECK_MESSAGE(TEXT("Find true in an array full of true, starting at index 31"), ArrTrue2.FindFrom(true, NumBitsPerDWORD - 1) == NumBitsPerDWORD - 1);
-		CHECK_MESSAGE(TEXT("Find true in an array full of true, starting at index 32"), ArrTrue2.FindFrom(true, NumBitsPerDWORD) == INDEX_NONE);
-		CHECK_MESSAGE(TEXT("Find false in an array full of true, starting at index 1"), ArrTrue2.FindFrom(false, 1) == INDEX_NONE);
-		CHECK_MESSAGE(TEXT("Find false in an array full of true, starting at index 32"), ArrTrue2.FindFrom(false, NumBitsPerDWORD) == INDEX_NONE);
+		TBitArray<> ArrTrue(true, NumBitsPerDWORD);
+		CHECK_MESSAGE(TEXT("Find true in an array full of true, starting at index 0"), ArrTrue.FindFrom(true, 0) == 0);
+		CHECK_MESSAGE(TEXT("Find true in an array full of true, starting at index 1"), ArrTrue.FindFrom(true, 1) == 1);
+		CHECK_MESSAGE(TEXT("Find true in an array full of true, starting at index 31"), ArrTrue.FindFrom(true, NumBitsPerDWORD - 1) == NumBitsPerDWORD - 1);
+		CHECK_MESSAGE(TEXT("Find true in an array full of true, starting at index 32"), ArrTrue.FindFrom(true, NumBitsPerDWORD) == INDEX_NONE);
+		CHECK_MESSAGE(TEXT("Find false in an array full of true, starting at index 1"), ArrTrue.FindFrom(false, 1) == INDEX_NONE);
+		CHECK_MESSAGE(TEXT("Find false in an array full of true, starting at index 32"), ArrTrue.FindFrom(false, NumBitsPerDWORD) == INDEX_NONE);
 	}
 
 	{
-		TBitArray<> ArrFalse2(false, NumBitsPerDWORD);
-		CHECK_MESSAGE(TEXT("Find false in an array full of false, starting at index 0"), ArrFalse2.FindFrom(false, 0) == 0);
-		CHECK_MESSAGE(TEXT("Find false in an array full of false, starting at index 1"), ArrFalse2.FindFrom(false, 1) == 1);
-		CHECK_MESSAGE(TEXT("Find false in an array full of false, starting at index 31"), ArrFalse2.FindFrom(false, NumBitsPerDWORD - 1) == NumBitsPerDWORD - 1);
-		CHECK_MESSAGE(TEXT("Find false in an array full of false, starting at index 32"), ArrFalse2.FindFrom(false, NumBitsPerDWORD) == INDEX_NONE);
-		CHECK_MESSAGE(TEXT("Find true in an array full of false, starting at index 1"), ArrFalse2.FindFrom(true, 1) == INDEX_NONE);
-		CHECK_MESSAGE(TEXT("Find false in an array full of false, starting at index 32"), ArrFalse2.FindFrom(true, NumBitsPerDWORD) == INDEX_NONE);
+		TBitArray<> ArrFalse(false, NumBitsPerDWORD);
+		CHECK_MESSAGE(TEXT("Find false in an array full of false, starting at index 0"), ArrFalse.FindFrom(false, 0) == 0);
+		CHECK_MESSAGE(TEXT("Find false in an array full of false, starting at index 1"), ArrFalse.FindFrom(false, 1) == 1);
+		CHECK_MESSAGE(TEXT("Find false in an array full of false, starting at index 31"), ArrFalse.FindFrom(false, NumBitsPerDWORD - 1) == NumBitsPerDWORD - 1);
+		CHECK_MESSAGE(TEXT("Find false in an array full of false, starting at index 32"), ArrFalse.FindFrom(false, NumBitsPerDWORD) == INDEX_NONE);
+		CHECK_MESSAGE(TEXT("Find true in an array full of false, starting at index 1"), ArrFalse.FindFrom(true, 1) == INDEX_NONE);
+		CHECK_MESSAGE(TEXT("Find false in an array full of false, starting at index 32"), ArrFalse.FindFrom(true, NumBitsPerDWORD) == INDEX_NONE);
 	}
 
 	{
@@ -2091,6 +2091,83 @@ TEST_CASE_NAMED(FBitArrayFindFromTest, "System::Core::Containers::BitArray::Find
 	}
 }
 
+TEST_CASE_NAMED(FBitArrayTestFindLastTest, "System::Core::Containers::BitArray::FindLast", "[ApplicationContextMask][SmokeFilter]")
+{
+	{
+		TBitArray<> ArrEmpty;
+		CHECK_MESSAGE(TEXT("Find last true in an empty array"), ArrEmpty.FindLast(true) == INDEX_NONE);
+		CHECK_MESSAGE(TEXT("Find last false in an empty array"), ArrEmpty.FindLast(false) == INDEX_NONE);
+	}
+
+	{
+		TBitArray<> ArrTrue(true, NumBitsPerDWORD);
+		CHECK_MESSAGE(TEXT("Find last true in an array containing 32 true"), ArrTrue.FindLast(true) == NumBitsPerDWORD - 1);
+		CHECK_MESSAGE(TEXT("Find last false in an array containing 32 true"), ArrTrue.FindLast(false) == INDEX_NONE);
+	}
+
+	{
+		TBitArray<> ArrFalse(false, NumBitsPerDWORD);
+		CHECK_MESSAGE(TEXT("Find last false in an array containing 32 false"), ArrFalse.FindLast(false) == NumBitsPerDWORD - 1);
+		CHECK_MESSAGE(TEXT("Find last true in an array containing 32 false"), ArrFalse.FindLast(true) == INDEX_NONE);
+	}
+
+	{
+		TBitArray<> ArrFalseThenTrue(false, NumBitsPerDWORD);
+		ArrFalseThenTrue.Add(true);
+		CHECK_MESSAGE(TEXT("Find last true in an array containing multiple false and then a single true"), ArrFalseThenTrue.FindLast(true) == NumBitsPerDWORD);
+		CHECK_MESSAGE(TEXT("Find last false in an array containing multiple false and then a single true"), ArrFalseThenTrue.FindLast(false) == NumBitsPerDWORD - 1);
+	}
+
+	{
+		TBitArray<> ArrTrueThenFalse(true, NumBitsPerDWORD);
+		ArrTrueThenFalse.Add(false);
+		CHECK_MESSAGE(TEXT("Find last false in an array containing multiple true and then a single false"), ArrTrueThenFalse.FindLast(false) == NumBitsPerDWORD);
+		CHECK_MESSAGE(TEXT("Find last true in an array containing multiple true and then a single false"), ArrTrueThenFalse.FindLast(true) == NumBitsPerDWORD - 1);
+	}
+}
+
+TEST_CASE_NAMED(FBitArrayTestFindAndSetFirstZeroBitTest, "System::Core::Containers::BitArray::FindAndSetFirstZeroBit", "[ApplicationContextMask][SmokeFilter]")
+{
+	{
+		TBitArray<> ArrEmpty;
+		CHECK_MESSAGE(TEXT("Find and set first zero bit in an empty array"), ArrEmpty.FindAndSetFirstZeroBit() == INDEX_NONE);
+		CHECK_MESSAGE(TEXT("Array should remain empty and have no set bits"), ArrEmpty.IsEmpty() && ArrEmpty.CountSetBits() == 0);
+	}
+
+	{
+		TBitArray<> ArrTrue(true, NumBitsPerDWORD);
+		CHECK_MESSAGE(TEXT("Find and set first zero bit in an array containing 32 true"), ArrTrue.FindAndSetFirstZeroBit() == INDEX_NONE);
+		CHECK_MESSAGE(TEXT("Array should still contain 32 true bits"), ArrTrue.CountSetBits() == NumBitsPerDWORD);
+	}
+
+	{
+		TBitArray<> ArrFalse(false, NumBitsPerDWORD);
+		CHECK_MESSAGE(TEXT("Find and set first zero bit in an array containing 32 false"), ArrFalse.FindAndSetFirstZeroBit() == 0);
+		CHECK_MESSAGE(TEXT("Array should contain 1 true bit"), ArrFalse.CountSetBits() == 1);
+	}
+}
+
+TEST_CASE_NAMED(FBitArrayTestFindAndSetLastZeroBitTest, "System::Core::Containers::BitArray::FindAndSetLastZeroBit", "[ApplicationContextMask][SmokeFilter]")
+{
+	{
+		TBitArray<> ArrEmpty;
+		CHECK_MESSAGE(TEXT("Find and set last zero bit in an empty array"), ArrEmpty.FindAndSetLastZeroBit() == INDEX_NONE);
+		CHECK_MESSAGE(TEXT("Find and set last zero bit in an empty array"), ArrEmpty == TBitArray<>());
+	}
+
+	{
+		TBitArray<> ArrTrue(true, NumBitsPerDWORD);
+		CHECK_MESSAGE(TEXT("Find and set last zero bit in an array containing 32 true"), ArrTrue.FindAndSetLastZeroBit() == INDEX_NONE);
+		CHECK_MESSAGE(TEXT("Array should still contain 32 true bits"), ArrTrue.CountSetBits() == NumBitsPerDWORD);
+	}
+
+	{
+		TBitArray<> ArrFalse(false, NumBitsPerDWORD);
+		CHECK_MESSAGE(TEXT("Find and set last zero bit in an array containing 32 false"), ArrFalse.FindAndSetLastZeroBit() == NumBitsPerDWORD - 1);
+		CHECK_MESSAGE(TEXT("Array should contain 1 true bit"), ArrFalse.CountSetBits() == 1);
+	}
+}
+
 class FBitArrayMemoryTest
 {
 public:
