FormatStringSan: correctly handle UTF8 args

diff --git a/Engine/Source/Runtime/Core/Public/String/FormatStringSan.h b/Engine/Source/Runtime/Core/Public/String/FormatStringSan.h
--- a/Engine/Source/Runtime/Core/Public/String/FormatStringSan.h
+++ b/Engine/Source/Runtime/Core/Public/String/FormatStringSan.h
@@ -44,7 +44,9 @@ namespace UE::Core::Private
 			StatusLLNeedsIntegerSpec,
 			StatusLLNeedsIntegerArg,
 			StatusI64BadSpec,
-			StatusI64NeedsIntegerArg
+			StatusI64NeedsIntegerArg,
+			StatusUTF8NeedsStringCastAndS,
+			StatusUTF8NeedsStringCast
 		};
 
 		template <int N>
@@ -102,6 +104,8 @@ namespace UE::Core::Private
 		FMT_STR_ERR(StatusI64BadSpec, "'%I' must appear as '%I64' with an integral suffix (eg. '%I64d', '%I64u', etc.)");
 		FMT_STR_ERR(StatusI64NeedsIntegerArg, "'%I64[ ]' expects integral arg (eg. `char`, `int`, `long`, etc.)");
 		FMT_STR_ERR(StatusDynamicLengthSpecNeedsIntegerArg, "dynamic field width specifier '*' expects integral arg (eg. `char`, `int`, `long`, etc.)");
+		FMT_STR_ERR(StatusUTF8NeedsStringCastAndS, "Pass UTF8 strings to StringCast<TCHAR> and use %s.");
+		FMT_STR_ERR(StatusUTF8NeedsStringCast, "'%s' expects `TCHAR*`; pass UTF8 strings to StringCast<TCHAR> first.");
 
 #undef FMT_STR_ERR
 
@@ -109,7 +113,7 @@ namespace UE::Core::Private
 		inline constexpr bool bIsPointerTo =
 			std::is_pointer_v<T> && std::is_same_v<std::remove_cv_t<std::remove_pointer_t<T>>, V>;
 		template <typename T>
-		inline constexpr bool bIsCharPtr = bIsPointerTo<T, char> || bIsPointerTo<T, UTF8CHAR>;
+		inline constexpr bool bIsCharPtr = bIsPointerTo<T, char>;
 		template <typename T>
 		inline constexpr bool bIsTCharPtr = bIsPointerTo<T, TCHAR>;
 		template <typename T>
@@ -241,7 +245,11 @@ namespace UE::Core::Private
 				}
 				else if (P[0] == TEXT('h') && P[1] == TEXT('s'))
 				{
-					if constexpr (!bIsCharPtr<Arg>)
+					if constexpr (bIsPointerTo<Arg, UTF8CHAR>)
+					{
+						return {StatusUTF8NeedsStringCastAndS, CurArgPos};
+					}
+					else if constexpr (!bIsCharPtr<Arg>)
 					{
 						return {StatusHSNeedsCharPtrArg, CurArgPos};
 					}
@@ -256,7 +264,11 @@ namespace UE::Core::Private
 				case TEXT('%'):
 					return TCheckFormatString<TFmtArgTypes<Arg, Args...>>::Check(false, CurArgPos, P + 1);
 				case TEXT('s'):
-					if constexpr (!bIsTCharPtr<Arg>)
+					if constexpr (bIsPointerTo<Arg, UTF8CHAR>)
+					{
+						return {StatusUTF8NeedsStringCast, CurArgPos};
+					}
+					else if constexpr (!bIsTCharPtr<Arg>)
 					{
 						return {StatusSNeedsTCHARPtrArg, CurArgPos};
 					}
