Change LogLatentCommands category to be declared extern rather than to be declared static in a header.

diff --git a/Engine/Source/Runtime/Core/Private/Misc/AutomationTest.cpp b/Engine/Source/Runtime/Core/Private/Misc/AutomationTest.cpp
--- a/Engine/Source/Runtime/Core/Private/Misc/AutomationTest.cpp
+++ b/Engine/Source/Runtime/Core/Private/Misc/AutomationTest.cpp
@@ -20,6 +20,7 @@
 #include "Misc/ScopedSlowTask.h"
 #include "Modules/ModuleManager.h"
 
+DEFINE_LOG_CATEGORY(LogLatentCommands)
 DEFINE_LOG_CATEGORY_STATIC(LogAutomationTest, Warning, All);
 
 namespace AutomationTest
diff --git a/Engine/Source/Runtime/Core/Public/Misc/AutomationTest.h b/Engine/Source/Runtime/Core/Public/Misc/AutomationTest.h
--- a/Engine/Source/Runtime/Core/Public/Misc/AutomationTest.h
+++ b/Engine/Source/Runtime/Core/Public/Misc/AutomationTest.h
@@ -47,7 +47,7 @@
 
 #include <atomic>
 
-DEFINE_LOG_CATEGORY_STATIC(LogLatentCommands, Log, All);
+CORE_API DECLARE_LOG_CATEGORY_EXTERN(LogLatentCommands, Log, All);
 class FAutomationTestBase;
 
 #ifndef WITH_AUTOMATION_TESTS
