Turned off callstack trace for Win ARM64 and ARM64EC.

diff --git a/Engine/Source/Runtime/Core/Public/ProfilingDebugging/CallstackTrace.h b/Engine/Source/Runtime/Core/Public/ProfilingDebugging/CallstackTrace.h
--- a/Engine/Source/Runtime/Core/Public/ProfilingDebugging/CallstackTrace.h
+++ b/Engine/Source/Runtime/Core/Public/ProfilingDebugging/CallstackTrace.h
@@ -9,7 +9,7 @@
 ////////////////////////////////////////////////////////////////////////////////
 #if !defined(UE_CALLSTACK_TRACE_ENABLED)
 	#if UE_TRACE_ENABLED && !UE_BUILD_SHIPPING
-		#if PLATFORM_WINDOWS 
+		#if PLATFORM_WINDOWS && !PLATFORM_CPU_ARM_FAMILY
 			#define UE_CALLSTACK_TRACE_ENABLED 1
 		#endif
 	#endif
