Remove the initial set of deprecated 'custom chunk' API functions, deprecated in 5.1.

 - Titles should have migrated to Named Chunk installation at this point.

diff --git a/Engine/Source/Runtime/Core/Private/GenericPlatform/GenericPlatformChunkInstall.cpp b/Engine/Source/Runtime/Core/Private/GenericPlatform/GenericPlatformChunkInstall.cpp
--- a/Engine/Source/Runtime/Core/Private/GenericPlatform/GenericPlatformChunkInstall.cpp
+++ b/Engine/Source/Runtime/Core/Private/GenericPlatform/GenericPlatformChunkInstall.cpp
@@ -1,58 +1,6 @@
 // Copyright Epic Games, Inc. All Rights Reserved.
 
 #include "GenericPlatform/GenericPlatformChunkInstall.h"
-#include "HAL/PlatformMisc.h"
 
 DEFINE_LOG_CATEGORY(LogChunkInstaller)
 
-
-
-/*
- *  Functionality to to implement the deprecated 'FCustomChunk' API using the new 'named chunk' API
- */
-
-PRAGMA_DISABLE_DEPRECATION_WARNINGS
-
-
-bool FGenericPlatformChunkInstall_WithEmulatedCustomChunks::IsChunkInstallationPending(const TArray<FCustomChunk>& ChunkTagsID)
-{
-	for (const FCustomChunk& CustomChunk : ChunkTagsID)
-	{
-		if (IsNamedChunkInProgress(FName(CustomChunk.ChunkTag)) || IsNamedChunkInProgress(FName(CustomChunk.ChunkTag2)) )
-		{
-			return true;
-		}
-	}
-
-	return false;
-}
-
-bool FGenericPlatformChunkInstall_WithEmulatedCustomChunks::InstallChunks(const TArray<FCustomChunk>& ChunkTagsID)
-{
-	for (const FCustomChunk& CustomChunk : ChunkTagsID)
-	{
-		if (!InstallNamedChunk(FName(CustomChunk.ChunkTag)) && !InstallNamedChunk(FName(CustomChunk.ChunkTag2)))
-		{
-			return false;
-		}
-	}
-
-	return true;
-}
-
-bool FGenericPlatformChunkInstall_WithEmulatedCustomChunks::UninstallChunks(const TArray<FCustomChunk>& ChunkTagsID)
-{
-	for (const FCustomChunk& CustomChunk : ChunkTagsID)
-	{
-		if (!UninstallNamedChunk(FName(CustomChunk.ChunkTag)) && !UninstallNamedChunk(FName(CustomChunk.ChunkTag2)))
-		{
-			return false;
-		}
-	}
-
-	return true;
-}
-
-
-
-PRAGMA_ENABLE_DEPRECATION_WARNINGS
diff --git a/Engine/Source/Runtime/Core/Public/GenericPlatform/GenericPlatformChunkInstall.h b/Engine/Source/Runtime/Core/Public/GenericPlatform/GenericPlatformChunkInstall.h
--- a/Engine/Source/Runtime/Core/Public/GenericPlatform/GenericPlatformChunkInstall.h
+++ b/Engine/Source/Runtime/Core/Public/GenericPlatform/GenericPlatformChunkInstall.h
@@ -211,18 +211,6 @@ public:
 	virtual void RemoveChunkInstallDelegate( FDelegateHandle Delegate ) = 0;
 
 
-	UE_DEPRECATED(5.1, "Call SupportsNamedChunkInstall instead")
-	virtual bool SupportsIntelligentInstall() = 0;
-
-	UE_DEPRECATED(5.1, "Call IsNamedChunkInProgress instead")
-	virtual bool IsChunkInstallationPending(const TArray<FCustomChunk>& ChunkTagsID) = 0;
-
-	UE_DEPRECATED(5.1, "Call InstallNamedChunks instead")
-	virtual bool InstallChunks(const TArray<FCustomChunk>& ChunkTagsID) = 0;
-
-	UE_DEPRECATED(5.1, "Call UninstallNamedChunks instead")
-	virtual bool UninstallChunks(const TArray<FCustomChunk>& ChunkTagsID) = 0;
-
 	UE_DEPRECATED(5.2, "Call GetNamedChunksByType instead")
 	virtual TArray<FCustomChunk> GetCustomChunksByType(ECustomChunkType DesiredChunkType) = 0;
 
@@ -408,26 +396,6 @@ public:
 		InstallDelegate.Remove(Delegate);
 	}
 
-	virtual bool SupportsIntelligentInstall() override
-	{
-		return false;
-	}
-
-	virtual bool IsChunkInstallationPending(const TArray<FCustomChunk>& ChunkTags) override
-	{
-		return false;
-	}
-
-	virtual bool InstallChunks(const TArray<FCustomChunk>& ChunkTagIDs) override
-	{
-		return false;
-	}
-
-	virtual bool UninstallChunks(const TArray<FCustomChunk>& ChunkTagsID) override
-	{
-		return false;
-	}
-
 	virtual TArray<FCustomChunk> GetCustomChunksByType(ECustomChunkType DesiredChunkType) override
 	{
 		return TArray<FCustomChunk>();
@@ -512,14 +480,5 @@ protected:
 };
 
 
-// temporary helper base class for platform chunk installers that have implemented named chunk support to provide FCustomChunk emulation
-class FGenericPlatformChunkInstall_WithEmulatedCustomChunks : public FGenericPlatformChunkInstall
-{
-public:
-	virtual bool SupportsIntelligentInstall() override final { return true; }
-	CORE_API virtual bool IsChunkInstallationPending(const TArray<FCustomChunk>& ChunkTagsID) override final;
-	CORE_API virtual bool InstallChunks(const TArray<FCustomChunk>& ChunkTagsID)  override final;
-	CORE_API virtual bool UninstallChunks(const TArray<FCustomChunk>& ChunkTagsID) override final;
-};
 
 PRAGMA_ENABLE_DEPRECATION_WARNINGS
