Replace __FILE__ with __builtin_FILE() in constexpr

#rnx

diff --git a/Engine/Source/Runtime/Core/Public/Logging/LogMacros.h b/Engine/Source/Runtime/Core/Public/Logging/LogMacros.h
--- a/Engine/Source/Runtime/Core/Public/Logging/LogMacros.h
+++ b/Engine/Source/Runtime/Core/Public/Logging/LogMacros.h
@@ -281,7 +281,7 @@ CORE_API void BasicFatalLog(const FLogCategoryBase& Category, const FStaticBasic
 		static_assert(TIsArrayOrRefOfTypeByPredicate<decltype(Format), TIsCharEncodingCompatibleWithTCHAR>::Value, "Formatting string must be a TCHAR array."); \
 		UE_VALIDATE_FORMAT_STRING(Format, ##__VA_ARGS__); \
 		static ::UE::Logging::Private::FStaticBasicLogDynamicData LOG_Dynamic; \
-		static constexpr ::UE::Logging::Private::FStaticBasicLogRecord LOG_Static(Format, __FILE__, __builtin_LINE(), ::ELogVerbosity::Verbosity, LOG_Dynamic); \
+		static constexpr ::UE::Logging::Private::FStaticBasicLogRecord LOG_Static(Format, __builtin_FILE(), __builtin_LINE(), ::ELogVerbosity::Verbosity, LOG_Dynamic); \
 		static_assert((::ELogVerbosity::Verbosity & ::ELogVerbosity::VerbosityMask) < ::ELogVerbosity::NumVerbosity && ::ELogVerbosity::Verbosity > 0, "Verbosity must be constant and in range."); \
 		if constexpr ((::ELogVerbosity::Verbosity & ELogVerbosity::VerbosityMask) == ::ELogVerbosity::Fatal) \
 		{ \
diff --git a/Engine/Source/Runtime/Core/Public/Logging/StructuredLog.h b/Engine/Source/Runtime/Core/Public/Logging/StructuredLog.h
--- a/Engine/Source/Runtime/Core/Public/Logging/StructuredLog.h
+++ b/Engine/Source/Runtime/Core/Public/Logging/StructuredLog.h
@@ -433,7 +433,7 @@ inline void LogIfActive(const LogCategoryType& Category, const LogRecordType& Lo
 	do \
 	{ \
 		static ::UE::Logging::Private::FStaticLogDynamicData LOG_Dynamic; \
-		static constexpr ::UE::Logging::Private::FStaticLogRecord LOG_Static UE_PRIVATE_LOGFMT_AGGREGATE(TEXT(Format), __FILE__, __builtin_LINE(), ::ELogVerbosity::Verbosity, LOG_Dynamic); \
+		static constexpr ::UE::Logging::Private::FStaticLogRecord LOG_Static UE_PRIVATE_LOGFMT_AGGREGATE(TEXT(Format), __builtin_FILE(), __builtin_LINE(), ::ELogVerbosity::Verbosity, LOG_Dynamic); \
 		::UE::Logging::Private::LogIfActive<FLogCategory##CategoryName, ::ELogVerbosity::Verbosity>(CategoryName, LOG_Static, ##__VA_ARGS__); \
 	} \
 	while (false)
@@ -442,7 +442,7 @@ inline void LogIfActive(const LogCategoryType& Category, const LogRecordType& Lo
 	do \
 	{ \
 		static ::UE::Logging::Private::FStaticLogDynamicData LOG_Dynamic; \
-		static constexpr ::UE::Logging::Private::FStaticLocalizedLogRecord LOG_Static UE_PRIVATE_LOGFMT_AGGREGATE(TEXT(Namespace), TEXT(Key), TEXT(Format), __FILE__, __builtin_LINE(), ::ELogVerbosity::Verbosity, LOG_Dynamic); \
+		static constexpr ::UE::Logging::Private::FStaticLocalizedLogRecord LOG_Static UE_PRIVATE_LOGFMT_AGGREGATE(TEXT(Namespace), TEXT(Key), TEXT(Format), __builtin_FILE(), __builtin_LINE(), ::ELogVerbosity::Verbosity, LOG_Dynamic); \
 		::UE::Logging::Private::LogIfActive<FLogCategory##CategoryName, ::ELogVerbosity::Verbosity>(CategoryName, LOG_Static, ##__VA_ARGS__); \
 	} \
 	while (false)
diff --git a/Engine/Source/Runtime/RenderCore/Public/ShaderParameterMacros.h b/Engine/Source/Runtime/RenderCore/Public/ShaderParameterMacros.h
--- a/Engine/Source/Runtime/RenderCore/Public/ShaderParameterMacros.h
+++ b/Engine/Source/Runtime/RenderCore/Public/ShaderParameterMacros.h
@@ -1322,7 +1322,7 @@ struct TShaderParameterStructTypeInfo<StructType[InNumElements]>
 			static constexpr int32 NumElements = 0; \
 			static constexpr int32 Alignment = SHADER_PARAMETER_STRUCT_ALIGNMENT; \
 			static constexpr bool bIsStoredInConstantBuffer = true; \
-			static constexpr const ANSICHAR* const FileName = UE_LOG_SOURCE_FILE(__FILE__); \
+			static constexpr const ANSICHAR* const FileName = UE_LOG_SOURCE_FILE(__builtin_FILE()); \
 			static constexpr int32 FileLine = __builtin_LINE(); \
 			using TAlignedType = StructTypeName; \
 			static const FShaderParametersMetadata* GetStructMetadata() GetStructMetadataScope \
