Change TaskArray back to use shared pointers for performance reasons (it was using this before I hacked out the custom shared pointer).

diff --git a/Engine/Source/Runtime/Core/Private/AutoRTFM/TaskArray.h b/Engine/Source/Runtime/Core/Private/AutoRTFM/TaskArray.h
--- a/Engine/Source/Runtime/Core/Private/AutoRTFM/TaskArray.h
+++ b/Engine/Source/Runtime/Core/Private/AutoRTFM/TaskArray.h
@@ -4,60 +4,80 @@
 
 #include "Containers/Array.h"
 #include "Templates/Function.h"
-#include "Utils.h"
+#include "Templates/SharedPointer.h"
 
 namespace AutoRTFM
 {
 
+template<typename T> class TBackwards final
+{
+	T& Obj;
+public:
+	TBackwards(T& Obj) : Obj(Obj) {}
+	auto begin() { return Obj.rbegin(); }
+	auto end() { return Obj.rend(); }
+};
+
 template<typename T>
 class TTaskArray
 {
+	template<typename U> using USharedPtr = TSharedPtr<U, ESPMode::NotThreadSafe>;
+
 public:
-    TTaskArray() = default;
+	TTaskArray() : Latest(new TArray<T>()) {}
+	TTaskArray(const TTaskArray&) = delete;
+	void operator=(const TTaskArray&) = delete;
 
-    bool IsEmpty() const { return Latest.IsEmpty() && Stash.IsEmpty(); }
+    bool IsEmpty() const { return Latest->IsEmpty() && Stash.IsEmpty(); }
 
     void Add(T&& value)
     {
-        Latest.Push(MoveTemp(value));
+        Latest->Push(MoveTemp(value));
     }
 
     void Add(const T& value)
     {
-        Latest.Push(value);
+        Latest->Push(value);
     }
 
     void AddAll(TTaskArray<T>&& Other)
     {
         Canonicalize();
 
-        for (TArray<T>& StashedVectorBox : Other.Stash)
+        for (USharedPtr<TArray<T>>& StashedVectorBox : Other.Stash)
         {
             Stash.Push(MoveTemp(StashedVectorBox));
         }
 
         Other.Stash.Empty();
 
-		ASSERT(Latest.IsEmpty());
-		Latest = MoveTemp(Other.Latest);
+        if (!Other.Latest->IsEmpty())
+        {
+            Stash.Push(MoveTemp(Other.Latest));
+			USharedPtr<TArray<T>> NewLatest(new TArray<T>());
+			Other.Latest = MoveTemp(NewLatest);
+        }
     }
 
     void AddAll(const TTaskArray<T>& Other)
     {
         Canonicalize();
         Other.Canonicalize();
-        for (const TArray<T>& StashedVectorBox : Other.Stash)
+        for (const USharedPtr<TArray<T>>& StashedVectorBox : Other.Stash)
         {
             Stash.Push(StashedVectorBox);
         }
+
+
+		ASSERT(Latest != nullptr);
     }
 
     template<typename TFunc>
-    bool ForEachForward(const TFunc& Func)
+    bool ForEachForward(const TFunc& Func) const
     {
-        for (const TArray<T>& StashedVectorBox : Stash)
+        for (USharedPtr<TArray<T>>& StashedVectorBox : Stash)
         {
-            for (const T& Entry : StashedVectorBox)
+            for (const T& Entry : *StashedVectorBox)
             {
                 if (!Func(Entry))
                 {
@@ -66,55 +86,54 @@ public:
             }
         }
 
-        for (const T& Entry : Latest)
+        for (const T& Entry : *Latest)
         {
             if (!Func(Entry))
             {
                 return false;
             }
         }
-
         return true;
     }
 
     template<typename TFunc>
-    bool ForEachBackward(const TFunc& Func)
+    bool ForEachBackward(const TFunc& Func) const
     {
-        for (size_t Index = Latest.Num(); Index--;)
-        {
-            if (!Func(Latest[Index]))
-            {
-                return false;
-            }
-        }
+		for (const T& Entry : TBackwards(*Latest))
+		{
+			if (!Func(Entry))
+			{
+				return false;
+			}
+		}
+
+		for (USharedPtr<TArray<T>>& StashedVectorBox : TBackwards(Stash))
+		{
+			for (const T& Entry : TBackwards(*StashedVectorBox))
+			{
+				if (!Func(Entry))
+				{
+					return false;
+				}
+			}
+		}
 
-        for (size_t IndexInStash = Stash.Num(); IndexInStash--;)
-        {
-            const TArray<T>& StashedVector = Stash[IndexInStash];
-            for (size_t Index = StashedVector.Num(); Index--;)
-            {
-                if (!Func(StashedVector[Index]))
-                {
-                    return false;
-                }
-            }
-        }
         return true;
     }
 
     void Reset()
     {
-        Latest.Empty();
+        Latest->Empty();
         Stash.Empty();
     }
 
     size_t Num() const
     {
-        size_t Result = Latest.Num();
+        size_t Result = Latest->Num();
 
         for (size_t Index = 0; Index < Stash.Num(); Index++)
         {
-            const TArray<T>& StashedVector = Stash[Index];
+            const TArray<T>& StashedVector = *Stash[Index];
             Result += StashedVector.Num();
         }
 
@@ -127,15 +146,19 @@ private:
     // method is `const`.
     void Canonicalize() const
     {
-        if (!Latest.IsEmpty())
+        if (!Latest->IsEmpty())
         {
-            Stash.Push(MoveTemp(Latest));
+			Stash.Push(MoveTemp(Latest));
+			USharedPtr<TArray<T>> NewLatest(new TArray<T>());
+			Latest = MoveTemp(NewLatest);
         }
+
+		ASSERT(Latest != nullptr);
     }
     
-    mutable TArray<T> Latest;
+    mutable USharedPtr<TArray<T>> Latest;
     
-    mutable TArray<TArray<T>> Stash;
+    mutable TArray<USharedPtr<TArray<T>>> Stash;
 };
 
 } // namespace AutoRTFM
