Add lambda delegate tied to lifetime of a shared-pointer, similar to WeakLambdas for UObjects.

diff --git a/Engine/Source/Runtime/Core/Public/Delegates/Delegate.h b/Engine/Source/Runtime/Core/Public/Delegates/Delegate.h
--- a/Engine/Source/Runtime/Core/Public/Delegates/Delegate.h
+++ b/Engine/Source/Runtime/Core/Public/Delegates/Delegate.h
@@ -183,6 +183,7 @@
  *	BindThreadSafeSP(SharedPtr, &FClass::Function)		|	Call a native class member function via a TWeakPtr, will not be called if shared pointer is invalid
  *	BindRaw(RawPtr, &FClass::Function)					|	Call a native class member function with no safety checks. You MUST call Unbind or Remove when object dies to avoid crashes!
  *	BindLambda(Lambda)									|	Call a lambda function with no safety checks. You MUST make sure all captures will be safe at a later point to avoid crashes!
+ *	BindSPLambda(SharedPtr, Lambda)						|	Call a lambda function only if shared pointer is still valid. Captured 'this' will always be valid but any other captures may not be
  *	BindWeakLambda(UObject, Lambda)						|	Call a lambda function only if UObject is still valid. Captured 'this' will always be valid but any other captures may not be
  *	BindUFunction(UObject, FName("FunctionName"))		|	Usable for both native and dynamic delegates, will call a UFUNCTION with specified name
  *	BindDynamic(UObject, &UClass::FunctionName)			|	Convenience wrapper only available for dynamic delegates, FunctionName must be declared as a UFUNCTION
diff --git a/Engine/Source/Runtime/Core/Public/Delegates/DelegateBase.h b/Engine/Source/Runtime/Core/Public/Delegates/DelegateBase.h
--- a/Engine/Source/Runtime/Core/Public/Delegates/DelegateBase.h
+++ b/Engine/Source/Runtime/Core/Public/Delegates/DelegateBase.h
@@ -112,6 +112,9 @@ private:
 	template <bool, class, ESPMode, typename, typename, typename...>
 	friend class TBaseSPMethodDelegateInstance;
 
+	template <typename, ESPMode, typename, typename, typename, typename...>
+	friend class TBaseSPLambdaDelegateInstance;
+
 	template <bool, class, typename, typename, typename...>
 	friend class TBaseRawMethodDelegateInstance;
 
diff --git a/Engine/Source/Runtime/Core/Public/Delegates/DelegateInstancesImpl.h b/Engine/Source/Runtime/Core/Public/Delegates/DelegateInstancesImpl.h
--- a/Engine/Source/Runtime/Core/Public/Delegates/DelegateInstancesImpl.h
+++ b/Engine/Source/Runtime/Core/Public/Delegates/DelegateInstancesImpl.h
@@ -313,6 +313,106 @@ protected:
 };
 
 
+template <typename UserClass, ESPMode SPMode, typename RetValType, typename... ParamTypes, typename UserPolicy, typename FunctorType, typename... VarTypes>
+	class TBaseSPLambdaDelegateInstance<UserClass, SPMode, RetValType(ParamTypes...), UserPolicy, FunctorType, VarTypes...> : public TCommonDelegateInstanceState<RetValType(ParamTypes...), UserPolicy, VarTypes...>
+{
+private:
+	static_assert(std::is_same_v<FunctorType, typename TRemoveReference<FunctorType>::Type>, "FunctorType cannot be a reference");
+
+	using Super = TCommonDelegateInstanceState<RetValType(ParamTypes...), UserPolicy, VarTypes...>;
+	using DelegateBaseType = typename UserPolicy::FDelegateExtras;
+
+public:
+	template <typename InFunctorType, typename... InVarTypes>
+	explicit TBaseSPLambdaDelegateInstance(const TSharedPtr<UserClass, SPMode>& InContextObject, InFunctorType&& InFunctor, InVarTypes&&... Vars)
+		: Super(Forward<InVarTypes>(Vars)...)
+		, ContextObject(InContextObject)
+		, Functor(Forward<InFunctorType>(InFunctor))
+	{
+	}
+
+	// IDelegateInstance interface
+
+#if USE_DELEGATE_TRYGETBOUNDFUNCTIONNAME
+
+	FName TryGetBoundFunctionName() const final
+	{
+		return NAME_None;
+	}
+
+#endif
+
+	UObject* GetUObject() const final
+	{
+		return nullptr;
+	}
+
+	const void* GetObjectForTimerManager() const final
+	{
+		return ContextObject.Pin().Get();
+	}
+
+	uint64 GetBoundProgramCounterForTimerManager() const final
+	{
+		return 0;
+	}
+
+	// Deprecated
+	bool HasSameObject(const void* InContextObject) const final
+	{
+		return ContextObject.Pin().Get() == InContextObject;
+	}
+
+	bool IsSafeToExecute() const final
+	{
+		return ContextObject.IsValid();
+	}
+
+public:
+
+public:
+	// IBaseDelegateInstance interface
+	void CreateCopy(TDelegateBase<FThreadSafeDelegateMode>& Base) const final
+	{
+		Base.template CreateDelegateInstance<TBaseSPLambdaDelegateInstance>(*this);
+	}
+
+	void CreateCopy(TDelegateBase<FNotThreadSafeDelegateMode>& Base) const final
+	{
+		Base.template CreateDelegateInstance<TBaseSPLambdaDelegateInstance>(*this);
+	}
+
+	void CreateCopy(TDelegateBase<FNotThreadSafeNotCheckedDelegateMode>& Base) const final
+	{
+		Base.template CreateDelegateInstance<TBaseSPLambdaDelegateInstance>(*this);
+	}
+
+	RetValType Execute(ParamTypes... Params) const final
+	{
+		return this->Payload.ApplyAfter(Functor, Forward<ParamTypes>(Params)...);
+	}
+
+	bool ExecuteIfSafe(ParamTypes... Params) const final
+	{
+		if (ContextObject.IsValid())
+		{
+			(void)this->Payload.ApplyAfter(Functor, Forward<ParamTypes>(Params)...);
+			return true;
+		}
+
+		return false;
+	}
+
+private:
+
+	// Weak reference to an instance of the user's class that controls the validity of the lambda.
+	TWeakPtr<UserClass, SPMode> ContextObject;
+
+	// We make this mutable to allow mutable lambdas to be bound and executed.
+	mutable std::remove_const_t<FunctorType> Functor;
+};
+
+
 template <bool bConst, class UserClass, typename RetValType, typename... ParamTypes, typename UserPolicy, typename... VarTypes>
 class TBaseRawMethodDelegateInstance<bConst, UserClass, RetValType(ParamTypes...), UserPolicy, VarTypes...> : public TCommonDelegateInstanceState<RetValType(ParamTypes...), UserPolicy, VarTypes...>
 {
diff --git a/Engine/Source/Runtime/Core/Public/Delegates/DelegateSignatureImpl.inl b/Engine/Source/Runtime/Core/Public/Delegates/DelegateSignatureImpl.inl
--- a/Engine/Source/Runtime/Core/Public/Delegates/DelegateSignatureImpl.inl
+++ b/Engine/Source/Runtime/Core/Public/Delegates/DelegateSignatureImpl.inl
@@ -105,6 +105,25 @@ public:
 		return Result;
 	}
 
+	/**
+	 * Static: Creates a weak shared pointer C++ lambda delegate
+	 * technically this works for any functor types, but lambdas are the primary use case
+	 */
+	template<typename UserClass, ESPMode Mode, typename FunctorType, typename... VarTypes>
+	UE_NODISCARD inline static TDelegate<RetValType(ParamTypes...), UserPolicy> CreateSPLambda(const TSharedRef<UserClass, Mode>& InUserObjectRef, FunctorType&& InFunctor, VarTypes&&... Vars)
+	{
+		TDelegate<RetValType(ParamTypes...), UserPolicy> Result;
+		Result.template CreateDelegateInstance<TBaseSPLambdaDelegateInstance<UserClass, Mode, FuncType, UserPolicy, typename TRemoveReference<FunctorType>::Type, std::decay_t<VarTypes>...>>(InUserObjectRef, Forward<FunctorType>(InFunctor), Forward<VarTypes>(Vars)...);
+		return Result;
+	}
+	template <typename UserClass, typename FunctorType, typename... VarTypes>
+	UE_NODISCARD inline static TDelegate<RetValType(ParamTypes...), UserPolicy> CreateSPLambda(UserClass* InUserObject, FunctorType&& InFunctor, VarTypes&&... Vars)
+	{
+		TDelegate<RetValType(ParamTypes...), UserPolicy> Result;
+		Result.template CreateDelegateInstance<TBaseSPLambdaDelegateInstance<UserClass, decltype(InUserObject->AsShared())::Mode, FuncType, UserPolicy, typename TRemoveReference<FunctorType>::Type, std::decay_t<VarTypes>...>>(StaticCastSharedRef<UserClass>(InUserObject->AsShared()), Forward<FunctorType>(InFunctor), Forward<VarTypes>(Vars)...);
+		return Result;
+	}
+
 	/**
 	 * Static: Creates a weak object C++ lambda delegate
 	 * technically this works for any functor types, but lambdas are the primary use case
@@ -339,6 +358,21 @@ public:
 		Super::template CreateDelegateInstance<TBaseFunctorDelegateInstance<FuncType, UserPolicy, typename TRemoveReference<FunctorType>::Type, std::decay_t<VarTypes>...>>(Forward<FunctorType>(InFunctor), Forward<VarTypes>(Vars)...);
 	}
 
+	/**
+	 * Static: Binds a weak shared pointer C++ lambda delegate
+	 * technically this works for any functor types, but lambdas are the primary use case
+	 */
+	template<typename UserClass, ESPMode Mode, typename FunctorType, typename... VarTypes>
+	inline void BindSPLambda(const TSharedRef<UserClass, Mode>& InUserObjectRef, FunctorType&& InFunctor, VarTypes&&... Vars)
+	{
+		Super::template CreateDelegateInstance<TBaseSPLambdaDelegateInstance<const UserClass, Mode, FuncType, UserPolicy, typename TRemoveReference<FunctorType>::Type, std::decay_t<VarTypes>...>>(InUserObjectRef, Forward<FunctorType>(InFunctor), Forward<VarTypes>(Vars)...);
+	}
+	template <typename UserClass, typename FunctorType, typename... VarTypes>
+	inline void BindSPLambda(const UserClass* InUserObject, FunctorType&& InFunctor, VarTypes&&... Vars)
+	{
+		Super::template CreateDelegateInstance<TBaseSPLambdaDelegateInstance<const UserClass, decltype(InUserObject->AsShared())::Mode, FuncType, UserPolicy, typename TRemoveReference<FunctorType>::Type, std::decay_t<VarTypes>...>>(StaticCastSharedRef<const UserClass>(InUserObject->AsShared()), Forward<FunctorType>(InFunctor), Forward<VarTypes>(Vars)...);
+	}
+
 	/**
 	 * Static: Binds a weak object C++ lambda delegate
 	 * technically this works for any functor types, but lambdas are the primary use case
@@ -706,6 +740,19 @@ public:
 		return Super::AddDelegateInstance(FDelegate::CreateLambda(Forward<FunctorType>(InFunctor), Forward<VarTypes>(Vars)...));
 	}
 
+	/**
+	 * Adds a weak shared pointer C++ lambda delegate
+	 * technically this works for any functor types, but lambdas are the primary use case
+	 *
+	 * @param	InUserObjectRef	User object to bind to
+	 * @param	InFunctor		Functor (e.g. Lambda)
+	 */
+	template <typename UserClass, typename FunctorType, typename... VarTypes>
+	inline FDelegateHandle AddSPLambda(const UserClass* InUserObject, FunctorType&& InFunctor, VarTypes&&... Vars)
+	{
+		return Super::AddDelegateInstance(FDelegate::CreateSPLambda(InUserObject, Forward<FunctorType>(InFunctor), Forward<VarTypes>(Vars)...));
+	}
+
 	/**
 	 * Adds a weak object C++ lambda delegate
 	 * technically this works for any functor types, but lambdas are the primary use case
diff --git a/Engine/Source/Runtime/Core/Public/HAL/IConsoleManager.h b/Engine/Source/Runtime/Core/Public/HAL/IConsoleManager.h
--- a/Engine/Source/Runtime/Core/Public/HAL/IConsoleManager.h
+++ b/Engine/Source/Runtime/Core/Public/HAL/IConsoleManager.h
@@ -216,6 +216,12 @@ struct FNullConsoleVariableDelegate
 		return {};
 	}
 
+	template<typename UserClass, typename FunctorType, typename... VarTypes>
+	inline static DerivedType CreateSPLambda(UserClass*, FunctorType&&, VarTypes...)
+	{
+		return {};
+	}
+
 	template<typename UserClass, typename FunctorType, typename... VarTypes>
 	inline static DerivedType CreateWeakLambda(UserClass*, FunctorType&&, VarTypes...)
 	{
