Fix for lambda delegate tied to lifetime of a shared-pointer.

diff --git a/Engine/Source/Runtime/Core/Public/Delegates/DelegateInstancesImplFwd.h b/Engine/Source/Runtime/Core/Public/Delegates/DelegateInstancesImplFwd.h
--- a/Engine/Source/Runtime/Core/Public/Delegates/DelegateInstancesImplFwd.h
+++ b/Engine/Source/Runtime/Core/Public/Delegates/DelegateInstancesImplFwd.h
@@ -22,6 +22,12 @@ class TBaseUFunctionDelegateInstance;
 template <bool bConst, class UserClass, ESPMode SPMode, typename FuncType, typename UserPolicy, typename... VarTypes>
 class TBaseSPMethodDelegateInstance;
 
+/**
+ * Implements a delegate binding for shared pointer functors, e.g. lambdas.
+ */
+template <typename UserClass, ESPMode SPMode, typename FuncType, typename UserPolicy, typename FunctorType, typename... VarTypes>
+class TBaseSPLambdaDelegateInstance;
+
 /**
  * Implements a delegate binding for C++ member functions.
  */
