ModuleManager: Logging: Downgrade log message to verbose to make cookcommandlet logs easier to read.
#rnx

diff --git a/Engine/Source/Runtime/Core/Private/Modules/ModuleManager.cpp b/Engine/Source/Runtime/Core/Private/Modules/ModuleManager.cpp
--- a/Engine/Source/Runtime/Core/Private/Modules/ModuleManager.cpp
+++ b/Engine/Source/Runtime/Core/Private/Modules/ModuleManager.cpp
@@ -848,7 +848,7 @@ void FModuleManager::UnloadModulesAtShutdown()
 	// Now actually unload all modules
 	for (FModulePair& ModuleToUnload : ModulesToUnload)
 	{
-		UE_LOG(LogModuleManager, Log, TEXT("Shutting down and abandoning module %s (%d)"), *ModuleToUnload.ModuleName.ToString(), ModuleToUnload.LoadOrder);
+		UE_LOG(LogModuleManager, Verbose, TEXT("Shutting down and abandoning module %s (%d)"), *ModuleToUnload.ModuleName.ToString(), ModuleToUnload.LoadOrder);
 		const bool bIsShutdown = true;
 		UnloadModule(ModuleToUnload.ModuleName, bIsShutdown);
 		UE_LOG(LogModuleManager, Verbose, TEXT( "Returned from UnloadModule." ));
