UniqueLock: Added comments to TUniqueLock and TDynamicUniqueLock

#rnx

diff --git a/Engine/Source/Runtime/Core/Public/Async/UniqueLock.h b/Engine/Source/Runtime/Core/Public/Async/UniqueLock.h
--- a/Engine/Source/Runtime/Core/Public/Async/UniqueLock.h
+++ b/Engine/Source/Runtime/Core/Public/Async/UniqueLock.h
@@ -7,11 +7,19 @@
 namespace UE
 {
 
+/** Use with dynamic locks to defer locking on construction. */
 struct FDeferLock final
 {
 	explicit FDeferLock() = default;
 };
 
+/**
+ * A basic mutex ownership wrapper that locks on construction and unlocks on destruction.
+ *
+ * LockType must contain Lock() and Unlock() functions.
+ * 
+ * Use with mutex types like FMutex and FRecursiveMutex.
+ */
 template <typename LockType>
 class TUniqueLock final
 {
@@ -34,6 +42,13 @@ private:
 	LockType& Mutex;
 };
 
+/**
+ * A mutex ownership wrapper that allows dynamic locking, unlocking, and deferred locking.
+ *
+ * LockType must contain Lock() and Unlock() functions.
+ * 
+ * Use with mutex types like FMutex and FRecursiveMutex.
+ */
 template <typename LockType>
 class TDynamicUniqueLock final
 {
@@ -43,6 +58,7 @@ public:
 	TDynamicUniqueLock(const TDynamicUniqueLock&) = delete;
 	TDynamicUniqueLock& operator=(const TDynamicUniqueLock&) = delete;
 
+	/** Wrap a mutex and lock it. */
 	[[nodiscard]] inline explicit TDynamicUniqueLock(LockType& Lock)
 		: Mutex(&Lock)
 	{
@@ -50,11 +66,13 @@ public:
 		bLocked = true;
 	}
 
+	/** Wrap a mutex without locking it. */
 	[[nodiscard]] inline explicit TDynamicUniqueLock(LockType& Lock, FDeferLock)
 		: Mutex(&Lock)
 	{
 	}
 
+	/** Move from another lock, transferring any ownership to this lock. */
 	[[nodiscard]] inline TDynamicUniqueLock(TDynamicUniqueLock&& Other)
 		: Mutex(Other.Mutex)
 		, bLocked(Other.bLocked)
@@ -63,6 +81,7 @@ public:
 		Other.bLocked = false;
 	}
 
+	/** Move from another lock, transferring any ownership to this lock, and unlocking the previous mutex if locked. */
 	inline TDynamicUniqueLock& operator=(TDynamicUniqueLock&& Other)
 	{
 		if (bLocked)
@@ -76,6 +95,7 @@ public:
 		return *this;
 	}
 
+	/** Unlock the mutex if locked. */
 	inline ~TDynamicUniqueLock()
 	{
 		if (bLocked)
@@ -84,6 +104,7 @@ public:
 		}
 	}
 
+	/** Lock the associated mutex. This lock must have a mutex and must not be locked. */
 	void Lock()
 	{
 		check(!bLocked);
@@ -92,6 +113,7 @@ public:
 		bLocked = true;
 	}
 
+	/** Unlock the associated mutex. This lock must have a mutex and must be locked. */
 	void Unlock()
 	{
 		check(bLocked);
@@ -99,11 +121,13 @@ public:
 		Mutex->Unlock();
 	}
 
+	/** Returns true if this lock has its associated mutex locked. */
 	inline bool OwnsLock() const
 	{
 		return bLocked;
 	}
 
+	/** Returns true if this lock has its associated mutex locked. */
 	inline explicit operator bool() const
 	{
 		return OwnsLock();
