Added missing API macros

#rnx

diff --git a/Engine/Source/Runtime/Core/Public/Misc/CoreDelegates.h b/Engine/Source/Runtime/Core/Public/Misc/CoreDelegates.h
--- a/Engine/Source/Runtime/Core/Public/Misc/CoreDelegates.h
+++ b/Engine/Source/Runtime/Core/Public/Misc/CoreDelegates.h
@@ -197,10 +197,10 @@ public:
 	static CORE_API FSimpleMulticastDelegate OnFEngineLoopInitComplete;
 
 	// Called prior to running the Main function for commandlets
-	static FSimpleMulticastDelegate OnCommandletPreMain;
+	static CORE_API FSimpleMulticastDelegate OnCommandletPreMain;
 
 	// Called after running the Main function for commandlets
-	static FSimpleMulticastDelegate OnCommandletPostMain;
+	static CORE_API FSimpleMulticastDelegate OnCommandletPostMain;
 
 	// Called when the application is about to exit.
 	static CORE_API FSimpleMulticastDelegate OnExit;
