Add a new UE macro UE_ALLOCATION_FUNCTION that tells Clang-based compilers that our Unreal allocation functions create new memory. This results in more optimal codegen in the LLVM optimization passes.

diff --git a/Engine/Source/Runtime/Core/Public/Clang/ClangPlatform.h b/Engine/Source/Runtime/Core/Public/Clang/ClangPlatform.h
--- a/Engine/Source/Runtime/Core/Public/Clang/ClangPlatform.h
+++ b/Engine/Source/Runtime/Core/Public/Clang/ClangPlatform.h
@@ -22,3 +22,18 @@
 #define PLATFORM_RETURN_ADDRESS_POINTER()	__builtin_frame_address(0)
 
 #define UE_LIFETIMEBOUND [[clang::lifetimebound]]
+
+// Clang supports the following GNU attributes to let us describe allocation functions.
+// `gnu::malloc` maps to the `noalias` function return on the caller in the LLVM IR.
+// `gnu::alloc_size` maps to the `allocsize` function attribute on the caller in the LLVM IR.
+// `gnu::alloc_align` will cause any memory operations in the callee to have the specified alignment in the LLVM IR.
+// We use a selector macro trick to map up to 2 inputs to the correct macros.
+#define UE_ALLOCATION_FUNCTION_0() [[gnu::malloc]]
+#define UE_ALLOCATION_FUNCTION_1(SIZE) [[gnu::malloc, gnu::alloc_size(SIZE)]]
+
+// Note: we cannot use gnu::alloc_align(ALIGN) here yet, because the DEFAULT_ALIGNMENT argument is 0, which is not
+//       a power of two, and Clang will complain about uses of a function with a default alignment that is not a
+//       power of two.
+#define UE_ALLOCATION_FUNCTION_2(SIZE, ALIGN) [[gnu::malloc, gnu::alloc_size(SIZE)]]
+#define UE_ALLOCATION_FUNCTION_X(x, SIZE, ALIGN, FUNC, ...) FUNC
+#define UE_ALLOCATION_FUNCTION(...) UE_ALLOCATION_FUNCTION_X(,##__VA_ARGS__, UE_ALLOCATION_FUNCTION_2(__VA_ARGS__), UE_ALLOCATION_FUNCTION_1(__VA_ARGS__), UE_ALLOCATION_FUNCTION_0(__VA_ARGS__)) 
diff --git a/Engine/Source/Runtime/Core/Public/HAL/Platform.h b/Engine/Source/Runtime/Core/Public/HAL/Platform.h
--- a/Engine/Source/Runtime/Core/Public/HAL/Platform.h
+++ b/Engine/Source/Runtime/Core/Public/HAL/Platform.h
@@ -704,6 +704,12 @@
 	#define UE_LIFETIMEBOUND
 #endif
 
+/* Annotate functions that allocate new memory to ensure the compiler can optimize them accordingly.
+   The arguments to this macro specify the 1-based arguments of the annotated function that specify the size of the allocation. */
+#ifndef UE_ALLOCATION_FUNCTION
+	#define UE_ALLOCATION_FUNCTION(...)
+#endif
+
 /** Promise expression is true. Compiler can optimize accordingly with undefined behavior if wrong. Static analyzers understand this.  */
 #ifndef UE_ASSUME
 	#if defined(__clang__)
diff --git a/Engine/Source/Runtime/Core/Public/HAL/UnrealMemory.h b/Engine/Source/Runtime/Core/Public/HAL/UnrealMemory.h
--- a/Engine/Source/Runtime/Core/Public/HAL/UnrealMemory.h
+++ b/Engine/Source/Runtime/Core/Public/HAL/UnrealMemory.h
@@ -191,7 +191,7 @@ struct FMemory
 	//
 	// C style memory allocation stubs that fall back to C runtime
 	//
-	static FORCEINLINE void* SystemMalloc(SIZE_T Size)
+	UE_ALLOCATION_FUNCTION(1) static FORCEINLINE void* SystemMalloc(SIZE_T Size)
 	{
 		/* TODO: Trace! */
 		return ::malloc(Size);
@@ -207,12 +207,12 @@ struct FMemory
 	// C style memory allocation stubs.
 	//
 
-	static CORE_API void* Malloc(SIZE_T Count, uint32 Alignment = DEFAULT_ALIGNMENT);
-	static CORE_API void* Realloc(void* Original, SIZE_T Count, uint32 Alignment = DEFAULT_ALIGNMENT);
+	UE_ALLOCATION_FUNCTION(1, 2) static CORE_API void* Malloc(SIZE_T Count, uint32 Alignment = DEFAULT_ALIGNMENT);
+	UE_ALLOCATION_FUNCTION(2, 3) static CORE_API void* Realloc(void* Original, SIZE_T Count, uint32 Alignment = DEFAULT_ALIGNMENT);
 	static CORE_API void Free(void* Original);
 	static CORE_API SIZE_T GetAllocSize(void* Original);
 
-	static FORCEINLINE_DEBUGGABLE void* MallocZeroed(SIZE_T Count, uint32 Alignment = DEFAULT_ALIGNMENT)
+	UE_ALLOCATION_FUNCTION(1, 2) static FORCEINLINE_DEBUGGABLE void* MallocZeroed(SIZE_T Count, uint32 Alignment = DEFAULT_ALIGNMENT)
 	{
 		void* Memory = Malloc(Count, Alignment);
 		Memzero(Memory, Count);
