Fixed ETextGender serialization within FFormatArgumentValue

#rnx

diff --git a/Engine/Source/Runtime/Core/Private/Internationalization/Text.cpp b/Engine/Source/Runtime/Core/Private/Internationalization/Text.cpp
--- a/Engine/Source/Runtime/Core/Private/Internationalization/Text.cpp
+++ b/Engine/Source/Runtime/Core/Private/Internationalization/Text.cpp
@@ -1281,6 +1281,7 @@ void operator<<(FStructuredArchive::FSlot Slot, FFormatArgumentValue& Value)
 			break;
 		}
 	case EFormatArgumentType::UInt:
+	case EFormatArgumentType::Gender: // Gender is stored as a UInt
 		{
 			Record << SA_VALUE(TEXT("Value"), Value.UIntValue);
 			break;
