Adds a [Startup_<platform>] section in ConsoleVariables.ini

diff --git a/Engine/Source/Runtime/Core/Private/Misc/ConfigCacheIni.cpp b/Engine/Source/Runtime/Core/Private/Misc/ConfigCacheIni.cpp
--- a/Engine/Source/Runtime/Core/Private/Misc/ConfigCacheIni.cpp
+++ b/Engine/Source/Runtime/Core/Private/Misc/ConfigCacheIni.cpp
@@ -3973,27 +3973,34 @@ FConfigFile* FConfigCacheIni::FindOrLoadPlatformConfig(FConfigFile& LocalFile, c
 
 void FConfigCacheIni::LoadConsoleVariablesFromINI()
 {
-	FString ConsoleVariablesPath = FPaths::EngineDir() + TEXT("Config/ConsoleVariables.ini");
-
 #if !DISABLE_CHEAT_CVARS
-	// First we read from "../../../Engine/Config/ConsoleVariables.ini" [Startup] section if it exists
-	// This is the only ini file where we allow cheat commands (this is why it's not there for UE_BUILD_SHIPPING || UE_BUILD_TEST)
-	UE::ConfigUtilities::ApplyCVarSettingsFromIni(TEXT("Startup"), *ConsoleVariablesPath, ECVF_SetByConsoleVariablesIni, true);
-#endif // !DISABLE_CHEAT_CVARS
-
-#if !DISABLE_CHEAT_CVARS && !UE_BUILD_SHIPPING
 	{
-		FString OverrideConsoleVariablesPath;
-		FParse::Value(FCommandLine::Get(), TEXT("-cvarsini="), OverrideConsoleVariablesPath);
+		const TCHAR* StartupSectionName = TEXT("Startup");
+		FString PlatformName = FPlatformProperties::IniPlatformName();
+		FString StartupPlatformSectionName = FString::Printf(TEXT("Startup_%s"), *PlatformName);
+		FString ConsoleVariablesPath = FPaths::EngineDir() + TEXT("Config/ConsoleVariables.ini");
 
-		if (!OverrideConsoleVariablesPath.IsEmpty())
+		// First we read from "../../../Engine/Config/ConsoleVariables.ini" [Startup] section if it exists
+		// This is the only ini file where we allow cheat commands (this is why it's not there for UE_BUILD_SHIPPING || UE_BUILD_TEST)
+		UE::ConfigUtilities::ApplyCVarSettingsFromIni(StartupSectionName, *ConsoleVariablesPath, ECVF_SetByConsoleVariablesIni, true);
+		UE::ConfigUtilities::ApplyCVarSettingsFromIni(*StartupPlatformSectionName, *ConsoleVariablesPath, ECVF_SetByConsoleVariablesIni, true);
+
+		#if !UE_BUILD_SHIPPING
 		{
-			ensureMsgf(FPaths::FileExists(OverrideConsoleVariablesPath), TEXT("-cvarsini's file %s doesn't exist"), *OverrideConsoleVariablesPath);
-			UE::ConfigUtilities::ApplyCVarSettingsFromIni(TEXT("Startup"), *OverrideConsoleVariablesPath, ECVF_SetByConsoleVariablesIni, true);
-		}
+			FString OverrideConsoleVariablesPath;
+			FParse::Value(FCommandLine::Get(), TEXT("-cvarsini="), OverrideConsoleVariablesPath);
 
+			if (!OverrideConsoleVariablesPath.IsEmpty())
+			{
+				ensureMsgf(FPaths::FileExists(OverrideConsoleVariablesPath), TEXT("-cvarsini's file %s doesn't exist"), *OverrideConsoleVariablesPath);
+				UE::ConfigUtilities::ApplyCVarSettingsFromIni(StartupSectionName, *OverrideConsoleVariablesPath, ECVF_SetByConsoleVariablesIni, true);
+				UE::ConfigUtilities::ApplyCVarSettingsFromIni(*StartupPlatformSectionName, *OverrideConsoleVariablesPath, ECVF_SetByConsoleVariablesIni, true);
+			}
+
+		}
+		#endif
 	}
-#endif
+#endif // !DISABLE_CHEAT_CVARS
 
 	// We also apply from Engine.ini [ConsoleVariables] section
 	UE::ConfigUtilities::ApplyCVarSettingsFromIni(TEXT("ConsoleVariables"), *GEngineIni, ECVF_SetBySystemSettingsIni);
