Cooker log readability pass: Localization: Move the (possibly multiline) details from the warning message into a separate message at higher verbosity level.
#rnx

diff --git a/Engine/Source/Runtime/Core/Private/Internationalization/TextLocalizationResource.cpp b/Engine/Source/Runtime/Core/Private/Internationalization/TextLocalizationResource.cpp
--- a/Engine/Source/Runtime/Core/Private/Internationalization/TextLocalizationResource.cpp
+++ b/Engine/Source/Runtime/Core/Private/Internationalization/TextLocalizationResource.cpp
@@ -520,9 +520,9 @@ bool FTextLocalizationResource::ShouldReplaceEntry(const FTextKey& Namespace, co
 		const bool bDidConflict = CurrentEntry.SourceStringHash != NewEntry.SourceStringHash || !CurrentEntry.LocalizedString->Equals(*NewEntry.LocalizedString, ESearchCase::CaseSensitive);
 		if (bDidConflict)
 		{
-			const FString LogMsg = FString::Printf(TEXT("Text translation conflict for namespace \"%s\" and key \"%s\". The current translation is \"%s\" (from \"%s\" and source hash 0x%08x) and the conflicting translation of \"%s\" (from \"%s\" and source hash 0x%08x) will be ignored."), 
-				Namespace.GetChars(),
-				Key.GetChars(),
+			const FString SummaryMessage = FString::Printf(TEXT("Text translation conflict for namespace \"%s\" and key \"%s\"."),
+				Namespace.GetChars(), Key.GetChars());
+			const FString DetailsMessage = FString::Printf(TEXT("The current translation is \"%s\" (from \"%s\" and source hash 0x%08x) and the conflicting translation of \"%s\" (from \"%s\" and source hash 0x%08x) will be ignored."), 
 				**CurrentEntry.LocalizedString,
 				CurrentEntry.LocResID.GetChars(),
 				CurrentEntry.SourceStringHash,
@@ -531,14 +531,22 @@ bool FTextLocalizationResource::ShouldReplaceEntry(const FTextKey& Namespace, co
 				NewEntry.SourceStringHash
 				);
 
-			static const bool bLogConflictAsWarning = FParse::Param(FCommandLine::Get(), TEXT("LogLocalizationConflicts")) || !GIsBuildMachine;
-			if (bLogConflictAsWarning)
+			static const bool bLoggingConflictsRequested = FParse::Param(FCommandLine::Get(), TEXT("LogLocalizationConflicts"));
+			if (bLoggingConflictsRequested)
 			{
-				UE_LOG(LogTextLocalizationResource, Warning, TEXT("%s"), *LogMsg);
+				// Log the entire message as warning. Include the details in the warning log so they appear in stdout.
+				UE_LOG(LogTextLocalizationResource, Warning, TEXT("%s"), *(SummaryMessage + TEXT(" ") + DetailsMessage));
+			}
+			else if (!GIsBuildMachine)
+			{
+				// Log just the summary as a warning and hide the rest in the log.
+				UE_LOG(LogTextLocalizationResource, Warning, TEXT("%s"), *SummaryMessage);
+				UE_LOG(LogTextLocalizationResource, Log, TEXT("%s"), *DetailsMessage);
 			}
 			else
 			{
-				UE_LOG(LogTextLocalizationResource, Log, TEXT("%s"), *LogMsg);
+				// Hide the entire message in the log.
+				UE_LOG(LogTextLocalizationResource, Log, TEXT("%s"), *(SummaryMessage + TEXT(" ") + DetailsMessage));
 			}
 		}
 	}
