Backout StallDetector pause on SlowTask
* Some unknown cases are causing ensure() to fire - backout until sorted

diff --git a/Engine/Source/Runtime/Core/Private/ProfilingDebugging/StallDetector.cpp b/Engine/Source/Runtime/Core/Private/ProfilingDebugging/StallDetector.cpp
--- a/Engine/Source/Runtime/Core/Private/ProfilingDebugging/StallDetector.cpp
+++ b/Engine/Source/Runtime/Core/Private/ProfilingDebugging/StallDetector.cpp
@@ -368,7 +368,7 @@ void UE::FStallTimer::Resume(const double InSeconds)
 	if (--PauseCount == 0)
 	{
 		// We should be paused!
-		ensure(LastCheckSeconds == InvalidSeconds);
+		//ensure(LastCheckSeconds == InvalidSeconds);
 
 		// Resume checking from this time, we don't need atomics here as this thread is the only one that should issue this write (we were paused which means the StallDetectorThread won't Check())
 		LastCheckSeconds = InSeconds;
