Add debug commands for SlowTask and StallDetector
* StallDetector gets StallDetector.ForceLogOnStall which will cause the editor to emit a log from the LogStall category at verbosity Log even if the StallDetector scope is set to disable reporting
* Editor.Debug.SlowTask.Simulate will cause the FEngineLoop::Tick to enter a SlowTask for the given number of seconds.
* Together, these can be used to validate that the StallDetector scope in FEngineLoop::Tick is triggered when a SlowTask runs past the detector's stall threshold time (2 seconds at time of writing)

diff --git a/Engine/Source/Runtime/Core/Private/Misc/SlowTask.cpp b/Engine/Source/Runtime/Core/Private/Misc/SlowTask.cpp
--- a/Engine/Source/Runtime/Core/Private/Misc/SlowTask.cpp
+++ b/Engine/Source/Runtime/Core/Private/Misc/SlowTask.cpp
@@ -4,8 +4,11 @@
 
 #include "Containers/Array.h"
 #include "CoreTypes.h"
+#include "HAL/IConsoleManager.h"
 #include "HAL/PlatformMisc.h"
+#include "Hal/PlatformProcess.h"
 #include "HAL/PlatformTime.h"
+#include "Math/Color.h"
 #include "Math/UnrealMathUtility.h"
 #include "Misc/AssertionMacros.h"
 #include "Misc/FeedbackContext.h"
@@ -226,3 +229,40 @@ bool FSlowTask::ShouldCancel() const
 	}
 	return false;
 }
+
+#if WITH_EDITOR
+namespace Private
+{
+	static void SimulateSlowTask(const TArray<FString>& Arguments)
+	{
+		double SecondsToStall = 2.0;
+		if (Arguments.Num() >= 1)
+		{
+			LexFromString(SecondsToStall, *Arguments[0]);
+		}
+		
+		SCOPED_NAMED_EVENT_TEXT(TEXT("Simulated SlowTask"), FColor::Red);
+
+		FSlowTask SlowTask(static_cast<float>(SecondsToStall));
+		SlowTask.Initialize();
+		SlowTask.MakeDialog();
+
+		const double StartTime = FPlatformTime::Seconds();
+		while (FPlatformTime::Seconds() - StartTime < SecondsToStall)
+		{
+			// Busy wait the rest if not slept long enough
+			const float SleepTimeSeconds = 0.1;
+			SlowTask.EnterProgressFrame(SleepTimeSeconds);
+			FPlatformProcess::SleepNoStats(SleepTimeSeconds);					
+		}
+
+		SlowTask.Destroy();	
+	}
+	
+	static FAutoConsoleCommand CmdEditorSimulateSlowTask(
+		TEXT("Editor.Debug.SlowTask.Simulate"),
+		TEXT("Runs a busy loop for N seconds. Will tick the slow task every 100ms until it is complete"),
+		FConsoleCommandWithArgsDelegate::CreateStatic(&SimulateSlowTask)
+	);
+}
+#endif
\ No newline at end of file
diff --git a/Engine/Source/Runtime/Core/Private/ProfilingDebugging/StallDetector.cpp b/Engine/Source/Runtime/Core/Private/ProfilingDebugging/StallDetector.cpp
--- a/Engine/Source/Runtime/Core/Private/ProfilingDebugging/StallDetector.cpp
+++ b/Engine/Source/Runtime/Core/Private/ProfilingDebugging/StallDetector.cpp
@@ -47,6 +47,13 @@ UE::FOnStallCompleted UE::FStallDetector::StallCompleted;
 static FCriticalSection StallScopesSection;
 static TSet<UE::FStallDetector*> StallScopes;
 
+static bool ForceStallDetectedLog = false;
+static FAutoConsoleVariableRef CVarForceStallDetectedLog(
+	TEXT("StallDetector.ForceLogOnStall"),
+	ForceStallDetectedLog,
+	TEXT("Forces StallDetector to make a log to LogStall at verbosity=Log even if reporting mode is disabled"),
+	ECVF_SetByConsole);
+
 /**
 * Stall Detector Thread
 **/
@@ -201,7 +208,7 @@ uint32 UE::FStallDetectorRunnable::Run()
 						ReportSeconds = FPlatformTime::Seconds() - ReportSeconds;
 						UE_LOG(LogStall, Log, TEXT("Stall detector '%s' report submitted, and took %fs"), Stall.Name, ReportSeconds);
 					}
-					else if (Stall.ReportingMode != EStallDetectorReportingMode::Disabled)
+					else if (ForceStallDetectedLog || Stall.ReportingMode != EStallDetectorReportingMode::Disabled)
 					{
 						UE_LOG(LogStall, Log, TEXT("Stall detector '%s' exceeded budget of %fs"), Stall.Name, Stall.BudgetSeconds);
 					}
