Fix CIS on Mac/Linux platforms
* Implicit double -> float conversion of constant warning on these
* Case-sensitive includes

diff --git a/Engine/Source/Runtime/Core/Private/Misc/SlowTask.cpp b/Engine/Source/Runtime/Core/Private/Misc/SlowTask.cpp
--- a/Engine/Source/Runtime/Core/Private/Misc/SlowTask.cpp
+++ b/Engine/Source/Runtime/Core/Private/Misc/SlowTask.cpp
@@ -6,7 +6,7 @@
 #include "CoreTypes.h"
 #include "HAL/IConsoleManager.h"
 #include "HAL/PlatformMisc.h"
-#include "Hal/PlatformProcess.h"
+#include "HAL/PlatformProcess.h"
 #include "HAL/PlatformTime.h"
 #include "Math/Color.h"
 #include "Math/UnrealMathUtility.h"
@@ -251,7 +251,7 @@ namespace Private
 		while (FPlatformTime::Seconds() - StartTime < SecondsToStall)
 		{
 			// Busy wait the rest if not slept long enough
-			const float SleepTimeSeconds = 0.1;
+			const float SleepTimeSeconds = 0.1f;
 			SlowTask.EnterProgressFrame(SleepTimeSeconds);
 			FPlatformProcess::SleepNoStats(SleepTimeSeconds);					
 		}
