Fix PreSubmitTest.bat running ubsan crashing on a misaligned load

diff --git a/Engine/Source/Runtime/Core/Public/GenericPlatform/GenericPlatformMemory.h b/Engine/Source/Runtime/Core/Public/GenericPlatform/GenericPlatformMemory.h
--- a/Engine/Source/Runtime/Core/Public/GenericPlatform/GenericPlatformMemory.h
+++ b/Engine/Source/Runtime/Core/Public/GenericPlatform/GenericPlatformMemory.h
@@ -665,11 +665,13 @@ public:
 	__pragma(pack(pop))
 #else 
 	// assume it is either clang or something that supports the clang attributes
+#pragma pack(push, 1)
 	template <typename U>
-	struct __attribute__((packed)) TUnaligned
+	struct TUnaligned
 	{
 		U Value;
 	};
+#pragma pack(pop)
 #endif
 
 	/**
