Add back the gnu::alloc_align change to explain to clang the alignment of allocated pointers.

diff --git a/Engine/Plugins/Runtime/Database/SQLiteCore/Source/SQLiteCore/Private/SQLiteEmbeddedPlatform.cpp b/Engine/Plugins/Runtime/Database/SQLiteCore/Source/SQLiteCore/Private/SQLiteEmbeddedPlatform.cpp
--- a/Engine/Plugins/Runtime/Database/SQLiteCore/Source/SQLiteCore/Private/SQLiteEmbeddedPlatform.cpp
+++ b/Engine/Plugins/Runtime/Database/SQLiteCore/Source/SQLiteCore/Private/SQLiteEmbeddedPlatform.cpp
@@ -51,7 +51,7 @@ private:
 	/** Allocate memory */
 	static void* Malloc(int InSizeBytes)
 	{
-		return FMemory::Malloc(InSizeBytes, DEFAULT_ALIGNMENT);
+		return FMemory::Malloc(InSizeBytes);
 	}
 
 	/** Free memory returned by Alloc or Realloc */
@@ -63,7 +63,7 @@ private:
 	/** Reallocate memory returned by Alloc or Realloc */
 	static void* Realloc(void* InPtr, int InSizeBytes)
 	{
-		return FMemory::Realloc(InPtr, InSizeBytes, DEFAULT_ALIGNMENT);
+		return FMemory::Realloc(InPtr, InSizeBytes);
 	}
 
 	/** Get the actual size of an allocation returned by Alloc or Realloc */
@@ -75,7 +75,7 @@ private:
 	/** Roundup to the expected allocation size */
 	static int Roundup(int InSizeBytes)
 	{
-		return (int)FMemory::QuantizeSize(InSizeBytes, DEFAULT_ALIGNMENT);
+		return (int)FMemory::QuantizeSize(InSizeBytes);
 	}
 };
 
diff --git a/Engine/Source/Runtime/Core/Public/Clang/ClangPlatform.h b/Engine/Source/Runtime/Core/Public/Clang/ClangPlatform.h
--- a/Engine/Source/Runtime/Core/Public/Clang/ClangPlatform.h
+++ b/Engine/Source/Runtime/Core/Public/Clang/ClangPlatform.h
@@ -30,10 +30,6 @@
 // We use a selector macro trick to map up to 2 inputs to the correct macros.
 #define UE_ALLOCATION_FUNCTION_0() [[gnu::malloc]]
 #define UE_ALLOCATION_FUNCTION_1(SIZE) [[gnu::malloc, gnu::alloc_size(SIZE)]]
-
-// Note: we cannot use gnu::alloc_align(ALIGN) here yet, because the DEFAULT_ALIGNMENT argument is 0, which is not
-//       a power of two, and Clang will complain about uses of a function with a default alignment that is not a
-//       power of two.
-#define UE_ALLOCATION_FUNCTION_2(SIZE, ALIGN) [[gnu::malloc, gnu::alloc_size(SIZE)]]
+#define UE_ALLOCATION_FUNCTION_2(SIZE, ALIGN) [[gnu::malloc, gnu::alloc_size(SIZE), gnu::alloc_align(ALIGN)]]
 #define UE_ALLOCATION_FUNCTION_X(x, SIZE, ALIGN, FUNC, ...) FUNC
 #define UE_ALLOCATION_FUNCTION(...) UE_ALLOCATION_FUNCTION_X(,##__VA_ARGS__, UE_ALLOCATION_FUNCTION_2(__VA_ARGS__), UE_ALLOCATION_FUNCTION_1(__VA_ARGS__), UE_ALLOCATION_FUNCTION_0(__VA_ARGS__)) 
diff --git a/Engine/Source/Runtime/Core/Public/HAL/UnrealMemory.h b/Engine/Source/Runtime/Core/Public/HAL/UnrealMemory.h
--- a/Engine/Source/Runtime/Core/Public/HAL/UnrealMemory.h
+++ b/Engine/Source/Runtime/Core/Public/HAL/UnrealMemory.h
@@ -207,18 +207,47 @@ struct FMemory
 	// C style memory allocation stubs.
 	//
 
-	UE_ALLOCATION_FUNCTION(1, 2) static CORE_API void* Malloc(SIZE_T Count, uint32 Alignment = DEFAULT_ALIGNMENT);
-	UE_ALLOCATION_FUNCTION(2, 3) static CORE_API void* Realloc(void* Original, SIZE_T Count, uint32 Alignment = DEFAULT_ALIGNMENT);
+	UE_ALLOCATION_FUNCTION(1, 2) static CORE_API void* Malloc(SIZE_T Count, uint32 Alignment);
+	UE_ALLOCATION_FUNCTION(2, 3) static CORE_API void* Realloc(void* Original, SIZE_T Count, uint32 Alignment);
 	static CORE_API void Free(void* Original);
 	static CORE_API SIZE_T GetAllocSize(void* Original);
 
-	UE_ALLOCATION_FUNCTION(1, 2) static FORCEINLINE_DEBUGGABLE void* MallocZeroed(SIZE_T Count, uint32 Alignment = DEFAULT_ALIGNMENT)
+	UE_ALLOCATION_FUNCTION(1, 2) static FORCEINLINE_DEBUGGABLE void* MallocZeroed(SIZE_T Count, uint32 Alignment)
 	{
 		void* Memory = Malloc(Count, Alignment);
 		Memzero(Memory, Count);
 		return Memory;
 	}
 
+#if defined(__clang__)
+	// When recording alignment information to aid the LLVM optimizer, clang will complain that DEFAULT_ALIGNMENT
+	// (which is zero) is not a power-of-two. But clang will always ignore this 'bad' non power-of-two alignment
+	// value and revert to the natural alignment, which is what we'd want in this case anyway. So to let us keep
+	// the semantics where we use DEFAULT_ALIGNMENT == 0 to effectively mean an optional alignment, we just tell
+	// clang to ignore the 'bad' alignment in this case.
+#pragma clang diagnostic push
+#pragma clang diagnostic ignored "-Wnon-power-of-two-alignment"
+#endif
+
+	static FORCEINLINE void* Malloc(SIZE_T Count)
+	{
+		return Malloc(Count, DEFAULT_ALIGNMENT);
+	}
+
+	static FORCEINLINE void* Realloc(void* Original, SIZE_T Count)
+	{
+		return Realloc(Original, Count, DEFAULT_ALIGNMENT);
+	}
+
+	static FORCEINLINE_DEBUGGABLE void* MallocZeroed(SIZE_T Count)
+	{
+		return MallocZeroed(Count, DEFAULT_ALIGNMENT);
+	}
+
+#if defined(__clang__)
+#pragma clang diagnostic pop
+#endif
+
 	/**
 	* For some allocators this will return the actual size that should be requested to eliminate
 	* internal fragmentation. The return value will always be >= Count. This can be used to grow
diff --git a/Engine/Source/Runtime/CoreUObject/Private/Serialization/BulkData.cpp b/Engine/Source/Runtime/CoreUObject/Private/Serialization/BulkData.cpp
--- a/Engine/Source/Runtime/CoreUObject/Private/Serialization/BulkData.cpp
+++ b/Engine/Source/Runtime/CoreUObject/Private/Serialization/BulkData.cpp
@@ -316,11 +316,11 @@ inline void* SafeMalloc(int64 SizeInBytes, uint32 Alignment)
 inline void* SafeRealloc(void* Original, int64 SizeInBytes, uint32 Alignment)
 {
 #if !PLATFORM_32BITS
-	return FMemory::Realloc(Original, SizeInBytes, DEFAULT_ALIGNMENT);
+	return FMemory::Realloc(Original, SizeInBytes);
 #else
 	if (SizeInBytes >= 0 && (uint64)SizeInBytes <= TNumericLimits<SIZE_T>::Max())
 	{
-		return FMemory::Realloc(Original, SizeInBytes, DEFAULT_ALIGNMENT);
+		return FMemory::Realloc(Original, SizeInBytes);
 	}
 	else
 	{
diff --git a/Engine/Source/Runtime/CoreUObject/Public/Serialization/EditorBulkDataWriter.h b/Engine/Source/Runtime/CoreUObject/Public/Serialization/EditorBulkDataWriter.h
--- a/Engine/Source/Runtime/CoreUObject/Public/Serialization/EditorBulkDataWriter.h
+++ b/Engine/Source/Runtime/CoreUObject/Public/Serialization/EditorBulkDataWriter.h
@@ -26,7 +26,7 @@ public:
 
 			// Clone the payload so that we have a local copy that we can
 			// append additional data to.
-			Buffer = FMemory::Malloc(CurrentDataLength, DEFAULT_ALIGNMENT);
+			Buffer = FMemory::Malloc(CurrentDataLength);
 			FMemory::Memcpy(Buffer, Payload.GetData(), CurrentDataLength);
 
 			BufferLength = CurrentDataLength;
@@ -48,7 +48,7 @@ public:
 	~FEditorBulkDataWriter()
 	{
 		// Remove the slack from the allocated bulk data
-		Buffer = FMemory::Realloc(Buffer, DataLength, DEFAULT_ALIGNMENT);
+		Buffer = FMemory::Realloc(Buffer, DataLength);
 		BulkData.UpdatePayload(FSharedBuffer::TakeOwnership(Buffer, DataLength, FMemory::Free));
 	}
 
@@ -68,7 +68,7 @@ public:
 		{
 			// If so, resize to the new size + 3/8 additional slack
 			const int64 NewLength = NewPos + 3 * NewPos / 8 + 16;
-			Buffer = FMemory::Realloc(Buffer, NewLength, DEFAULT_ALIGNMENT);
+			Buffer = FMemory::Realloc(Buffer, NewLength);
 			BufferLength = NewLength;
 		}
 
diff --git a/Engine/Source/Runtime/Engine/Public/Animation/AnimCompressionTypes.h b/Engine/Source/Runtime/Engine/Public/Animation/AnimCompressionTypes.h
--- a/Engine/Source/Runtime/Engine/Public/Animation/AnimCompressionTypes.h
+++ b/Engine/Source/Runtime/Engine/Public/Animation/AnimCompressionTypes.h
@@ -611,7 +611,17 @@ public:
 
 				// Avoid calling FMemory::Realloc( nullptr, 0 ) as ANSI C mandates returning a valid pointer which is not what we want.
 				//checkSlow(((uint64)NumElements*(uint64)ElementTypeInfo.GetSize() < (uint64)INT_MAX));
-				Data = (FScriptContainerElement*)FMemory::Realloc(Data, NumElements*NumBytesPerElement, Alignment);
+
+				const SIZE_T NewSize = NumElements * NumBytesPerElement;
+
+				if constexpr (Alignment == DEFAULT_ALIGNMENT)
+				{
+					Data = (FScriptContainerElement*)FMemory::Realloc(Data, NewSize);
+				}
+				else
+				{
+					Data = (FScriptContainerElement*)FMemory::Realloc(Data, NewSize, Alignment);
+				}
 			}
 		}
 		SizeType CalculateSlackReserve(SizeType NumElements, SIZE_T NumBytesPerElement) const
