Pause StallDetector when SlowTask is present
* A SlowTask is created when there is delibrately going to be a stall-like condition but the developer wants to retain some limited interactivity.  This is traditionally not considered a stall.
* This change will pause stall detectors when slow tasks are created and resume them when the slow task ends.
* Implementation uses a TLS variable to store SlowTask specific pause contexts so as to not put a dependency of SlowTasks on the StallDetector.  Thread local stack is used since SlowTasks can be created from any thread context and can also be recursively created - therefore the Start/Finalize SlowTask events can be fired from any thread and the number of pause contexts need to match the created slow tasks.

diff --git a/Engine/Source/Runtime/Core/Private/ProfilingDebugging/StallDetector.cpp b/Engine/Source/Runtime/Core/Private/ProfilingDebugging/StallDetector.cpp
--- a/Engine/Source/Runtime/Core/Private/ProfilingDebugging/StallDetector.cpp
+++ b/Engine/Source/Runtime/Core/Private/ProfilingDebugging/StallDetector.cpp
@@ -415,6 +415,18 @@ static ThreadLocalSlotAcquire ThreadLocalSlot;
 
 TAtomic<uint64> UE::FStallDetector::UIDGenerator = 0;
 
+/**
+ * Pause on slow task
+ */
+namespace UE::StallDetector::Private
+{
+	FDelegateHandle GSlowTaskStartDelegate;
+	FDelegateHandle GSlowTaskFinalizedDelegate;
+
+	void RegisterSlowTaskHandlers();
+	void UnregisterSlowTaskHandlers();
+}
+
 UE::FStallDetector::FStallDetector(FStallDetectorStats& InStats)
 	: Stats(InStats)
 	, Parent(nullptr)
@@ -618,6 +630,8 @@ void UE::FStallDetector::Startup()
 				{
 					FPlatformProcess::YieldThread();
 				}
+
+				StallDetector::Private::RegisterSlowTaskHandlers();
 			}
 		}
 
@@ -631,6 +645,8 @@ void UE::FStallDetector::Shutdown()
 	{
 		UE_LOG(LogStall, Log, TEXT("Shutdown..."));
 
+		StallDetector::Private::UnregisterSlowTaskHandlers();
+
 		delete StallDetectorThread;
 		StallDetectorThread = nullptr;
 
@@ -683,4 +699,69 @@ UE::FStallDetectorPause::~FStallDetectorPause()
 	}
 }
 
+bool UE::FStallDetectorPause::IsPaused() const
+{
+	return bPaused;
+}
+
+
+/**
+ * Pause on slow task
+ */
+namespace UE::StallDetector::Private
+{
+	static ThreadLocalSlotAcquire SlowTaskPauseSlot;
+
+	/**
+	 * A StallDetector pause context which handles cases where SlowTask is invoked
+	 * Each pause context forms a node in a linked list - the tail is held in the SlowTaskPauseSlot TLS slot
+	 */
+	class FSlowTaskStallDetectorPause : public FStallDetectorPause
+	{
+	public:
+		explicit FSlowTaskStallDetectorPause(FSlowTaskStallDetectorPause* InParent);
+		FSlowTaskStallDetectorPause* GetParent() const;
+	private:
+		FSlowTaskStallDetectorPause* Parent;
+	};
+	
+	FSlowTaskStallDetectorPause::FSlowTaskStallDetectorPause(FSlowTaskStallDetectorPause* InParent)
+		: FStallDetectorPause()
+		, Parent(InParent)
+	{
+	}
+
+	FSlowTaskStallDetectorPause* FSlowTaskStallDetectorPause::GetParent() const
+	{
+		return Parent;
+	}
+
+	static void HandleSlowTaskStart(const FText& TaskName)
+	{
+		FSlowTaskStallDetectorPause* Parent = static_cast<FSlowTaskStallDetectorPause*>(FPlatformTLS::GetTlsValue(SlowTaskPauseSlot.GetSlot()));
+		FSlowTaskStallDetectorPause* PauseState = new FSlowTaskStallDetectorPause(Parent);
+		FPlatformTLS::SetTlsValue(SlowTaskPauseSlot.GetSlot(), PauseState);
+	}
+
+	static void HandleSlowTaskFinalize(const FText& TaskName, double DurationInSeconds)
+	{
+		FSlowTaskStallDetectorPause* PauseState = static_cast<FSlowTaskStallDetectorPause*>(FPlatformTLS::GetTlsValue(SlowTaskPauseSlot.GetSlot()));
+		FSlowTaskStallDetectorPause* Parent = PauseState->GetParent();
+		FPlatformTLS::SetTlsValue(SlowTaskPauseSlot.GetSlot(), Parent);
+		delete PauseState;
+	}
+	
+	void RegisterSlowTaskHandlers()
+	{
+		GSlowTaskStartDelegate = GWarn->OnStartSlowTask().AddStatic(&HandleSlowTaskStart);
+		GSlowTaskFinalizedDelegate = GWarn->OnFinalizeSlowTask().AddStatic(&HandleSlowTaskFinalize);
+	}
+	
+	void UnregisterSlowTaskHandlers()
+	{
+		GWarn->OnFinalizeSlowTask().Remove(GSlowTaskFinalizedDelegate);
+		GWarn->OnStartSlowTask().Remove(GSlowTaskStartDelegate);
+	}
+}
+
 #endif // STALL_DETECTOR
\ No newline at end of file
diff --git a/Engine/Source/Runtime/Core/Public/ProfilingDebugging/StallDetector.h b/Engine/Source/Runtime/Core/Public/ProfilingDebugging/StallDetector.h
--- a/Engine/Source/Runtime/Core/Public/ProfilingDebugging/StallDetector.h
+++ b/Engine/Source/Runtime/Core/Public/ProfilingDebugging/StallDetector.h
@@ -368,6 +368,7 @@ namespace UE
 	public:
 		CORE_API FStallDetectorPause();
 		CORE_API ~FStallDetectorPause();
+		bool IsPaused() const;
 
 	private:
 		bool bPaused;
