[PerfFix] Removed 10ms sleep in WindowsPlatformNamedPipe which is running every frame while CmdLinkServer is active

diff --git a/Engine/Source/Runtime/Core/Private/Windows/WindowsPlatformNamedPipe.cpp b/Engine/Source/Runtime/Core/Private/Windows/WindowsPlatformNamedPipe.cpp
--- a/Engine/Source/Runtime/Core/Private/Windows/WindowsPlatformNamedPipe.cpp
+++ b/Engine/Source/Runtime/Core/Private/Windows/WindowsPlatformNamedPipe.cpp
@@ -306,8 +306,6 @@ bool FWindowsPlatformNamedPipe::UpdateAsyncStatus()
 					switch (LastError)
 					{
 						case ERROR_IO_INCOMPLETE:
-							// Yield CPU time while waiting
-							FPlatformProcess::Sleep(.01f);
 							break;
 
 						case ERROR_BROKEN_PIPE:
