Fix the screen resolution (on mobile) stretched issue when casted to a larger screen with a USBC to HDMI cable.

#b Chris.Babcock, Allan.Bentham

diff --git a/Engine/Build/Android/Java/src/com/epicgames/unreal/GameActivity.java.template b/Engine/Build/Android/Java/src/com/epicgames/unreal/GameActivity.java.template
--- a/Engine/Build/Android/Java/src/com/epicgames/unreal/GameActivity.java.template
+++ b/Engine/Build/Android/Java/src/com/epicgames/unreal/GameActivity.java.template
@@ -594,9 +594,9 @@ public class GameActivity extends $${gameActivitySuperClass}$$ implements Surfac
 	{
 		// onConfigurationChanged won't be called if orientation changes from landscape/portrait to reverse landscape/portrait
 		// use this function instead of caching DeviceRotation
-		int orientation = AndroidThunkJava_GetDeviceOrientation();
+		int rotation = AndroidThunkJava_GetDeviceRotation();
 
-		switch (orientation)
+		switch (rotation)
 		{
 			case Surface.ROTATION_0:	return 0;
 			case Surface.ROTATION_90:	return 90;
@@ -2940,8 +2940,7 @@ public class GameActivity extends $${gameActivitySuperClass}$$ implements Surfac
 		nativeSetGlobalActivity(UseExternalFilesDir, PublicLogFiles, getFilesDir().getPath(), getExternalFilesDir(null).getPath(), PackageDataInsideApkValue==1, APKPath);
 
 		// tell the engine if this is a portrait app
-		int orientation = AndroidThunkJava_GetDeviceOrientation();
-		nativeSetWindowInfo(orientation, DepthBufferPreference, PropagateAlpha);
+		nativeSetWindowInfo(getResources().getConfiguration().orientation, DepthBufferPreference, PropagateAlpha);
 
 		// get the full language code, like en-US
 		// note: this may need to be Locale.getDefault().getLanguage()
@@ -3384,8 +3383,7 @@ public class GameActivity extends $${gameActivitySuperClass}$$ implements Surfac
 		}
 
 		// invalidate window cache
-		int orientation = AndroidThunkJava_GetDeviceOrientation();
-		nativeSetWindowInfo(orientation, DepthBufferPreference, PropagateAlpha);
+		nativeSetWindowInfo(getResources().getConfiguration().orientation, DepthBufferPreference, PropagateAlpha);
 
 		// only do this on KitKat and above
 		if (ShouldHideUI)
@@ -3880,8 +3878,7 @@ public class GameActivity extends $${gameActivitySuperClass}$$ implements Surfac
 		super.onConfigurationChanged(newConfig);
 		//$${gameActivityonConfigurationChangedAdditions}$$
 		// forward the orientation
-		int orientation = AndroidThunkJava_GetDeviceOrientation();
-		nativeOnConfigurationChanged(orientation);
+		nativeOnConfigurationChanged(newConfig.orientation);
 	}
 
 	public void surfaceChanged(SurfaceHolder holder, int format, int width, int height)
@@ -6786,11 +6783,16 @@ public class GameActivity extends $${gameActivitySuperClass}$$ implements Surfac
 		return "";
 	}
 
-	public int AndroidThunkJava_GetDeviceOrientation() 
+	public int AndroidThunkJava_GetDeviceRotation() 
 	{
 		return getWindowManager().getDefaultDisplay().getRotation();
 	}
 
+	public int AndroidThunkJava_GetConfigurationOrientation()
+	{
+		return getResources().getConfiguration().orientation;
+	}
+
 	@SuppressWarnings("deprecation")
 	public boolean AndroidThunkJava_CookieManager_SetCookie(String Url, String Value)
 	{
diff --git a/Engine/Source/Runtime/ApplicationCore/Private/Android/AndroidWindow.cpp b/Engine/Source/Runtime/ApplicationCore/Private/Android/AndroidWindow.cpp
--- a/Engine/Source/Runtime/ApplicationCore/Private/Android/AndroidWindow.cpp
+++ b/Engine/Source/Runtime/ApplicationCore/Private/Android/AndroidWindow.cpp
@@ -139,7 +139,7 @@ JNI_METHOD void Java_com_epicgames_unreal_GameActivity_nativeSetWindowInfo(JNIEn
 {
 	ClearCachedWindowRects();
 	GWindowOrientation = orientation;
-	bool bIsPortrait = GWindowOrientation == EAndroidSurfaceOrientation::ROTATION_0 || GWindowOrientation == EAndroidSurfaceOrientation::ROTATION_180;
+	bool bIsPortrait = GWindowOrientation == EAndroidConfigurationOrientation::ORIENTATION_PORTRAIT;
 	GAndroidIsPortrait = bIsPortrait == JNI_TRUE;
 	GAndroidDepthBufferPreference = DepthBufferPreference;
 	GAndroidPropagateAlpha = PropagateAlpha;
@@ -172,6 +172,7 @@ JNI_METHOD void Java_com_epicgames_unreal_GameActivity_nativeSetSafezoneInfo(JNI
 	}
 
 	GAndroidSafezoneRequiresUpdate = true;
+	UE_LOG(LogAndroid, Log, TEXT("nativeSetSafezoneInfo bIsPortrait=%d, left=%f, top=%f, right=%f, bottom=%f"), bIsPortrait ? 1 : 0, left, top, right, bottom);
 }
 #endif
 
@@ -565,7 +566,7 @@ bool FAndroidWindow::OnWindowOrientationChanged(int Orientation)
 	if (GWindowOrientation != Orientation)
 	{
 		GWindowOrientation = Orientation;
-		bool bIsPortrait = GWindowOrientation == EAndroidSurfaceOrientation::ROTATION_0 || GWindowOrientation == EAndroidSurfaceOrientation::ROTATION_180;
+		bool bIsPortrait = GWindowOrientation == EAndroidConfigurationOrientation::ORIENTATION_PORTRAIT;
 		UE_LOG(LogAndroid, Log, TEXT("Window orientation changed: %s"), bIsPortrait ? TEXT("Portrait") : TEXT("Landscape"));
 		GAndroidIsPortrait = bIsPortrait;
 		return true;
diff --git a/Engine/Source/Runtime/Core/Private/Android/AndroidPlatformMisc.cpp b/Engine/Source/Runtime/Core/Private/Android/AndroidPlatformMisc.cpp
--- a/Engine/Source/Runtime/Core/Private/Android/AndroidPlatformMisc.cpp
+++ b/Engine/Source/Runtime/Core/Private/Android/AndroidPlatformMisc.cpp
@@ -618,27 +618,43 @@ void FAndroidMisc::UpdateDeviceOrientation()
 	JNIEnv* JEnv = AndroidJavaEnv::GetJavaEnv();
 	if (JEnv)
 	{
+		static jmethodID getRotationMethod = 0;
 		static jmethodID getOrientationMethod = 0;
 
-		if (getOrientationMethod == 0)
+		if (getRotationMethod == 0 || getOrientationMethod == 0)
 		{
 			jclass MainClass = AndroidJavaEnv::FindJavaClassGlobalRef("com/epicgames/unreal/GameActivity");
 			if (MainClass != nullptr)
 			{
-				getOrientationMethod = JEnv->GetMethodID(MainClass, "AndroidThunkJava_GetDeviceOrientation", "()I");
+				getRotationMethod = JEnv->GetMethodID(MainClass, "AndroidThunkJava_GetDeviceRotation", "()I");
+				getOrientationMethod = JEnv->GetMethodID(MainClass, "AndroidThunkJava_GetConfigurationOrientation", "()I");
 				JEnv->DeleteGlobalRef(MainClass);
 			}
 		}
 
-		if (getOrientationMethod != 0)
+		if (getRotationMethod != 0 && getOrientationMethod != 0)
 		{
+			const int Rotation = JEnv->CallIntMethod(AndroidJavaEnv::GetGameActivityThis(), getRotationMethod);
 			const int Orientation = JEnv->CallIntMethod(AndroidJavaEnv::GetGameActivityThis(), getOrientationMethod);
-			switch (Orientation)
+			if (Orientation == EAndroidConfigurationOrientation::ORIENTATION_PORTRAIT)
 			{
-			case EAndroidSurfaceOrientation::ROTATION_0:	DeviceOrientation = EDeviceScreenOrientation::Portrait;             break;
-			case EAndroidSurfaceOrientation::ROTATION_90:	DeviceOrientation = EDeviceScreenOrientation::LandscapeLeft;        break;
-			case EAndroidSurfaceOrientation::ROTATION_180:	DeviceOrientation = EDeviceScreenOrientation::PortraitUpsideDown;   break;
-			case EAndroidSurfaceOrientation::ROTATION_270:	DeviceOrientation = EDeviceScreenOrientation::LandscapeRight;       break;
+				switch (Rotation)
+				{
+				case EAndroidSurfaceRotation::ROTATION_0:	DeviceOrientation = EDeviceScreenOrientation::Portrait;             break;
+				case EAndroidSurfaceRotation::ROTATION_90:	DeviceOrientation = EDeviceScreenOrientation::LandscapeLeft;        break;
+				case EAndroidSurfaceRotation::ROTATION_180:	DeviceOrientation = EDeviceScreenOrientation::PortraitUpsideDown;   break;
+				case EAndroidSurfaceRotation::ROTATION_270:	DeviceOrientation = EDeviceScreenOrientation::LandscapeRight;       break;
+				}
+			}
+			else if (Orientation == EAndroidConfigurationOrientation::ORIENTATION_LANDSCAPE)
+			{
+				switch (Rotation)
+				{
+				case EAndroidSurfaceRotation::ROTATION_0:	DeviceOrientation = EDeviceScreenOrientation::LandscapeLeft;        break;
+				case EAndroidSurfaceRotation::ROTATION_90:	DeviceOrientation = EDeviceScreenOrientation::PortraitUpsideDown;				break;
+				case EAndroidSurfaceRotation::ROTATION_180:	DeviceOrientation = EDeviceScreenOrientation::LandscapeRight;		break;
+				case EAndroidSurfaceRotation::ROTATION_270:	DeviceOrientation = EDeviceScreenOrientation::Portrait;   break;
+				}
 			}
 		}
 	}
diff --git a/Engine/Source/Runtime/Core/Public/Android/AndroidPlatformMisc.h b/Engine/Source/Runtime/Core/Public/Android/AndroidPlatformMisc.h
--- a/Engine/Source/Runtime/Core/Public/Android/AndroidPlatformMisc.h
+++ b/Engine/Source/Runtime/Core/Public/Android/AndroidPlatformMisc.h
@@ -23,8 +23,18 @@ class TFunction;
 
 enum class ECrashContextType;
 
-//The android surface orientations defined in C++ map to the values in https://developer.android.com/reference/android/view/Surface#ROTATION_0
-namespace EAndroidSurfaceOrientation
+//The android configuration orientations defined in C++ map to the values in https://developer.android.com/reference/android/content/res/Configuration#ORIENTATION_PORTRAIT
+namespace EAndroidConfigurationOrientation
+{
+	enum Type
+	{
+		ORIENTATION_PORTRAIT = 0x00000001,
+		ORIENTATION_LANDSCAPE = 0x00000002,
+	};
+};
+
+//The android surface rotations defined in C++ map to the values in https://developer.android.com/reference/android/view/Surface#ROTATION_0
+namespace EAndroidSurfaceRotation
 {
 	enum Type
 	{
