Disable launching of UnrealTraceServer from the Editor on Mac.

diff --git a/Engine/Source/Runtime/Core/Private/ProfilingDebugging/TraceAuxiliary.cpp b/Engine/Source/Runtime/Core/Private/ProfilingDebugging/TraceAuxiliary.cpp
--- a/Engine/Source/Runtime/Core/Private/ProfilingDebugging/TraceAuxiliary.cpp
+++ b/Engine/Source/Runtime/Core/Private/ProfilingDebugging/TraceAuxiliary.cpp
@@ -16,7 +16,7 @@
 #endif
 
 #if !defined(WITH_UNREAL_TRACE_LAUNCH)
-#	define WITH_UNREAL_TRACE_LAUNCH (PLATFORM_DESKTOP && !UE_BUILD_SHIPPING && !IS_PROGRAM)
+#	define WITH_UNREAL_TRACE_LAUNCH (PLATFORM_DESKTOP && !UE_BUILD_SHIPPING && !IS_PROGRAM && !PLATFORM_MAC)
 #endif
 
 #if !defined(UE_TRACE_AUTOSTART)
