Tasks System: switched from YieldThread() to Yield() for short spinning on waiting for task completion, as it was causing long preemption and resulted with poor performance on RHIT in particular

diff --git a/Engine/Source/Runtime/Core/Private/Tasks/TaskPrivate.cpp b/Engine/Source/Runtime/Core/Private/Tasks/TaskPrivate.cpp
--- a/Engine/Source/Runtime/Core/Private/Tasks/TaskPrivate.cpp
+++ b/Engine/Source/Runtime/Core/Private/Tasks/TaskPrivate.cpp
@@ -204,7 +204,7 @@ namespace UE::Tasks
 			const uint32 MaxSpinCount = 40;
 			for (uint32 SpinCount = 0; SpinCount != MaxSpinCount && !IsCompleted() && !Timeout; ++SpinCount)
 			{
-				FPlatformProcess::YieldThread();
+				FPlatformProcess::Yield(); // YieldThread() was much slower on some platforms with low core count and contention for CPU
 			}
 
 			if (IsCompleted() || Timeout)
