cosmetic, comment correction

diff --git a/Engine/Source/Runtime/Core/Public/Tasks/TaskPrivate.h b/Engine/Source/Runtime/Core/Public/Tasks/TaskPrivate.h
--- a/Engine/Source/Runtime/Core/Public/Tasks/TaskPrivate.h
+++ b/Engine/Source/Runtime/Core/Public/Tasks/TaskPrivate.h
@@ -207,7 +207,7 @@ namespace UE::Tasks
 				// prerequisite can be added successfully, and release the lock if it wasn't
 				uint32 PrevNumLocks = NumLocks.fetch_add(1, std::memory_order_relaxed); // relaxed because the following
 				// `AddSubsequent` provides required sync
-				checkf(PrevNumLocks + 1 < ExecutionFlag, TEXT("Max number of nested tasks reached: %d"), ExecutionFlag);
+				checkf(PrevNumLocks + 1 < ExecutionFlag, TEXT("Max number of task prerequisites reached: %d"), ExecutionFlag);
 
 				if (!Prerequisite.AddSubsequent(*this)) // linearisation point, acq_rel semantic
 				{
