bigger format buffer for `ensure` to avoid trancating long messages

diff --git a/Engine/Source/Runtime/Core/Private/Misc/AssertionMacros.cpp b/Engine/Source/Runtime/Core/Private/Misc/AssertionMacros.cpp
--- a/Engine/Source/Runtime/Core/Private/Misc/AssertionMacros.cpp
+++ b/Engine/Source/Runtime/Core/Private/Misc/AssertionMacros.cpp
@@ -637,15 +637,20 @@ FORCENOINLINE bool VARARGS FDebug::OptionallyLogFormattedEnsureMessageReturningF
 	return false;
 }
 
+namespace AssertionMacros_Private
+{
+	const int32 FormatBufferSize = 65535;
+	TCHAR FormatBuffer[FormatBufferSize];
+	UE::FWordMutex FormatMutex;
+}
+
 FORCENOINLINE bool FDebug::OptionallyLogFormattedEnsureMessageReturningFalseImpl(bool bLog, const ANSICHAR* Expr, const ANSICHAR* File, int32 Line, void* ProgramCounter, const TCHAR* FormattedMsg, va_list Args)
 {
 	if (bLog)
 	{
-		const int32 TempStrSize = 4096;
-		TCHAR TempStr[TempStrSize];
-		FCString::GetVarArgs(TempStr, TempStrSize, FormattedMsg, Args);
-
-		EnsureFailed(Expr, File, Line, ProgramCounter, TempStr);
+		UE::TUniqueLock Lock(AssertionMacros_Private::FormatMutex);
+		FCString::GetVarArgs(AssertionMacros_Private::FormatBuffer, AssertionMacros_Private::FormatBufferSize, FormattedMsg, Args);
+		EnsureFailed(Expr, File, Line, ProgramCounter, AssertionMacros_Private::FormatBuffer);
 	}
 
 	return false;
