Added support to track previously untracked NVN memory.
This is the bulk of untracked memory on the Switch.

  - NVN allocates memory much the same way as the MallocBinned series do.
    - So we need to treat the memory like we do for FMalloc.
    - Added to LLM ELLMTag::RHIUnused and ELLMAllocType::RHI.
    - Both trackers (Platform and Default), use the new ELLMAllocType::RHI to track total allocations for NVN.
    - Both still use ELLMTags so that they still be tracked in more fine details.
    - Now ELLMTag::RHIUnused is used to track the difference between ELLMAllocType::RHI totals.
  - Removed sizeof(FLowLevelMemTracker) from LLM overhead because it's not dynamically allocated.
  - Removed forcing negative values for Untracked to zero. It only hides possible errors.

MERGING THIS ON BEHALF OF [REDACTED_HANDLE] while he's out on vacation.
Swarm reviews / approvals are on 26407978

#tests Built all Win64 targets and all Switch clients. Did a local ReplayRun with LLM. Horde Minimal Code Only Build with Perf and LLM.

diff --git a/Engine/Source/Runtime/Core/Private/HAL/LowLevelMemTracker.cpp b/Engine/Source/Runtime/Core/Private/HAL/LowLevelMemTracker.cpp
--- a/Engine/Source/Runtime/Core/Private/HAL/LowLevelMemTracker.cpp
+++ b/Engine/Source/Runtime/Core/Private/HAL/LowLevelMemTracker.cpp
@@ -206,6 +206,7 @@ DECLARE_LLM_MEMORY_STAT(TEXT("Overhead"), STAT_PlatformOverheadLLM, STATGROUP_LL
 DECLARE_LLM_MEMORY_STAT(TEXT("OS Available"), STAT_PlatformOSAvailableLLM, STATGROUP_LLMPlatform);
 DECLARE_LLM_MEMORY_STAT(TEXT("FMalloc"), STAT_FMallocLLM, STATGROUP_LLMPlatform);
 DECLARE_LLM_MEMORY_STAT(TEXT("FMalloc Unused"), STAT_FMallocUnusedLLM, STATGROUP_LLMFULL);
+DECLARE_LLM_MEMORY_STAT(TEXT("RHI Unused"), STAT_RHIUnusedLLM, STATGROUP_LLMFULL);
 DECLARE_LLM_MEMORY_STAT(TEXT("ThreadStack"), STAT_ThreadStackLLM, STATGROUP_LLMFULL);
 DECLARE_LLM_MEMORY_STAT(TEXT("ThreadStackPlatform"), STAT_ThreadStackPlatformLLM, STATGROUP_LLMPlatform);
 DECLARE_LLM_MEMORY_STAT(TEXT("Program Size"), STAT_ProgramSizePlatformLLM, STATGROUP_LLMPlatform);
@@ -956,7 +957,28 @@ void FLowLevelMemTracker::TickInternal()
 	FLLMTracker& DefaultTracker = *Trackers[static_cast<int32>(ELLMTracker::Default)];
 	FLLMTracker& PlatformTracker = *Trackers[static_cast<int32>(ELLMTracker::Platform)];
 
-	// calculate FMalloc unused stat and set it in the Default tracker
+	// Cache the amount of memory used early, since some of these functions (FindOrAddTagData) can
+	// cause allocations, which will throw the numbers off slightly.
+	FPlatformMemoryStats PlatformStats = FPlatformMemory::GetStats();
+#if PLATFORM_DESKTOP
+	// virtual is working set + paged out memory.
+	const int64 PlatformProcessMemory = static_cast<int64>(PlatformStats.UsedVirtual);
+#elif PLATFORM_ANDROID
+	// On some Android devices total used mem (VmRss + VmSwap) does not include GPU memory allocations
+	// and so it's going to be lower than TrackedTotal
+	const int64 PlatformProcessMemory = FMath::Max((int64)PlatformStats.UsedPhysical, TrackedTotal);
+#elif  PLATFORM_IOS || UE_SERVER
+	const int64 PlatformProcessMemory = static_cast<int64>(PlatformStats.UsedPhysical);
+#else
+	const int64 PlatformProcessMemory = static_cast<int64>(PlatformStats.TotalPhysical) -
+		static_cast<int64>(PlatformStats.AvailablePhysical);
+#endif
+	// Update how much overhead LLM is using. Note we used to also added sizeof(FLowLevelMemTracker), but this
+	// isn't allocated and is in the data segment. So adding it throws the numbers off.
+	MemoryUsageCurrentOverhead = Allocator.GetTotal();
+	PlatformTracker.SetTagAmountInUpdate(FindOrAddTagData(ELLMTag::PlatformOverhead), MemoryUsageCurrentOverhead, !FPlatformMemory::TracksLLMAllocations());
+
+	// Calculate FMalloc unused stat and set it in the Default tracker.
 	const int64 FMallocAmount = DefaultTracker.GetAllocTypeAmount(ELLMAllocType::FMalloc);
 	const int64 FMallocPlatformAmount = PlatformTracker.GetTagAmount(FindOrAddTagData(ELLMTag::FMalloc));
 	int64 FMallocUnused = FMallocPlatformAmount - FMallocAmount;
@@ -968,34 +990,20 @@ void FLowLevelMemTracker::TickInternal()
 	}
 	DefaultTracker.SetTagAmountInUpdate(FindOrAddTagData(ELLMTag::FMallocUnused), FMallocUnused, true);
 
-	const int64 StaticOverhead = sizeof(FLowLevelMemTracker);
-	MemoryUsageCurrentOverhead = StaticOverhead + Allocator.GetTotal();
-	PlatformTracker.SetTagAmountInUpdate(FindOrAddTagData(ELLMTag::PlatformOverhead), MemoryUsageCurrentOverhead, !FPlatformMemory::TracksLLMAllocations());
+	// Determine the UnusedRHI amount by finding the difference between the Platform and Default tracker.
+	int64 PlatformRHIAmount = PlatformTracker.GetAllocTypeAmount(ELLMAllocType::RHI);
+	int64 DefaultRHIAmount = DefaultTracker.GetAllocTypeAmount(ELLMAllocType::RHI);
+	int64 UnusedRHIAmount = PlatformRHIAmount - DefaultRHIAmount;
+	DefaultTracker.SetTagAmountInUpdate(FindOrAddTagData(ELLMTag::RHIUnused), UnusedRHIAmount, true);
 
 	const int64 TrackedTotal = DefaultTracker.GetTrackedTotal();
 	// Compare memory the platform thinks we have allocated to what we have tracked, including the program memory
-	const FPlatformMemoryStats PlatformStats = FPlatformMemory::GetStats();
-#if PLATFORM_DESKTOP
-	// virtual is working set + paged out memory.
-	const int64 PlatformProcessMemory = static_cast<int64>(PlatformStats.UsedVirtual);
-#elif PLATFORM_ANDROID
-	// On some Android devices total used mem (VmRss + VmSwap) does not include GPU memory allocations
-	// and so it's going to be lower than TrackedTotal
-	const int64 PlatformProcessMemory = FMath::Max((int64)PlatformStats.UsedPhysical, TrackedTotal);
-#elif  PLATFORM_IOS || UE_SERVER
-	const int64 PlatformProcessMemory = static_cast<int64>(PlatformStats.UsedPhysical);
-#else
-	const int64 PlatformProcessMemory = static_cast<int64>(PlatformStats.TotalPhysical) -
-		static_cast<int64>(PlatformStats.AvailablePhysical);
-#endif
 	const int64 PlatformTrackedTotal = PlatformTracker.GetTrackedTotal();
 	MemoryUsagePlatformTotalUntracked = FMath::Max<int64>(0, PlatformProcessMemory - PlatformTrackedTotal);
 
 	PlatformTracker.SetTagAmountInUpdate(FindOrAddTagData(ELLMTag::PlatformTotal), PlatformProcessMemory, false);
-	PlatformTracker.SetTagAmountInUpdate(FindOrAddTagData(ELLMTag::PlatformUntracked),
-		MemoryUsagePlatformTotalUntracked, false);
-	PlatformTracker.SetTagAmountInUpdate(FindOrAddTagData(ELLMTag::PlatformOSAvailable),
-		PlatformStats.AvailablePhysical, false);
+	PlatformTracker.SetTagAmountInUpdate(FindOrAddTagData(ELLMTag::PlatformUntracked), MemoryUsagePlatformTotalUntracked, false);
+	PlatformTracker.SetTagAmountInUpdate(FindOrAddTagData(ELLMTag::PlatformOSAvailable), PlatformStats.AvailablePhysical, false);
 
 	// remove the MemoryUsageCurrentOverhead from the "Total" for the default LLM as it's not something anyone needs
 	// to investigate when finding what to reduce the platform LLM will have the info 
diff --git a/Engine/Source/Runtime/Core/Public/HAL/LowLevelMemTracker.h b/Engine/Source/Runtime/Core/Public/HAL/LowLevelMemTracker.h
--- a/Engine/Source/Runtime/Core/Public/HAL/LowLevelMemTracker.h
+++ b/Engine/Source/Runtime/Core/Public/HAL/LowLevelMemTracker.h
@@ -164,6 +164,7 @@ ENUM_CLASS_FLAGS(ESizeParams);
 	/*FMalloc is a special tag that is reserved for the Platform Tracker only. It's used with ELLMAllocType::FMalloc to calculate ELLMTag::FMallocUnused. */								   \
 	macro(FMalloc,								"FMalloc",						GET_STATFNAME(STAT_FMallocLLM),								NAME_None,										-1)\
 	macro(FMallocUnused,						"FMallocUnused",				GET_STATFNAME(STAT_FMallocUnusedLLM),						GET_STATFNAME(STAT_EngineSummaryLLM),			-1)\
+	macro(RHIUnused,							"RHIUnused",					GET_STATFNAME(STAT_RHIUnusedLLM),							GET_STATFNAME(STAT_EngineSummaryLLM),			-1)\
 	macro(ThreadStack,							"ThreadStack",					GET_STATFNAME(STAT_ThreadStackLLM),							GET_STATFNAME(STAT_EngineSummaryLLM),			-1)\
 	macro(ThreadStackPlatform,					"ThreadStack",					GET_STATFNAME(STAT_ThreadStackPlatformLLM),					NAME_None,										-1)\
 	macro(ProgramSizePlatform,					"ProgramSize",					GET_STATFNAME(STAT_ProgramSizePlatformLLM),					NAME_None,										-1)\
@@ -304,7 +305,7 @@ enum class ELLMAllocType
 	None = 0,
 	FMalloc,
 	System,
-
+	RHI,
 	Count
 };
 
