Fix misaligned load/store UB in FGenericPlatformMemory::(Read|Write)Unaligned.
Use memcpy to load from unaligned to aligned, and store from aligned to unaligned. On platforms that support unaligned loads/stores, the memcpy will compile to an unaligned load/store.

diff --git a/Engine/Source/Runtime/Core/Private/GenericPlatform/GenericPlatformMemory.cpp b/Engine/Source/Runtime/Core/Private/GenericPlatform/GenericPlatformMemory.cpp
--- a/Engine/Source/Runtime/Core/Private/GenericPlatform/GenericPlatformMemory.cpp
+++ b/Engine/Source/Runtime/Core/Private/GenericPlatform/GenericPlatformMemory.cpp
@@ -82,16 +82,6 @@ namespace UE::Memory::Private
 }
 #endif
 
-struct TUnalignedTester
-{
-	FGenericPlatformMemory::TUnaligned<uint8> A;
-	FGenericPlatformMemory::TUnaligned<uint16> B;
-	TUnalignedTester()
-	{
-		static_assert(STRUCT_OFFSET(TUnalignedTester, B) == 1, "TUnaligned failure.");
-	}
-};
-
 TRACE_DECLARE_MEMORY_COUNTER(PlatformMemoryTotalPhysical, TEXT("PlatformMemory/TotalPhysical"));
 TRACE_DECLARE_MEMORY_COUNTER(PlatformMemoryTotalVirtual, TEXT("PlatformMemory/TotalVirtual"));
 TRACE_DECLARE_MEMORY_COUNTER(PlatformMemoryPageSize, TEXT("PlatformMemory/PageSize"));
diff --git a/Engine/Source/Runtime/Core/Public/GenericPlatform/GenericPlatformMemory.h b/Engine/Source/Runtime/Core/Public/GenericPlatform/GenericPlatformMemory.h
--- a/Engine/Source/Runtime/Core/Public/GenericPlatform/GenericPlatformMemory.h
+++ b/Engine/Source/Runtime/Core/Public/GenericPlatform/GenericPlatformMemory.h
@@ -654,26 +654,6 @@ public:
 		}
 	}
 
-#if defined(_MSC_VER)
-	// MSVC version
-	__pragma(pack(push, 1))
-		template <typename U>
-	struct TUnaligned
-	{
-		U Value;
-	};
-	__pragma(pack(pop))
-#else 
-	// assume it is either clang or something that supports the clang attributes
-#pragma pack(push, 1)
-	template <typename U>
-	struct TUnaligned
-	{
-		U Value;
-	};
-#pragma pack(pop)
-#endif
-
 	/**
 	* Loads a simple POD type from unaligned memory.
 	*
@@ -683,7 +663,9 @@ public:
 	template <typename T>
 	static FORCEINLINE T ReadUnaligned(const void* Ptr)
 	{
-		return reinterpret_cast<const TUnaligned<T>*>(Ptr)->Value;
+		T AlignedT;
+		memcpy(&AlignedT, Ptr, sizeof(T));
+		return AlignedT;
 	}
 
 	/**
@@ -695,7 +677,7 @@ public:
 	template <typename T>
 	static FORCEINLINE void WriteUnaligned(void* Ptr, const T& InValue)
 	{
-		reinterpret_cast<TUnaligned<T>*>(Ptr)->Value = InValue;
+		memcpy(Ptr, &InValue, sizeof(T));
 	}
 
 	/**
