TaskGraph: fixed a race reported by TSAN - `StartWorkers()` adds to `ReserveEvents` array while a just created worker accessed it in CreateWorker() via `EventStack`

diff --git a/Engine/Source/Runtime/Core/Private/Async/Fundamental/ReserveScheduler.cpp b/Engine/Source/Runtime/Core/Private/Async/Fundamental/ReserveScheduler.cpp
--- a/Engine/Source/Runtime/Core/Private/Async/Fundamental/ReserveScheduler.cpp
+++ b/Engine/Source/Runtime/Core/Private/Async/Fundamental/ReserveScheduler.cpp
@@ -81,12 +81,12 @@ void FReserveScheduler::StartWorkers(uint32 NumWorkers, FThread::EForkable IsFor
 		check(!WorkerThreads.Num());
 		check(NextWorkerId == 0);
 
-		ReserveEvents.Reserve(NumWorkers);
+		ReserveEvents.AddDefaulted(NumWorkers);
+		WorkerThreads.Reserve(NumWorkers);
 		UE::Trace::ThreadGroupBegin(TEXT("Reserve Workers"));
 		for (uint32 WorkerId = 0; WorkerId < NumWorkers; ++WorkerId)
 		{
-			ReserveEvents.Emplace();
-			WorkerThreads.Add(CreateWorker(IsForkable, &ReserveEvents.Last(), WorkerPriority));
+			WorkerThreads.Add(CreateWorker(IsForkable, &ReserveEvents[WorkerId], WorkerPriority));
 		}
 		UE::Trace::ThreadGroupEnd();
 	}
