Remove debug code that was not intended to be submitted

#rnx

diff --git a/Engine/Source/Runtime/Core/Private/Misc/ConfigContext.cpp b/Engine/Source/Runtime/Core/Private/Misc/ConfigContext.cpp
--- a/Engine/Source/Runtime/Core/Private/Misc/ConfigContext.cpp
+++ b/Engine/Source/Runtime/Core/Private/Misc/ConfigContext.cpp
@@ -113,10 +113,6 @@ const FConfigContext::FPerPlatformDirs& FConfigContext::GetPerPlatformDirs(const
 		FString PluginExtDir;
 		if (bIsForPlugin)
 		{
- 			if (DestIniFilename.Contains(TEXT("Enhanced")))
-			{
-				UE_LOG(LogConfig, Display, TEXT("paper2d"));
-			}
 			// look if there's a plugin extension for this platform, it will have the platform name in the path
 			for (const FString& ChildDir : ChildPluginBaseDirs)
 			{
