Localization: Downgrading intermittent localization write error to warning pending investigation

diff --git a/Engine/Source/Runtime/Core/Private/Misc/FileHelper.cpp b/Engine/Source/Runtime/Core/Private/Misc/FileHelper.cpp
--- a/Engine/Source/Runtime/Core/Private/Misc/FileHelper.cpp
+++ b/Engine/Source/Runtime/Core/Private/Misc/FileHelper.cpp
@@ -613,7 +613,10 @@ bool FFileHelper::SaveStringToFile( FStringView String, const TCHAR* Filename,
 	// max size of the string is a UCS2CHAR for each character and some UNICODE magic 
 	TUniquePtr<FArchive> Ar = TUniquePtr<FArchive>( FileManager->CreateFileWriter( Filename, WriteFlags ) );
 	if( !Ar )
+	{
+		UE_LOG(LogStreaming, Warning, TEXT("SaveStringToFile: couldn't create file writer"));
 		return false;
+	}
 
 	if( String.IsEmpty() )
 	{
@@ -653,6 +656,14 @@ bool FFileHelper::SaveStringToFile( FStringView String, const TCHAR* Filename,
 	// Always explicitly close to catch errors from flush/close
 	Ar->Close();
 
+	if (Ar->IsError())
+	{
+		UE_LOG(LogStreaming, Warning, TEXT("SaveStringToFile: Ar->IsError() == true"));
+	}
+	if (Ar->IsCriticalError())
+	{
+		UE_LOG(LogStreaming, Warning, TEXT("SaveStringToFile: Ar->IsCriticalError() == true"));
+	}
 	return !Ar->IsError() && !Ar->IsCriticalError();
 }
 
