Compile fixes when compiling for Win64 ARM64EC

diff --git a/Engine/Source/Runtime/Core/Public/Windows/AllowWindowsPlatformAtomics.h b/Engine/Source/Runtime/Core/Public/Windows/AllowWindowsPlatformAtomics.h
--- a/Engine/Source/Runtime/Core/Public/Windows/AllowWindowsPlatformAtomics.h
+++ b/Engine/Source/Runtime/Core/Public/Windows/AllowWindowsPlatformAtomics.h
@@ -14,7 +14,7 @@
 
 #define InterlockedIncrement _InterlockedIncrement
 #define InterlockedDecrement _InterlockedDecrement
-#if !defined(InterlockedAdd)
+#if !defined(InterlockedAdd) && !defined(_M_ARM64EC)
 	#define InterlockedAdd _InterlockedAdd
 #endif
 #define InterlockedExchange _InterlockedExchange
