#UE Remove warning about not being able to open a file in FFileHelper::SaveStringToFile as this is sometimes expected. Added more detail to other warnings that were recently added to include the filename

diff --git a/Engine/Source/Runtime/Core/Private/Misc/FileHelper.cpp b/Engine/Source/Runtime/Core/Private/Misc/FileHelper.cpp
--- a/Engine/Source/Runtime/Core/Private/Misc/FileHelper.cpp
+++ b/Engine/Source/Runtime/Core/Private/Misc/FileHelper.cpp
@@ -613,10 +613,7 @@ bool FFileHelper::SaveStringToFile( FStringView String, const TCHAR* Filename,
 	// max size of the string is a UCS2CHAR for each character and some UNICODE magic 
 	TUniquePtr<FArchive> Ar = TUniquePtr<FArchive>( FileManager->CreateFileWriter( Filename, WriteFlags ) );
 	if( !Ar )
-	{
-		UE_LOG(LogStreaming, Warning, TEXT("SaveStringToFile: couldn't create file writer"));
 		return false;
-	}
 
 	if( String.IsEmpty() )
 	{
@@ -658,11 +655,11 @@ bool FFileHelper::SaveStringToFile( FStringView String, const TCHAR* Filename,
 
 	if (Ar->IsError())
 	{
-		UE_LOG(LogStreaming, Warning, TEXT("SaveStringToFile: Ar->IsError() == true"));
+		UE_LOG(LogStreaming, Warning, TEXT("SaveStringToFile: Ar->IsError() == true. File:%s"), Filename);
 	}
 	if (Ar->IsCriticalError())
 	{
-		UE_LOG(LogStreaming, Warning, TEXT("SaveStringToFile: Ar->IsCriticalError() == true"));
+		UE_LOG(LogStreaming, Warning, TEXT("SaveStringToFile: Ar->IsCriticalError() == true. File:%s"), Filename);
 	}
 	return !Ar->IsError() && !Ar->IsCriticalError();
 }
