Fixed Windows Arm64 compile issues.

diff --git a/Engine/Source/Runtime/Core/Private/Windows/WindowsPlatformMisc.cpp b/Engine/Source/Runtime/Core/Private/Windows/WindowsPlatformMisc.cpp
--- a/Engine/Source/Runtime/Core/Private/Windows/WindowsPlatformMisc.cpp
+++ b/Engine/Source/Runtime/Core/Private/Windows/WindowsPlatformMisc.cpp
@@ -2594,6 +2594,8 @@ public:
 	}
 
 private:
+
+#if !PLATFORM_CPU_ARM_FAMILY
 	/**
 	 * Checks if __cpuid instruction is present on current machine.
 	 *
@@ -2641,8 +2643,7 @@ private:
 	 */
 	static bool CheckForTimedPauseInstruction()
 	{
-		// _tpause isn't defined for Arm64EC
-#if PLATFORM_SEH_EXCEPTIONS_DISABLED || defined(_M_ARM64EC)
+#if PLATFORM_SEH_EXCEPTIONS_DISABLED
 		return false;
 #else
 		bool bSupportsTpause = false;
@@ -2760,6 +2761,14 @@ private:
 
 		return Result;
 	}
+#else
+	static bool CheckForCPUIDInstruction() { return false; }
+	static bool CheckForTimedPauseInstruction() { return false; }
+	static void GetCPUVendor(ANSICHAR(&OutBuffer)[12 + 1]) {}
+	static void GetCPUBrand(ANSICHAR(&OutBrandString)[0x40]) {}
+	static void QueryCPUInfo(int Args[4]) {}
+	static int32 QueryCacheLineSize() { return PLATFORM_CACHE_LINE_SIZE; }
+#endif
 
 	/** Static field with pre-cached __cpuid data. */
 	static FCPUIDQueriedData CPUIDStaticCache;
@@ -2804,6 +2813,9 @@ FString FWindowsPlatformMisc::GetCPUBrand()
 
 bool FWindowsPlatformMisc::HasAVX2InstructionSupport()
 {
+#if PLATFORM_CPU_ARM_FAMILY
+	return false;
+#else
 	if (!HasCPUIDInstruction())
 	{
 		return false;
@@ -2847,6 +2859,7 @@ bool FWindowsPlatformMisc::HasAVX2InstructionSupport()
 	}
 
 	return true;
+#endif
 }
 
 #include "Windows/AllowWindowsPlatformTypes.h"
diff --git a/Engine/Source/Runtime/Core/Private/Windows/WindowsPlatformStackWalk.cpp b/Engine/Source/Runtime/Core/Private/Windows/WindowsPlatformStackWalk.cpp
--- a/Engine/Source/Runtime/Core/Private/Windows/WindowsPlatformStackWalk.cpp
+++ b/Engine/Source/Runtime/Core/Private/Windows/WindowsPlatformStackWalk.cpp
@@ -218,10 +218,17 @@ static int32 CaptureStackTraceExternalProcess(uint64* BackTrace, uint32 MaxDepth
 		StackFrame64.AddrStack.Mode      = AddrModeFlat;
 		StackFrame64.AddrFrame.Mode      = AddrModeFlat;
 #if PLATFORM_64BITS
-		StackFrame64.AddrPC.Offset = ContextCopy.Rip;
-		StackFrame64.AddrStack.Offset = ContextCopy.Rsp;
-		StackFrame64.AddrFrame.Offset = ContextCopy.Rbp;
+#if defined(_M_ARM64)
+		StackFrame64.AddrPC.Offset       = ContextCopy.Pc;
+		StackFrame64.AddrStack.Offset    = ContextCopy.Sp;
+		StackFrame64.AddrFrame.Offset    = ContextCopy.Fp;
+		MachineType                      = IMAGE_FILE_MACHINE_ARM64;
+#else
+		StackFrame64.AddrPC.Offset       = ContextCopy.Rip;
+		StackFrame64.AddrStack.Offset    = ContextCopy.Rsp;
+		StackFrame64.AddrFrame.Offset    = ContextCopy.Rbp;
 		MachineType                      = IMAGE_FILE_MACHINE_AMD64;
+#endif
 #else	//PLATFORM_64BITS
 		StackFrame64.AddrPC.Offset       = ContextCopy.Eip;
 		StackFrame64.AddrStack.Offset    = ContextCopy.Esp;
