CONTROL_FLOW_PERF_TRACE_STEP. Adding support in FControlFlow and TRACE_CPUPROFILE macros to allow you to begin an event in a control flow node implementation and have it automatically terminated by FControlFlow whenever/however that node completes. This ensures that in a profile trace sequential steps show up sequentially rather than appearing nested (despite the technical fact that they are actually invoked as nested functions).

diff --git a/Engine/Plugins/Experimental/ControlFlows/Source/ControlFlows/Private/ControlFlow.cpp b/Engine/Plugins/Experimental/ControlFlows/Source/ControlFlows/Private/ControlFlow.cpp
--- a/Engine/Plugins/Experimental/ControlFlows/Source/ControlFlows/Private/ControlFlow.cpp
+++ b/Engine/Plugins/Experimental/ControlFlows/Source/ControlFlows/Private/ControlFlow.cpp
@@ -74,6 +74,14 @@ void FControlFlow::HandleControlFlowNodeCompleted(TSharedRef<const FControlFlowN
 {
 	UE_LOG(LogControlFlows, Verbose, TEXT("ControlFlow - Executing %s%s(x%d).%s (FlowControlNodeCompleted)"), *GetFlowPath(), *DebugName, GetRepeatedFlowCount(), *NodeCompleted->GetNodeName());
 
+#if CPUPROFILERTRACE_ENABLED
+	if (bProfilerEventStarted)
+	{
+		FCpuProfilerTrace::OutputEndEvent();
+		bProfilerEventStarted = false;
+	}
+#endif 
+
 	SubFlowStack_ForDebugging.Add(SharedThis(this));
 
 	if (ensureMsgf(CurrentNode.IsValid(), TEXT("CurrentNode isn't valid when handling control flow node (%s) being completed.  This will likely become a static_assert/compile-error"),
diff --git a/Engine/Plugins/Experimental/ControlFlows/Source/ControlFlows/Private/ControlFlowNode.cpp b/Engine/Plugins/Experimental/ControlFlows/Source/ControlFlows/Private/ControlFlowNode.cpp
--- a/Engine/Plugins/Experimental/ControlFlows/Source/ControlFlows/Private/ControlFlowNode.cpp
+++ b/Engine/Plugins/Experimental/ControlFlows/Source/ControlFlows/Private/ControlFlowNode.cpp
@@ -63,6 +63,14 @@ FString FControlFlowNode::GetNodeName() const
 	return NodeName;
 }
 
+void FControlFlowNode::SetProfilerEventStarted()
+{
+	if (ensureAlways(Parent.IsValid()))
+	{
+		Parent.Pin()->SetProfilerEventStarted();
+	}
+}
+
 ///////////////////////////////////////////////////
 
 FControlFlowNode_SelfCompleting::FControlFlowNode_SelfCompleting(TSharedRef<FControlFlow> ControlFlowParent, const FString& FlowNodeDebugName)
diff --git a/Engine/Plugins/Experimental/ControlFlows/Source/ControlFlows/Public/ControlFlow.h b/Engine/Plugins/Experimental/ControlFlows/Source/ControlFlows/Public/ControlFlow.h
--- a/Engine/Plugins/Experimental/ControlFlows/Source/ControlFlows/Public/ControlFlow.h
+++ b/Engine/Plugins/Experimental/ControlFlows/Source/ControlFlows/Public/ControlFlow.h
@@ -335,6 +335,8 @@ public:
 	CONTROLFLOWS_API FControlFlowPopulator& QueueLoop(FControlFlowLoopComplete& LoopCompleteDelgate, const FString& TaskName = TEXT(""), const FString& FlowNodeDebugName = TEXT(""));
 	
 private:
+	FORCEINLINE void SetProfilerEventStarted() { ensureMsgf(!bProfilerEventStarted, TEXT("Started a new control flow profiler event before previous step completed.")); bProfilerEventStarted = true; }
+
 	void HandleControlFlowNodeCompleted(TSharedRef<const FControlFlowNode> NodeCompleted);
 
 	mutable FSimpleMulticastDelegate OnStepCompletedDelegate;
@@ -394,4 +396,22 @@ private:
 	TArray<TSharedRef<FControlFlowNode>> FlowQueue;
 
 	TSharedPtr<FTrackedActivity> Activity;
+
+	bool bProfilerEventStarted = false;
+
 };
+
+#if CPUPROFILERTRACE_ENABLED
+
+#define CONTROL_FLOW_PERF_TRACE_STEP(FlowHandle, EventName) \
+	TRACE_CPUPROFILER_EVENT_MANUAL_START(#EventName); \
+	if (TRACE_CPUPROFILER_EVENT_MANUAL_IS_ENABLED()) \
+	{ \
+		FlowHandle->SetProfilerEventStarted(); \
+	}
+
+#else
+
+#define CONTROL_FLOW_PERF_TRACE_STEP(FlowHandle, EventName)
+
+#endif
diff --git a/Engine/Plugins/Experimental/ControlFlows/Source/ControlFlows/Public/ControlFlowNode.h b/Engine/Plugins/Experimental/ControlFlows/Source/ControlFlows/Public/ControlFlowNode.h
--- a/Engine/Plugins/Experimental/ControlFlows/Source/ControlFlows/Public/ControlFlowNode.h
+++ b/Engine/Plugins/Experimental/ControlFlows/Source/ControlFlows/Public/ControlFlowNode.h
@@ -52,6 +52,8 @@ public:
 
 	virtual FString GetNodeName() const;
 
+	CONTROLFLOWS_API void SetProfilerEventStarted();
+
 protected:
 	friend class FControlFlow;
 	
diff --git a/Engine/Source/Runtime/Core/Public/ProfilingDebugging/CpuProfilerTrace.h b/Engine/Source/Runtime/Core/Public/ProfilingDebugging/CpuProfilerTrace.h
--- a/Engine/Source/Runtime/Core/Public/ProfilingDebugging/CpuProfilerTrace.h
+++ b/Engine/Source/Runtime/Core/Public/ProfilingDebugging/CpuProfilerTrace.h
@@ -247,6 +247,19 @@ struct FCpuProfilerTrace
 #define TRACE_CPUPROFILER_EVENT_SCOPE_USE_ON_CHANNEL(DeclName, NameStr, ScopeName, Channel, Condition) \
 	FCpuProfilerTrace::FEventScope ScopeName(DeclName, NameStr, Channel, Condition, __FILE__, __LINE__);
 
+// Advanced macro that will check if CpuChannel is enabled and, if so, declare a new stat event and start it
+#define TRACE_CPUPROFILER_EVENT_MANUAL_START(EventNameStr) \
+	if (CpuChannel) \
+	{ \
+		static uint32 PREPROCESSOR_JOIN(__CpuProfilerEventSpecId, __LINE__) = FCpuProfilerTrace::OutputEventType(EventNameStr, __FILE__, __LINE__); \
+		FCpuProfilerTrace::OutputBeginEvent(PREPROCESSOR_JOIN(__CpuProfilerEventSpecId, __LINE__)); \
+	}
+
+// Advanced macro that can be used with TRACE_CPUPROFILER_EVENT_MANUAL_START to wrap code that should only be executed if 
+// the event was actually started
+#define TRACE_CPUPROFILER_EVENT_MANUAL_IS_ENABLED() \
+	bool(CpuChannel)
+
 // Trace a scoped cpu timing event providing a static string (const ANSICHAR* or const TCHAR*)
 // as the scope name. It will use the Cpu trace channel.
 // Example: TRACE_CPUPROFILER_EVENT_SCOPE_STR("My Scoped Timer A")
@@ -302,6 +315,8 @@ struct FCpuProfilerTrace
 #define TRACE_CPUPROFILER_EVENT_DECLARE(DeclName)
 #define TRACE_CPUPROFILER_EVENT_SCOPE_USE(DeclName, NameStr, ScopeName, Condition)
 #define TRACE_CPUPROFILER_EVENT_SCOPE_USE_ON_CHANNEL(DeclName, NameStr, ScopeName, Channel, Condition)
+#define TRACE_CPUPROFILER_EVENT_MANUAL_START(EventNameStr)
+#define TRACE_CPUPROFILER_EVENT_MANUAL_IS_ENABLED() false
 #define TRACE_CPUPROFILER_EVENT_SCOPE_STR(NameStr)
 #define TRACE_CPUPROFILER_EVENT_SCOPE_ON_CHANNEL_STR(NameStr, Channel)
 #define TRACE_CPUPROFILER_EVENT_SCOPE(Name)
