Add extra context information when messaging skip reason

diff --git a/Engine/Source/Developer/AutomationWindow/Private/SAutomationWindow.cpp b/Engine/Source/Developer/AutomationWindow/Private/SAutomationWindow.cpp
--- a/Engine/Source/Developer/AutomationWindow/Private/SAutomationWindow.cpp
+++ b/Engine/Source/Developer/AutomationWindow/Private/SAutomationWindow.cpp
@@ -1991,9 +1991,9 @@ TSharedRef<ITableRow> SAutomationWindow::OnGenerateWidgetForLog(TSharedPtr<FAuto
 {
 	check(Message.IsValid());
 
-	// \[((?:[\w]\:|\\)(?:(?:\\[A-Za-z_\-\s0-9\.]+)+)\.(?:cpp|h))\((\d+)\)\]$
+	// \[((?:[\w]\:|\\)(?:(?:\\[A-Za-z_\-\s0-9\.]+)+)\.(?:cpp|h|ini))\((\d+)\)\]$
 	// https://regex101.com/r/vV4cV7/25
-	FRegexPattern FileAndLinePattern(TEXT("\\[((?:[\\w]\\:|\\\\)(?:(?:\\\\[A-Za-z_\\-\\s0-9\\.]+)+)\\.(?:cpp|h))\\((\\d+)\\)\\]$"));
+	FRegexPattern FileAndLinePattern(TEXT("\\[((?:[\\w]\\:|\\\\)(?:(?:\\\\[A-Za-z_\\-\\s0-9\\.]+)+)\\.(?:cpp|h|ini))\\((\\d+)\\)\\]$"));
 	FRegexMatcher FileAndLineRegexMatcher(FileAndLinePattern, Message->Text);
 
 	TSharedRef<SWidget> SourceLink = SNullWidget::NullWidget;
diff --git a/Engine/Source/Runtime/AutomationWorker/Private/AutomationWorkerModule.cpp b/Engine/Source/Runtime/AutomationWorker/Private/AutomationWorkerModule.cpp
--- a/Engine/Source/Runtime/AutomationWorker/Private/AutomationWorkerModule.cpp
+++ b/Engine/Source/Runtime/AutomationWorker/Private/AutomationWorkerModule.cpp
@@ -527,6 +527,8 @@ bool FAutomationWorkerModule::IsTestExcluded(const FString& InTestToRun, FString
 {
 	FName SkipReason;
 	UAutomationTestExcludelist* Excludelist = UAutomationTestExcludelist::Get();
+	check(nullptr != Excludelist);
+	
 	static const TSet<FName> RHI = GetRHIForAutomation();
 	if (Excludelist->IsTestExcluded(InTestToRun, RHI, &SkipReason, OutWarn))
 	{
@@ -534,6 +536,23 @@ bool FAutomationWorkerModule::IsTestExcluded(const FString& InTestToRun, FString
 		{
 			(*OutReason) = (SkipReason.IsNone() ? TEXT("unknown reason") : SkipReason.ToString());
 			(*OutReason) += TEXT(" [config]");
+
+			if (Excludelist->GetFName().ToString().StartsWith("Default_"))
+			{
+				// We can handle only CDO to correctly detect the filename
+				FString Filename = Excludelist->GetConfigFilename();
+				
+				if (!Filename.IsEmpty())
+				{
+					Filename = FPaths::ConvertRelativePathToFull(Filename);
+						FPaths::MakePlatformFilename(Filename);
+						*OutReason += TEXT(" [");
+						*OutReason += Filename;
+						// Using of line number 1 as a default value to get it working as a hyperlink
+						*OutReason += TEXT("(1)]");
+				}
+			}
+
 		}
 
 		return true;
diff --git a/Engine/Source/Runtime/Core/Private/Misc/AutomationTest.cpp b/Engine/Source/Runtime/Core/Private/Misc/AutomationTest.cpp
--- a/Engine/Source/Runtime/Core/Private/Misc/AutomationTest.cpp
+++ b/Engine/Source/Runtime/Core/Private/Misc/AutomationTest.cpp
@@ -42,6 +42,23 @@ namespace AutomationTest
 		TEXT("Automation.LogBPTestMetadata"),
 		bLogBPTestMetadata,
 		TEXT("Whether to output blueprint functional test metadata to the log when test is running"));
+
+	// The method prepares the filename and LineNumber to be placed in the form that could be extracted by SAutomationWindow widget if it is additionally eclosed into []
+	// The result format is filename(line)
+	static FString CreateFileLineDescription(const FString& Filename, const int32 LineNumber)
+	{
+		FString Result;
+
+		if (!Filename.IsEmpty() && LineNumber > 0)
+		{
+			Result += Filename;
+			Result += TEXT("(");
+			Result += FString::FromInt(LineNumber);
+			Result += TEXT(")");
+		}
+
+		return Result;
+	}
 };
 
 bool FAutomationTestBase::bSuppressLogWarnings = false;
@@ -950,6 +967,15 @@ bool FAutomationTestFramework::CanRunTestInEnvironment(const FString& InTestToRu
 			}
 
 			*OutReason += TEXT(" [code]");
+			FString Filename = Test->GetTestSourceFileName();
+			FPaths::MakePlatformFilename(Filename);
+			const FString FileLineDescription = AutomationTest::CreateFileLineDescription(Filename, Test->GetTestSourceFileLine());
+			if (!FileLineDescription.IsEmpty())
+			{
+				*OutReason += TEXT(" [");
+				*OutReason += FileLineDescription;
+				*OutReason += TEXT("]");
+			}
 		}
 		
 		return false;
@@ -1022,13 +1048,12 @@ FString FAutomationExecutionEntry::ToString() const
 
 	// Place the filename at the end so it can be extracted by the SAutomationWindow widget
 	// Expectation is "[filename(line)]"
-	if ( !Filename.IsEmpty() && LineNumber > 0 )
+	const FString FileLineDescription = AutomationTest::CreateFileLineDescription(Filename, LineNumber);
+	if ( !FileLineDescription.IsEmpty() )
 	{
 		ComplexString += TEXT(" [");
-		ComplexString += Filename;
-		ComplexString += TEXT("(");
-		ComplexString += FString::FromInt(LineNumber);
-		ComplexString += TEXT(")]");
+		ComplexString += FileLineDescription;
+		ComplexString += TEXT("]");
 	}
 
 	return ComplexString;
@@ -1047,13 +1072,11 @@ FString FAutomationExecutionEntry::ToStringFormattedEditorLog() const
 		ComplexString += TEXT("] ");
 	}
 
-	if (!Filename.IsEmpty() && LineNumber > 0)
+	const FString FileLineDescription = AutomationTest::CreateFileLineDescription(Filename, LineNumber);
+	if (!FileLineDescription.IsEmpty())
 	{
 		ComplexString += TEXT(" ");
-		ComplexString += Filename;
-		ComplexString += TEXT("(");
-		ComplexString += FString::FromInt(LineNumber);
-		ComplexString += TEXT(")");
+		ComplexString += FileLineDescription;
 	}
 
 	return ComplexString;
