Add metric density to unit conversions. Use g/cm3 for physical material density.

diff --git a/Engine/Source/Runtime/Core/Private/Math/UnitConversion.cpp b/Engine/Source/Runtime/Core/Private/Math/UnitConversion.cpp
--- a/Engine/Source/Runtime/Core/Private/Math/UnitConversion.cpp
+++ b/Engine/Source/Runtime/Core/Private/Math/UnitConversion.cpp
@@ -61,6 +61,11 @@ FParseCandidate ParseCandidates[] = {
 	{ TEXT("Pounds"),				EUnit::Pounds },				{ TEXT("lb"),		EUnit::Pounds },
 	{ TEXT("Stones"),				EUnit::Stones },				{ TEXT("st"),		EUnit::Stones },
 
+	{ TEXT("GramsPerCubicCentimeter"),		EUnit::GramsPerCubicCentimeter },		{ TEXT("g/cm3"),	EUnit::GramsPerCubicCentimeter },		{ TEXT("g/cm\u00B3"), EUnit::GramsPerCubicCentimeter },
+	{ TEXT("GramsPerCubicMeter"),			EUnit::GramsPerCubicMeter },			{ TEXT("g/m3"),		EUnit::GramsPerCubicMeter },			{ TEXT("g/m\u00B3"), EUnit::GramsPerCubicMeter },
+	{ TEXT("KilogramsPerCubicCentimeter"),	EUnit::KilogramsPerCubicCentimeter },	{ TEXT("kg/cm3"),	EUnit::KilogramsPerCubicCentimeter },	{ TEXT("kg/cm\u00B3"), EUnit::KilogramsPerCubicCentimeter },
+	{ TEXT("KilogramsPerCubicMeter"),		EUnit::KilogramsPerCubicMeter },		{ TEXT("kg/m3"),	EUnit::KilogramsPerCubicMeter },		{ TEXT("kg/m\u00B3"), EUnit::KilogramsPerCubicMeter },
+
 	{ TEXT("Newtons"),				EUnit::Newtons },				{ TEXT("N"),		EUnit::Newtons },
 	{ TEXT("PoundsForce"),			EUnit::PoundsForce },			{ TEXT("lbf"),		EUnit::PoundsForce },
 	{ TEXT("KilogramsForce"),		EUnit::KilogramsForce },		{ TEXT("kgf"),		EUnit::KilogramsForce },
@@ -124,6 +129,8 @@ const TCHAR* const DisplayStrings[] = {
 	TEXT("\u00B5g"), TEXT("mg"), TEXT("g"), TEXT("kg"), TEXT("t"),
 	TEXT("oz"), TEXT("lb"), TEXT("st"),
 
+	TEXT("g/cm\u00B3"), TEXT("g/m\u00B3"), TEXT("kg/cm\u00B3"), TEXT("kg/m\u00B3"),
+
 	TEXT("N"), TEXT("lbf"), TEXT("kgf"), TEXT("kgcm/s\u00B2"),
 
 	TEXT("Nm"), TEXT("kgcm\u00B2/s\u00B2"),
@@ -159,6 +166,8 @@ const EUnitType UnitTypes[] = {
 	EUnitType::Mass,		EUnitType::Mass,		EUnitType::Mass,		EUnitType::Mass,		EUnitType::Mass,
 	EUnitType::Mass,		EUnitType::Mass,		EUnitType::Mass,
 
+	EUnitType::Density,		EUnitType::Density,		EUnitType::Density,		EUnitType::Density,
+
 	EUnitType::Force,		EUnitType::Force,		EUnitType::Force,		EUnitType::Force,
 
 	EUnitType::Torque,		EUnitType::Torque,
@@ -444,6 +453,7 @@ FUnitSettings::FUnitSettings()
 	DisplayUnits[(uint8)EUnitType::Speed].Add(EUnit::MetersPerSecond);
 	DisplayUnits[(uint8)EUnitType::Temperature].Add(EUnit::Celsius);
 	DisplayUnits[(uint8)EUnitType::Mass].Add(EUnit::Kilograms);
+	DisplayUnits[(uint8)EUnitType::Density].Add(EUnit::GramsPerCubicCentimeter);
 	DisplayUnits[(uint8)EUnitType::Force].Add(EUnit::Newtons);
 	DisplayUnits[(uint8)EUnitType::Torque].Add(EUnit::NewtonMeters);
 	DisplayUnits[(uint8)EUnitType::Frequency].Add(EUnit::Hertz);
@@ -626,6 +636,19 @@ namespace UnitConversion
 		}
 	}
 
+	double DensityUnificationFactor(EUnit From)
+	{
+		// Convert to g/cm^3
+		switch (From)
+		{
+		case EUnit::GramsPerCubicCentimeter:		return 1;
+		case EUnit::GramsPerCubicMeter:				return 0.000001;
+		case EUnit::KilogramsPerCubicCentimeter:	return 1000;
+		case EUnit::KilogramsPerCubicMeter:			return 0.001;
+		default:									return 1;
+		}
+	}
+
 	double ForceUnificationFactor(EUnit From)
 	{
 		// Convert to Newtons
@@ -754,6 +777,8 @@ namespace UnitConversion
 			TArray<FQuantizationInfo> MetricMass;
 			TArray<FQuantizationInfo> ImperialMass;
 
+			TArray<FQuantizationInfo> MetricDensity;
+
 			TArray<FQuantizationInfo> Frequency;
 			TArray<FQuantizationInfo> DataSize;
 
@@ -784,6 +809,11 @@ namespace UnitConversion
 				ImperialMass.Emplace(EUnit::Pounds,	14.0f);
 				ImperialMass.Emplace(EUnit::Stones,	0.0f);
 
+				MetricDensity.Emplace(EUnit::GramsPerCubicCentimeter,		0.000001f);
+				MetricDensity.Emplace(EUnit::GramsPerCubicMeter,			1000.0f);
+				MetricDensity.Emplace(EUnit::KilogramsPerCubicCentimeter,	0.000001f);
+				MetricDensity.Emplace(EUnit::KilogramsPerCubicMeter,		0.0f);
+
 				Frequency.Emplace(EUnit::Hertz,		1000.0f);
 				Frequency.Emplace(EUnit::Kilohertz,	1000.0f);
 				Frequency.Emplace(EUnit::Megahertz,	1000.0f);
@@ -826,6 +856,12 @@ namespace UnitConversion
 		case EUnit::Ounces: case EUnit::Pounds: case EUnit::Stones:
 			return &Bounds.ImperialMass;
 
+		case EUnit::GramsPerCubicCentimeter:
+		case EUnit::GramsPerCubicMeter:
+		case EUnit::KilogramsPerCubicCentimeter:
+		case EUnit::KilogramsPerCubicMeter:
+			return &Bounds.MetricDensity;
+
 		case EUnit::Hertz: case EUnit::Kilohertz: case EUnit::Megahertz: case EUnit::Gigahertz: case EUnit::RevolutionsPerMinute:
 			return &Bounds.Frequency;
 
diff --git a/Engine/Source/Runtime/Core/Public/Math/UnitConversion.h b/Engine/Source/Runtime/Core/Public/Math/UnitConversion.h
--- a/Engine/Source/Runtime/Core/Public/Math/UnitConversion.h
+++ b/Engine/Source/Runtime/Core/Public/Math/UnitConversion.h
@@ -35,6 +35,9 @@ enum class EUnit : uint8
 	Micrograms, Milligrams, Grams, Kilograms, MetricTons,
 	Ounces, Pounds, Stones,
 
+	/** Density units */
+	GramsPerCubicCentimeter, GramsPerCubicMeter, KilogramsPerCubicCentimeter, KilogramsPerCubicMeter,
+
 	/** Force units */
 	Newtons, PoundsForce, KilogramsForce, KilogramCentimetersPerSecondSquared,
 
@@ -69,7 +72,7 @@ enum class EUnit : uint8
 /** Enumeration that specifies particular classes of unit */
 enum class EUnitType
 {
-	Distance, Angle, Speed, Temperature, Mass, Force, Torque, Frequency, DataSize, LuminousFlux, LuminousIntensity, Illuminance, Luminance, Time, PixelDensity, Multipliers, ExposureValue, Stress,
+	Distance, Angle, Speed, Temperature, Mass, Density, Force, Torque, Frequency, DataSize, LuminousFlux, LuminousIntensity, Illuminance, Luminance, Time, PixelDensity, Multipliers, ExposureValue, Stress,
 
 	// Symbolic entry - do not use directly
 	NumberOf,
diff --git a/Engine/Source/Runtime/Core/Public/Math/UnitConversion.inl b/Engine/Source/Runtime/Core/Public/Math/UnitConversion.inl
--- a/Engine/Source/Runtime/Core/Public/Math/UnitConversion.inl
+++ b/Engine/Source/Runtime/Core/Public/Math/UnitConversion.inl
@@ -28,6 +28,8 @@ namespace UnitConversion
 	CORE_API double TemperatureUnificationFactor(EUnit From);
 	/** Find the common quantization factor for the specified mass unit. Quantizes to Grams. */
 	CORE_API double MassUnificationFactor(EUnit From);
+	/** Find the common quantization factor for the specified density unit. Quantizes to GramsPerCubicCentimeter. */
+	CORE_API double DensityUnificationFactor(EUnit From);
 	/** Find the common quantization factor for the specified force unit. Quantizes to Newtons. */
 	CORE_API double ForceUnificationFactor(EUnit From);
 	/** Find the common quantization factor for the specified torque unit. Quantizes to Newton meters. */
@@ -84,6 +86,7 @@ T FUnitConversion::Convert(T InValue, EUnit From, EUnit To)
 		case EUnitType::Angle:				return InValue * AngleUnificationFactor(From) 			* (1.0 / AngleUnificationFactor(To));
 		case EUnitType::Speed:				return InValue * SpeedUnificationFactor(From) 			* (1.0 / SpeedUnificationFactor(To));
 		case EUnitType::Mass:				return InValue * MassUnificationFactor(From) 			* (1.0 / MassUnificationFactor(To));
+		case EUnitType::Density:			return InValue * DensityUnificationFactor(From)			* (1.0 / DensityUnificationFactor(To));
 		case EUnitType::Force:				return InValue * ForceUnificationFactor(From) 			* (1.0 / ForceUnificationFactor(To));
 		case EUnitType::Torque:				return InValue * TorqueUnificationFactor(From)			* (1.0 / TorqueUnificationFactor(To));
 		case EUnitType::Frequency:			return InValue * FrequencyUnificationFactor(From) 		* (1.0 / FrequencyUnificationFactor(To));
diff --git a/Engine/Source/Runtime/PhysicsCore/Public/PhysicalMaterials/PhysicalMaterial.h b/Engine/Source/Runtime/PhysicsCore/Public/PhysicalMaterials/PhysicalMaterial.h
--- a/Engine/Source/Runtime/PhysicsCore/Public/PhysicalMaterials/PhysicalMaterial.h
+++ b/Engine/Source/Runtime/PhysicsCore/Public/PhysicalMaterials/PhysicalMaterial.h
@@ -99,7 +99,7 @@ class UPhysicalMaterial : public UObject
 	//
 	
 	/** Used with the shape of the object to calculate its mass properties. The higher the number, the heavier the object. g per cubic cm. */
-	UPROPERTY(EditAnywhere, BlueprintReadOnly, Category = PhysicalMaterial, meta=(ClampMin=0))
+	UPROPERTY(EditAnywhere, BlueprintReadOnly, Category = PhysicalMaterial, meta = (ClampMin = 0, ForceUnits = "g/cm3"))
 	float Density;
 
 	/**  How low the linear velocity can be before solver puts body to sleep. */
