Fix the compiler ICE by writing my own CloneFunctionsInto.

diff --git a/Engine/Source/Runtime/Core/Public/AutoRTFM/AutoRTFM.h b/Engine/Source/Runtime/Core/Public/AutoRTFM/AutoRTFM.h
--- a/Engine/Source/Runtime/Core/Public/AutoRTFM/AutoRTFM.h
+++ b/Engine/Source/Runtime/Core/Public/AutoRTFM/AutoRTFM.h
@@ -541,9 +541,18 @@ struct FRegisterOpenFunction
 
 // Macro-based variants so we completely compile away when not in use, even in debug builds
 #if UE_AUTORTFM
+
+#if defined(__clang__) && __has_warning("-Wdeprecated-this-capture")
+#define UE_AUTORTFM_BEGIN_DISABLE_WARNINGS _Pragma("clang diagnostic push") _Pragma("clang diagnostic ignored \"-Wdeprecated-this-capture\"")
+#define UE_AUTORTFM_END_DISABLE_WARNINGS _Pragma("clang diagnostic pop")
+#else
+#define UE_AUTORTFM_BEGIN_DISABLE_WARNINGS
+#define UE_AUTORTFM_END_DISABLE_WARNINGS
+#endif
+
 #define UE_AUTORTFM_OPEN(...) AutoRTFM::Open([&]() { __VA_ARGS__ })
-#define UE_AUTORTFM_OPENABORT(...) AutoRTFM::OpenAbort([=]() { __VA_ARGS__ })
-#define UE_AUTORTFM_OPENCOMMIT(...) AutoRTFM::OpenCommit([=]() { __VA_ARGS__ })
+#define UE_AUTORTFM_OPENABORT(...) UE_AUTORTFM_BEGIN_DISABLE_WARNINGS AutoRTFM::OpenAbort([=]() { __VA_ARGS__ }) UE_AUTORTFM_END_DISABLE_WARNINGS
+#define UE_AUTORTFM_OPENCOMMIT(...) UE_AUTORTFM_BEGIN_DISABLE_WARNINGS AutoRTFM::OpenCommit([=]() { __VA_ARGS__ }) UE_AUTORTFM_END_DISABLE_WARNINGS
 #define UE_AUTORTFM_TRANSACT(...) AutoRTFM::Transact([&]() { __VA_ARGS__ })
 #else
 #define UE_AUTORTFM_OPEN(...) do { __VA_ARGS__ } while (false)
