FWindowsPlatformProcess::TerminateProc fixes.
Prevent infinite recursion by keeping a set of visited process ids (UE-186400).
Don't kill the whole tree when KillTree is false (fallout from 19196914).
Close previously leaking snapshot and child process handles.

diff --git a/Engine/Source/Runtime/Core/Private/Windows/WindowsPlatformProcess.cpp b/Engine/Source/Runtime/Core/Private/Windows/WindowsPlatformProcess.cpp
--- a/Engine/Source/Runtime/Core/Private/Windows/WindowsPlatformProcess.cpp
+++ b/Engine/Source/Runtime/Core/Private/Windows/WindowsPlatformProcess.cpp
@@ -513,16 +513,26 @@ void FWindowsPlatformProcess::CloseProc(FProcHandle & ProcessHandle)
 
 void FWindowsPlatformProcess::TerminateProc( FProcHandle & ProcessHandle, bool KillTree )
 {
-	TerminateProcTreeWithPredicate(ProcessHandle, [](uint32 ProcessId, const TCHAR* ApplicationName) { return true; });
+	if (KillTree)
+	{
+		TerminateProcTreeWithPredicate(ProcessHandle, [](uint32 ProcessId, const TCHAR* ApplicationName) { return true; });
+	}
+	else
+	{
+		TerminateProcess(ProcessHandle.Get(),0);
+	}
 }
 
-void FWindowsPlatformProcess::TerminateProcTreeWithPredicate(
-	FProcHandle& ProcessHandle,
-	TFunctionRef<bool(uint32 ProcessId, const TCHAR* ApplicationName)> Predicate)
+static void TerminateProcTreeWithPredicateInternal(
+	HANDLE ProcessHandle,
+	TFunctionRef<bool(uint32 ProcessId, const TCHAR* ApplicationName)> Predicate,
+	TSet<DWORD>& VisitedProcessIds)
 {
-	::DWORD ProcessId = ::GetProcessId(ProcessHandle.Get());
+	::DWORD ProcessId = ::GetProcessId(ProcessHandle);
 	FString ProcessName = FPlatformProcess::GetApplicationName(ProcessId);
 
+	VisitedProcessIds.Add(ProcessId);
+
 	if (!Predicate(ProcessId, *ProcessName))
 	{
 		return;
@@ -546,16 +556,28 @@ void FWindowsPlatformProcess::TerminateProcTreeWithPredicate(
 
 					if (ChildProcHandle)
 					{
-						FProcHandle ChildHandle(ChildProcHandle);
-						TerminateProcTreeWithPredicate(ChildHandle, Predicate);
+						if (!VisitedProcessIds.Contains(Entry.th32ProcessID))
+						{
+							TerminateProcTreeWithPredicateInternal(ChildProcHandle, Predicate, VisitedProcessIds);
+						}
+						::CloseHandle(ChildProcHandle);
 					}
 				}
 			}
 			while(::Process32Next(SnapShot, &Entry));
 		}
 	}
+	::CloseHandle(SnapShot);
+
+	TerminateProcess(ProcessHandle,0);
+}
 
-	TerminateProcess(ProcessHandle.Get(),0);
+void FWindowsPlatformProcess::TerminateProcTreeWithPredicate(
+	FProcHandle& ProcessHandle,
+	TFunctionRef<bool(uint32 ProcessId, const TCHAR* ApplicationName)> Predicate)
+{
+	TSet<DWORD> VisitedProcessIds;
+	TerminateProcTreeWithPredicateInternal(ProcessHandle.Get(), Predicate, VisitedProcessIds);
 }
 
 uint32 FWindowsPlatformProcess::GetCurrentProcessId()
