Fix WinGDK build by ensuring that Core only depends on libpas on Win64

diff --git a/Engine/Source/Runtime/Core/Core.Build.cs b/Engine/Source/Runtime/Core/Core.Build.cs
--- a/Engine/Source/Runtime/Core/Core.Build.cs
+++ b/Engine/Source/Runtime/Core/Core.Build.cs
@@ -46,6 +46,12 @@ public class Core : ModuleRules
 			DynamicallyLoadedModuleNames.Add("DirectoryWatcher");
 		}
 
+		if (Target.Platform == UnrealTargetPlatform.Win64)
+		{
+			PrivateDefinitions.Add("PLATFORM_BUILDS_LIBPAS=1");
+			PrivateDependencyModuleNames.Add("libpas");
+		}
+			
 		if (Target.Platform.IsInGroup(UnrealPlatformGroup.Windows))
 		{
 			AddEngineThirdPartyPrivateStaticDependencies(Target,
@@ -67,9 +73,6 @@ public class Core : ModuleRules
 				PrivateDefinitions.Add("PLATFORM_BUILDS_MIMALLOC=1");
 			}
 			
-			PrivateDefinitions.Add("PLATFORM_BUILDS_LIBPAS=1");
-			PrivateDependencyModuleNames.Add("libpas");
-			
 			if (Target.WindowsPlatform.bUseBundledDbgHelp)
 			{
 				PublicDelayLoadDLLs.Add("DBGHELP.DLL");
