[Insights] CpuProfilerTrace:
  - Enabled cpu scope tracing of object names (blueprints, functions) by default when app is not running as a commandlet (still requires "cpu" trace channel and -statnamedevents to be enabled). Enabling AssetLoadTime trace channel will enable tracing of object names also in commandlets (ex. cooking).
  - Improved performance when tracing "statnamedevents" cpu scopes (by using a new optimized code path that uses an FName stat id together with a TCHAR* descriptive name).

diff --git a/Engine/Source/Runtime/Core/Private/ProfilingDebugging/CpuProfilerTrace.cpp b/Engine/Source/Runtime/Core/Private/ProfilingDebugging/CpuProfilerTrace.cpp
--- a/Engine/Source/Runtime/Core/Private/ProfilingDebugging/CpuProfilerTrace.cpp
+++ b/Engine/Source/Runtime/Core/Private/ProfilingDebugging/CpuProfilerTrace.cpp
@@ -206,27 +206,39 @@ void FCpuProfilerTrace::OutputBeginDynamicEvent(const TCHAR* Name, const ANSICHA
 	CPUPROFILERTRACE_OUTPUTBEGINEVENT_EPILOGUE();
 }
 
-void FCpuProfilerTrace::OutputBeginDynamicEvent(const FName& Name, const ANSICHAR* File, uint32 Line)
+void FCpuProfilerTrace::OutputBeginDynamicEvent(const FName Name, const ANSICHAR* File, uint32 Line)
+{
+	OutputBeginDynamicEventWithId(Name, nullptr, File, Line);
+}
+
+void FCpuProfilerTrace::OutputBeginDynamicEventWithId(const FName Id, const TCHAR* Name, const ANSICHAR* File, uint32 Line)
 {
 	CPUPROFILERTRACE_OUTPUTBEGINEVENT_PROLOGUE();
-	uint32 SpecId = ThreadBuffer->DynamicFNameScopeNamesMap.FindRef(Name.GetComparisonIndex());
+	uint32 SpecId = ThreadBuffer->DynamicFNameScopeNamesMap.FindRef(Id.GetComparisonIndex());
 	if (!SpecId)
 	{
 		LLM_SCOPE_BYNAME(TEXT("Trace/CpuProfiler"));
-		const FNameEntry* NameEntry = Name.GetDisplayNameEntry();
-		if (NameEntry->IsWide())
+		if (Name != nullptr)
 		{
-			WIDECHAR WideName[NAME_SIZE];
-			NameEntry->GetWideName(WideName);
-			SpecId = OutputEventType(WideName, File, Line);
+			SpecId = OutputEventType(Name, File, Line);
 		}
 		else
 		{
-			ANSICHAR AnsiName[NAME_SIZE];
-			NameEntry->GetAnsiName(AnsiName);
-			SpecId = OutputEventType(AnsiName, File, Line);
+			const FNameEntry* NameEntry = Id.GetDisplayNameEntry();
+			if (NameEntry->IsWide())
+			{
+				WIDECHAR WideName[NAME_SIZE];
+				NameEntry->GetWideName(WideName);
+				SpecId = OutputEventType(WideName, File, Line);
+			}
+			else
+			{
+				ANSICHAR AnsiName[NAME_SIZE];
+				NameEntry->GetAnsiName(AnsiName);
+				SpecId = OutputEventType(AnsiName, File, Line);
+			}
 		}
-		ThreadBuffer->DynamicFNameScopeNamesMap.Add(Name.GetComparisonIndex(), SpecId);
+		ThreadBuffer->DynamicFNameScopeNamesMap.Add(Id.GetComparisonIndex(), SpecId);
 	}
 	CPUPROFILERTRACE_OUTPUTBEGINEVENT_EPILOGUE();
 }
diff --git a/Engine/Source/Runtime/Core/Public/ProfilingDebugging/CpuProfilerTrace.h b/Engine/Source/Runtime/Core/Public/ProfilingDebugging/CpuProfilerTrace.h
--- a/Engine/Source/Runtime/Core/Public/ProfilingDebugging/CpuProfilerTrace.h
+++ b/Engine/Source/Runtime/Core/Public/ProfilingDebugging/CpuProfilerTrace.h
@@ -83,7 +83,17 @@ struct FCpuProfilerTrace
 	 * @param File Source filename
 	 * @param Line Line number in source file
 	 */
-	CORE_API static void OutputBeginDynamicEvent(const FName& Name, const ANSICHAR* File = nullptr, uint32 Line = 0);
+	FORCENOINLINE CORE_API static void OutputBeginDynamicEvent(const FName Name, const ANSICHAR* File = nullptr, uint32 Line = 0);
+	/*
+	 * Output begin event marker for a dynamic event identified by an FName. This is more expensive than
+	 * statically known event names using \ref OutputBeginEvent, but it is faster than \ref OutputBeginDynamicEvent
+	 * that receives ANSICHAR* / TCHAR* name. Must always be matched with an end event.
+	 * @param Id Id of event
+	 * @param Name Name of event
+	 * @param File Source filename
+	 * @param Line Line number in source file
+	 */
+	CORE_API static void OutputBeginDynamicEventWithId(const FName Id, const TCHAR* Name, const ANSICHAR* File = nullptr, uint32 Line = 0);
 	/*
 	 * Output end event marker for static or dynamic event for the currently open scope.
 	 */
diff --git a/Engine/Source/Runtime/Core/Public/Stats/Stats2.h b/Engine/Source/Runtime/Core/Public/Stats/Stats2.h
--- a/Engine/Source/Runtime/Core/Public/Stats/Stats2.h
+++ b/Engine/Source/Runtime/Core/Public/Stats/Stats2.h
@@ -1683,7 +1683,8 @@ public:
 #if CPUPROFILERTRACE_ENABLED
 			if (UE_TRACE_CHANNELEXPR_IS_ENABLED(CpuChannel))
 			{
-				FCpuProfilerTrace::OutputBeginDynamicEvent(InStatId.GetStatDescriptionANSI()); //todo: Could we use FName index as event id?
+				FName StatName = MinimalNameToName(StatMinimalName);
+				FCpuProfilerTrace::OutputBeginDynamicEventWithId(StatName, WCHAR_TO_TCHAR(InStatId.GetStatDescriptionWIDE()));
 				EmittedEvent |= TraceEvent;
 			}
 #endif
diff --git a/Engine/Source/Runtime/CoreUObject/Private/UObject/ObjectBaseUtility.cpp b/Engine/Source/Runtime/CoreUObject/Private/UObject/ObjectBaseUtility.cpp
--- a/Engine/Source/Runtime/CoreUObject/Private/UObject/ObjectBaseUtility.cpp
+++ b/Engine/Source/Runtime/CoreUObject/Private/UObject/ObjectBaseUtility.cpp
@@ -130,7 +130,7 @@ int32 UObjectBaseUtility::GetLinkerLicenseeUEVersion() const
 COREUOBJECT_API FName GetClassTraceScope(const UObjectBaseUtility* Object)
 {
 	UClass* Class = Object->GetClass();
-	if (Class->IsNative())
+	if (Class->IsNative() || !IsRunningCommandlet())
 	{
 		return Class->GetFName();
 	}
@@ -145,13 +145,20 @@ COREUOBJECT_API FName GetClassTraceScope(const UObjectBaseUtility* Object)
 #if STATS && CPUPROFILERTRACE_ENABLED
 void FScopeCycleCounterUObject::StartObjectTrace(const UObjectBaseUtility* Object)
 {
+	bool bTraceObjectFName = !IsRunningCommandlet();
+
 #if LOADTIMEPROFILERTRACE_ENABLED
 	if (UE_TRACE_CHANNELEXPR_IS_ENABLED(AssetLoadTimeChannel))
+	{
+		bTraceObjectFName = true;
+	}
+#endif // LOADTIMEPROFILERTRACE_ENABLED
+
+	if (bTraceObjectFName)
 	{
 		StartTrace(Object->GetFName());
 	}
 	else
-#endif
 	{
 		StartTrace(GetClassTraceScope(Object));
 	}
