Add a global FCbObjectId SessionObjectId usable to coordinate session id's to zenserver.
Change GUID SessionId in FApp to base it on the SessionObjectId to make it easier to correlate session ids in zenserver

diff --git a/Engine/Source/Developer/DerivedDataCache/Private/Http/CurlHttpClient.cpp b/Engine/Source/Developer/DerivedDataCache/Private/Http/CurlHttpClient.cpp
--- a/Engine/Source/Developer/DerivedDataCache/Private/Http/CurlHttpClient.cpp
+++ b/Engine/Source/Developer/DerivedDataCache/Private/Http/CurlHttpClient.cpp
@@ -145,7 +145,7 @@ public:
 	void SetDefaultOptions(CURL* Curl, FCurlHttpHeaders& Headers);
 
 private:
-	const FCbObjectId SessionId = FCbObjectId::NewObjectId();
+	const FCbObjectId SessionId = FApp::GetSessionObjectId();
 	std::atomic<uint32> RequestId = 1;
 	FUtf8StringBuilderBase UserAgent;
 };
diff --git a/Engine/Source/Developer/Zen/Private/ZenServerHttp.cpp b/Engine/Source/Developer/Zen/Private/ZenServerHttp.cpp
--- a/Engine/Source/Developer/Zen/Private/ZenServerHttp.cpp
+++ b/Engine/Source/Developer/Zen/Private/ZenServerHttp.cpp
@@ -283,8 +283,7 @@ namespace UE::Zen {
 	}
 
 	static const char* GetSessionIdHeader() {
-		static FCbObjectId SessionId = FCbObjectId::NewObjectId();
-
+		static FCbObjectId SessionId = FApp::GetSessionObjectId();
 		static const char* HeaderString = [&] {
 			static TAnsiStringBuilder<64> SessionIdHeader;
 			SessionIdHeader << "UE-Session: " << SessionId;
diff --git a/Engine/Source/Developer/Zen/Private/ZenServerInterface.cpp b/Engine/Source/Developer/Zen/Private/ZenServerInterface.cpp
--- a/Engine/Source/Developer/Zen/Private/ZenServerInterface.cpp
+++ b/Engine/Source/Developer/Zen/Private/ZenServerInterface.cpp
@@ -23,7 +23,6 @@
 #include "Misc/PathViews.h"
 #include "Misc/ScopedSlowTask.h"
 #include "Misc/ScopeExit.h"
-#include "Serialization/CompactBinary.h"
 #include "Serialization/CompactBinarySerialization.h"
 #include "Serialization/CompactBinaryValidation.h"
 #include "Serialization/JsonReader.h"
diff --git a/Engine/Source/Runtime/Core/Private/Misc/App.cpp b/Engine/Source/Runtime/Core/Private/Misc/App.cpp
--- a/Engine/Source/Runtime/Core/Private/Misc/App.cpp
+++ b/Engine/Source/Runtime/Core/Private/Misc/App.cpp
@@ -6,6 +6,7 @@
 #include "Misc/ConfigCacheIni.h"
 #include "Misc/CoreDelegates.h"
 #include "Runtime/Launch/Resources/Version.h"
+#include "Serialization/CompactBinary.h"
 #include "HAL/LowLevelMemTracker.h"
 #include "BuildSettings.h"
 #include "UObject/DevObjectVersion.h"
@@ -22,7 +23,7 @@ DEFINE_LOG_CATEGORY_STATIC(LogApp, Log, All);
 bool FApp::bIsDebugGame = false;
 #endif
 
-FGuid FApp::SessionId = FGuid::NewGuid();
+FGuid FApp::SessionId = ToGuid(FApp::GetSessionObjectId());
 FString FApp::SessionName = FString();
 FString FApp::SessionOwner = FString();
 FString FApp::GraphicsRHI = FString();
@@ -466,3 +467,9 @@ FGuid FApp::GetInstanceId()
 	static FGuid InstanceId = FGuid::NewGuid();
 	return InstanceId;
 }
+
+const FCbObjectId& FApp::GetSessionObjectId()
+{
+	static const FCbObjectId SessionObjectId = FCbObjectId::NewObjectId();
+	return SessionObjectId;
+}
diff --git a/Engine/Source/Runtime/Core/Private/Serialization/CompactBinary.cpp b/Engine/Source/Runtime/Core/Private/Serialization/CompactBinary.cpp
--- a/Engine/Source/Runtime/Core/Private/Serialization/CompactBinary.cpp
+++ b/Engine/Source/Runtime/Core/Private/Serialization/CompactBinary.cpp
@@ -106,6 +106,40 @@ FCbObjectId FCbObjectId::NewObjectId()
 	return ObjectId;
 }
 
+FGuid ToGuid(const FCbObjectId& Id)
+{
+	const FCbObjectId::ByteArray& Bytes = Id.GetBytes();
+	uint32 A = static_cast<uint32>(Bytes[0]) << 24 | static_cast<uint32>(Bytes[1]) << 16 | static_cast<uint32>(Bytes[2]) << 8 | static_cast<uint32>(Bytes[3]);
+	uint32 B = static_cast<uint32>(Bytes[4]) << 24 | static_cast<uint32>(Bytes[5]) << 16 | static_cast<uint32>(Bytes[6]) << 8 | static_cast<uint32>(Bytes[7]);
+	uint32 C = static_cast<uint32>(Bytes[8]) << 24 | static_cast<uint32>(Bytes[9]) << 16 | static_cast<uint32>(Bytes[10]) << 8 | static_cast<uint32>(Bytes[11]);
+	uint32 D = (B & 0x0000ff00u) | (C & 0xff000000u);;
+	B = (B & ~0x0000f000u) | 0x00004000u;
+	C = (C & ~0xc0000000u) | 0x80000000u;
+	return FGuid(A, B, C, D);
+}
+
+FCbObjectId FromGuid(const FGuid& Id)
+{
+	uint32 Bits0 = Id.A;
+	uint32 Bits1 = (Id.B & ~0x0000ff00u) | (Id.D & 0x0000ff00u);
+	uint32 Bits2 = (Id.C & ~0xff000000u) | (Id.D & 0xff000000u);
+	FCbObjectId::ByteArray Bytes = {
+		static_cast<uint8>((Bits0 >> 24) & 0xff),
+		static_cast<uint8>((Bits0 >> 16) & 0xff),
+		static_cast<uint8>((Bits0 >> 8) & 0xff),
+		static_cast<uint8>(Bits0 & 0xff),
+		static_cast<uint8>((Bits1 >> 24) & 0xff),
+		static_cast<uint8>((Bits1 >> 16) & 0xff),
+		static_cast<uint8>((Bits1 >> 8) & 0xff),
+		static_cast<uint8>(Bits1 & 0xff),
+		static_cast<uint8>((Bits2 >> 24) & 0xff),
+		static_cast<uint8>((Bits2 >> 16) & 0xff),
+		static_cast<uint8>((Bits2 >> 8) & 0xff),
+		static_cast<uint8>(Bits2 & 0xff)
+	};
+	return FCbObjectId(Bytes);
+}
+
 ///////////////////////////////////////////////////////////////////////////////////////////////////////////////////////
 
 FCbFieldView::FCbFieldView(const void* const Data, ECbFieldType Type)
diff --git a/Engine/Source/Runtime/Core/Public/Misc/App.h b/Engine/Source/Runtime/Core/Public/Misc/App.h
--- a/Engine/Source/Runtime/Core/Public/Misc/App.h
+++ b/Engine/Source/Runtime/Core/Public/Misc/App.h
@@ -23,6 +23,8 @@
 #include "Templates/UnrealTemplate.h"
 #include "UObject/NameTypes.h"
 
+class FCbObjectId;
+
 /**
  * Provides information about the application.
  */
@@ -279,6 +281,14 @@ public:
 		return SessionId;
 	}
 
+	/**
+	 * Gets the identifier of the session that this application is part of as a FCbObject.
+	 *
+	 * @return Session identifier
+	 * @see GetSessionId
+	 */
+	static CORE_API const FCbObjectId& GetSessionObjectId();
+
 	/**
 	 * Gets the name of the session that this application is part of, if any.
 	 *
diff --git a/Engine/Source/Runtime/Core/Public/Serialization/CompactBinary.h b/Engine/Source/Runtime/Core/Public/Serialization/CompactBinary.h
--- a/Engine/Source/Runtime/Core/Public/Serialization/CompactBinary.h
+++ b/Engine/Source/Runtime/Core/Public/Serialization/CompactBinary.h
@@ -535,6 +535,9 @@ inline FCbObjectId::FCbObjectId(const ByteArray& ObjectId)
 	FMemory::Memcpy(Bytes, ObjectId, sizeof(ByteArray));
 }
 
+FGuid			ToGuid(const FCbObjectId& Id);
+FCbObjectId		FromGuid(const FGuid& Id);
+
 ///////////////////////////////////////////////////////////////////////////////////////////////////////////////////////
 
 /** A custom compact binary field type with an integer identifier. */
