A minimal support for task cancellation:
* it's user's decision and responsibility to manage cancellation token lifetime, to check cancellation token, return early, to do any required cleanup, or to do nothing at all
* no way to cancel a task to skip its execution completely
* waiting for a canceled task is blocking until its prerequisites are completed and the task is executed and completed, basically same as for not canceled tasks except a canceled one can quit execution early
* canceling a task doesn't affect its subsequents (unless they use the same cancellation token instance)

diff --git a/Engine/Source/Runtime/Core/Public/Tasks/Task.h b/Engine/Source/Runtime/Core/Public/Tasks/Task.h
--- a/Engine/Source/Runtime/Core/Public/Tasks/Task.h
+++ b/Engine/Source/Runtime/Core/Public/Tasks/Task.h
@@ -473,4 +473,35 @@ namespace UE::Tasks
 			ETaskPriority::Default, // doesn't matter
 			EExtendedTaskPriority::Inline);
 	}
+
+	// support for canceling tasks mid-execution
+	// usage:
+	//		FCancellationToken Token;
+	//		Launch(UE_SOURCE_LOCATION, [&Token] { ... if (Token.IsCanceled()) return; ... });
+	//		Token.Cancel();
+	// * it's user's decision and responsibility to manage cancellation token lifetime, to check cancellation token, return early, 
+	//		to do any required cleanup, or to do nothing at all
+	// * no way to cancel a task to skip its execution completely
+	// * waiting for a canceled task is blocking until its prerequisites are completed and the task is executed and completed, 
+	//		basically same as for not canceled tasks except a canceled one can quit execution early
+	// * canceling a task doesn't affect its subsequents (unless they use the same cancellation token instance)
+	class FCancellationToken
+	{
+	public:
+		UE_NONCOPYABLE(FCancellationToken);
+		FCancellationToken() = default;
+
+		void Cancel()
+		{
+			bCanceled = true;
+		}
+
+		bool IsCanceled() const
+		{
+			return bCanceled;
+		}
+
+	private:
+		std::atomic<bool> bCanceled{ false };
+	};
 }
diff --git a/Engine/Source/Runtime/Core/Tests/Tasks/TasksTest.cpp b/Engine/Source/Runtime/Core/Tests/Tasks/TasksTest.cpp
--- a/Engine/Source/Runtime/Core/Tests/Tasks/TasksTest.cpp
+++ b/Engine/Source/Runtime/Core/Tests/Tasks/TasksTest.cpp
@@ -1428,6 +1428,31 @@ namespace UE { namespace TasksTests
 			check(Task.GetResult().Int == 42 && Task.GetResult().Str == TEXT("Test"));
 		}
 	}
+
+	TEST_CASE_NAMED(FTasksCancellation, "System::Core::Tasks::Cancellation", "[.][ApplicationContextMask][EngineFilter]")
+	{
+		{
+			FCancellationToken CancellationToken;
+			FTaskEvent BlockExecution{ UE_SOURCE_LOCATION };
+
+			// check that a task sees cancellation request
+			FTask Task1 = Launch(UE_SOURCE_LOCATION,
+				[&CancellationToken, BlockExecution]
+				{
+					BlockExecution.Wait();
+					verify(CancellationToken.IsCanceled());
+				}
+			);
+			// same token can be used with multiple tasks to cancel them all
+			// a task can ignore cancellation request
+			FTask Task2 = Launch(UE_SOURCE_LOCATION, [&CancellationToken] {});
+
+			CancellationToken.Cancel();
+			BlockExecution.Trigger();
+
+			Wait(TArray{ Task1, Task2 });
+		}
+	}
 }}
 
 #endif // WITH_TESTS
