When running in the iOS simulator, set the correct device type.

diff --git a/Engine/Source/Runtime/Core/Private/IOS/IOSPlatformMisc.cpp b/Engine/Source/Runtime/Core/Private/IOS/IOSPlatformMisc.cpp
--- a/Engine/Source/Runtime/Core/Private/IOS/IOSPlatformMisc.cpp
+++ b/Engine/Source/Runtime/Core/Private/IOS/IOSPlatformMisc.cpp
@@ -368,6 +368,10 @@ FString GetIOSDeviceIDString()
 	static bool bCached = false;
 	if (!bCached)
 	{
+#if WITH_IOS_SIMULATOR
+		NSString* ModelID = [[NSProcessInfo processInfo] environment][@"SIMULATOR_MODEL_IDENTIFIER"];
+		CachedResult = FString(ModelID);
+#else
 		// get the device hardware type string length
 		size_t DeviceIDLen;
 		sysctlbyname("hw.machine", NULL, &DeviceIDLen, NULL, 0);
@@ -377,7 +381,6 @@ FString GetIOSDeviceIDString()
 		sysctlbyname("hw.machine", DeviceID, &DeviceIDLen, NULL, 0);
 
 		CachedResult = ANSI_TO_TCHAR(DeviceID);
-		bCached = true;
 
 		free(DeviceID);
 		
@@ -393,6 +396,8 @@ FString GetIOSDeviceIDString()
 			CachedResult = TEXT("iPhone0,1");
 #endif
 		}
+#endif
+		bCached = true;
 	}
 
 	return CachedResult;
