Only set device type for iOS when running in a simulator

diff --git a/Engine/Source/Runtime/Core/Private/IOS/IOSPlatformMisc.cpp b/Engine/Source/Runtime/Core/Private/IOS/IOSPlatformMisc.cpp
--- a/Engine/Source/Runtime/Core/Private/IOS/IOSPlatformMisc.cpp
+++ b/Engine/Source/Runtime/Core/Private/IOS/IOSPlatformMisc.cpp
@@ -368,10 +368,6 @@ FString GetIOSDeviceIDString()
 	static bool bCached = false;
 	if (!bCached)
 	{
-#if WITH_IOS_SIMULATOR
-		NSString* ModelID = [[NSProcessInfo processInfo] environment][@"SIMULATOR_MODEL_IDENTIFIER"];
-		CachedResult = FString(ModelID);
-#else
 		// get the device hardware type string length
 		size_t DeviceIDLen;
 		sysctlbyname("hw.machine", NULL, &DeviceIDLen, NULL, 0);
@@ -381,6 +377,7 @@ FString GetIOSDeviceIDString()
 		sysctlbyname("hw.machine", DeviceID, &DeviceIDLen, NULL, 0);
 
 		CachedResult = ANSI_TO_TCHAR(DeviceID);
+		bCached = true;
 
 		free(DeviceID);
 		
@@ -393,11 +390,10 @@ FString GetIOSDeviceIDString()
 #elif PLATFORM_TVOS
 			CachedResult = TEXT("AppleTV0,1");
 #else
-			CachedResult = TEXT("iPhone0,1");
+			NSString* ModelID = [[NSProcessInfo processInfo] environment][@"SIMULATOR_MODEL_IDENTIFIER"];
+			CachedResult = FString(ModelID);
 #endif
 		}
-#endif
-		bCached = true;
 	}
 
 	return CachedResult;
