[Insights] MemoryInsights: Fixed missing ProgramSize allocations for Unix platform.

diff --git a/Engine/Source/Runtime/Core/Private/ProfilingDebugging/Unix/UnixModuleDiagnostics.cpp b/Engine/Source/Runtime/Core/Private/ProfilingDebugging/Unix/UnixModuleDiagnostics.cpp
--- a/Engine/Source/Runtime/Core/Private/ProfilingDebugging/Unix/UnixModuleDiagnostics.cpp
+++ b/Engine/Source/Runtime/Core/Private/ProfilingDebugging/Unix/UnixModuleDiagnostics.cpp
@@ -1,11 +1,17 @@
 // Copyright Epic Games, Inc. All Rights Reserved.
 
 #include "CoreTypes.h"
+#include "HAL/LowLevelMemTracker.h"
 #include "HAL/PlatformProcess.h"
 #include "Math/UnrealMathUtility.h"
 #include "Misc/Paths.h"
 #include "ProfilingDebugging/ModuleDiagnostics.h"
+#include "ProfilingDebugging/MemoryTrace.h"
+#include "ProfilingDebugging/MetadataTrace.h"
+
+#if !UE_BUILD_SHIPPING
 #include <link.h>
+#endif
 
 ////////////////////////////////////////////////////////////////////////////////
 void Modules_Initialize()
@@ -14,10 +20,12 @@ void Modules_Initialize()
 	using namespace UE::Trace;
 
 	constexpr uint32 SizeOfSymbolFormatString = 5;
-	UE_TRACE_LOG(Diagnostics, ModuleInit, ModuleChannel, sizeof(TCHAR) * SizeOfSymbolFormatString)
+	UE_TRACE_LOG(Diagnostics, ModuleInit, ModuleChannel, sizeof(ANSICHAR) * SizeOfSymbolFormatString)
 		<< ModuleInit.SymbolFormat("dwarf", SizeOfSymbolFormatString)
 		<< ModuleInit.ModuleBaseShift(uint8(0));
 
+	HeapId ProgramHeapId = MemoryTrace_HeapSpec(EMemoryTraceRootHeap::SystemMemory, TEXT("Program"), EMemoryTraceHeapFlags::NeverFrees);
+
 	auto IterateCallback = [](struct dl_phdr_info *Info, size_t /*size*/, void *Data)
 	{
 		int32 TotalMemSize = 0;
@@ -68,8 +76,19 @@ void Modules_Initialize()
 			<< ModuleLoad.Size(TotalMemSize)
 			<< ModuleLoad.ImageId(BuildId, BuildIdSize);
 
+#if UE_MEMORY_TRACE_ENABLED
+		{
+			UE_TRACE_METADATA_CLEAR_SCOPE();
+			LLM(UE_MEMSCOPE(ELLMTag::ProgramSize));
+			MemoryTrace_Alloc(RealBase, TotalMemSize, 1);
+			HeapId ProgramHeapId = *(HeapId*)Data;
+			MemoryTrace_MarkAllocAsHeap(RealBase, ProgramHeapId);
+			MemoryTrace_Alloc(RealBase, TotalMemSize, 1);
+		}
+#endif // UE_MEMORY_TRACE_ENABLED
+
 		return 0;
 	};
-	dl_iterate_phdr(IterateCallback, nullptr);
+	dl_iterate_phdr(IterateCallback, &ProgramHeapId);
 #endif
 }
