Rename TGuardedSignedInt to TGuardedInt and add support for unsigned integers.
Factor some of the arithmetic operators out of TGuardedInt into AddAndCheckForOverflow et al functions in FPlatformMath so they can use the __builtin_*_overflow functions on clang that compile to much more efficient code than the portable C code does in most cases.

diff --git a/Engine/Source/Runtime/Core/Public/Clang/ClangPlatformMath.h b/Engine/Source/Runtime/Core/Public/Clang/ClangPlatformMath.h
--- a/Engine/Source/Runtime/Core/Public/Clang/ClangPlatformMath.h
+++ b/Engine/Source/Runtime/Core/Public/Clang/ClangPlatformMath.h
@@ -98,4 +98,37 @@ struct FClangPlatformMath : public FGenericPlatformMath
 	{
 		return 63 - __builtin_clzll(Value | 1);
 	}
+	
+	/**
+	 * Adds two integers of any integer type, checking for overflow.
+	 * If there was overflow, it returns false, and OutResult may or may not be written.
+	 * If there wasn't overflow, it returns true, and the result of the addition is written to OutResult.
+	 */
+	template<typename IntType, std::enable_if_t<std::is_integral_v<IntType>>* = nullptr>
+	static FORCEINLINE bool AddAndCheckForOverflow(IntType A, IntType B, IntType& OutResult)
+	{
+		return !__builtin_add_overflow(A, B, &OutResult);
+	}
+
+	/**
+	 * Subtracts two integers of any integer type, checking for overflow.
+	 * If there was overflow, it returns false, and OutResult may or may not be written.
+	 * If there wasn't overflow, it returns true, and the result of the subtraction is written to OutResult.
+	 */
+	template<typename IntType, std::enable_if_t<std::is_integral_v<IntType>>* = nullptr>
+	static FORCEINLINE bool SubtractAndCheckForOverflow(IntType A, IntType B, IntType& OutResult)
+	{
+		return !__builtin_sub_overflow(A, B, &OutResult);
+	}
+
+	/**
+	 * Multiplies two integers of any integer type, checking for overflow.
+	 * If there was overflow, it returns false, and OutResult may or may not be written.
+	 * If there wasn't overflow, it returns true, and the result of the multiplication is written to OutResult.
+	 */
+	template<typename IntType, std::enable_if_t<std::is_integral_v<IntType>>* = nullptr>
+	static FORCEINLINE bool MultiplyAndCheckForOverflow(IntType A, IntType B, IntType& OutResult)
+	{
+		return !__builtin_mul_overflow(A, B, &OutResult);
+	}
 };
diff --git a/Engine/Source/Runtime/Core/Public/GenericPlatform/GenericPlatformMath.h b/Engine/Source/Runtime/Core/Public/GenericPlatform/GenericPlatformMath.h
--- a/Engine/Source/Runtime/Core/Public/GenericPlatform/GenericPlatformMath.h
+++ b/Engine/Source/Runtime/Core/Public/GenericPlatform/GenericPlatformMath.h
@@ -16,6 +16,8 @@
 #include "Templates/UnrealTypeTraits.h"
 #include "Templates/ResolveTypeAmbiguity.h"
 #include "Templates/TypeCompatibleBytes.h"
+#include <limits>
+#include <type_traits>
 
 #if PLATFORM_HAS_FENV_H 
 #include <fenv.h>
@@ -1019,6 +1021,165 @@ struct FGenericPlatformMath
 	static void AutoTest();
 #endif
 
+	/**
+	 * Adds two integers of any integer type, checking for overflow.
+	 * If there was overflow, it returns false, and OutResult may or may not be written.
+	 * If there wasn't overflow, it returns true, and the result of the addition is written to OutResult.
+	 */
+	template<typename IntType, std::enable_if_t<std::is_integral_v<IntType>>* = nullptr>
+	static FORCEINLINE bool AddAndCheckForOverflow(IntType A, IntType B, IntType& OutResult)
+	{
+		// This follows Hacker's Delight, Chapter 2-12
+		// Signed->unsigned conversion and unsigned addition have defined behavior always
+		typedef typename std::make_unsigned_t<IntType> UnsignedType;
+		const UnsignedType UnsignedA = static_cast<UnsignedType>(A);
+		const UnsignedType UnsignedB = static_cast<UnsignedType>(B);
+		const UnsignedType UnsignedSum = UnsignedA + UnsignedB;
+
+		if constexpr (std::is_signed_v<IntType>)
+		{
+			// Check for signed overflow.
+			// The underlying logic here is pretty simple: if A and B had opposite signs, their sum can't
+			// overflow. If they had the same sign and the sum has the opposite value in the sign bit, we
+			// had an overflow. (See Hacker's Delight Chapter 2-12 for more details.)
+			constexpr UnsignedType UnsignedMSB = static_cast<UnsignedType>(std::numeric_limits<IntType>::min());
+			if((UnsignedSum ^ UnsignedA) & (UnsignedSum ^ UnsignedB) & UnsignedMSB)
+			{
+				return false;
+			}
+			else
+			{
+				OutResult = A + B;
+				return true;
+			}
+		}
+		else
+		{
+			// Iff there was overflow, the modular sum must be less than both the operands.
+			if (UnsignedSum >= UnsignedA)
+			{
+				OutResult = UnsignedSum;
+				return true;
+			}
+			else
+			{
+				return false;
+			}
+		}
+	}
+
+	/**
+	 * Subtracts two integers of any integer type, checking for overflow.
+	 * If there was overflow, it returns false, and OutResult may or may not be written.
+	 * If there wasn't overflow, it returns true, and the result of the subtraction is written to OutResult.
+	 */
+	template<typename IntType, std::enable_if_t<std::is_integral_v<IntType>>* = nullptr>
+	static FORCEINLINE bool SubtractAndCheckForOverflow(IntType A, IntType B, IntType& OutResult)
+	{
+		// This follows Hacker's Delight, Chapter 2-12
+		// Signed->unsigned conversion and unsigned addition have defined behavior always
+		typedef typename std::make_unsigned_t<IntType> UnsignedType;
+		const UnsignedType UnsignedA = static_cast<UnsignedType>(A);
+		const UnsignedType UnsignedB = static_cast<UnsignedType>(B);
+		const UnsignedType UnsignedDiff = UnsignedA - UnsignedB;
+
+		if constexpr (std::is_signed_v<IntType>)
+		{
+			// Check for signed overflow.
+			// If A and B have the same sign, the difference can't overflow. Therefore, we test for cases
+			// where the sign bit differs meaning ((UnsignedA ^ UnsignedB) & UnsignedMSB) != 0, and
+			// simultaneously the sign of the difference differs from the sign of the minuend (which should
+			// keep its sign when we're subtracting a value of the opposite sign), meaning
+			// ((UnsignedDiff ^ UnsignedA) & UnsignedMSB) != 0. Combining the two yields:
+			constexpr UnsignedType UnsignedMSB = static_cast<UnsignedType>(std::numeric_limits<IntType>::min());
+			if ((UnsignedA ^ UnsignedB) & (UnsignedDiff ^ UnsignedA) & UnsignedMSB)
+			{
+				return false;
+			}
+			else
+			{
+				OutResult = A - B;
+				return true;
+			}
+		}
+		else
+		{
+			OutResult = UnsignedDiff;
+			return A >= B;
+		}
+	}
+
+	/**
+	 * Multiplies two integers of any integer type, checking for overflow.
+	 * If there was overflow, it returns false, and OutResult may or may not be written.
+	 * If there wasn't overflow, it returns true, and the result of the multiplication is written to OutResult.
+	 */
+	template<typename IntType, std::enable_if_t<std::is_integral_v<IntType>>* = nullptr>
+	static FORCEINLINE bool MultiplyAndCheckForOverflow(IntType A, IntType B, IntType& OutResult)
+	{
+		// Handle the case where the second factor is 0 specially (why will become clear in a minute).
+		if (B == 0)
+		{
+			// Anything times 0 is 0.
+			OutResult = 0;
+			return true;
+		}
+
+		if constexpr (std::is_signed_v<IntType>)
+		{
+			// The overflow check is annoying and expensive, but the basic idea is fairly simple:
+			// reduce to an unsigned check of the absolute values. (Again the basic algorithm is
+			// in Hacker's Delight, Chapter 2-12).
+			//
+			// We need the absolute value of the product to be <=MaxValue when the result is positive
+			// (signs of factors same) and <= -MinValue = MaxValue + 1 if the result is negative
+			// (signs of factors opposite).
+			typedef typename std::make_unsigned_t<IntType> UnsignedType;
+			UnsignedType UnsignedA = static_cast<UnsignedType>(A);
+			UnsignedType UnsignedB = static_cast<UnsignedType>(B);
+			bool bSignsDiffer = false;
+
+			// Determine the unsigned absolute values of A and B carefully (note we can't negate signed
+			// A or B, because negating MinValue is UB). We can however subtract their
+			// unsigned values from 0 if the original value was less than zero. While doing this, also
+			// keep track of the sign parity.
+			if (A < 0)
+			{
+				UnsignedA = UnsignedType(0) - UnsignedA;
+				bSignsDiffer = !bSignsDiffer;
+			}
+
+			if (B < 0)
+			{
+				UnsignedB = UnsignedType(0) - UnsignedB;
+				bSignsDiffer = !bSignsDiffer;
+			}
+
+			// Determine the unsigned product bound we need based on whether the signs were same or different.
+			const UnsignedType ProductBound = UnsignedType(std::numeric_limits<IntType>::max()) + (bSignsDiffer ? 1 : 0);
+
+			// We're now in the unsigned case, 0 <= UnsignedA, 0 < UnsignedB (we established b != 0), and for
+			// there not to be overflows we need
+			//   a * b <= ProductBound
+			// <=> a <= ProductBound/b
+			// <=> a <= floor(ProductBound/b)   since a is integer
+			if (UnsignedA <= ProductBound / UnsignedB)
+			{
+				OutResult = A * B;
+				return true;
+			}
+			else
+			{
+				return false;
+			}
+		}
+		else
+		{
+			OutResult = A * B;
+			return OutResult / B == A;
+		}
+	}
+
 private:
 
 	/** Error reporting for Fmod. Not inlined to avoid compilation issues and avoid all the checks and error reporting at all callsites. */
diff --git a/Engine/Source/Runtime/Core/Public/Math/GuardedInt.h b/Engine/Source/Runtime/Core/Public/Math/GuardedInt.h
--- a/Engine/Source/Runtime/Core/Public/Math/GuardedInt.h
+++ b/Engine/Source/Runtime/Core/Public/Math/GuardedInt.h
@@ -3,8 +3,9 @@
 #pragma once
 
 #include "CoreTypes.h"
+#include "HAL/PlatformMath.h"
 #include "Misc/AssertionMacros.h"
-#include "Math/NumericLimits.h"
+#include <limits>
 #include <type_traits>
 
 /**
@@ -28,7 +29,7 @@
  *
  * Checked integers keep both the integer value and a "valid" flag. Default-constructed guarded ints
  * are invalid, and guarded integers constructed from an integer value are valid and hold that value.
- * Guarded integers are somewhat analogous to a TOptional<SignedType> in semantics, and borrow some of
+ * Guarded integers are somewhat analogous to a TOptional<IntType> in semantics, and borrow some of
  * the function names.
  *
  * The main feature of guarded integers is that all arithmetic on them is overflow-checked. Any arithmetic
@@ -48,32 +49,31 @@
  * As-is, the current approach is not the fastest, but it's not presently intended to be used in contexts
  * where speed of arithmetic operations is a major concern.
  */
-template<typename SignedType>
-class TGuardedSignedInt
+template<typename IntType>
+class TGuardedInt
 {
 private:
-	static_assert(std::is_integral_v<SignedType> && std::is_signed_v<SignedType>, "Only defined for signed ints");
-	typedef typename std::make_unsigned_t<SignedType> UnsignedType;
+	static_assert(std::is_integral_v<IntType>, "Only defined for integer types");
+	typedef typename std::make_unsigned_t<IntType> UnsignedType;
 
-	static constexpr SignedType MinValue = TNumericLimits<SignedType>::Min();
-	static constexpr SignedType MaxValue = TNumericLimits<SignedType>::Max();
-	static constexpr SignedType NumBits = SignedType((sizeof(SignedType) / sizeof(char)) * 8); // Using sizeof to guess the bit count, ugh.
-	static constexpr UnsignedType UnsignedMSB = (UnsignedType)MinValue; // Assuming two's complement
+	static constexpr IntType MinValue = std::numeric_limits<IntType>::min();
+	static constexpr IntType MaxValue = std::numeric_limits<IntType>::max();
+	static constexpr IntType NumBits = IntType((sizeof(IntType) / sizeof(char)) * 8); // Using sizeof to guess the bit count, ugh.
 
-	SignedType Value = 0;
+	IntType Value = 0;
 	bool bIsValid = false;
 
 public:
-	/** Construct a TGuardedSignedInt with an invalid value. */
-	TGuardedSignedInt() = default;
+	/** Construct a TGuardedInt with an invalid value. */
+	TGuardedInt() = default;
 
-	/** Construct a TGuardedSignedInt from a regular signed integer value. If it's out of range, it results in an invalid value. */
+	/** Construct a TGuardedInt from a regular signed integer value. If it's out of range, it results in an invalid value. */
 	template<
 		typename ValueType,
 		std::enable_if_t<std::is_integral_v<ValueType>>* = nullptr
 	>
-	explicit TGuardedSignedInt(ValueType InValue)
-		: Value((SignedType)InValue), bIsValid(false)
+	explicit TGuardedInt(ValueType InValue)
+		: Value((IntType)InValue), bIsValid(false)
 	{
 		if constexpr (std::is_signed_v<ValueType>)
 		{
@@ -85,17 +85,17 @@ public:
 		}
 	}
 
-	/** Copy-construct a TGuardedSignedInt from another of matching type. */
-	TGuardedSignedInt(const TGuardedSignedInt& Other) = default;
+	/** Copy-construct a TGuardedInt from another of matching type. */
+	TGuardedInt(const TGuardedInt& Other) = default;
 
-	/** Assign a TGuardedSignedInt to another. */
-	TGuardedSignedInt& operator=(const TGuardedSignedInt& Other) = default;
+	/** Assign a TGuardedInt to another. */
+	TGuardedInt& operator=(const TGuardedInt& Other) = default;
 
 	/** @return true if current value is valid (assigned and no overflows or other errors occurred), false otherwise. */
 	bool IsValid() const { return bIsValid; }
 
 	/** @return The value if valid, DefaultValue otherwise. */
-	const SignedType Get(const SignedType DefaultValue) const
+	const IntType Get(const IntType DefaultValue) const
 	{
 		return IsValid() ? Value : DefaultValue;
 	}
@@ -103,14 +103,14 @@ public:
 	/** @return Returns the value if valid, DefaultValue otherwise, but also check()s that the value is valid.
 	 * Intended for cases where the value is not expected to be invalid.
 	 */
-	const SignedType GetChecked(const SignedType DefaultValue = 0) const
+	const IntType GetChecked(const IntType DefaultValue = 0) const
 	{
-		checkf(IsValid(), TEXT("Invalid value in TGuardedSignedInt::GetChecked."));
+		checkf(IsValid(), TEXT("Invalid value in TGuardedInt::GetChecked."));
 		return Get(DefaultValue);
 	}
 
 	/** @return true if *this and Other are either both invalid or both valid and have the same value, false otherwise. */
-	bool operator ==(const TGuardedSignedInt Other) const
+	bool operator ==(const TGuardedInt Other) const
 	{
 		if (bIsValid != Other.bIsValid)
 		{
@@ -121,7 +121,7 @@ public:
 	}
 
 	/** @return true if *this and Other either have different "valid" states or are both valid and have different values, false otherwise (logical negation of ==). */
-	bool operator !=(const TGuardedSignedInt Other) const
+	bool operator !=(const TGuardedInt Other) const
 	{
 		if (bIsValid != Other.bIsValid)
 		{
@@ -135,13 +135,13 @@ public:
 	// to decide what to do about validity as well. Instead, do this.
 
 	/** @return true if *this and Other are both valid so they can be compared. */
-	bool ComparisonValid(const TGuardedSignedInt Other) const { return bIsValid && Other.bIsValid; }
+	bool ComparisonValid(const TGuardedInt Other) const { return bIsValid && Other.bIsValid; }
 
 	/** @return true if *this and Other are both valid and *this is less than Other. */
 	template<typename ValueType>
 	bool ValidAndLessThan(const ValueType Other) const
 	{
-		TGuardedSignedInt CheckedOther{ Other };
+		TGuardedInt CheckedOther{ Other };
 		return ComparisonValid(CheckedOther) && Value < CheckedOther.Value;
 	}
 
@@ -149,7 +149,7 @@ public:
 	template<typename ValueType>
 	bool ValidAndLessOrEqual(const ValueType Other) const
 	{
-		TGuardedSignedInt CheckedOther{ Other };
+		TGuardedInt CheckedOther{ Other };
 		return ComparisonValid(CheckedOther) && Value <= CheckedOther.Value;
 	}
 
@@ -157,7 +157,7 @@ public:
 	template<typename ValueType>
 	bool ValidAndGreaterThan(const ValueType Other) const
 	{
-		TGuardedSignedInt CheckedOther{ Other };
+		TGuardedInt CheckedOther{ Other };
 		return ComparisonValid(CheckedOther) && Value > CheckedOther.Value;
 	}
 
@@ -165,7 +165,7 @@ public:
 	template<typename ValueType>
 	bool ValidAndGreaterOrEqual(const ValueType Other) const
 	{
-		TGuardedSignedInt CheckedOther{ Other };
+		TGuardedInt CheckedOther{ Other };
 		return ComparisonValid(CheckedOther) && Value >= CheckedOther.Value;
 	}
 
@@ -173,7 +173,7 @@ public:
 	template<typename ValueType>
 	bool InvalidOrLessThan(const ValueType Other) const
 	{
-		TGuardedSignedInt CheckedOther{ Other };
+		TGuardedInt CheckedOther{ Other };
 		return !ComparisonValid(CheckedOther) || Value < CheckedOther.Value;
 	}
 
@@ -181,7 +181,7 @@ public:
 	template<typename ValueType>
 	bool InvalidOrLessOrEqual(const ValueType Other) const
 	{
-		TGuardedSignedInt CheckedOther{ Other };
+		TGuardedInt CheckedOther{ Other };
 		return !ComparisonValid(CheckedOther) || Value <= CheckedOther.Value;
 	}
 
@@ -189,7 +189,7 @@ public:
 	template<typename ValueType>
 	bool InvalidOrGreaterThan(const ValueType Other) const
 	{
-		TGuardedSignedInt CheckedOther{ Other };
+		TGuardedInt CheckedOther{ Other };
 		return !ComparisonValid(CheckedOther) || Value > CheckedOther.Value;
 	}
 
@@ -197,240 +197,200 @@ public:
 	template<typename ValueType>
 	bool InvalidOrGreaterOrEqual(const ValueType Other) const
 	{
-		TGuardedSignedInt CheckedOther{ Other };
+		TGuardedInt CheckedOther{ Other };
 		return !ComparisonValid(CheckedOther) || Value >= CheckedOther.Value;
 	}
 
 	// Arithmetic operations
 
 	/** @return The negated value. */
-	TGuardedSignedInt operator-() const
+	TGuardedInt operator-() const
 	{
-		// Unary negation (for two's complement) overflows iff the operand is MinValue.
-		return (bIsValid && Value > MinValue) ? TGuardedSignedInt(-Value) : TGuardedSignedInt();
+		if constexpr (std::is_signed_v<IntType>)
+		{
+			// Unary negation (for two's complement) overflows iff the operand is MinValue.
+			return (bIsValid && Value > MinValue) ? TGuardedInt(-Value) : TGuardedInt();
+		}
+		else
+		{
+			// Negating anything but zero overflows unsigned integers.
+			return (bIsValid && Value == 0) ? TGuardedInt(-Value) : TGuardedInt();
+		}
 	}
 
 	/** @return The sum of the two operands. */
-	TGuardedSignedInt operator +(const TGuardedSignedInt Other) const
+	TGuardedInt operator +(const TGuardedInt Other) const
 	{
 		// Any sum involving an invalid value is invalid.
 		if (!bIsValid || !Other.bIsValid)
 		{
-			return TGuardedSignedInt();
-		}
-
-		// This follows Hacker's Delight, Chapter 2-12
-		// Signed->unsigned conversion and unsigned addition have defined behavior always
-		const UnsignedType UnsignedA = (UnsignedType)Value;
-		const UnsignedType UnsignedB = (UnsignedType)Other.Value;
-		const UnsignedType UnsignedSum = UnsignedA + UnsignedB;
-
-		// Check for signed overflow.
-		// The underlying logic here is pretty simple: if A and B had opposite signs, their sum can't
-		// overflow. If they had the same sign and the sum has the opposite value in the sign bit, we
-		// had an overflow. (See Hacker's Delight Chapter 2-12 for more details.)
-		if ((UnsignedSum ^ UnsignedA) & (UnsignedSum ^ UnsignedB) & UnsignedMSB)
-		{
-			return TGuardedSignedInt();
+			return TGuardedInt();
 		}
 
-		return TGuardedSignedInt(Value + Other.Value);
+		IntType Result;
+		return FPlatformMath::AddAndCheckForOverflow<IntType>(Value, Other.Value, Result)
+			? TGuardedInt(Result)
+			: TGuardedInt();
 	}
 
 	/** @return The difference between the two operands. */
-	TGuardedSignedInt operator -(const TGuardedSignedInt Other) const
+	TGuardedInt operator -(const TGuardedInt Other) const
 	{
 		// Any difference involving an invalid value is invalid.
 		if (!bIsValid || !Other.bIsValid)
 		{
-			return TGuardedSignedInt();
+			return TGuardedInt();
 		}
-
-		// This follows Hacker's Delight, Chapter 2-12
-		// Signed->unsigned conversion and unsigned subtraction have defined behavior always
-		const UnsignedType UnsignedA = (UnsignedType)Value;
-		const UnsignedType UnsignedB = (UnsignedType)Other.Value;
-		const UnsignedType UnsignedDiff = UnsignedA - UnsignedB;
-
-		// Check for signed overflow.
-		// If A and B have the same sign, the difference can't overflow. Therefore, we test for cases
-		// where the sign bit differs meaning ((UnsignedA ^ UnsignedB) & UnsignedMSB) != 0, and
-		// simultaneously the sign of the difference differs from the sign of the minuend (which should
-		// keep its sign when we're subtracting a value of the opposite sign), meaning
-		// ((UnsignedDiff ^ UnsignedA) & UnsignedMSB) != 0. Combining the two yields:
-		if ((UnsignedA ^ UnsignedB) & (UnsignedDiff ^ UnsignedA) & UnsignedMSB)
-		{
-			return TGuardedSignedInt();
-		}
-
-		return TGuardedSignedInt(Value - Other.Value);
+		
+		IntType Result;
+		return FPlatformMath::SubtractAndCheckForOverflow<IntType>(Value, Other.Value, Result)
+			? TGuardedInt(Result)
+			: TGuardedInt();
 	}
 
 	/** @return The product of the two operands. */
-	TGuardedSignedInt operator *(const TGuardedSignedInt Other) const
+	TGuardedInt operator *(const TGuardedInt Other) const
 	{
 		// Any product involving invalid values is invalid.
 		if (!bIsValid || !Other.bIsValid)
 		{
-			return TGuardedSignedInt();
+			return TGuardedInt();
 		}
 
-		// Handle the case where the second factor is 0 specially (why will become clear in a minute).
-		if (Other.Value == 0)
-		{
-			// Anything times 0 is 0.
-			return TGuardedSignedInt(0);
-		}
-
-		// The overflow check is annoying and expensive, but the basic idea is fairly simple:
-		// reduce to an unsigned check of the absolute values. (Again the basic algorithm is
-		// in Hacker's Delight, Chapter 2-12).
-		//
-		// We need the absolute value of the product to be <=MaxValue when the result is positive
-		// (signs of factors same) and <= -MinValue = MaxValue + 1 if the result is negative
-		// (signs of factors opposite).
-		UnsignedType UnsignedA = (UnsignedType)Value;
-		UnsignedType UnsignedB = (UnsignedType)Other.Value;
-		bool bSignsDiffer = false;
-
-		// Determine the unsigned absolute values of A and B carefully (note we can't negate signed
-		// Value or Other.Value, because negating MinValue is UB). We can however subtract their
-		// unsigned values from 0 if the original value was less than zero. While doing this, also
-		// keep track of the sign parity.
-		if (Value < 0)
-		{
-			UnsignedA = UnsignedType(0) - UnsignedA;
-			bSignsDiffer = !bSignsDiffer;
-		}
-
-		if (Other.Value < 0)
-		{
-			UnsignedB = UnsignedType(0) - UnsignedB;
-			bSignsDiffer = !bSignsDiffer;
-		}
-
-		// Determine the unsigned product bound we need based on whether the signs were same or different.
-		const UnsignedType ProductBound = UnsignedType(MaxValue) + (bSignsDiffer ? 1 : 0);
-
-		// We're now in the unsigned case, 0 <= UnsignedA, 0 < UnsignedB (we established b != 0), and for
-		// there not to be overflows we need
-		//   a * b <= ProductBound
-		// <=> a <= ProductBound/b
-		// <=> a <= floor(ProductBound/b)   since a is integer
-		return (UnsignedA <= ProductBound / UnsignedB) ? TGuardedSignedInt(Value * Other.Value) : TGuardedSignedInt();
+		IntType Result;
+		return FPlatformMath::MultiplyAndCheckForOverflow<IntType>(Value, Other.Value, Result)
+			? TGuardedInt(Result)
+			: TGuardedInt();
 	}
 
 	/** @return The quotient when dividing *this by Other. */
-	TGuardedSignedInt operator /(const TGuardedSignedInt Other) const
+	TGuardedInt operator /(const TGuardedInt Other) const
 	{
 		// Any quotient involving invalid values is invalid.
 		if (!bIsValid || !Other.bIsValid)
 		{
-			return TGuardedSignedInt();
+			return TGuardedInt();
 		}
 
 		// Luckily for us, division generally makes things smaller, so there's only two things to watch
 		// out for: division by zero is not allowed, and division of MinValue by -1 would give -MinValue
 		// which overflows. All other combinations are fine.
-		if (Other.Value == 0 || (Value == MinValue && Other.Value == -1))
+		if (Other.Value == 0 || (std::is_signed_v<IntType> && Value == MinValue && Other.Value == -1))
 		{
-			return TGuardedSignedInt();
+			return TGuardedInt();
 		}
 
-		return TGuardedSignedInt(Value / Other.Value);
+		return TGuardedInt(Value / Other.Value);
 	}
 
 	/** @return The remainder when dividing *this by Other. */
-	TGuardedSignedInt operator %(const TGuardedSignedInt Other) const
+	TGuardedInt operator %(const TGuardedInt Other) const
 	{
 		// Any quotient involving invalid values is invalid.
 		if (!bIsValid || !Other.bIsValid)
 		{
-			return TGuardedSignedInt();
+			return TGuardedInt();
 		}
 
 		// Same error cases as for division.
-		if (Other.Value == 0 || (Value == MinValue && Other.Value == -1))
+		if (Other.Value == 0 || (std::is_signed_v<IntType> && Value == MinValue && Other.Value == -1))
 		{
-			return TGuardedSignedInt();
+			return TGuardedInt();
 		}
 
-		return TGuardedSignedInt(Value % Other.Value);
+		return TGuardedInt(Value % Other.Value);
 	}
 
 	/** @return This value bitwise left-shifted by the operand. */
-	TGuardedSignedInt operator <<(const TGuardedSignedInt Other) const
+	TGuardedInt operator <<(const TGuardedInt Other) const
 	{
 		// Any shift involving invalid values is invalid.
 		if (!bIsValid || !Other.bIsValid)
 		{
-			return TGuardedSignedInt();
+			return TGuardedInt();
 		}
 
 		// Left-shifts by negative values or >= the width of the type are always invalid.
 		if (Other.Value < 0 || Other.Value >= NumBits)
 		{
-			return TGuardedSignedInt();
+			return TGuardedInt();
 		}
 
-		const int ShiftAmount = Other.Value;
-
-		// Once again, taking our overflow-prone expression and using algebra to find
-		// a form that doesn't overflow:
-		//
-		// MinValue <= a * 2^b <= MaxValue
-		// <=> MinValue * 2^(-b) <= a <= MaxValue * 2^(-b)
-		//
-		// The LHS is exact because MinValue is -2^(NumBits - 1) for two's complement,
-		// and we just ensured that 0 <= b < NumBits (with b integer).
-		//
-		// The RHS has a fractional part whereas a is integer; therefore, we can
-		// substitute floor(MaxValue * 2^(-b)) for the RHS without changing the result.
-		//
-		// And that gives us our test!
-		return ((MinValue >> ShiftAmount) <= Value && Value <= (MaxValue >> ShiftAmount)) ? TGuardedSignedInt(Value << ShiftAmount) : TGuardedSignedInt();
+		if constexpr (std::is_signed_v<IntType>)
+		{
+			const int ShiftAmount = Other.Value;
+
+			// Once again, taking our overflow-prone expression and using algebra to find
+			// a form that doesn't overflow:
+			//
+			// MinValue <= a * 2^b <= MaxValue
+			// <=> MinValue * 2^(-b) <= a <= MaxValue * 2^(-b)
+			//
+			// The LHS is exact because MinValue is -2^(NumBits - 1) for two's complement,
+			// and we just ensured that 0 <= b < NumBits (with b integer).
+			//
+			// The RHS has a fractional part whereas a is integer; therefore, we can
+			// substitute floor(MaxValue * 2^(-b)) for the RHS without changing the result.
+			//
+			// And that gives us our test!
+			return ((MinValue >> ShiftAmount) <= Value && Value <= (MaxValue >> ShiftAmount)) ? TGuardedInt(Value << ShiftAmount) : TGuardedInt();
+		}
+		else
+		{
+			return (static_cast<IntType>(Value << Other.Value) >> Other.Value) == Value
+				? TGuardedInt(static_cast<IntType>(Value << Other.Value))
+				: TGuardedInt();
+		}
 	}
 
 	/** @return This value bitwise right-shifted by the operand. */
-	TGuardedSignedInt operator >>(const TGuardedSignedInt Other) const
+	TGuardedInt operator >>(const TGuardedInt Other) const
 	{
 		// Any shift involving invalid values is invalid.
 		if (!bIsValid || !Other.bIsValid)
 		{
-			return TGuardedSignedInt();
+			return TGuardedInt();
 		}
 
 		// Right-shifts by negative values or >= the width of the type are always invalid.
 		if (Other.Value < 0 || Other.Value >= NumBits)
 		{
-			return TGuardedSignedInt();
+			return TGuardedInt();
 		}
 
 		// Right-shifts don't have any overflow conditions, so we're good!
-		return TGuardedSignedInt(Value >> Other.Value);
+		return TGuardedInt(Value >> Other.Value);
 	}
 
 	/** @return The absolute value of *this. */
-	TGuardedSignedInt Abs() const
+	TGuardedInt Abs() const
 	{
 		if (!bIsValid)
 		{
-			return TGuardedSignedInt();
+			return TGuardedInt();
 		}
 
-		// Note the absolute value of MinValue overflows, so this is not completely trivial.
-		// Can't just use TGuardedSignedInt(abs(Value)) here!
-		return (Value < 0) ? -*this : *this;
+		if constexpr (std::is_signed_v<IntType>)
+		{
+			// Note the absolute value of MinValue overflows, so this is not completely trivial.
+			// Can't just use TGuardedInt(abs(Value)) here!
+			return (Value < 0) ? -*this : *this;
+		}
+		else
+		{
+			// Abs on unsigned integers is the identity function, and can't overflow.
+			return TGuardedInt(Value);
+		}
 	}
 
 	// Mixed-type operators and assignment operators reduce to the base operators systematically
 #define UE_GUARDED_SIGNED_INT_IMPL_BINARY_OPERATOR(OP) \
 	/* Mixed-type expressions that coerce both operands to guarded ints */ \
-	TGuardedSignedInt operator OP(SignedType InB) const { return *this OP TGuardedSignedInt(InB); } \
-	friend TGuardedSignedInt operator OP(SignedType InA, TGuardedSignedInt InB) { return TGuardedSignedInt(InA) OP InB; } \
+	TGuardedInt operator OP(IntType InB) const { return *this OP TGuardedInt(InB); } \
+	friend TGuardedInt operator OP(IntType InA, TGuardedInt InB) { return TGuardedInt(InA) OP InB; } \
 	/* Assignment operators, direct and mixed */ \
-	TGuardedSignedInt& operator OP##=(TGuardedSignedInt InB) { return *this = *this OP InB; } \
-	TGuardedSignedInt& operator OP##=(SignedType InB) { return *this = *this OP TGuardedSignedInt(InB); } \
+	TGuardedInt& operator OP##=(TGuardedInt InB) { return *this = *this OP InB; } \
+	TGuardedInt& operator OP##=(IntType InB) { return *this = *this OP TGuardedInt(InB); } \
 	/* end */
 
 	UE_GUARDED_SIGNED_INT_IMPL_BINARY_OPERATOR(+)
@@ -444,8 +404,12 @@ public:
 #undef UE_GUARDED_SIGNED_INT_IMPL_BINARY_OPERATOR
 };
 
+/** Legacy alias for the previously signed-integer-only implementation. */
+template<typename SignedType>
+using TGuardedSignedInt = TGuardedInt<SignedType>;
+
 /** Guarded 32-bit integer class. Used to deal with integer data from untrusted sources in size computations etc. */
-using FGuardedInt32 = TGuardedSignedInt<int32>;
+using FGuardedInt32 = TGuardedInt<int32>;
 
 /** Guarded 64-bit integer class. Used to deal with integer data from untrusted sources in size computations etc. */
-using FGuardedInt64 = TGuardedSignedInt<int64>;
+using FGuardedInt64 = TGuardedInt<int64>;
diff --git a/Engine/Source/Runtime/Core/Tests/Math/GuardedIntTest.cpp b/Engine/Source/Runtime/Core/Tests/Math/GuardedIntTest.cpp
--- a/Engine/Source/Runtime/Core/Tests/Math/GuardedIntTest.cpp
+++ b/Engine/Source/Runtime/Core/Tests/Math/GuardedIntTest.cpp
@@ -2,93 +2,105 @@
 
 #include "Math/GuardedInt.h"
 #include "Misc/AutomationTest.h"
+#include "Templates/ChooseClass.h"
+#include <limits>
+#include <type_traits>
 
 #if WITH_DEV_AUTOMATION_TESTS
 
-IMPLEMENT_SIMPLE_AUTOMATION_TEST(FGuardedIntTest, "System.ImageCore.GuardedInt", EAutomationTestFlags::ApplicationContextMask | EAutomationTestFlags::EngineFilter)
-bool FGuardedIntTest::RunTest(const FString& Parameters)
+namespace {
+template<typename IntType>
+bool DoGuardedIntTestForType(FAutomationTestBase& Test)
 {
-	// Do an exhaustive test of all checked signed int operations with int8s. This is a small enough space of possibilities
+	// IntType may only be int8 or uint8.
+	static_assert(std::is_same_v<IntType, int8> || std::is_same_v<IntType, uint8>);
+
+	constexpr bool bIsSigned = std::is_signed_v<IntType>;
+	using IntType32 = typename TChooseClass<bIsSigned, int32, uint32>::Result;
+	constexpr IntType32 IntTypeMin = std::numeric_limits<IntType>::min();
+	constexpr IntType32 IntTypeMax = std::numeric_limits<IntType>::max();
+
+	// Do an exhaustive test of all checked 8-bit int operations. This is a small enough space of possibilities
 	// that it's easy to try all options exhaustively, and it's likewise easy to compute reference results in int32 arithmetic
-	// without having to worry about overflows. TCheckedSignedInt uses the same logic and implementation for all type
+	// without having to worry about overflows. TGuardedInt uses the same logic and implementation for all signed type
 	// instantiations (as of this writing), so this is a good way to get coverage on all relevant code paths.
-	using FGuardedInt8 = TGuardedSignedInt<int8>;
+	using FGuardedIntType = TGuardedInt<IntType>;
 	int NumErrorsFound = 0;
 	const int NumErrorsMax = 100; // Stop reporting errors after this in case something is completely busted.
 
 	// Helper functions to log test results
-	auto CheckArithResult = [this, &NumErrorsFound, NumErrorsMax](const TCHAR* Operation, int32 A, int32 B, FGuardedInt8 CheckedResult, int32 UncheckedValue)
+	auto CheckArithResult = [&Test, &NumErrorsFound, NumErrorsMax](const TCHAR* Operation, IntType32 A, IntType32 B, FGuardedIntType CheckedResult, IntType32 UncheckedValue)
 	{
-		FGuardedInt8 UncheckedResult{ UncheckedValue };
+		FGuardedIntType UncheckedResult{ UncheckedValue };
 
 		// If either operand was invalid, the unchecked result is also invalid
-		if (!FGuardedInt8(A).IsValid() || !FGuardedInt8(B).IsValid())
+		if (!FGuardedIntType(A).IsValid() || !FGuardedIntType(B).IsValid())
 		{
-			UncheckedResult = FGuardedInt8();
+			UncheckedResult = FGuardedIntType();
 		}
 
 		if (CheckedResult != UncheckedResult && NumErrorsFound < NumErrorsMax)
 		{
 			NumErrorsFound++;
 
-			const int UncheckedNum = UncheckedResult.Get(-128);
-			const int CheckedNum = CheckedResult.Get(-128);
+			const IntType32 UncheckedNum = UncheckedResult.IsValid() ? UncheckedResult.GetChecked() : (IntTypeMax + 1);
+			const IntType32 CheckedNum = CheckedResult.IsValid() ? CheckedResult.GetChecked() : (IntTypeMax + 1);
 
-			AddError(FString::Printf(TEXT("Arith %s A=%d B=%d Expected=(%d,Valid=%d) Got=(%d,Valid=%d)"),
+			Test.AddError(FString::Printf(TEXT("Arith %s A=%d B=%d Expected=(%d,Valid=%d) Got=(%d,Valid=%d)"),
 				Operation, A, B, UncheckedNum, UncheckedResult.IsValid(), CheckedNum, CheckedResult.IsValid()));
 		}
 	};
 
-	auto CheckCompareResult = [this, &NumErrorsFound, NumErrorsMax](const TCHAR* Operation, int32 A, int32 B, bool bAcceptInvalid, bool bTestResult, bool bRefResult)
+	auto CheckCompareResult = [&Test, &NumErrorsFound, NumErrorsMax](const TCHAR* Operation, IntType32 A, IntType32 B, bool bAcceptInvalid, bool bTestResult, bool bRefResult)
 	{
 		// If either operands was invalid, override the ref result
-		if (!FGuardedInt8(A).IsValid() || !FGuardedInt8(B).IsValid())
+		if (!FGuardedIntType(A).IsValid() || !FGuardedIntType(B).IsValid())
 			bRefResult = bAcceptInvalid;
 
 		if (bTestResult != bRefResult && NumErrorsFound < NumErrorsMax)
 		{
 			NumErrorsFound++;
 
-			AddError(FString::Printf(TEXT("Cmp %s A=%d B=%d ref=%d test=%d\n"), Operation, A, B, bRefResult, bTestResult));
+			Test.AddError(FString::Printf(TEXT("Cmp %s A=%d B=%d ref=%d test=%d\n"), Operation, A, B, bRefResult, bTestResult));
 		}
 	};
 
-	// Note: we include 128 in the ranges here to explicitly try out-of-range (invalid) values.
-	for (int32 A = -128; A <= 128; ++A)
+	// Note: we include IntTypeMax + 1 in the ranges here to explicitly try out-of-range (invalid) values.
+	for (IntType32 A = IntTypeMin; A <= IntTypeMax + 1; ++A)
 	{
-		for (int32 B = -128; B <= 128; ++B)
+		for (IntType32 B = IntTypeMin; B <= IntTypeMax + 1; ++B)
 		{
-			FGuardedInt8 CheckedA{ A };
-			FGuardedInt8 CheckedB{ B };
+			FGuardedIntType CheckedA{ A };
+			FGuardedIntType CheckedB{ B };
 
 			// This doesn't test much by itself but ensures we have at least compiled GetChecked.
-			if (A <= 127 && B == 0)
+			if (A >= IntTypeMin && A <= IntTypeMax && B == 0)
 			{
 				if (CheckedA.GetChecked() != A)
 				{
-					AddError(FString::Printf(TEXT("GetChecked A=%d"), A));
+					Test.AddError(FString::Printf(TEXT("GetChecked A=%d"), A));
 				}
 			}
 
 			// Compute reference results for the operations that are not defined everywhere
-			int32 ADivB = -256;
-			int32 AModB = -256;
+			IntType32 ADivB = -256;
+			IntType32 AModB = -256;
 
 			// -128 / -1 is undefined in int8 arithmetic because it overflows; when we're
 			// computing it with int32s, it of course works just fine, but we want to make sure
 			// our test results are correct.
-			if (B != 0 && (A != -128 || B != -1))
+			if (B != 0 && !(bIsSigned && A == IntTypeMin && B == -1))
 			{
 				ADivB = A / B;
 				AModB = A % B;
 			}
 
 			// Shifts are only defined for in-range shift amounts, which for our checked ints
-			// is [0,NumBits-1]. This is tighter than usual C++ arithmetic for the FGuardedInt8 we
+			// is [0,NumBits-1]. This is tighter than usual C++ arithmetic for the FGuardedIntType we
 			// use for testing since normal C++ promotes to int and shifts down (so usually [0,31]
 			// count limit) so we need to test for these explicitly.
-			int32 ALshB = -256;
-			int32 ARshB = -256;
+			IntType32 ALshB = -256;
+			IntType32 ARshB = -256;
 			if (B >= 0 && B <= 7)
 			{
 				ALshB = A << B;
@@ -102,14 +114,17 @@ bool FGuardedIntTest::RunTest(const FString& Parameters)
 			CheckArithResult(TEXT("mod1"), A, B, CheckedA % CheckedB, AModB);
 			CheckArithResult(TEXT("lsh1"), A, B, CheckedA << CheckedB, ALshB);
 			CheckArithResult(TEXT("rsh1"), A, B, CheckedA >> CheckedB, ARshB);
-			CheckArithResult(TEXT("neg"), A, A, -CheckedA, -A);
-			CheckArithResult(TEXT("abs"), A, A, CheckedA.Abs(), FMath::Abs(A));
+			if (bIsSigned || A == 0)
+			{
+				CheckArithResult(TEXT("neg"), A, A, -CheckedA, static_cast<IntType32>(-static_cast<int32>(A)));
+			}
+			CheckArithResult(TEXT("abs"), A, A, CheckedA.Abs(), static_cast<IntType32>(FMath::Abs(static_cast<int32>(A))));
 
 			// Mixed-type arithmetic needs int8 arguments, which we can only actually do when the values in question _are_ int8s.
 			// (i.e. we can't do this with our intentional out-of-range value of 128.)
 			if (CheckedB.IsValid())
 			{
-				int8 BAs8 = (int8)B;
+				IntType BAs8 = (IntType)B;
 
 				CheckArithResult(TEXT("add2"), A, B, CheckedA + BAs8, A + B);
 				CheckArithResult(TEXT("sub2"), A, B, CheckedA - BAs8, A - B);
@@ -122,7 +137,7 @@ bool FGuardedIntTest::RunTest(const FString& Parameters)
 
 			if (CheckedA.IsValid())
 			{
-				int8 AAs8 = (int8)A;
+				IntType AAs8 = (IntType)A;
 
 				CheckArithResult(TEXT("add3"), A, B, AAs8 + CheckedB, A + B);
 				CheckArithResult(TEXT("sub3"), A, B, AAs8 - CheckedB, A - B);
@@ -155,10 +170,20 @@ bool FGuardedIntTest::RunTest(const FString& Parameters)
 
 	if (NumErrorsFound > NumErrorsMax)
 	{
-		AddError(FString::Printf(TEXT("Too many errors (%d), stopped reporting after %d instances.\n"), NumErrorsFound, NumErrorsMax));
+		Test.AddError(FString::Printf(TEXT("Too many errors (%d), stopped reporting after %d instances.\n"), NumErrorsFound, NumErrorsMax));
 	}
 
 	return NumErrorsFound == 0;
 }
+}
+
+IMPLEMENT_SIMPLE_AUTOMATION_TEST(FGuardedIntTest, "Core.GuardedInt", EAutomationTestFlags::ApplicationContextMask | EAutomationTestFlags::EngineFilter)
+bool FGuardedIntTest::RunTest(const FString&)
+{
+	bool bResult = true;
+	bResult &= DoGuardedIntTestForType<int8>(*this);
+	bResult &= DoGuardedIntTestForType<uint8>(*this);
+	return bResult;
+}
 
 #endif
\ No newline at end of file
