- Delayed the kickoff of async loading of other platform configs, because ProjectDir can get reset on main thread while the async threads are using it for loading other platform configs
- Allow for projects to have ConfigRedirects.ini files - it was not oringally allowed but a licensee found that it can be needed for User saved ini files that can't be updated, so the redirects is useful for a licensee to specify
- Fixed a copy/paste error in a log about ConfigRedirects

diff --git a/Engine/Config/ConfigRedirects.ini b/Engine/Config/ConfigRedirects.ini
--- a/Engine/Config/ConfigRedirects.ini
+++ b/Engine/Config/ConfigRedirects.ini
@@ -10,7 +10,10 @@
 ;; The main intention of this file is to allow Epic to change section/key names without breaking licensee config files. So,
 ;; it's not expected that licensees would need to modify this file - instead, licensees can fix up their config files directly
 ;; instead of modifying a config file like this one.
-;; 
+;;
+;; However, if a licensee can't fix the ini file because saved user settings need to be redirected (and can't be fixed easily)
+;; a project can have a <ProjectDir>/Config/ConfigRedirects.ini file that will be read.
+;;
 ;; This first section contains a list of section names to be remapped. All other sections are the name of a section that contains
 ;; one of more keys to remap/redirect, followed up the old and new name of the keys to remap
 
@@ -50,4 +53,4 @@ EncryptionKey=ClientEncryptionKey
 EncryptionKey=ClientEncryptionKey
 
 [HTTP]
-AllowedDomains=Online.HttpManager https:AllowedDomains
\ No newline at end of file
+AllowedDomains=Online.HttpManager https:AllowedDomains
diff --git a/Engine/Source/Runtime/Core/Private/Misc/ConfigCacheIni.cpp b/Engine/Source/Runtime/Core/Private/Misc/ConfigCacheIni.cpp
--- a/Engine/Source/Runtime/Core/Private/Misc/ConfigCacheIni.cpp
+++ b/Engine/Source/Runtime/Core/Private/Misc/ConfigCacheIni.cpp
@@ -694,7 +694,7 @@ static void WarnAboutKeyRemap(const FString& OldValue, const FString& NewValue,
 	Arguments.Add(TEXT("NewValue"), FText::FromString(NewValue));
 	Arguments.Add(TEXT("Section"), FText::FromString(Section));
 	Arguments.Add(TEXT("File"), FText::FromString(File));
-	FText Msg = FText::Format(LOCTEXT("DeprecatedConfig", "Found a deprecated ini section name in {File}. Search for [{OldValue}] and replace with [{NewValue}]"), Arguments);
+	FText Msg = FText::Format(LOCTEXT("DeprecatedConfig", "Found a deprecated ini key name in {File}. Search for [{OldValue}] and replace with [{NewValue}]"), Arguments);
 	
 	FString Key = OldValue+Section;
 	if (!IsInGameThread())
@@ -3741,23 +3741,28 @@ static void InitializeConfigRemap()
 	// read in the single remap file
 	FConfigFile RemapFile;
 	FConfigContext Context = FConfigContext::ReadSingleIntoLocalFile(RemapFile);
-	Context.Load(*FPaths::Combine(FPaths::EngineDir(), TEXT("Config/ConfigRedirects.ini")));
 	
-	for (const TPair<FString, FConfigSection>& Section : RemapFile)
+	// read in engine and project ini files (these are not hierarchical, so it has to be done in two passes)
+	for (int Pass = 0; Pass < 2; Pass++)
 	{
-		if (Section.Key == TEXT("SectionNameRemap"))
+		Context.Load(*FPaths::Combine(Pass == 0 ? FPaths::EngineDir() : FPaths::ProjectDir(), TEXT("Config/ConfigRedirects.ini")));
+		
+		for (const TPair<FString, FConfigSection>& Section : RemapFile)
 		{
-			for (const TPair<FName, FConfigValue>& Line : Section.Value)
+			if (Section.Key == TEXT("SectionNameRemap"))
 			{
-				SectionRemap.Add(Line.Key.ToString(), Line.Value.GetSavedValue());
+				for (const TPair<FName, FConfigValue>& Line : Section.Value)
+				{
+					SectionRemap.Add(Line.Key.ToString(), Line.Value.GetSavedValue());
+				}
 			}
-		}
-		else
-		{
-			TMap<FString, FString>& KeyRemaps = KeyRemap.FindOrAdd(Section.Key);
-			for (const TPair<FName, FConfigValue>& Line : Section.Value)
+			else
 			{
-				KeyRemaps.Add(Line.Key.ToString(), Line.Value.GetSavedValue());
+				TMap<FString, FString>& KeyRemaps = KeyRemap.FindOrAdd(Section.Key);
+				for (const TPair<FName, FConfigValue>& Line : Section.Value)
+				{
+					KeyRemaps.Add(Line.Key.ToString(), Line.Value.GetSavedValue());
+				}
 			}
 		}
 	}
@@ -3774,10 +3779,6 @@ void FConfigCacheIni::InitializeConfigSystem()
 	
 	InitializeConfigRemap();
 	
-#if WITH_EDITOR
-	AsyncInitializeConfigForPlatforms();
-#endif
-
 #if PLATFORM_SUPPORTS_BINARYCONFIG && PRELOAD_BINARY_CONFIG
 	// attempt to load from staged binary config data
 	if (!FParse::Param(FCommandLine::Get(), TEXT("textconfig")) &&
@@ -3862,6 +3863,12 @@ void FConfigCacheIni::InitializeConfigSystem()
 	// now we can make use of GConfig
 	GConfig->bIsReadyForUse = true;
 
+#if WITH_EDITOR
+	// this needs to be called after setting bIsReadyForUse, because it uses ProjectDir, and bIsReadyForUse can reset the 
+	// ProjectDir array while the async threads are using it and crash
+	AsyncInitializeConfigForPlatforms();
+#endif
+
 	TRACE_CPUPROFILER_EVENT_SCOPE(ConfigReadyForUseBroadcast);
 	FCoreDelegates::TSConfigReadyForUse().Broadcast();
 }
