FBox constructors update:
- default construct now creates a box marked as invalid
- added ENABLE_UNINITIALIZED_BOX_DIAGNOSTIC to force Min and Max values to NaN and to perform disagnostic checks at runtime.
- updated some constructors to be constexpr
- added FBox(ENoInit) constructor that creates a box without any initialization
- added DiagnosticCheck to detect cases where a default initialized is used with values for Min and Max still set to NaN (ENABLE_UNINITIALIZED_BOX_DIAGNOSTIC = 1) or to detect introduce NaN during computation (ENABLE_NAN_DIAGNOSTIC = 1)

Goals with this change since the implementation is currently not consistent regarding the IsValid flag:
1- Initialize IsValid flag to 0 in default constructor so 'operator+=' can be reliable when used on a default initialized FBox. DiagnosticCheck can then be used when box is marked as valid.
2- Since it is already possible for user to construct a FBox, set Min and Max (without marking it as Valid) then pass it around or perform other operations calls to DiagnosticCheck are added when a box is constructed (mainly to catch cases when a NaN initialized box is used to create another box: ExpandBy, ShiftBy, MoveTo, etc.). Otherwhise DiagnosticCheck is used in methods where Min or Max are used without consideration of the IsValid flag (GetCenter, GetExtent, GetSize, etc.)

diff --git a/Engine/Source/Runtime/Core/Public/Math/Box.h b/Engine/Source/Runtime/Core/Public/Math/Box.h
--- a/Engine/Source/Runtime/Core/Public/Math/Box.h
+++ b/Engine/Source/Runtime/Core/Public/Math/Box.h
@@ -10,14 +10,22 @@
 #include "Math/Sphere.h"
 #include "Misc/LargeWorldCoordinatesSerializer.h"
 
+#ifndef ENABLE_UNINITIALIZED_BOX_DIAGNOSTIC
+#define ENABLE_UNINITIALIZED_BOX_DIAGNOSTIC 0
+#endif // ENABLE_UNINITIALIZED_BOX_DIAGNOSTIC
+
+#ifndef ENABLE_BOX_DIAGNOSTIC_CHECKS
+#define ENABLE_BOX_DIAGNOSTIC_CHECKS (ENABLE_UNINITIALIZED_BOX_DIAGNOSTIC || ENABLE_NAN_DIAGNOSTIC)
+#endif// ENABLE_BOX_DIAGNOSTIC_CHECKS
+
 /**
  * Implements an axis-aligned box.
  *
  * Boxes describe an axis-aligned extent in three dimensions. They are used for many different things in the
  * Engine and in games, such as bounding volumes, collision detection and visibility calculation.
  */
-namespace UE {
-namespace Math {
+namespace UE::Math
+{
 
 struct TBoxConstInit {};
 
@@ -38,17 +46,34 @@ public:
 
 public:
 
-	/** Default constructor (no initialization). */
-	TBox() { }
-
 	/**
-	 * Creates and initializes a new box with zero extent and marks it as invalid.
-	 *
-	 * Use enum value EForceInit::ForceInit to force box initialization.
+	 * Default constructor.
+	 * Creates a new box with uninitialized extents and marks it as invalid.
 	 */
-	explicit TBox( EForceInit )
+	constexpr TBox()
+#if ENABLE_UNINITIALIZED_BOX_DIAGNOSTIC
+		: Min(std::numeric_limits<T>::quiet_NaN(), TVectorConstInit{})
+		, Max(std::numeric_limits<T>::quiet_NaN(), TVectorConstInit{})
+		, IsValid(0)
+#else
+		// To be consistent with other engine core types (i.e. Vector) the other members are not initialized
+		// since there is not meaningful default values.
+		: IsValid(0)
+#endif
+	{
+	}
+
+	/** Creates a new box without initialization. */
+	explicit constexpr TBox(ENoInit)
+	{
+	}
+
+	/** Creates and initializes a new box with zero extent and marks it as invalid. */
+	explicit constexpr TBox(EForceInit)
+		: Min(0, TVectorConstInit{})
+		, Max(0, TVectorConstInit{})
+		, IsValid(0)
 	{
-		Init();
 	}
 
 	constexpr TBox(EForceInit, TBoxConstInit)
@@ -56,6 +81,30 @@ public:
 	{
 	}
 
+#if ENABLE_BOX_DIAGNOSTIC_CHECKS
+	FORCEINLINE void DiagnosticCheck() const
+	{
+		if (ContainsNaN())
+		{
+#if ENABLE_UNINITIALIZED_BOX_DIAGNOSTIC
+			// In case ENABLE_UNINITIALIZED_BOX_DIAGNOSTIC was explicitly enabled then report each encountered accesses to uninitialized members.
+			ensureAlwaysMsgf(TEXT("FBox contains NaN: %s. Initializing with zero extent and marking as invalid for safety."), *ToString());			
+#else
+			// In case diagnostic check is enable from NaN diagnostic we report the error using the NaN error reporting behavior.
+			logOrEnsureNanError(TEXT("FBox contains NaN: %s. Initializing with zero extent and marking as invalid for safety."), *ToString());
+#endif
+			const_cast<TBox<T>*>(this)->Init();
+		}
+	}
+#else
+	FORCEINLINE void DiagnosticCheck() const {}
+#endif
+
+	bool ContainsNaN() const
+	{
+		return Min.ContainsNaN() || Max.ContainsNaN();
+	}
+	
 	/**
 	 * Creates and initializes a new box from the specified extents.
 	 *
@@ -70,13 +119,17 @@ public:
 	{
 		// Intended to catch TBox<float>(TVector<double>(), TVector<double>())
 		static_assert(sizeof(FArg) <= sizeof(T), "Losing precision when constructing a box of floats from vectors of doubles");
+
+		DiagnosticCheck();
 	}
 
 	TBox(const TVector4<T>& InMin, const TVector4<T>& InMax)
 		: Min(InMin)
 		, Max(InMax)
 		, IsValid(1)
-	{ }
+	{
+		DiagnosticCheck();
+	}
 
 	/**
 	 * Creates and initializes a new box from the given set of points.
@@ -84,7 +137,7 @@ public:
 	 * @param Points Array of Points to create for the bounding volume.
 	 * @param Count The number of points.
 	 */
-	TBox(const TVector<T>* Points, int32 Count) : Min(0, 0, 0), Max(0, 0, 0), IsValid(0)
+	TBox(const TVector<T>* Points, const int32 Count) : TBox(ForceInit)
 	{
 		for (int32 i = 0; i < Count; i++)
 		{
@@ -97,7 +150,7 @@ public:
 	 *
 	 * @param Points Array of Points to create for the bounding volume.
 	 */
-	TBox(const TArray<TVector<T>>& Points) : TBox<T>(&Points[0], Points.Num()) {};
+	explicit TBox(const TArray<TVector<T>>& Points) : TBox<T>(&Points[0], Points.Num()) {};
 
 	// Conversion from other type.
 	template<typename FArg, TEMPLATE_REQUIRES(!std::is_same_v<T, FArg>)>
@@ -285,6 +338,7 @@ public:
 	 */
 	FORCEINLINE TVector<T> GetCenter() const
 	{
+		DiagnosticCheck();
 		return TVector<T>((Min + Max) * 0.5f);
 	}
 
@@ -317,21 +371,10 @@ public:
 	 */
 	FORCEINLINE TVector<T> GetExtent() const
 	{
+		DiagnosticCheck();
 		return 0.5f * (Max - Min);
 	}
 
-	UE_DEPRECATED(4.24, "This method performed unsafe operations and should be replaced with using .Min and .Max directly or using the [] operator on this class instead.")
-	FORCEINLINE TVector<T>& GetExtrema( int PointIndex )
-	{
-		return (&Min)[PointIndex];
-	}
-
-	UE_DEPRECATED(4.24, "This method performed unsafe operations and should be replaced with using .Min and .Max directly or using the [] operator on this class instead.")
-	FORCEINLINE const TVector<T>& GetExtrema( int PointIndex ) const
-	{
-		return (&Min)[PointIndex];
-	}
-
 	/**
 	 * Gets the size of this box.
 	 *
@@ -340,6 +383,7 @@ public:
 	 */
 	FORCEINLINE TVector<T> GetSize() const
 	{
+		DiagnosticCheck();
 		return (Max - Min);
 	}
 
@@ -351,6 +395,7 @@ public:
 	 */
 	FORCEINLINE T GetVolume() const
 	{
+		DiagnosticCheck();
 		return (Max.X - Min.X) * (Max.Y - Min.Y) * (Max.Z - Min.Z);
 	}
 
@@ -413,6 +458,7 @@ public:
 	 */
 	FORCEINLINE bool IsInside( const TVector<T>& In ) const
 	{
+		DiagnosticCheck();
 		return ((In.X > Min.X) && (In.X < Max.X) && (In.Y > Min.Y) && (In.Y < Max.Y) && (In.Z > Min.Z) && (In.Z < Max.Z));
 	}
 
@@ -428,6 +474,7 @@ public:
 	 */
 	FORCEINLINE bool IsInsideOrOn( const TVector<T>& In ) const
 	{
+		DiagnosticCheck();
 		return ((In.X >= Min.X) && (In.X <= Max.X) && (In.Y >= Min.Y) && (In.Y <= Max.Y) && (In.Z >= Min.Z) && (In.Z <= Max.Z));
 	}
 
@@ -458,6 +505,7 @@ public:
 	 */
 	FORCEINLINE bool IsInsideXY( const TVector<T>& In ) const
 	{
+		DiagnosticCheck();
 		return ((In.X > Min.X) && (In.X < Max.X) && (In.Y > Min.Y) && (In.Y < Max.Y));
 	}
 
@@ -473,6 +521,7 @@ public:
 	 */
 	FORCEINLINE bool IsInsideOrOnXY(const TVector<T>& In) const
 	{
+		DiagnosticCheck();
 		return ((In.X >= Min.X) && (In.X <= Max.X) && (In.Y >= Min.Y) && (In.Y <= Max.Y));
 	}
 
@@ -597,6 +646,9 @@ FORCEINLINE TBox<T>& TBox<T>::operator+=( const TVector<T> &Other )
 {
 	if (IsValid)
 	{
+		DiagnosticCheck();
+		// Not testing Other.DiagnosticCheck() since TVector is not initialized to NaN
+
 		Min.X = FMath::Min(Min.X, Other.X);
 		Min.Y = FMath::Min(Min.Y, Other.Y);
 		Min.Z = FMath::Min(Min.Z, Other.Z);
@@ -620,6 +672,10 @@ FORCEINLINE TBox<T>& TBox<T>::operator+=( const TBox<T>& Other )
 {
 	if (IsValid && Other.IsValid)
 	{
+		DiagnosticCheck();
+		// Testing Other.DiagnosticCheck() since TBox can be initialized to NaN
+		Other.DiagnosticCheck();
+
 		Min.X = FMath::Min(Min.X, Other.Min.X);
 		Min.Y = FMath::Min(Min.Y, Other.Min.Y);
 		Min.Z = FMath::Min(Min.Z, Other.Min.Z);
@@ -639,6 +695,8 @@ FORCEINLINE TBox<T>& TBox<T>::operator+=( const TBox<T>& Other )
 template<typename T>
 FORCEINLINE TVector<T> TBox<T>::GetClosestPointTo( const TVector<T>& Point ) const
 {
+	DiagnosticCheck();
+
 	// start by considering the point inside the box
 	TVector<T> ClosestPoint = Point;
 
@@ -679,6 +737,10 @@ FORCEINLINE TVector<T> TBox<T>::GetClosestPointTo( const TVector<T>& Point ) con
 template<typename T>
 FORCEINLINE bool TBox<T>::Intersect( const TBox<T>& Other ) const
 {
+	DiagnosticCheck();
+	// Testing Other.DiagnosticCheck() since TBox can be initialized to NaN
+	Other.DiagnosticCheck();
+
 	if ((Min.X > Other.Max.X) || (Other.Min.X > Max.X))
 	{
 		return false;
@@ -701,6 +763,10 @@ FORCEINLINE bool TBox<T>::Intersect( const TBox<T>& Other ) const
 template<typename T>
 FORCEINLINE bool TBox<T>::IntersectXY( const TBox<T>& Other ) const
 {
+	DiagnosticCheck();
+	// Testing Other.DiagnosticCheckUninitialized() since TBox can be initialized to NaN
+	Other.DiagnosticCheck();
+
 	if ((Min.X > Other.Max.X) || (Other.Min.X > Max.X))
 	{
 		return false;
@@ -731,6 +797,8 @@ TBox<T> TBox<T>::TransformBy(const TMatrix<T>& M) const
 		return TBox<T>(ForceInit);
 	}
 
+	DiagnosticCheck();
+
 	TBox<T> NewBox;
 
 	const TVectorRegisterType<T> VecMin = VectorLoadFloat3_W0(&Min);
@@ -774,6 +842,8 @@ TBox<T> TBox<T>::TransformBy(const TTransform<T>& M) const
 template<typename T>
 void TBox<T>::GetVertices(TVector<T>(&Vertices)[8]) const
 {
+	DiagnosticCheck();
+
 	Vertices[0] = TVector<T>(Min);
 	Vertices[1] = TVector<T>(Min.X, Min.Y, Max.Z);
 	Vertices[2] = TVector<T>(Min.X, Max.Y, Min.Z);
@@ -844,8 +914,7 @@ TBox<T> TBox<T>::Overlap(const TBox<T>& Other) const
 }
 
 
-} // namespace Math
-} // namespace UE
+} // namespace UE::Math
 
 UE_DECLARE_LWC_TYPE(Box, 3);
 
