Remove platform extension WindowsHWrapperPrivate.h because there's already PreWindowsApiPrivate/MinWindowsPrivate & PostWindowsApiPrivate.h & the structure is always the same

#rnx

diff --git a/Engine/Source/Runtime/Core/Public/Microsoft/WindowsHWrapper.h b/Engine/Source/Runtime/Core/Public/Microsoft/WindowsHWrapper.h
--- a/Engine/Source/Runtime/Core/Public/Microsoft/WindowsHWrapper.h
+++ b/Engine/Source/Runtime/Core/Public/Microsoft/WindowsHWrapper.h
@@ -6,5 +6,12 @@
 #if PLATFORM_WINDOWS
 	#include "Windows/WindowsHWrapper.h" // HEADER_UNIT_IGNORE
 #else
-    #include "Microsoft/WindowsHWrapperPrivate.h"
+	#include "CoreTypes.h"
+	#include "HAL/PlatformMemory.h"
+	#include "Microsoft/PreWindowsApiPrivate.h"
+	#ifndef STRICT
+	#define STRICT
+	#endif
+	#include "Microsoft/MinWindowsPrivate.h"
+	#include "Microsoft/PostWindowsApiPrivate.h"
 #endif
