Increase size of symbols field and warn on truncation

diff --git a/Engine/Source/Runtime/Core/Private/GenericPlatform/GenericPlatformCrashContext.cpp b/Engine/Source/Runtime/Core/Private/GenericPlatform/GenericPlatformCrashContext.cpp
--- a/Engine/Source/Runtime/Core/Private/GenericPlatform/GenericPlatformCrashContext.cpp
+++ b/Engine/Source/Runtime/Core/Private/GenericPlatform/GenericPlatformCrashContext.cpp
@@ -191,6 +191,11 @@ void FGenericCrashContext::Initialize()
 	Symbols = FString::Printf(TEXT("%s-%s"), *Symbols, *FString(UE_BUILD_FLAVOR));
 #endif
 	NCached::Set(NCached::Session.SymbolsLabel, *Symbols);
+	if (Symbols.Len() >= UE_ARRAY_COUNT(NCached::Session.SymbolsLabel))
+	{
+		UE_LOG(LogInit, Error, TEXT("Symbols label too long (%d) for field size(%d), truncated. This may cause problems with crash report symbolication."),
+			Symbols.Len(), UE_ARRAY_COUNT(NCached::Session.SymbolsLabel));
+	}
 
 	FString OsVersion, OsSubVersion;
 	FPlatformMisc::GetOSVersions(OsVersion, OsSubVersion);
diff --git a/Engine/Source/Runtime/Core/Public/GenericPlatform/GenericPlatformCrashContext.h b/Engine/Source/Runtime/Core/Public/GenericPlatform/GenericPlatformCrashContext.h
--- a/Engine/Source/Runtime/Core/Public/GenericPlatform/GenericPlatformCrashContext.h
+++ b/Engine/Source/Runtime/Core/Public/GenericPlatform/GenericPlatformCrashContext.h
@@ -141,6 +141,7 @@ enum class EUnattendedStatus : uint8
 
 #define CR_MAX_ERROR_MESSAGE_CHARS 2048
 #define CR_MAX_DIRECTORY_CHARS 256
+#define CR_MAX_SYMBOL_CHARS 128
 #define CR_MAX_STACK_FRAMES 256
 #define CR_MAX_THREAD_NAME_CHARS 64
 #define CR_MAX_THREADS 512
@@ -183,7 +184,7 @@ struct FSessionContext
 	TCHAR 					RootDir[CR_MAX_DIRECTORY_CHARS];
 	TCHAR 					EpicAccountId[CR_MAX_GENERIC_FIELD_CHARS];
 	TCHAR 					LoginIdStr[CR_MAX_GENERIC_FIELD_CHARS];
-	TCHAR					SymbolsLabel[CR_MAX_GENERIC_FIELD_CHARS];
+	TCHAR					SymbolsLabel[CR_MAX_SYMBOL_CHARS];
 	TCHAR 					OsVersion[CR_MAX_GENERIC_FIELD_CHARS];
 	TCHAR 					OsSubVersion[CR_MAX_GENERIC_FIELD_CHARS];
 	TCHAR 					CPUVendor[CR_MAX_GENERIC_FIELD_CHARS];
