CSVProfiler: Return -1 when hot reloading to prevent a crash.

diff --git a/Engine/Source/Runtime/Core/Private/ProfilingDebugging/CsvProfiler.cpp b/Engine/Source/Runtime/Core/Private/ProfilingDebugging/CsvProfiler.cpp
--- a/Engine/Source/Runtime/Core/Private/ProfilingDebugging/CsvProfiler.cpp
+++ b/Engine/Source/Runtime/Core/Private/ProfilingDebugging/CsvProfiler.cpp
@@ -596,6 +596,13 @@ public:
 	int32 RegisterCategory(const FString& CategoryName, bool bEnableByDefault, bool bIsGlobal)
 	{
 		int32 Index = -1;
+		// During a hot-reload, we attempt to re-register categories and/or statics
+		// that not have been loaded/init'd during a hot-reload, can result in crashing.
+		// Thus when doing a reload, bail. 
+		if (IsReloadActive())
+		{
+			return Index;
+		}
 
 		FScopeLock Lock(&CS);
 		{
