Global lock using recursive mutex to for stackwalking access on windows
cl originally by devin.doucette

diff --git a/Engine/Source/Runtime/Core/Private/Windows/WindowsPlatformStackWalk.cpp b/Engine/Source/Runtime/Core/Private/Windows/WindowsPlatformStackWalk.cpp
--- a/Engine/Source/Runtime/Core/Private/Windows/WindowsPlatformStackWalk.cpp
+++ b/Engine/Source/Runtime/Core/Private/Windows/WindowsPlatformStackWalk.cpp
@@ -1,6 +1,8 @@
 // Copyright Epic Games, Inc. All Rights Reserved.
 
 #include "Windows/WindowsPlatformStackWalk.h"
+#include "Async/RecursiveMutex.h"
+#include "Async/UniqueLock.h"
 #include "HAL/PlatformMemory.h"
 #include "HAL/PlatformMisc.h"
 #include "Logging/LogMacros.h"
@@ -36,6 +38,9 @@ THIRD_PARTY_INCLUDES_END
 	@TODO To be removed
 -----------------------------------------------------------------------------*/
 
+//@note Global lock as a tentative fix to narrow down UE-192420
+static UE::FRecursiveMutex GStackWalkingLock;
+
 /** Whether appInitStackWalking() has been called successfully or not. */
 static bool GStackWalkingInitialized = false;
 static bool GNeedToRefreshSymbols = false;
@@ -133,6 +138,8 @@ void DetermineMaxCallstackDepth()
 
 void FWindowsPlatformStackWalk::StackWalkAndDump( ANSICHAR* HumanReadableString, SIZE_T HumanReadableStringSize, int32 IgnoreCount, void* Context )
 {
+	UE::TUniqueLock GlobalLock(GStackWalkingLock);
+
 	InitStackWalking();
 
 	// If the callstack is for the executing thread, ignore this function
@@ -152,12 +159,16 @@ void FWindowsPlatformStackWalk::StackWalkAndDump( ANSICHAR* HumanReadableString,
 
 void FWindowsPlatformStackWalk::StackWalkAndDump( ANSICHAR* HumanReadableString, SIZE_T HumanReadableStringSize, void* ProgramCounter, void* Context )
 {
+	UE::TUniqueLock GlobalLock(GStackWalkingLock);
+
 	InitStackWalking();
 	FGenericPlatformStackWalk::StackWalkAndDump(HumanReadableString, HumanReadableStringSize, ProgramCounter, Context);
 }
 
 FORCENOINLINE TArray<FProgramCounterSymbolInfo> FWindowsPlatformStackWalk::GetStack(int32 IgnoreCount, int32 MaxDepth, void* Context)
 {
+	UE::TUniqueLock GlobalLock(GStackWalkingLock);
+
 	InitStackWalking();
 
 	// If the callstack is for the executing thread, ignore this function
@@ -170,6 +181,8 @@ FORCENOINLINE TArray<FProgramCounterSymbolInfo> FWindowsPlatformStackWalk::GetSt
 
 void FWindowsPlatformStackWalk::ThreadStackWalkAndDump(ANSICHAR* HumanReadableString, SIZE_T HumanReadableStringSize, int32 IgnoreCount, uint32 ThreadId)
 {
+	UE::TUniqueLock GlobalLock(GStackWalkingLock);
+
 	InitStackWalking();
 	HANDLE ThreadHandle = OpenThread(THREAD_GET_CONTEXT | THREAD_SET_CONTEXT | THREAD_TERMINATE | THREAD_SUSPEND_RESUME, false, ThreadId);
 	if (ThreadHandle)
@@ -302,6 +315,8 @@ void FWindowsPlatformStackWalk::CaptureStackTraceByProcess(uint64* OutBacktrace,
 
 uint32 FWindowsPlatformStackWalk::CaptureThreadStackBackTrace(uint64 ThreadId, uint64* BackTrace, uint32 MaxDepth, void* Context)
 {
+	UE::TUniqueLock GlobalLock(GStackWalkingLock);
+
 	if (BackTrace == nullptr || MaxDepth == 0)
 	{
 		return 0;
@@ -370,6 +385,8 @@ uint32 FWindowsPlatformStackWalk::CaptureThreadStackBackTrace(uint64 ThreadId, u
  */
 uint32 FWindowsPlatformStackWalk::CaptureStackBackTrace( uint64* BackTrace, uint32 MaxDepth, void* Context )
 {
+	UE::TUniqueLock GlobalLock(GStackWalkingLock);
+
 	// Make sure we have place to store the information before we go through the process of raising
 	// an exception and handling it.
 	if (BackTrace == NULL || MaxDepth == 0)
@@ -441,6 +458,8 @@ uint32 FWindowsPlatformStackWalk::CaptureStackBackTrace( uint64* BackTrace, uint
 
 void FWindowsPlatformStackWalk::ProgramCounterToSymbolInfo( uint64 ProgramCounter, FProgramCounterSymbolInfo& out_SymbolInfo )
 {
+	UE::TUniqueLock GlobalLock(GStackWalkingLock);
+
 	// Initialize stack walking as it loads up symbol information which we require.
 	InitStackWalking();
 
@@ -511,6 +530,8 @@ void FWindowsPlatformStackWalk::ProgramCounterToSymbolInfo( uint64 ProgramCounte
 
 void FWindowsPlatformStackWalk::ProgramCounterToSymbolInfoEx(uint64 ProgramCounter, FProgramCounterSymbolInfoEx& out_SymbolInfo)
 {
+	UE::TUniqueLock GlobalLock(GStackWalkingLock);
+
 #if ON_DEMAND_SYMBOL_LOADING
 	// Load symbols for the module
 	bool bShouldReloadModuleMissingDebugSymbols = !FPlatformProperties::IsMonolithicBuild() && FPlatformStackWalk::WantsDetailedCallstacksInNonMonolithicBuilds();
@@ -748,6 +769,8 @@ void LoadSymbolsForModule(HMODULE ModuleHandle, const FString& RemoteStorage)
 		SearchPathList.Append(RemoteStorage);
 	}
 
+	UE::TUniqueLock GlobalLock(GStackWalkingLock);
+
 	SymSetSearchPathW(ProcessHandle, *SearchPathList);
 
 	// Load module.
@@ -793,6 +816,8 @@ void LoadSymbolsForProcessModules(const FString &RemoteStorage)
 
 void LoadSymbolsForModuleByAddress(uint64 Address, const FString& RemoteStorage, bool bShouldReloadModuleMissingDebugSymbols)
 {
+	UE::TUniqueLock GlobalLock(GStackWalkingLock);
+
 	HMODULE ModuleHandle = NULL;
 
 	if (GetModuleHandleEx(GET_MODULE_HANDLE_EX_FLAG_FROM_ADDRESS | GET_MODULE_HANDLE_EX_FLAG_UNCHANGED_REFCOUNT, (LPCTSTR)Address, &ModuleHandle))
@@ -839,6 +864,8 @@ int32 FWindowsPlatformStackWalk::GetProcessModuleCount()
 
 int32 FWindowsPlatformStackWalk::GetProcessModuleSignatures(FStackWalkModuleInfo *ModuleSignatures, const int32 ModuleSignaturesSize)
 {
+	UE::TUniqueLock GlobalLock(GStackWalkingLock);
+
 	FPlatformStackWalk::InitStackWalking();
 
 	HANDLE		ProcessHandle = GProcessHandle; 
@@ -1031,10 +1058,7 @@ FString GetSymbolSearchPath()
  */
 bool FWindowsPlatformStackWalk::InitStackWalkingInternal(void* Process, bool bForceReinitOnProcessMismatch)
 {
-	// DbgHelp functions are not thread safe, but this function can potentially be called from different
-	// threads in our engine, so we take a critical section
-	static FCriticalSection CriticalSection;
-	FScopeLock Lock(&CriticalSection);
+	UE::TUniqueLock GlobalLock(GStackWalkingLock);
 
 	if (GStackWalkingInitialized)
 	{
@@ -1194,6 +1218,8 @@ void FWindowsPlatformStackWalk::RegisterOnModulesChanged()
 
 bool FWindowsPlatformStackWalk::GetFunctionDefinitionLocation(const FString& FunctionSymbolName, const FString& FunctionModuleName, FString& OutPathname, uint32& OutLineNumber, uint32& OutColumnNumber)
 {
+	UE::TUniqueLock GlobalLock(GStackWalkingLock);
+
 #if ON_DEMAND_SYMBOL_LOADING
 	bool bShouldReloadModuleMissingDebugSymbols = !FPlatformProperties::IsMonolithicBuild() && !FunctionModuleName.IsEmpty();
 	if (bShouldReloadModuleMissingDebugSymbols)
