Fix startup crash on Apple A9/A9X cpus as they don't support hardware CRC.

#rnx

diff --git a/Engine/Source/Runtime/Core/Private/Misc/Crc.cpp b/Engine/Source/Runtime/Core/Private/Misc/Crc.cpp
--- a/Engine/Source/Runtime/Core/Private/Misc/Crc.cpp
+++ b/Engine/Source/Runtime/Core/Private/Misc/Crc.cpp
@@ -6,12 +6,13 @@
 #include "Misc/ByteSwap.h"
 
 // As Apple sets __ARM_FEATURE_CRYPTO by default for all of it's arm chips even when not fully supported,
-// it can not be relied upon to detect the lack of hardware CRC in an A8 in an AppleHDTV, so manually force the CRC check
-#if PLATFORM_TVOS
+// it can not be relied upon to detect the lack of hardware CRC in an A8 in an AppleHDTV, so manually force
+// the CRC check. A9/A9X also doesn't support this feature, so force the check until A9 support is dropped.
+#if PLATFORM_TVOS || PLATFORM_IOS
 #	define DETECT_HW_CRC32_SUPPORT_IN_RUNTIME 1
 #elif defined(__clang__) && defined(__ARM_FEATURE_CRYPTO)
 #	define DETECT_HW_CRC32_SUPPORT_IN_RUNTIME 0
-#elif PLATFORM_ANDROID_ARM64 || PLATFORM_LINUXARM64 || PLATFORM_IOS
+#elif PLATFORM_ANDROID_ARM64 || PLATFORM_LINUXARM64
 #	define DETECT_HW_CRC32_SUPPORT_IN_RUNTIME 1
 #else
 #	define DETECT_HW_CRC32_SUPPORT_IN_RUNTIME 0
