CSVProfiler: Return -1 when hot reloading to prevent a crash on Mac/Linux.

diff --git a/Engine/Source/Runtime/Core/Private/ProfilingDebugging/CsvProfiler.cpp b/Engine/Source/Runtime/Core/Private/ProfilingDebugging/CsvProfiler.cpp
--- a/Engine/Source/Runtime/Core/Private/ProfilingDebugging/CsvProfiler.cpp
+++ b/Engine/Source/Runtime/Core/Private/ProfilingDebugging/CsvProfiler.cpp
@@ -35,6 +35,7 @@
 #include "Misc/Fork.h"
 #include "Misc/Guid.h"
 #include "Misc/WildcardString.h"
+#include "Modules/ModuleManager.h"
 
 #include "HAL/PlatformMisc.h"
 
@@ -596,6 +597,13 @@ public:
 	int32 RegisterCategory(const FString& CategoryName, bool bEnableByDefault, bool bIsGlobal)
 	{
 		int32 Index = -1;
+		// During a hot-reload, we attempt to re-register categories and/or statics
+		// that not have been loaded/init'd during a hot-reload, can result in crashing.
+		// Thus when doing a reload, bail. 
+		if (IsReloadActive())
+		{
+			return Index;
+		}
 
 		FScopeLock Lock(&CS);
 		{
