- Fixed some CongicCache warnings (an issue where the per-project redirects would warn if there was no active project, and a LOCTEXT copy/paste problem)

diff --git a/Engine/Source/Runtime/Core/Private/Misc/ConfigCacheIni.cpp b/Engine/Source/Runtime/Core/Private/Misc/ConfigCacheIni.cpp
--- a/Engine/Source/Runtime/Core/Private/Misc/ConfigCacheIni.cpp
+++ b/Engine/Source/Runtime/Core/Private/Misc/ConfigCacheIni.cpp
@@ -635,9 +635,17 @@ ValueType& FindOrAddHeterogeneous(TMap<KeyType, ValueType>& Map, const AltKeyTyp
 namespace
 {
 
+// don't allow warning until all redirects are read in
+bool GAllowConfigRemapWarning = false;
+
 // either show an editor warning, or just write to log for non-editor
 void LogOrEditorWarning(const FText& Msg, const FString& PartialKey, const FString& File)
 {
+	if (!GAllowConfigRemapWarning)
+	{
+		return;
+	}
+	
 	if (GIsEditor)
 	{
 		static TSet<FString> AlreadyWarnedKeys;
@@ -669,6 +677,11 @@ void LogOrEditorWarning(const FText& Msg, const FString& PartialKey, const FStri
 // warn about a section name that's deprecated
 void WarnAboutSectionRemap(const FString& OldValue, const FString& NewValue, const FString& File)
 {
+	if (!GAllowConfigRemapWarning)
+	{
+		return;
+	}
+
 	FFormatNamedArguments Arguments;
 	Arguments.Add(TEXT("OldValue"), FText::FromString(OldValue));
 	Arguments.Add(TEXT("NewValue"), FText::FromString(NewValue));
@@ -694,7 +707,7 @@ static void WarnAboutKeyRemap(const FString& OldValue, const FString& NewValue,
 	Arguments.Add(TEXT("NewValue"), FText::FromString(NewValue));
 	Arguments.Add(TEXT("Section"), FText::FromString(Section));
 	Arguments.Add(TEXT("File"), FText::FromString(File));
-	FText Msg = FText::Format(LOCTEXT("DeprecatedConfig", "Found a deprecated ini key name in {File}. Search for [{OldValue}] and replace with [{NewValue}]"), Arguments);
+	FText Msg = FText::Format(LOCTEXT("DeprecatedConfigKey", "Found a deprecated ini key name in {File}. Search for [{OldValue}] and replace with [{NewValue}]"), Arguments);
 	
 	FString Key = OldValue+Section;
 	if (!IsInGameThread())
@@ -3745,6 +3758,12 @@ static void InitializeConfigRemap()
 	// read in engine and project ini files (these are not hierarchical, so it has to be done in two passes)
 	for (int Pass = 0; Pass < 2; Pass++)
 	{
+		// if there isn't an active project, then skip the project pass
+		if (Pass == 1 && FPaths::ProjectDir() == FPaths::EngineDir())
+		{
+			continue;
+		}
+		
 		Context.Load(*FPaths::Combine(Pass == 0 ? FPaths::EngineDir() : FPaths::ProjectDir(), TEXT("Config/ConfigRedirects.ini")));
 		
 		for (const TPair<FString, FConfigSection>& Section : RemapFile)
@@ -3766,6 +3785,8 @@ static void InitializeConfigRemap()
 			}
 		}
 	}
+	
+	GAllowConfigRemapWarning = true;
 }
 
 void FConfigCacheIni::InitializeConfigSystem()
