Iris - Make sure incoming packet payloads are stored in buffers suitable for Iris bitstreams.

#rnx

diff --git a/Engine/Source/Runtime/Core/Private/Serialization/BitReader.cpp b/Engine/Source/Runtime/Core/Private/Serialization/BitReader.cpp
--- a/Engine/Source/Runtime/Core/Private/Serialization/BitReader.cpp
+++ b/Engine/Source/Runtime/Core/Private/Serialization/BitReader.cpp
@@ -211,7 +211,7 @@ void FBitReader::SetData( FBitReader& Src, int64 CountBits )
 	Src.SerializeBits(Buffer.GetData(), CountBits);
 }
 
-void FBitReader::ResetData(FBitReader& Src, int64 CountBits)
+void FBitReader::ResetData(FBitReader& Src, int64 CountBits, int64 CountBitsWithSlack)
 {
 	Num = CountBits;
 	Pos = 0;
@@ -220,7 +220,7 @@ void FBitReader::ResetData(FBitReader& Src, int64 CountBits)
 	// Setup network version
 	this->SetNetVersionsFromArchive(Src);
 
-	Buffer.Reset();
+	Buffer.Reset((CountBitsWithSlack + 7) >> 3);
 	Buffer.AddUninitialized((CountBits + 7) >> 3);
 	Src.SerializeBits(Buffer.GetData(), CountBits);
 }
diff --git a/Engine/Source/Runtime/Core/Public/Serialization/BitReader.h b/Engine/Source/Runtime/Core/Public/Serialization/BitReader.h
--- a/Engine/Source/Runtime/Core/Public/Serialization/BitReader.h
+++ b/Engine/Source/Runtime/Core/Public/Serialization/BitReader.h
@@ -37,8 +37,8 @@ public:
 	CORE_API void SetData( uint8* Src, int64 CountBits );
 	CORE_API void SetData( TArray<uint8>&& Src, int64 CountBits );
 
-	/** Equivalent to SetData (reset position, copy from Src into internal buffer), but uses Reset not Empty to avoid a realloc if possible. */
-	CORE_API void ResetData(FBitReader& Src, int64 CountBits);
+	/** Equivalent to SetData (reset position, copy from Src into internal buffer), but uses Reset not Empty to avoid a realloc if possible. The internal buffer wiill be able to hold at least the maximum of CountBits and CountBitsWithSlack */
+	CORE_API void ResetData(FBitReader& Src, int64 CountBits, int64 CountBitsWithSlack=0);
 
 // Disable false positive buffer overrun warning during pgoprofile linking step
 PRAGMA_DISABLE_BUFFER_OVERRUN_WARNING
diff --git a/Engine/Source/Runtime/Engine/Private/NetConnection.cpp b/Engine/Source/Runtime/Engine/Private/NetConnection.cpp
--- a/Engine/Source/Runtime/Engine/Private/NetConnection.cpp
+++ b/Engine/Source/Runtime/Engine/Private/NetConnection.cpp
@@ -3302,8 +3302,10 @@ void UNetConnection::DispatchPacket( FBitReader& Reader, int32 PacketId, bool& b
 			// Trace bunch read
 			UE_NET_TRACE_BUNCH_SCOPE(InTraceCollector, Bunch, StartPos, HeaderPos - StartPos);
 
-			Bunch.SetData( Reader, BunchDataBits );
-			if( Reader.IsError() )
+			// Iris requires bitstream data allocations to be a multiple of four bytes.
+			const int64 BunchDataBits32BitAligned = (BunchDataBits + 31) & ~31;
+			Bunch.ResetData(Reader, BunchDataBits, BunchDataBits32BitAligned);
+			if (Reader.IsError())
 			{
 				// Bunch claims it's larger than the enclosing packet.
 				UE_LOG(LogNet, Warning, TEXT("Bunch data overflowed (%i %i+%i/%i)"), IncomingStartPos, HeaderPos, BunchDataBits,
