Workaround for a VC bug: When inlining, VC works to avoid the load->store forwarding stall but loses track of things and inadvertently ends up requiring alignment on temps it creates causing semi random crashes depending on whether the allocator/stack decided to alignment them. The fix is to switch to directly setting the values.

diff --git a/Engine/Source/Runtime/Core/Public/Math/UnrealMathSSE.h b/Engine/Source/Runtime/Core/Public/Math/UnrealMathSSE.h
--- a/Engine/Source/Runtime/Core/Public/Math/UnrealMathSSE.h
+++ b/Engine/Source/Runtime/Core/Public/Math/UnrealMathSSE.h
@@ -656,8 +656,10 @@ FORCEINLINE VectorRegister4Double VectorLoadDouble1(const double* Ptr)
  */
 FORCEINLINE VectorRegister4Float VectorLoadFloat2(const float* Ptr)
 {
-	// This intentionally casts to a double* to be able to load 64 bits of data using the "load 1 double" instruction to fill in the two 32-bit floats.
-	return _mm_castpd_ps(_mm_load1_pd(reinterpret_cast<const double*>(Ptr))); // -V615
+	// Switched from _mm_load1_pd and a cast to avoid a compiler bug in VC. This has the benefit of
+	// being very clear about not needing any alignment, and the optimizer will still result in
+	// movsd and movlhps in both clang and vc.
+	return _mm_setr_ps(Ptr[0], Ptr[1], Ptr[0], Ptr[1]);
 }
 
 FORCEINLINE VectorRegister4Double VectorLoadFloat2(const double* Ptr)
