LLM: Fix crash due to invalid call to GetParent during handling of race condition for registering tags and their parents. SetParent might not yet have been called at the point in the race handling when we call GetParent. Call a new safe function GetParentNameSafeBeforeFinishConstruct instead.
#rnx

diff --git a/Engine/Source/Runtime/Core/Private/HAL/LowLevelMemTracker.cpp b/Engine/Source/Runtime/Core/Private/HAL/LowLevelMemTracker.cpp
--- a/Engine/Source/Runtime/Core/Private/HAL/LowLevelMemTracker.cpp
+++ b/Engine/Source/Runtime/Core/Private/HAL/LowLevelMemTracker.cpp
@@ -1807,9 +1807,7 @@ UE::LLMPrivate::FTagData& FLowLevelMemTracker::RegisterTagData(FName Name, FName
 		// by another thread in between FindOrAddTagData's check of TagDataNameMap->Find and now. Note that it is not
 		// valid for an LLM_SCOPE to be called before a formal registration (e.g. LLM_DECLARE_TAG). If a formal
 		// registration exists for a tag, it must precede its use in any LLM_SCOPE calls.
-		if ((ParentName.IsNone() && TagDataForName->GetParent() != nullptr) ||
-			(!ParentName.IsNone() &&
-				(TagDataForName->GetParent() == nullptr || TagDataForName->GetParent()->GetName() != ParentName)) ||
+		if (ParentName != TagDataForName->GetParentNameSafeBeforeFinishConstruct() ||
 			StatName != TagDataForName->GetStatName() ||
 			SummaryStatName != TagDataForName->GetSummaryStatName() ||
 			bHasEnumTag != TagDataForName->HasEnumTag() ||
@@ -3374,6 +3372,18 @@ FName FTagData::GetParentName() const
 	return ParentName;
 }
 
+FName FTagData::GetParentNameSafeBeforeFinishConstruct() const
+{
+	if (bParentIsName)
+	{
+		return ParentName;
+	}
+	else
+	{
+		return Parent ? Parent->GetName() : NAME_None;
+	}
+}
+
 FName FTagData::GetStatName() const
 {
 	return StatName;
diff --git a/Engine/Source/Runtime/Core/Private/HAL/LowLevelMemTrackerPrivate.h b/Engine/Source/Runtime/Core/Private/HAL/LowLevelMemTrackerPrivate.h
--- a/Engine/Source/Runtime/Core/Private/HAL/LowLevelMemTrackerPrivate.h
+++ b/Engine/Source/Runtime/Core/Private/HAL/LowLevelMemTrackerPrivate.h
@@ -99,6 +99,7 @@ namespace LLMPrivate
 		void AppendDisplayPath(FStringBuilderBase& Result) const;
 		const FTagData* GetParent() const;
 		FName GetParentName() const;
+		FName GetParentNameSafeBeforeFinishConstruct() const;
 		FName GetStatName() const;
 		FName GetSummaryStatName() const;
 		ELLMTag GetEnumTag() const;
