[Insights] Fixed memory tracing for ConcurrentLinearAllocator.

diff --git a/Engine/Source/Runtime/Core/Private/Experimental/ConcurrentLinearAllocator.cpp b/Engine/Source/Runtime/Core/Private/Experimental/ConcurrentLinearAllocator.cpp
--- a/Engine/Source/Runtime/Core/Private/Experimental/ConcurrentLinearAllocator.cpp
+++ b/Engine/Source/Runtime/Core/Private/Experimental/ConcurrentLinearAllocator.cpp
@@ -10,16 +10,11 @@ void* FOsAllocator::Malloc(SIZE_T Size, uint32 Alignment)
 		FMemory::Free(FMemory::Malloc(0));
 		check(GMalloc);
 	}
-	void* Pointer = GMalloc->Malloc(Size, Alignment);
-	MemoryTrace_Alloc(uint64(Pointer), Size, Alignment, EMemoryTraceRootHeap::SystemMemory);
-	MemoryTrace_MarkAllocAsHeap(uint64(Pointer), EMemoryTraceRootHeap::SystemMemory);
-	return Pointer;
+	return GMalloc->Malloc(Size, Alignment);
 }
 
 void FOsAllocator::Free(void* Pointer, SIZE_T Size)
 {
-	MemoryTrace_UnmarkAllocAsHeap(uint64(Pointer), EMemoryTraceRootHeap::SystemMemory);
-	MemoryTrace_Free(uint64(Pointer), EMemoryTraceRootHeap::SystemMemory);
 	GMalloc->Free(Pointer);
 }
 
diff --git a/Engine/Source/Runtime/Core/Public/Experimental/ConcurrentLinearAllocator.h b/Engine/Source/Runtime/Core/Public/Experimental/ConcurrentLinearAllocator.h
--- a/Engine/Source/Runtime/Core/Public/Experimental/ConcurrentLinearAllocator.h
+++ b/Engine/Source/Runtime/Core/Public/Experimental/ConcurrentLinearAllocator.h
@@ -67,13 +67,11 @@ struct FAlignedAllocatorInternal
 		check(IsAligned(Pointer, Alignment));
 #endif
 		MemoryTrace_Alloc(uint64(Pointer), Size, Alignment, EMemoryTraceRootHeap::SystemMemory);
-		MemoryTrace_MarkAllocAsHeap(uint64(Pointer), EMemoryTraceRootHeap::SystemMemory);
 		return Pointer;
 	}
 
 	FORCEINLINE static void Free(void* Pointer, SIZE_T Size)
 	{
-		MemoryTrace_UnmarkAllocAsHeap(uint64(Pointer), EMemoryTraceRootHeap::SystemMemory);
 		MemoryTrace_Free(uint64(Pointer), EMemoryTraceRootHeap::SystemMemory);
 		return AnsiFree(Pointer);
 	}
@@ -275,6 +273,7 @@ class TConcurrentLinearAllocator
 					//if all allocations are already freed we can reuse the Block again
 					Header->~FBlockHeader();
 					ASAN_UNPOISON_MEMORY_REGION( Header, BlockAllocationTag::BlockSize );
+					MemoryTrace_UnmarkAllocAsHeap(uint64(Header), EMemoryTraceRootHeap::SystemMemory);
 					BlockAllocationTag::Allocator::Free(Header, BlockAllocationTag::BlockSize);
 				}
 			}
@@ -287,6 +286,7 @@ class TConcurrentLinearAllocator
 		{
 			static_assert(BlockAllocationTag::BlockSize >= sizeof(FBlockHeader) + sizeof(FAllocationHeader));
 			Header = new (BlockAllocationTag::Allocator::Malloc(BlockAllocationTag::BlockSize, BlockAllocationTag::BlockSize)) FBlockHeader;
+			MemoryTrace_MarkAllocAsHeap(uint64(Header), EMemoryTraceRootHeap::SystemMemory);
 			checkSlow(IsAligned(Header, BlockAllocationTag::BlockSize));
 			if constexpr (!SupportsFastPath)
 			{
@@ -326,6 +326,7 @@ public:
 			{
 				static_assert(BlockAllocationTag::BlockSize >= sizeof(FBlockHeader) + sizeof(FAllocationHeader));
 				Header = new (BlockAllocationTag::Allocator::Malloc(BlockAllocationTag::BlockSize, BlockAllocationTag::BlockSize)) FBlockHeader;
+				MemoryTrace_MarkAllocAsHeap(uint64(Header), EMemoryTraceRootHeap::SystemMemory);
 				checkSlow(IsAligned(Header, BlockAllocationTag::BlockSize));
 				if constexpr (!SupportsFastPath)
 				{
@@ -357,6 +358,7 @@ public:
 				if (HeaderSize + Size + Alignment > BlockAllocationTag::BlockSize)
 				{
 					FBlockHeader* LargeHeader = new (BlockAllocationTag::Allocator::Malloc(HeaderSize + Size + Alignment, BlockAllocationTag::BlockSize)) FBlockHeader;
+					MemoryTrace_MarkAllocAsHeap(uint64(LargeHeader), EMemoryTraceRootHeap::SystemMemory);
 					checkSlow(IsAligned(LargeHeader, alignof(FBlockHeader)));
 
 					uintptr_t LargeAlignedOffset = Align(LargeHeader->NextAllocationPtr, Alignment);
@@ -398,6 +400,7 @@ public:
 				if (HeaderSize + Size + Alignment > BlockAllocationTag::BlockSize)
 				{
 					FBlockHeader* LargeHeader = new (BlockAllocationTag::Allocator::Malloc(HeaderSize + Size + Alignment, BlockAllocationTag::BlockSize)) FBlockHeader;
+					MemoryTrace_MarkAllocAsHeap(uint64(LargeHeader), EMemoryTraceRootHeap::SystemMemory);
 					checkSlow(IsAligned(LargeHeader, alignof(FBlockHeader)));
 
 					uintptr_t LargeAlignedOffset = Align(LargeHeader->NextAllocationPtr, Alignment);
@@ -448,6 +451,7 @@ public:
 				if (Header->NumAllocations.fetch_sub(1, std::memory_order_acq_rel) == 1)
 				{
 					Header->~FBlockHeader();
+					MemoryTrace_UnmarkAllocAsHeap(uint64(Header), EMemoryTraceRootHeap::SystemMemory);
 					BlockAllocationTag::Allocator::Free(Header, Header->NextAllocationPtr - uintptr_t(Header));
 				}
 			}
@@ -465,6 +469,7 @@ public:
 				{
 					Header->~FBlockHeader();
 					ASAN_UNPOISON_MEMORY_REGION( Header, Header->NextAllocationPtr - uintptr_t(Header) );
+					MemoryTrace_UnmarkAllocAsHeap(uint64(Header), EMemoryTraceRootHeap::SystemMemory);
 					BlockAllocationTag::Allocator::Free(Header, Header->NextAllocationPtr - uintptr_t(Header));
 				}
 			}
