CpuProfilerTrace: Fixed incorrect usage of WCHAR_TO_TCHAR (for platforms where wchar_t is 4 bytes) when tracing named events.

#rnx

diff --git a/Engine/Source/Runtime/Core/Public/Stats/Stats2.h b/Engine/Source/Runtime/Core/Public/Stats/Stats2.h
--- a/Engine/Source/Runtime/Core/Public/Stats/Stats2.h
+++ b/Engine/Source/Runtime/Core/Public/Stats/Stats2.h
@@ -1684,7 +1684,7 @@ public:
 			if (UE_TRACE_CHANNELEXPR_IS_ENABLED(CpuChannel))
 			{
 				FName StatName = MinimalNameToName(StatMinimalName);
-				FCpuProfilerTrace::OutputBeginDynamicEventWithId(StatName, WCHAR_TO_TCHAR(InStatId.GetStatDescriptionWIDE()));
+				FCpuProfilerTrace::OutputBeginDynamicEventWithId(StatName, InStatId.GetStatDescriptionWIDE());
 				EmittedEvent |= TraceEvent;
 			}
 #endif
