Remove UE_TRACE_ENABLED check

It should be possible to launch UTS without having tracing enabled in the application. This is important to allow that for example UnrealInsights can launch the trace server. With this the log messages has to be emitted on the LogCore channel, as the LogTrace only exists when tracing is enabled.

diff --git a/Engine/Source/Runtime/Core/Private/ProfilingDebugging/TraceAuxiliary.cpp b/Engine/Source/Runtime/Core/Private/ProfilingDebugging/TraceAuxiliary.cpp
--- a/Engine/Source/Runtime/Core/Private/ProfilingDebugging/TraceAuxiliary.cpp
+++ b/Engine/Source/Runtime/Core/Private/ProfilingDebugging/TraceAuxiliary.cpp
@@ -1227,10 +1227,10 @@ static void LaunchUnrealTraceInternal(const TCHAR* CommandLine)
 static void LaunchUnrealTraceInternal(const TCHAR* CommandLine)
 {
 // TSAN doesn't like fork(), so disable this for now.
-#if !USING_THREAD_SANITISER && UE_TRACE_ENABLED
+#if !USING_THREAD_SANITISER 
 	if (GUnrealTraceLaunched.load(std::memory_order_relaxed))
 	{
-		UE_LOG(LogTrace, Log, TEXT("UnrealTraceServer: Trace store already started"));
+		UE_LOG(LogCore, Log, TEXT("UnrealTraceServer: Trace store already started"));
 		return;
 	}
 
@@ -1245,7 +1245,7 @@ static void LaunchUnrealTraceInternal(const TCHAR* CommandLine)
 
 	if (access(*BinPath, F_OK) < 0)
 	{
-		UE_LOG(LogTrace, Display, TEXT("UnrealTraceServer: Binary not found (%s)"), ANSI_TO_TCHAR(*BinPath));
+		UE_LOG(LogCore, Display, TEXT("UnrealTraceServer: Binary not found (%s)"), ANSI_TO_TCHAR(*BinPath));
 		return;
 	}
 
@@ -1256,7 +1256,7 @@ static void LaunchUnrealTraceInternal(const TCHAR* CommandLine)
 	pid_t UtsPid = fork();
 	if (UtsPid < 0)
 	{
-		UE_LOG(LogTrace, Display, TEXT("UnrealTraceServer: Unable to fork (errno: %d)"), errno);
+		UE_LOG(LogCore, Display, TEXT("UnrealTraceServer: Unable to fork (errno: %d)"), errno);
 		return;
 	}
 	else if (UtsPid == 0)
@@ -1273,7 +1273,7 @@ static void LaunchUnrealTraceInternal(const TCHAR* CommandLine)
 		int32 WaitRet = waitpid(UtsPid, &WaitStatus, 0);
 		if (WaitRet < 0)
 		{
-			UE_LOG(LogTrace, Display, TEXT("UnrealTraceServer: waitpid() error; (errno: %d)"), errno);
+			UE_LOG(LogCore, Display, TEXT("UnrealTraceServer: waitpid() error; (errno: %d)"), errno);
 			return;
 		}
 	}
@@ -1282,11 +1282,11 @@ static void LaunchUnrealTraceInternal(const TCHAR* CommandLine)
 	int32 UtsRet = WEXITSTATUS(WaitStatus);
 	if (UtsRet)
 	{
-		UE_LOG(LogTrace, Display, TEXT("UnrealTraceServer: Trace store returned an error (0x%08x)"), UtsRet);
+		UE_LOG(LogCore, Display, TEXT("UnrealTraceServer: Trace store returned an error (0x%08x)"), UtsRet);
 	}
 	else
 	{
-		UE_LOG(LogTrace, Log, TEXT("UnrealTraceServer: Trace store launch successful"));
+		UE_LOG(LogCore, Log, TEXT("UnrealTraceServer: Trace store launch successful"));
 		GUnrealTraceLaunched.fetch_add(1, std::memory_order_relaxed);
 	}
 #endif // #if !USING_THREAD_SANITISER
