Replace some references to UE4 with UE

#rnx

diff --git a/Engine/Source/Runtime/Core/Private/Serialization/Archive.cpp b/Engine/Source/Runtime/Core/Private/Serialization/Archive.cpp
--- a/Engine/Source/Runtime/Core/Private/Serialization/Archive.cpp
+++ b/Engine/Source/Runtime/Core/Private/Serialization/Archive.cpp
@@ -755,7 +755,7 @@ void FArchive::SerializeCompressedNew(void* V, int64 Length, FName CompressionFo
 		else
 		{
 			UE_LOG(LogSerialization, Log, TEXT("ArchiveName: %s"), *GetArchiveName());
-			UE_LOG(LogSerialization, Log, TEXT("Archive UE4 Version: %d"), UEVer().ToValue());
+			UE_LOG(LogSerialization, Log, TEXT("Archive UE Version: %d"), UEVer().ToValue());
 			UE_LOG(LogSerialization, Log, TEXT("Archive Licensee Version: %d"), LicenseeUEVer());
 			UE_LOG(LogSerialization, Log, TEXT("Position: %lld"), Tell());
 			UE_LOG(LogSerialization, Log, TEXT("Read Size: %lld"), Length);
@@ -796,7 +796,7 @@ void FArchive::SerializeCompressedNew(void* V, int64 Length, FName CompressionFo
 		if ( ! FCompression::IsFormatValid(CompressionFormatToDecode) )
 		{
 			UE_LOG(LogSerialization, Log, TEXT("ArchiveName: %s"), *GetArchiveName());
-			UE_LOG(LogSerialization, Log, TEXT("Archive UE4 Version: %d"), UEVer().ToValue());
+			UE_LOG(LogSerialization, Log, TEXT("Archive UE Version: %d"), UEVer().ToValue());
 			UE_LOG(LogSerialization, Log, TEXT("Archive Licensee Version: %d"), LicenseeUEVer());
 			UE_LOG(LogSerialization, Log, TEXT("Position: %lld"), Tell());
 			UE_LOG(LogSerialization, Log, TEXT("Read Size: %lld"), Length);
