Prevent future inclusion of MinWindows.h, PreWindowsApi.h and PostWindowsApi.h via header guard warning.

 - use WindowsHWrapper.h, AllowWindowsPlatformTypes.h and HideWindowsPlatformTypes.h instead.

diff --git a/Engine/Source/Runtime/Core/Public/Microsoft/MinWindows.h b/Engine/Source/Runtime/Core/Public/Microsoft/MinWindows.h
--- a/Engine/Source/Runtime/Core/Public/Microsoft/MinWindows.h
+++ b/Engine/Source/Runtime/Core/Public/Microsoft/MinWindows.h
@@ -4,5 +4,9 @@
 #if PLATFORM_WINDOWS
 	#include "Windows/MinWindows.h" // HEADER_UNIT_IGNORE
 #else
+	#if !defined(WINDOWS_H_WRAPPER_GUARD)
+	#pragma message("WARNING: do not include Microsoft/MinWindows.h directly. Use Microsoft/WindowsHWrapper.h instead")
+	#endif
+
 	#include "Microsoft/MinWindowsPrivate.h"
 #endif
diff --git a/Engine/Source/Runtime/Core/Public/Microsoft/PostWindowsApi.h b/Engine/Source/Runtime/Core/Public/Microsoft/PostWindowsApi.h
--- a/Engine/Source/Runtime/Core/Public/Microsoft/PostWindowsApi.h
+++ b/Engine/Source/Runtime/Core/Public/Microsoft/PostWindowsApi.h
@@ -5,5 +5,10 @@
 #if PLATFORM_WINDOWS
 	#include "Windows/PostWindowsApi.h"
 #else
+	// this file should only be included from WindowsHWrapper.h
+	#if !defined(WINDOWS_H_WRAPPER_GUARD) 
+	#pragma message("WARNING: do not include Microsoft/PostWindowsApi.h directly. Use Microsoft/WindowsHWrapper.h or Microsoft/HideMicrosoftPlatformTypes.h instead") 
+	#endif
+
 	#include "Microsoft/PostWindowsApiPrivate.h"
 #endif
diff --git a/Engine/Source/Runtime/Core/Public/Microsoft/PreWindowsApi.h b/Engine/Source/Runtime/Core/Public/Microsoft/PreWindowsApi.h
--- a/Engine/Source/Runtime/Core/Public/Microsoft/PreWindowsApi.h
+++ b/Engine/Source/Runtime/Core/Public/Microsoft/PreWindowsApi.h
@@ -5,5 +5,10 @@
 #if PLATFORM_WINDOWS
 	#include "Windows/PreWindowsApi.h" // HEADER_UNIT_IGNORE
 #else
+	// this file should only be included from WindowsHWrapper.h
+	#if !defined(WINDOWS_H_WRAPPER_GUARD) 
+	#pragma message("WARNING: do not include Microsoft/PreWindowsApi.h directly. Use Microsoft/WindowsHWrapper.h or Microsoft/AllowMicrosoftPlatformTypes.h instead") 
+	#endif
+
 	#include "Microsoft/PreWindowsApiPrivate.h"
 #endif
diff --git a/Engine/Source/Runtime/Core/Public/Microsoft/WindowsHWrapper.h b/Engine/Source/Runtime/Core/Public/Microsoft/WindowsHWrapper.h
--- a/Engine/Source/Runtime/Core/Public/Microsoft/WindowsHWrapper.h
+++ b/Engine/Source/Runtime/Core/Public/Microsoft/WindowsHWrapper.h
@@ -8,10 +8,18 @@
 #else
 	#include "CoreTypes.h"
 	#include "HAL/PlatformMemory.h"
+
+	#if defined(WINDOWS_H_WRAPPER_GUARD)
+	#error WINDOWS_H_WRAPPER_GUARD already defined
+	#endif
+	#define WINDOWS_H_WRAPPER_GUARD
+	
 	#include "Microsoft/PreWindowsApiPrivate.h"
 	#ifndef STRICT
 	#define STRICT
 	#endif
 	#include "Microsoft/MinWindowsPrivate.h"
 	#include "Microsoft/PostWindowsApiPrivate.h"
+	
+	#undef WINDOWS_H_WRAPPER_GUARD
 #endif
diff --git a/Engine/Source/Runtime/Core/Public/Windows/MinWindows.h b/Engine/Source/Runtime/Core/Public/Windows/MinWindows.h
--- a/Engine/Source/Runtime/Core/Public/Windows/MinWindows.h
+++ b/Engine/Source/Runtime/Core/Public/Windows/MinWindows.h
@@ -2,6 +2,10 @@
 
 #pragma once
 
+#if !defined(WINDOWS_H_WRAPPER_GUARD)
+#pragma message("WARNING: do not include Windows/MinWindows.h file directly. Use Windows/WindowsHWrapper.h instead")
+#endif
+
 // #TODO: redirect to platform-agnostic version for the time being. Eventually this will become an error
 #include "HAL/Platform.h"
 #if !PLATFORM_WINDOWS
diff --git a/Engine/Source/Runtime/Core/Public/Windows/PostWindowsApi.h b/Engine/Source/Runtime/Core/Public/Windows/PostWindowsApi.h
--- a/Engine/Source/Runtime/Core/Public/Windows/PostWindowsApi.h
+++ b/Engine/Source/Runtime/Core/Public/Windows/PostWindowsApi.h
@@ -8,6 +8,11 @@
 	#include "Microsoft/PostWindowsApi.h"
 #else
 
+// this file should only be included from WindowsHWrapper.h
+#if !defined(WINDOWS_H_WRAPPER_GUARD) 
+#pragma message("WARNING: do not include Windows/PostWindowsApi.h directly. Use Windows/WindowsHWrapper.h or Windows/HideWindowsPlatformTypes.h instead") 
+#endif
+
 // Re-enable warnings
 THIRD_PARTY_INCLUDES_END
 
diff --git a/Engine/Source/Runtime/Core/Public/Windows/PreWindowsApi.h b/Engine/Source/Runtime/Core/Public/Windows/PreWindowsApi.h
--- a/Engine/Source/Runtime/Core/Public/Windows/PreWindowsApi.h
+++ b/Engine/Source/Runtime/Core/Public/Windows/PreWindowsApi.h
@@ -6,6 +6,11 @@
 	#include "Microsoft/PreWindowsApi.h"
 #else
 
+// this file should only be included from WindowsHWrapper.h
+#if !defined(WINDOWS_H_WRAPPER_GUARD) 
+#pragma message("WARNING: do not include Windows/PreWindowsApi.h directly. Use Windows/WindowsHWrapper.h or Windows/AllowWindowsPlatformTypes.h instead") 
+#endif
+
 // Disable the warning that the pack size is changed in this header. We do this globally for
 // clang, since the mechanism of using a sub-header to modify packing generates a -Wpragma-pack
 // warning about modifying packing alignemnt in a header.
diff --git a/Engine/Source/Runtime/Core/Public/Windows/WindowsHWrapper.h b/Engine/Source/Runtime/Core/Public/Windows/WindowsHWrapper.h
--- a/Engine/Source/Runtime/Core/Public/Windows/WindowsHWrapper.h
+++ b/Engine/Source/Runtime/Core/Public/Windows/WindowsHWrapper.h
@@ -9,6 +9,12 @@
 
 #include "CoreTypes.h"
 #include "HAL/PlatformMemory.h"
+
+#if defined(WINDOWS_H_WRAPPER_GUARD)
+#error WINDOWS_H_WRAPPER_GUARD already defined
+#endif
+#define WINDOWS_H_WRAPPER_GUARD
+
 #include "Windows/PreWindowsApi.h"
 #ifndef STRICT
 #define STRICT
@@ -16,4 +22,6 @@
 #include "Windows/MinWindows.h"
 #include "Windows/PostWindowsApi.h"
 
+#undef WINDOWS_H_WRAPPER_GUARD
+
 #endif //PLATFORM_*
\ No newline at end of file
