[Insights] Fixed memory tracing for virtual allocs to track allocations also when just MEM_COMMIT is used (with a nullptr address).

diff --git a/Engine/Source/Runtime/Core/Private/ProfilingDebugging/Microsoft/WindowsMemoryTrace.cpp b/Engine/Source/Runtime/Core/Private/ProfilingDebugging/Microsoft/WindowsMemoryTrace.cpp
--- a/Engine/Source/Runtime/Core/Private/ProfilingDebugging/Microsoft/WindowsMemoryTrace.cpp
+++ b/Engine/Source/Runtime/Core/Private/ProfilingDebugging/Microsoft/WindowsMemoryTrace.cpp
@@ -303,7 +303,8 @@ LPVOID WINAPI FVirtualWinApiHooks::VmAlloc(LPVOID Address, SIZE_T Size, DWORD Ty
 
 	// Track any reserve for now. Going forward we need events to differentiate reserves/commits and
 	// corresponding information on frees.
-	if (Ret != nullptr && (Type & MEM_RESERVE))
+	if (Ret != nullptr &&
+		((Type & MEM_RESERVE) || ((Type & MEM_COMMIT) && Address == nullptr)))
 	{
 		MemoryTrace_Alloc((uint64)Ret, Size, 0, EMemoryTraceRootHeap::SystemMemory);
 		MemoryTrace_MarkAllocAsHeap((uint64)Ret, EMemoryTraceRootHeap::SystemMemory);
@@ -329,7 +330,8 @@ LPVOID WINAPI FVirtualWinApiHooks::VmAllocEx(HANDLE Process, LPVOID Address, SIZ
 {
 	LPVOID Ret = VmAllocExOrig(Process, Address, Size, Type, Protect);
 
-	if (Process == GetCurrentProcess() && Ret != nullptr && (Type & MEM_RESERVE))
+	if (Process == GetCurrentProcess() && Ret != nullptr &&
+		((Type & MEM_RESERVE) || ((Type & MEM_COMMIT) && Address == nullptr)))
 	{
 		MemoryTrace_Alloc((uint64)Ret, Size, 0, EMemoryTraceRootHeap::SystemMemory);
 		MemoryTrace_MarkAllocAsHeap((uint64)Ret, EMemoryTraceRootHeap::SystemMemory);
@@ -355,7 +357,8 @@ LPVOID WINAPI FVirtualWinApiHooks::VmAlloc2(HANDLE Process, LPVOID BaseAddress,
 {
 	LPVOID Ret = VmAlloc2Orig(Process, BaseAddress, Size, Type, PageProtection, ExtendedParameters, ParameterCount);
 
-	if (Process == GetCurrentProcess() && Ret != nullptr && (Type & MEM_RESERVE))
+	if (Process == GetCurrentProcess() && Ret != nullptr &&
+		((Type & MEM_RESERVE) || ((Type & MEM_COMMIT) && BaseAddress == nullptr)))
 	{
 		MemoryTrace_Alloc((uint64)Ret, Size, 0, EMemoryTraceRootHeap::SystemMemory);
 		MemoryTrace_MarkAllocAsHeap((uint64)Ret, EMemoryTraceRootHeap::SystemMemory);
