[Insights] MemoryTrace: Fixed GetDescriptiveName() for FMallocWrapper to report name of wrapped allocator. Also replaced few "checks" in the memory trace API with "checkSlow".

diff --git a/Engine/Source/Runtime/Core/Private/ProfilingDebugging/MemoryTrace.cpp b/Engine/Source/Runtime/Core/Private/ProfilingDebugging/MemoryTrace.cpp
--- a/Engine/Source/Runtime/Core/Private/ProfilingDebugging/MemoryTrace.cpp
+++ b/Engine/Source/Runtime/Core/Private/ProfilingDebugging/MemoryTrace.cpp
@@ -173,6 +173,7 @@ private:
 	};
 
 	static uint32 GetActualAlignment(SIZE_T Size, uint32 Alignment);
+
 	virtual void* Malloc(SIZE_T Size, uint32 Alignment) override;
 	virtual void* Realloc(void* PrevAddress, SIZE_T NewSize, uint32 Alignment) override;
 	virtual void Free(void* Address) override;
@@ -186,6 +187,7 @@ private:
 	virtual void OnMallocInitialized() override                             { InnerMalloc->OnMallocInitialized(); }
 	virtual void OnPreFork() override                                       { InnerMalloc->OnPreFork(); }
 	virtual void OnPostFork() override                                      { InnerMalloc->OnPostFork(); }
+	virtual const TCHAR* GetDescriptiveName() override                      { return InnerMalloc->GetDescriptiveName(); }
 
 	FMalloc* InnerMalloc;
 };
@@ -435,7 +437,7 @@ void MemoryTrace_Alloc(uint64 Address, uint64 Size, uint32 Alignment, HeapId Roo
 		return;
 	}
 
-	check(RootHeap < 16);
+	checkSlow(RootHeap < 16);
 
 	const uint32 AlignmentPow2 = uint32(FPlatformMath::CountTrailingZeros(Alignment));
 	const uint32 Alignment_SizeLower = (AlignmentPow2 << SizeShift) | uint32(Size & ((1 << SizeShift) - 1));
@@ -486,7 +488,7 @@ void MemoryTrace_Free(uint64 Address, HeapId RootHeap)
 		return;
 	}
 
-	check(RootHeap < 16);
+	checkSlow(RootHeap < 16);
 
 	const uint32 CallstackId = GDoNotAllocateInTrace ? 0 : CallstackTrace_GetCurrentId();
 
@@ -527,7 +529,7 @@ void MemoryTrace_ReallocAlloc(uint64 Address, uint64 Size, uint32 Alignment, Hea
 		return;
 	}
 
-	check(RootHeap < 16);
+	checkSlow(RootHeap < 16);
 
 	const uint32 AlignmentPow2 = uint32(FPlatformMath::CountTrailingZeros(Alignment));
 	const uint32 Alignment_SizeLower = (AlignmentPow2 << SizeShift) | uint32(Size & ((1 << SizeShift) - 1));
@@ -568,7 +570,7 @@ void MemoryTrace_ReallocFree(uint64 Address, HeapId RootHeap)
 		return;
 	}
 
-	check(RootHeap < 16);
+	checkSlow(RootHeap < 16);
 
 	const uint32 CallstackId = GDoNotAllocateInTrace ? 0 : CallstackTrace_GetCurrentId();
 
