Add some missing functions to self regsiter open

diff --git a/Engine/Source/Runtime/Core/Private/AutoRTFM/StdLibSupport.cpp b/Engine/Source/Runtime/Core/Private/AutoRTFM/StdLibSupport.cpp
--- a/Engine/Source/Runtime/Core/Private/AutoRTFM/StdLibSupport.cpp
+++ b/Engine/Source/Runtime/Core/Private/AutoRTFM/StdLibSupport.cpp
@@ -115,6 +115,17 @@ UE_AUTORTFM_REGISTER_SELF_FUNCTION(static_cast<const char* (*)(const char*, cons
 UE_AUTORTFM_REGISTER_SELF_FUNCTION(strlen);
 
 UE_AUTORTFM_REGISTER_SELF_FUNCTION(wcscmp);
+UE_AUTORTFM_REGISTER_SELF_FUNCTION(iswupper);
+UE_AUTORTFM_REGISTER_SELF_FUNCTION(iswlower);
+UE_AUTORTFM_REGISTER_SELF_FUNCTION(iswalpha);
+UE_AUTORTFM_REGISTER_SELF_FUNCTION(iswgraph);
+UE_AUTORTFM_REGISTER_SELF_FUNCTION(iswprint);
+UE_AUTORTFM_REGISTER_SELF_FUNCTION(iswpunct);
+UE_AUTORTFM_REGISTER_SELF_FUNCTION(iswalnum);
+UE_AUTORTFM_REGISTER_SELF_FUNCTION(iswdigit);
+UE_AUTORTFM_REGISTER_SELF_FUNCTION(iswxdigit);
+UE_AUTORTFM_REGISTER_SELF_FUNCTION(iswspace);
+UE_AUTORTFM_REGISTER_SELF_FUNCTION(iswcntrl);
 
 #if PLATFORM_WINDOWS
 UE_AUTORTFM_REGISTER_SELF_FUNCTION(sqrt);
@@ -256,6 +267,7 @@ UE_AUTORTFM_REGISTER_SELF_FUNCTION(static_cast<float(*)(float, float)>(&powf));
 UE_AUTORTFM_REGISTER_SELF_FUNCTION(static_cast<double(*)(double, double)>(&pow));
 
 #if PLATFORM_WINDOWS
+UE_AUTORTFM_REGISTER_SELF_FUNCTION(_tcsncmp);
 UE_AUTORTFM_REGISTER_SELF_FUNCTION(_tcslen);
 UE_AUTORTFM_REGISTER_SELF_FUNCTION(_isnan);
 UE_AUTORTFM_REGISTER_SELF_FUNCTION(_finite);
