Moved the Superluminal defines and includes to being private instead of public because all the code is private in Core.

#rnx

diff --git a/Engine/Source/Runtime/Core/Core.Build.cs b/Engine/Source/Runtime/Core/Core.Build.cs
--- a/Engine/Source/Runtime/Core/Core.Build.cs
+++ b/Engine/Source/Runtime/Core/Core.Build.cs
@@ -239,12 +239,12 @@ public class Core : ModuleRules
 				{
 					PublicAdditionalLibraries.Add(Path.Combine(SuperluminalLibDir, "PerformanceAPI_MD.lib"));
 				}
-				PublicDefinitions.Add("WITH_SUPERLUMINAL_PROFILER=1");					  
-				PublicSystemIncludePaths.Add(Path.Combine(SuperluminalApiDir, "include/"));
+				PrivateDefinitions.Add("WITH_SUPERLUMINAL_PROFILER=1");
+				PrivateIncludePaths.Add(Path.Combine(SuperluminalApiDir, "include/"));
 			}
 			else
 			{
-				PublicDefinitions.Add("WITH_SUPERLUMINAL_PROFILER=0");
+				PrivateDefinitions.Add("WITH_SUPERLUMINAL_PROFILER=0");
 			}
 		}
 
