Fix static analysis false positives by explicitly using the condition to early out

diff --git a/Engine/Source/Runtime/Core/Public/Misc/AutomationTest.h b/Engine/Source/Runtime/Core/Public/Misc/AutomationTest.h
--- a/Engine/Source/Runtime/Core/Public/Misc/AutomationTest.h
+++ b/Engine/Source/Runtime/Core/Public/Misc/AutomationTest.h
@@ -20,6 +20,7 @@
 #include "HAL/PlatformProcess.h"
 #include "HAL/PlatformStackWalk.h"
 #include "HAL/PlatformTime.h"
+#include "HAL/PreprocessorHelpers.h"
 #include "HAL/ThreadSafeBool.h"
 #include "Internationalization/Regex.h"
 #include "Logging/LogVerbosity.h"
@@ -65,8 +66,10 @@ class FAutomationTestBase;
 		Info.LineNumber = 1;																				\
 	}
 
+// This macro allows for early exit of the executing unit test function when the condition is false
+// It explicitly uses the condition to ensure static analysis is happy with nullptr checks
 #ifndef UE_RETURN_ON_ERROR
-#define UE_RETURN_ON_ERROR(Condition, Message) if(!AddErrorIfFalse(Condition, Message)) return false;
+#define UE_RETURN_ON_ERROR(Condition, Message) const bool PREPROCESSOR_JOIN(UE____bCondition_Line_, __LINE__) = (Condition); AddErrorIfFalse(PREPROCESSOR_JOIN(UE____bCondition_Line_, __LINE__), (Message)); if(!PREPROCESSOR_JOIN(UE____bCondition_Line_, __LINE__)) return false
 #endif
 
 /**
