Linux: make sure the file with PGO data is flushed and closed.

#rnx

diff --git a/Engine/Source/Runtime/Core/Private/Unix/UnixPlatformMisc.cpp b/Engine/Source/Runtime/Core/Private/Unix/UnixPlatformMisc.cpp
--- a/Engine/Source/Runtime/Core/Private/Unix/UnixPlatformMisc.cpp
+++ b/Engine/Source/Runtime/Core/Private/Unix/UnixPlatformMisc.cpp
@@ -2072,6 +2072,16 @@ bool FUnixPlatformMisc::StopPGOCollectionAndCloseFile()
 		UE_LOG(LogHAL, Error, TEXT("Error writing out PGO file."));
 	}
 	UnixPlatformMisc::GPGICollectionUnderway = false;
+	fflush(stdout);
+
+	// write out a fake file to make sure the previous file is no longer kept open by the LLVM machinery
+	FString DummyProfileFileName = TEXT("/tmp/");
+	DummyProfileFileName += FGuid::NewGuid().ToString();
+	FPlatformMisc::SetEnvironmentVar(TEXT("LLVM_PROFILE_FILE"), *DummyProfileFileName);
+	__llvm_profile_reset_counters();
+	__llvm_profile_initialize_file();
+	__llvm_profile_write_file();
+
 	return true;
 }
 #endif
