Allow casting functions like MoveTemp and Forward to be elided and replaced with direct casts.

diff --git a/Engine/Source/Runtime/Core/Public/HAL/Platform.h b/Engine/Source/Runtime/Core/Public/HAL/Platform.h
--- a/Engine/Source/Runtime/Core/Public/HAL/Platform.h
+++ b/Engine/Source/Runtime/Core/Public/HAL/Platform.h
@@ -721,6 +721,11 @@
 	#endif
 #endif
 
+/** Improves the debugging of functions which act as pure reference casts, eliding the call and doing a direct cast from the argument to the return type */
+#ifndef UE_INTRINSIC_CAST
+	#define UE_INTRINSIC_CAST
+#endif
+
 /** Branch prediction hints */
 #ifndef LIKELY						/* Hints compiler that expression is likely to be true, much softer than UE_ASSUME - allows (penalized by worse performance) expression to be false */
 	#if ( defined(__clang__) || defined(__GNUC__) ) && (PLATFORM_UNIX)	// effect of these on non-Linux platform has not been analyzed as of 2016-03-21
diff --git a/Engine/Source/Runtime/Core/Public/MSVC/MSVCPlatform.h b/Engine/Source/Runtime/Core/Public/MSVC/MSVCPlatform.h
--- a/Engine/Source/Runtime/Core/Public/MSVC/MSVCPlatform.h
+++ b/Engine/Source/Runtime/Core/Public/MSVC/MSVCPlatform.h
@@ -20,3 +20,6 @@
 
 #define PLATFORM_RETURN_ADDRESS()	        _ReturnAddress()
 #define PLATFORM_RETURN_ADDRESS_POINTER()	_AddressOfReturnAddress()
+
+// https://devblogs.microsoft.com/cppblog/improving-the-state-of-debug-performance-in-c/
+#define UE_INTRINSIC_CAST [[msvc::intrinsic]]
diff --git a/Engine/Source/Runtime/Core/Public/Templates/UnrealTemplate.h b/Engine/Source/Runtime/Core/Public/Templates/UnrealTemplate.h
--- a/Engine/Source/Runtime/Core/Public/Templates/UnrealTemplate.h
+++ b/Engine/Source/Runtime/Core/Public/Templates/UnrealTemplate.h
@@ -504,7 +504,7 @@ template <typename T> struct TRemovePointer<T*> { typedef T Type; };
  * const object, because we would prefer to be informed when MoveTemp will have no effect.
  */
 template <typename T>
-FORCEINLINE typename TRemoveReference<T>::Type&& MoveTemp(T&& Obj)
+UE_INTRINSIC_CAST FORCEINLINE typename TRemoveReference<T>::Type&& MoveTemp(T&& Obj)
 {
 	typedef typename TRemoveReference<T>::Type CastType;
 
@@ -522,7 +522,7 @@ FORCEINLINE typename TRemoveReference<T>::Type&& MoveTemp(T&& Obj)
  * where you can but not stop compilation.
  */
 template <typename T>
-FORCEINLINE typename TRemoveReference<T>::Type&& MoveTempIfPossible(T&& Obj)
+UE_INTRINSIC_CAST FORCEINLINE typename TRemoveReference<T>::Type&& MoveTempIfPossible(T&& Obj)
 {
 	typedef typename TRemoveReference<T>::Type CastType;
 	return (CastType&&)Obj;
@@ -584,13 +584,13 @@ FORCEINLINE std::decay_t<T> CopyTempIfNecessary(T&& Val)
  * This is UE's equivalent of std::forward.
  */
 template <typename T>
-FORCEINLINE T&& Forward(typename TRemoveReference<T>::Type& Obj)
+UE_INTRINSIC_CAST FORCEINLINE T&& Forward(typename TRemoveReference<T>::Type& Obj)
 {
 	return (T&&)Obj;
 }
 
 template <typename T>
-FORCEINLINE T&& Forward(typename TRemoveReference<T>::Type&& Obj)
+UE_INTRINSIC_CAST FORCEINLINE T&& Forward(typename TRemoveReference<T>::Type&& Obj)
 {
 	return (T&&)Obj;
 }
@@ -755,7 +755,7 @@ template <
 	typename Base,
 	decltype(ImplicitConv<const volatile Base*>((typename TRemoveReference<T>::Type*)nullptr))* = nullptr
 >
-FORCEINLINE decltype(auto) ForwardAsBase(typename TRemoveReference<T>::Type& Obj)
+UE_INTRINSIC_CAST FORCEINLINE decltype(auto) ForwardAsBase(typename TRemoveReference<T>::Type& Obj)
 {
 	return (TCopyQualifiersAndRefsFromTo_T<T&&, Base>)Obj;
 }
@@ -765,7 +765,7 @@ template <
 	typename Base,
 	decltype(ImplicitConv<const volatile Base*>((typename TRemoveReference<T>::Type*)nullptr))* = nullptr
 >
-FORCEINLINE decltype(auto) ForwardAsBase(typename TRemoveReference<T>::Type&& Obj)
+UE_INTRINSIC_CAST FORCEINLINE decltype(auto) ForwardAsBase(typename TRemoveReference<T>::Type&& Obj)
 {
 	return (TCopyQualifiersAndRefsFromTo_T<T&&, Base>)Obj;
 }
\ No newline at end of file
