Force ConcurrentLinearAllocator allocations alignment to a minimum of 8 bytes when ASAN is enabled to avoid false/positive of use-after-poison
- 8 bytes of heap space are mapped to 1 byte in ASAN shadow space
- Having FAllocationHeader aligned to lower than 8 bytes makes FAllocationHeader ASAN shadow space byte address to be shared with previous FAllocationHeader payload
- When thread A does an allocation, manually unpoison FAllocationHeader, its ASAN shadow space is byte set to 0
- Then thread B frees the allocation that is at the left of new allocated memory from thread A. Thread B Manually poison the payload, FAllocationHeader ASAN shadow space byte address of thread A will then be marked as poisoned again (0x7F)
- Thread A updates FAllocationHeader members -> ASAN exception use-after-poison

#rnx
#tests soaked for 8 hours, no repro of the crash

diff --git a/Engine/Source/Runtime/Core/Public/Experimental/ConcurrentLinearAllocator.h b/Engine/Source/Runtime/Core/Public/Experimental/ConcurrentLinearAllocator.h
--- a/Engine/Source/Runtime/Core/Public/Experimental/ConcurrentLinearAllocator.h
+++ b/Engine/Source/Runtime/Core/Public/Experimental/ConcurrentLinearAllocator.h
@@ -315,7 +315,11 @@ public:
 		checkSlow(Alignment >= 1 && FMath::IsPowerOfTwo(Alignment));
 		if constexpr (!SupportsFastPath)
 		{
-			Alignment = (Alignment < alignof(FAllocationHeader)) ? alignof(FAllocationHeader) : Alignment;
+			Alignment = (Alignment < alignof(FAllocationHeader)) ? alignof(FAllocationHeader) : Alignment; 
+#if IS_ASAN_ENABLED
+			// Make sure FAllocationHeader is 8 bytes aligned so poison / unpoison doesnt end up as false positive
+			Alignment = Align(Alignment, 8);
+#endif
 		}
 
 		static thread_local FBlockHeader* Header;
