Add support for plain text to FAutomationTestBase::AddExpectedMessage

diff --git a/Engine/Plugins/Tests/CQTest/Source/CQTest/Public/Assert/NoDiscardAsserter.inl b/Engine/Plugins/Tests/CQTest/Source/CQTest/Public/Assert/NoDiscardAsserter.inl
--- a/Engine/Plugins/Tests/CQTest/Source/CQTest/Public/Assert/NoDiscardAsserter.inl
+++ b/Engine/Plugins/Tests/CQTest/Source/CQTest/Public/Assert/NoDiscardAsserter.inl
@@ -6,18 +6,7 @@
 
 inline void FNoDiscardAsserter::ExpectError(FString Error, int32 Count)
 {
-	static const FRegexPattern Pattern(TEXT("([\\[\\]{}()^$.?\\\\*+|])"));
-	FRegexMatcher Matcher(Pattern, Error);
-	FString SanitizedError(Error);
-	int32 Offset = 0;
-
-	while (Matcher.FindNext())
-	{
-		int32 Pos = Matcher.GetMatchBeginning() + Offset++;
-		SanitizedError.InsertAt(Pos, TEXT("\\"));
-	}
-
-	TestRunner.AddExpectedError(SanitizedError, EAutomationExpectedErrorFlags::Contains, Count);
+	TestRunner.AddExpectedError(Error, EAutomationExpectedErrorFlags::Contains, Count, false);
 }
 
 inline void FNoDiscardAsserter::ExpectErrorRegex(FString Error, int32 Count)
diff --git a/Engine/Plugins/Tests/TestFramework/Source/TestFramework/Private/AutomationExpectedMessage.spec.cpp b/Engine/Plugins/Tests/TestFramework/Source/TestFramework/Private/AutomationExpectedMessage.spec.cpp
--- a/Engine/Plugins/Tests/TestFramework/Source/TestFramework/Private/AutomationExpectedMessage.spec.cpp
+++ b/Engine/Plugins/Tests/TestFramework/Source/TestFramework/Private/AutomationExpectedMessage.spec.cpp
@@ -37,15 +37,12 @@ void FAutomationExpectedErrorTest::Define()
 
 		It("will not duplicate an existing expected error using the same matcher", [this]()
 		{
-			// Suppress warning logged when adding duplicate value
-			AddExpectedError(TEXT("cannot add duplicate entries"), EAutomationExpectedErrorFlags::Contains, 1);
-			
 			AddExpectedError(TEXT("Expected Error"), EAutomationExpectedErrorFlags::Contains, 1);
 			AddExpectedError(TEXT("Expected Error"), EAutomationExpectedErrorFlags::Contains, 1);
 
 			TArray<FAutomationExpectedMessage> Errors;
 			GetExpectedMessages(Errors, ELogVerbosity::Warning);
-			TestEqual("Expected Errors Count", Errors.Num(), 2);
+			TestEqual("Expected Errors Count", Errors.Num(), 1);
 
 			// Add the expected error to ensure all test conditions pass
 			AddError(TEXT("Expected Error"));
@@ -53,9 +50,6 @@ void FAutomationExpectedErrorTest::Define()
 
 		It("will not duplicate an expected error using a different matcher", [this]()
 		{
-			// Suppress warning logged when adding duplicate value
-			AddExpectedError(TEXT("cannot add duplicate entries"), EAutomationExpectedErrorFlags::Contains, 2);
-
 			AddExpectedError(TEXT("Expected Exact Error"), EAutomationExpectedErrorFlags::Exact, 1);
 			AddExpectedError(TEXT("Expected Exact Error"), EAutomationExpectedErrorFlags::Contains, 1);
 
@@ -64,7 +58,7 @@ void FAutomationExpectedErrorTest::Define()
 
 			TArray<FAutomationExpectedMessage> Errors;
 			GetExpectedMessages(Errors, ELogVerbosity::Warning);
-			TestEqual("Expected Errors Count", Errors.Num(), 3);
+			TestEqual("Expected Errors Count", Errors.Num(), 2);
 
 			// Add the expected errors to ensure all test conditions pass
 			AddError(TEXT("Expected Exact Error"));
@@ -112,6 +106,29 @@ void FAutomationExpectedErrorTest::Define()
 			AddError(FString::Printf(TEXT("Response (%d)"), MAX_uint64));
 		});
 
+		It("will match a regex pattern as case insensitive", [this]()
+		{
+			AddExpectedError(TEXT("Expecte. Error"), EAutomationExpectedErrorFlags::Contains, 0);
+			AddError(TEXT("EXPECTED ERROR"));
+		});
+
+		It("will match a message that contains the plain string pattern", [this]()
+		{
+			AddExpectedErrorPlain(TEXT("Expected Error"), EAutomationExpectedErrorFlags::Contains, 1);
+			AddError(TEXT("Some Expected Error not catched"));
+		});
+
+		It("will match a message as case insensitive", [this]()
+		{
+			AddExpectedErrorPlain(TEXT("Expected Error"), EAutomationExpectedErrorFlags::Contains, 1);
+			AddError(TEXT("Some EXPECTED ERROR not catched"));
+		});
+
+		It("will match a message that is the exact match of the plain string pattern", [this]()
+		{
+			AddExpectedErrorPlain(TEXT("Expected Error not catched"), EAutomationExpectedErrorFlags::Exact, 1);
+			AddError(TEXT("Expected Error not catched"));
+		});
 	});
 
 }
@@ -167,6 +184,18 @@ void FAutomationExpectedErrorFailureTest::Define()
 			AddExpectedError(TEXT("Expected"), EAutomationExpectedErrorFlags::Exact, 1);
 			AddError(TEXT("Expected error"));
 		});
+
+		It("will occur if a message that is not the exact match of the plain string pattern is encountered", [this]()
+		{
+			AddExpectedErrorPlain(TEXT("Expected Error"), EAutomationExpectedErrorFlags::Exact, 0);
+			AddError(TEXT("Some Expected Error not catched"));
+		});
+
+		It("will occur if a message that is a regex match pattern is encountered", [this]()
+		{
+			AddExpectedErrorPlain(TEXT("Expecte. Error"), EAutomationExpectedErrorFlags::Contains, 0);
+			AddError(TEXT("Some Expected Error not catched"));
+		});
 	});
 }
 
diff --git a/Engine/Source/Developer/FunctionalTesting/Private/AutomationBlueprintFunctionLibrary.cpp b/Engine/Source/Developer/FunctionalTesting/Private/AutomationBlueprintFunctionLibrary.cpp
--- a/Engine/Source/Developer/FunctionalTesting/Private/AutomationBlueprintFunctionLibrary.cpp
+++ b/Engine/Source/Developer/FunctionalTesting/Private/AutomationBlueprintFunctionLibrary.cpp
@@ -1385,11 +1385,35 @@ FAutomationScreenshotOptions UAutomationBlueprintFunctionLibrary::GetDefaultScre
 	return Options;
 }
 
-void UAutomationBlueprintFunctionLibrary::AddExpectedLogError(FString ExpectedPatternString, int32 Occurrences, bool ExactMatch)
+void UAutomationBlueprintFunctionLibrary::AddExpectedLogError(FString ExpectedPatternString, int32 Occurrences, bool ExactMatch, bool IsRegex)
 {
 	if (FAutomationTestBase* CurrentTest = FAutomationTestFramework::Get().GetCurrentTest())
 	{
-		CurrentTest->AddExpectedError(ExpectedPatternString, ExactMatch? EAutomationExpectedErrorFlags::Exact:EAutomationExpectedErrorFlags::Contains, Occurrences);
+		CurrentTest->AddExpectedError(ExpectedPatternString, ExactMatch? EAutomationExpectedErrorFlags::Exact:EAutomationExpectedErrorFlags::Contains, Occurrences, IsRegex);
+	}
+}
+
+void UAutomationBlueprintFunctionLibrary::AddExpectedPlainLogError(FString ExpectedString, int32 Occurrences, bool ExactMatch)
+{
+	if (FAutomationTestBase* CurrentTest = FAutomationTestFramework::Get().GetCurrentTest())
+	{
+		CurrentTest->AddExpectedErrorPlain(ExpectedString, ExactMatch ? EAutomationExpectedErrorFlags::Exact : EAutomationExpectedErrorFlags::Contains, Occurrences);
+	}
+}
+
+void UAutomationBlueprintFunctionLibrary::AddExpectedLogMessage(FString ExpectedPatternString, int32 Occurrences, bool ExactMatch, bool IsRegex)
+{
+	if (FAutomationTestBase* CurrentTest = FAutomationTestFramework::Get().GetCurrentTest())
+	{
+		CurrentTest->AddExpectedMessage(ExpectedPatternString, ExactMatch ? EAutomationExpectedErrorFlags::Exact : EAutomationExpectedErrorFlags::Contains, Occurrences, IsRegex);
+	}
+}
+
+void UAutomationBlueprintFunctionLibrary::AddExpectedPlainLogMessage(FString ExpectedString, int32 Occurrences, bool ExactMatch)
+{
+	if (FAutomationTestBase* CurrentTest = FAutomationTestFramework::Get().GetCurrentTest())
+	{
+		CurrentTest->AddExpectedMessagePlain(ExpectedString, ExactMatch ? EAutomationExpectedErrorFlags::Exact : EAutomationExpectedErrorFlags::Contains, Occurrences);
 	}
 }
 
diff --git a/Engine/Source/Developer/FunctionalTesting/Public/AutomationBlueprintFunctionLibrary.h b/Engine/Source/Developer/FunctionalTesting/Public/AutomationBlueprintFunctionLibrary.h
--- a/Engine/Source/Developer/FunctionalTesting/Public/AutomationBlueprintFunctionLibrary.h
+++ b/Engine/Source/Developer/FunctionalTesting/Public/AutomationBlueprintFunctionLibrary.h
@@ -174,10 +174,29 @@ public:
 	static FAutomationScreenshotOptions GetDefaultScreenshotOptionsForRendering(EComparisonTolerance Tolerance = EComparisonTolerance::Low, float Delay = 0.2);
 
 	/**
-	 * Mute the report of log error and warning matching a pattern during an automated test
+	 * Mute the report of log error and warning matching a pattern during an automated test. Treat the pattern as regex by default.
+	 * @param ExpectedPatternString	Expects a Regex pattern.
+	 */
+	UFUNCTION(BlueprintCallable, Category = "Automation", meta = (AdvancedDisplay = "Occurrences, ExactMatch, IsRegex"))
+	static void AddExpectedLogError(FString ExpectedPatternString, int32 Occurrences = 1, bool ExactMatch = false, bool IsRegex = true);
+
+	/**
+	 * Mute the report of log error and warning matching a plain string during an automated test
+	 */
+	UFUNCTION(BlueprintCallable, Category = "Automation", meta = (AdvancedDisplay = "Occurrences, ExactMatch"))
+	static void AddExpectedPlainLogError(FString ExpectedString, int32 Occurrences = 1, bool ExactMatch = false);
+
+	/**
+	 * Expect a specific log message to match a pattern during an automated test regardless of its verbosity. Treat the pattern as regex by default.
+	 */
+	UFUNCTION(BlueprintCallable, Category = "Automation", meta = (AdvancedDisplay = "Occurrences, ExactMatch, IsRegex"))
+	static void AddExpectedLogMessage(FString ExpectedPatternString, int32 Occurrences = 1, bool ExactMatch = false, bool IsRegex = true);
+
+	/**
+	 * Expect a specific log message to match a plain string during an automated test regardless of its verbosity
 	 */
 	UFUNCTION(BlueprintCallable, Category = "Automation", meta = (AdvancedDisplay = "Occurrences, ExactMatch"))
-	static void AddExpectedLogError(FString ExpectedPatternString, int32 Occurrences = 1, bool ExactMatch = false);
+	static void AddExpectedPlainLogMessage(FString ExpectedString, int32 Occurrences = 1, bool ExactMatch = false);
 
 	/**
 	 * Sets all other settings based on an overall value
diff --git a/Engine/Source/Runtime/Core/Private/Misc/AutomationTest.cpp b/Engine/Source/Runtime/Core/Private/Misc/AutomationTest.cpp
--- a/Engine/Source/Runtime/Core/Private/Misc/AutomationTest.cpp
+++ b/Engine/Source/Runtime/Core/Private/Misc/AutomationTest.cpp
@@ -1330,7 +1330,7 @@ bool FAutomationTestBase::HasMetExpectedMessages(ELogVerbosity::Type VerbosityTy
 			bHasMetAllExpectedMessages = false;
 
 			ExecutionInfo.AddEvent(FAutomationEvent(EAutomationEventType::Error,
-				FString::Printf(TEXT("Expected ('%s') log message matching '%s' to occur %d times with %s match type, but it was found %d time(s).")
+				FString::Printf(TEXT("Expected ('%s') level log message or higher matching '%s' to occur %d times with %s match type, but it was found %d time(s).")
 					, LogVerbosityString
 					, *ExpectedMessage.MessagePatternString
 					, ExpectedMessage.ExpectedNumberOfOccurrences
@@ -1345,7 +1345,7 @@ bool FAutomationTestBase::HasMetExpectedMessages(ELogVerbosity::Type VerbosityTy
 				bHasMetAllExpectedMessages = false;
 
 				ExecutionInfo.AddEvent(FAutomationEvent(EAutomationEventType::Error,
-					FString::Printf(TEXT("Expected suppressed ('%s') log message matching '%s' did not occur.")
+					FString::Printf(TEXT("Expected suppressed ('%s') level log message or higher matching '%s' did not occur.")
 						, LogVerbosityString
 						, *ExpectedMessage.MessagePatternString)
 					, ExecutionInfo.GetContext()));
@@ -1353,7 +1353,7 @@ bool FAutomationTestBase::HasMetExpectedMessages(ELogVerbosity::Type VerbosityTy
 			else
 			{
 				ExecutionInfo.AddEvent(FAutomationEvent(EAutomationEventType::Info,
-					FString::Printf(TEXT("Suppressed expected ('%s') log message matching '%s' %d times.")
+					FString::Printf(TEXT("Suppressed expected ('%s') level log message or higher matching '%s' %d times.")
 						, LogVerbosityString
 						, *ExpectedMessage.MessagePatternString
 						, ExpectedMessage.ActualNumberOfOccurrences)
@@ -1389,28 +1389,13 @@ void FAutomationTestBase::AddExpectedMessage(
 	FString ExpectedPatternString,
 	ELogVerbosity::Type ExpectedVerbosity,
 	EAutomationExpectedMessageFlags::MatchType CompareType,
-	int32 Occurrences)
+	int32 Occurrences,
+	bool IsRegex)
 {
 	if (Occurrences >= 0)
 	{
 		FScopeLock Lock(&ActionCS);
-		// If we already have a message matching string in our list, let's not add it again.
-		FAutomationExpectedMessage* FoundEntry = ExpectedMessages.FindByPredicate(
-			[ExpectedPatternString](const FAutomationExpectedMessage& InItem) 
-				{
-					return InItem.MessagePatternString == ExpectedPatternString; 
-				}
-		);
-
-		if (FoundEntry)
-		{
-			UE_LOG(LogAutomationTest, Warning, TEXT("Adding expected log message matching '%s' failed: cannot add duplicate entries"), *ExpectedPatternString)
-		}
-		else
-		{
-			// ToDo: Check that ExpectedPatternString is valid before adding
-			ExpectedMessages.Add(FAutomationExpectedMessage(ExpectedPatternString, ExpectedVerbosity, CompareType, Occurrences));
-		}
+		ExpectedMessages.Add(FAutomationExpectedMessage(ExpectedPatternString, ExpectedVerbosity, CompareType, Occurrences, IsRegex));
 	}
 	else
 	{
@@ -1421,9 +1406,27 @@ void FAutomationTestBase::AddExpectedMessage(
 void FAutomationTestBase::AddExpectedMessage(
 	FString ExpectedPatternString,
 	EAutomationExpectedMessageFlags::MatchType CompareType,
+	int32 Occurrences,
+	bool IsRegex)
+{
+	AddExpectedMessage(MoveTemp(ExpectedPatternString), ELogVerbosity::All, CompareType, Occurrences, IsRegex);	
+}
+
+void FAutomationTestBase::AddExpectedMessagePlain(
+	FString ExpectedString,
+	ELogVerbosity::Type ExpectedVerbosity,
+	EAutomationExpectedMessageFlags::MatchType CompareType,
+	int32 Occurrences)
+{
+	AddExpectedMessage(MoveTemp(ExpectedString), ExpectedVerbosity, CompareType, Occurrences, false);
+}
+
+void FAutomationTestBase::AddExpectedMessagePlain(
+	FString ExpectedString,
+	EAutomationExpectedMessageFlags::MatchType CompareType,
 	int32 Occurrences)
 {
-	AddExpectedMessage(MoveTemp(ExpectedPatternString), ELogVerbosity::All, CompareType, Occurrences);	
+	AddExpectedMessagePlain(MoveTemp(ExpectedString), ELogVerbosity::All, CompareType, Occurrences);
 }
 
 void FAutomationTestBase::GetExpectedMessages(
@@ -1432,7 +1435,7 @@ void FAutomationTestBase::GetExpectedMessages(
 {
 	if (Verbosity == ELogVerbosity::All)
 	{
-		OutInfo = ExpectedMessages;
+		OutInfo = ExpectedMessages.Array();
 	}
 	else
 	{
@@ -1442,12 +1445,21 @@ void FAutomationTestBase::GetExpectedMessages(
 			return FAutomationTestBase::LogCategoryMatchesSeverityInclusive(Message.Verbosity, Verbosity);
 		});
 	}
+	OutInfo.Sort();
 }
 
-void FAutomationTestBase::AddExpectedError(FString ExpectedErrorPattern, EAutomationExpectedErrorFlags::MatchType InCompareType, int32 Occurrences)
+void FAutomationTestBase::AddExpectedError(FString ExpectedErrorPattern, EAutomationExpectedErrorFlags::MatchType InCompareType, int32 Occurrences, bool IsRegex)
 {
 	// Set verbosity to Warning as it's inclusive, and so checks for both Warnings and Errors
-	AddExpectedMessage(MoveTemp(ExpectedErrorPattern), ELogVerbosity::Warning, static_cast<EAutomationExpectedMessageFlags::MatchType>(InCompareType), Occurrences);
+	AddExpectedMessage(MoveTemp(ExpectedErrorPattern), ELogVerbosity::Warning, static_cast<EAutomationExpectedMessageFlags::MatchType>(InCompareType), Occurrences, IsRegex);
+}
+
+void FAutomationTestBase::AddExpectedErrorPlain(
+	FString ExpectedString,
+	EAutomationExpectedErrorFlags::MatchType CompareType,
+	int32 Occurrences)
+{
+	AddExpectedMessagePlain(MoveTemp(ExpectedString), ELogVerbosity::Warning, static_cast<EAutomationExpectedMessageFlags::MatchType>(CompareType), Occurrences);
 }
 
 uint32 FAutomationTestBase::ExtractAutomationTestFlags(FString InTagNotation)
@@ -1712,8 +1724,6 @@ bool FAutomationTestBase::IsExpectedMessage(
 	FScopeLock Lock(&ActionCS);
 	for (FAutomationExpectedMessage& ExpectedMessage : ExpectedMessages)
 	{
-		FRegexMatcher MessageMatcher(ExpectedMessage.MessagePattern, Message);
-
 		// Maintains previous behavior: Adjust so that error and fatal messages are tested against when the input verbosity is "Warning"
 		// Similarly, any message above warning should be considered an "info" message
 		const ELogVerbosity::Type AdjustedMessageVerbosity =
@@ -1722,9 +1732,8 @@ bool FAutomationTestBase::IsExpectedMessage(
 			: ELogVerbosity::VeryVerbose;
 
 		// Compare the incoming message verbosity with the expected verbosity,
-		if (LogCategoryMatchesSeverityInclusive(Verbosity, AdjustedMessageVerbosity) && MessageMatcher.FindNext())
+		if (LogCategoryMatchesSeverityInclusive(Verbosity, AdjustedMessageVerbosity) && ExpectedMessage.Matches(Message))
 		{
-			ExpectedMessage.ActualNumberOfOccurrences++;
 			return true;
 		}
 	}
diff --git a/Engine/Source/Runtime/Core/Public/Misc/AutomationTest.h b/Engine/Source/Runtime/Core/Public/Misc/AutomationTest.h
--- a/Engine/Source/Runtime/Core/Public/Misc/AutomationTest.h
+++ b/Engine/Source/Runtime/Core/Public/Misc/AutomationTest.h
@@ -605,12 +605,13 @@ public:
 
 struct FAutomationExpectedMessage
 {
-	// Original regular expression pattern string matching expected log message.
-	// NOTE: using the Exact comparison type wraps the pattern string with ^ and $ tokens,
-	// but the base pattern string is preserved to allow checks for duplicate entries.
+	// Original string pattern matching expected log message.
+	// If IsRegex is false, it is this string that is used to match.
+	// Otherwise MessagePatternRegex is set to a valid pointer of FRegexPattern.
+	// The base pattern string is preserved none the less to allow checks for duplicate entries.
 	FString MessagePatternString;
-	// Regular expression pattern for MessagePatternString
-	FRegexPattern MessagePattern;
+	// Regular expression pattern from MessagePatternString if regex option was true(default), otherwise it is not set.
+	TOptional<FRegexPattern> MessagePatternRegex;
 	// Type of comparison to perform on error log using MessagePattern.
 	EAutomationExpectedMessageFlags::MatchType CompareType;
 	/** 
@@ -625,25 +626,76 @@ struct FAutomationExpectedMessage
 	/**
 	* Constructor
 	*/
-	FAutomationExpectedMessage(FString& InMessagePattern, ELogVerbosity::Type InVerbosity, EAutomationExpectedMessageFlags::MatchType InCompareType, int32 InExpectedNumberOfOccurrences = 1)
+	FAutomationExpectedMessage(FString& InMessagePattern, ELogVerbosity::Type InVerbosity, EAutomationExpectedMessageFlags::MatchType InCompareType, int32 InExpectedNumberOfOccurrences = 1, bool IsRegex = true)
 		: MessagePatternString(InMessagePattern)
-		, MessagePattern((InCompareType == EAutomationExpectedMessageFlags::Exact) ? FString::Printf(TEXT("^%s$"), *InMessagePattern) : InMessagePattern)
 		, CompareType(InCompareType)
 		, ExpectedNumberOfOccurrences(InExpectedNumberOfOccurrences)
 		, ActualNumberOfOccurrences(0)
 		, Verbosity(InVerbosity)
-	{}
+	{
+		if (IsRegex)
+		{
+			MessagePatternRegex = FRegexPattern((InCompareType == EAutomationExpectedMessageFlags::Exact) ? FString::Printf(TEXT("^%s$"), *InMessagePattern) : InMessagePattern, ERegexPatternFlags::CaseInsensitive);
+		}		
+	}
 
 	FAutomationExpectedMessage(FString& InMessagePattern, ELogVerbosity::Type InVerbosity, int32 InExpectedNumberOfOccurrences)
 		: MessagePatternString(InMessagePattern)
-		, MessagePattern(InMessagePattern)
+		, MessagePatternRegex(FRegexPattern(InMessagePattern, ERegexPatternFlags::CaseInsensitive))
 		, CompareType(EAutomationExpectedMessageFlags::Contains)
 		, ExpectedNumberOfOccurrences(InExpectedNumberOfOccurrences)
 		, ActualNumberOfOccurrences(0)
 		, Verbosity(InVerbosity)
 	{}
+
+	inline bool IsRegex() const
+	{
+		return MessagePatternRegex.IsSet();
+	}
+
+	inline bool IsExactCompareType() const
+	{
+		return CompareType == EAutomationExpectedMessageFlags::Exact;
+	}
+
+	/// <summary>
+	/// Look if Message matches the expected message and increment internal counter if true.
+	/// </summary>
+	/// <param name="Message"></param>
+	/// <returns></returns>
+	bool Matches(const FString& Message)
+	{
+		bool HasMatch = false;
+		if (IsRegex())
+		{
+			FRegexMatcher MessageMatcher(MessagePatternRegex.GetValue(), Message);
+			HasMatch = MessageMatcher.FindNext();
+		}
+		else
+		{
+			HasMatch = Message.Contains(MessagePatternString) && (!IsExactCompareType() || Message.Len() == MessagePatternString.Len());
+		}
+		ActualNumberOfOccurrences += HasMatch;
+		return HasMatch;
+	}
+
+	bool operator==(const FAutomationExpectedMessage& Other) const
+	{
+		return MessagePatternString == Other.MessagePatternString;
+	}
+
+	bool operator<(const FAutomationExpectedMessage& Other) const
+	{
+		return MessagePatternString < Other.MessagePatternString;
+	}
+
 };
 
+FORCEINLINE uint32 GetTypeHash(const FAutomationExpectedMessage& Object)
+{
+	return GetTypeHash(Object.MessagePatternString);
+}
+
 struct FAutomationScreenshotData
 {
 	FString ScreenShotName;
@@ -1489,14 +1541,28 @@ public:
 	* Adds a regex pattern to an internal list that this test will expect to encounter in logs (of the specified verbosity) during its execution. If an expected pattern
 	* is not encountered, it will cause this test to fail.
 	*
-	* @param ExpectedPatternString - The expected message string. Supports basic regex patterns.
+	* @param ExpectedPatternString - The expected message string. Supports basic regex patterns if IsRegex is set to true (the default).
 	* @param ExpectedVerbosity - The expected message verbosity. This is treated as a minimum requirement, so for example the Warning level will intercept Warnings, Errors and Fatal.
 	* @param CompareType - How to match this string with an encountered message, should it match exactly or simply just contain the string.
 	* @param Occurrences - How many times to expect this message string to be seen. If > 0, the message must be seen the exact number of times
 	* specified or the test will fail. If == 0, the message must be seen one or more times (with no upper limit) or the test will fail.
+	* @param IsRegex - If the pattern is to be used as regex or plain string. Default is true.
 	*/
-	CORE_API void AddExpectedMessage(FString ExpectedPatternString, ELogVerbosity::Type ExpectedVerbosity, EAutomationExpectedMessageFlags::MatchType CompareType = EAutomationExpectedMessageFlags::Contains, int32 Occurrences = 1);
-	
+	CORE_API void AddExpectedMessage(FString ExpectedPatternString, ELogVerbosity::Type ExpectedVerbosity, EAutomationExpectedMessageFlags::MatchType CompareType = EAutomationExpectedMessageFlags::Contains, int32 Occurrences = 1, bool IsRegex = true);
+
+	/**
+
+	* Adds a plain string to an internal list that this test will expect to encounter in logs (of the specified verbosity) during its execution. If an expected pattern
+	* is not encountered, it will cause this test to fail.
+	*
+	* @param ExpectedString - The expected message string.
+	* @param ExpectedVerbosity - The expected message verbosity. This is treated as a minimum requirement, so for example the Warning level will intercept Warnings, Errors and Fatal.
+	* @param CompareType - How to match this string with an encountered message, should it match exactly or simply just contain the string.
+	* @param Occurrences - How many times to expect this message string to be seen. If > 0, the message must be seen the exact number of times
+	* specified or the test will fail. If == 0, the message must be seen one or more times (with no upper limit) or the test will fail.
+	*/
+	CORE_API void AddExpectedMessagePlain(FString ExpectedString, ELogVerbosity::Type ExpectedVerbosity, EAutomationExpectedMessageFlags::MatchType CompareType = EAutomationExpectedMessageFlags::Contains, int32 Occurrences = 1);
+
 	/**
 	* Adds a regex pattern to an internal list that this test will expect to encounter in logs (of all severities) during its execution. If an expected pattern
 	* is not encountered, it will cause this test to fail.
@@ -1505,8 +1571,21 @@ public:
 	* @param CompareType - How to match this string with an encountered message, should it match exactly or simply just contain the string.
 	* @param Occurrences - How many times to expect this message string to be seen. If > 0, the message must be seen the exact number of times
 	* specified or the test will fail. If == 0, the message must be seen one or more times (with no upper limit) or the test will fail.
+	* @param IsRegex - If the pattern is to be used as regex or plain string. Default is true.
 	*/
-	CORE_API void AddExpectedMessage(FString ExpectedPatternString, EAutomationExpectedMessageFlags::MatchType CompareType = EAutomationExpectedMessageFlags::Contains, int32 Occurrences = 1);
+	CORE_API void AddExpectedMessage(FString ExpectedPatternString, EAutomationExpectedMessageFlags::MatchType CompareType = EAutomationExpectedMessageFlags::Contains, int32 Occurrences = 1, bool IsRegex = true);
+
+	/**
+	* Adds a plain string to an internal list that this test will expect to encounter in logs (of all severities) during its execution. If an expected pattern
+	* is not encountered, it will cause this test to fail.
+	*
+	* @param ExpectedString - The expected message string.
+	* @param CompareType - How to match this string with an encountered message, should it match exactly or simply just contain the string.
+	* @param Occurrences - How many times to expect this message string to be seen. If > 0, the message must be seen the exact number of times
+	* specified or the test will fail. If == 0, the message must be seen one or more times (with no upper limit) or the test will fail.
+	*/
+	CORE_API void AddExpectedMessagePlain(FString ExpectedString, EAutomationExpectedMessageFlags::MatchType CompareType = EAutomationExpectedMessageFlags::Contains, int32 Occurrences = 1);
+
 
 	/**
 	* Populate the provided expected log messages object with the expected messages contained within the test. Not particularly efficient,
@@ -1524,8 +1603,20 @@ public:
 	* @param CompareType - How to match this string with an encountered error, should it match exactly or simply just contain the string.
 	* @param Occurrences - How many times to expect this error string to be seen. If > 0, the error must be seen the exact number of times
 	* specified or the test will fail. If == 0, the error must be seen one or more times (with no upper limit) or the test will fail.
+	* @param IsRegex - If the pattern is to be used as regex or plain string. Default is true.
+	*/
+	CORE_API void AddExpectedError(FString ExpectedPatternString, EAutomationExpectedErrorFlags::MatchType CompareType = EAutomationExpectedErrorFlags::Contains, int32 Occurrences = 1, bool IsRegex = true);
+
+	/**
+	* Adds a plain string to an internal list that this test will expect to encounter in error or warning logs during its execution. If an expected pattern
+	* is not encountered, it will cause this test to fail.
+	*
+	* @param ExpectedString - The expected message string.
+	* @param CompareType - How to match this string with an encountered error, should it match exactly or simply just contain the string.
+	* @param Occurrences - How many times to expect this error string to be seen. If > 0, the error must be seen the exact number of times
+	* specified or the test will fail. If == 0, the error must be seen one or more times (with no upper limit) or the test will fail.
 	*/
-	CORE_API void AddExpectedError(FString ExpectedPatternString, EAutomationExpectedErrorFlags::MatchType CompareType = EAutomationExpectedErrorFlags::Contains, int32 Occurrences = 1);
+	CORE_API void AddExpectedErrorPlain(FString ExpectedString, EAutomationExpectedErrorFlags::MatchType CompareType = EAutomationExpectedErrorFlags::Contains, int32 Occurrences = 1);
 
 	/**
 	 * Is this a complex tast - if so it will be a stress test.
@@ -2050,7 +2141,7 @@ private:
 	CORE_API void InternalSetSuccessState(bool bSuccessful);
 
 	/* Log messages to be expected while processing this test.*/
-	TArray<FAutomationExpectedMessage> ExpectedMessages;
+	TSet<FAutomationExpectedMessage> ExpectedMessages;
 
 	/** Critical section lock */
 	FCriticalSection ActionCS;
