autortfm:
- remove Context pointer from being passed through every function
- add temporary clang-compatible nop hack to suppress ICF deduplication

diff --git a/Engine/Source/Runtime/Core/Private/AutoRTFM/API.cpp b/Engine/Source/Runtime/Core/Private/AutoRTFM/API.cpp
--- a/Engine/Source/Runtime/Core/Private/AutoRTFM/API.cpp
+++ b/Engine/Source/Runtime/Core/Private/AutoRTFM/API.cpp
@@ -48,6 +48,8 @@ UE_AUTORTFM_FORCEINLINE autortfm_result TransactThenOpenImpl(void (*Work)(void*
 
 extern "C" UE_AUTORTFM_NOAUTORTFM bool autortfm_is_transactional()
 {
+	asm volatile("mov %0, %%rax;"::"i"(__LINE__):"%rax");
+
 	if (GAutoRTFMRuntimeEnabled)
 	{
 		return FContext::Get()->IsTransactional();
@@ -58,12 +60,16 @@ extern "C" UE_AUTORTFM_NOAUTORTFM bool autortfm_is_transactional()
 
 extern "C" UE_AUTORTFM_NOAUTORTFM bool autortfm_is_closed()
 {
+	asm volatile("mov %0, %%rax;"::"i"(__LINE__):"%rax");
+
     return false;
 }
 
 // First Part - the API exposed outside transactions.
 extern "C" UE_AUTORTFM_NOAUTORTFM autortfm_result autortfm_transact(void (*Work)(void* Arg), void* Arg)
 {
+	asm volatile("mov %0, %%rax;"::"i"(__LINE__):"%rax");
+
 	if (GAutoRTFMRuntimeEnabled)
 	{
 	    return static_cast<autortfm_result>(FContext::Get()->Transact(Work, Arg));
@@ -75,47 +81,63 @@ extern "C" UE_AUTORTFM_NOAUTORTFM autortfm_result autortfm_transact(void (*Work)
 
 extern "C" UE_AUTORTFM_NOAUTORTFM autortfm_result autortfm_transact_then_open(void (*Work)(void* Arg), void* Arg)
 {
+	asm volatile("mov %0, %%rax;"::"i"(__LINE__):"%rax");
+
     return TransactThenOpenImpl(Work, Arg);
 }
 
 extern "C" UE_AUTORTFM_NOAUTORTFM void autortfm_commit(void (*Work)(void* Arg), void* Arg)
 {
+	asm volatile("mov %0, %%rax;"::"i"(__LINE__):"%rax");
+
     autortfm_result Result = autortfm_transact(Work, Arg);
 	UE_CLOG(Result != autortfm_committed, LogAutoRTFM, Fatal, TEXT("Unexpected transaction result: %u."), Result);
 }
 
 extern "C" UE_AUTORTFM_NOAUTORTFM void autortfm_abort()
 {
+	asm volatile("mov %0, %%rax;"::"i"(__LINE__):"%rax");
+
 	UE_CLOG(!FContext::IsTransactional(), LogAutoRTFM, Fatal, TEXT("The function `autortfm_abort` was called from outside a transaction."));
 	FContext::Get()->AbortByRequestAndThrow();
 }
 
 extern "C" UE_AUTORTFM_NOAUTORTFM bool autortfm_start_transaction()
 {
+	asm volatile("mov %0, %%rax;"::"i"(__LINE__):"%rax");
+
 	UE_CLOG(!FContext::IsTransactional(), LogAutoRTFM, Fatal, TEXT("The function `autortfm_start_transaction` was called from outside a transact."));
 	return FContext::Get()->StartTransaction();
 }
 
 extern "C" UE_AUTORTFM_NOAUTORTFM autortfm_result autortfm_commit_transaction()
 {
+	asm volatile("mov %0, %%rax;"::"i"(__LINE__):"%rax");
+
 	UE_CLOG(!FContext::IsTransactional(), LogAutoRTFM, Fatal, TEXT("The function `autortfm_commit_transaction` was called from outside a transact."));
 	return static_cast<autortfm_result>(FContext::Get()->CommitTransaction());
 }
 
 extern "C" UE_AUTORTFM_NOAUTORTFM autortfm_result autortfm_abort_transaction()
 {
+	asm volatile("mov %0, %%rax;"::"i"(__LINE__):"%rax");
+
 	UE_CLOG(!FContext::IsTransactional(), LogAutoRTFM, Fatal, TEXT("The function `autortfm_abort_transaction` was called from outside a transact."));
 	return static_cast<autortfm_result>(FContext::Get()->AbortTransaction(false));
 }
 
 extern "C" UE_AUTORTFM_NOAUTORTFM void autortfm_clear_transaction_status()
 {
+	asm volatile("mov %0, %%rax;"::"i"(__LINE__):"%rax");
+
 	ASSERT(FContext::Get()->IsAborting());
 	FContext::Get()->ClearTransactionStatus();
 }
 
 extern "C" UE_AUTORTFM_NOAUTORTFM bool autortfm_is_aborting()
 {
+	asm volatile("mov %0, %%rax;"::"i"(__LINE__):"%rax");
+
 	if (GAutoRTFMRuntimeEnabled)
 	{
 		return FContext::Get()->IsAborting();
@@ -126,26 +148,35 @@ extern "C" UE_AUTORTFM_NOAUTORTFM bool autortfm_is_aborting()
 
 extern "C" UE_AUTORTFM_NOAUTORTFM bool autortfm_current_nest_throw()
 {
+	asm volatile("mov %0, %%rax;"::"i"(__LINE__):"%rax");
+
 	FContext::Get()->Throw();
 	return true;
 }
 
 extern "C" UE_AUTORTFM_NOAUTORTFM void autortfm_abort_if_transactional()
 {
+	asm volatile("mov %0, %%rax;"::"i"(__LINE__):"%rax");
+
 	UE_CLOG(FContext::IsTransactional(), LogAutoRTFM, Fatal, TEXT("The function `autortfm_abort_if_transactional` was called from an open inside a transaction."));
 }
 
 extern "C" UE_AUTORTFM_NOAUTORTFM void autortfm_abort_if_closed()
 {
+	asm volatile("mov %0, %%rax;"::"i"(__LINE__):"%rax");
 }
 
 extern "C" UE_AUTORTFM_NOAUTORTFM void autortfm_open(void (*Work)(void* Arg), void* Arg)
 {
+	asm volatile("mov %0, %%rax;"::"i"(__LINE__):"%rax");
+
 	Work(Arg);
 }
 
 extern "C" UE_AUTORTFM_NOAUTORTFM autortfm_status autortfm_close(void (*Work)(void* Arg), void* Arg)
 {
+	asm volatile("mov %0, %%rax;"::"i"(__LINE__):"%rax");
+
 	autortfm_status Result = autortfm_status_ontrack;
 
 	if (GAutoRTFMRuntimeEnabled)
@@ -153,7 +184,7 @@ extern "C" UE_AUTORTFM_NOAUTORTFM autortfm_status autortfm_close(void (*Work)(vo
 		UE_CLOG(!FContext::IsTransactional(), LogAutoRTFM, Fatal, TEXT("Close called from an outside a transaction."));
 
 		FContext* Context = FContext::Get();
-		void (*WorkClone)(void* Arg, FContext* Context) = FunctionMapLookup(Work, Context, "autortfm_close");
+		void (*WorkClone)(void* Arg) = FunctionMapLookup(Work, "autortfm_close");
 		if (WorkClone)
 		{
 			Result = static_cast<autortfm_status>(Context->CallClosedNest(WorkClone, Arg));
@@ -169,41 +200,55 @@ extern "C" UE_AUTORTFM_NOAUTORTFM autortfm_status autortfm_close(void (*Work)(vo
 
 extern "C" UE_AUTORTFM_NOAUTORTFM void autortfm_record_open_write(void* Ptr, size_t Size)
 {
+	asm volatile("mov %0, %%rax;"::"i"(__LINE__):"%rax");
+
     FContext::Get()->CheckOpenRecordWrite(Ptr);
 	FContext::Get()->RecordWrite(Ptr, Size);
 }
 
 extern "C" UE_AUTORTFM_NOAUTORTFM void autortfm_register_open_function(void* OriginalFunction, void* NewFunction)
 {
+	asm volatile("mov %0, %%rax;"::"i"(__LINE__):"%rax");
+
 	UE_LOG(LogAutoRTFM, Verbose, TEXT("Registering open %p->%p"), OriginalFunction, NewFunction);
     FunctionMapAdd(OriginalFunction, NewFunction);
 }
 
 UE_AUTORTFM_NOAUTORTFM void OpenCommit(TFunction<void()>&& Work)
 {
+	asm volatile("mov %0, %%rax;"::"i"(__LINE__):"%rax");
+
     Work();
 }
 
 UE_AUTORTFM_NOAUTORTFM void OpenAbort(TFunction<void()>&& Work)
 {
+	asm volatile("mov %0, %%rax;"::"i"(__LINE__):"%rax");
 }
 
 extern "C" UE_AUTORTFM_NOAUTORTFM void autortfm_open_commit(void (*Work)(void* Arg), void* Arg)
 {
+	asm volatile("mov %0, %%rax;"::"i"(__LINE__):"%rax");
+
     Work(Arg);
 }
 
 extern "C" UE_AUTORTFM_NOAUTORTFM void autortfm_open_abort(void (*Work)(void* arg), void* Arg)
 {
+	asm volatile("mov %0, %%rax;"::"i"(__LINE__):"%rax");
 }
 
 extern "C" UE_AUTORTFM_NOAUTORTFM void* autortfm_did_allocate(void* Ptr, size_t Size)
 {
+	asm volatile("mov %0, %%rax;"::"i"(__LINE__):"%rax");
+
     return Ptr;
 }
 
 extern "C" UE_AUTORTFM_NOAUTORTFM void autortfm_check_consistency_assuming_no_races()
 {
+	asm volatile("mov %0, %%rax;"::"i"(__LINE__):"%rax");
+
     if (FContext::IsTransactional())
     {
         AutoRTFM::Unreachable();
@@ -259,85 +304,91 @@ extern "C" UE_AUTORTFM_NOAUTORTFM void autortfm_check_abi(void* const Ptr, const
 // Second Part - the same API exposed inside transactions. Note that we don't expose all of the API
 // to transactions! That's intentional. However, things like autortfm_defer_until_commit can be called
 // from an open nest in a transaction.
-bool RTFM_autortfm_is_transactional(FContext* Context)
+bool RTFM_autortfm_is_transactional()
 {
     return true;
 }
 UE_AUTORTFM_REGISTER_OPEN_FUNCTION(autortfm_is_transactional);
 
-bool RTFM_autortfm_is_closed(FContext* Context)
+bool RTFM_autortfm_is_closed()
 {
     return true;
 }
 UE_AUTORTFM_REGISTER_OPEN_FUNCTION(autortfm_is_closed);
 
-autortfm_result RTFM_autortfm_transact(void (*Work)(void* Arg), void* Arg, FContext* Context)
+autortfm_result RTFM_autortfm_transact(void (*Work)(void* Arg), void* Arg)
 {
+	FContext* Context = FContext::Get();
     return static_cast<autortfm_result>(Context->Transact(Work, Arg));
 }
 UE_AUTORTFM_REGISTER_OPEN_FUNCTION(autortfm_transact);
 
-autortfm_result RTFM_autortfm_transact_then_open(void (*Work)(void* Arg), void* Arg, FContext*)
+autortfm_result RTFM_autortfm_transact_then_open(void (*Work)(void* Arg), void* Arg)
 {
     return TransactThenOpenImpl(Work, Arg);
 }
 UE_AUTORTFM_REGISTER_OPEN_FUNCTION(autortfm_transact_then_open);
 
-void RTFM_autortfm_commit(void (*Work)(void* Arg), void* Arg, FContext*)
+void RTFM_autortfm_commit(void (*Work)(void* Arg), void* Arg)
 {
     autortfm_result Result = autortfm_transact(Work, Arg);
 	UE_CLOG(Result != autortfm_committed, LogAutoRTFM, Fatal, TEXT("Unexpected transaction result: %u."), Result);
 }
 UE_AUTORTFM_REGISTER_OPEN_FUNCTION(autortfm_commit);
 
-void RTFM_autortfm_abort(FContext* Context)
+void RTFM_autortfm_abort()
 {
+	FContext* Context = FContext::Get();
     Context->AbortByRequestAndThrow();
 }
 UE_AUTORTFM_REGISTER_OPEN_FUNCTION(autortfm_abort);
 
-void RTFM_autortfm_start_transaction(FContext* Context)
+void RTFM_autortfm_start_transaction()
 {
 	UE_LOG(LogAutoRTFM, Fatal, TEXT("The function `autortfm_start_transaction` was called from closed code."));
 }
 UE_AUTORTFM_REGISTER_OPEN_FUNCTION(autortfm_start_transaction);
 
-void RTFM_autortfm_commit_transaction(FContext* Context)
+void RTFM_autortfm_commit_transaction()
 {
 	UE_LOG(LogAutoRTFM, Fatal, TEXT("The function `RTFM_autortfm_commit_transaction` was called from closed code."));
 }
 UE_AUTORTFM_REGISTER_OPEN_FUNCTION(autortfm_commit_transaction);
 
-autortfm_result RTFM_autortfm_abort_transaction(FContext* Context)
+autortfm_result RTFM_autortfm_abort_transaction()
 {
 	UE_CLOG(!FContext::IsTransactional(), LogAutoRTFM, Fatal, TEXT("The function `autortfm_abort_transaction` was called from outside a transaction"));
+	FContext* Context = FContext::Get();
 	return static_cast<autortfm_result>(Context->AbortTransaction(true));
 }
 UE_AUTORTFM_REGISTER_OPEN_FUNCTION(autortfm_abort_transaction);
 
-void RTFM_autortfm_clear_transaction_status(FContext* Context)
+void RTFM_autortfm_clear_transaction_status()
 {
 	UE_LOG(LogAutoRTFM, Fatal, TEXT("The function `autortfm_clear_transaction_status` was called from closed code."));
 }
 
-void RTFM_autortfm_abort_if_transactional(FContext* Context)
+void RTFM_autortfm_abort_if_transactional()
 {
     UE_LOG(LogAutoRTFM, Verbose, TEXT("The function `autortfm_abort_if_transactional` was called from inside a transaction."));
+	FContext* Context = FContext::Get();
     Context->AbortByRequestAndThrow();
 }
 UE_AUTORTFM_REGISTER_OPEN_FUNCTION(autortfm_abort_if_transactional);
 
-void RTFM_autortfm_abort_if_closed(FContext* Context)
+void RTFM_autortfm_abort_if_closed()
 {
 	UE_LOG(LogAutoRTFM, Verbose, TEXT("The function `autortfm_abort_if_closed` was called from closed inside a transaction."));
+	FContext* Context = FContext::Get();
     Context->AbortByRequestAndThrow();
 }
 UE_AUTORTFM_REGISTER_OPEN_FUNCTION(autortfm_abort_if_closed);
 
-void RTFM_autortfm_open(void (*Work)(void* Arg), void* Arg, FContext* Context)
+void RTFM_autortfm_open(void (*Work)(void* Arg), void* Arg)
 {
 	Work(Arg);
 
+	FContext* Context = FContext::Get();
 	if (Context->IsAborting())
 	{
 		Context->Throw();
@@ -345,58 +396,61 @@ void RTFM_autortfm_open(void (*Work)(void* Arg), void* Arg, FContext* Context)
 }
 UE_AUTORTFM_REGISTER_OPEN_FUNCTION(autortfm_open);
 
-autortfm_status RTFM_autortfm_close(void (*Work)(void* Arg), void* Arg, FContext* Context)
+autortfm_status RTFM_autortfm_close(void (*Work)(void* Arg), void* Arg)
 {
-    void (*WorkClone)(void* Arg, FContext* Context) = FunctionMapLookup(Work, Context, "RTFM_autortfm_close");
+    void (*WorkClone)(void* Arg) = FunctionMapLookup(Work, "RTFM_autortfm_close");
     if (WorkClone)
     {
-        WorkClone(Arg, Context);
+        WorkClone(Arg);
     }
 
 	return static_cast<autortfm_status>(FContext::Get()->GetStatus());
 }
 UE_AUTORTFM_REGISTER_OPEN_FUNCTION(autortfm_close);
 
-extern "C" void RTFM_autortfm_record_open_write(void*, size_t, FContext*)
+extern "C" void RTFM_autortfm_record_open_write(void*, size_t)
 {
 	UE_LOG(LogAutoRTFM, Fatal, TEXT("The function `autortfm_record_open_write` was called from closed code."));
 }
 UE_AUTORTFM_REGISTER_OPEN_FUNCTION(autortfm_record_open_write);
 
-void RTFM_OpenCommit(TFunction<void()>&& Work, FContext* Context)
+void RTFM_OpenCommit(TFunction<void()>&& Work)
 {
+	FContext* Context = FContext::Get();
     ASSERT(Context->GetStatus() == EContextStatus::OnTrack);
     Context->GetCurrentTransaction()->DeferUntilCommit(MoveTemp(Work));
 }
 UE_AUTORTFM_REGISTER_OPEN_FUNCTION(OpenCommit);
 
-void RTFM_OpenAbort(TFunction<void()>&& Work, FContext* Context)
+void RTFM_OpenAbort(TFunction<void()>&& Work)
 {
+	FContext* Context = FContext::Get();
     ASSERT(Context->GetStatus() == EContextStatus::OnTrack);
     Context->GetCurrentTransaction()->DeferUntilAbort(MoveTemp(Work));
 }
 UE_AUTORTFM_REGISTER_OPEN_FUNCTION(OpenAbort);
 
-extern "C" void RTFM_autortfm_open_commit(void (*Work)(void* Arg), void* Arg, FContext* Context)
+extern "C" void RTFM_autortfm_open_commit(void (*Work)(void* Arg), void* Arg)
 {
-    RTFM_OpenCommit([Work, Arg] { Work(Arg); }, Context);
+    RTFM_OpenCommit([Work, Arg] { Work(Arg); });
 }
 UE_AUTORTFM_REGISTER_OPEN_FUNCTION(autortfm_open_commit);
 
-extern "C" void RTFM_autortfm_open_abort(void (*Work)(void* arg), void* Arg, FContext* Context)
+extern "C" void RTFM_autortfm_open_abort(void (*Work)(void* arg), void* Arg)
 {
-    RTFM_OpenAbort([Work, Arg] { Work(Arg); }, Context);
+    RTFM_OpenAbort([Work, Arg] { Work(Arg); });
 }
 UE_AUTORTFM_REGISTER_OPEN_FUNCTION(autortfm_open_abort);
 
-void* RTFM_autortfm_did_allocate(void* Ptr, size_t Size, FContext* Context)
+void* RTFM_autortfm_did_allocate(void* Ptr, size_t Size)
 {
+	FContext* Context = FContext::Get();
     Context->DidAllocate(Ptr, Size);
     return Ptr;
 }
 UE_AUTORTFM_REGISTER_OPEN_FUNCTION(autortfm_did_allocate);
 
-void RTFM_autortfm_check_consistency_assuming_no_races(FContext* Context)
+void RTFM_autortfm_check_consistency_assuming_no_races()
 {
 }
 UE_AUTORTFM_REGISTER_OPEN_FUNCTION(autortfm_check_consistency_assuming_no_races);
diff --git a/Engine/Source/Runtime/Core/Private/AutoRTFM/CompilerInterface.cpp b/Engine/Source/Runtime/Core/Private/AutoRTFM/CompilerInterface.cpp
--- a/Engine/Source/Runtime/Core/Private/AutoRTFM/CompilerInterface.cpp
+++ b/Engine/Source/Runtime/Core/Private/AutoRTFM/CompilerInterface.cpp
@@ -28,7 +28,7 @@ void CheckAlignment(FContext* Context, void* Ptr, size_t AlignmentMask)
     }
 }
 
-extern "C" void autortfm_record_write(FContext* Context, void* Ptr, size_t Size)
+extern "C" void autortfm_record_write(void* Ptr, size_t Size)
 {
 	// check for writes to null here so we end up crashing in the user
 	// code rather than in the autortfm runtime.
@@ -37,10 +37,11 @@ extern "C" void autortfm_record_write(FContext* Context, void* Ptr, size_t Size)
 		return;
 	}
 
+	FContext* Context = FContext::Get();
 	Context->RecordWrite(Ptr, Size);
 }
 
-extern "C" void autortfm_record_write_8(FContext* Context, void* Ptr)
+extern "C" void autortfm_record_write_8(void* Ptr)
 {
 	// check for writes to null here so we end up crashing in the user
 	// code rather than in the autortfm runtime.
@@ -49,30 +50,35 @@ extern "C" void autortfm_record_write_8(FContext* Context, void* Ptr)
 		return;
 	}
 
+	FContext* Context = FContext::Get();
 	Context->RecordWrite<8>(Ptr);
 }
 
-extern "C" void* autortfm_lookup_function(FContext* Context, void* OriginalFunction, const char* Where)
+extern "C" void* autortfm_lookup_function(void* OriginalFunction, const char* Where)
 {
-    return FunctionMapLookup(OriginalFunction, Context, Where);
+	FContext* Context = FContext::Get();
+    return FunctionMapLookup(OriginalFunction, Where);
 }
 
-extern "C" void autortfm_memcpy(void* Dst, const void* Src, size_t Size, FContext* Context)
+extern "C" void autortfm_memcpy(void* Dst, const void* Src, size_t Size)
 {
+	FContext* Context = FContext::Get();
     Memcpy(Dst, Src, Size, Context);
 }
 
-extern "C" void autortfm_memmove(void* Dst, const void* Src, size_t Size, FContext* Context)
+extern "C" void autortfm_memmove(void* Dst, const void* Src, size_t Size)
 {
+	FContext* Context = FContext::Get();
     Memmove(Dst, Src, Size, Context);
 }
 
-extern "C" void autortfm_memset(void* Dst, int Value, size_t Size, FContext* Context)
+extern "C" void autortfm_memset(void* Dst, int Value, size_t Size)
 {
+	FContext* Context = FContext::Get();
     Memset(Dst, Value, Size, Context);
 }
 
-extern "C" void autortfm_llvm_fail(FContext* Context, const char* Message)
+extern "C" void autortfm_llvm_fail(const char* Message)
 {
     if (Message)
     {
@@ -83,6 +89,7 @@ extern "C" void autortfm_llvm_fail(FContext* Context, const char* Message)
 		UE_LOG(LogAutoRTFM, Warning, TEXT("Transaction failing because of language issue."));
 	}
 
+	FContext* Context = FContext::Get();
     Context->AbortByLanguageAndThrow();
 }
 
@@ -91,7 +98,7 @@ extern "C" void autortfm_llvm_alignment_error(FContext* Context, void* Ptr, size
     AbortDueToBadAlignment(Context, Ptr, Alignment, Message);
 }
 
-extern "C" void autortfm_llvm_error(FContext* Context, const char* Message)
+extern "C" void autortfm_llvm_error(const char* Message)
 {
 	if (Message)
 	{
diff --git a/Engine/Source/Runtime/Core/Private/AutoRTFM/Context.cpp b/Engine/Source/Runtime/Core/Private/AutoRTFM/Context.cpp
--- a/Engine/Source/Runtime/Core/Private/AutoRTFM/Context.cpp
+++ b/Engine/Source/Runtime/Core/Private/AutoRTFM/Context.cpp
@@ -149,13 +149,13 @@ bool FContext::IsAborting() const
 	return Status != EContextStatus::OnTrack && Status != EContextStatus::Idle;
 }
 
-EContextStatus FContext::CallClosedNest(void (*ClosedFunction)(void* Arg, FContext* Context), void* Arg)
+EContextStatus FContext::CallClosedNest(void (*ClosedFunction)(void* Arg), void* Arg)
 {
 	TScopedGuard<void*> CurrentNestStackAddressGuard(CurrentTransactStackAddress, &CurrentNestStackAddressGuard);
 
 	PushCallNest(new FCallNest(this));
 
-	CurrentNest->Try([&]() { ClosedFunction(Arg, this); });
+	CurrentNest->Try([&]() { ClosedFunction(Arg); });
 
 	PopCallNest();
 
@@ -254,7 +254,7 @@ ETransactionResult FContext::Transact(void (*Function)(void* Arg), void* Arg)
     
     ASSERT(Status == EContextStatus::Idle || Status == EContextStatus::OnTrack);
 
-    void (*ClonedFunction)(void* Arg, FContext* Context) = FunctionMapTryLookup(Function);
+    void (*ClonedFunction)(void* Arg) = FunctionMapTryLookup(Function);
     if (!ClonedFunction)
     {
 		UE_LOG(LogAutoRTFM, Warning, TEXT("Could not find function %p (%s) in AutoRTFM::FContext::Transact."), Function, *GetFunctionDescription(Function));
@@ -287,7 +287,7 @@ ETransactionResult FContext::Transact(void (*Function)(void* Arg), void* Arg)
         {
             Status = EContextStatus::OnTrack;
             ASSERT(CurrentTransaction->IsFresh());
-			CurrentNest->Try([&] () { ClonedFunction(Arg, this); });
+			CurrentNest->Try([&] () { ClonedFunction(Arg); });
 			ASSERT(CurrentTransaction == NewTransaction); // The transaction lambda should have unwound any nested transactions.
             ASSERT(Status != EContextStatus::Idle);
 
@@ -342,7 +342,7 @@ ETransactionResult FContext::Transact(void (*Function)(void* Arg), void* Arg)
 		PushTransaction(NewTransaction);
 		PushCallNest(NewNest);
 
-		CurrentNest->Try([&]() { ClonedFunction(Arg, this); });
+		CurrentNest->Try([&]() { ClonedFunction(Arg); });
 		ASSERT(CurrentTransaction == NewTransaction);
 
 		Result = ResolveNestedTransaction(NewTransaction);
diff --git a/Engine/Source/Runtime/Core/Private/AutoRTFM/Context.h b/Engine/Source/Runtime/Core/Private/AutoRTFM/Context.h
--- a/Engine/Source/Runtime/Core/Private/AutoRTFM/Context.h
+++ b/Engine/Source/Runtime/Core/Private/AutoRTFM/Context.h
@@ -22,7 +22,7 @@ public:
     // This is public API
     ETransactionResult Transact(void (*Function)(void* Arg), void* Arg);
     
-	EContextStatus CallClosedNest(void (*ClosedFunction)(void* Arg, FContext* Context), void* Arg);
+	EContextStatus CallClosedNest(void (*ClosedFunction)(void* Arg), void* Arg);
 
 	void AbortByRequestAndThrow();
 	void AbortByRequestWithoutThrowing();
diff --git a/Engine/Source/Runtime/Core/Private/AutoRTFM/FunctionMap.h b/Engine/Source/Runtime/Core/Private/AutoRTFM/FunctionMap.h
--- a/Engine/Source/Runtime/Core/Private/AutoRTFM/FunctionMap.h
+++ b/Engine/Source/Runtime/Core/Private/AutoRTFM/FunctionMap.h
@@ -62,9 +62,9 @@ inline void* FunctionMapTryLookup(void* OldFunction)
 }
 
 template<typename TReturnType, typename... TParameterTypes>
-auto FunctionMapTryLookup(TReturnType (*Function)(TParameterTypes...)) -> TReturnType (*)(TParameterTypes..., FContext*)
+auto FunctionMapTryLookup(TReturnType (*Function)(TParameterTypes...)) -> TReturnType (*)(TParameterTypes...)
 {
-    return reinterpret_cast<TReturnType (*)(TParameterTypes..., FContext*)>(FunctionMapTryLookup(reinterpret_cast<void*>(Function)));
+    return reinterpret_cast<TReturnType (*)(TParameterTypes...)>(FunctionMapTryLookup(reinterpret_cast<void*>(Function)));
 }
 
 } // namespace AutoRTFM
diff --git a/Engine/Source/Runtime/Core/Private/AutoRTFM/FunctionMapInlines.h b/Engine/Source/Runtime/Core/Private/AutoRTFM/FunctionMapInlines.h
--- a/Engine/Source/Runtime/Core/Private/AutoRTFM/FunctionMapInlines.h
+++ b/Engine/Source/Runtime/Core/Private/AutoRTFM/FunctionMapInlines.h
@@ -10,7 +10,7 @@
 namespace AutoRTFM
 {
 
-inline void* FunctionMapLookup(void* OldFunction, FContext* Context, const char* Where)
+inline void* FunctionMapLookup(void* OldFunction, const char* Where)
 {
     void* Result = FunctionMapTryLookup(OldFunction);
     if (!Result)
@@ -23,16 +23,16 @@ inline void* FunctionMapLookup(void* OldFunction, FContext* Context, const char*
 		{
 			UE_LOG(LogAutoRTFM, Warning, TEXT("Could not find function %p '%s'."), OldFunction, *GetFunctionDescription(OldFunction));
 		}
-
+		FContext* Context = FContext::Get();
         Context->AbortByLanguageAndThrow();
     }
     return Result;
 }
 
 template<typename TReturnType, typename... TParameterTypes>
-auto FunctionMapLookup(TReturnType (*Function)(TParameterTypes...), FContext* Context, const char* Where) -> TReturnType (*)(TParameterTypes..., FContext*)
+auto FunctionMapLookup(TReturnType (*Function)(TParameterTypes...), const char* Where) -> TReturnType (*)(TParameterTypes...)
 {
-    return reinterpret_cast<TReturnType (*)(TParameterTypes..., FContext*)>(FunctionMapLookup(reinterpret_cast<void*>(Function), Context, Where));
+    return reinterpret_cast<TReturnType (*)(TParameterTypes...)>(FunctionMapLookup(reinterpret_cast<void*>(Function), Where));
 }
 
 } // namespace AutoRTFM
diff --git a/Engine/Source/Runtime/Core/Private/AutoRTFM/StdLibSupport.cpp b/Engine/Source/Runtime/Core/Private/AutoRTFM/StdLibSupport.cpp
--- a/Engine/Source/Runtime/Core/Private/AutoRTFM/StdLibSupport.cpp
+++ b/Engine/Source/Runtime/Core/Private/AutoRTFM/StdLibSupport.cpp
@@ -27,9 +27,10 @@ UE_AUTORTFM_REGISTER_OPEN_FUNCTION_EXPLICIT(memcpy, Memcpy);
 UE_AUTORTFM_REGISTER_OPEN_FUNCTION_EXPLICIT(memmove, Memmove);
 UE_AUTORTFM_REGISTER_OPEN_FUNCTION_EXPLICIT(memset, Memset);
 
-void* RTFM_malloc(size_t Size, FContext* Context)
+void* RTFM_malloc(size_t Size)
 {
     void* Result = malloc(Size);
+	FContext* Context = FContext::Get();
     Context->GetCurrentTransaction()->DeferUntilAbort([Result]
     {
         free(Result);
@@ -40,10 +41,11 @@ void* RTFM_malloc(size_t Size, FContext* Context)
 
 UE_AUTORTFM_REGISTER_OPEN_FUNCTION(malloc);
 
-void RTFM_free(void* Ptr, FContext* Context)
+void RTFM_free(void* Ptr)
 {
     if (Ptr)
     {
+		FContext* Context = FContext::Get();
         Context->GetCurrentTransaction()->DeferUntilCommit([Ptr]
         {
             free(Ptr);
@@ -52,9 +54,9 @@ void RTFM_free(void* Ptr, FContext* Context)
 }
 UE_AUTORTFM_REGISTER_OPEN_FUNCTION(free);
 
-void* RTFM_realloc(void* Ptr, size_t Size, FContext* Context)
+void* RTFM_realloc(void* Ptr, size_t Size)
 {
-    void* NewObject = RTFM_malloc(Size, Context);
+    void* NewObject = RTFM_malloc(Size);
     if (Ptr)
     {
 #if defined(__APPLE__)
@@ -64,43 +66,48 @@ void* RTFM_realloc(void* Ptr, size_t Size, FContext* Context)
 #else
 		const size_t OldSize = malloc_usable_size(Ptr);
 #endif
+		FContext* Context = FContext::Get();
         MemcpyToNew(NewObject, Ptr,  FMath::Min(OldSize, Size), Context);
-        RTFM_free(Ptr, Context);
+        RTFM_free(Ptr);
     }
     return NewObject;
 }
 UE_AUTORTFM_REGISTER_OPEN_FUNCTION(realloc);
 
-char* RTFM_strcpy(char* const Dst, const char* const Src, FContext* const Context)
+char* RTFM_strcpy(char* const Dst, const char* const Src)
 {
     const size_t SrcLen = strlen(Src);
 
+	FContext* Context = FContext::Get();
     Context->RecordWrite(Dst, SrcLen);
     return strcpy(Dst, Src);
 }
 UE_AUTORTFM_REGISTER_OPEN_FUNCTION(strcpy);
 
-char* RTFM_strncpy(char* const Dst, const char* const Src, const size_t Num, FContext* const Context)
+char* RTFM_strncpy(char* const Dst, const char* const Src, const size_t Num)
 {
+	FContext* Context = FContext::Get();
     Context->RecordWrite(Dst, Num);
     return strncpy(Dst, Src, Num);
 }
 UE_AUTORTFM_REGISTER_OPEN_FUNCTION(strncpy);
 
-char* RTFM_strcat(char* const Dst, const char* const Src, FContext* const Context)
+char* RTFM_strcat(char* const Dst, const char* const Src)
 {
     const size_t DstLen = strlen(Dst);
     const size_t SrcLen = strlen(Src);
 
+	FContext* Context = FContext::Get();
     Context->RecordWrite(Dst + DstLen, SrcLen + 1);
     return strcat(Dst, Src);
 }
 UE_AUTORTFM_REGISTER_OPEN_FUNCTION(strcat);
 
-char* RTFM_strncat(char* const Dst, const char* const Src, const size_t Num, FContext* const Context)
+char* RTFM_strncat(char* const Dst, const char* const Src, const size_t Num)
 {
     const size_t DstLen = strlen(Dst);
 
+	FContext* Context = FContext::Get();
     Context->RecordWrite(Dst + DstLen, Num + 1);
     return strncat(Dst, Src, Num);
 }
@@ -205,33 +212,35 @@ UE_AUTORTFM_REGISTER_SELF_FUNCTION(powf);
 // - Str is newly allocated
 // - Format is either newly allocated or not mutated
 // - any strings passed as arguments are either newly allocated or not mutated
-int RTFM_snprintf(char* Str, size_t Size, char* Format, FContext* Context, ...)
+int RTFM_snprintf(char* Str, size_t Size, char* Format, ...)
 {
+	FContext* Context = FContext::Get();
     va_list ArgList;
-    va_start(ArgList, Context);
+    va_start(ArgList, Format);
     int Result = vsnprintf(Str, Size, Format, ArgList);
     va_end(ArgList);
     return Result;
 }
 UE_AUTORTFM_REGISTER_OPEN_FUNCTION(snprintf);
 
-int RTFM_printf(const char* Format, FContext* Context, ...)
+int RTFM_printf(const char* Format, ...)
 {
+	FContext* Context = FContext::Get();
     va_list ArgList;
-    va_start(ArgList, Context);
+    va_start(ArgList, Format);
     int Result = vprintf(Format, ArgList);
     va_end(ArgList);
     return Result;
 }
 UE_AUTORTFM_REGISTER_OPEN_FUNCTION(printf);
 
-int RTFM_putchar(int Char, FContext* Context)
+int RTFM_putchar(int Char)
 {
     return putchar(Char);
 }
 UE_AUTORTFM_REGISTER_OPEN_FUNCTION(putchar);
 
-int RTFM_puts(const char* Str, FContext* Context)
+int RTFM_puts(const char* Str)
 {
     return puts(Str);
 }
@@ -239,7 +248,7 @@ UE_AUTORTFM_REGISTER_OPEN_FUNCTION(puts);
 
 #if PLATFORM_WINDOWS
 
-FILE* RTFM___acrt_iob_func(int Index, FContext* Context)
+FILE* RTFM___acrt_iob_func(int Index)
 {
     switch (Index)
     {
@@ -247,9 +256,12 @@ FILE* RTFM___acrt_iob_func(int Index, FContext* Context)
     case 2:
         return __acrt_iob_func(Index);
     default:
+	{
 		UE_LOG(LogAutoRTFM, Warning, TEXT("Attempt to get file descriptor %d (not 1 or 2) in __acrt_iob_func."), Index);
+		FContext* Context = FContext::Get();
         Context->AbortByLanguageAndThrow();
         return NULL;
+	}
     }
 }
 UE_AUTORTFM_REGISTER_OPEN_FUNCTION(__acrt_iob_func);
@@ -279,7 +291,7 @@ UE_AUTORTFM_REGISTER_SELF_FUNCTION(GetCurrentThreadId);
 
 UE_AUTORTFM_REGISTER_SELF_FUNCTION(TlsGetValue);
 
-BOOL RTFM_TlsSetValue(DWORD dwTlsIndex, LPVOID lpTlsValue, FContext* Context)
+BOOL RTFM_TlsSetValue(DWORD dwTlsIndex, LPVOID lpTlsValue)
 {
 	LPVOID CurrentValue = TlsGetValue(dwTlsIndex);
 
@@ -293,7 +305,7 @@ BOOL RTFM_TlsSetValue(DWORD dwTlsIndex, LPVOID lpTlsValue, FContext* Context)
 UE_AUTORTFM_REGISTER_OPEN_FUNCTION(TlsSetValue);
 #endif // PLATFORM_WINDOWS
 
-wchar_t* RTFM_wcsncpy(wchar_t* Dst, const wchar_t* Src, size_t Count, FContext* Context)
+wchar_t* RTFM_wcsncpy(wchar_t* Dst, const wchar_t* Src, size_t Count)
 {
 	AutoRTFM::Unreachable();
 }
@@ -313,8 +325,9 @@ UE_AUTORTFM_REGISTER_OPEN_FUNCTION(wcsncpy);
 #pragma warning(pop)
 #endif
 
-int RTFM_atexit(void(__cdecl*Callback)(void), FContext* Context)
+int RTFM_atexit(void(__cdecl*Callback)(void))
 {
+	FContext* Context = FContext::Get();
 	Context->GetCurrentTransaction()->DeferUntilCommit([Callback]
 		{
 			atexit(Callback);
