fixed "signed/unsigned mismatch" warnign on using a TArray constructor with `ExtraSlack`, happens only for TArray64 for unknown reason. example: `TArray64<uint32>{ TArray64<uint32>{}, 1 }`

diff --git a/Engine/Source/Runtime/Core/Public/Containers/Array.h b/Engine/Source/Runtime/Core/Public/Containers/Array.h
--- a/Engine/Source/Runtime/Core/Public/Containers/Array.h
+++ b/Engine/Source/Runtime/Core/Public/Containers/Array.h
@@ -606,7 +606,7 @@ private:
 			USizeType NewMax        = (USizeType)LocalArrayNum + (USizeType)ExtraSlack;
 
 			// This should only happen when we've underflowed or overflowed SizeType
-			if ((SizeType)NewMax < LocalArrayNum)
+			if ((SizeType)NewMax < (SizeType)LocalArrayNum)
 			{
 				UE::Core::Private::OnInvalidArrayNum((unsigned long long)ExtraSlack);
 			}
