- Fix CVar SetByHotfix warnings

diff --git a/Engine/Source/Runtime/Core/Private/HAL/ConsoleManager.cpp b/Engine/Source/Runtime/Core/Private/HAL/ConsoleManager.cpp
--- a/Engine/Source/Runtime/Core/Private/HAL/ConsoleManager.cpp
+++ b/Engine/Source/Runtime/Core/Private/HAL/ConsoleManager.cpp
@@ -182,7 +182,7 @@ public:
 				*GetString()
 				);
 				
-			if (OldPri == ECVF_SetByConsoleVariablesIni || OldPri == ECVF_SetByCommandline || OldPri == ECVF_SetBySystemSettingsIni)
+			if (OldPri == ECVF_SetByConsoleVariablesIni || OldPri == ECVF_SetByCommandline || OldPri == ECVF_SetBySystemSettingsIni || OldPri == ECVF_SetByHotfix)
 			{
 				// Set by an ini that has to be hand edited, a deliberate fail
 				UE_LOG(LogConsoleManager, Verbose, TEXT("%s"), *Message);
