GetVarArgs/LLM: Move the tests for GetVarArgs into a test function rather than doing them on the first call to GetVarArgs. The first call to GetVarArgs may come from inside the LowLevelMemTracker, and trying to run the test will allocate memory and cause reentry into LLM when it is in a bad startup state.
#rnx

diff --git a/Engine/Source/Runtime/Core/Private/GenericPlatform/GenericWidePlatformString.cpp b/Engine/Source/Runtime/Core/Private/GenericPlatform/GenericWidePlatformString.cpp
--- a/Engine/Source/Runtime/Core/Private/GenericPlatform/GenericWidePlatformString.cpp
+++ b/Engine/Source/Runtime/Core/Private/GenericPlatform/GenericWidePlatformString.cpp
@@ -1,12 +1,18 @@
 // Copyright Epic Games, Inc. All Rights Reserved.
 
 #include "GenericPlatform/GenericWidePlatformString.h"
+
+#if PLATFORM_USE_GENERIC_STRING_IMPLEMENTATION
+
 #include "HAL/UnrealMemory.h"
-#include "Templates/UnrealTemplate.h"
 #include "Logging/LogCategory.h"
 #include "Logging/LogMacros.h"
+#include "Templates/UnrealTemplate.h"
 
-#if PLATFORM_USE_GENERIC_STRING_IMPLEMENTATION
+#if WITH_LOW_LEVEL_TESTS
+#include "TestHarness.h"
+#include <catch2/generators/catch_generators.hpp>
+#endif
 
 DEFINE_LOG_CATEGORY_STATIC(LogStandardPlatformString, Log, All);
 
@@ -220,8 +226,8 @@ int iswspace(wint_t wc)
 #endif
 
 
-#if !(UE_BUILD_SHIPPING || UE_BUILD_TEST)
-static const int OUTPUT_SIZE = 256;
+#if WITH_LOW_LEVEL_TESTS
+constexpr int OUTPUT_SIZE = 256;
 int32 TestGetVarArgs(WIDECHAR* OutputString, const WIDECHAR* Format, ...)
 {
 	va_list ArgPtr;
@@ -231,49 +237,49 @@ int32 TestGetVarArgs(WIDECHAR* OutputString, const WIDECHAR* Format, ...)
 	return Result;
 }
 
-void RunGetVarArgsTests()
+TEST_CASE("Core::PlatformString::GetVarArgs", "[Core][String][Smoke]")
 {
 	WIDECHAR OutputString[OUTPUT_SIZE];
 
 	TestGetVarArgs(OutputString, TEXT("Test A|%-20s|%20s|%10.2f|%-10.2f|"), TEXT("LEFT"), TEXT("RIGHT"), 33.333333, 66.666666);
-	checkf(FString(OutputString) == FString(TEXT("Test A|LEFT                |               RIGHT|     33.33|66.67     |")), OutputString);
+	CHECK_MESSAGE(OutputString, FString(OutputString) == FString(TEXT("Test A|LEFT                |               RIGHT|     33.33|66.67     |")));
 
 	TestGetVarArgs(OutputString, TEXT("Test B|Percents:%%%%%%%d|"), 3);
-	checkf(FString(OutputString) == FString(TEXT("Test B|Percents:%%%3|")), OutputString);
+	CHECK_MESSAGE(OutputString, FString(OutputString) == FString(TEXT("Test B|Percents:%%%3|")));
 
 	TestGetVarArgs(OutputString, TEXT("Test C|%d|%i|%X|%x|%u|"), 12345, 54321, 0x123AbC, 15, 99);
-	checkf(FString(OutputString) == FString(TEXT("Test C|12345|54321|123ABC|f|99|")), OutputString);
+	CHECK_MESSAGE(OutputString, FString(OutputString) == FString(TEXT("Test C|12345|54321|123ABC|f|99|")));
 
 	TestGetVarArgs(OutputString, TEXT("Test D|%p|"), 0x12345);
 #ifdef _MSC_VER
 	// MSVC's standard library formats pointers differently
-	checkf(FString(OutputString) == FString(TEXT("Test D|0000000000012345|")), OutputString);
+	CHECK_MESSAGE(OutputString, FString(OutputString) == FString(TEXT("Test D|0000000000012345|")));
 #else
-	checkf(FString(OutputString) == FString(TEXT("Test D|0x12345|")), OutputString);
+	CHECK_MESSAGE(OutputString, FString(OutputString) == FString(TEXT("Test D|0x12345|")));
 #endif
 
 	TestGetVarArgs(OutputString, TEXT("Test E|%" INT64_FMT "|"), int64(12345678912345LL));
-	checkf(FString(OutputString) == FString(TEXT("Test E|12345678912345|")), OutputString);
+	CHECK_MESSAGE(OutputString, FString(OutputString) == FString(TEXT("Test E|12345678912345|")));
 
 	TestGetVarArgs(OutputString, TEXT("Test F|%f|%e|%g|"), 123.456, 123.456, 123.456);
-	checkf(FString(OutputString) == FString(TEXT("Test F|123.456000|1.234560e+02|123.456|")), OutputString);
+	CHECK_MESSAGE(OutputString, FString(OutputString) == FString(TEXT("Test F|123.456000|1.234560e+02|123.456|")));
 
 	TestGetVarArgs(OutputString, TEXT("Test G|%" UPTRINT_X_FMT"|"), UPTRINT(49374));
-	checkf(FString(OutputString) == FString(TEXT("Test G|C0DE|")), OutputString);
+	CHECK_MESSAGE(OutputString, FString(OutputString) == FString(TEXT("Test G|C0DE|")));
 
 	TestGetVarArgs(OutputString, TEXT("Test H|%" UPTRINT_x_FMT "|"), UPTRINT(49374));
-	checkf(FString(OutputString) == FString(TEXT("Test H|c0de|")), OutputString);
+	CHECK_MESSAGE(OutputString, FString(OutputString) == FString(TEXT("Test H|c0de|")));
 
 	TestGetVarArgs(OutputString, TEXT("Test I|%" UINT64_FMT "|"), MAX_uint64);
-	checkf(FString(OutputString) == FString(TEXT("Test I|18446744073709551615|")), OutputString);
+	CHECK_MESSAGE(OutputString, FString(OutputString) == FString(TEXT("Test I|18446744073709551615|")));
 
 	TestGetVarArgs(OutputString, TEXT("Test J|%*c|"), 10, 'J');
-	checkf(FString(OutputString) == FString(TEXT("Test J|         J|")), OutputString);
+	CHECK_MESSAGE(OutputString, FString(OutputString) == FString(TEXT("Test J|         J|")));
 
 	TestGetVarArgs(OutputString, TEXT("Test K|%-5c|"), 'K');
-	checkf(FString(OutputString) == FString(TEXT("Test K|K    |")), OutputString);
+	CHECK_MESSAGE(OutputString, FString(OutputString) == FString(TEXT("Test K|K    |")));
 }
-#endif
+#endif // WITH_LOW_LEVEL_TESTS
 
 namespace
 {
@@ -404,15 +410,6 @@ int32 FGenericWidePlatformString::GetVarArgs( WIDECHAR* Dest, SIZE_T DestSize, c
 	static_assert(sizeof(TCHAR) == 2, "TCHAR is expected to be 16-bit");
 #endif
 
-#if !(UE_BUILD_SHIPPING || UE_BUILD_TEST)
-	static bool bTested = false;
-	if(!bTested)
-	{
-		bTested = true;
-		RunGetVarArgsTests();
-	}
-#endif
-
 	if (Fmt == nullptr)
 	{
 		if ((DestSize > 0) && (Dest != nullptr))
@@ -909,4 +906,4 @@ int32 FGenericWidePlatformString::GetVarArgs( WIDECHAR* Dest, SIZE_T DestSize, c
 	return Result;
 }
 
-#endif
+#endif // PLATFORM_USE_GENERIC_STRING_IMPLEMENTATION
