LowLevelMemTracker: Add more assertions to narrow down why PopTag might fail to find a state.
#rnx

diff --git a/Engine/Source/Runtime/Core/Private/HAL/LowLevelMemTracker.cpp b/Engine/Source/Runtime/Core/Private/HAL/LowLevelMemTracker.cpp
--- a/Engine/Source/Runtime/Core/Private/HAL/LowLevelMemTracker.cpp
+++ b/Engine/Source/Runtime/Core/Private/HAL/LowLevelMemTracker.cpp
@@ -2864,8 +2864,8 @@ bool HandleAssert(bool bLog, const TCHAR* Format, ...)
 	if (bLog)
 	{
 		TCHAR DescriptionString[4096];
-	GET_VARARGS(DescriptionString, UE_ARRAY_COUNT(DescriptionString), UE_ARRAY_COUNT(DescriptionString) - 1,
-		Format, Format);
+		GET_VARARGS(DescriptionString, UE_ARRAY_COUNT(DescriptionString), UE_ARRAY_COUNT(DescriptionString) - 1,
+			Format, Format);
 
 		FPlatformMisc::LowLevelOutputDebugString(DescriptionString);
 
@@ -3555,6 +3555,7 @@ FLLMThreadState* FLLMTracker::GetOrCreateState()
 	if (State == nullptr)
 	{
 		State = LLMRef.Allocator.New<FLLMThreadState>();
+		LLMCheckf(State != nullptr, TEXT("LLMRef.Allocator.New returned nullptr."));
 
 		// Add to pending thread states, these will be consumed on the GT
 		{
@@ -3564,6 +3565,10 @@ FLLMThreadState* FLLMTracker::GetOrCreateState()
 
 		// push to Tls
 		FPlatformTLS::SetTlsValue(TlsSlot, State);
+
+		// Verify the Tls slot is working
+		FLLMThreadState* StoredPointer = GetState();
+		LLMCheckf(StoredPointer == State, TEXT("SetTlsValue/GetTlsValue failed to store and retrieve the state pointer."));
 	}
 	return State;
 }
