Annotate the AutoRTFM API with the correct CORE_API.

diff --git a/Engine/Source/Runtime/Core/Public/AutoRTFM/AutoRTFM.h b/Engine/Source/Runtime/Core/Public/AutoRTFM/AutoRTFM.h
--- a/Engine/Source/Runtime/Core/Public/AutoRTFM/AutoRTFM.h
+++ b/Engine/Source/Runtime/Core/Public/AutoRTFM/AutoRTFM.h
@@ -33,11 +33,18 @@
 #include <memory.h>
 #include <type_traits>
 
+#if defined(UE_AUTORTFM_STANDALONE)
+#define UE_AUTORTFM_API
+#else
+#include <HAL/Platform.h>
+#define UE_AUTORTFM_API CORE_API
+#endif
+
 #if UE_AUTORTFM
 template <typename FuncType> class TFunction;
 #endif
 
-#define UE_AUTORTFM_UNUSED(UNUSEDVAR)   (void)UNUSEDVAR
+#define UE_AUTORTFM_UNUSED(UNUSEDVAR) (void)UNUSEDVAR
 
 #ifdef __cplusplus
 extern "C"
@@ -76,7 +83,7 @@ typedef enum
 // Tells if we are currently running in a transaction. This will return true in an
 // open nest (see autortfm_open).
 #if UE_AUTORTFM
-bool autortfm_is_transactional(void);
+UE_AUTORTFM_API bool autortfm_is_transactional(void);
 #else
 UE_AUTORTFM_FORCEINLINE bool autortfm_is_transactional(void)
 {
@@ -97,7 +104,7 @@ UE_AUTORTFM_FORCEINLINE bool autortfm_is_transactional(void)
 //   the transactional openation in a closed nest. So, it's often more correct
 //   to test is_closed than is_transactional.
 #if UE_AUTORTFM
-bool autortfm_is_closed(void);
+UE_AUTORTFM_API bool autortfm_is_closed(void);
 #else
 UE_AUTORTFM_FORCEINLINE bool autortfm_is_closed(void)
 {
@@ -111,8 +118,8 @@ UE_AUTORTFM_FORCEINLINE bool autortfm_is_closed(void)
 // transaction, so the effects can be reversed later if the root transaction aborts, even
 // if this nested transaction succeeds.
 #if UE_AUTORTFM
-autortfm_result autortfm_transact(void (*work)(void* arg), void* arg);
-autortfm_result autortfm_transact_then_open(void (*work)(void* arg), void* arg);
+UE_AUTORTFM_API autortfm_result autortfm_transact(void (*work)(void* arg), void* arg);
+UE_AUTORTFM_API autortfm_result autortfm_transact_then_open(void (*work)(void* arg), void* arg);
 #else
 UE_AUTORTFM_FORCEINLINE autortfm_result autortfm_transact(void (*work)(void* arg), void* arg)
 {
@@ -130,7 +137,7 @@ UE_AUTORTFM_FORCEINLINE autortfm_result autortfm_transact_then_open(void (*work)
 // execution if the result is anything other than autortfm_committed. Useful for
 // testing.
 #if UE_AUTORTFM
-void autortfm_commit(void (*work)(void* arg), void* arg);
+UE_AUTORTFM_API void autortfm_commit(void (*work)(void* arg), void* arg);
 #else
 UE_AUTORTFM_FORCEINLINE void autortfm_commit(void (*work)(void* arg), void* arg)
 {
@@ -142,7 +149,7 @@ UE_AUTORTFM_FORCEINLINE void autortfm_commit(void (*work)(void* arg), void* arg)
 
 // Create a new transaction in the open.
 #if UE_AUTORTFM
-bool autortfm_start_transaction();
+UE_AUTORTFM_API bool autortfm_start_transaction();
 #else
 UE_AUTORTFM_FORCEINLINE bool autortfm_start_transaction()
 {
@@ -152,7 +159,7 @@ UE_AUTORTFM_FORCEINLINE bool autortfm_start_transaction()
 
 // End a transaction and commit the changes to be visible to all
 #if UE_AUTORTFM
-autortfm_result autortfm_commit_transaction();
+UE_AUTORTFM_API autortfm_result autortfm_commit_transaction();
 #else
 UE_AUTORTFM_FORCEINLINE autortfm_result autortfm_commit_transaction()
 {
@@ -162,7 +169,7 @@ UE_AUTORTFM_FORCEINLINE autortfm_result autortfm_commit_transaction()
 
 // End a transaction and discard all changes
 #if UE_AUTORTFM
-autortfm_result autortfm_abort_transaction();
+UE_AUTORTFM_API autortfm_result autortfm_abort_transaction();
 #else
 UE_AUTORTFM_FORCEINLINE autortfm_result autortfm_abort_transaction()
 {
@@ -172,7 +179,7 @@ UE_AUTORTFM_FORCEINLINE autortfm_result autortfm_abort_transaction()
 
 // Clear the status of a transaction that was aborted in the open
 #if UE_AUTORTFM
-void autortfm_clear_transaction_status();
+UE_AUTORTFM_API void autortfm_clear_transaction_status();
 #else
 UE_AUTORTFM_FORCEINLINE void autortfm_clear_transaction_status()
 {
@@ -184,7 +191,7 @@ UE_AUTORTFM_FORCEINLINE void autortfm_clear_transaction_status()
 // depending on configuration. If called from open code that is inside a
 // transaction, this aborts the program.
 #if UE_AUTORTFM
-void autortfm_abort_if_transactional(void);
+UE_AUTORTFM_API void autortfm_abort_if_transactional(void);
 #else
 UE_AUTORTFM_FORCEINLINE void autortfm_abort_if_transactional(void) { }
 #endif
@@ -200,7 +207,7 @@ UE_AUTORTFM_FORCEINLINE void autortfm_abort_if_closed(void) { }
 // Executes the given code non-transactionally regardless of whether we are in
 // a transaction or not.
 #if UE_AUTORTFM
-void autortfm_open(void (*work)(void* arg), void* arg);
+UE_AUTORTFM_API void autortfm_open(void (*work)(void* arg), void* arg);
 #else
 UE_AUTORTFM_FORCEINLINE void autortfm_open(void (*work)(void* arg), void* arg)
 {
@@ -213,7 +220,7 @@ UE_AUTORTFM_FORCEINLINE void autortfm_open(void (*work)(void* arg), void* arg)
 //
 // Guaranteed to crash if called outside a transaction.
 #if UE_AUTORTFM
-[[nodiscard]] autortfm_status autortfm_close(void (*work)(void* arg), void* arg);
+UE_AUTORTFM_API [[nodiscard]] autortfm_status autortfm_close(void (*work)(void* arg), void* arg);
 #else
 [[nodiscard]] UE_AUTORTFM_FORCEINLINE autortfm_status autortfm_close(void (*work)(void* arg), void* arg)
 {
@@ -229,7 +236,7 @@ UE_AUTORTFM_FORCEINLINE void autortfm_open(void (*work)(void* arg), void* arg)
 // Guaranteed to assert if called outside a transaction.
 // Guaranteed to assert if called from closed code.
 #if UE_AUTORTFM
-void autortfm_record_open_write(void* Ptr, size_t Size);
+UE_AUTORTFM_API void autortfm_record_open_write(void* Ptr, size_t Size);
 #else
 UE_AUTORTFM_FORCEINLINE void autortfm_record_open_write(void* Ptr, size_t Size)
 {
@@ -263,7 +270,7 @@ UE_AUTORTFM_FORCEINLINE void autortfm_record_open_write(void* Ptr, size_t Size)
 // Note that the string describing the function may be any non-NULL string. It
 // does not have to be unique.
 #if UE_AUTORTFM
-void autortfm_register_open_function(void* original_function, void* new_function);
+UE_AUTORTFM_API void autortfm_register_open_function(void* original_function, void* new_function);
 #else
 UE_AUTORTFM_FORCEINLINE void autortfm_register_open_function(void* original_function, void* new_function) 
 { 
@@ -277,7 +284,7 @@ UE_AUTORTFM_FORCEINLINE void autortfm_register_open_function(void* original_func
 // If this is called outside a transaction or from an open nest then the work
 // happens immediately.
 #if UE_AUTORTFM
-void autortfm_open_commit(void (*work)(void* arg), void* arg);
+UE_AUTORTFM_API void autortfm_open_commit(void (*work)(void* arg), void* arg);
 #else
 UE_AUTORTFM_FORCEINLINE void autortfm_open_commit(void (*work)(void* arg), void* arg)
 {
@@ -288,7 +295,7 @@ UE_AUTORTFM_FORCEINLINE void autortfm_open_commit(void (*work)(void* arg), void*
 // Have some work happen when this transaction aborts. If this is called
 // outside a transaction or from an open nest then the work is ignored.
 #if UE_AUTORTFM
-void autortfm_open_abort(void (*work)(void* arg), void* arg);
+UE_AUTORTFM_API void autortfm_open_abort(void (*work)(void* arg), void* arg);
 #else
 UE_AUTORTFM_FORCEINLINE void autortfm_open_abort(void (*work)(void* arg), void* arg)
 {
@@ -304,7 +311,7 @@ UE_AUTORTFM_FORCEINLINE void autortfm_open_abort(void (*work)(void* arg), void*
 // passed, but it's blessed specially from the compiler's perspective, leading
 // to some nice optimizations. This does nothing when called from open code.
 #if UE_AUTORTFM
-void* autortfm_did_allocate(void* ptr, size_t size);
+UE_AUTORTFM_API void* autortfm_did_allocate(void* ptr, size_t size);
 #else
 UE_AUTORTFM_FORCEINLINE void* autortfm_did_allocate(void* ptr, size_t size)
 {
@@ -319,7 +326,7 @@ UE_AUTORTFM_FORCEINLINE void* autortfm_did_allocate(void* ptr, size_t size)
 // called outside of a transaction. May do nothing if debugging features aren't
 // enabled in the autortfm runtime.
 #if UE_AUTORTFM
-void autortfm_check_consistency_assuming_no_races(void);
+UE_AUTORTFM_API void autortfm_check_consistency_assuming_no_races(void);
 #else
 UE_AUTORTFM_FORCEINLINE void autortfm_check_consistency_assuming_no_races(void) { }
 #endif
@@ -331,7 +338,7 @@ UE_AUTORTFM_FORCEINLINE void autortfm_check_consistency_assuming_no_races(void)
 // be injected by the compiler into a global constructor in the AutoRTFM compiled
 // code.
 #if UE_AUTORTFM
-void autortfm_check_abi(void* ptr, size_t size);
+UE_AUTORTFM_API void autortfm_check_abi(void* ptr, size_t size);
 #else
 UE_AUTORTFM_FORCEINLINE void autortfm_check_abi(void* ptr, size_t size)
 {
@@ -512,8 +519,8 @@ UE_AUTORTFM_FORCEINLINE void RegisterOpenFunction(void* OriginalFunction, void*
 }
 
 #if UE_AUTORTFM
-void OpenCommit(TFunction<void()>&& Work);
-void OpenAbort(TFunction<void()>&& Work);
+UE_AUTORTFM_API void OpenCommit(TFunction<void()>&& Work);
+UE_AUTORTFM_API void OpenAbort(TFunction<void()>&& Work);
 #else
 template<typename TFunctor>
 UE_AUTORTFM_FORCEINLINE void OpenCommit(const TFunctor& Work) { Work(); }
