Two Way Map: Statically allocates 2 maps.

diff --git a/Engine/Source/Runtime/Core/Public/Experimental/Containers/TwoWayMap.h b/Engine/Source/Runtime/Core/Public/Experimental/Containers/TwoWayMap.h
new file mode 100644
--- /dev/null
+++ b/Engine/Source/Runtime/Core/Public/Experimental/Containers/TwoWayMap.h
@@ -0,0 +1,44 @@
+// Copyright Epic Games, Inc. All Rights Reserved.
+
+#pragma once
+
+#include "Containers/Map.h"
+
+/** Statically allocates two maps from the same array of pairs */
+template<typename ClassT, typename FirstType, typename SecondType>
+struct TTwoWayMap
+{
+	typedef TArray<TPair<FirstType, SecondType>> InitListType;
+
+	static const TMap<FirstType, SecondType>& GetFirstToSecond()
+	{
+		InitIfNeeded();
+		return FirstToSecond;
+	}
+
+	static const TMap<SecondType, FirstType>& GetSecondToFirst()
+	{
+		InitIfNeeded();
+		return SecondToFirst;
+	}
+	static void InitIfNeeded()
+	{
+		if (FirstToSecond.Num() == 0 || SecondToFirst.Num() == 0)
+		{
+			FirstToSecond.Reserve(ClassT::PairDefinitions.Num());
+			SecondToFirst.Reserve(ClassT::PairDefinitions.Num());
+
+			for (const TPair<FirstType, SecondType>& Pair : ClassT::PairDefinitions)
+			{
+				FirstToSecond.Emplace(Pair.Key, Pair.Value);
+				SecondToFirst.Emplace(Pair.Value, Pair.Key);
+			}
+
+			ClassT::PairDefinitions.Reset();
+		}
+	}
+
+protected:
+	inline static TMap<FirstType, SecondType> FirstToSecond = TMap<FirstType, SecondType>();
+	inline static TMap<SecondType, FirstType> SecondToFirst = TMap<SecondType, FirstType>();
+};
\ No newline at end of file
