Inject the ICF defeating ASM fun using the compiler on methods explicitly marked as [[clang::noautortfm]].

diff --git a/Engine/Source/Runtime/Core/Private/AutoRTFM/API.cpp b/Engine/Source/Runtime/Core/Private/AutoRTFM/API.cpp
--- a/Engine/Source/Runtime/Core/Private/AutoRTFM/API.cpp
+++ b/Engine/Source/Runtime/Core/Private/AutoRTFM/API.cpp
@@ -48,8 +48,6 @@ UE_AUTORTFM_FORCEINLINE autortfm_result TransactThenOpenImpl(void (*Work)(void*
 
 extern "C" UE_AUTORTFM_NOAUTORTFM bool autortfm_is_transactional()
 {
-	asm volatile("mov %0, %%rax;"::"i"(__LINE__):"%rax");
-
 	if (GAutoRTFMRuntimeEnabled)
 	{
 		return FContext::Get()->IsTransactional();
@@ -60,16 +58,12 @@ extern "C" UE_AUTORTFM_NOAUTORTFM bool autortfm_is_transactional()
 
 extern "C" UE_AUTORTFM_NOAUTORTFM bool autortfm_is_closed()
 {
-	asm volatile("mov %0, %%rax;"::"i"(__LINE__):"%rax");
-
     return false;
 }
 
 // First Part - the API exposed outside transactions.
 extern "C" UE_AUTORTFM_NOAUTORTFM autortfm_result autortfm_transact(void (*Work)(void* Arg), void* Arg)
 {
-	asm volatile("mov %0, %%rax;"::"i"(__LINE__):"%rax");
-
 	if (GAutoRTFMRuntimeEnabled)
 	{
 	    return static_cast<autortfm_result>(FContext::Get()->Transact(Work, Arg));
@@ -81,63 +75,47 @@ extern "C" UE_AUTORTFM_NOAUTORTFM autortfm_result autortfm_transact(void (*Work)
 
 extern "C" UE_AUTORTFM_NOAUTORTFM autortfm_result autortfm_transact_then_open(void (*Work)(void* Arg), void* Arg)
 {
-	asm volatile("mov %0, %%rax;"::"i"(__LINE__):"%rax");
-
     return TransactThenOpenImpl(Work, Arg);
 }
 
 extern "C" UE_AUTORTFM_NOAUTORTFM void autortfm_commit(void (*Work)(void* Arg), void* Arg)
 {
-	asm volatile("mov %0, %%rax;"::"i"(__LINE__):"%rax");
-
     autortfm_result Result = autortfm_transact(Work, Arg);
 	UE_CLOG(Result != autortfm_committed, LogAutoRTFM, Fatal, TEXT("Unexpected transaction result: %u."), Result);
 }
 
 extern "C" UE_AUTORTFM_NOAUTORTFM void autortfm_abort()
 {
-	asm volatile("mov %0, %%rax;"::"i"(__LINE__):"%rax");
-
 	UE_CLOG(!FContext::IsTransactional(), LogAutoRTFM, Fatal, TEXT("The function `autortfm_abort` was called from outside a transaction."));
 	FContext::Get()->AbortByRequestAndThrow();
 }
 
 extern "C" UE_AUTORTFM_NOAUTORTFM bool autortfm_start_transaction()
 {
-	asm volatile("mov %0, %%rax;"::"i"(__LINE__):"%rax");
-
 	UE_CLOG(!FContext::IsTransactional(), LogAutoRTFM, Fatal, TEXT("The function `autortfm_start_transaction` was called from outside a transact."));
 	return FContext::Get()->StartTransaction();
 }
 
 extern "C" UE_AUTORTFM_NOAUTORTFM autortfm_result autortfm_commit_transaction()
 {
-	asm volatile("mov %0, %%rax;"::"i"(__LINE__):"%rax");
-
 	UE_CLOG(!FContext::IsTransactional(), LogAutoRTFM, Fatal, TEXT("The function `autortfm_commit_transaction` was called from outside a transact."));
 	return static_cast<autortfm_result>(FContext::Get()->CommitTransaction());
 }
 
 extern "C" UE_AUTORTFM_NOAUTORTFM autortfm_result autortfm_abort_transaction()
 {
-	asm volatile("mov %0, %%rax;"::"i"(__LINE__):"%rax");
-
 	UE_CLOG(!FContext::IsTransactional(), LogAutoRTFM, Fatal, TEXT("The function `autortfm_abort_transaction` was called from outside a transact."));
 	return static_cast<autortfm_result>(FContext::Get()->AbortTransaction(false));
 }
 
 extern "C" UE_AUTORTFM_NOAUTORTFM void autortfm_clear_transaction_status()
 {
-	asm volatile("mov %0, %%rax;"::"i"(__LINE__):"%rax");
-
 	ASSERT(FContext::Get()->IsAborting());
 	FContext::Get()->ClearTransactionStatus();
 }
 
 extern "C" UE_AUTORTFM_NOAUTORTFM bool autortfm_is_aborting()
 {
-	asm volatile("mov %0, %%rax;"::"i"(__LINE__):"%rax");
-
 	if (GAutoRTFMRuntimeEnabled)
 	{
 		return FContext::Get()->IsAborting();
@@ -148,35 +126,26 @@ extern "C" UE_AUTORTFM_NOAUTORTFM bool autortfm_is_aborting()
 
 extern "C" UE_AUTORTFM_NOAUTORTFM bool autortfm_current_nest_throw()
 {
-	asm volatile("mov %0, %%rax;"::"i"(__LINE__):"%rax");
-
 	FContext::Get()->Throw();
 	return true;
 }
 
 extern "C" UE_AUTORTFM_NOAUTORTFM void autortfm_abort_if_transactional()
 {
-	asm volatile("mov %0, %%rax;"::"i"(__LINE__):"%rax");
-
 	UE_CLOG(FContext::IsTransactional(), LogAutoRTFM, Fatal, TEXT("The function `autortfm_abort_if_transactional` was called from an open inside a transaction."));
 }
 
 extern "C" UE_AUTORTFM_NOAUTORTFM void autortfm_abort_if_closed()
 {
-	asm volatile("mov %0, %%rax;"::"i"(__LINE__):"%rax");
 }
 
 extern "C" UE_AUTORTFM_NOAUTORTFM void autortfm_open(void (*Work)(void* Arg), void* Arg)
 {
-	asm volatile("mov %0, %%rax;"::"i"(__LINE__):"%rax");
-
 	Work(Arg);
 }
 
 extern "C" UE_AUTORTFM_NOAUTORTFM autortfm_status autortfm_close(void (*Work)(void* Arg), void* Arg)
 {
-	asm volatile("mov %0, %%rax;"::"i"(__LINE__):"%rax");
-
 	autortfm_status Result = autortfm_status_ontrack;
 
 	if (GAutoRTFMRuntimeEnabled)
@@ -200,55 +169,41 @@ extern "C" UE_AUTORTFM_NOAUTORTFM autortfm_status autortfm_close(void (*Work)(vo
 
 extern "C" UE_AUTORTFM_NOAUTORTFM void autortfm_record_open_write(void* Ptr, size_t Size)
 {
-	asm volatile("mov %0, %%rax;"::"i"(__LINE__):"%rax");
-
     FContext::Get()->CheckOpenRecordWrite(Ptr);
 	FContext::Get()->RecordWrite(Ptr, Size);
 }
 
 extern "C" UE_AUTORTFM_NOAUTORTFM void autortfm_register_open_function(void* OriginalFunction, void* NewFunction)
 {
-	asm volatile("mov %0, %%rax;"::"i"(__LINE__):"%rax");
-
 	UE_LOG(LogAutoRTFM, Verbose, TEXT("Registering open %p->%p"), OriginalFunction, NewFunction);
     FunctionMapAdd(OriginalFunction, NewFunction);
 }
 
 UE_AUTORTFM_NOAUTORTFM void OpenCommit(TFunction<void()>&& Work)
 {
-	asm volatile("mov %0, %%rax;"::"i"(__LINE__):"%rax");
-
     Work();
 }
 
 UE_AUTORTFM_NOAUTORTFM void OpenAbort(TFunction<void()>&& Work)
 {
-	asm volatile("mov %0, %%rax;"::"i"(__LINE__):"%rax");
 }
 
 extern "C" UE_AUTORTFM_NOAUTORTFM void autortfm_open_commit(void (*Work)(void* Arg), void* Arg)
 {
-	asm volatile("mov %0, %%rax;"::"i"(__LINE__):"%rax");
-
     Work(Arg);
 }
 
 extern "C" UE_AUTORTFM_NOAUTORTFM void autortfm_open_abort(void (*Work)(void* arg), void* Arg)
 {
-	asm volatile("mov %0, %%rax;"::"i"(__LINE__):"%rax");
 }
 
 extern "C" UE_AUTORTFM_NOAUTORTFM void* autortfm_did_allocate(void* Ptr, size_t Size)
 {
-	asm volatile("mov %0, %%rax;"::"i"(__LINE__):"%rax");
-
     return Ptr;
 }
 
 extern "C" UE_AUTORTFM_NOAUTORTFM void autortfm_check_consistency_assuming_no_races()
 {
-	asm volatile("mov %0, %%rax;"::"i"(__LINE__):"%rax");
-
     if (FContext::IsTransactional())
     {
         AutoRTFM::Unreachable();
@@ -304,58 +259,58 @@ extern "C" UE_AUTORTFM_NOAUTORTFM void autortfm_check_abi(void* const Ptr, const
 // Second Part - the same API exposed inside transactions. Note that we don't expose all of the API
 // to transactions! That's intentional. However, things like autortfm_defer_until_commit can be called
 // from an open nest in a transaction.
-bool RTFM_autortfm_is_transactional()
+UE_AUTORTFM_NOAUTORTFM bool RTFM_autortfm_is_transactional()
 {
     return true;
 }
 UE_AUTORTFM_REGISTER_OPEN_FUNCTION(autortfm_is_transactional);
 
-bool RTFM_autortfm_is_closed()
+UE_AUTORTFM_NOAUTORTFM bool RTFM_autortfm_is_closed()
 {
     return true;
 }
 UE_AUTORTFM_REGISTER_OPEN_FUNCTION(autortfm_is_closed);
 
-autortfm_result RTFM_autortfm_transact(void (*Work)(void* Arg), void* Arg)
+UE_AUTORTFM_NOAUTORTFM autortfm_result RTFM_autortfm_transact(void (*Work)(void* Arg), void* Arg)
 {
 	FContext* Context = FContext::Get();
     return static_cast<autortfm_result>(Context->Transact(Work, Arg));
 }
 UE_AUTORTFM_REGISTER_OPEN_FUNCTION(autortfm_transact);
 
-autortfm_result RTFM_autortfm_transact_then_open(void (*Work)(void* Arg), void* Arg)
+UE_AUTORTFM_NOAUTORTFM autortfm_result RTFM_autortfm_transact_then_open(void (*Work)(void* Arg), void* Arg)
 {
     return TransactThenOpenImpl(Work, Arg);
 }
 UE_AUTORTFM_REGISTER_OPEN_FUNCTION(autortfm_transact_then_open);
 
-void RTFM_autortfm_commit(void (*Work)(void* Arg), void* Arg)
+UE_AUTORTFM_NOAUTORTFM void RTFM_autortfm_commit(void (*Work)(void* Arg), void* Arg)
 {
     autortfm_result Result = autortfm_transact(Work, Arg);
 	UE_CLOG(Result != autortfm_committed, LogAutoRTFM, Fatal, TEXT("Unexpected transaction result: %u."), Result);
 }
 UE_AUTORTFM_REGISTER_OPEN_FUNCTION(autortfm_commit);
 
-void RTFM_autortfm_abort()
+UE_AUTORTFM_NOAUTORTFM void RTFM_autortfm_abort()
 {
 	FContext* Context = FContext::Get();
     Context->AbortByRequestAndThrow();
 }
 UE_AUTORTFM_REGISTER_OPEN_FUNCTION(autortfm_abort);
 
-void RTFM_autortfm_start_transaction()
+UE_AUTORTFM_NOAUTORTFM void RTFM_autortfm_start_transaction()
 {
 	UE_LOG(LogAutoRTFM, Fatal, TEXT("The function `autortfm_start_transaction` was called from closed code."));
 }
 UE_AUTORTFM_REGISTER_OPEN_FUNCTION(autortfm_start_transaction);
 
-void RTFM_autortfm_commit_transaction()
+UE_AUTORTFM_NOAUTORTFM void RTFM_autortfm_commit_transaction()
 {
 	UE_LOG(LogAutoRTFM, Fatal, TEXT("The function `RTFM_autortfm_commit_transaction` was called from closed code."));
 }
 UE_AUTORTFM_REGISTER_OPEN_FUNCTION(autortfm_commit_transaction);
 
-autortfm_result RTFM_autortfm_abort_transaction()
+UE_AUTORTFM_NOAUTORTFM autortfm_result RTFM_autortfm_abort_transaction()
 {
 	UE_CLOG(!FContext::IsTransactional(), LogAutoRTFM, Fatal, TEXT("The function `autortfm_abort_transaction` was called from outside a transaction"));
 	FContext* Context = FContext::Get();
@@ -363,12 +318,13 @@ autortfm_result RTFM_autortfm_abort_transaction()
 }
 UE_AUTORTFM_REGISTER_OPEN_FUNCTION(autortfm_abort_transaction);
 
-void RTFM_autortfm_clear_transaction_status()
+UE_AUTORTFM_NOAUTORTFM void RTFM_autortfm_clear_transaction_status()
 {
 	UE_LOG(LogAutoRTFM, Fatal, TEXT("The function `autortfm_clear_transaction_status` was called from closed code."));
 }
+UE_AUTORTFM_REGISTER_OPEN_FUNCTION(autortfm_clear_transaction_status);
 
-void RTFM_autortfm_abort_if_transactional()
+UE_AUTORTFM_NOAUTORTFM void RTFM_autortfm_abort_if_transactional()
 {
     UE_LOG(LogAutoRTFM, Verbose, TEXT("The function `autortfm_abort_if_transactional` was called from inside a transaction."));
 	FContext* Context = FContext::Get();
@@ -376,7 +332,7 @@ void RTFM_autortfm_abort_if_transactional()
 }
 UE_AUTORTFM_REGISTER_OPEN_FUNCTION(autortfm_abort_if_transactional);
 
-void RTFM_autortfm_abort_if_closed()
+UE_AUTORTFM_NOAUTORTFM void RTFM_autortfm_abort_if_closed()
 {
 	UE_LOG(LogAutoRTFM, Verbose, TEXT("The function `autortfm_abort_if_closed` was called from closed inside a transaction."));
 	FContext* Context = FContext::Get();
@@ -384,7 +340,7 @@ void RTFM_autortfm_abort_if_closed()
 }
 UE_AUTORTFM_REGISTER_OPEN_FUNCTION(autortfm_abort_if_closed);
 
-void RTFM_autortfm_open(void (*Work)(void* Arg), void* Arg)
+UE_AUTORTFM_NOAUTORTFM void RTFM_autortfm_open(void (*Work)(void* Arg), void* Arg)
 {
 	Work(Arg);
 
@@ -396,7 +352,7 @@ void RTFM_autortfm_open(void (*Work)(void* Arg), void* Arg)
 }
 UE_AUTORTFM_REGISTER_OPEN_FUNCTION(autortfm_open);
 
-autortfm_status RTFM_autortfm_close(void (*Work)(void* Arg), void* Arg)
+UE_AUTORTFM_NOAUTORTFM autortfm_status RTFM_autortfm_close(void (*Work)(void* Arg), void* Arg)
 {
     void (*WorkClone)(void* Arg) = FunctionMapLookup(Work, "RTFM_autortfm_close");
     if (WorkClone)
@@ -408,13 +364,13 @@ autortfm_status RTFM_autortfm_close(void (*Work)(void* Arg), void* Arg)
 }
 UE_AUTORTFM_REGISTER_OPEN_FUNCTION(autortfm_close);
 
-extern "C" void RTFM_autortfm_record_open_write(void*, size_t)
+extern "C" UE_AUTORTFM_NOAUTORTFM void RTFM_autortfm_record_open_write(void*, size_t)
 {
 	UE_LOG(LogAutoRTFM, Fatal, TEXT("The function `autortfm_record_open_write` was called from closed code."));
 }
 UE_AUTORTFM_REGISTER_OPEN_FUNCTION(autortfm_record_open_write);
 
-void RTFM_OpenCommit(TFunction<void()>&& Work)
+UE_AUTORTFM_NOAUTORTFM void RTFM_OpenCommit(TFunction<void()>&& Work)
 {
 	FContext* Context = FContext::Get();
     ASSERT(Context->GetStatus() == EContextStatus::OnTrack);
@@ -422,7 +378,7 @@ void RTFM_OpenCommit(TFunction<void()>&& Work)
 }
 UE_AUTORTFM_REGISTER_OPEN_FUNCTION(OpenCommit);
 
-void RTFM_OpenAbort(TFunction<void()>&& Work)
+UE_AUTORTFM_NOAUTORTFM void RTFM_OpenAbort(TFunction<void()>&& Work)
 {
 	FContext* Context = FContext::Get();
     ASSERT(Context->GetStatus() == EContextStatus::OnTrack);
@@ -430,19 +386,19 @@ void RTFM_OpenAbort(TFunction<void()>&& Work)
 }
 UE_AUTORTFM_REGISTER_OPEN_FUNCTION(OpenAbort);
 
-extern "C" void RTFM_autortfm_open_commit(void (*Work)(void* Arg), void* Arg)
+extern "C" UE_AUTORTFM_NOAUTORTFM void RTFM_autortfm_open_commit(void (*Work)(void* Arg), void* Arg)
 {
     RTFM_OpenCommit([Work, Arg] { Work(Arg); });
 }
 UE_AUTORTFM_REGISTER_OPEN_FUNCTION(autortfm_open_commit);
 
-extern "C" void RTFM_autortfm_open_abort(void (*Work)(void* arg), void* Arg)
+extern "C" UE_AUTORTFM_NOAUTORTFM void RTFM_autortfm_open_abort(void (*Work)(void* arg), void* Arg)
 {
     RTFM_OpenAbort([Work, Arg] { Work(Arg); });
 }
 UE_AUTORTFM_REGISTER_OPEN_FUNCTION(autortfm_open_abort);
 
-void* RTFM_autortfm_did_allocate(void* Ptr, size_t Size)
+UE_AUTORTFM_NOAUTORTFM void* RTFM_autortfm_did_allocate(void* Ptr, size_t Size)
 {
 	FContext* Context = FContext::Get();
     Context->DidAllocate(Ptr, Size);
@@ -450,7 +406,7 @@ void* RTFM_autortfm_did_allocate(void* Ptr, size_t Size)
 }
 UE_AUTORTFM_REGISTER_OPEN_FUNCTION(autortfm_did_allocate);
 
-void RTFM_autortfm_check_consistency_assuming_no_races()
+UE_AUTORTFM_NOAUTORTFM void RTFM_autortfm_check_consistency_assuming_no_races()
 {
 }
 UE_AUTORTFM_REGISTER_OPEN_FUNCTION(autortfm_check_consistency_assuming_no_races);
