Add an RBTree based Interval Tree to track the ranges of new memory on large allocations, to avoid spewing memory on some torture cases. This adds an additional 17ms cost to a 5000ms run of the DeltaBlue benchmark, but in the test case of VerseTestScriptCmd that this caused issues with, reduced peak memory used by 2.5x.

diff --git a/Engine/Source/Runtime/Core/Private/AutoRTFM/IntervalTree.h b/Engine/Source/Runtime/Core/Private/AutoRTFM/IntervalTree.h
new file mode 100644
--- /dev/null
+++ b/Engine/Source/Runtime/Core/Private/AutoRTFM/IntervalTree.h
@@ -0,0 +1,333 @@
+// Copyright Epic Games, Inc. All Rights Reserved.
+
+#pragma once
+
+#include <cstdint>
+#include "Containers/PagedArray.h"
+
+namespace AutoRTFM
+{
+
+struct FIntervalTree final
+{
+    FIntervalTree() = default;
+
+    FIntervalTree(const FIntervalTree&) = delete;
+    FIntervalTree& operator=(const FIntervalTree&) = delete;
+
+	FORCENOINLINE bool Insert(const void* const Address, const size_t Size)
+    {
+        FRange NewRange(Address, Size);
+
+        if (UNLIKELY(IntervalTreeNodeIndexNone == Root))
+        {
+            ASSERT(0 == Nodes.Num());
+            Nodes.Add(FIntervalTreeNode(NewRange));
+            Root = 0;
+            return true;
+        }
+
+		ASSERT(Nodes[Root].Parent == IntervalTreeNodeIndexNone);
+
+        FIntervalTreeNodeIndex Current = Root;
+
+        for(;;)
+        {
+            FIntervalTreeNode& Node = Nodes[Current];
+
+            const FRange Range = Node.Range;
+
+            if (UNLIKELY((NewRange.Start < Range.End) && (Range.Start < NewRange.End)))
+            {
+                return false;
+            }
+
+            if (NewRange.Start < Range.Start)
+            {
+				if (NewRange.End == Range.Start)
+				{
+                    // We can just modify the existing node in place.
+					ASSERT(NewRange.Start < Range.Start);
+					Node.Range.Start = NewRange.Start;
+                    return true;
+				}
+				else if (IntervalTreeNodeIndexNone == Node.Left)
+                {
+                    const FIntervalTreeNodeIndex Index = Nodes.Add(FIntervalTreeNode(NewRange, Current));
+                    Node.Left = Index;
+                    Current = Index;
+                    break;
+                }
+
+                Current = Node.Left;
+            }
+            else
+            {
+				if (NewRange.Start == Range.End)
+				{
+                    // We can just modify the existing node in place.
+					ASSERT(NewRange.End > Range.End);
+					Node.Range.End = NewRange.End;
+					return true;
+				}
+				else if (IntervalTreeNodeIndexNone == Node.Right)
+                {
+                    const FIntervalTreeNodeIndex Index = Nodes.Add(FIntervalTreeNode(NewRange, Current));
+                    Node.Right = Index;
+                    Current = Index;
+                    break;
+                }
+
+                Current = Node.Right;
+            }
+
+			ASSERT(Root != Current);
+        }
+
+        auto IsBlack = [this](FIntervalTreeNodeIndex Index)
+        {
+            return IntervalTreeNodeIndexNone == Index || Nodes[Index].bIsBlack;
+        };
+
+        for(;;)
+        {
+            FIntervalTreeNodeIndex Parent = Nodes[Current].Parent;
+
+            UE_ASSUME(Current != Parent);
+
+            // The root will always have a black parent, so this check covers both.
+			if (Parent == Root)
+			{
+				Nodes[Parent].bIsBlack = true;
+				break;
+			}
+            else if (IsBlack(Parent))
+            {
+				break;
+            }
+
+            const FIntervalTreeNodeIndex GrandParent = Nodes[Parent].Parent;
+
+            UE_ASSUME((Parent != GrandParent) && (Current != GrandParent));
+
+			const bool bGrandParentIsRoot = GrandParent == Root;
+
+            const bool bParentIsLeft = Nodes[GrandParent].Left == Parent;
+
+            // The uncle is the other node of our parent.
+            const FIntervalTreeNodeIndex Uncle = bParentIsLeft ? Nodes[GrandParent].Right : Nodes[GrandParent].Left;
+
+            UE_ASSUME((GrandParent != Uncle) && (Parent != Uncle) && (Current != Uncle));
+
+            if (!IsBlack(Uncle))
+            {
+				ASSERT(IsBlack(GrandParent));
+                Nodes[Parent].bIsBlack = true;
+                Nodes[Uncle].bIsBlack = true;
+
+				if (bGrandParentIsRoot)
+				{
+					break;
+				}
+				else
+				{
+					Nodes[GrandParent].bIsBlack = false;
+				}
+
+                Current = GrandParent;
+                continue;
+            }
+
+            // Our uncle is black, so we need to swizzle around.
+            const bool CurrentIsLeft = Nodes[Parent].Left == Current;
+
+            if (bParentIsLeft)
+            {
+                if (!CurrentIsLeft)
+                {
+                    Nodes[Parent].Right = Nodes[Current].Left;
+
+					if (IntervalTreeNodeIndexNone != Nodes[Parent].Right)
+					{
+						Nodes[Nodes[Parent].Right].Parent = Parent;
+					}
+
+                    Nodes[Parent].Parent = Current;
+                    Nodes[Current].Parent = GrandParent;
+                    Nodes[Current].Left = Parent;
+					Nodes[GrandParent].Left = Current;
+                    std::swap(Parent, Current);
+                }
+
+                Nodes[Parent].Parent = Nodes[GrandParent].Parent;
+                Nodes[GrandParent].Left = Nodes[Parent].Right;
+
+				if (IntervalTreeNodeIndexNone != Nodes[GrandParent].Left)
+				{
+					Nodes[Nodes[GrandParent].Left].Parent = GrandParent;
+				}
+
+                Nodes[Parent].Right = GrandParent;
+                Nodes[GrandParent].Parent = Parent;
+
+                std::swap(Nodes[GrandParent].bIsBlack, Nodes[Parent].bIsBlack);
+            }
+            else
+            {
+                if (CurrentIsLeft)
+                {
+                    Nodes[Parent].Left = Nodes[Current].Right;
+					
+					if (IntervalTreeNodeIndexNone != Nodes[Parent].Left)
+					{
+						Nodes[Nodes[Parent].Left].Parent = Parent;
+					}
+
+                    Nodes[Parent].Parent = Current;
+                    Nodes[Current].Parent = GrandParent;
+                    Nodes[Current].Right = Parent;
+					Nodes[GrandParent].Right = Current;
+                    std::swap(Parent, Current);
+                }
+
+                Nodes[Parent].Parent = Nodes[GrandParent].Parent;
+                Nodes[GrandParent].Right = Nodes[Parent].Left;
+
+				if (IntervalTreeNodeIndexNone != Nodes[GrandParent].Right)
+				{
+					Nodes[Nodes[GrandParent].Right].Parent = GrandParent;
+				}
+
+                Nodes[Parent].Left = GrandParent;
+				Nodes[GrandParent].Parent = Parent;
+
+                std::swap(Nodes[GrandParent].bIsBlack, Nodes[Parent].bIsBlack);
+            }
+
+			if (bGrandParentIsRoot)
+			{
+				Root = Parent;
+			}
+			else
+			{
+				FIntervalTreeNode& GreatGrandParent = Nodes[Nodes[Parent].Parent];
+				if (GrandParent == GreatGrandParent.Left)
+				{
+					GreatGrandParent.Left = Parent;
+				}
+				else
+				{
+					GreatGrandParent.Right = Parent;
+				}
+			}
+
+            break;
+        }
+
+		if constexpr (bExtraDebugging)
+		{
+			for (FIntervalTreeNodeIndex Index = 0; Index < Nodes.Num(); Index++)
+			{
+				const FIntervalTreeNodeIndex Parent = Nodes[Index].Parent;
+				if (IntervalTreeNodeIndexNone == Parent)
+				{
+					ASSERT(Root == Index);
+				}
+				else
+				{
+					ASSERT(Nodes[Parent].bIsBlack || Nodes[Index].bIsBlack);
+					ASSERT((Nodes[Parent].Left == Index) ^ (Nodes[Parent].Right == Index));
+				}
+			}
+		}
+
+        return true;
+    }
+
+    FORCENOINLINE bool Contains(const void* const Address, const size_t Size) const
+    {
+        FRange NewRange(Address, Size);
+
+        if (UNLIKELY(IntervalTreeNodeIndexNone == Root))
+        {
+            return false;
+        }
+
+        FIntervalTreeNodeIndex Current = Root;
+
+        do
+        {
+            const FIntervalTreeNode& Node = Nodes[Current];
+
+            const FRange Range = Node.Range;
+
+            // This check does not need to prove that NewRange is entirely
+            // enclosed within Range, because if any byte of NewRange was in the
+            // original Range then it **must** already have been new memory.
+            if ((NewRange.Start < Range.End) && (Range.Start < NewRange.End))
+            {
+                return true;
+            }
+            else if (NewRange.Start < Range.Start)
+            {
+                Current = Node.Left;
+            }
+            else
+            {
+                Current = Node.Right;
+            }
+        } while (IntervalTreeNodeIndexNone != Current);
+
+        return false;
+    }
+
+    bool IsEmpty() const
+    {
+        return IntervalTreeNodeIndexNone == Root;
+    }
+
+    void Reset()
+    {
+        Root = IntervalTreeNodeIndexNone;
+        Nodes.Empty();
+    }
+
+private:
+    struct FRange final
+    {
+        FRange(const void* const Address, const size_t Size) :
+            Start(reinterpret_cast<uintptr_t>(Address)), End(reinterpret_cast<uintptr_t>(Address) + Size) {}
+
+        uintptr_t Start;
+        uintptr_t End;
+    };
+
+private:
+    using FIntervalTreeNodeIndex = uint32_t;
+
+    static constexpr FIntervalTreeNodeIndex IntervalTreeNodeIndexNone = UINT32_MAX;
+
+    struct FIntervalTreeNode final
+    {
+        explicit FIntervalTreeNode(const FRange Range) :
+            Parent(IntervalTreeNodeIndexNone), bIsBlack(true), Range(Range) {}
+
+		FIntervalTreeNode(const FRange Range, FIntervalTreeNodeIndex Parent) :
+			Parent(Parent), bIsBlack(false), Range(Range) {}
+
+        FIntervalTreeNodeIndex Left = IntervalTreeNodeIndexNone;
+        FIntervalTreeNodeIndex Right = IntervalTreeNodeIndexNone;
+        FIntervalTreeNodeIndex Parent;
+
+        // TODO: optimize this waste of memory.
+        bool bIsBlack;
+        FRange Range;
+    };
+
+    FIntervalTreeNodeIndex Root = IntervalTreeNodeIndexNone;
+    TPagedArray<FIntervalTreeNode> Nodes;
+
+	static constexpr bool bExtraDebugging = false;
+};
+
+} // namespace AutoRTFM
diff --git a/Engine/Source/Runtime/Core/Private/AutoRTFM/Stats.cpp b/Engine/Source/Runtime/Core/Private/AutoRTFM/Stats.cpp
--- a/Engine/Source/Runtime/Core/Private/AutoRTFM/Stats.cpp
+++ b/Engine/Source/Runtime/Core/Private/AutoRTFM/Stats.cpp
@@ -88,6 +88,17 @@ template<> void AutoRTFM::FStats::Report<AutoRTFM::EStatsKind::MaximumAbortTasks
     UE_LOG(LogAutoRTFM, Display, TEXT("  Maximum abort tasks:       %11u"), Data);
 }
 
+template<> void AutoRTFM::FStats::Report<AutoRTFM::EStatsKind::NewMemoryTrackerHit>(const uint64_t Data) const
+{
+    UE_LOG(LogAutoRTFM, Display, TEXT("  New memory hits:           %11u"), Data);
+}
+
+template<> void AutoRTFM::FStats::Report<AutoRTFM::EStatsKind::NewMemoryTrackerMiss>(const uint64_t Data) const
+{
+    UE_LOG(LogAutoRTFM, Display, TEXT("  New memory misses:         %11u"), Data);
+}
+
+
 void AutoRTFM::FStats::Report() const
 {
 	if constexpr (bCollectStats)
@@ -116,6 +127,8 @@ void AutoRTFM::FStats::Report() const
             REPORT_CASE(AutoRTFM::EStatsKind::MaximumCommitTasks);
             REPORT_CASE(AutoRTFM::EStatsKind::AverageAbortTasks);
             REPORT_CASE(AutoRTFM::EStatsKind::MaximumAbortTasks);
+            REPORT_CASE(AutoRTFM::EStatsKind::NewMemoryTrackerHit);
+            REPORT_CASE(AutoRTFM::EStatsKind::NewMemoryTrackerMiss);
 #undef REPORT_CASE
 			}
 		}
diff --git a/Engine/Source/Runtime/Core/Private/AutoRTFM/Stats.h b/Engine/Source/Runtime/Core/Private/AutoRTFM/Stats.h
--- a/Engine/Source/Runtime/Core/Private/AutoRTFM/Stats.h
+++ b/Engine/Source/Runtime/Core/Private/AutoRTFM/Stats.h
@@ -27,6 +27,8 @@ enum class EStatsKind : uint8_t
 	MaximumCommitTasks,
 	AverageAbortTasks,
 	MaximumAbortTasks,
+	NewMemoryTrackerHit,
+	NewMemoryTrackerMiss,
 	Total
 };
 
@@ -42,6 +44,8 @@ template<EStatsKind Kind> constexpr bool IsStatsKindNoArgs()
 	case EStatsKind::HitSetHit:
 	case EStatsKind::HitSetMiss:
 	case EStatsKind::HitSetSkippedBecauseOfStackLocalMemory:
+	case EStatsKind::NewMemoryTrackerHit:
+	case EStatsKind::NewMemoryTrackerMiss:
 		return true;
 	}
 }
diff --git a/Engine/Source/Runtime/Core/Private/AutoRTFM/Transaction.cpp b/Engine/Source/Runtime/Core/Private/AutoRTFM/Transaction.cpp
--- a/Engine/Source/Runtime/Core/Private/AutoRTFM/Transaction.cpp
+++ b/Engine/Source/Runtime/Core/Private/AutoRTFM/Transaction.cpp
@@ -17,6 +17,7 @@ FTransaction::FTransaction(FContext* Context)
 bool FTransaction::IsFresh() const
 {
     return HitSet.IsEmpty()
+        && NewMemoryTracker.IsEmpty()
         && WriteLog.IsEmpty()
         && CommitTasks.IsEmpty()
         && AbortTasks.IsEmpty()
@@ -205,6 +206,7 @@ void FTransaction::Reset()
     CommitTasks.Reset();
     AbortTasks.Reset();
 	HitSet.Reset();
+    NewMemoryTracker.Reset();
 	WriteLog.Reset();
 	WriteLogBumpAllocator.Reset();
 }
diff --git a/Engine/Source/Runtime/Core/Private/AutoRTFM/Transaction.h b/Engine/Source/Runtime/Core/Private/AutoRTFM/Transaction.h
--- a/Engine/Source/Runtime/Core/Private/AutoRTFM/Transaction.h
+++ b/Engine/Source/Runtime/Core/Private/AutoRTFM/Transaction.h
@@ -3,6 +3,7 @@
 #pragma once
 
 #include "HitSet.h"
+#include "IntervalTree.h"
 #include "LongJump.h"
 #include "Stats.h"
 #include "TaggedPtr.h"
@@ -99,6 +100,7 @@ private:
 	bool bIsStackScoped{false};
 
     FHitSet HitSet;
+    FIntervalTree NewMemoryTracker;
     FWriteLog WriteLog;
     FWriteLogBumpAllocator WriteLogBumpAllocator;
     TStatStorage<uint64_t> StatDepth = 1;
diff --git a/Engine/Source/Runtime/Core/Private/AutoRTFM/TransactionInlines.h b/Engine/Source/Runtime/Core/Private/AutoRTFM/TransactionInlines.h
--- a/Engine/Source/Runtime/Core/Private/AutoRTFM/TransactionInlines.h
+++ b/Engine/Source/Runtime/Core/Private/AutoRTFM/TransactionInlines.h
@@ -38,9 +38,20 @@ UE_AUTORTFM_FORCEINLINE void FTransaction::RecordWrite(void* LogicalAddress, siz
             Stats.Collect<EStatsKind::HitSetHit>();
             return;
         }
+    
+        Stats.Collect<EStatsKind::HitSetMiss>();
+    }
+    else
+    {
+        if (NewMemoryTracker.Contains(LogicalAddress, Size))
+        {
+            Stats.Collect<EStatsKind::NewMemoryTrackerHit>();
+            return;
+        }
+
+        Stats.Collect<EStatsKind::NewMemoryTrackerMiss>();
     }
 
-    Stats.Collect<EStatsKind::HitSetMiss>();
 
     uint8_t* const Address = reinterpret_cast<uint8_t*>(LogicalAddress);
 
@@ -83,9 +94,13 @@ template<unsigned SIZE> UE_AUTORTFM_FORCEINLINE void FTransaction::RecordWrite(v
     RecordWriteMaxPageSized(LogicalAddress, SIZE);
 }
 
-UE_AUTORTFM_FORCEINLINE void FTransaction::DidAllocate(void* LogicalAddress, size_t Size)
+UE_AUTORTFM_FORCEINLINE void FTransaction::DidAllocate(void* LogicalAddress, const size_t Size)
 {
-    if ((0 < Size) && (Size <= FWriteLogBumpAllocator::MaxSize))
+	if (0 == Size)
+	{
+		return;
+	}
+	else if (Size <= FWriteLogBumpAllocator::MaxSize)
     {
         FMemoryLocation Key(LogicalAddress);
         Key.SetTopTag(static_cast<uint16_t>(Size));
@@ -93,6 +108,11 @@ UE_AUTORTFM_FORCEINLINE void FTransaction::DidAllocate(void* LogicalAddress, siz
         const bool DidInsert = HitSet.Insert(Key);
         ASSERT(DidInsert); 
     }
+    else
+    {
+        const bool DidInsert = NewMemoryTracker.Insert(LogicalAddress, Size);
+        ASSERT(DidInsert);
+    }
 }
 
 UE_AUTORTFM_FORCEINLINE void FTransaction::DeferUntilCommit(TFunction<void()>&& Callback)
