GenericWidePlatformString GetVarArgs LowLevelTests Bugfix: Loosen the format requirements on the results for %p, it can vary from platform to platform.
#rnx

diff --git a/Engine/Source/Runtime/Core/Private/GenericPlatform/GenericWidePlatformString.cpp b/Engine/Source/Runtime/Core/Private/GenericPlatform/GenericWidePlatformString.cpp
--- a/Engine/Source/Runtime/Core/Private/GenericPlatform/GenericWidePlatformString.cpp
+++ b/Engine/Source/Runtime/Core/Private/GenericPlatform/GenericWidePlatformString.cpp
@@ -255,7 +255,14 @@ TEST_CASE("Core::PlatformString::GetVarArgs", "[Core][String][Smoke]")
 	// MSVC's standard library formats pointers differently
 	CHECK_MESSAGE(OutputString, FString(OutputString) == FString(TEXT("Test D|0000000000012345|")));
 #else
-	CHECK_MESSAGE(OutputString, FString(OutputString) == FString(TEXT("Test D|0x12345|")));
+	// Pointer format (which we get from snprintf) varies from platform to platform.
+	// Call it to make sure it puts the output in the proper place, but don't test for the exact format
+//	CHECK_MESSAGE(OutputString, FString(OutputString) == FString(TEXT("Test D|0x12345|")));
+	{
+		FString OutputStringStr(OutputString);
+		CHECK_MESSAGE(OutputString, OutputStringStr.StartsWith(TEXT("Test D|"))
+			&& OutputStringStr.EndsWith(TEXT("|")) && OutputStringStr.Len() < TEXTVIEW("Test D||").Len() + 20);
+	}
 #endif
 
 	TestGetVarArgs(OutputString, TEXT("Test E|%" INT64_FMT "|"), int64(12345678912345LL));
