Fix for compilation error on older ms compilers.

diff --git a/Engine/Source/Runtime/Core/Public/MSVC/MSVCPlatform.h b/Engine/Source/Runtime/Core/Public/MSVC/MSVCPlatform.h
--- a/Engine/Source/Runtime/Core/Public/MSVC/MSVCPlatform.h
+++ b/Engine/Source/Runtime/Core/Public/MSVC/MSVCPlatform.h
@@ -22,4 +22,6 @@
 #define PLATFORM_RETURN_ADDRESS_POINTER()	_AddressOfReturnAddress()
 
 // https://devblogs.microsoft.com/cppblog/improving-the-state-of-debug-performance-in-c/
+#if __has_cpp_attribute(msvc::intrinsic)
 #define UE_INTRINSIC_CAST [[msvc::intrinsic]]
+#endif
