Fix use-without-include of ASSERT macro in IntervalTree.h to make IWYU happy.

diff --git a/Engine/Source/Runtime/Core/Private/AutoRTFM/IntervalTree.h b/Engine/Source/Runtime/Core/Private/AutoRTFM/IntervalTree.h
--- a/Engine/Source/Runtime/Core/Private/AutoRTFM/IntervalTree.h
+++ b/Engine/Source/Runtime/Core/Private/AutoRTFM/IntervalTree.h
@@ -4,6 +4,7 @@
 
 #include <cstdint>
 #include "Containers/PagedArray.h"
+#include "Utils.h"
 
 namespace AutoRTFM
 {
