[Insights] Added memory tracing for FMemory::SystemMalloc and FMemory::SystemFree.

diff --git a/Engine/Source/Runtime/Core/Public/HAL/UnrealMemory.h b/Engine/Source/Runtime/Core/Public/HAL/UnrealMemory.h
--- a/Engine/Source/Runtime/Core/Public/HAL/UnrealMemory.h
+++ b/Engine/Source/Runtime/Core/Public/HAL/UnrealMemory.h
@@ -6,6 +6,7 @@
 #include "GenericPlatform/GenericPlatformMemory.h"
 #include "HAL/MemoryBase.h"
 #include "HAL/PlatformMemory.h"
+#include "ProfilingDebugging/MemoryTrace.h"
 #include "Templates/IsPointer.h"
 
 #ifndef UE_USE_VERYLARGEPAGEALLOCATOR
@@ -193,13 +194,14 @@ struct FMemory
 	//
 	UE_ALLOCATION_FUNCTION(1) static FORCEINLINE void* SystemMalloc(SIZE_T Size)
 	{
-		/* TODO: Trace! */
-		return ::malloc(Size);
+		void* Ptr = ::malloc(Size);
+		MemoryTrace_Alloc(uint64(Ptr), Size, 0, EMemoryTraceRootHeap::SystemMemory);
+		return Ptr;
 	}
 
 	static FORCEINLINE void SystemFree(void* Ptr)
 	{
-		/* TODO: Trace! */
+		MemoryTrace_Free(uint64(Ptr), EMemoryTraceRootHeap::SystemMemory);
 		::free(Ptr);
 	}
 
