Fixed TSparseArray::operator() indexing usage from UE3.

diff --git a/Engine/Source/Runtime/Core/Public/Containers/SparseArray.h b/Engine/Source/Runtime/Core/Public/Containers/SparseArray.h
--- a/Engine/Source/Runtime/Core/Public/Containers/SparseArray.h
+++ b/Engine/Source/Runtime/Core/Public/Containers/SparseArray.h
@@ -1060,8 +1060,8 @@ public:
 			return !(bool)*this;
 		}
 
-		FORCEINLINE const ElementType& operator*() const { return Array(GetIndex()); }
-		FORCEINLINE const ElementType* operator->() const { return &Array(GetIndex()); }
+		FORCEINLINE const ElementType& operator*() const { return Array[GetIndex()]; }
+		FORCEINLINE const ElementType* operator->() const { return &Array[GetIndex()]; }
 		FORCEINLINE const FRelativeBitReference& GetRelativeBitReference() const { return BitArrayIt; }
 	private:
 		const TSparseArray& Array;
