Fixed Decommit error when freeing FMallocBinned3::FPerThreadFreeBlockLists.

  - The problem was that a FPlatformVirtualMemoryBlock was being saved along with the actual FPerThreadFreeBlockLists.
    - This was done so that later the FPerThreadFreeBlockLists could be deallocated using the original Block.
    - The problem was that the Block gets decommitted while the Block is freeing the virtual memory, causing an access violation.
    - At this point accessing any of the members of FPlatformVirtualMemoryBlock is an access violation.
    - This doesn't currently crash because USE_PAGE_CACHE 1 is the default, and the page cache just marks for decommit.
  - The fix was to add FMallocBinned3::FreeMetaDataMemory that can free memory by pointer.

#tests Built Win64 Editor and all Switch clients. Did a local ReplayRun with LLM with USE_PAGE_CACHE 0 and 1.

diff --git a/Engine/Source/Runtime/Core/Private/HAL/MallocBinned3.cpp b/Engine/Source/Runtime/Core/Private/HAL/MallocBinned3.cpp
--- a/Engine/Source/Runtime/Core/Private/HAL/MallocBinned3.cpp
+++ b/Engine/Source/Runtime/Core/Private/HAL/MallocBinned3.cpp
@@ -929,11 +929,6 @@ void FMallocBinned3::Decommit(uint32 InPoolIndex, void *Ptr, SIZE_T Size)
 }
 
 void* FMallocBinned3::AllocateMetaDataMemory(SIZE_T Size)
-{
-	return AllocateMemoryBlock(Size).GetVirtualPointer();
-}
-
-FPlatformMemory::FPlatformVirtualMemoryBlock FMallocBinned3::AllocateMemoryBlock(SIZE_T Size)
 {
 	LLM_PLATFORM_SCOPE(ELLMTag::FMalloc);
 	size_t VirtualAlignedSize = Align(Size, FPlatformMemory::FPlatformVirtualMemoryBlock::GetVirtualSizeAlignment());
@@ -941,13 +936,19 @@ FPlatformMemory::FPlatformVirtualMemoryBlock FMallocBinned3::AllocateMemoryBlock
 	size_t CommitAlignedSize = Align(Size, FPlatformMemory::FPlatformVirtualMemoryBlock::GetCommitAlignment());
 	Block.Commit(0, CommitAlignedSize);
 	LLM(FLowLevelMemTracker::Get().OnLowLevelAlloc(ELLMTracker::Platform, Block.GetVirtualPointer(), CommitAlignedSize));
-	return Block;
+	return Block.GetVirtualPointer();
 }
 
-void FMallocBinned3::DeallocateMemoryBlock(FPlatformMemory::FPlatformVirtualMemoryBlock& Block)
+void FMallocBinned3::FreeMetaDataMemory(void *Ptr, SIZE_T InSize)
 {
-	LLM(FLowLevelMemTracker::Get().OnLowLevelFree(ELLMTracker::Platform, Block.GetVirtualPointer()));
-	Block.FreeVirtual();
+	if (Ptr)
+	{
+		LLM(FLowLevelMemTracker::Get().OnLowLevelFree(ELLMTracker::Platform, Ptr));
+
+		size_t VirtualAlignedSize = Align(InSize, FPlatformMemory::FPlatformVirtualMemoryBlock::GetVirtualSizeAlignment());
+		FPlatformMemory::FPlatformVirtualMemoryBlock Block(Ptr, VirtualAlignedSize / FPlatformMemory::FPlatformVirtualMemoryBlock::GetVirtualSizeAlignment());
+		Block.FreeVirtual();
+	}
 }
 
 bool FMallocBinned3::IsInternallyThreadSafe() const
@@ -1490,13 +1491,8 @@ void FMallocBinned3::FPerThreadFreeBlockLists::SetTLS()
 	FPerThreadFreeBlockLists* ThreadSingleton = (FPerThreadFreeBlockLists*)FPlatformTLS::GetTlsValue(FMallocBinned3::Binned3TlsSlot);
 	if (!ThreadSingleton)
 	{
-		const int VMBlockSize = Align(sizeof(FPlatformMemory::FPlatformVirtualMemoryBlock), alignof(FPerThreadFreeBlockLists));
-		const int64 TLSSize = VMBlockSize + sizeof(FPerThreadFreeBlockLists);
-		FPlatformMemory::FPlatformVirtualMemoryBlock Block = FMallocBinned3::AllocateMemoryBlock(TLSSize);
-		char* Ptr = (char*)Block.GetVirtualPointer();
-		FMemory::Memcpy(Ptr, &Block, sizeof(Block));
-		Ptr += VMBlockSize;
-		ThreadSingleton = new (Ptr) FPerThreadFreeBlockLists();
+		const int64 TLSSize = Align(sizeof(FPerThreadFreeBlockLists), FPlatformMemory::FPlatformVirtualMemoryBlock::GetVirtualSizeAlignment());
+		ThreadSingleton = new (FMallocBinned3::AllocateMetaDataMemory(TLSSize)) FPerThreadFreeBlockLists();
 #if BINNED3_ALLOCATOR_STATS
 		Binned3TLSMemory += TLSSize;
 #endif
@@ -1512,11 +1508,15 @@ void FMallocBinned3::FPerThreadFreeBlockLists::ClearTLS()
 	FPerThreadFreeBlockLists* ThreadSingleton = (FPerThreadFreeBlockLists*)FPlatformTLS::GetTlsValue(FMallocBinned3::Binned3TlsSlot);
 	if (ThreadSingleton)
 	{
+		const int64 TLSSize = Align(sizeof(FPerThreadFreeBlockLists), FPlatformMemory::FPlatformVirtualMemoryBlock::GetVirtualSizeAlignment());
+#if BINNED3_ALLOCATOR_STATS
+		Binned3TLSMemory -= TLSSize;
+#endif
 		FMallocBinned3::Private::UnregisterThreadFreeBlockLists(ThreadSingleton);
-
+	
 		ThreadSingleton->~FPerThreadFreeBlockLists();
-		FPlatformMemory::FPlatformVirtualMemoryBlock* Block = (FPlatformMemory::FPlatformVirtualMemoryBlock*)AlignDown(ThreadSingleton, FPlatformMemory::FPlatformVirtualMemoryBlock::GetCommitAlignment());
-		DeallocateMemoryBlock(*Block);
+
+		FMallocBinned3::FreeMetaDataMemory(ThreadSingleton, TLSSize);
 	}
 	FPlatformTLS::SetTlsValue(FMallocBinned3::Binned3TlsSlot, nullptr);
 }
diff --git a/Engine/Source/Runtime/Core/Public/HAL/MallocBinned3.h b/Engine/Source/Runtime/Core/Public/HAL/MallocBinned3.h
--- a/Engine/Source/Runtime/Core/Public/HAL/MallocBinned3.h
+++ b/Engine/Source/Runtime/Core/Public/HAL/MallocBinned3.h
@@ -733,8 +733,7 @@ public:
 
 private:
 	static void* AllocateMetaDataMemory(SIZE_T Size);
-	static FPlatformMemory::FPlatformVirtualMemoryBlock AllocateMemoryBlock(SIZE_T Size);
-	static void DeallocateMemoryBlock(FPlatformMemory::FPlatformVirtualMemoryBlock& Block);
+	static void FreeMetaDataMemory(void *Ptr, SIZE_T Size);
 };
 
 PRAGMA_RESTORE_UNSAFE_TYPECAST_WARNINGS
