CookIndeterminism: Add function FGuid::NewDeterministicGuid for use during PostLoad or new object initialization.
#rnx

diff --git a/Engine/Source/Runtime/Core/Private/Misc/Guid.cpp b/Engine/Source/Runtime/Core/Private/Misc/Guid.cpp
--- a/Engine/Source/Runtime/Core/Private/Misc/Guid.cpp
+++ b/Engine/Source/Runtime/Core/Private/Misc/Guid.cpp
@@ -7,6 +7,7 @@
 #include "Containers/StaticArray.h"
 #include "HAL/PlatformMisc.h"
 #include "HAL/PlatformString.h"
+#include "Hash/Blake3.h"
 #include "Math/Int128.h"
 #include "Misc/Base64.h"
 #include "Misc/ByteSwap.h"
@@ -235,6 +236,36 @@ FGuid FGuid::NewGuid()
 	return Result;
 }
 
+FGuid FGuid::NewDeterministicGuid(FStringView ObjectPath, uint64 Seed)
+{
+	// Convert the objectpath to utf8 so that whether TCHAR is UTF8 or UTF16 does not alter the hash.
+	TUtf8StringBuilder<1024> Utf8ObjectPath(InPlace, ObjectPath);
+
+	FBlake3 Builder;
+
+	// Hash this as the namespace of the Version 3 UUID, to avoid collisions with any other guids created using Blake3.
+	static FGuid BaseVersion(TEXT("182b8dd3-f963-477f-a57d-70a339d922d8"));
+	Builder.Update(&BaseVersion, sizeof(FGuid));
+	Builder.Update(&Seed, sizeof(Seed));
+	Builder.Update(Utf8ObjectPath.GetData(), Utf8ObjectPath.Len() * sizeof(UTF8CHAR));
+
+	// We use the first 16 bytes of the BLAKE3 hash to create the guid, there is no specific reason why these were
+	// chosen, we could take any pattern or combination of bytes.
+	FBlake3Hash Hash = Builder.Finalize();
+	uint32* HashBytes = (uint32*)Hash.GetBytes();
+	uint32 A = uint32(HashBytes[0]);
+	uint32 B = uint32(HashBytes[1]);
+	uint32 C = uint32(HashBytes[2]);
+	uint32 D = uint32(HashBytes[3]);
+
+	// Convert to a Variant 1 Version 3 UUID, which is meant to be created by hashing the namespace and name with MD5.
+	// We use that version even though we're using BLAKE3 instead of MD5 because we don't have a better alternative.
+	// It is still very useful for avoiding collisions with other UUID variants.
+	B = (B & ~0x0000f000) | 0x00003000; // Version 3 (MD5)
+	C = (C & ~0xc0000000) | 0x80000000; // Variant 1 (RFC 4122 UUID)
+	return FGuid(A, B, C, D);
+}
+
 FGuid FGuid::Combine(const FGuid& GuidA, const FGuid& GuidB)
 {
 	return FGuid(
diff --git a/Engine/Source/Runtime/Core/Public/Misc/Guid.h b/Engine/Source/Runtime/Core/Public/Misc/Guid.h
--- a/Engine/Source/Runtime/Core/Public/Misc/Guid.h
+++ b/Engine/Source/Runtime/Core/Public/Misc/Guid.h
@@ -360,6 +360,11 @@ public:
 	 * @return A new GUID.
 	 */
 	static CORE_API FGuid NewGuid();
+	/**
+	 * Create a guid by hashing the given path; this guid will be deterministic when called in multiple cook processes
+	 * and will thus avoid cook indeterminism caused by FGuid::NewGuid. ObjectPath and Seed must be deterministic.
+	 */
+	static CORE_API FGuid NewDeterministicGuid(FStringView ObjectPath, uint64 Seed = 0);
 
 	/**
 	 * Returns a GUID which is a combinationof the two provided ones.
