Fix logging a ANSI string as a TCHAR

diff --git a/Engine/Source/Runtime/Core/Private/AutoRTFM/CompilerInterface.cpp b/Engine/Source/Runtime/Core/Private/AutoRTFM/CompilerInterface.cpp
--- a/Engine/Source/Runtime/Core/Private/AutoRTFM/CompilerInterface.cpp
+++ b/Engine/Source/Runtime/Core/Private/AutoRTFM/CompilerInterface.cpp
@@ -82,7 +82,7 @@ extern "C" void autortfm_llvm_fail(const char* Message)
 {
     if (Message)
     {
-		UE_LOG(LogAutoRTFM, Warning, TEXT("Transaction failing because of language issue '%s'."), Message);
+		UE_LOG(LogAutoRTFM, Warning, TEXT("Transaction failing because of language issue '%s'."), ANSI_TO_TCHAR(Message));
     }
     else
     {
@@ -102,7 +102,7 @@ extern "C" void autortfm_llvm_error(const char* Message)
 {
 	if (Message)
 	{
-		UE_LOG(LogAutoRTFM, Fatal, TEXT("Transaction failing because of LLVM issue '%s'."), Message);
+		UE_LOG(LogAutoRTFM, Fatal, TEXT("Transaction failing because of LLVM issue '%s'."), ANSI_TO_TCHAR(Message));
 	}
 	else
 	{
diff --git a/Engine/Source/Runtime/Core/Private/AutoRTFM/FunctionMapInlines.h b/Engine/Source/Runtime/Core/Private/AutoRTFM/FunctionMapInlines.h
--- a/Engine/Source/Runtime/Core/Private/AutoRTFM/FunctionMapInlines.h
+++ b/Engine/Source/Runtime/Core/Private/AutoRTFM/FunctionMapInlines.h
@@ -17,7 +17,7 @@ inline void* FunctionMapLookup(void* OldFunction, const char* Where)
     {
 		if (Where)
 		{
-			UE_LOG(LogAutoRTFM, Warning, TEXT("Could not find function %p '%s' where '%s'."), OldFunction, *GetFunctionDescription(OldFunction), Where);
+			UE_LOG(LogAutoRTFM, Warning, TEXT("Could not find function %p '%s' where '%s'."), OldFunction, *GetFunctionDescription(OldFunction), ANSI_TO_TCHAR(Where));
 		}
 		else
 		{
