Add missing 2d/3d features :
- ZeroValue and NoneValue in TIntVector2 (from TIntVector3)
- GridValues() in TDenseGrid3 (from TDenseGrid2)

diff --git a/Engine/Source/Runtime/Core/Public/Math/IntVector.h b/Engine/Source/Runtime/Core/Public/Math/IntVector.h
--- a/Engine/Source/Runtime/Core/Public/Math/IntVector.h
+++ b/Engine/Source/Runtime/Core/Public/Math/IntVector.h
@@ -581,6 +581,12 @@ struct TIntVector2
 		IntType XY[2];
 	};
 
+	/** An int point with zeroed values. */
+	static const TIntVector2 ZeroValue;
+
+	/** An int point with INDEX_NONE values. */
+	static const TIntVector2 NoneValue;
+
 	TIntVector2() = default;
 
 	TIntVector2(IntType InX, IntType InY)
@@ -814,6 +820,12 @@ PRAGMA_ENABLE_DEPRECATION_WARNINGS
 	bool SerializeFromMismatchedTag(FName StructTag, FArchive& Ar);
 };
 
+template <typename IntType>
+const TIntVector2<IntType> TIntVector2<IntType>::ZeroValue(0, 0);
+
+template <typename IntType>
+const TIntVector2<IntType> TIntVector2<IntType>::NoneValue(INDEX_NONE, INDEX_NONE);
+
 ///////////////////////////////////////////////////////////////////////////////////////////////////
 
 template <typename InIntType>
diff --git a/Engine/Source/Runtime/GeometryCore/Public/Spatial/DenseGrid3.h b/Engine/Source/Runtime/GeometryCore/Public/Spatial/DenseGrid3.h
--- a/Engine/Source/Runtime/GeometryCore/Public/Spatial/DenseGrid3.h
+++ b/Engine/Source/Runtime/GeometryCore/Public/Spatial/DenseGrid3.h
@@ -116,7 +116,14 @@ public:
 		return Buffer[I + Dimensions.X * (J + Dimensions.Y * K)];
 	}
 
-
+	TArray<ElemType>& GridValues()
+	{
+		return Buffer;
+	}
+	const TArray<ElemType>& GridValues() const
+	{
+		return Buffer;
+	}
 
 	/**
 	* @return the grid value at (X,Y,Z)
