Fixing UE_LOG sanitizer errors

diff --git a/Engine/Source/Runtime/Core/Private/AutoRTFM/Transaction.cpp b/Engine/Source/Runtime/Core/Private/AutoRTFM/Transaction.cpp
--- a/Engine/Source/Runtime/Core/Private/AutoRTFM/Transaction.cpp
+++ b/Engine/Source/Runtime/Core/Private/AutoRTFM/Transaction.cpp
@@ -26,7 +26,7 @@ bool FTransaction::IsFresh() const
 
 void FTransaction::AbortWithoutThrowing()
 {
-	UE_LOG(LogAutoRTFM, Verbose, TEXT("Aborting '%s'!"), GetContextStatusName(Context->GetStatus()));
+	UE_LOG(LogAutoRTFM, Verbose, TEXT("Aborting '%hs'!"), GetContextStatusName(Context->GetStatus()));
 
     ASSERT(Context->GetStatus() == EContextStatus::AbortedByFailedLockAcquisition
            || Context->GetStatus() == EContextStatus::AbortedByLanguage
