FGenericPlatformMisc::ClipboardCopy() and FGenericPlatformMisc::ClipboardPaste() removed after deprecation.

diff --git a/Engine/Source/Runtime/Core/Private/GenericPlatform/GenericPlatformMisc.cpp b/Engine/Source/Runtime/Core/Private/GenericPlatform/GenericPlatformMisc.cpp
--- a/Engine/Source/Runtime/Core/Private/GenericPlatform/GenericPlatformMisc.cpp
+++ b/Engine/Source/Runtime/Core/Private/GenericPlatform/GenericPlatformMisc.cpp
@@ -990,30 +990,6 @@ const TCHAR* FGenericPlatformMisc::GetSystemErrorMessage(TCHAR* OutBuffer, int32
 	return OutBuffer;
 }
 
-void FGenericPlatformMisc::ClipboardCopy(const TCHAR* Str)
-{
-	if(ClipboardCopyShim == nullptr)
-	{
-		UE_LOG(LogGenericPlatformMisc, Warning, TEXT("ClipboardCopyShim() is not bound; ignoring."));
-	}
-	else
-	{
-		ClipboardCopyShim(Str);
-	}
-}
-
-void FGenericPlatformMisc:: ClipboardPaste(class FString& Dest)
-{
-	if(ClipboardPasteShim == nullptr)
-	{
-		UE_LOG(LogGenericPlatformMisc, Warning, TEXT("ClipboardPasteShim() is not bound; ignoring."));
-	}
-	else
-	{
-		ClipboardPasteShim(Dest);
-	}
-}
-
 void FGenericPlatformMisc::CreateGuid(FGuid& Guid)
 {
 	QUICK_SCOPE_CYCLE_COUNTER(STAT_FGenericPlatformMisc_CreateGuid);
diff --git a/Engine/Source/Runtime/Core/Public/GenericPlatform/GenericPlatformMisc.h b/Engine/Source/Runtime/Core/Public/GenericPlatform/GenericPlatformMisc.h
--- a/Engine/Source/Runtime/Core/Public/GenericPlatform/GenericPlatformMisc.h
+++ b/Engine/Source/Runtime/Core/Public/GenericPlatform/GenericPlatformMisc.h
@@ -1055,14 +1055,6 @@ public:
 	 */
 	static CORE_API const TCHAR* GetSystemErrorMessage(TCHAR* OutBuffer, int32 BufferCount, int32 Error);
 
-	/** Copies text to the operating system clipboard. */
-	UE_DEPRECATED(4.18, "FPlatformMisc::ClipboardCopy() has been superseded by FPlatformApplicationMisc::ClipboardCopy()")
-	static CORE_API void ClipboardCopy(const TCHAR* Str);
-
-	/** Pastes in text from the operating system clipboard. */
-	UE_DEPRECATED(4.18, "FPlatformMisc::ClipboardPaste() has been superseded by FPlatformApplicationMisc::ClipboardPaste()")
-	static CORE_API void ClipboardPaste(class FString& Dest);
-
 	/** Create a new globally unique identifier. **/
 	static CORE_API void CreateGuid(struct FGuid& Result);
 
