Add FTextBuilder::GetNumLines

diff --git a/Engine/Source/Runtime/Core/Private/Internationalization/Text.cpp b/Engine/Source/Runtime/Core/Private/Internationalization/Text.cpp
--- a/Engine/Source/Runtime/Core/Private/Internationalization/Text.cpp
+++ b/Engine/Source/Runtime/Core/Private/Internationalization/Text.cpp
@@ -1921,11 +1921,16 @@ void FTextBuilder::Clear()
 	Lines.Reset();
 }
 
-bool FTextBuilder::IsEmpty()
+bool FTextBuilder::IsEmpty() const
 {
 	return Lines.Num() == 0;
 }
 
+int32 FTextBuilder::GetNumLines() const
+{
+	return Lines.Num();
+}
+
 FText FTextBuilder::ToText() const
 {
 	return FText::Join(FText::AsCultureInvariant(LINE_TERMINATOR), Lines);
diff --git a/Engine/Source/Runtime/Core/Public/Internationalization/Text.h b/Engine/Source/Runtime/Core/Public/Internationalization/Text.h
--- a/Engine/Source/Runtime/Core/Public/Internationalization/Text.h
+++ b/Engine/Source/Runtime/Core/Public/Internationalization/Text.h
@@ -1265,7 +1265,12 @@ public:
 	/**
 	 * Check to see if the builder has any data.
 	 */
-	CORE_API bool IsEmpty();
+	CORE_API bool IsEmpty() const;
+
+	/**
+	 * Returns the number of lines.
+	 */
+	CORE_API int32 GetNumLines() const;
 
 	/**
 	 * Build the current set of input into a FText.
