Moved FString usage out of various headers to allow them to be compiled with only a forward-declaration.

diff --git a/Engine/Source/Runtime/Core/Private/Math/Color.cpp b/Engine/Source/Runtime/Core/Private/Math/Color.cpp
--- a/Engine/Source/Runtime/Core/Private/Math/Color.cpp
+++ b/Engine/Source/Runtime/Core/Private/Math/Color.cpp
@@ -525,6 +525,25 @@ FLinearColor FLinearColor::MakeRandomSeededColor(int32 Seed)
 	return FLinearColor(R,G,B);
 }
 
+FString FLinearColor::ToString() const
+{
+	return FString::Printf(TEXT("(R=%f,G=%f,B=%f,A=%f)"), R, G, B, A);
+}
+
+bool FLinearColor::InitFromString(const FString& InSourceString)
+{
+	R = G = B = 0.f;
+	A = 1.f;
+
+	// The initialization is only successful if the R, G, and B values can all be parsed from the string
+	const bool bSuccessful = FParse::Value(*InSourceString, TEXT("R="), R) && FParse::Value(*InSourceString, TEXT("G="), G) && FParse::Value(*InSourceString, TEXT("B="), B);
+
+	// Alpha is optional, so don't factor in its presence (or lack thereof) in determining initialization success
+	FParse::Value(*InSourceString, TEXT("A="), A);
+
+	return bSuccessful;
+}
+
 FColor FColor::MakeRandomSeededColor(int32 Seed)
 {
 	return FLinearColor::MakeRandomSeededColor( Seed ).ToFColor( true );
@@ -545,6 +564,30 @@ FColor FColor::MakeRedToGreenColorFromScalar(float Scalar)
 	return FColor(R, G, B);
 }
 
+FString FColor::ToHex() const
+{
+	return FString::Printf(TEXT("%02X%02X%02X%02X"), R, G, B, A);
+}
+
+FString FColor::ToString() const
+{
+	return FString::Printf(TEXT("(R=%i,G=%i,B=%i,A=%i)"), R, G, B, A);
+}
+
+bool FColor::InitFromString(const FString& InSourceString)
+{
+	R = G = B = 0;
+	A = 255;
+
+	// The initialization is only successful if the R, G, and B values can all be parsed from the string
+	const bool bSuccessful = FParse::Value(*InSourceString, TEXT("R="), R) && FParse::Value(*InSourceString, TEXT("G="), G) && FParse::Value(*InSourceString, TEXT("B="), B);
+
+	// Alpha is optional, so don't factor in its presence (or lack thereof) in determining initialization success
+	FParse::Value(*InSourceString, TEXT("A="), A);
+
+	return bSuccessful;
+}
+
 void ComputeAndFixedColorAndIntensity(const FLinearColor& InLinearColor,FColor& OutColor,float& OutIntensity)
 {
 	float MaxComponent = FMath::Max(UE_DELTA,FMath::Max(InLinearColor.R,FMath::Max(InLinearColor.G,InLinearColor.B)));
diff --git a/Engine/Source/Runtime/Core/Private/UObject/UnrealNames.cpp b/Engine/Source/Runtime/Core/Private/UObject/UnrealNames.cpp
--- a/Engine/Source/Runtime/Core/Private/UObject/UnrealNames.cpp
+++ b/Engine/Source/Runtime/Core/Private/UObject/UnrealNames.cpp
@@ -3315,6 +3315,11 @@ FString FName::ToString() const
 	return Out;
 }
 
+FString LexToString(const FName& Name)
+{
+	return Name.ToString();
+}
+
 void FName::ToString(FString& Out) const
 {
 #if UE_FNAME_OUTLINE_NUMBER
@@ -3527,6 +3532,16 @@ FString FName::SafeString(FNameEntryId InDisplayIndex, int32 InstanceNumber)
 #endif // UE_FNAME_OUTLINE_NUMBER
 }
 
+bool FName::IsValidXName() const
+{
+	return IsValidXName(*this, FString(INVALID_NAME_CHARACTERS), (FText*)nullptr, (const FText*)nullptr);
+}
+
+bool FName::IsValidXName(FText& OutReason) const
+{
+	return IsValidXName(*this, FString(INVALID_NAME_CHARACTERS), &OutReason);
+}
+
 bool FName::IsValidXName(const FName InName, const FString& InInvalidChars, FText* OutReason, const FText* InErrorCtx)
 {
 	return IsValidXName(FNameBuilder(InName), InInvalidChars, OutReason, InErrorCtx);
@@ -3632,6 +3647,16 @@ bool FName::IsValidXName(const FStringView& InName, const FString& InInvalidChar
 	return true;
 }
 
+bool FName::IsValidObjectName(FText& OutReason) const
+{
+	return IsValidXName(*this, INVALID_OBJECTNAME_CHARACTERS, &OutReason);
+}
+
+bool FName::IsValidGroupName(FText& OutReason, bool bIsGroupName) const
+{
+	return IsValidXName(*this, INVALID_LONGPACKAGE_CHARACTERS, &OutReason);
+}
+
 FString FName::SanitizeWhitespace(const FString& FNameString)
 {
 	FString ReturnStr;
diff --git a/Engine/Source/Runtime/Core/Public/Math/Color.h b/Engine/Source/Runtime/Core/Public/Math/Color.h
--- a/Engine/Source/Runtime/Core/Public/Math/Color.h
+++ b/Engine/Source/Runtime/Core/Public/Math/Color.h
@@ -409,10 +409,7 @@ struct FLinearColor
 		return FMath::Min( FMath::Min( FMath::Min( R, G ), B ), A );
 	}
 
-	FString ToString() const
-	{
-		return FString::Printf(TEXT("(R=%f,G=%f,B=%f,A=%f)"),R,G,B,A);
-	}
+	CORE_API FString ToString() const;
 
 	/**
 	 * Initialize this Color based on an FString. The String is expected to contain R=, G=, B=, A=.
@@ -421,19 +418,7 @@ struct FLinearColor
 	 * @param InSourceString FString containing the color values.
 	 * @return true if the R,G,B values were read successfully; false otherwise.
 	 */
-	bool InitFromString( const FString& InSourceString )
-	{
-		R = G = B = 0.f;
-		A = 1.f;
-
-		// The initialization is only successful if the R, G, and B values can all be parsed from the string
-		const bool bSuccessful = FParse::Value( *InSourceString, TEXT("R=") , R ) && FParse::Value( *InSourceString, TEXT("G="), G ) && FParse::Value( *InSourceString, TEXT("B="), B );
-		
-		// Alpha is optional, so don't factor in its presence (or lack thereof) in determining initialization success
-		FParse::Value( *InSourceString, TEXT("A="), A );
-		
-		return bSuccessful;
-	}
+	CORE_API bool InitFromString( const FString& InSourceString );
 
 	/**
 	 * Helper for pixel format conversions. Clamps to [0,1], mapping NaNs to 0,
@@ -693,10 +678,7 @@ public:
 	 * @return Hexadecimal string.
 	 * @see FromHex, ToString
 	 */
-	FORCEINLINE FString ToHex() const
-	{
-		return FString::Printf(TEXT("%02X%02X%02X%02X"), R, G, B, A);
-	}
+	CORE_API FString ToHex() const;
 
 	/**
 	 * Converts this color value to a string.
@@ -704,10 +686,7 @@ public:
 	 * @return The string representation.
 	 * @see ToHex
 	 */
-	FORCEINLINE FString ToString() const
-	{
-		return FString::Printf(TEXT("(R=%i,G=%i,B=%i,A=%i)"), R, G, B, A);
-	}
+	CORE_API FString ToString() const;
 
 	/**
 	 * Initialize this Color based on an FString. The String is expected to contain R=, G=, B=, A=.
@@ -716,19 +695,7 @@ public:
 	 * @param	InSourceString	FString containing the color values.
 	 * @return true if the R,G,B values were read successfully; false otherwise.
 	 */
-	bool InitFromString( const FString& InSourceString )
-	{
-		R = G = B = 0;
-		A = 255;
-
-		// The initialization is only successful if the R, G, and B values can all be parsed from the string
-		const bool bSuccessful = FParse::Value( *InSourceString, TEXT("R=") , R ) && FParse::Value( *InSourceString, TEXT("G="), G ) && FParse::Value( *InSourceString, TEXT("B="), B );
-		
-		// Alpha is optional, so don't factor in its presence (or lack thereof) in determining initialization success
-		FParse::Value( *InSourceString, TEXT("A="), A );
-		
-		return bSuccessful;
-	}
+	CORE_API bool InitFromString( const FString& InSourceString );
 
 	/**
 	 * Gets the color in a packed uint32 format packed in the order ARGB.
diff --git a/Engine/Source/Runtime/Core/Public/UObject/NameTypes.h b/Engine/Source/Runtime/Core/Public/UObject/NameTypes.h
--- a/Engine/Source/Runtime/Core/Public/UObject/NameTypes.h
+++ b/Engine/Source/Runtime/Core/Public/UObject/NameTypes.h
@@ -760,10 +760,11 @@ public:
 	 *
 	 * @return	true if the name is valid
 	 */
-	bool IsValidXName( const FString& InInvalidChars = INVALID_NAME_CHARACTERS, FText* OutReason = nullptr, const FText* InErrorCtx = nullptr ) const
+	bool IsValidXName( const FString& InInvalidChars, FText* OutReason = nullptr, const FText* InErrorCtx = nullptr ) const
 	{
 		return IsValidXName(*this, InInvalidChars, OutReason, InErrorCtx);
 	}
+	CORE_API bool IsValidXName() const;
 
 	/**
 	 * Takes an FName and checks to see that it follows the rules that Unreal requires.
@@ -773,10 +774,11 @@ public:
 	 *
 	 * @return	true if the name is valid
 	 */
-	bool IsValidXName( FText& OutReason, const FString& InInvalidChars = INVALID_NAME_CHARACTERS ) const
+	bool IsValidXName( FText& OutReason, const FString& InInvalidChars ) const
 	{
 		return IsValidXName(*this, InInvalidChars, &OutReason);
 	}
+	CORE_API bool IsValidXName( FText& OutReason ) const;
 
 	/**
 	 * Takes an FName and checks to see that it follows the rules that Unreal requires for object names.
@@ -785,10 +787,7 @@ public:
 	 *
 	 * @return	true if the name is valid
 	 */
-	bool IsValidObjectName( FText& OutReason ) const
-	{
-		return IsValidXName(*this, INVALID_OBJECTNAME_CHARACTERS, &OutReason);
-	}
+	CORE_API bool IsValidObjectName( FText& OutReason ) const;
 
 	/**
 	 * Takes an FName and checks to see that it follows the rules that Unreal requires for package or group names.
@@ -798,10 +797,7 @@ public:
 	 *
 	 * @return	true if the name is valid
 	 */
-	bool IsValidGroupName( FText& OutReason, bool bIsGroupName=false ) const
-	{
-		return IsValidXName(*this, INVALID_LONGPACKAGE_CHARACTERS, &OutReason);
-	}
+	CORE_API bool IsValidGroupName( FText& OutReason, bool bIsGroupName=false ) const;
 
 	/**
 	 * Printing FNames in logging or on screen can be problematic when they contain Whitespace characters such as \n and \r,
@@ -1500,10 +1496,7 @@ FORCEINLINE FScriptName NameToScriptName(FName InName)
 	return FScriptName(InName);
 }
 
-FORCEINLINE FString LexToString(const FName& Name)
-{
-	return Name.ToString();
-}
+CORE_API FString LexToString(const FName& Name);
 
 FORCEINLINE void LexFromString(FName& Name, const TCHAR* Str)
 {
