LLM - Handle FPlatformTLS::GetTlsValue returning null; this can happen during thread termination on some platforms. Instead, create a backup map from ThreadId to ThreadState that we fall back to if GetTlsValue returns null.
#rnx

diff --git a/Engine/Source/Runtime/Core/Private/HAL/LowLevelMemTracker.cpp b/Engine/Source/Runtime/Core/Private/HAL/LowLevelMemTracker.cpp
--- a/Engine/Source/Runtime/Core/Private/HAL/LowLevelMemTracker.cpp
+++ b/Engine/Source/Runtime/Core/Private/HAL/LowLevelMemTracker.cpp
@@ -589,6 +589,12 @@ protected:
 
 	FCriticalSection PendingThreadStatesGuard;
 	TArray<FLLMThreadState*, FDefaultLLMAllocator> PendingThreadStates;
+	/**
+	 * Backup map from thread to threadstate. The primary lookup method is FPlatformTLS, but that is unavailable
+	 * during thread termination on some platforms. When unavailable, we use this TMap (under PendingThreadStatesGuard
+	 * critical section) to find the state.
+	 */
+	TMap<uint32, FLLMThreadState*, FDefaultSetLLMAllocator> ThreadIdToThreadState;
 
 	/** Sum of memory from all tracked tags. Duplicated in separate storage to make it instantly available at any time of frame without waiting for accumulation from threads during update.
 	    GCC_ALIGN is required because it is modified in FPlatformAtomics::InterlockedAdd, which requires aligned values. */
@@ -3550,32 +3556,46 @@ void FLLMTracker::Initialise(
 FLLMThreadState* FLLMTracker::GetOrCreateState()
 {
 	// look for already allocated thread state
-	FLLMThreadState* State = (FLLMThreadState*)FPlatformTLS::GetTlsValue(TlsSlot);
-	// get one if needed
+	FLLMThreadState* State = GetState();
+	// Create one if needed
 	if (State == nullptr)
 	{
 		State = LLMRef.Allocator.New<FLLMThreadState>();
 		LLMCheckf(State != nullptr, TEXT("LLMRef.Allocator.New returned nullptr."));
 
-		// Add to pending thread states, these will be consumed on the GT
+		// Add to pending thread states, (these will be consumed on the main thread and transferred to ThreadStates,
+		// which is only read/write on main thread). Also add to our backup map from thread id to thread state.
+		uint32 ThreadId = FPlatformTLS::GetCurrentThreadId();
 		{
 			FScopeLock Lock(&PendingThreadStatesGuard);
 			PendingThreadStates.Add(State);
+			ThreadIdToThreadState.Add(ThreadId, State);
 		}
 
 		// push to Tls
 		FPlatformTLS::SetTlsValue(TlsSlot, State);
-
-		// Verify the Tls slot is working
-		FLLMThreadState* StoredPointer = GetState();
-		LLMCheckf(StoredPointer == State, TEXT("SetTlsValue/GetTlsValue failed to store and retrieve the state pointer."));
 	}
 	return State;
 }
 
 FLLMThreadState* FLLMTracker::GetState()
 {
-	return (FLLMThreadState*)FPlatformTLS::GetTlsValue(TlsSlot);
+	FLLMThreadState* State = (FLLMThreadState*)FPlatformTLS::GetTlsValue(TlsSlot);
+	if (!State)
+	{
+		// GetTlsValue might return null even if we previously set it, if called during thread termination
+		// Check our backup mapping from thread id to thread state before concluding the state does not exist.
+		uint32 ThreadId = FPlatformTLS::GetCurrentThreadId();
+		{
+			FScopeLock Lock(&PendingThreadStatesGuard);
+			FLLMThreadState** ExistingState = ThreadIdToThreadState.Find(ThreadId);
+			if (ExistingState)
+			{
+				State = *ExistingState;
+			}
+		}
+	}
+	return State; // Can be nullptr if not yet created
 }
 
 void FLLMTracker::PushTag(ELLMTag EnumTag, ELLMTagSet TagSet)
@@ -3867,6 +3887,7 @@ void FLLMTracker::Clear()
 			LLMRef.Allocator.Delete(ThreadState);
 		}
 		PendingThreadStates.Empty();
+		ThreadIdToThreadState.Empty();
 	}
 
 	for (FLLMThreadState* ThreadState : ThreadStates)
