removed tasks LLM instrumentation because they are not useful and steal data from other tags. tasks are a low-level thing that shouldn't be represented in a high-level mem usage overview

diff --git a/Engine/Source/Runtime/Core/Private/Async/TaskGraph.cpp b/Engine/Source/Runtime/Core/Private/Async/TaskGraph.cpp
--- a/Engine/Source/Runtime/Core/Private/Async/TaskGraph.cpp
+++ b/Engine/Source/Runtime/Core/Private/Async/TaskGraph.cpp
@@ -1037,8 +1037,6 @@ private:
 	**/
 	uint64 ProcessTasks()
 	{
-		LLM_SCOPE_BYNAME(TEXT("Tasks/AnyThread/ProcessTasks"));
-
 		TStatId StallStatId;
 		bool bCountAsStall = true;
 		uint64 ProcessedTasks = 0;
@@ -1951,8 +1949,6 @@ public:
 private:
 	void QueueTask(class FBaseGraphTask* Task, bool bWakeUpWorker, ENamedThreads::Type InThreadToExecuteOn, ENamedThreads::Type InCurrentThreadIfKnown) override
 	{
-		LLM_SCOPE_BYNAME(TEXT("Tasks/QueueTask"));
-
 		if (ENamedThreads::GetThreadIndex(InThreadToExecuteOn) == ENamedThreads::AnyThread)
 		{
 #if TASKGRAPH_NEW_FRONTEND
@@ -2081,7 +2077,6 @@ private:
 
 	void WaitUntilTasksComplete(const FGraphEventArray& Tasks, ENamedThreads::Type CurrentThreadIfKnown = ENamedThreads::AnyThread) final override
 	{
-		LLM_SCOPE_BYNAME(TEXT("Tasks/WaitUntilTasksComplete"));
 		TaskTrace::FWaitingScope WaitingScope(GetTraceIds(Tasks));
 		TRACE_CPUPROFILER_EVENT_SCOPE(WaitUntilTasksComplete);
 
@@ -2431,7 +2426,6 @@ static TLockFreeClassAllocator_TLSCache<FGraphEvent, PLATFORM_CACHE_LINE_SIZE>&
 
 FGraphEventRef FGraphEvent::CreateGraphEvent()
 {
-	LLM_SCOPE_BYNAME(TEXT("Tasks/FGraphEvent/CreateGraphEvent"));
 	FGraphEvent* Instance = new(GetGraphEventAllocator().Allocate()) FGraphEvent{};
 	return Instance;
 }
diff --git a/Engine/Source/Runtime/Core/Public/Async/TaskGraphInterfaces.h b/Engine/Source/Runtime/Core/Public/Async/TaskGraphInterfaces.h
--- a/Engine/Source/Runtime/Core/Public/Async/TaskGraphInterfaces.h
+++ b/Engine/Source/Runtime/Core/Public/Async/TaskGraphInterfaces.h
@@ -508,8 +508,6 @@ public:
 			return;
 		}
 
-		LLM_SCOPE_BYNAME(TEXT("Tasks/FGraphEvent/DontCompleteUntil"));
-
 		if (GetExtendedPriority() == UE::Tasks::EExtendedTaskPriority::TaskEvent)
 		{	// TaskEvent can't have nested tasks, add it as a prerequisite instead
 			AddPrerequisites(*NestedTask);
@@ -608,7 +606,6 @@ public:
 		template<typename...T>
 		FORCEINLINE_DEBUGGABLE TGraphTask* ConstructAndHoldImpl(T&&... Args)
 		{
-			LLM_SCOPE_BYNAME(TEXT("Tasks/TGraphTask/ConstructAndHoldImpl"));
 			TGraphTask* Task = new TGraphTask(Prerequisites);
 			TTask* TaskObject = new(&Task->TaskStorage) TTask(Forward<T>(Args)...);
 
@@ -712,8 +709,6 @@ inline void FGraphEventImpl::operator delete(void* Ptr)
 
 inline FGraphEventRef FBaseGraphTask::CreateGraphEvent()
 {
-	LLM_SCOPE_BYNAME(TEXT("Tasks/FGraphEvent/CreateGraphEvent"));
-
 	FGraphEventImpl* GraphEvent = new FGraphEventImpl;
 	return FGraphEventRef{ GraphEvent, /*bAddRef = */ false };
 }
@@ -930,7 +925,6 @@ public:
 	**/
 	bool AddSubsequent(class FBaseGraphTask* Subsequent)
 	{
-		LLM_SCOPE_BYNAME(TEXT("Tasks/FGraphEvent/AddSubsequent"));
 		bool bSucceeded = SubsequentList.PushIfNotClosed(Subsequent);
 		if (bSucceeded)
 		{
@@ -954,7 +948,6 @@ public:
 	**/
 	void DontCompleteUntil(FGraphEventRef EventToWaitFor)
 	{
-		LLM_SCOPE_BYNAME(TEXT("Tasks/FGraphEvent/DontCompleteUntil"));
 		checkThreadGraph(!IsComplete()); // it is not legal to add a DontCompleteUntil after the event has been completed. Basically, this is only legal within a task function.
 		EventsToWaitFor.Emplace(EventToWaitFor);
 		TaskTrace::SubsequentAdded(EventToWaitFor->GetTraceId(), GetTraceId());
@@ -1213,8 +1206,6 @@ public:
 	{
 		FGraphEventRef GraphEvent = TTask::GetSubsequentsMode() == ESubsequentsMode::FireAndForget ? NULL : FGraphEvent::CreateGraphEvent();
 
-		LLM_SCOPE_BYNAME(TEXT("Tasks/TGraphTask/CreateTask"));
-
 		int32 NumPrereq = Prerequisites ? Prerequisites->Num() : 0;
 		return FConstructor(new TGraphTask(MoveTemp(GraphEvent), NumPrereq), Prerequisites, CurrentThreadIfKnown);
 	}
diff --git a/Engine/Source/Runtime/Core/Public/Tasks/Task.h b/Engine/Source/Runtime/Core/Public/Tasks/Task.h
--- a/Engine/Source/Runtime/Core/Public/Tasks/Task.h
+++ b/Engine/Source/Runtime/Core/Public/Tasks/Task.h
@@ -80,7 +80,6 @@ namespace UE::Tasks
 			// @return true if the task is completed
 			bool Wait(FTimespan Timeout = FTimespan::MaxValue()) const
 			{
-				LLM_SCOPE_BYNAME(TEXT("Tasks/Wait"));
 				return !IsValid() || Pimpl->Wait(FTimeout{ Timeout });
 			}
 
@@ -89,7 +88,6 @@ namespace UE::Tasks
 			// @return true if the task is completed
 			bool BusyWait(FTimespan Timeout = FTimespan::MaxValue()) const
 			{
-				LLM_SCOPE_BYNAME(TEXT("Tasks/BusyWait"));
 				return !IsValid() || Pimpl->BusyWait(FTimeout{ Timeout });
 			}
 
@@ -99,7 +97,6 @@ namespace UE::Tasks
 			template<typename ConditionType>
 			bool BusyWait(ConditionType&& Condition) const
 			{
-				LLM_SCOPE_BYNAME(TEXT("Tasks/BusyWait"));
 				return !IsValid() || Pimpl->BusyWait(Forward<ConditionType>(Condition));
 			}
 
@@ -118,8 +115,6 @@ namespace UE::Tasks
 			{
 				check(!IsValid());
 
-				LLM_SCOPE_BYNAME(TEXT("Tasks/Launch"));
-
 				using FExecutableTask = Private::TExecutableTask<std::decay_t<TaskBodyType>>;
 				FExecutableTask* Task = FExecutableTask::Create(DebugName, Forward<TaskBodyType>(TaskBody), Priority, ExtendedPriority);
 				// this must happen before launching, to support an ability to access the task itself from inside it
@@ -145,8 +140,6 @@ namespace UE::Tasks
 			{
 				check(!IsValid());
 
-				LLM_SCOPE_BYNAME(TEXT("Tasks/Launch"));
-
 				using FExecutableTask = Private::TExecutableTask<std::decay_t<TaskBodyType>>;
 				FExecutableTask* Task = FExecutableTask::Create(DebugName, Forward<TaskBodyType>(TaskBody), Priority, ExtendedPriority);
 				Task->AddPrerequisites(Forward<PrerequisitesCollectionType>(Prerequisites));
@@ -235,7 +228,6 @@ namespace UE::Tasks
 		template<typename PrerequisitesType>
 		void AddPrerequisites(const PrerequisitesType& Prerequisites)
 		{
-			LLM_SCOPE_BYNAME(TEXT("Tasks/FTaskEvent/AddPrerequisites"));
 			Pimpl->AddPrerequisites(Prerequisites);
 		}
 
@@ -243,8 +235,6 @@ namespace UE::Tasks
 		{
 			if (!IsCompleted()) // event can be triggered multiple times
 			{
-				LLM_SCOPE_BYNAME(TEXT("Tasks/FTaskEvent/Trigger"));
-
 				// An event is not "in the system" until it's triggered, and should be kept alive only by external references. Once it's triggered it's in the system 
 				// and can outlive external references, so we need to keep it alive by holding an internal reference. It will be released when the event is signalled
 				Pimpl->AddRef();
@@ -328,7 +318,6 @@ namespace UE::Tasks
 	template<typename TaskCollectionType>
 	bool Wait(const TaskCollectionType& Tasks, FTimespan InTimeout/* = FTimespan::MaxValue()*/)
 	{
-		LLM_SCOPE_BYNAME(TEXT("Tasks/Wait"));
 		TaskTrace::FWaitingScope WaitingScope(Private::GetTraceIds(Tasks));
 		TRACE_CPUPROFILER_EVENT_SCOPE(Tasks::Wait);
 
@@ -359,7 +348,6 @@ namespace UE::Tasks
 	template<typename TaskCollectionType>
 	bool BusyWait(const TaskCollectionType& Tasks, FTimespan InTimeout/* = FTimespan::MaxValue()*/)
 	{
-		LLM_SCOPE_BYNAME(TEXT("Tasks/BusyWait"));
 		TaskTrace::FWaitingScope WaitingScope(Private::GetTraceIds(Tasks));
 		TRACE_CPUPROFILER_EVENT_SCOPE(Tasks::BusyWait);
 
@@ -418,7 +406,6 @@ namespace UE::Tasks
 	template<typename TaskType>
 	void AddNested(const TaskType& Nested)
 	{
-		LLM_SCOPE_BYNAME(TEXT("Tasks/AddNested"));
 		Private::FTaskBase* Parent = Private::GetCurrentTask();
 		check(Parent != nullptr);
 		Parent->AddNested(*Nested.Pimpl);
diff --git a/Engine/Source/Runtime/Core/Public/Tasks/TaskPrivate.h b/Engine/Source/Runtime/Core/Public/Tasks/TaskPrivate.h
--- a/Engine/Source/Runtime/Core/Public/Tasks/TaskPrivate.h
+++ b/Engine/Source/Runtime/Core/Public/Tasks/TaskPrivate.h
@@ -754,7 +754,6 @@ namespace UE::Tasks
 		public:
 			static FTaskEventBase* Create(const TCHAR* DebugName)
 			{
-				LLM_SCOPE_BYNAME(TEXT("Tasks/FTaskEvent/Create"));
 				return new FTaskEventBase(DebugName);
 			}
 
