Add floating point append operators to FStringBuilderBase. Uses the same format specifiers as LexToString does (via TFormatSpecifier in UnrealTypeTraits.h)

diff --git a/Engine/Source/Runtime/Core/Public/Misc/StringBuilder.h b/Engine/Source/Runtime/Core/Public/Misc/StringBuilder.h
--- a/Engine/Source/Runtime/Core/Public/Misc/StringBuilder.h
+++ b/Engine/Source/Runtime/Core/Public/Misc/StringBuilder.h
@@ -503,6 +503,16 @@ inline FWideStringBuilderBase&		operator<<(FWideStringBuilderBase& Builder, uint
 inline FUtf8StringBuilderBase&		operator<<(FUtf8StringBuilderBase& Builder, int16 Value)							{ return Builder << int32(Value); }
 inline FUtf8StringBuilderBase&		operator<<(FUtf8StringBuilderBase& Builder, uint16 Value)							{ return Builder << uint32(Value); }
 
+// Floating Point Append Operators
+
+inline FAnsiStringBuilderBase&		operator<<(FAnsiStringBuilderBase& Builder, float Value)							{ return Builder.Appendf("%f", Value); }
+inline FWideStringBuilderBase&		operator<<(FWideStringBuilderBase& Builder, float Value)							{ return Builder.Appendf(WIDETEXT("%f"), Value); }
+inline FUtf8StringBuilderBase&		operator<<(FUtf8StringBuilderBase& Builder, float Value)							{ return Builder.Appendf(UTF8TEXT("%f"), Value); }
+
+inline FAnsiStringBuilderBase&		operator<<(FAnsiStringBuilderBase& Builder, double Value)							{ return Builder.Appendf("%f", Value); }
+inline FWideStringBuilderBase&		operator<<(FWideStringBuilderBase& Builder, double Value)							{ return Builder.Appendf(WIDETEXT("%f"), Value); }
+inline FUtf8StringBuilderBase&		operator<<(FUtf8StringBuilderBase& Builder, double Value)							{ return Builder.Appendf(UTF8TEXT("%f"), Value); }
+
 template <typename CharType, int32 BufferSize>
 class UE_DEPRECATED(5.3, "Use WriteToString<N>(...) or TStringBuilder<N>(InPlace, ...).") TWriteToString : public TStringBuilderWithBuffer<CharType, BufferSize>
 {
