Detect when we try to run a new transaction nest from within an open-commit or open-abort, and abort if so.

diff --git a/Engine/Source/Programs/AutoRTFMTests/Private/Abort.cpp b/Engine/Source/Programs/AutoRTFMTests/Private/Abort.cpp
--- a/Engine/Source/Programs/AutoRTFMTests/Private/Abort.cpp
+++ b/Engine/Source/Programs/AutoRTFMTests/Private/Abort.cpp
@@ -56,25 +56,73 @@ TEST_CASE("Abort.NestedAbortOrder")
 	unsigned Orderer = 0;
 
 	AutoRTFM::Commit([&]
+	{
+		InnerResult = AutoRTFM::Transact([&]
+			{
+				AutoRTFM::OpenAbort([&]
+					{
+						REQUIRE(1 == Orderer);
+						Orderer += 1;
+					});
+
+				AutoRTFM::OpenAbort([&]
+					{
+						REQUIRE(0 == Orderer);
+						Orderer += 1;
+					});
+
+				AutoRTFM::AbortTransaction();
+			});
+	});
+
+	REQUIRE(AutoRTFM::ETransactionResult::AbortedByRequest == InnerResult);
+	REQUIRE(2 == Orderer);
+}
+
+TEST_CASE("Abort.TransactionInOpenCommit")
+{
+	AutoRTFM::ETransactionResult InnerResult;
+
+	AutoRTFM::Commit([&]
+	{
+		AutoRTFM::OpenCommit([&]
+		{
+			bool bDidSomething = false;
+
+			InnerResult = AutoRTFM::Transact([&]
+			{
+				bDidSomething = true;
+			});
+
+			REQUIRE(false == bDidSomething);
+		});
+	});
+
+	REQUIRE(AutoRTFM::ETransactionResult::AbortedByTransactInOpenCommit == InnerResult);
+}
+
+TEST_CASE("Abort.TransactionInOpenAbort")
+{
+	AutoRTFM::ETransactionResult Result;
+	AutoRTFM::ETransactionResult InnerResult;
+
+	Result = AutoRTFM::Transact([&]
+	{
+		AutoRTFM::OpenAbort([&]
 		{
+			bool bDidSomething = false;
+
 			InnerResult = AutoRTFM::Transact([&]
-				{
-					AutoRTFM::OpenAbort([&]
-						{
-							REQUIRE(1 == Orderer);
-							Orderer += 1;
-						});
-
-					AutoRTFM::OpenAbort([&]
-						{
-							REQUIRE(0 == Orderer);
-							Orderer += 1;
-						});
-
-					AutoRTFM::AbortTransaction();
-				});
+			{
+				bDidSomething = true;
+			});
+
+			REQUIRE(false == bDidSomething);
 		});
 
-	REQUIRE(AutoRTFM::ETransactionResult::AbortedByRequest == InnerResult);
-	REQUIRE(2 == Orderer);
+		AutoRTFM::AbortTransaction();
+	});
+
+	REQUIRE(AutoRTFM::ETransactionResult::AbortedByTransactInOpenAbort == InnerResult);
+	REQUIRE(AutoRTFM::ETransactionResult::AbortedByRequest == Result);
 }
diff --git a/Engine/Source/Runtime/Core/Private/AutoRTFM/Context.cpp b/Engine/Source/Runtime/Core/Private/AutoRTFM/Context.cpp
--- a/Engine/Source/Runtime/Core/Private/AutoRTFM/Context.cpp
+++ b/Engine/Source/Runtime/Core/Private/AutoRTFM/Context.cpp
@@ -100,7 +100,7 @@ ETransactionResult FContext::CommitTransaction()
 		DumpState();
 		UE_LOG(LogAutoRTFM, Verbose, TEXT("Committing..."));
 
-		if (CurrentTransaction->AttemptToCommit())
+		if (AttemptToCommitTransaction(CurrentTransaction))
 		{
 			Result = ETransactionResult::Committed;
 		}
@@ -146,7 +146,15 @@ ETransactionResult FContext::AbortTransaction(bool bIsClosed)
 
 bool FContext::IsAborting() const
 {
-	return Status != EContextStatus::OnTrack && Status != EContextStatus::Idle;
+	switch (Status)
+	{
+	default:
+		return true;
+	case EContextStatus::OnTrack:
+	case EContextStatus::Idle:
+	case EContextStatus::Committing:
+		return false;
+	}
 }
 
 EContextStatus FContext::CallClosedNest(void (*ClosedFunction)(void* Arg), void* Arg)
@@ -230,7 +238,7 @@ ETransactionResult FContext::ResolveNestedTransaction(FTransaction* NewTransacti
 
 	if (Status == EContextStatus::OnTrack)
 	{
-		bool bCommitResult = NewTransaction->AttemptToCommit();
+		bool bCommitResult = AttemptToCommitTransaction(NewTransaction);
 		ASSERT(bCommitResult);
 		ASSERT(Status == EContextStatus::OnTrack);
 		return ETransactionResult::Committed;
@@ -251,6 +259,16 @@ ETransactionResult FContext::ResolveNestedTransaction(FTransaction* NewTransacti
 ETransactionResult FContext::Transact(void (*Function)(void* Arg), void* Arg)
 {
     constexpr bool bVerbose = false;
+
+    if (UNLIKELY(EContextStatus::Committing == Status))
+    {
+    	return ETransactionResult::AbortedByTransactInOpenCommit;
+    }
+
+    if (UNLIKELY(IsAborting()))
+    {
+    	return ETransactionResult::AbortedByTransactInOpenAbort;
+    }
     
     ASSERT(Status == EContextStatus::Idle || Status == EContextStatus::OnTrack);
 
@@ -297,7 +315,7 @@ ETransactionResult FContext::Transact(void (*Function)(void* Arg), void* Arg)
 				DumpState();
 				UE_LOG(LogAutoRTFM, Verbose, TEXT("Committing..."));
 
-                if (CurrentTransaction->AttemptToCommit())
+                if (AttemptToCommitTransaction(CurrentTransaction))
                 {
                     Result = ETransactionResult::Committed;
                     break;
diff --git a/Engine/Source/Runtime/Core/Private/AutoRTFM/Context.h b/Engine/Source/Runtime/Core/Private/AutoRTFM/Context.h
--- a/Engine/Source/Runtime/Core/Private/AutoRTFM/Context.h
+++ b/Engine/Source/Runtime/Core/Private/AutoRTFM/Context.h
@@ -68,6 +68,7 @@ private:
 	void PopTransaction();
 
 	ETransactionResult ResolveNestedTransaction(FTransaction* NewTransaction);
+	bool AttemptToCommitTransaction(FTransaction* const Transaction);
 
     void Set();
     
diff --git a/Engine/Source/Runtime/Core/Private/AutoRTFM/ContextInlines.h b/Engine/Source/Runtime/Core/Private/AutoRTFM/ContextInlines.h
--- a/Engine/Source/Runtime/Core/Private/AutoRTFM/ContextInlines.h
+++ b/Engine/Source/Runtime/Core/Private/AutoRTFM/ContextInlines.h
@@ -34,4 +34,20 @@ UE_AUTORTFM_FORCEINLINE void FContext::DidAllocate(void* LogicalAddress, size_t
     CurrentTransaction->DidAllocate(LogicalAddress, Size);
 }
 
+UE_AUTORTFM_FORCEINLINE bool FContext::AttemptToCommitTransaction(FTransaction* const Transaction)
+{
+    ASSERT(EContextStatus::OnTrack == Status);
+
+    Status = EContextStatus::Committing;
+
+    const bool bResult = Transaction->AttemptToCommit();
+
+    if (bResult)
+    {
+        Status = EContextStatus::OnTrack;    
+    }
+
+    return bResult;
+}
+
 } // namespace AutoRTFM
diff --git a/Engine/Source/Runtime/Core/Private/AutoRTFM/Transaction.cpp b/Engine/Source/Runtime/Core/Private/AutoRTFM/Transaction.cpp
--- a/Engine/Source/Runtime/Core/Private/AutoRTFM/Transaction.cpp
+++ b/Engine/Source/Runtime/Core/Private/AutoRTFM/Transaction.cpp
@@ -55,7 +55,7 @@ void FTransaction::AbortAndThrow()
 
 bool FTransaction::AttemptToCommit()
 {
-    ASSERT(Context->GetStatus() == EContextStatus::OnTrack);
+    ASSERT(Context->GetStatus() == EContextStatus::Committing);
     ASSERT(Context->GetCurrentTransaction() == this);
 
     Stats.Collect<EStatsKind::Commit>();
diff --git a/Engine/Source/Runtime/Core/Public/AutoRTFM/AutoRTFM.h b/Engine/Source/Runtime/Core/Public/AutoRTFM/AutoRTFM.h
--- a/Engine/Source/Runtime/Core/Public/AutoRTFM/AutoRTFM.h
+++ b/Engine/Source/Runtime/Core/Public/AutoRTFM/AutoRTFM.h
@@ -67,7 +67,9 @@ typedef enum
 {
     autortfm_aborted_by_request,
     autortfm_aborted_by_language,
-    autortfm_committed
+    autortfm_committed,
+	autortfm_aborted_by_transact_in_open_commit,
+	autortfm_aborted_by_transact_in_open_abort
 } autortfm_result;
 
 // This must match AutoRTFM::EContextStatus.
@@ -77,7 +79,8 @@ typedef enum
 	autortfm_status_ontrack,
 	autortfm_status_aborted_by_failed_lock_aquisition,
 	autortfm_status_aborted_by_language,
-	autortfm_status_aborted_by_request
+	autortfm_status_aborted_by_request,
+	autortfm_status_committing
 } autortfm_status;
 
 // Tells if we are currently running in a transaction. This will return true in an
@@ -358,7 +361,9 @@ enum class ETransactionResult
 {
     AbortedByRequest = autortfm_aborted_by_request,
     AbortedByLanguage = autortfm_aborted_by_language,
-    Committed = autortfm_committed
+    Committed = autortfm_committed,
+	AbortedByTransactInOpenCommit = autortfm_aborted_by_transact_in_open_commit,
+	AbortedByTransactInOpenAbort = autortfm_aborted_by_transact_in_open_abort
 };
 
 enum class EContextStatus
@@ -367,7 +372,8 @@ enum class EContextStatus
 	OnTrack = autortfm_status_ontrack,
 	AbortedByFailedLockAcquisition = autortfm_status_aborted_by_failed_lock_aquisition,
 	AbortedByLanguage = autortfm_status_aborted_by_language,
-	AbortedByRequest = autortfm_status_aborted_by_request
+	AbortedByRequest = autortfm_status_aborted_by_request,
+	Committing = autortfm_status_committing,
 };
 
 UE_AUTORTFM_FORCEINLINE bool IsTransactional() { return autortfm_is_transactional(); }
