Add an autortfm attribute that lets us say 'this other function is my transactionalized variant'.

diff --git a/Engine/Source/Runtime/Core/Private/AutoRTFM/API.cpp b/Engine/Source/Runtime/Core/Private/AutoRTFM/API.cpp
--- a/Engine/Source/Runtime/Core/Private/AutoRTFM/API.cpp
+++ b/Engine/Source/Runtime/Core/Private/AutoRTFM/API.cpp
@@ -46,7 +46,7 @@ UE_AUTORTFM_FORCEINLINE autortfm_result TransactThenOpenImpl(void (*Work)(void*
 		}));
 }
 
-extern "C" UE_AUTORTFM_NOAUTORTFM bool autortfm_is_transactional()
+extern "C" UE_AUTORTFM_AUTORTFM("RTFM_autortfm_is_transactional") bool autortfm_is_transactional()
 {
 	if (GAutoRTFMRuntimeEnabled)
 	{
@@ -56,13 +56,13 @@ extern "C" UE_AUTORTFM_NOAUTORTFM bool autortfm_is_transactional()
 	return false;
 }
 
-extern "C" UE_AUTORTFM_NOAUTORTFM bool autortfm_is_closed()
+extern "C" UE_AUTORTFM_AUTORTFM("RTFM_autortfm_is_closed") bool autortfm_is_closed()
 {
     return false;
 }
 
 // First Part - the API exposed outside transactions.
-extern "C" UE_AUTORTFM_NOAUTORTFM autortfm_result autortfm_transact(void (*Work)(void* Arg), void* Arg)
+extern "C" UE_AUTORTFM_AUTORTFM("RTFM_autortfm_transact") autortfm_result autortfm_transact(void (*Work)(void* Arg), void* Arg)
 {
 	if (GAutoRTFMRuntimeEnabled)
 	{
@@ -73,42 +73,42 @@ extern "C" UE_AUTORTFM_NOAUTORTFM autortfm_result autortfm_transact(void (*Work)
 	return autortfm_committed;
 }
 
-extern "C" UE_AUTORTFM_NOAUTORTFM autortfm_result autortfm_transact_then_open(void (*Work)(void* Arg), void* Arg)
+extern "C" UE_AUTORTFM_AUTORTFM("RTFM_autortfm_transact_then_open") autortfm_result autortfm_transact_then_open(void (*Work)(void* Arg), void* Arg)
 {
     return TransactThenOpenImpl(Work, Arg);
 }
 
-extern "C" UE_AUTORTFM_NOAUTORTFM void autortfm_commit(void (*Work)(void* Arg), void* Arg)
+extern "C" UE_AUTORTFM_AUTORTFM("RTFM_autortfm_commit") void autortfm_commit(void (*Work)(void* Arg), void* Arg)
 {
     autortfm_result Result = autortfm_transact(Work, Arg);
 	UE_CLOG(Result != autortfm_committed, LogAutoRTFM, Fatal, TEXT("Unexpected transaction result: %u."), Result);
 }
 
-extern "C" UE_AUTORTFM_NOAUTORTFM void autortfm_abort()
+extern "C" UE_AUTORTFM_AUTORTFM("RTFM_autortfm_abort") void autortfm_abort()
 {
 	UE_CLOG(!FContext::IsTransactional(), LogAutoRTFM, Fatal, TEXT("The function `autortfm_abort` was called from outside a transaction."));
 	FContext::Get()->AbortByRequestAndThrow();
 }
 
-extern "C" UE_AUTORTFM_NOAUTORTFM bool autortfm_start_transaction()
+extern "C" UE_AUTORTFM_AUTORTFM("RTFM_autortfm_start_transaction") bool autortfm_start_transaction()
 {
 	UE_CLOG(!FContext::IsTransactional(), LogAutoRTFM, Fatal, TEXT("The function `autortfm_start_transaction` was called from outside a transact."));
 	return FContext::Get()->StartTransaction();
 }
 
-extern "C" UE_AUTORTFM_NOAUTORTFM autortfm_result autortfm_commit_transaction()
+extern "C" UE_AUTORTFM_AUTORTFM("RTFM_autortfm_commit_transaction") autortfm_result autortfm_commit_transaction()
 {
 	UE_CLOG(!FContext::IsTransactional(), LogAutoRTFM, Fatal, TEXT("The function `autortfm_commit_transaction` was called from outside a transact."));
 	return static_cast<autortfm_result>(FContext::Get()->CommitTransaction());
 }
 
-extern "C" UE_AUTORTFM_NOAUTORTFM autortfm_result autortfm_abort_transaction()
+extern "C" UE_AUTORTFM_AUTORTFM("RTFM_autortfm_abort_transaction") autortfm_result autortfm_abort_transaction()
 {
 	UE_CLOG(!FContext::IsTransactional(), LogAutoRTFM, Fatal, TEXT("The function `autortfm_abort_transaction` was called from outside a transact."));
 	return static_cast<autortfm_result>(FContext::Get()->AbortTransaction(false));
 }
 
-extern "C" UE_AUTORTFM_NOAUTORTFM void autortfm_clear_transaction_status()
+extern "C" UE_AUTORTFM_AUTORTFM("RTFM_autortfm_clear_transaction_status") void autortfm_clear_transaction_status()
 {
 	ASSERT(FContext::Get()->IsAborting());
 	FContext::Get()->ClearTransactionStatus();
@@ -130,21 +130,21 @@ extern "C" UE_AUTORTFM_NOAUTORTFM bool autortfm_current_nest_throw()
 	return true;
 }
 
-extern "C" UE_AUTORTFM_NOAUTORTFM void autortfm_abort_if_transactional()
+extern "C" UE_AUTORTFM_AUTORTFM("RTFM_autortfm_abort_if_transactional") void autortfm_abort_if_transactional()
 {
 	UE_CLOG(FContext::IsTransactional(), LogAutoRTFM, Fatal, TEXT("The function `autortfm_abort_if_transactional` was called from an open inside a transaction."));
 }
 
-extern "C" UE_AUTORTFM_NOAUTORTFM void autortfm_abort_if_closed()
+extern "C" UE_AUTORTFM_AUTORTFM("RTFM_autortfm_abort_if_closed") void autortfm_abort_if_closed()
 {
 }
 
-extern "C" UE_AUTORTFM_NOAUTORTFM void autortfm_open(void (*Work)(void* Arg), void* Arg)
+extern "C" UE_AUTORTFM_AUTORTFM("RTFM_autortfm_open") void autortfm_open(void (*Work)(void* Arg), void* Arg)
 {
 	Work(Arg);
 }
 
-extern "C" UE_AUTORTFM_NOAUTORTFM autortfm_status autortfm_close(void (*Work)(void* Arg), void* Arg)
+extern "C" UE_AUTORTFM_AUTORTFM("RTFM_autortfm_close") autortfm_status autortfm_close(void (*Work)(void* Arg), void* Arg)
 {
 	autortfm_status Result = autortfm_status_ontrack;
 
@@ -167,7 +167,7 @@ extern "C" UE_AUTORTFM_NOAUTORTFM autortfm_status autortfm_close(void (*Work)(vo
 	return Result;
 }
 
-extern "C" UE_AUTORTFM_NOAUTORTFM void autortfm_record_open_write(void* Ptr, size_t Size)
+extern "C" UE_AUTORTFM_AUTORTFM("RTFM_autortfm_record_open_write") void autortfm_record_open_write(void* Ptr, size_t Size)
 {
     FContext::Get()->CheckOpenRecordWrite(Ptr);
 	FContext::Get()->RecordWrite(Ptr, Size);
@@ -179,30 +179,30 @@ extern "C" UE_AUTORTFM_NOAUTORTFM void autortfm_register_open_function(void* Ori
     FunctionMapAdd(OriginalFunction, NewFunction);
 }
 
-UE_AUTORTFM_NOAUTORTFM void OpenCommit(TFunction<void()>&& Work)
+UE_AUTORTFM_AUTORTFM("RTFM_OpenCommit") void OpenCommit(TFunction<void()> && Work)
 {
     Work();
 }
 
-UE_AUTORTFM_NOAUTORTFM void OpenAbort(TFunction<void()>&& Work)
+UE_AUTORTFM_AUTORTFM("RTFM_OpenAbort") void OpenAbort(TFunction<void()> && Work)
 {
 }
 
-extern "C" UE_AUTORTFM_NOAUTORTFM void autortfm_open_commit(void (*Work)(void* Arg), void* Arg)
+extern "C" UE_AUTORTFM_AUTORTFM("RTFM_autortfm_open_commit") void autortfm_open_commit(void (*Work)(void* Arg), void* Arg)
 {
     Work(Arg);
 }
 
-extern "C" UE_AUTORTFM_NOAUTORTFM void autortfm_open_abort(void (*Work)(void* arg), void* Arg)
+extern "C" UE_AUTORTFM_AUTORTFM("RTFM_autortfm_open_abort") void autortfm_open_abort(void (*Work)(void* arg), void* Arg)
 {
 }
 
-extern "C" UE_AUTORTFM_NOAUTORTFM void* autortfm_did_allocate(void* Ptr, size_t Size)
+extern "C" UE_AUTORTFM_AUTORTFM("RTFM_autortfm_did_allocate") void* autortfm_did_allocate(void* Ptr, size_t Size)
 {
     return Ptr;
 }
 
-extern "C" UE_AUTORTFM_NOAUTORTFM void autortfm_check_consistency_assuming_no_races()
+extern "C" UE_AUTORTFM_AUTORTFM("RTFM_autortfm_check_consistency_assuming_no_races") void autortfm_check_consistency_assuming_no_races()
 {
     if (FContext::IsTransactional())
     {
@@ -259,88 +259,76 @@ extern "C" UE_AUTORTFM_NOAUTORTFM void autortfm_check_abi(void* const Ptr, const
 // Second Part - the same API exposed inside transactions. Note that we don't expose all of the API
 // to transactions! That's intentional. However, things like autortfm_defer_until_commit can be called
 // from an open nest in a transaction.
-UE_AUTORTFM_NOAUTORTFM bool RTFM_autortfm_is_transactional()
+extern "C" UE_AUTORTFM_NOAUTORTFM bool RTFM_autortfm_is_transactional()
 {
     return true;
 }
-UE_AUTORTFM_REGISTER_OPEN_FUNCTION(autortfm_is_transactional);
 
-UE_AUTORTFM_NOAUTORTFM bool RTFM_autortfm_is_closed()
+extern "C" UE_AUTORTFM_NOAUTORTFM bool RTFM_autortfm_is_closed()
 {
     return true;
 }
-UE_AUTORTFM_REGISTER_OPEN_FUNCTION(autortfm_is_closed);
 
-UE_AUTORTFM_NOAUTORTFM autortfm_result RTFM_autortfm_transact(void (*Work)(void* Arg), void* Arg)
+extern "C" UE_AUTORTFM_NOAUTORTFM autortfm_result RTFM_autortfm_transact(void (*Work)(void* Arg), void* Arg)
 {
 	FContext* Context = FContext::Get();
     return static_cast<autortfm_result>(Context->Transact(Work, Arg));
 }
-UE_AUTORTFM_REGISTER_OPEN_FUNCTION(autortfm_transact);
 
-UE_AUTORTFM_NOAUTORTFM autortfm_result RTFM_autortfm_transact_then_open(void (*Work)(void* Arg), void* Arg)
+extern "C" UE_AUTORTFM_NOAUTORTFM autortfm_result RTFM_autortfm_transact_then_open(void (*Work)(void* Arg), void* Arg)
 {
     return TransactThenOpenImpl(Work, Arg);
 }
-UE_AUTORTFM_REGISTER_OPEN_FUNCTION(autortfm_transact_then_open);
 
-UE_AUTORTFM_NOAUTORTFM void RTFM_autortfm_commit(void (*Work)(void* Arg), void* Arg)
+extern "C" UE_AUTORTFM_NOAUTORTFM void RTFM_autortfm_commit(void (*Work)(void* Arg), void* Arg)
 {
     autortfm_result Result = autortfm_transact(Work, Arg);
 	UE_CLOG(Result != autortfm_committed, LogAutoRTFM, Fatal, TEXT("Unexpected transaction result: %u."), Result);
 }
-UE_AUTORTFM_REGISTER_OPEN_FUNCTION(autortfm_commit);
 
-UE_AUTORTFM_NOAUTORTFM void RTFM_autortfm_abort()
+extern "C" UE_AUTORTFM_NOAUTORTFM void RTFM_autortfm_abort()
 {
 	FContext* Context = FContext::Get();
     Context->AbortByRequestAndThrow();
 }
-UE_AUTORTFM_REGISTER_OPEN_FUNCTION(autortfm_abort);
 
-UE_AUTORTFM_NOAUTORTFM void RTFM_autortfm_start_transaction()
+extern "C" UE_AUTORTFM_NOAUTORTFM void RTFM_autortfm_start_transaction()
 {
 	UE_LOG(LogAutoRTFM, Fatal, TEXT("The function `autortfm_start_transaction` was called from closed code."));
 }
-UE_AUTORTFM_REGISTER_OPEN_FUNCTION(autortfm_start_transaction);
 
-UE_AUTORTFM_NOAUTORTFM void RTFM_autortfm_commit_transaction()
+extern "C" UE_AUTORTFM_NOAUTORTFM void RTFM_autortfm_commit_transaction()
 {
 	UE_LOG(LogAutoRTFM, Fatal, TEXT("The function `RTFM_autortfm_commit_transaction` was called from closed code."));
 }
-UE_AUTORTFM_REGISTER_OPEN_FUNCTION(autortfm_commit_transaction);
 
-UE_AUTORTFM_NOAUTORTFM autortfm_result RTFM_autortfm_abort_transaction()
+extern "C" UE_AUTORTFM_NOAUTORTFM autortfm_result RTFM_autortfm_abort_transaction()
 {
 	UE_CLOG(!FContext::IsTransactional(), LogAutoRTFM, Fatal, TEXT("The function `autortfm_abort_transaction` was called from outside a transaction"));
 	FContext* Context = FContext::Get();
 	return static_cast<autortfm_result>(Context->AbortTransaction(true));
 }
-UE_AUTORTFM_REGISTER_OPEN_FUNCTION(autortfm_abort_transaction);
 
-UE_AUTORTFM_NOAUTORTFM void RTFM_autortfm_clear_transaction_status()
+extern "C" UE_AUTORTFM_NOAUTORTFM void RTFM_autortfm_clear_transaction_status()
 {
 	UE_LOG(LogAutoRTFM, Fatal, TEXT("The function `autortfm_clear_transaction_status` was called from closed code."));
 }
-UE_AUTORTFM_REGISTER_OPEN_FUNCTION(autortfm_clear_transaction_status);
 
-UE_AUTORTFM_NOAUTORTFM void RTFM_autortfm_abort_if_transactional()
+extern "C" UE_AUTORTFM_NOAUTORTFM void RTFM_autortfm_abort_if_transactional()
 {
     UE_LOG(LogAutoRTFM, Verbose, TEXT("The function `autortfm_abort_if_transactional` was called from inside a transaction."));
 	FContext* Context = FContext::Get();
     Context->AbortByRequestAndThrow();
 }
-UE_AUTORTFM_REGISTER_OPEN_FUNCTION(autortfm_abort_if_transactional);
 
-UE_AUTORTFM_NOAUTORTFM void RTFM_autortfm_abort_if_closed()
+extern "C" UE_AUTORTFM_NOAUTORTFM void RTFM_autortfm_abort_if_closed()
 {
 	UE_LOG(LogAutoRTFM, Verbose, TEXT("The function `autortfm_abort_if_closed` was called from closed inside a transaction."));
 	FContext* Context = FContext::Get();
     Context->AbortByRequestAndThrow();
 }
-UE_AUTORTFM_REGISTER_OPEN_FUNCTION(autortfm_abort_if_closed);
 
-UE_AUTORTFM_NOAUTORTFM void RTFM_autortfm_open(void (*Work)(void* Arg), void* Arg)
+extern "C" UE_AUTORTFM_NOAUTORTFM void RTFM_autortfm_open(void (*Work)(void* Arg), void* Arg)
 {
 	Work(Arg);
 
@@ -350,9 +338,8 @@ UE_AUTORTFM_NOAUTORTFM void RTFM_autortfm_open(void (*Work)(void* Arg), void* Ar
 		Context->Throw();
 	}
 }
-UE_AUTORTFM_REGISTER_OPEN_FUNCTION(autortfm_open);
 
-UE_AUTORTFM_NOAUTORTFM autortfm_status RTFM_autortfm_close(void (*Work)(void* Arg), void* Arg)
+extern "C" UE_AUTORTFM_NOAUTORTFM autortfm_status RTFM_autortfm_close(void (*Work)(void* Arg), void* Arg)
 {
     void (*WorkClone)(void* Arg) = FunctionMapLookup(Work, "RTFM_autortfm_close");
     if (WorkClone)
@@ -362,51 +349,44 @@ UE_AUTORTFM_NOAUTORTFM autortfm_status RTFM_autortfm_close(void (*Work)(void* Ar
 
 	return static_cast<autortfm_status>(FContext::Get()->GetStatus());
 }
-UE_AUTORTFM_REGISTER_OPEN_FUNCTION(autortfm_close);
 
 extern "C" UE_AUTORTFM_NOAUTORTFM void RTFM_autortfm_record_open_write(void*, size_t)
 {
 	UE_LOG(LogAutoRTFM, Fatal, TEXT("The function `autortfm_record_open_write` was called from closed code."));
 }
-UE_AUTORTFM_REGISTER_OPEN_FUNCTION(autortfm_record_open_write);
 
-UE_AUTORTFM_NOAUTORTFM void RTFM_OpenCommit(TFunction<void()>&& Work)
+extern "C" UE_AUTORTFM_NOAUTORTFM void RTFM_OpenCommit(TFunction<void()> && Work)
 {
 	FContext* Context = FContext::Get();
     ASSERT(Context->GetStatus() == EContextStatus::OnTrack);
     Context->GetCurrentTransaction()->DeferUntilCommit(MoveTemp(Work));
 }
-UE_AUTORTFM_REGISTER_OPEN_FUNCTION(OpenCommit);
 
-UE_AUTORTFM_NOAUTORTFM void RTFM_OpenAbort(TFunction<void()>&& Work)
+extern "C" UE_AUTORTFM_NOAUTORTFM void RTFM_OpenAbort(TFunction<void()>&& Work)
 {
 	FContext* Context = FContext::Get();
     ASSERT(Context->GetStatus() == EContextStatus::OnTrack);
     Context->GetCurrentTransaction()->DeferUntilAbort(MoveTemp(Work));
 }
-UE_AUTORTFM_REGISTER_OPEN_FUNCTION(OpenAbort);
 
 extern "C" UE_AUTORTFM_NOAUTORTFM void RTFM_autortfm_open_commit(void (*Work)(void* Arg), void* Arg)
 {
     RTFM_OpenCommit([Work, Arg] { Work(Arg); });
 }
-UE_AUTORTFM_REGISTER_OPEN_FUNCTION(autortfm_open_commit);
 
 extern "C" UE_AUTORTFM_NOAUTORTFM void RTFM_autortfm_open_abort(void (*Work)(void* arg), void* Arg)
 {
     RTFM_OpenAbort([Work, Arg] { Work(Arg); });
 }
-UE_AUTORTFM_REGISTER_OPEN_FUNCTION(autortfm_open_abort);
 
-UE_AUTORTFM_NOAUTORTFM void* RTFM_autortfm_did_allocate(void* Ptr, size_t Size)
+extern "C" UE_AUTORTFM_NOAUTORTFM void* RTFM_autortfm_did_allocate(void* Ptr, size_t Size)
 {
 	FContext* Context = FContext::Get();
     Context->DidAllocate(Ptr, Size);
     return Ptr;
 }
-UE_AUTORTFM_REGISTER_OPEN_FUNCTION(autortfm_did_allocate);
 
-UE_AUTORTFM_NOAUTORTFM void RTFM_autortfm_check_consistency_assuming_no_races()
+extern "C" UE_AUTORTFM_NOAUTORTFM void RTFM_autortfm_check_consistency_assuming_no_races()
 {
 }
 UE_AUTORTFM_REGISTER_OPEN_FUNCTION(autortfm_check_consistency_assuming_no_races);
diff --git a/Engine/Source/Runtime/Core/Public/AutoRTFM/AutoRTFM.h b/Engine/Source/Runtime/Core/Public/AutoRTFM/AutoRTFM.h
--- a/Engine/Source/Runtime/Core/Public/AutoRTFM/AutoRTFM.h
+++ b/Engine/Source/Runtime/Core/Public/AutoRTFM/AutoRTFM.h
@@ -15,8 +15,10 @@
 #endif
 
 #if UE_AUTORTFM
+#define UE_AUTORTFM_AUTORTFM(F) [[clang::autortfm(F), clang::noinline]]
 #define UE_AUTORTFM_NOAUTORTFM [[clang::noautortfm, clang::noinline]]
 #else
+#define UE_AUTORTFM_AUTORTFM(F)
 #define UE_AUTORTFM_NOAUTORTFM
 #endif
 
