Added a compilation error when checkf or ensureMsgf uses a const string as the condition.
This generally is caused by a programmer adding the printf string as the first parameter inadvertently.
Use != nullptr if your condition is intentional.

diff --git a/Engine/Source/Runtime/Core/Public/Misc/AssertionMacros.h b/Engine/Source/Runtime/Core/Public/Misc/AssertionMacros.h
--- a/Engine/Source/Runtime/Core/Public/Misc/AssertionMacros.h
+++ b/Engine/Source/Runtime/Core/Public/Misc/AssertionMacros.h
@@ -12,6 +12,7 @@
 #include "Templates/IsValidVariadicFunctionArg.h"
 #include "Traits/IsCharEncodingCompatibleWith.h"
 #include "Misc/VarArgs.h"
+#include "String/FormatStringSan.h"
 
 #if (DO_CHECK || DO_GUARD_SLOW || DO_ENSURE) && !PLATFORM_CPU_ARM_FAMILY
 	// We'll put all assert implementation code into a separate section in the linked
@@ -252,6 +253,7 @@ RetType FORCENOINLINE UE_DEBUG_SECTION DispatchCheckVerify(InnerType&& Inner, Ar
 
 	#define UE_CHECK_F_IMPL(expr, format, ...) \
 		{ \
+			static_assert(!UE::Core::Private::FormatStringSan::bIsAConstString<decltype(expr)>, "The checkf condition cannot be a const string. Use != nullptr if it's intentional"); \
 			if(UNLIKELY(!(expr))) \
 			{ \
 				if (FDebug::CheckVerifyFailedImpl(#expr, __FILE__, __LINE__, PLATFORM_RETURN_ADDRESS(), format, ##__VA_ARGS__)) \
@@ -373,6 +375,7 @@ RetType FORCENOINLINE UE_DEBUG_SECTION DispatchCheckVerify(InnerType&& Inner, Ar
 	#define UE_ENSURE_IMPL2(Capture, Always, InExpression, ...) \
 		(LIKELY(!!(InExpression)) || (DispatchCheckVerify<bool>([Capture] () UE_DEBUG_SECTION \
 		{ \
+			static_assert(!UE::Core::Private::FormatStringSan::bIsAConstString<decltype(InExpression)>, "The ensureMsgf condition cannot be a const string. Use != nullptr if it's intentional"); \
 			static bool bExecuted = false; \
 			FValidateArgsInternal(__VA_ARGS__); \
 			return CheckVerifyImpl(bExecuted, Always, __FILE__, __LINE__, PLATFORM_RETURN_ADDRESS(), #InExpression, ##__VA_ARGS__); \
diff --git a/Engine/Source/Runtime/Core/Public/String/FormatStringSan.h b/Engine/Source/Runtime/Core/Public/String/FormatStringSan.h
--- a/Engine/Source/Runtime/Core/Public/String/FormatStringSan.h
+++ b/Engine/Source/Runtime/Core/Public/String/FormatStringSan.h
@@ -20,6 +20,10 @@ namespace UE::Core::Private
 {
 	namespace FormatStringSan
 	{
+		// Returns true when the type is a const char*, const TCHAR*, const char[] or const TCHAR[]
+		template<typename T>
+		inline constexpr bool bIsAConstString = !(std::is_convertible_v<std::decay_t<T>, char*> || std::is_convertible_v<std::decay_t<T>, TCHAR*>) && (std::is_convertible_v<std::decay_t<T>, const char*> || std::is_convertible_v<std::decay_t<T>, const TCHAR*>);
+
 		enum EFormatStringSanStatus
 		{
 			StatusOk,
diff --git a/Engine/Source/Runtime/Core/Tests/String/FormatStringSanTest.cpp b/Engine/Source/Runtime/Core/Tests/String/FormatStringSanTest.cpp
--- a/Engine/Source/Runtime/Core/Tests/String/FormatStringSanTest.cpp
+++ b/Engine/Source/Runtime/Core/Tests/String/FormatStringSanTest.cpp
@@ -95,6 +95,109 @@ TEST_CASE_NAMED(FFormatStringValidatorTest, "Core::String::FormatStringSan", "[C
 	}
 }
 
+TEST_CASE_NAMED(FFormatStringConstStringValidationTest, "Core::String::FormatStringSan::ConstString", "[Core][String][FormatStringSan]")
+{
+	using namespace UE::Core::Private::FormatStringSan;
+
+	SECTION("Valid Const String Conditions")
+	{
+		STATIC_CHECK(bIsAConstString<decltype("Raw CString")>);
+		STATIC_CHECK(bIsAConstString<decltype(TEXT("Raw WString"))>);
+
+		{
+			const char Array[] = "CString";
+			STATIC_CHECK(bIsAConstString<decltype(Array)>);
+		}
+
+		{
+			const TCHAR Array[] = TEXT("WString");
+			STATIC_CHECK(bIsAConstString<decltype(Array)>);
+		}
+
+		{
+			const char* ConstPtr = "CString";
+			STATIC_CHECK(bIsAConstString<decltype(ConstPtr)>);
+		}
+
+		{
+			const TCHAR* ConstPtr = TEXT("WString");
+			STATIC_CHECK(bIsAConstString<decltype(ConstPtr)>);
+		}
+
+		{
+			const char* const ConstPtrConst = "CString";
+			STATIC_CHECK(bIsAConstString<decltype(ConstPtrConst)>);
+		}
+
+		{
+			const TCHAR* const ConstPtrConst = TEXT("WString");
+			STATIC_CHECK(bIsAConstString<decltype(ConstPtrConst)>);
+		}
+		
+		{
+			struct FImplicitConvertToChar
+			{
+				operator const char*() const { return (const char*)this; }
+			};
+			FImplicitConvertToChar ToChar;
+			STATIC_CHECK(bIsAConstString<decltype(ToChar)>);
+		}
+
+		{
+			struct FImplicitConvertToTChar
+			{
+				operator const TCHAR* () const { return (const TCHAR*)this; }
+			};
+			FImplicitConvertToTChar ToTChar;
+			STATIC_CHECK(bIsAConstString<decltype(ToTChar)>);
+		}
+	}
+
+	SECTION("Invalid Const String Conditions")
+	{
+		bool bBool = true;
+		STATIC_CHECK_FALSE(bIsAConstString<decltype(bBool)>);
+
+		{
+			char* Ptr = nullptr;
+			STATIC_CHECK_FALSE(bIsAConstString<decltype(Ptr)>);
+		}
+
+		{
+			TCHAR* Ptr = nullptr;
+			STATIC_CHECK_FALSE(bIsAConstString<decltype(Ptr)>);
+		}
+
+		{
+			char* const Ptr = nullptr;
+			STATIC_CHECK_FALSE(bIsAConstString<decltype(Ptr)>);
+		}
+
+		{
+			TCHAR* const Ptr = nullptr;
+			STATIC_CHECK_FALSE(bIsAConstString<decltype(Ptr)>);
+		}
+
+		{
+			struct FExplicitConvertToChar
+			{
+				explicit operator const char* () const { return (const char*)this; }
+			};
+			FExplicitConvertToChar ToChar;
+			STATIC_CHECK_FALSE(bIsAConstString<decltype(ToChar)>);
+		}
+
+		{
+			struct FExplicitConvertToTChar
+			{
+				explicit operator const TCHAR* () const { return (const TCHAR*)this; }
+			};
+			FExplicitConvertToTChar ToTChar;
+			STATIC_CHECK_FALSE(bIsAConstString<decltype(ToTChar)>);
+		}
+	}
+}
+
 } // namespace UE
 
 #endif // WITH_TESTS
