Fix autortfm/transactional failures on Linux

diff --git a/Engine/Source/Runtime/Core/Private/AutoRTFM/StdLibSupport.cpp b/Engine/Source/Runtime/Core/Private/AutoRTFM/StdLibSupport.cpp
--- a/Engine/Source/Runtime/Core/Private/AutoRTFM/StdLibSupport.cpp
+++ b/Engine/Source/Runtime/Core/Private/AutoRTFM/StdLibSupport.cpp
@@ -305,6 +305,10 @@ BOOL RTFM_TlsSetValue(DWORD dwTlsIndex, LPVOID lpTlsValue)
 UE_AUTORTFM_REGISTER_OPEN_FUNCTION(TlsSetValue);
 #endif // PLATFORM_WINDOWS
 
+#if PLATFORM_LINUX
+UE_AUTORTFM_REGISTER_SELF_FUNCTION(clock_gettime);
+#endif // PLATFORM_LINUX
+
 wchar_t* RTFM_wcsncpy(wchar_t* Dst, const wchar_t* Src, size_t Count)
 {
 	AutoRTFM::Unreachable();
diff --git a/Engine/Source/Runtime/Core/Private/HAL/ThreadingBase.cpp b/Engine/Source/Runtime/Core/Private/HAL/ThreadingBase.cpp
--- a/Engine/Source/Runtime/Core/Private/HAL/ThreadingBase.cpp
+++ b/Engine/Source/Runtime/Core/Private/HAL/ThreadingBase.cpp
@@ -1361,14 +1361,19 @@ FTlsAutoCleanup* FThreadSingletonInitializer::Get( TFunctionRef<FTlsAutoCleanup*
 			}
 		}
 	});
-	
-	FTlsAutoCleanup* ThreadSingleton = (FTlsAutoCleanup*)FPlatformTLS::GetTlsValue( TlsSlot );
-	if( !ThreadSingleton )
+
+	FTlsAutoCleanup* ThreadSingleton = nullptr;
+	UE_AUTORTFM_OPEN(
 	{
-		ThreadSingleton = CreateInstance();
-		ThreadSingleton->Register();
-		FPlatformTLS::SetTlsValue( TlsSlot, ThreadSingleton );
-	}
+		ThreadSingleton = (FTlsAutoCleanup*)FPlatformTLS::GetTlsValue( TlsSlot );
+		if( !ThreadSingleton )
+		{
+			// these are generally left open and only get cleaned up on thread exit so avoiding dealing with an OPENABORT here to clean this up
+			ThreadSingleton = CreateInstance();
+			ThreadSingleton->Register();
+			FPlatformTLS::SetTlsValue( TlsSlot, ThreadSingleton );
+		}
+	});
 	return ThreadSingleton;
 }
 
diff --git a/Engine/Source/Runtime/Core/Private/Unix/UnixPlatformMisc.cpp b/Engine/Source/Runtime/Core/Private/Unix/UnixPlatformMisc.cpp
--- a/Engine/Source/Runtime/Core/Private/Unix/UnixPlatformMisc.cpp
+++ b/Engine/Source/Runtime/Core/Private/Unix/UnixPlatformMisc.cpp
@@ -1395,35 +1395,37 @@ bool FUnixPlatformMisc::IsDebuggerPresent()
 	// If a process is tracing this one then TracerPid in /proc/self/status will
 	// be the id of the tracing process. Use SignalHandler safe functions
 
-	int StatusFile = open("/proc/self/status", O_RDONLY);
-	if (StatusFile == -1)
+	int StatusFile = -1;
+	bool bDebugging = false;
+	UE_AUTORTFM_OPEN(
 	{
-		// Failed - unknown debugger status.
-		return false;
-	}
+		StatusFile = open("/proc/self/status", O_RDONLY);
 
-	char Buffer[256];
-	ssize_t Length = read(StatusFile, Buffer, sizeof(Buffer));
+		if (StatusFile >= 0)
+		{
+			char Buffer[256];
+			ssize_t Length = read(StatusFile, Buffer, sizeof(Buffer));
 
-	bool bDebugging = false;
-	const char* TracerString = "TracerPid:\t";
-	const ssize_t LenTracerString = strlen(TracerString);
-	int i = 0;
+			const char* TracerString = "TracerPid:\t";
+			const ssize_t LenTracerString = strlen(TracerString);
+			int i = 0;
 
-	while((Length - i) > LenTracerString)
-	{
-		// TracerPid is found
-		if (strncmp(&Buffer[i], TracerString, LenTracerString) == 0)
-		{
-			// 0 if no process is tracing.
-			bDebugging = Buffer[i + LenTracerString] != '0';
-			break;
-		}
+			while((Length - i) > LenTracerString)
+			{
+				// TracerPid is found
+				if (strncmp(&Buffer[i], TracerString, LenTracerString) == 0)
+				{
+					// 0 if no process is tracing.
+					bDebugging = Buffer[i + LenTracerString] != '0';
+					break;
+				}
 
-		++i;
-	}
+				++i;
+			}
 
-	close(StatusFile);
+			close(StatusFile);
+		}
+	});
 	return bDebugging;
 }
 #endif // !UE_BUILD_SHIPPING
