CsvProfiler - Add end capture timestamp to the csv metadata.

This allows easily calculating capture duration by checking the delta between StartTimestamp and the new EndTimestamp. This is much faster than trying to calculate the duration by looking at all the frame times.

diff --git a/Engine/Source/Runtime/Core/Private/ProfilingDebugging/CsvProfiler.cpp b/Engine/Source/Runtime/Core/Private/ProfilingDebugging/CsvProfiler.cpp
--- a/Engine/Source/Runtime/Core/Private/ProfilingDebugging/CsvProfiler.cpp
+++ b/Engine/Source/Runtime/Core/Private/ProfilingDebugging/CsvProfiler.cpp
@@ -3284,6 +3284,8 @@ TSharedFuture<FString> FCsvProfiler::EndCapture(FGraphEventRef EventToSignal)
 	// Fire before we copy the metadata so it gives other systems a chance to write any final information.
 	OnCSVProfileEndRequestedDelegate.Broadcast();
 
+	SetNonPersistentMetadata(TEXT("EndTimestamp"), *FString::Printf(TEXT("%lld"), FDateTime::UtcNow().ToUnixTimestamp()));
+
 	TPromise<FString>* Completion = new TPromise<FString>([EventToSignal]()
 	{
 		if (EventToSignal)
