Add ECVF_Preview to cvar r.Mobile.AntiAliasing so it is previewable

diff --git a/Engine/Source/Runtime/Core/Private/HAL/ConsoleManager.cpp b/Engine/Source/Runtime/Core/Private/HAL/ConsoleManager.cpp
--- a/Engine/Source/Runtime/Core/Private/HAL/ConsoleManager.cpp
+++ b/Engine/Source/Runtime/Core/Private/HAL/ConsoleManager.cpp
@@ -3489,7 +3489,7 @@ static TAutoConsoleVariable<int32> CVarMobileDefaultAntiAliasing(
 		 " 1: FXAA (default, faster than TemporalAA but much more shimmering for non static cases)\n"
 		 " 2: TemporalAA(it will fallback to FXAA if SupportsGen4TAA is disabled) \n"
 		 " 3: MSAA"),
-	ECVF_RenderThreadSafe);
+	ECVF_RenderThreadSafe | ECVF_Preview);
 
 static TAutoConsoleVariable<int32> CVarMobileSupportsGen4TAA(
 	TEXT("r.Mobile.SupportsGen4TAA"),
