Fixed FString's private namespacing.

diff --git a/Engine/Source/Runtime/Core/Private/Containers/String.cpp b/Engine/Source/Runtime/Core/Private/Containers/String.cpp
--- a/Engine/Source/Runtime/Core/Private/Containers/String.cpp
+++ b/Engine/Source/Runtime/Core/Private/Containers/String.cpp
@@ -27,7 +27,7 @@
 /* FString implementation
  *****************************************************************************/
 
-namespace UE::Core::String::Private
+namespace UE::Core::Private
 {
 	struct FCompareCharsCaseSensitive
 	{
@@ -126,125 +126,120 @@ namespace UE::Core::String::Private
 		}
 		return false;
 	}
-}
-
-template<typename CharType>
-void AppendCharacters(TArray<TCHAR>& Out, const CharType* Str, int32 Count)
-{
-	check(Count >= 0);
 
-	if (!Count)
+	template<typename CharType>
+	void AppendCharacters(TArray<TCHAR>& Out, const CharType* Str, int32 Count)
 	{
-		return;
-	}
+		check(Count >= 0);
 
-	checkSlow(Str);
+		if (!Count)
+		{
+			return;
+		}
+
+		checkSlow(Str);
 
-	int32 OldEnd = Out.Num();
+		int32 OldEnd = Out.Num();
 	
-	// Try to reserve enough space by guessing that the new length will be the same as the input length.
-	// Include an extra gap for a null terminator if we don't already have a string allocated
-	Out.AddUninitialized(Count + (OldEnd ? 0 : 1));
-	OldEnd -= OldEnd ? 1 : 0;
+		// Try to reserve enough space by guessing that the new length will be the same as the input length.
+		// Include an extra gap for a null terminator if we don't already have a string allocated
+		Out.AddUninitialized(Count + (OldEnd ? 0 : 1));
+		OldEnd -= OldEnd ? 1 : 0;
 
-	TCHAR* Dest = Out.GetData() + OldEnd;
+		TCHAR* Dest = Out.GetData() + OldEnd;
 
-	// Try copying characters to end of string, overwriting null terminator if we already have one
-	TCHAR* NewEnd = FPlatformString::Convert(Dest, Count, Str, Count);
-	if (!NewEnd)
-	{
-		// If that failed, it will have meant that conversion likely contained multi-code unit characters
-		// and so the buffer wasn't long enough, so calculate it properly.
-		int32 Length = FPlatformString::ConvertedLength<TCHAR>(Str, Count);
+		// Try copying characters to end of string, overwriting null terminator if we already have one
+		TCHAR* NewEnd = FPlatformString::Convert(Dest, Count, Str, Count);
+		if (!NewEnd)
+		{
+			// If that failed, it will have meant that conversion likely contained multi-code unit characters
+			// and so the buffer wasn't long enough, so calculate it properly.
+			int32 Length = FPlatformString::ConvertedLength<TCHAR>(Str, Count);
 
-		// Add the extra bytes that we need
-		Out.AddUninitialized(Length - Count);
+			// Add the extra bytes that we need
+			Out.AddUninitialized(Length - Count);
 
-		// Restablish destination pointer in case a realloc happened
-		Dest = Out.GetData() + OldEnd;
+			// Restablish destination pointer in case a realloc happened
+			Dest = Out.GetData() + OldEnd;
 
-		NewEnd = FPlatformString::Convert(Dest, Length, Str, Count);
-		checkSlow(NewEnd);
-	}
-	else
-	{
-		int32 NewEndIndex = (int32)(NewEnd - Dest);
-		if (NewEndIndex < Count)
+			NewEnd = FPlatformString::Convert(Dest, Length, Str, Count);
+			checkSlow(NewEnd);
+		}
+		else
 		{
-			Out.SetNumUninitialized(OldEnd + NewEndIndex + 1, /*bAllowShrinking=*/false);
+			int32 NewEndIndex = (int32)(NewEnd - Dest);
+			if (NewEndIndex < Count)
+			{
+				Out.SetNumUninitialized(OldEnd + NewEndIndex + 1, /*bAllowShrinking=*/false);
+			}
 		}
-	}
 
-	// (Re-)establish the null terminator
-	*NewEnd = '\0';
-}
-
-namespace UE::String::Private
-{
+		// (Re-)establish the null terminator
+		*NewEnd = '\0';
+	}
 
-template<typename CharType>
-FORCEINLINE void ConstructFromCString(/* Out */ TArray<TCHAR>& Data, const CharType* Src)
-{
-	if (Src && *Src)
+	template<typename CharType>
+	FORCEINLINE void ConstructFromCString(/* Out */ TArray<TCHAR>& Data, const CharType* Src)
 	{
-		int32 SrcLen  = TCString<CharType>::Strlen(Src) + 1;
-		int32 DestLen = FPlatformString::ConvertedLength<TCHAR>(Src, SrcLen);
-		Data.Reserve(DestLen);
-		Data.AddUninitialized(DestLen);
+		if (Src && *Src)
+		{
+			int32 SrcLen  = TCString<CharType>::Strlen(Src) + 1;
+			int32 DestLen = FPlatformString::ConvertedLength<TCHAR>(Src, SrcLen);
+			Data.Reserve(DestLen);
+			Data.AddUninitialized(DestLen);
 
-		FPlatformString::Convert(Data.GetData(), DestLen, Src, SrcLen);
+			FPlatformString::Convert(Data.GetData(), DestLen, Src, SrcLen);
+		}
 	}
-}
 
-template<typename CharType>
-FORCEINLINE void ConstructWithLength(/* Out */ TArray<TCHAR>& Data, int32 InCount, const CharType* InSrc)
-{
-	if (InSrc)
+	template<typename CharType>
+	FORCEINLINE void ConstructWithLength(/* Out */ TArray<TCHAR>& Data, int32 InCount, const CharType* InSrc)
 	{
-		int32 DestLen = FPlatformString::ConvertedLength<TCHAR>(InSrc, InCount);
-		if (DestLen > 0 && *InSrc)
+		if (InSrc)
 		{
-			Data.Reserve(DestLen + 1);
-			Data.AddUninitialized(DestLen + 1);
+			int32 DestLen = FPlatformString::ConvertedLength<TCHAR>(InSrc, InCount);
+			if (DestLen > 0 && *InSrc)
+			{
+				Data.Reserve(DestLen + 1);
+				Data.AddUninitialized(DestLen + 1);
 
-			FPlatformString::Convert(Data.GetData(), DestLen, InSrc, InCount);
-			*(Data.GetData() + Data.Num() - 1) = TEXT('\0');
+				FPlatformString::Convert(Data.GetData(), DestLen, InSrc, InCount);
+				*(Data.GetData() + Data.Num() - 1) = TEXT('\0');
+			}
 		}
 	}
-}
 
-template<typename CharType>
-FORCEINLINE void ConstructWithSlack(/* Out */ TArray<TCHAR>& Data, const CharType* Src, int32 ExtraSlack)
-{
-	if (Src && *Src)
+	template<typename CharType>
+	FORCEINLINE void ConstructWithSlack(/* Out */ TArray<TCHAR>& Data, const CharType* Src, int32 ExtraSlack)
 	{
-		int32 SrcLen = TCString<CharType>::Strlen(Src) + 1;
-		int32 DestLen = FPlatformString::ConvertedLength<TCHAR>(Src, SrcLen);
-		Data.Reserve(DestLen + ExtraSlack);
-		Data.AddUninitialized(DestLen);
+		if (Src && *Src)
+		{
+			int32 SrcLen = TCString<CharType>::Strlen(Src) + 1;
+			int32 DestLen = FPlatformString::ConvertedLength<TCHAR>(Src, SrcLen);
+			Data.Reserve(DestLen + ExtraSlack);
+			Data.AddUninitialized(DestLen);
 
-		FPlatformString::Convert(Data.GetData(), DestLen, Src, SrcLen);
-	}
-	else if (ExtraSlack > 0)
-	{
-		Data.Reserve(ExtraSlack + 1); 
+			FPlatformString::Convert(Data.GetData(), DestLen, Src, SrcLen);
+		}
+		else if (ExtraSlack > 0)
+		{
+			Data.Reserve(ExtraSlack + 1); 
+		}
 	}
-}
-
-} // namespace UE::String::Private
-
-FString::FString(const ANSICHAR* Str)								{ UE::String::Private::ConstructFromCString(Data, Str); }
-FString::FString(const WIDECHAR* Str)								{ UE::String::Private::ConstructFromCString(Data, Str); }
-FString::FString(const UTF8CHAR* Str)								{ UE::String::Private::ConstructFromCString(Data, Str); }
-FString::FString(const UCS2CHAR* Str)								{ UE::String::Private::ConstructFromCString(Data, Str); }
-FString::FString(int32 Len, const ANSICHAR* Str)					{ UE::String::Private::ConstructWithLength(Data, Len, Str); }
-FString::FString(int32 Len, const WIDECHAR* Str)					{ UE::String::Private::ConstructWithLength(Data, Len, Str); }
-FString::FString(int32 Len, const UTF8CHAR* Str)					{ UE::String::Private::ConstructWithLength(Data, Len, Str); }
-FString::FString(int32 Len, const UCS2CHAR* Str)					{ UE::String::Private::ConstructWithLength(Data, Len, Str); }
-FString::FString(const ANSICHAR* Str, int32 ExtraSlack)				{ UE::String::Private::ConstructWithSlack(Data, Str, ExtraSlack); }
-FString::FString(const WIDECHAR* Str, int32 ExtraSlack)				{ UE::String::Private::ConstructWithSlack(Data, Str, ExtraSlack); }
-FString::FString(const UTF8CHAR* Str, int32 ExtraSlack)				{ UE::String::Private::ConstructWithSlack(Data, Str, ExtraSlack); }
-FString::FString(const UCS2CHAR* Str, int32 ExtraSlack)				{ UE::String::Private::ConstructWithSlack(Data, Str, ExtraSlack); }
+} // namespace UE::Core::Private
+
+FString::FString(const ANSICHAR* Str)								{ UE::Core::Private::ConstructFromCString(Data, Str); }
+FString::FString(const WIDECHAR* Str)								{ UE::Core::Private::ConstructFromCString(Data, Str); }
+FString::FString(const UTF8CHAR* Str)								{ UE::Core::Private::ConstructFromCString(Data, Str); }
+FString::FString(const UCS2CHAR* Str)								{ UE::Core::Private::ConstructFromCString(Data, Str); }
+FString::FString(int32 Len, const ANSICHAR* Str)					{ UE::Core::Private::ConstructWithLength(Data, Len, Str); }
+FString::FString(int32 Len, const WIDECHAR* Str)					{ UE::Core::Private::ConstructWithLength(Data, Len, Str); }
+FString::FString(int32 Len, const UTF8CHAR* Str)					{ UE::Core::Private::ConstructWithLength(Data, Len, Str); }
+FString::FString(int32 Len, const UCS2CHAR* Str)					{ UE::Core::Private::ConstructWithLength(Data, Len, Str); }
+FString::FString(const ANSICHAR* Str, int32 ExtraSlack)				{ UE::Core::Private::ConstructWithSlack(Data, Str, ExtraSlack); }
+FString::FString(const WIDECHAR* Str, int32 ExtraSlack)				{ UE::Core::Private::ConstructWithSlack(Data, Str, ExtraSlack); }
+FString::FString(const UTF8CHAR* Str, int32 ExtraSlack)				{ UE::Core::Private::ConstructWithSlack(Data, Str, ExtraSlack); }
+FString::FString(const UCS2CHAR* Str, int32 ExtraSlack)				{ UE::Core::Private::ConstructWithSlack(Data, Str, ExtraSlack); }
 
 FString& FString::operator=( const TCHAR* Other )
 {
@@ -361,25 +356,25 @@ FString& FString::AppendChar(TCHAR InChar)
 void FString::AppendChars(const ANSICHAR* Str, int32 Count)
 {
 	CheckInvariants();
-	AppendCharacters(Data, Str, Count);
+	UE::Core::Private::AppendCharacters(Data, Str, Count);
 }
 
 void FString::AppendChars(const WIDECHAR* Str, int32 Count)
 {
 	CheckInvariants();
-	AppendCharacters(Data, Str, Count);
+	UE::Core::Private::AppendCharacters(Data, Str, Count);
 }
 
 void FString::AppendChars(const UCS2CHAR* Str, int32 Count)
 {
 	CheckInvariants();
-	AppendCharacters(Data, Str, Count);
+	UE::Core::Private::AppendCharacters(Data, Str, Count);
 }
 
 void FString::AppendChars(const UTF8CHAR* Str, int32 Count)
 {
 	CheckInvariants();
-	AppendCharacters(Data, Str, Count);
+	UE::Core::Private::AppendCharacters(Data, Str, Count);
 }
 
 void FString::TrimToNullTerminator()
@@ -1301,11 +1296,11 @@ bool FString::MatchesWildcard(const TCHAR* InWildcard, int32 InWildcardLen, ESea
 
 	if (SearchCase == ESearchCase::CaseSensitive)
 	{
-		return UE::Core::String::Private::MatchesWildcardRecursive<UE::Core::String::Private::FCompareCharsCaseSensitive>(Target, TargetLength, InWildcard, InWildcardLen);
+		return UE::Core::Private::MatchesWildcardRecursive<UE::Core::Private::FCompareCharsCaseSensitive>(Target, TargetLength, InWildcard, InWildcardLen);
 	}
 	else
 	{
-		return UE::Core::String::Private::MatchesWildcardRecursive<UE::Core::String::Private::FCompareCharsCaseInsensitive>(Target, TargetLength, InWildcard, InWildcardLen);
+		return UE::Core::Private::MatchesWildcardRecursive<UE::Core::Private::FCompareCharsCaseInsensitive>(Target, TargetLength, InWildcard, InWildcardLen);
 	}
 }
 
