Added angular speed units (deg/s, rad/s) for display in editor and conversion

diff --git a/Engine/Config/BaseEditor.ini b/Engine/Config/BaseEditor.ini
--- a/Engine/Config/BaseEditor.ini
+++ b/Engine/Config/BaseEditor.ini
@@ -482,6 +482,7 @@ bDisplayUnits=True
 +TimeUnits=EUnit::Minutes
 AngleUnits=EUnit::Degrees
 SpeedUnits=EUnit::MetersPerSecond
+AngularSpeedUnits=EUnit::DegreesPerSecond
 TemperatureUnits=EUnit::Celsius
 ForceUnits=EUnit::Newtons
 TorqueUnits=EUnit::NewtonMeters
diff --git a/Engine/Programs/UnrealHeaderTool/Config/DefaultEngine.ini b/Engine/Programs/UnrealHeaderTool/Config/DefaultEngine.ini
--- a/Engine/Programs/UnrealHeaderTool/Config/DefaultEngine.ini
+++ b/Engine/Programs/UnrealHeaderTool/Config/DefaultEngine.ini
@@ -66,6 +66,10 @@ gc.MaxObjectsInProgram=262144
 +Units="MilesPerHour"
 +Units="mi/h"
 +Units="mph"
++Units="DegreesPerSecond"
++Units="deg/s"
++Units="RadiansPerSecond"
++Units="rad/s"
 +Units="Celsius"
 +Units="C"
 +Units="degC"
diff --git a/Engine/Source/Editor/UnrealEd/Private/Settings/EditorProjectSettings.cpp b/Engine/Source/Editor/UnrealEd/Private/Settings/EditorProjectSettings.cpp
--- a/Engine/Source/Editor/UnrealEd/Private/Settings/EditorProjectSettings.cpp
+++ b/Engine/Source/Editor/UnrealEd/Private/Settings/EditorProjectSettings.cpp
@@ -105,6 +105,10 @@ void UEditorProjectAppearanceSettings::PostEditChangeProperty( struct FPropertyC
 	{
 		Settings.SetDisplayUnits(EUnitType::Speed, SpeedUnits);
 	}
+	else if (Name == GET_MEMBER_NAME_CHECKED(UEditorProjectAppearanceSettings, AngularSpeedUnits))
+	{
+		Settings.SetDisplayUnits(EUnitType::AngularSpeed, AngularSpeedUnits);
+	}
 	else if (Name == GET_MEMBER_NAME_CHECKED(UEditorProjectAppearanceSettings, TemperatureUnits))
 	{
 		Settings.SetDisplayUnits(EUnitType::Temperature, TemperatureUnits);
diff --git a/Engine/Source/Editor/UnrealEd/Public/Settings/EditorProjectSettings.h b/Engine/Source/Editor/UnrealEd/Public/Settings/EditorProjectSettings.h
--- a/Engine/Source/Editor/UnrealEd/Public/Settings/EditorProjectSettings.h
+++ b/Engine/Source/Editor/UnrealEd/Public/Settings/EditorProjectSettings.h
@@ -89,6 +89,9 @@ public:
 	UPROPERTY(EditAnywhere, config, Category=Units, AdvancedDisplay, meta=(DisplayName="Speed/Velocity", Tooltip="Choose the units in which to display speeds and velocities.", ValidEnumValues="CentimetersPerSecond, MetersPerSecond, KilometersPerHour, MilesPerHour"))
 	EUnit SpeedUnits;
 
+	UPROPERTY(EditAnywhere, config, Category=Units, AdvancedDisplay, meta=(DisplayName="Angular Speed", Tooltip="Choose the units in which to display angular speeds.", ValidEnumValues="DegreesPerSecond, RadiansPerSecond"))
+	EUnit AngularSpeedUnits;
+
 	UPROPERTY(EditAnywhere, config, Category=Units, AdvancedDisplay, meta=(DisplayName="Temperature", Tooltip="Choose the units in which to display temperatures.", ValidEnumValues="Celsius, Farenheit, Kelvin"))
 	EUnit TemperatureUnits;
 
diff --git a/Engine/Source/Runtime/Core/Private/Math/UnitConversion.cpp b/Engine/Source/Runtime/Core/Private/Math/UnitConversion.cpp
--- a/Engine/Source/Runtime/Core/Private/Math/UnitConversion.cpp
+++ b/Engine/Source/Runtime/Core/Private/Math/UnitConversion.cpp
@@ -47,6 +47,9 @@ FParseCandidate ParseCandidates[] = {
 	{ TEXT("MetersPerSecond"),		EUnit::MetersPerSecond },		{ TEXT("m/s"),		EUnit::MetersPerSecond },
 	{ TEXT("KilometersPerHour"),	EUnit::KilometersPerHour },		{ TEXT("km/h"),		EUnit::KilometersPerHour },		{ TEXT("kmph"),		EUnit::KilometersPerHour },
 	{ TEXT("MilesPerHour"),			EUnit::MilesPerHour },			{ TEXT("mi/h"),		EUnit::MilesPerHour },			{ TEXT("mph"),		EUnit::MilesPerHour },
+
+	{ TEXT("DegreesPerSecond"),		EUnit::DegreesPerSecond },		{ TEXT("deg/s"),	EUnit::DegreesPerSecond },
+	{ TEXT("RadiansPerSecond"),		EUnit::RadiansPerSecond },		{ TEXT("rad/s"),	EUnit::RadiansPerSecond },
 		
 	{ TEXT("Celsius"),				EUnit::Celsius },				{ TEXT("C"),		EUnit::Celsius },				{ TEXT("degC"),		EUnit::Celsius },			{ TEXT("\u00B0C"),		EUnit::Celsius },
 	{ TEXT("Farenheit"),			EUnit::Farenheit },				{ TEXT("F"),		EUnit::Farenheit },				{ TEXT("degF"),		EUnit::Farenheit },			{ TEXT("\u00B0F"),		EUnit::Farenheit },
@@ -124,6 +127,8 @@ const TCHAR* const DisplayStrings[] = {
 
 	TEXT("cm/s"), TEXT("m/s"), TEXT("km/h"), TEXT("mi/h"),
 
+	TEXT("deg/s"), TEXT("rad/s"),
+
 	TEXT("\u00B0C"), TEXT("\u00B0F"), TEXT("K"),
 
 	TEXT("\u00B5g"), TEXT("mg"), TEXT("g"), TEXT("kg"), TEXT("t"),
@@ -161,6 +166,8 @@ const EUnitType UnitTypes[] = {
 
 	EUnitType::Speed,		EUnitType::Speed,		EUnitType::Speed, 		EUnitType::Speed,
 
+	EUnitType::AngularSpeed, EUnitType::AngularSpeed,
+
 	EUnitType::Temperature,	EUnitType::Temperature,	EUnitType::Temperature,
 
 	EUnitType::Mass,		EUnitType::Mass,		EUnitType::Mass,		EUnitType::Mass,		EUnitType::Mass,
@@ -451,6 +458,7 @@ FUnitSettings::FUnitSettings()
 	DisplayUnits[(uint8)EUnitType::Distance].Add(EUnit::Centimeters);
 	DisplayUnits[(uint8)EUnitType::Angle].Add(EUnit::Degrees);
 	DisplayUnits[(uint8)EUnitType::Speed].Add(EUnit::MetersPerSecond);
+	DisplayUnits[(uint8)EUnitType::AngularSpeed].Add(EUnit::DegreesPerSecond);
 	DisplayUnits[(uint8)EUnitType::Temperature].Add(EUnit::Celsius);
 	DisplayUnits[(uint8)EUnitType::Mass].Add(EUnit::Kilograms);
 	DisplayUnits[(uint8)EUnitType::Density].Add(EUnit::GramsPerCubicCentimeter);
@@ -618,6 +626,17 @@ namespace UnitConversion
 		}
 	}
 
+
+	double AngularSpeedUnificationFactor(EUnit From)
+	{
+		// Convert to degrees/second
+		switch (From)
+		{
+			case EUnit::RadiansPerSecond:	return (180.0 / UE_DOUBLE_PI);
+			default: 						return 1.0;
+		}
+	}
+
 	double MassUnificationFactor(EUnit From)
 	{
 		double Factor = 1;
diff --git a/Engine/Source/Runtime/Core/Public/Math/UnitConversion.h b/Engine/Source/Runtime/Core/Public/Math/UnitConversion.h
--- a/Engine/Source/Runtime/Core/Public/Math/UnitConversion.h
+++ b/Engine/Source/Runtime/Core/Public/Math/UnitConversion.h
@@ -28,6 +28,9 @@ enum class EUnit : uint8
 	/** Speed units */
 	CentimetersPerSecond, MetersPerSecond, KilometersPerHour, MilesPerHour,
 
+	/** Angular speed units */
+	DegreesPerSecond, RadiansPerSecond,
+
 	/** Temperature units */
 	Celsius, Farenheit, Kelvin,
 
@@ -72,7 +75,7 @@ enum class EUnit : uint8
 /** Enumeration that specifies particular classes of unit */
 enum class EUnitType
 {
-	Distance, Angle, Speed, Temperature, Mass, Density, Force, Torque, Frequency, DataSize, LuminousFlux, LuminousIntensity, Illuminance, Luminance, Time, PixelDensity, Multipliers, ExposureValue, Stress,
+	Distance, Angle, Speed, AngularSpeed, Temperature, Mass, Density, Force, Torque, Frequency, DataSize, LuminousFlux, LuminousIntensity, Illuminance, Luminance, Time, PixelDensity, Multipliers, ExposureValue, Stress,
 
 	// Symbolic entry - do not use directly
 	NumberOf,
diff --git a/Engine/Source/Runtime/Core/Public/Math/UnitConversion.inl b/Engine/Source/Runtime/Core/Public/Math/UnitConversion.inl
--- a/Engine/Source/Runtime/Core/Public/Math/UnitConversion.inl
+++ b/Engine/Source/Runtime/Core/Public/Math/UnitConversion.inl
@@ -24,6 +24,8 @@ namespace UnitConversion
 	CORE_API double AngleUnificationFactor(EUnit From);
 	/** Find the common quantization factor for the specified speed unit. Quantizes to km/h. */
 	CORE_API double SpeedUnificationFactor(EUnit From);
+	/** Find the common quantization factor for the specified angular speed unit. Quantizes to DegreesPerSecond. */
+	CORE_API double AngularSpeedUnificationFactor(EUnit From);
 	/** Find the common quantization factor for the specified temperature unit. Quantizes to Kelvin. */
 	CORE_API double TemperatureUnificationFactor(EUnit From);
 	/** Find the common quantization factor for the specified mass unit. Quantizes to Grams. */
@@ -85,6 +87,7 @@ T FUnitConversion::Convert(T InValue, EUnit From, EUnit To)
 		case EUnitType::Distance:			return InValue * DistanceUnificationFactor(From)		* (1.0 / DistanceUnificationFactor(To));
 		case EUnitType::Angle:				return InValue * AngleUnificationFactor(From) 			* (1.0 / AngleUnificationFactor(To));
 		case EUnitType::Speed:				return InValue * SpeedUnificationFactor(From) 			* (1.0 / SpeedUnificationFactor(To));
+		case EUnitType::AngularSpeed:		return InValue * AngularSpeedUnificationFactor(From)	* (1.0 / AngularSpeedUnificationFactor(To));
 		case EUnitType::Mass:				return InValue * MassUnificationFactor(From) 			* (1.0 / MassUnificationFactor(To));
 		case EUnitType::Density:			return InValue * DensityUnificationFactor(From)			* (1.0 / DensityUnificationFactor(To));
 		case EUnitType::Force:				return InValue * ForceUnificationFactor(From) 			* (1.0 / ForceUnificationFactor(To));
diff --git a/Engine/Source/Runtime/Core/Tests/Math/UnitConversionTests.cpp b/Engine/Source/Runtime/Core/Tests/Math/UnitConversionTests.cpp
--- a/Engine/Source/Runtime/Core/Tests/Math/UnitConversionTests.cpp
+++ b/Engine/Source/Runtime/Core/Tests/Math/UnitConversionTests.cpp
@@ -44,6 +44,9 @@ TEST_CASE_NAMED(FUnitUnitTests, "System::Core::Math::Unit Conversion", "[Applica
 		{ 1,		0.6214,		1e-4,	EUnit::KilometersPerHour,							EUnit::MilesPerHour									},
 		{ 15,		6.7056,		1e-4,	EUnit::MilesPerHour,								EUnit::MetersPerSecond								},
 
+		{ 90.0,			UE_DOUBLE_HALF_PI,	1e-3,	EUnit::DegreesPerSecond, 				EUnit::RadiansPerSecond								},
+		{ UE_DOUBLE_PI,	180.0,				1e-3,	EUnit::RadiansPerSecond, 				EUnit::DegreesPerSecond								},
+
 		{ 100, 		212,		0.1,	EUnit::Celsius,										EUnit::Farenheit									},
 		{ 400, 		477.594,	1e-3,	EUnit::Farenheit,									EUnit::Kelvin										},
 		{ 72, 		-201.15,	0.01,	EUnit::Kelvin,										EUnit::Celsius										},
diff --git a/Engine/Source/Runtime/CoreUObject/Public/UObject/NoExportTypes.h b/Engine/Source/Runtime/CoreUObject/Public/UObject/NoExportTypes.h
--- a/Engine/Source/Runtime/CoreUObject/Public/UObject/NoExportTypes.h
+++ b/Engine/Source/Runtime/CoreUObject/Public/UObject/NoExportTypes.h
@@ -367,7 +367,10 @@ enum class EUnit : uint8
 	
 	/** Speed units */
 	CentimetersPerSecond, MetersPerSecond, KilometersPerHour, MilesPerHour,
-	
+
+	/** Angular speed units */
+	DegreesPerSecond, RadiansPerSecond,
+
 	/** Temperature units */
 	Celsius, Farenheit, Kelvin,
 	
