autortfm: catch a few cases of invoking Verse outside of an autortfm transaction

diff --git a/Engine/Source/Runtime/Core/Public/Delegates/Delegate.h b/Engine/Source/Runtime/Core/Public/Delegates/Delegate.h
--- a/Engine/Source/Runtime/Core/Public/Delegates/Delegate.h
+++ b/Engine/Source/Runtime/Core/Public/Delegates/Delegate.h
@@ -9,6 +9,7 @@
 #include "UObject/WeakObjectPtrTemplates.h"
 #include "Delegates/MulticastDelegateBase.h" // IWYU pragma: export
 #include "Delegates/IntegerSequence.h" // IWYU pragma: export
+#include "AutoRTFM/AutoRTFM.h"
 
 /**
  *  C++ DELEGATES
@@ -370,8 +371,13 @@ class DynamicMulticastDelegateClassName : public TBaseDynamicMulticastDelegate<F
 		{
 			static FName Get()
 			{
-				static FName Result = Create();
-				return Result;
+				FName* Result = nullptr;
+				UE_AUTORTFM_OPEN(
+				{
+					static FName StaticResult = Create();
+					Result = &StaticResult;
+				});
+				return *Result;
 			}
 
 		private:
