Fix concurrency test in AutoRTFM.

diff --git a/Engine/Source/Runtime/Core/Private/AutoRTFM/API.cpp b/Engine/Source/Runtime/Core/Private/AutoRTFM/API.cpp
--- a/Engine/Source/Runtime/Core/Private/AutoRTFM/API.cpp
+++ b/Engine/Source/Runtime/Core/Private/AutoRTFM/API.cpp
@@ -1,7 +1,7 @@
 // Copyright Epic Games, Inc. All Rights Reserved.
 
-
 #include "AutoRTFM/AutoRTFM.h"
+#include "CoreGlobals.h"
 #include "HAL/IConsoleManager.h"
 
 #if UE_AUTORTFM
@@ -24,6 +24,7 @@ static constexpr bool GAutoRTFMRuntimeEnabled = false;
 #include "ContextStatus.h"
 #include "FunctionMapInlines.h"
 #include "TransactionInlines.h"
+#include "Toggles.h"
 #include "Utils.h"
 
 #include "Templates/Tuple.h"
@@ -152,7 +153,7 @@ extern "C" UE_AUTORTFM_AUTORTFM("RTFM_autortfm_close") autortfm_status autortfm_
 	{
 		UE_CLOG(!FContext::IsTransactional(), LogAutoRTFM, Fatal, TEXT("Close called from an outside a transaction."));
 
-		FContext* Context = FContext::Get();
+		FContext* const Context = FContext::Get();
 		void (*WorkClone)(void* Arg) = FunctionMapLookup(Work, "autortfm_close");
 		if (WorkClone)
 		{
@@ -169,8 +170,10 @@ extern "C" UE_AUTORTFM_AUTORTFM("RTFM_autortfm_close") autortfm_status autortfm_
 
 extern "C" UE_AUTORTFM_AUTORTFM("RTFM_autortfm_record_open_write") void autortfm_record_open_write(void* Ptr, size_t Size)
 {
-    FContext::Get()->CheckOpenRecordWrite(Ptr);
-	FContext::Get()->RecordWrite(Ptr, Size);
+	FContext* const Context = FContext::Get();
+
+    Context->CheckOpenRecordWrite(Ptr);
+	Context->RecordWrite(Ptr, Size);
 }
 
 extern "C" UE_AUTORTFM_NOAUTORTFM void autortfm_register_open_function(void* OriginalFunction, void* NewFunction)
@@ -202,6 +205,32 @@ extern "C" UE_AUTORTFM_AUTORTFM("RTFM_autortfm_did_allocate") void* autortfm_did
     return Ptr;
 }
 
+UE_DISABLE_OPTIMIZATION
+extern "C" UE_AUTORTFM_AUTORTFM("RTFM_autortfm_did_free") void autortfm_did_free(void* Ptr)
+{
+	// We only need to process did free if we need to track allocation locations.
+	if constexpr (bTrackAllocationLocations)
+	{
+		if (UNLIKELY(GIsCriticalError))
+		{
+			return;
+		}
+
+		if (FContext::IsTransactional())
+		{
+			FContext* const Context = FContext::Get();
+
+		    // We only care about frees that are occuring when the transaction
+		    // is in an on-going state (it's not committing or aborting).
+		    if (EContextStatus::OnTrack == Context->GetStatus())
+		    {
+		    	Context->DidFree(Ptr);
+		    }
+		}
+	}
+}
+UE_ENABLE_OPTIMIZATION
+
 extern "C" UE_AUTORTFM_AUTORTFM("RTFM_autortfm_check_consistency_assuming_no_races") void autortfm_check_consistency_assuming_no_races()
 {
     if (FContext::IsTransactional())
@@ -386,6 +415,13 @@ extern "C" UE_AUTORTFM_NOAUTORTFM void* RTFM_autortfm_did_allocate(void* Ptr, si
     return Ptr;
 }
 
+extern "C" UE_AUTORTFM_NOAUTORTFM void RTFM_autortfm_did_free(void* Ptr)
+{
+	// We should never-ever-ever actually free memory from within closed code of
+	// a transaction.
+	AutoRTFM::Unreachable();
+}
+
 extern "C" UE_AUTORTFM_NOAUTORTFM void RTFM_autortfm_check_consistency_assuming_no_races()
 {
 }
diff --git a/Engine/Source/Runtime/Core/Private/AutoRTFM/Context.h b/Engine/Source/Runtime/Core/Private/AutoRTFM/Context.h
--- a/Engine/Source/Runtime/Core/Private/AutoRTFM/Context.h
+++ b/Engine/Source/Runtime/Core/Private/AutoRTFM/Context.h
@@ -42,6 +42,7 @@ public:
     template<unsigned SIZE> void RecordWrite(void* LogicalAddress);
 
     void DidAllocate(void* LogicalAddress, size_t Size);
+    void DidFree(void* LogicalAddress);
 
     // The rest of this is internalish.
     void AbortByLanguageAndThrow();
diff --git a/Engine/Source/Runtime/Core/Private/AutoRTFM/ContextInlines.h b/Engine/Source/Runtime/Core/Private/AutoRTFM/ContextInlines.h
--- a/Engine/Source/Runtime/Core/Private/AutoRTFM/ContextInlines.h
+++ b/Engine/Source/Runtime/Core/Private/AutoRTFM/ContextInlines.h
@@ -34,6 +34,18 @@ UE_AUTORTFM_FORCEINLINE void FContext::DidAllocate(void* LogicalAddress, size_t
     CurrentTransaction->DidAllocate(LogicalAddress, Size);
 }
 
+UE_AUTORTFM_FORCEINLINE void FContext::DidFree(void* LogicalAddress)
+{
+    // We can do free's in the open within a transaction *during* when the
+    // transaction itself is being destroyed, so we need to check for that case.
+	if (UNLIKELY(!CurrentTransaction))
+	{
+		return;
+	}
+
+    CurrentTransaction->DidFree(LogicalAddress);
+}
+
 UE_AUTORTFM_FORCEINLINE bool FContext::AttemptToCommitTransaction(FTransaction* const Transaction)
 {
     ASSERT(EContextStatus::OnTrack == Status);
diff --git a/Engine/Source/Runtime/Core/Private/AutoRTFM/ContextStatus.h b/Engine/Source/Runtime/Core/Private/AutoRTFM/ContextStatus.h
--- a/Engine/Source/Runtime/Core/Private/AutoRTFM/ContextStatus.h
+++ b/Engine/Source/Runtime/Core/Private/AutoRTFM/ContextStatus.h
@@ -13,6 +13,9 @@ inline const char* GetContextStatusName(EContextStatus Status)
 {
     switch (Status)
     {
+    default:
+        Unreachable();
+        return nullptr;
     case EContextStatus::Idle:
         return "Idle";
     case EContextStatus::OnTrack:
@@ -23,9 +26,9 @@ inline const char* GetContextStatusName(EContextStatus Status)
         return "AbortedByLanguage";
     case EContextStatus::AbortedByRequest:
         return "AbortedByRequest";
+    case EContextStatus::Committing:
+        return "Committing";
     }
-    ASSERT(!"Should not be reached");
-    return NULL;
 }
 
 } // namespace AutoRTFM
diff --git a/Engine/Source/Runtime/Core/Private/AutoRTFM/HitSet.h b/Engine/Source/Runtime/Core/Private/AutoRTFM/HitSet.h
--- a/Engine/Source/Runtime/Core/Private/AutoRTFM/HitSet.h
+++ b/Engine/Source/Runtime/Core/Private/AutoRTFM/HitSet.h
@@ -70,7 +70,7 @@ namespace AutoRTFM
 		// Clear out the data stored in the set, but does not reduce the capacity.
 		void Reset()
 		{
-			memset(Payload, 0, sizeof(uintptr_t) * Capacity);
+			memset(Payload, 0, sizeof(uintptr_t) * (Capacity + LinearProbeDepth));
 			Size = 0;
 		}
 
diff --git a/Engine/Source/Runtime/Core/Private/AutoRTFM/IntervalTree.h b/Engine/Source/Runtime/Core/Private/AutoRTFM/IntervalTree.h
--- a/Engine/Source/Runtime/Core/Private/AutoRTFM/IntervalTree.h
+++ b/Engine/Source/Runtime/Core/Private/AutoRTFM/IntervalTree.h
@@ -17,9 +17,134 @@ struct FIntervalTree final
     FIntervalTree& operator=(const FIntervalTree&) = delete;
 
 	FORCENOINLINE bool Insert(const void* const Address, const size_t Size)
+    {        
+        FRange NewRange(Address, Size);
+        return Insert(MoveTemp(NewRange));
+    }
+
+    FORCENOINLINE bool Contains(const void* const Address, const size_t Size) const
     {
         FRange NewRange(Address, Size);
 
+        if (UNLIKELY(IntervalTreeNodeIndexNone == Root))
+        {
+            return false;
+        }
+
+        FIntervalTreeNodeIndex Current = Root;
+
+        do
+        {
+            const FIntervalTreeNode& Node = Nodes[Current];
+
+            const FRange Range = Node.Range;
+
+            // This check does not need to prove that NewRange is entirely
+            // enclosed within Range, because if any byte of NewRange was in the
+            // original Range then it **must** already have been new memory.
+            if ((NewRange.Start < Range.End) && (Range.Start < NewRange.End))
+            {
+                return true;
+            }
+            else if (NewRange.Start < Range.Start)
+            {
+                Current = Node.Left;
+            }
+            else
+            {
+                Current = Node.Right;
+            }
+        } while (IntervalTreeNodeIndexNone != Current);
+
+        return false;
+    }
+
+    bool IsEmpty() const
+    {
+        return IntervalTreeNodeIndexNone == Root;
+    }
+
+    void Reset()
+    {
+        Root = IntervalTreeNodeIndexNone;
+        Nodes.Empty();
+    }
+
+    void Merge(const FIntervalTree& Other)
+    {
+        if (Other.IsEmpty())
+        {
+            return;
+        }
+
+		ArrayType<FIntervalTreeNodeIndex> ToProcess;
+
+        ToProcess.Add(Other.Root);
+
+        do
+        {
+            const FIntervalTreeNodeIndex Current = ToProcess.Last();
+            ToProcess.Pop();
+
+            FRange Range = Other.Nodes[Current].Range;
+
+            Insert(MoveTemp(Range));
+
+            const FIntervalTreeNodeIndex Left = Other.Nodes[Current].Left;
+            const FIntervalTreeNodeIndex Right = Other.Nodes[Current].Right;
+
+            if (IntervalTreeNodeIndexNone != Left)
+            {
+                ToProcess.Push(Left);
+            }
+
+            if (IntervalTreeNodeIndexNone != Right)
+            {
+                ToProcess.Push(Right);
+            }
+        } while (!ToProcess.IsEmpty());
+    }
+
+private:
+	template<typename T> using ArrayType = TPagedArray<T>;
+
+    struct FRange final
+    {
+        FRange(const void* const Address, const size_t Size) :
+            Start(reinterpret_cast<uintptr_t>(Address)), End(reinterpret_cast<uintptr_t>(Address) + Size) {}
+
+        uintptr_t Start;
+        uintptr_t End;
+    };
+
+    using FIntervalTreeNodeIndex = uint32_t;
+
+    static constexpr FIntervalTreeNodeIndex IntervalTreeNodeIndexNone = UINT32_MAX;
+
+    struct FIntervalTreeNode final
+    {
+        explicit FIntervalTreeNode(const FRange Range) :
+            Parent(IntervalTreeNodeIndexNone), bIsBlack(true), Range(Range) {}
+
+		FIntervalTreeNode(const FRange Range, FIntervalTreeNodeIndex Parent) :
+			Parent(Parent), bIsBlack(false), Range(Range) {}
+
+        FIntervalTreeNodeIndex Left = IntervalTreeNodeIndexNone;
+        FIntervalTreeNodeIndex Right = IntervalTreeNodeIndexNone;
+        FIntervalTreeNodeIndex Parent;
+
+        // TODO: optimize this waste of memory.
+        bool bIsBlack;
+        FRange Range;
+    };
+
+    FIntervalTreeNodeIndex Root = IntervalTreeNodeIndexNone;
+	ArrayType<FIntervalTreeNode> Nodes;
+
+	static constexpr bool bExtraDebugging = false;
+
+    FORCENOINLINE bool Insert(FRange&& NewRange)
+    {
         if (UNLIKELY(IntervalTreeNodeIndexNone == Root))
         {
             ASSERT(0 == Nodes.Num());
@@ -28,8 +153,6 @@ struct FIntervalTree final
             return true;
         }
 
-		ASSERT(Nodes[Root].Parent == IntervalTreeNodeIndexNone);
-
         FIntervalTreeNodeIndex Current = Root;
 
         for(;;)
@@ -45,14 +168,14 @@ struct FIntervalTree final
 
             if (NewRange.Start < Range.Start)
             {
-				if (NewRange.End == Range.Start)
-				{
+                if (NewRange.End == Range.Start)
+                {
                     // We can just modify the existing node in place.
-					ASSERT(NewRange.Start < Range.Start);
-					Node.Range.Start = NewRange.Start;
+                    ASSERT(NewRange.Start < Range.Start);
+                    Node.Range.Start = NewRange.Start;
                     return true;
-				}
-				else if (IntervalTreeNodeIndexNone == Node.Left)
+                }
+                else if (IntervalTreeNodeIndexNone == Node.Left)
                 {
                     const FIntervalTreeNodeIndex Index = Nodes.Add(FIntervalTreeNode(NewRange, Current));
                     Node.Left = Index;
@@ -64,14 +187,14 @@ struct FIntervalTree final
             }
             else
             {
-				if (NewRange.Start == Range.End)
-				{
+                if (NewRange.Start == Range.End)
+                {
                     // We can just modify the existing node in place.
-					ASSERT(NewRange.End > Range.End);
-					Node.Range.End = NewRange.End;
-					return true;
-				}
-				else if (IntervalTreeNodeIndexNone == Node.Right)
+                    ASSERT(NewRange.End > Range.End);
+                    Node.Range.End = NewRange.End;
+                    return true;
+                }
+                else if (IntervalTreeNodeIndexNone == Node.Right)
                 {
                     const FIntervalTreeNodeIndex Index = Nodes.Add(FIntervalTreeNode(NewRange, Current));
                     Node.Right = Index;
@@ -82,7 +205,7 @@ struct FIntervalTree final
                 Current = Node.Right;
             }
 
-			ASSERT(Root != Current);
+            ASSERT(Root != Current);
         }
 
         auto IsBlack = [this](FIntervalTreeNodeIndex Index)
@@ -97,21 +220,21 @@ struct FIntervalTree final
             UE_ASSUME(Current != Parent);
 
             // The root will always have a black parent, so this check covers both.
-			if (Parent == Root)
-			{
-				Nodes[Parent].bIsBlack = true;
-				break;
-			}
+            if (Parent == Root)
+            {
+                Nodes[Parent].bIsBlack = true;
+                break;
+            }
             else if (IsBlack(Parent))
             {
-				break;
+                break;
             }
 
             const FIntervalTreeNodeIndex GrandParent = Nodes[Parent].Parent;
 
             UE_ASSUME((Parent != GrandParent) && (Current != GrandParent));
 
-			const bool bGrandParentIsRoot = GrandParent == Root;
+            const bool bGrandParentIsRoot = GrandParent == Root;
 
             const bool bParentIsLeft = Nodes[GrandParent].Left == Parent;
 
@@ -122,18 +245,18 @@ struct FIntervalTree final
 
             if (!IsBlack(Uncle))
             {
-				ASSERT(IsBlack(GrandParent));
+                ASSERT(IsBlack(GrandParent));
                 Nodes[Parent].bIsBlack = true;
                 Nodes[Uncle].bIsBlack = true;
 
-				if (bGrandParentIsRoot)
-				{
-					break;
-				}
-				else
-				{
-					Nodes[GrandParent].bIsBlack = false;
-				}
+                if (bGrandParentIsRoot)
+                {
+                    break;
+                }
+                else
+                {
+                    Nodes[GrandParent].bIsBlack = false;
+                }
 
                 Current = GrandParent;
                 continue;
@@ -148,25 +271,25 @@ struct FIntervalTree final
                 {
                     Nodes[Parent].Right = Nodes[Current].Left;
 
-					if (IntervalTreeNodeIndexNone != Nodes[Parent].Right)
-					{
-						Nodes[Nodes[Parent].Right].Parent = Parent;
-					}
+                    if (IntervalTreeNodeIndexNone != Nodes[Parent].Right)
+                    {
+                        Nodes[Nodes[Parent].Right].Parent = Parent;
+                    }
 
                     Nodes[Parent].Parent = Current;
                     Nodes[Current].Parent = GrandParent;
                     Nodes[Current].Left = Parent;
-					Nodes[GrandParent].Left = Current;
+                    Nodes[GrandParent].Left = Current;
                     std::swap(Parent, Current);
                 }
 
                 Nodes[Parent].Parent = Nodes[GrandParent].Parent;
                 Nodes[GrandParent].Left = Nodes[Parent].Right;
 
-				if (IntervalTreeNodeIndexNone != Nodes[GrandParent].Left)
-				{
-					Nodes[Nodes[GrandParent].Left].Parent = GrandParent;
-				}
+                if (IntervalTreeNodeIndexNone != Nodes[GrandParent].Left)
+                {
+                    Nodes[Nodes[GrandParent].Left].Parent = GrandParent;
+                }
 
                 Nodes[Parent].Right = GrandParent;
                 Nodes[GrandParent].Parent = Parent;
@@ -178,157 +301,106 @@ struct FIntervalTree final
                 if (CurrentIsLeft)
                 {
                     Nodes[Parent].Left = Nodes[Current].Right;
-					
-					if (IntervalTreeNodeIndexNone != Nodes[Parent].Left)
-					{
-						Nodes[Nodes[Parent].Left].Parent = Parent;
-					}
+                    
+                    if (IntervalTreeNodeIndexNone != Nodes[Parent].Left)
+                    {
+                        Nodes[Nodes[Parent].Left].Parent = Parent;
+                    }
 
                     Nodes[Parent].Parent = Current;
                     Nodes[Current].Parent = GrandParent;
                     Nodes[Current].Right = Parent;
-					Nodes[GrandParent].Right = Current;
+                    Nodes[GrandParent].Right = Current;
                     std::swap(Parent, Current);
                 }
 
                 Nodes[Parent].Parent = Nodes[GrandParent].Parent;
                 Nodes[GrandParent].Right = Nodes[Parent].Left;
 
-				if (IntervalTreeNodeIndexNone != Nodes[GrandParent].Right)
-				{
-					Nodes[Nodes[GrandParent].Right].Parent = GrandParent;
-				}
+                if (IntervalTreeNodeIndexNone != Nodes[GrandParent].Right)
+                {
+                    Nodes[Nodes[GrandParent].Right].Parent = GrandParent;
+                }
 
                 Nodes[Parent].Left = GrandParent;
-				Nodes[GrandParent].Parent = Parent;
+                Nodes[GrandParent].Parent = Parent;
 
                 std::swap(Nodes[GrandParent].bIsBlack, Nodes[Parent].bIsBlack);
             }
 
-			if (bGrandParentIsRoot)
-			{
-				Root = Parent;
-			}
-			else
-			{
-				FIntervalTreeNode& GreatGrandParent = Nodes[Nodes[Parent].Parent];
-				if (GrandParent == GreatGrandParent.Left)
-				{
-					GreatGrandParent.Left = Parent;
-				}
-				else
-				{
-					GreatGrandParent.Right = Parent;
-				}
-			}
-
-            break;
-        }
-
-		if constexpr (bExtraDebugging)
-		{
-			for (FIntervalTreeNodeIndex Index = 0; Index < Nodes.Num(); Index++)
-			{
-				const FIntervalTreeNodeIndex Parent = Nodes[Index].Parent;
-				if (IntervalTreeNodeIndexNone == Parent)
-				{
-					ASSERT(Root == Index);
-				}
-				else
-				{
-					ASSERT(Nodes[Parent].bIsBlack || Nodes[Index].bIsBlack);
-					ASSERT((Nodes[Parent].Left == Index) ^ (Nodes[Parent].Right == Index));
-				}
-			}
-		}
-
-        return true;
-    }
-
-    FORCENOINLINE bool Contains(const void* const Address, const size_t Size) const
-    {
-        FRange NewRange(Address, Size);
-
-        if (UNLIKELY(IntervalTreeNodeIndexNone == Root))
-        {
-            return false;
-        }
-
-        FIntervalTreeNodeIndex Current = Root;
-
-        do
-        {
-            const FIntervalTreeNode& Node = Nodes[Current];
-
-            const FRange Range = Node.Range;
-
-            // This check does not need to prove that NewRange is entirely
-            // enclosed within Range, because if any byte of NewRange was in the
-            // original Range then it **must** already have been new memory.
-            if ((NewRange.Start < Range.End) && (Range.Start < NewRange.End))
+            if (bGrandParentIsRoot)
             {
-                return true;
-            }
-            else if (NewRange.Start < Range.Start)
-            {
-                Current = Node.Left;
+                Root = Parent;
             }
             else
             {
-                Current = Node.Right;
+                FIntervalTreeNode& GreatGrandParent = Nodes[Nodes[Parent].Parent];
+                if (GrandParent == GreatGrandParent.Left)
+                {
+                    GreatGrandParent.Left = Parent;
+                }
+                else
+                {
+                    GreatGrandParent.Right = Parent;
+                }
             }
-        } while (IntervalTreeNodeIndexNone != Current);
 
-        return false;
-    }
+            break;
+        }
 
-    bool IsEmpty() const
-    {
-        return IntervalTreeNodeIndexNone == Root;
-    }
+		AssertStructureIsOk();
 
-    void Reset()
-    {
-        Root = IntervalTreeNodeIndexNone;
-        Nodes.Empty();
+        return true;
     }
 
-private:
-    struct FRange final
-    {
-        FRange(const void* const Address, const size_t Size) :
-            Start(reinterpret_cast<uintptr_t>(Address)), End(reinterpret_cast<uintptr_t>(Address) + Size) {}
-
-        uintptr_t Start;
-        uintptr_t End;
-    };
-
-private:
-    using FIntervalTreeNodeIndex = uint32_t;
-
-    static constexpr FIntervalTreeNodeIndex IntervalTreeNodeIndexNone = UINT32_MAX;
+	UE_AUTORTFM_FORCEINLINE void AssertStructureIsOk() const
+	{
+		if constexpr (bExtraDebugging)
+		{
+			if (IntervalTreeNodeIndexNone != Root)
+			{
+				AssertNodeIsOk(Root);
+			}
+		}
+	}
 
-    struct FIntervalTreeNode final
-    {
-        explicit FIntervalTreeNode(const FRange Range) :
-            Parent(IntervalTreeNodeIndexNone), bIsBlack(true), Range(Range) {}
+	FORCENOINLINE void AssertNodeIsOk(FIntervalTreeNodeIndex Index) const
+	{
+		// We need to use recursion to check because we cannot have any
+		// allocations within the checker itself!
+		if constexpr (bExtraDebugging)
+		{
+			ASSERT(IntervalTreeNodeIndexNone != Index);
+			ASSERT(Index < Nodes.Num());
 
-		FIntervalTreeNode(const FRange Range, FIntervalTreeNodeIndex Parent) :
-			Parent(Parent), bIsBlack(false), Range(Range) {}
+			const FIntervalTreeNodeIndex Parent = Nodes[Index].Parent;
+			const FIntervalTreeNodeIndex Left = Nodes[Index].Left;
+			const FIntervalTreeNodeIndex Right = Nodes[Index].Right;
 
-        FIntervalTreeNodeIndex Left = IntervalTreeNodeIndexNone;
-        FIntervalTreeNodeIndex Right = IntervalTreeNodeIndexNone;
-        FIntervalTreeNodeIndex Parent;
+			if (IntervalTreeNodeIndexNone == Parent)
+			{
+				ASSERT(Root == Index);
+			}
+			else
+			{
+				ASSERT(Nodes[Parent].bIsBlack || Nodes[Index].bIsBlack);
+				ASSERT((Nodes[Parent].Left == Index) ^ (Nodes[Parent].Right == Index));
+			}
 
-        // TODO: optimize this waste of memory.
-        bool bIsBlack;
-        FRange Range;
-    };
+			ASSERT(Left != Index);
+			ASSERT(Right != Index);
 
-    FIntervalTreeNodeIndex Root = IntervalTreeNodeIndexNone;
-    TPagedArray<FIntervalTreeNode> Nodes;
+			if (IntervalTreeNodeIndexNone != Left)
+			{
+				AssertNodeIsOk(Left);
+			}
 
-	static constexpr bool bExtraDebugging = false;
+			if (IntervalTreeNodeIndexNone != Right)
+			{
+				AssertNodeIsOk(Right);
+			}
+		}
+	}
 };
 
 } // namespace AutoRTFM
diff --git a/Engine/Source/Runtime/Core/Private/AutoRTFM/Stats.h b/Engine/Source/Runtime/Core/Private/AutoRTFM/Stats.h
--- a/Engine/Source/Runtime/Core/Private/AutoRTFM/Stats.h
+++ b/Engine/Source/Runtime/Core/Private/AutoRTFM/Stats.h
@@ -2,13 +2,12 @@
 
 #pragma once
 
+#include "Toggles.h"
 #include "Utils.h"
 
 namespace AutoRTFM
 {
 
-static constexpr bool bCollectStats = false;
-
 enum class EStatsKind : uint8_t
 {
 	Transaction = 0,
diff --git a/Engine/Source/Runtime/Core/Private/AutoRTFM/Toggles.h b/Engine/Source/Runtime/Core/Private/AutoRTFM/Toggles.h
new file mode 100644
--- /dev/null
+++ b/Engine/Source/Runtime/Core/Private/AutoRTFM/Toggles.h
@@ -0,0 +1,14 @@
+// Copyright Epic Games, Inc. All Rights Reserved.
+
+#pragma once
+
+namespace AutoRTFM
+{
+
+// Should we collect stats on the AutoRTFM runtime or not.
+static constexpr bool bCollectStats = false;
+
+// Should we track the locations of allocations.
+static constexpr bool bTrackAllocationLocations = false;
+
+} // namespace AutoRTFM
diff --git a/Engine/Source/Runtime/Core/Private/AutoRTFM/Transaction.cpp b/Engine/Source/Runtime/Core/Private/AutoRTFM/Transaction.cpp
--- a/Engine/Source/Runtime/Core/Private/AutoRTFM/Transaction.cpp
+++ b/Engine/Source/Runtime/Core/Private/AutoRTFM/Transaction.cpp
@@ -186,6 +186,8 @@ void FTransaction::CommitNested()
 
     Parent->CommitTasks.AddAll(MoveTemp(CommitTasks));
     Parent->AbortTasks.AddAll(MoveTemp(AbortTasks));
+
+    Parent->NewMemoryTracker.Merge(NewMemoryTracker);
 }
 
 bool FTransaction::AttemptToCommitOuterNest()
diff --git a/Engine/Source/Runtime/Core/Private/AutoRTFM/Transaction.h b/Engine/Source/Runtime/Core/Private/AutoRTFM/Transaction.h
--- a/Engine/Source/Runtime/Core/Private/AutoRTFM/Transaction.h
+++ b/Engine/Source/Runtime/Core/Private/AutoRTFM/Transaction.h
@@ -74,6 +74,7 @@ public:
     template<unsigned SIZE> void RecordWrite(void* LogicalAddress);
 
     void DidAllocate(void* LogicalAddress, size_t Size);
+    void DidFree(void* LogicalAddress);
 
 private:
     void Undo();
diff --git a/Engine/Source/Runtime/Core/Private/AutoRTFM/TransactionInlines.h b/Engine/Source/Runtime/Core/Private/AutoRTFM/TransactionInlines.h
--- a/Engine/Source/Runtime/Core/Private/AutoRTFM/TransactionInlines.h
+++ b/Engine/Source/Runtime/Core/Private/AutoRTFM/TransactionInlines.h
@@ -100,21 +100,34 @@ UE_AUTORTFM_FORCEINLINE void FTransaction::DidAllocate(void* LogicalAddress, con
 	{
 		return;
 	}
-	else if (Size <= FWriteLogBumpAllocator::MaxSize)
+
+    const bool bUseHitSet = Size <= FWriteLogBumpAllocator::MaxSize;
+
+	if (bUseHitSet)
     {
         FMemoryLocation Key(LogicalAddress);
         Key.SetTopTag(static_cast<uint16_t>(Size));
         // Otherwise we need to record the write.
         const bool DidInsert = HitSet.Insert(Key);
-        ASSERT(DidInsert); 
+        ASSERT(DidInsert);
     }
-    else
+    
+    if (bTrackAllocationLocations || !bUseHitSet)
     {
         const bool DidInsert = NewMemoryTracker.Insert(LogicalAddress, Size);
         ASSERT(DidInsert);
     }
 }
 
+UE_AUTORTFM_FORCEINLINE void FTransaction::DidFree(void* LogicalAddress)
+{
+    ASSERT(bTrackAllocationLocations);
+    
+    // Checking if one byte is in the interval map is enough to ascertain if it
+    // is new memory and we should be worried.
+    ASSERT(!NewMemoryTracker.Contains(LogicalAddress, 1));
+}
+
 UE_AUTORTFM_FORCEINLINE void FTransaction::DeferUntilCommit(TFunction<void()>&& Callback)
 {
     CommitTasks.Add(MoveTemp(Callback));
diff --git a/Engine/Source/Runtime/Core/Private/AutoRTFM/Utils.h b/Engine/Source/Runtime/Core/Private/AutoRTFM/Utils.h
--- a/Engine/Source/Runtime/Core/Private/AutoRTFM/Utils.h
+++ b/Engine/Source/Runtime/Core/Private/AutoRTFM/Utils.h
@@ -45,4 +45,4 @@ template<size_t A, size_t B> struct PrettyStaticAssert final
 
 } // namespace AutoRTFM
 
-#define ASSERT(exp) do { if (UNLIKELY(!(exp))) PrettyAbort(__FILE__, __LINE__, __PRETTY_FUNCTION__, #exp); } while (false)
+#define ASSERT(exp) do { if (UNLIKELY(!(exp))) { UE_DEBUG_BREAK(); PrettyAbort(__FILE__, __LINE__, __PRETTY_FUNCTION__, #exp); } } while (false)
diff --git a/Engine/Source/Runtime/Core/Public/AutoRTFM/AutoRTFM.h b/Engine/Source/Runtime/Core/Public/AutoRTFM/AutoRTFM.h
--- a/Engine/Source/Runtime/Core/Public/AutoRTFM/AutoRTFM.h
+++ b/Engine/Source/Runtime/Core/Public/AutoRTFM/AutoRTFM.h
@@ -325,6 +325,16 @@ UE_AUTORTFM_FORCEINLINE void* autortfm_did_allocate(void* ptr, size_t size)
 }
 #endif
 
+// Inform the runtime that we have free'd a given memory location.
+#if UE_AUTORTFM
+UE_AUTORTFM_API void autortfm_did_free(void* ptr);
+#else
+UE_AUTORTFM_FORCEINLINE void autortfm_did_free(void* ptr)
+{
+	UE_AUTORTFM_UNUSED(ptr);
+}
+#endif
+
 // If running in a transaction, then perform a consistency check of the
 // transaction's read-write set. If possible, this compares the read-write set's
 // expected values with the actual values in global memory. Does nothing when
@@ -541,6 +551,11 @@ UE_AUTORTFM_FORCEINLINE void* DidAllocate(void* Ptr, size_t Size)
     return autortfm_did_allocate(Ptr, Size);
 }
 
+UE_AUTORTFM_FORCEINLINE void DidFree(void* Ptr)
+{
+    autortfm_did_free(Ptr);
+}
+
 UE_AUTORTFM_FORCEINLINE void CheckConsistencyAssumingNoRaces()
 {
     autortfm_check_consistency_assuming_no_races();
diff --git a/Engine/Source/Runtime/Core/Public/HAL/FMemory.inl b/Engine/Source/Runtime/Core/Public/HAL/FMemory.inl
--- a/Engine/Source/Runtime/Core/Public/HAL/FMemory.inl
+++ b/Engine/Source/Runtime/Core/Public/HAL/FMemory.inl
@@ -105,6 +105,13 @@ FMEMORY_INLINE_FUNCTION_DECORATOR void* FMemory::Realloc(void* Original, SIZE_T
 		Ptr = FMEMORY_INLINE_GMalloc->Realloc(Original, Count, Alignment);
 	}
 
+	if (Ptr != Original)
+	{
+		// If the pointer we've got back from realloc is new memory, we are
+		// assuming the old memory was free'd.
+		AutoRTFM::DidFree(Original);
+	}
+
 	// optional tracking of every allocation - a realloc with a Count of zero is equivalent to a call 
 	// to free() and will return a null pointer which does not require tracking. If realloc returns null
 	// for some other reason (like failure to allocate) there's also no reason to track it
@@ -137,6 +144,8 @@ FMEMORY_INLINE_FUNCTION_DECORATOR void FMemory::Free(void* Original)
 		DoGamethreadHook(2);
 		FScopedMallocTimer Timer(2);
 		FMEMORY_INLINE_GMalloc->Free(Original);
+
+		AutoRTFM::DidFree(Original);
 	});
 }
 
diff --git a/Engine/Source/Runtime/CoreUObject/Public/UObject/UObjectGlobals.h b/Engine/Source/Runtime/CoreUObject/Public/UObject/UObjectGlobals.h
--- a/Engine/Source/Runtime/CoreUObject/Public/UObject/UObjectGlobals.h
+++ b/Engine/Source/Runtime/CoreUObject/Public/UObject/UObjectGlobals.h
@@ -1725,7 +1725,15 @@ FUNCTION_NON_NULL_RETURN_END
 	Params.bCopyTransientsFromClassDefaults = bCopyTransientsFromClassDefaults;
 	Params.InstanceGraph = InInstanceGraph;
 	Params.ExternalPackage = ExternalPackage;
-	return static_cast<T*>(StaticConstructObject_Internal(Params));
+
+	T* Result = nullptr;
+
+	UE_AUTORTFM_OPEN(
+	{
+		Result = static_cast<T*>(StaticConstructObject_Internal(Params));
+	});
+
+	return Result;
 }
 
 template< class T >
