remove FPlatformChunkInstallCompleteDelegate - it's been unused since at least 4.18

#rnx

diff --git a/Engine/Source/Runtime/Core/Public/GenericPlatform/GenericPlatformChunkInstall.h b/Engine/Source/Runtime/Core/Public/GenericPlatform/GenericPlatformChunkInstall.h
--- a/Engine/Source/Runtime/Core/Public/GenericPlatform/GenericPlatformChunkInstall.h
+++ b/Engine/Source/Runtime/Core/Public/GenericPlatform/GenericPlatformChunkInstall.h
@@ -87,9 +87,6 @@ public:
 	virtual IPlatformChunkInstall* GetPlatformChunkInstall() = 0;
 };
 
-/** Deprecated delegate */
-DECLARE_DELEGATE_OneParam(FPlatformChunkInstallCompleteDelegate, uint32);
-
 /** Delegate called when a chunk either successfully installs or fails to install, bool is success */
 DECLARE_DELEGATE_TwoParams(FPlatformChunkInstallDelegate, uint32, bool);
 DECLARE_MULTICAST_DELEGATE_TwoParams(FPlatformChunkInstallMultiDelegate, uint32, bool);
