Fix printf checkf warnings in CompactBinaryWriter due to UTF8 casts to Ansi char.

#rnx

diff --git a/Engine/Source/Runtime/Core/Private/Serialization/CompactBinaryWriter.cpp b/Engine/Source/Runtime/Core/Private/Serialization/CompactBinaryWriter.cpp
--- a/Engine/Source/Runtime/Core/Private/Serialization/CompactBinaryWriter.cpp
+++ b/Engine/Source/Runtime/Core/Private/Serialization/CompactBinaryWriter.cpp
@@ -137,8 +137,8 @@ void FCbWriter::BeginField()
 	else
 	{
 		checkf((State.Flags & EStateFlags::Name) == EStateFlags::Name,
-			TEXT("A new field cannot be written until the previous field '%.*hs' is finished."),
-			GetActiveName().Len(), GetActiveName().GetData());
+			TEXT("A new field cannot be written until the previous field '%s' is finished."),
+			StringCast<TCHAR>(GetActiveName().GetData(), GetActiveName().Len()).Get());
 	}
 }
 
@@ -174,13 +174,13 @@ FCbWriter& FCbWriter::SetName(const FUtf8StringView Name)
 {
 	FState& State = States.Last();
 	checkf((State.Flags & EStateFlags::Array) != EStateFlags::Array,
-		TEXT("It is invalid to write a name for an array field. Name '%.*hs'"), Name.Len(), Name.GetData());
+		TEXT("It is invalid to write a name for an array field. Name '%s'"), StringCast<TCHAR>(Name.GetData(), Name.Len()).Get());
 	checkf(!Name.IsEmpty(), TEXT("%s"), (State.Flags & EStateFlags::Object) == EStateFlags::Object
 		? TEXT("It is invalid to write an empty name for an object field. Specify a unique non-empty name.")
 		: TEXT("It is invalid to write an empty name for a top-level field. Specify a name or avoid this call."));
 	checkf((State.Flags & (EStateFlags::Name | EStateFlags::Field)) == EStateFlags::None,
-		TEXT("A new field '%.*hs' cannot be written until the previous field '%.*hs' is finished."),
-		Name.Len(), Name.GetData(), GetActiveName().Len(), GetActiveName().GetData());
+		TEXT("A new field '%s' cannot be written until the previous field '%s' is finished."),
+		StringCast<TCHAR>(Name.GetData(), Name.Len()).Get(), StringCast<TCHAR>(GetActiveName().GetData(), GetActiveName().Len()).Get());
 
 	BeginField();
 	State.Flags |= EStateFlags::Name;
@@ -595,8 +595,8 @@ void FCbWriter::AddCustom(const uint64 TypeId, const FMemoryView Value)
 
 void FCbWriter::AddCustom(const FUtf8StringView TypeName, const FMemoryView Value)
 {
-	checkf(!TypeName.IsEmpty(), TEXT("Field '%.*hs' requires a non-empty type name for its custom type."),
-		GetActiveName().Len(), GetActiveName().GetData());
+	checkf(!TypeName.IsEmpty(), TEXT("Field '%s' requires a non-empty type name for its custom type."),
+		StringCast<TCHAR>(GetActiveName().GetData(), GetActiveName().Len()).Get());
 
 	BeginField();
 
