- Added system to only save config sections, to user Saved/Config files, that are opted in
- Currently disabled, allowing people to test locally by setting a bool in Base.ini (affects all hierarchical configs)

diff --git a/Engine/Config/Base.ini b/Engine/Config/Base.ini
--- a/Engine/Config/Base.ini
+++ b/Engine/Config/Base.ini
@@ -1 +1,5 @@
 ; Base.ini is required to exist
+
+[SectionsToSave]
+bCanSaveAllSections=true
++Section=CurrentIniVersion
diff --git a/Engine/Config/BaseEditor.ini b/Engine/Config/BaseEditor.ini
--- a/Engine/Config/BaseEditor.ini
+++ b/Engine/Config/BaseEditor.ini
@@ -2,6 +2,10 @@
 ; These settings are overridden by a project's DefaultEditor.ini file, per-user editor setting defaults are in *EditorPerProjectUserSettings.ini
 ; These settings are not available in a packaged build
 
+[SectionsToSave]
++Section=/Script/UnrealEd.UnrealEdOptions
++Section=/Script/AdvancedPreviewScene.LocalProfiles
+
 [Internationalization]
 ; These are the paths for localizing the editor itself, the localized game content settings are in *Game.ini
 +LocalizationPaths=../../../Engine/Content/Localization/Editor
@@ -630,4 +634,4 @@ bUseWidgetTemplateSelector=True
 +WidgetClassesToHide=/Script/UMG.Slider
 +WidgetClassesToHide=/Script/UMG.SpinBox
 +WidgetClassesToHide=/Script/UMG.Throbber
-+WidgetClassesToHide=/Script/UMG.TreeView
\ No newline at end of file
++WidgetClassesToHide=/Script/UMG.TreeView
diff --git a/Engine/Config/BaseEngine.ini b/Engine/Config/BaseEngine.ini
--- a/Engine/Config/BaseEngine.ini
+++ b/Engine/Config/BaseEngine.ini
@@ -2,6 +2,9 @@
 ; These settings are overridden by a project's DefaultEngine.ini file and per-platform overrides
 ; Some of these settings can be modified from the project settings in the editor
 
+[SectionsToSave]
++Section=GameNetDriver StatelessConnectHandlerComponent
+
 [/Script/EngineSettings.GameMapsSettings]
 GameInstanceClass=/Script/Engine.GameInstance
 EditorStartupMap=/Engine/Maps/Templates/OpenWorld
@@ -3355,4 +3358,4 @@ ContentImportSettings=(DefaultPipelineStackOverride=((Materials, "Materials"),(T
 SceneImportSettings=(PipelineStacks=(("Scene", (Pipelines=("/Interchange/Pipelines/DefaultSceneAssetsPipeline.DefaultSceneAssetsPipeline","/Interchange/Pipelines/DefaultSceneLevelPipeline.DefaultSceneLevelPipeline"),PerTranslatorPipelines=((Translator="/Script/InterchangeImport.InterchangeMaterialXTranslator",Pipelines=("/Interchange/Pipelines/DefaultSceneAssetsPipeline.DefaultSceneAssetsPipeline","/Interchange/Pipelines/DefaultSceneLevelPipeline.DefaultSceneLevelPipeline","/Interchange/Pipelines/DefaultMaterialXPipeline.DefaultMaterialXPipeline")),(Translator="/Script/InterchangeImport.InterchangeGltfTranslator",Pipelines=("/Interchange/Pipelines/DefaultGLTFSceneAssetsPipeline.DefaultGLTFSceneAssetsPipeline","/Interchange/Pipelines/DefaultSceneLevelPipeline.DefaultSceneLevelPipeline","/Interchange/Pipelines/DefaultGLTFPipeline.DefaultGLTFPipeline")))))),DefaultPipelineStack="Scene",PipelineConfigurationDialogClass="/Script/InterchangeEditorPipelines.InterchangePipelineConfigurationGeneric",bShowPipelineStacksConfigurationDialog=True)
 FilePickerClass=/Script/InterchangeEditorUtilities.InterchangeFilePickerGeneric
 bStaticMeshUseSmoothEdgesIfSmoothingInformationIsMissing=True
-GenericPipelineClass=/Script/InterchangePipelines.InterchangeGenericAssetsPipeline
\ No newline at end of file
+GenericPipelineClass=/Script/InterchangePipelines.InterchangeGenericAssetsPipeline
diff --git a/Engine/Config/BaseGame.ini b/Engine/Config/BaseGame.ini
--- a/Engine/Config/BaseGame.ini
+++ b/Engine/Config/BaseGame.ini
@@ -2,6 +2,9 @@
 ; These settings are overridden by a project's DefaultGame.ini file and per-platform overrides
 ; Most of these options can be set from Project Settings in the editor
 
+[SectionsToSave]
++Section=DebugWindows
+
 [Internationalization]
 ; To add paths to use for localized game text, use the localization tools or add more lines like: +LocalizationPaths=%GAMEDIR%Content/Localization/SpecificConfiguration
 ; To map a culture code to a culture implemented in your game, add lines like: +CultureMappings="es-AR;es-419"
diff --git a/Engine/Config/BaseSourceControlSettings.ini b/Engine/Config/BaseSourceControlSettings.ini
new file mode 100644
--- /dev/null
+++ b/Engine/Config/BaseSourceControlSettings.ini
@@ -0,0 +1,3 @@
+[SectionsToSave]
+bCanSaveAllSections=true
+
diff --git a/Engine/Plugins/Compositing/OpenColorIO/Config/BaseOpenColorIO.ini b/Engine/Plugins/Compositing/OpenColorIO/Config/BaseOpenColorIO.ini
new file mode 100644
--- /dev/null
+++ b/Engine/Plugins/Compositing/OpenColorIO/Config/BaseOpenColorIO.ini
@@ -0,0 +1,2 @@
+[SectionsToSave]
++Section=/Script/OpenColorIOEditor.OpenColorIOLevelViewportSettings
diff --git a/Engine/Source/Runtime/Core/Private/Misc/ConfigCacheIni.cpp b/Engine/Source/Runtime/Core/Private/Misc/ConfigCacheIni.cpp
--- a/Engine/Source/Runtime/Core/Private/Misc/ConfigCacheIni.cpp
+++ b/Engine/Source/Runtime/Core/Private/Misc/ConfigCacheIni.cpp
@@ -32,7 +32,8 @@
 
 namespace
 {
-	FString CurrentIniVersionStr = TEXT("CurrentIniVersion");
+	const FString CurrentIniVersionStr = TEXT("CurrentIniVersion");
+	const FString SectionsToSaveStr = TEXT("SectionsToSave");
 
 	TMap<FString, FString> SectionRemap;
 	TMap<FString, TMap<FString, FString>> KeyRemap;
@@ -1491,6 +1492,13 @@ bool FConfigFile::WriteInternal(const FString& Filename, bool bDoRemoteWrite, TM
 		FRemoteConfig::Get()->Write(*Filename, Text);
 	}
 
+	// don't write out non-default configs that are only whitespace
+	if (!bIsADefaultIniWrite && Text.TrimStart().Len() == 0)
+	{
+		IFileManager::Get().Delete(*Filename);
+		return true;
+	}
+	
 	bool bResult = SaveConfigFileWrapper(*Filename, Text);
 
 	// File is still dirty if it didn't save.
@@ -1526,11 +1534,26 @@ void FConfigFile::WriteToStringInternal(FString& InOutText, bool bIsADefaultIniW
 	TSet<FName> PropertiesAddedLookup;
 	PropertiesAddedLookup.Reserve(HighestPropertiesInSection);
 	int32 EstimatedFinalTextSize = 0;
-
+	
+	// no need to look up the section if it's a default ini, or if we are always saving all sections
+	FConfigSection* SectionsToSaveSection = (bIsADefaultIniWrite || bCanSaveAllSections) ? nullptr : Find(SectionsToSaveStr);
+	TArray<FString> SectionsToSave;
+	if (SectionsToSaveSection != nullptr)
+	{
+		SectionsToSaveSection->MultiFind("Section", SectionsToSave);
+	}
+	
 	for( TIterator SectionIterator(*this); SectionIterator; ++SectionIterator )
 	{
 		const FString& SectionName = SectionIterator.Key();
 		const FConfigSection& Section = SectionIterator.Value();
+		
+		// null Sections array means to save everything, otherwise check if we can save this section
+		bool bCanSaveThisSection = SectionsToSaveSection == nullptr || SectionsToSave.Contains(SectionName);
+		if (!bCanSaveThisSection)
+		{
+			continue;
+		}
 
 		// If we have a config file to check against, have a look.
 		FConfigSection* SourceConfigSection = nullptr;
@@ -1680,12 +1703,6 @@ void FConfigFile::WriteToStringInternal(FString& InOutText, bool bIsADefaultIniW
 			AddSectionToText(SectionName);
 		}
 	}
-
-	// Ensure We have at least something to write
-	if (InOutText.Len() == 0)
-	{
-		InOutText.Append(LINE_TERMINATOR);
-	}
 }
 
 /** Adds any properties that exist in InSourceFile that this config file is missing */
diff --git a/Engine/Source/Runtime/Core/Private/Misc/ConfigContext.cpp b/Engine/Source/Runtime/Core/Private/Misc/ConfigContext.cpp
--- a/Engine/Source/Runtime/Core/Private/Misc/ConfigContext.cpp
+++ b/Engine/Source/Runtime/Core/Private/Misc/ConfigContext.cpp
@@ -19,6 +19,8 @@ namespace
 	FString LegacyIniVersionString = TEXT("IniVersion");
 	FString LegacyEngineString = TEXT("Engine.Engine");
 	FString CurrentIniVersionString = TEXT("CurrentIniVersion");
+	const TCHAR* SectionsToSaveString = TEXT("SectionsToSave");
+	const TCHAR* SaveAllSectionsKey = TEXT("bCanSaveAllSections");
 }
 
 
@@ -390,6 +392,12 @@ bool FConfigContext::PerformLoad()
 				ConfigFile->Write(DestIniFilename);
 			}
 		}
+
+		// chcek if the config file wants to save all sections
+		bool bLocalSaveAllSections;
+		ConfigFile->bCanSaveAllSections = ConfigFile->GetBool(SectionsToSaveString, SaveAllSectionsKey, bLocalSaveAllSections) && bLocalSaveAllSections;
+		// we can always save all sections of a User config file
+		ConfigFile->bCanSaveAllSections = ConfigFile->bCanSaveAllSections || BaseIniName.Contains(TEXT("User"));
 	}
 
 	// GenerateDestIniFile returns true if nothing is loaded, so check if we actually loaded something
diff --git a/Engine/Source/Runtime/Core/Public/Misc/ConfigCacheIni.h b/Engine/Source/Runtime/Core/Public/Misc/ConfigCacheIni.h
--- a/Engine/Source/Runtime/Core/Public/Misc/ConfigCacheIni.h
+++ b/Engine/Source/Runtime/Core/Public/Misc/ConfigCacheIni.h
@@ -377,6 +377,9 @@ public:
 	bool Dirty;
 	bool NoSave;
 	bool bHasPlatformName = false;
+	// by default, we allow saving - this is going to be applied to config files that are not loaded from disk
+	// (when loading, this will get set to false, and then the ini sections will be checked)
+	bool bCanSaveAllSections = true;
 
 	/** The name of this config file */	
 	FName Name;
