Attempted fix for warnings in SAL annotations, probably caused by our #define UINT.

diff --git a/Engine/Source/Runtime/Core/Public/Windows/AllowWindowsPlatformTypes.h b/Engine/Source/Runtime/Core/Public/Windows/AllowWindowsPlatformTypes.h
--- a/Engine/Source/Runtime/Core/Public/Windows/AllowWindowsPlatformTypes.h
+++ b/Engine/Source/Runtime/Core/Public/Windows/AllowWindowsPlatformTypes.h
@@ -18,6 +18,9 @@
 #pragma warning( push )
 #pragma warning( disable : 4459 )
 
+// Disable warnings about malformed SAL annotations - this is probably due to the UINT macro below
+#pragma warning( disable : 28285 ) // d3d12.h(6236) : warning C28285: For function 'SetComputeRoot32BitConstants' '_Param_(3)' syntax error in 'SAL_readableTo(elementCount(__formal(1,Num32BitValuesToSet)*sizeof(::UINT)))' near '::UINT)))'.
+
 #define INT ::INT
 #define UINT ::UINT
 #define DWORD ::DWORD
