Disable const string test code in checkf,ensuremsgf since it's failing to compile in VS219 min toolchain

#rnx

diff --git a/Engine/Source/Runtime/Core/Public/Misc/AssertionMacros.h b/Engine/Source/Runtime/Core/Public/Misc/AssertionMacros.h
--- a/Engine/Source/Runtime/Core/Public/Misc/AssertionMacros.h
+++ b/Engine/Source/Runtime/Core/Public/Misc/AssertionMacros.h
@@ -253,7 +253,6 @@ RetType FORCENOINLINE UE_DEBUG_SECTION DispatchCheckVerify(InnerType&& Inner, Ar
 
 	#define UE_CHECK_F_IMPL(expr, format, ...) \
 		{ \
-			static_assert(!UE::Core::Private::FormatStringSan::bIsAConstString<decltype(expr)>, "The checkf condition cannot be a const string. Use != nullptr if it's intentional"); \
 			if(UNLIKELY(!(expr))) \
 			{ \
 				if (FDebug::CheckVerifyFailedImpl(#expr, __FILE__, __LINE__, PLATFORM_RETURN_ADDRESS(), format, ##__VA_ARGS__)) \
@@ -375,7 +374,6 @@ RetType FORCENOINLINE UE_DEBUG_SECTION DispatchCheckVerify(InnerType&& Inner, Ar
 	#define UE_ENSURE_IMPL2(Capture, Always, InExpression, ...) \
 		(LIKELY(!!(InExpression)) || (DispatchCheckVerify<bool>([Capture] () UE_DEBUG_SECTION \
 		{ \
-			static_assert(!UE::Core::Private::FormatStringSan::bIsAConstString<decltype(InExpression)>, "The ensureMsgf condition cannot be a const string. Use != nullptr if it's intentional"); \
 			static bool bExecuted = false; \
 			FValidateArgsInternal(__VA_ARGS__); \
 			return CheckVerifyImpl(bExecuted, Always, __FILE__, __LINE__, PLATFORM_RETURN_ADDRESS(), #InExpression, ##__VA_ARGS__); \
@@ -391,7 +389,7 @@ RetType FORCENOINLINE UE_DEBUG_SECTION DispatchCheckVerify(InnerType&& Inner, Ar
 	#define ensure(           InExpression                ) (LIKELY(!!(InExpression)))
 	#define ensureMsgf(       InExpression, InFormat, ... ) (LIKELY(!!(InExpression)))
 	#define ensureAlways(     InExpression                ) (LIKELY(!!(InExpression)))
-	#define ensureAlwaysMsgf( InExpression, InFormat, ... ) (LIKELY(!!(InExpression)))
+#define ensureAlwaysMsgf( InExpression, InFormat, ... ) (LIKELY(!!(InExpression)))
 
 #endif	// DO_CHECK
 
