Get Address Sanitizer to work with transactions.

diff --git a/Engine/Source/Runtime/Core/Private/AutoRTFM/CompilerInterface.cpp b/Engine/Source/Runtime/Core/Private/AutoRTFM/CompilerInterface.cpp
--- a/Engine/Source/Runtime/Core/Private/AutoRTFM/CompilerInterface.cpp
+++ b/Engine/Source/Runtime/Core/Private/AutoRTFM/CompilerInterface.cpp
@@ -28,7 +28,7 @@ void CheckAlignment(FContext* Context, void* Ptr, size_t AlignmentMask)
     }
 }
 
-extern "C" void autortfm_record_write(void* Ptr, size_t Size)
+extern "C" AUTORTFM_NO_ASAN void autortfm_record_write(void* Ptr, size_t Size)
 {
 	// check for writes to null here so we end up crashing in the user
 	// code rather than in the autortfm runtime.
@@ -41,7 +41,7 @@ extern "C" void autortfm_record_write(void* Ptr, size_t Size)
 	Context->RecordWrite(Ptr, Size);
 }
 
-extern "C" void autortfm_record_write_8(void* Ptr)
+extern "C" AUTORTFM_NO_ASAN void autortfm_record_write_8(void* Ptr)
 {
 	// check for writes to null here so we end up crashing in the user
 	// code rather than in the autortfm runtime.
diff --git a/Engine/Source/Runtime/Core/Private/AutoRTFM/ContextInlines.h b/Engine/Source/Runtime/Core/Private/AutoRTFM/ContextInlines.h
--- a/Engine/Source/Runtime/Core/Private/AutoRTFM/ContextInlines.h
+++ b/Engine/Source/Runtime/Core/Private/AutoRTFM/ContextInlines.h
@@ -19,12 +19,12 @@ UE_AUTORTFM_FORCEINLINE void FContext::CheckOpenRecordWrite(void* LogicalAddress
     ASSERT(CurrentTransaction->IsScopedTransaction() || !IsInnerTransactionStack(LogicalAddress));
 }
 
-UE_AUTORTFM_FORCEINLINE void FContext::RecordWrite(void* LogicalAddress, size_t Size)
+AUTORTFM_NO_ASAN UE_AUTORTFM_FORCEINLINE void FContext::RecordWrite(void* LogicalAddress, size_t Size)
 {
     CurrentTransaction->RecordWrite(LogicalAddress, Size);
 }
 
-template<unsigned SIZE> UE_AUTORTFM_FORCEINLINE void FContext::RecordWrite(void* LogicalAddress)
+template<unsigned SIZE> AUTORTFM_NO_ASAN UE_AUTORTFM_FORCEINLINE void FContext::RecordWrite(void* LogicalAddress)
 {
     CurrentTransaction->RecordWrite<SIZE>(LogicalAddress);
 }
diff --git a/Engine/Source/Runtime/Core/Private/AutoRTFM/TransactionInlines.h b/Engine/Source/Runtime/Core/Private/AutoRTFM/TransactionInlines.h
--- a/Engine/Source/Runtime/Core/Private/AutoRTFM/TransactionInlines.h
+++ b/Engine/Source/Runtime/Core/Private/AutoRTFM/TransactionInlines.h
@@ -5,18 +5,34 @@
 #include "Context.h"
 #include "Transaction.h"
 
+#include "HAL/Platform.h"
+
+#if PLATFORM_HAS_ASAN_INCLUDE && USING_ADDRESS_SANITISER
+#include <sanitizer/asan_interface.h>
+#endif
+
 namespace AutoRTFM
 {
 
-UE_AUTORTFM_FORCEINLINE void FTransaction::RecordWriteMaxPageSized(void* LogicalAddress, size_t Size)
+AUTORTFM_NO_ASAN UE_AUTORTFM_FORCEINLINE void FTransaction::RecordWriteMaxPageSized(void* LogicalAddress, size_t Size)
 {
+#if PLATFORM_HAS_ASAN_INCLUDE && USING_ADDRESS_SANITISER
+    // TODO(SOL-5123): Can we detect shadow memory locations at compile-time instead?
+	const char* const Location = __asan_locate_address(LogicalAddress, nullptr, 0, nullptr, nullptr);
+
+    if (strstr(Location, "shadow"))
+	{
+		return;
+	}
+#endif
+
     void* CopyAddress = WriteLogBumpAllocator.Allocate(Size);
     memcpy(CopyAddress, LogicalAddress, Size);
 
     WriteLog.Push(FWriteLogEntry(LogicalAddress, Size, CopyAddress));
 }
 
-UE_AUTORTFM_FORCEINLINE void FTransaction::RecordWrite(void* LogicalAddress, size_t Size)
+AUTORTFM_NO_ASAN UE_AUTORTFM_FORCEINLINE void FTransaction::RecordWrite(void* LogicalAddress, size_t Size)
 {
     // If we are recording a stack address that is relative to our current
     // transactions stack location, we do not need to record the data in the
@@ -66,7 +82,7 @@ UE_AUTORTFM_FORCEINLINE void FTransaction::RecordWrite(void* LogicalAddress, siz
     RecordWriteMaxPageSized(Address + I, Size - I);
 }
 
-template<unsigned SIZE> UE_AUTORTFM_FORCEINLINE void FTransaction::RecordWrite(void* LogicalAddress)
+template<unsigned SIZE> AUTORTFM_NO_ASAN UE_AUTORTFM_FORCEINLINE void FTransaction::RecordWrite(void* LogicalAddress)
 {
     static_assert(SIZE <= FWriteLogBumpAllocator::MaxSize);
 
@@ -130,12 +146,20 @@ UE_AUTORTFM_FORCEINLINE void FTransaction::DidFree(void* LogicalAddress)
 
 UE_AUTORTFM_FORCEINLINE void FTransaction::DeferUntilCommit(TFunction<void()>&& Callback)
 {
-    CommitTasks.Add(MoveTemp(Callback));
+	// We explicitly must copy the function here because the original was allocated
+	// within a transactional context, and thus the memory is allocating under
+	// transactionalized conditions. By copying, we create an open copy of the callback.
+	TFunction<void()> Copy(Callback);
+    CommitTasks.Add(MoveTemp(Copy));
 }
 
 UE_AUTORTFM_FORCEINLINE void FTransaction::DeferUntilAbort(TFunction<void()>&& Callback)
 {
-    AbortTasks.Add(MoveTemp(Callback));
+	// We explicitly must copy the function here because the original was allocated
+	// within a transactional context, and thus the memory is allocating under
+	// transactionalized conditions. By copying, we create an open copy of the callback.
+	TFunction<void()> Copy(Callback);
+    AbortTasks.Add(MoveTemp(Copy));
 }
 
 UE_AUTORTFM_FORCEINLINE void FTransaction::CollectStats() const
diff --git a/Engine/Source/Runtime/Core/Private/AutoRTFM/Utils.h b/Engine/Source/Runtime/Core/Private/AutoRTFM/Utils.h
--- a/Engine/Source/Runtime/Core/Private/AutoRTFM/Utils.h
+++ b/Engine/Source/Runtime/Core/Private/AutoRTFM/Utils.h
@@ -46,3 +46,10 @@ template<size_t A, size_t B> struct PrettyStaticAssert final
 } // namespace AutoRTFM
 
 #define ASSERT(exp) do { if (UNLIKELY(!(exp))) { UE_DEBUG_BREAK(); PrettyAbort(__FILE__, __LINE__, __PRETTY_FUNCTION__, #exp); } } while (false)
+
+
+#if defined(__has_feature) && __has_feature(address_sanitizer)
+#define AUTORTFM_NO_ASAN [[clang::no_sanitize("address")]]
+#else
+#define AUTORTFM_NO_ASAN
+#endif
diff --git a/Engine/Source/Runtime/Core/Private/HAL/MallocAnsi.cpp b/Engine/Source/Runtime/Core/Private/HAL/MallocAnsi.cpp
--- a/Engine/Source/Runtime/Core/Private/HAL/MallocAnsi.cpp
+++ b/Engine/Source/Runtime/Core/Private/HAL/MallocAnsi.cpp
@@ -221,12 +221,8 @@ bool FMallocAnsi::GetAllocationSize( void *Original, SIZE_T &SizeOut )
 		return false;
 	}
 
-#if	PLATFORM_USES__ALIGNED_MALLOC
-	return false;
-#else
 	SizeOut = AnsiGetAllocationSize(Original);
 	return true;
-#endif
 }
 
 bool FMallocAnsi::IsInternallyThreadSafe() const
diff --git a/Engine/Source/Runtime/Core/Public/Modules/Boilerplate/ModuleBoilerplate.h b/Engine/Source/Runtime/Core/Public/Modules/Boilerplate/ModuleBoilerplate.h
--- a/Engine/Source/Runtime/Core/Public/Modules/Boilerplate/ModuleBoilerplate.h
+++ b/Engine/Source/Runtime/Core/Public/Modules/Boilerplate/ModuleBoilerplate.h
@@ -33,7 +33,11 @@ class FChunkedFixedUObjectArray;
 
 // Disable the replacement new/delete when running the Clang static analyzer, due to false positives in 15.0.x:
 // [REDACTED_PROFILE_URL]/llvm-project/issues/58820
-#if !FORCE_ANSI_ALLOCATOR && !defined(__clang_analyzer__)
+// For AutoRTFM when FORCE_ANSI_ALLOCATOR is specified we still need to re-route operator new/delete (to deal
+// with transactionalization of memory allocation). Address sanitizer will still work even with this re-routing
+// because we point the underlying FMemory allocator at ansi malloc/free, which ASan still hijacks for its
+// purposes.
+#if !(FORCE_ANSI_ALLOCATOR && !(defined(__AUTORTFM) && __AUTORTFM)) && !defined(__clang_analyzer__)
 static_assert(__STDCPP_DEFAULT_NEW_ALIGNMENT__ <= 16, "Expecting 16-byte default operator new alignment - alignments > 16 may have bloat");
 #define REPLACEMENT_OPERATOR_NEW_AND_DELETE \
 	OPERATOR_NEW_MSVC_PRAGMA void* operator new  ( size_t Size                                                    ) OPERATOR_NEW_THROW_SPEC      { return FMemory::Malloc( Size ? Size : 1, __STDCPP_DEFAULT_NEW_ALIGNMENT__ ); } \
