Fix UB sanitizer errors:

- signed integer overflow in MaterialInstanceSupport.h
- unitialized fast path load buffer causing applying non-zero offset to nullptr in Archive.h
- offset to nullptr on MemStack.h

diff --git a/Engine/Source/Runtime/Core/Public/Misc/MemStack.h b/Engine/Source/Runtime/Core/Public/Misc/MemStack.h
--- a/Engine/Source/Runtime/Core/Public/Misc/MemStack.h
+++ b/Engine/Source/Runtime/Core/Public/Misc/MemStack.h
@@ -122,8 +122,8 @@ public:
 		check(!bShouldEnforceAllocMarks || NumMarks > 0);
 
 		// Try to get memory from the current chunk.
-		uint8* Result = Align( Top, Alignment );
-		uint8* NewTop = Result + AllocSize;
+		SIZE_T Result = (SIZE_T)(Align( Top, Alignment ));
+		uint8* NewTop = (uint8*)(Result + AllocSize);
 
 		// Make sure we didn't overflow.
 		if ( NewTop <= End )
@@ -134,11 +134,11 @@ public:
 		{
 			// We'd pass the end of the current chunk, so allocate a new one.
 			AllocateNewChunk( AllocSize + Alignment );
-			Result = Align( Top, Alignment );
-			NewTop = Result + AllocSize;
+			Result = (SIZE_T)(Align( Top, Alignment ));
+			NewTop = (uint8*)(Result + AllocSize);
 			Top = NewTop;
 		}
-		return Result;
+		return (uint8*)(Result);
 	}
 
 	/** return true if this stack is empty. */
diff --git a/Engine/Source/Runtime/Core/Public/Serialization/Archive.h b/Engine/Source/Runtime/Core/Public/Serialization/Archive.h
--- a/Engine/Source/Runtime/Core/Public/Serialization/Archive.h
+++ b/Engine/Source/Runtime/Core/Public/Serialization/Archive.h
@@ -749,6 +749,10 @@ public:
 			EndFastPathLoadBuffer = nullptr;
 			OriginalFastPathLoadBuffer = nullptr;
 		}
+		FORCEINLINE bool IsValid() const
+		{
+			return ((StartFastPathLoadBuffer != nullptr) && (EndFastPathLoadBuffer != nullptr) && (OriginalFastPathLoadBuffer != nullptr));
+		}
 	};
 #if DEVIRTUALIZE_FLinkerLoad_Serialize
 	//@todoio FArchive is really a horrible class and the way it is proxied by FLinkerLoad is double terrible. It makes the fast path really hacky and slower than it would need to be.
@@ -1431,7 +1435,8 @@ public:
 		// Serialize bool as if it were UBOOL (legacy, 32 bit int).
 #if DEVIRTUALIZE_FLinkerLoad_Serialize
 		const uint8 * RESTRICT Src = Ar.ActiveFPLB->StartFastPathLoadBuffer;
-		if (Src + sizeof(uint32) <= Ar.ActiveFPLB->EndFastPathLoadBuffer)
+		const bool bValidFPLBAccess = Ar.ActiveFPLB->IsValid() && ((Src + sizeof(uint32)) <= Ar.ActiveFPLB->EndFastPathLoadBuffer);
+		if (bValidFPLBAccess)
 		{
 			D = !!FPlatformMemory::ReadUnaligned<uint32>(Src);
 			Ar.ActiveFPLB->StartFastPathLoadBuffer += 4;
@@ -1970,6 +1975,11 @@ public:
 	template<SIZE_T Size>
 	FORCEINLINE bool FastPathLoad(void* InDest)
 	{
+		if (!ActiveFPLB->IsValid())
+		{
+			return false;
+		}
+
 		const uint8* RESTRICT Src = ActiveFPLB->StartFastPathLoadBuffer;
 		if (Src + Size <= ActiveFPLB->EndFastPathLoadBuffer)
 		{
diff --git a/Engine/Source/Runtime/Engine/Private/Materials/MaterialInstanceSupport.h b/Engine/Source/Runtime/Engine/Private/Materials/MaterialInstanceSupport.h
--- a/Engine/Source/Runtime/Engine/Private/Materials/MaterialInstanceSupport.h
+++ b/Engine/Source/Runtime/Engine/Private/Materials/MaterialInstanceSupport.h
@@ -129,14 +129,14 @@ public:
 		// We expect the Name to be the only part of the key to vary in 99.995% of keys.  If this changes, we could make
 		// the hash take into account more fields than just Name, at some cost in perf generating the hash.  It's important
 		// to mix the key bits, because GetTypeHash(FScriptName) just returns an index, which is not good for a hash table.
-		int32 Key = GetTypeHash(Info.Name);
+		uint32 Key = GetTypeHash(Info.Name);
 		Key += ~(Key << 15);
 		Key ^= (Key >> 10);
 		Key += (Key << 3);
 		Key ^= (Key >> 6);
 		Key += ~(Key << 11);
 		Key ^= (Key >> 16);
-		return (uint32)Key;
+		return Key;
 	}
 
 	// Add an item from the Array to the hash table.  Hash table must already be allocated.
diff --git a/Engine/Source/Runtime/Experimental/Chaos/Public/Chaos/TriangleMeshImplicitObject.h b/Engine/Source/Runtime/Experimental/Chaos/Public/Chaos/TriangleMeshImplicitObject.h
--- a/Engine/Source/Runtime/Experimental/Chaos/Public/Chaos/TriangleMeshImplicitObject.h
+++ b/Engine/Source/Runtime/Experimental/Chaos/Public/Chaos/TriangleMeshImplicitObject.h
@@ -445,7 +445,7 @@ namespace Chaos
 	private:
 		TArray<TVec3<LargeIdxType>> LargeIdxBuffer;
 		TArray<TVec3<SmallIdxType>> SmallIdxBuffer;
-		bool bRequiresLargeIndices;
+		bool bRequiresLargeIndices = false;
 	};
 
 	FORCEINLINE_DEBUGGABLE FArchive& operator<<(FArchive& Ar, FTrimeshIndexBuffer& Buffer)
