[Automation] Overloaded TestValid/TestInvalid to accept anything with an IsValid member

diff --git a/Engine/Source/Runtime/Core/Public/Misc/AutomationTest.h b/Engine/Source/Runtime/Core/Public/Misc/AutomationTest.h
--- a/Engine/Source/Runtime/Core/Public/Misc/AutomationTest.h
+++ b/Engine/Source/Runtime/Core/Public/Misc/AutomationTest.h
@@ -2,7 +2,6 @@
 
 #pragma once
 
-#include "CoreTypes.h"
 #include "Async/Async.h"
 #include "Async/Future.h"
 #include "Containers/Array.h"
@@ -10,6 +9,7 @@
 #include "Containers/Queue.h"
 #include "Containers/Set.h"
 #include "Containers/UnrealString.h"
+#include "CoreTypes.h"
 #include "Delegates/Delegate.h"
 #include "Delegates/DelegateBase.h"
 #include "Delegates/DelegateInstancesImpl.h"
@@ -25,7 +25,6 @@
 #include "Internationalization/Regex.h"
 #include "Logging/LogVerbosity.h"
 #include "Math/Color.h"
-#include "Math/Color.h"
 #include "Math/MathFwd.h"
 #include "Math/Rotator.h"
 #include "Math/UnrealMathUtility.h"
@@ -1894,27 +1893,28 @@ public:
 	}
 
 	/**
-	 * Logs an error if the given shared pointer is valid.
+	 * Logs an error if the given object tests true when calling its IsValid member.
 	 *
 	 * @param Description - Description text for the test.
-	 * @param SharedPointer - The shared pointer to test.
+	 * @param Value - The value to test.
 	 *
 	 * @see TestValid
 	 */
 	template<typename ValueType>
-	FORCEINLINE bool TestInvalid(const TCHAR* Description, const TSharedPtr<ValueType>& SharedPointer)
+	FORCEINLINE bool TestInvalid(const TCHAR* Description, const ValueType& Value)
 	{
-		if (SharedPointer.IsValid())
+		if (Value.IsValid())
 		{
-			AddError(FString::Printf(TEXT("%s: The shared pointer is valid."), Description));
+			AddError(FString::Printf(TEXT("%s: The value is valid (.IsValid() returned true)."), Description));
 			return false;
 		}
 		return true;
 	}
 
-	template<typename ValueType> bool TestInvalid(const FString& Description, const TSharedPtr<ValueType>& SharedPointer)
+	template<typename ValueType>
+	bool TestInvalid(const FString& Description, const ValueType& Value)
 	{
-		return TestInvalid(*Description, SharedPointer);
+		return TestInvalid(*Description, Value);
 	}
 
 	/**
@@ -2050,27 +2050,28 @@ public:
 	#define TestTrueExpr(Expression) TestTrue(TEXT(#Expression), Expression)
 
 	/**
-	 * Logs an error if the given shared pointer is not valid.
+	 * Logs an error if the given object returns false when calling its IsValid member.
 	 *
 	 * @param Description - Description text for the test.
-	 * @param SharedPointer - The shared pointer to test.
+	 * @param Value - The value to test.
 	 *
 	 * @see TestInvalid
 	 */
 	template<typename ValueType>
-	FORCEINLINE bool TestValid(const TCHAR* Description, const TSharedPtr<ValueType>& SharedPointer)
+	FORCEINLINE bool TestValid(const TCHAR* Description, const ValueType& Value)
 	{
-		if (!SharedPointer.IsValid())
+		if (!Value.IsValid())
 		{
-			AddError(FString::Printf(TEXT("%s: The shared pointer is not valid."), Description));
+			AddError(FString::Printf(TEXT("%s: The value is not valid (.IsValid() returned false)."), Description));
 			return false;
 		}
 		return true;
 	}
 
-	template<typename ValueType> bool TestValid(const FString& Description, const TSharedPtr<ValueType>& SharedPointer)
+	template<typename ValueType>
+	bool TestValid(const FString& Description, const ValueType& Value)
 	{
-		return TestValid(*Description, SharedPointer);
+		return TestValid(*Description, Value);
 	}
 
 protected:
@@ -3905,26 +3906,26 @@ public: \
 		return false;\
 	}
 
-#define UTEST_VALID(What, SharedPointer)\
-	if (!TestValid(What, SharedPointer))\
+#define UTEST_VALID(What, Value)\
+	if (!TestValid(What, Value))\
 	{\
 		return false;\
 	}
 
-#define UTEST_VALID_EXPR(SharedPointer)\
-	if (!TestValid(TEXT(#SharedPointer), SharedPointer))\
+#define UTEST_VALID_EXPR(Value)\
+	if (!TestValid(TEXT(#Value), Value))\
 	{\
 		return false;\
 	}
 
-#define UTEST_INVALID(What, SharedPointer)\
-	if (!TestInvalid(What, SharedPointer))\
+#define UTEST_INVALID(What, Value)\
+	if (!TestInvalid(What, Value))\
 	{\
 		return false;\
 	}
 
-#define UTEST_INVALID_EXPR(SharedPointer)\
-	if (!TestInvalid(TEXT(#SharedPointer), SharedPointer))\
+#define UTEST_INVALID_EXPR(Value)\
+	if (!TestInvalid(TEXT(#Value), Value))\
 	{\
 		return false;\
 	}
diff --git a/Engine/Source/Runtime/Engine/Private/Tests/EngineAutomationTests.cpp b/Engine/Source/Runtime/Engine/Private/Tests/EngineAutomationTests.cpp
--- a/Engine/Source/Runtime/Engine/Private/Tests/EngineAutomationTests.cpp
+++ b/Engine/Source/Runtime/Engine/Private/Tests/EngineAutomationTests.cpp
@@ -897,14 +897,31 @@ bool FAutomationTrueFalseEXPR::RunTest(const FString& Parameters)
 IMPLEMENT_CUSTOM_SIMPLE_AUTOMATION_TEST(FAutomationValidInvalidEXPR, FAutomationUTestMacrosExpr, "TestFramework.Validation.UTestValidInvalid", EAutomationTestFlags::ApplicationContextMask | EAutomationTestFlags::EngineFilter);
 bool FAutomationValidInvalidEXPR::RunTest(const FString& Parameters)
 {
-	struct HandleResumeContext
+	struct FHasIsValid
 	{
-		const char* PackageIdentifier = nullptr;
+		explicit FHasIsValid(bool bInIsValid)
+			: bIsValid(bInIsValid)
+		{  }
+		
+		bool IsValid() const { return bIsValid; }
+
+	private:
+		bool bIsValid;
 	};
-	TSharedPtr<HandleResumeContext, ESPMode::ThreadSafe> Context = MakeShared<HandleResumeContext, ESPMode::ThreadSafe>();
-	TSharedPtr< UObject > Entry = nullptr;
-	UTEST_VALID_EXPR(Context);
-	UTEST_INVALID_EXPR(Entry);
+	
+	//** TEST **//
+	TSharedPtr<FVector> ValidSharedPtr = MakeShared<FVector>();
+	TSharedPtr<UObject> InvalidSharedPtr = nullptr;
+
+	FHasIsValid ValidObject(true);
+	FHasIsValid InvalidObject(false);
+
+	//** VERIFY **//
+	UTEST_VALID_EXPR(ValidSharedPtr);
+	UTEST_INVALID_EXPR(InvalidSharedPtr);
+
+	UTEST_VALID_EXPR(ValidObject);
+	UTEST_INVALID_EXPR(InvalidObject);
 
 	return true;
 }
