Enable reporting of "CustomTag" when FName based tags are disabled.

diff --git a/Engine/Source/Runtime/Core/Private/HAL/LowLevelMemTracker.cpp b/Engine/Source/Runtime/Core/Private/HAL/LowLevelMemTracker.cpp
--- a/Engine/Source/Runtime/Core/Private/HAL/LowLevelMemTracker.cpp
+++ b/Engine/Source/Runtime/Core/Private/HAL/LowLevelMemTracker.cpp
@@ -1587,10 +1587,12 @@ void FLowLevelMemTracker::BootstrapTagDatas()
 		LLM_ENUM_GENERIC_TAGS(REGISTER_ELLMTAG);
 #undef REGISTER_ELLMTAG
 
+#if LLM_ALLOW_NAMES_TAGS
 		// The CustomName tag is an adapter for connecting LLM_SCOPE_BYNAME tags with platforms that used the
 		// ELLMTag-based reporting. We want to hide this procedural tag in systems that use FName-based reporting;
 		// if it is displayed it confusingly just displays a sum of every LLM_SCOPE_BYNAME tag.
 		TagDataEnumMap[static_cast<int32>(ELLMTag::CustomName)]->SetIsReportable(false);
+#endif
 	}
 }
 
