- Updated TFormatSpecifier to templatize GetFormatSpecifier on CharType.
- Added TIsFormattedStringBuilderType_V trait to StringBuilder.h to specify the types for which TFormatSpecifier should be used for.
- Replaced string builder numeric append operator overloads with a single templated one, based on what we do for LexToString. Provides consistency and adds coverage for types we haven't yet provided overloads for (long, unsigned long). Uses TIsFormattedStringBuilderType_V as a requirement to use.
- Added an overload for bool via a template which requires the template type to be bool, which prevents things like pointers being printed as bools.

diff --git a/Engine/Source/Runtime/Core/Public/Misc/StringBuilder.h b/Engine/Source/Runtime/Core/Public/Misc/StringBuilder.h
--- a/Engine/Source/Runtime/Core/Public/Misc/StringBuilder.h
+++ b/Engine/Source/Runtime/Core/Public/Misc/StringBuilder.h
@@ -473,45 +473,50 @@ inline FStringBuilderBase&			operator+=(FStringBuilderBase& Builder, UTF8CHAR Ch
 inline FStringBuilderBase&			operator+=(FStringBuilderBase& Builder, FWideStringView Str)						{ return Builder.Append(Str); }
 inline FStringBuilderBase&			operator+=(FStringBuilderBase& Builder, FUtf8StringView Str)						{ return Builder.Append(Str); }
 
-// Integer Append Operators
-
-inline FAnsiStringBuilderBase&		operator<<(FAnsiStringBuilderBase& Builder, int32 Value)							{ return Builder.Appendf("%d", Value); }
-inline FAnsiStringBuilderBase&		operator<<(FAnsiStringBuilderBase& Builder, uint32 Value)							{ return Builder.Appendf("%u", Value); }
-inline FWideStringBuilderBase&		operator<<(FWideStringBuilderBase& Builder, int32 Value)							{ return Builder.Appendf(WIDETEXT("%d"), Value); }
-inline FWideStringBuilderBase&		operator<<(FWideStringBuilderBase& Builder, uint32 Value)							{ return Builder.Appendf(WIDETEXT("%u"), Value); }
-inline FUtf8StringBuilderBase&		operator<<(FUtf8StringBuilderBase& Builder, int32 Value)							{ return Builder.Appendf(UTF8TEXT("%d"), Value); }
-inline FUtf8StringBuilderBase&		operator<<(FUtf8StringBuilderBase& Builder, uint32 Value)							{ return Builder.Appendf(UTF8TEXT("%u"), Value); }
-
-inline FAnsiStringBuilderBase&		operator<<(FAnsiStringBuilderBase& Builder, int64 Value)							{ return Builder.Appendf("%" INT64_FMT, Value); }
-inline FAnsiStringBuilderBase&		operator<<(FAnsiStringBuilderBase& Builder, uint64 Value)							{ return Builder.Appendf("%" UINT64_FMT, Value); }
-inline FWideStringBuilderBase&		operator<<(FWideStringBuilderBase& Builder, int64 Value)							{ return Builder.Appendf(WIDETEXT("%" INT64_FMT), Value); }
-inline FWideStringBuilderBase&		operator<<(FWideStringBuilderBase& Builder, uint64 Value)							{ return Builder.Appendf(WIDETEXT("%" UINT64_FMT), Value); }
-inline FUtf8StringBuilderBase&		operator<<(FUtf8StringBuilderBase& Builder, int64 Value)							{ return Builder.Appendf(UTF8TEXT("%" INT64_FMT), Value); }
-inline FUtf8StringBuilderBase&		operator<<(FUtf8StringBuilderBase& Builder, uint64 Value)							{ return Builder.Appendf(UTF8TEXT("%" UINT64_FMT), Value); }
-
-inline FAnsiStringBuilderBase&		operator<<(FAnsiStringBuilderBase& Builder, int8 Value)								{ return Builder << int32(Value); }
-inline FAnsiStringBuilderBase&		operator<<(FAnsiStringBuilderBase& Builder, uint8 Value)							{ return Builder << uint32(Value); }
-inline FWideStringBuilderBase&		operator<<(FWideStringBuilderBase& Builder, int8 Value)								{ return Builder << int32(Value); }
-inline FWideStringBuilderBase&		operator<<(FWideStringBuilderBase& Builder, uint8 Value)							{ return Builder << uint32(Value); }
-inline FUtf8StringBuilderBase&		operator<<(FUtf8StringBuilderBase& Builder, int8 Value)								{ return Builder << int32(Value); }
-inline FUtf8StringBuilderBase&		operator<<(FUtf8StringBuilderBase& Builder, uint8 Value)							{ return Builder << uint32(Value); }
-
-inline FAnsiStringBuilderBase&		operator<<(FAnsiStringBuilderBase& Builder, int16 Value)							{ return Builder << int32(Value); }
-inline FAnsiStringBuilderBase&		operator<<(FAnsiStringBuilderBase& Builder, uint16 Value)							{ return Builder << uint32(Value); }
-inline FWideStringBuilderBase&		operator<<(FWideStringBuilderBase& Builder, int16 Value)							{ return Builder << int32(Value); }
-inline FWideStringBuilderBase&		operator<<(FWideStringBuilderBase& Builder, uint16 Value)							{ return Builder << uint32(Value); }
-inline FUtf8StringBuilderBase&		operator<<(FUtf8StringBuilderBase& Builder, int16 Value)							{ return Builder << int32(Value); }
-inline FUtf8StringBuilderBase&		operator<<(FUtf8StringBuilderBase& Builder, uint16 Value)							{ return Builder << uint32(Value); }
-
-// Floating Point Append Operators
-
-inline FAnsiStringBuilderBase&		operator<<(FAnsiStringBuilderBase& Builder, float Value)							{ return Builder.Appendf("%f", Value); }
-inline FWideStringBuilderBase&		operator<<(FWideStringBuilderBase& Builder, float Value)							{ return Builder.Appendf(WIDETEXT("%f"), Value); }
-inline FUtf8StringBuilderBase&		operator<<(FUtf8StringBuilderBase& Builder, float Value)							{ return Builder.Appendf(UTF8TEXT("%f"), Value); }
-
-inline FAnsiStringBuilderBase&		operator<<(FAnsiStringBuilderBase& Builder, double Value)							{ return Builder.Appendf("%f", Value); }
-inline FWideStringBuilderBase&		operator<<(FWideStringBuilderBase& Builder, double Value)							{ return Builder.Appendf(WIDETEXT("%f"), Value); }
-inline FUtf8StringBuilderBase&		operator<<(FUtf8StringBuilderBase& Builder, double Value)							{ return Builder.Appendf(UTF8TEXT("%f"), Value); }
+// Bool Append Operators
+
+template <typename T UE_REQUIRES(std::is_same_v<bool, T>)>
+inline FAnsiStringBuilderBase&		operator<<(FAnsiStringBuilderBase& Builder, T Value)								{ return Builder.Append(Value ? "true" : "false"); }
+template <typename T UE_REQUIRES(std::is_same_v<bool, T>)>
+inline FWideStringBuilderBase&		operator<<(FWideStringBuilderBase& Builder, T Value)								{ return Builder.Append(Value ? WIDETEXT("true") : WIDETEXT("false")); }
+template <typename T UE_REQUIRES(std::is_same_v<bool, T>)>
+inline FUtf8StringBuilderBase&		operator<<(FUtf8StringBuilderBase& Builder, T Value)								{ return Builder.Append(Value ? UTF8TEXT("true") : UTF8TEXT("false")); }
+
+/**
+ * Trait which determines whether or not a type can be appended to TStringBuilderBase via TFormatSpecifier.
+ */
+
+template <typename T>        constexpr bool TIsFormattedStringBuilderType_V                   = false;
+template <>           inline constexpr bool TIsFormattedStringBuilderType_V<uint8>            = true;
+template <>           inline constexpr bool TIsFormattedStringBuilderType_V<uint16>           = true;
+template <>           inline constexpr bool TIsFormattedStringBuilderType_V<uint32>           = true;
+template <>           inline constexpr bool TIsFormattedStringBuilderType_V<uint64>           = true;
+template <>           inline constexpr bool TIsFormattedStringBuilderType_V<int8>             = true;
+template <>           inline constexpr bool TIsFormattedStringBuilderType_V<int16>            = true;
+template <>           inline constexpr bool TIsFormattedStringBuilderType_V<int32>            = true;
+template <>           inline constexpr bool TIsFormattedStringBuilderType_V<int64>            = true;
+template <>           inline constexpr bool TIsFormattedStringBuilderType_V<float>            = true;
+template <>           inline constexpr bool TIsFormattedStringBuilderType_V<double>           = true;
+template <>           inline constexpr bool TIsFormattedStringBuilderType_V<long double>      = true;
+template <>           inline constexpr bool TIsFormattedStringBuilderType_V<long>             = true;
+template <>           inline constexpr bool TIsFormattedStringBuilderType_V<unsigned long>    = true;
+
+template <typename T>        constexpr bool TIsFormattedStringBuilderType_V<const          T> = TIsFormattedStringBuilderType_V<T>;
+template <typename T>        constexpr bool TIsFormattedStringBuilderType_V<      volatile T> = TIsFormattedStringBuilderType_V<T>;
+template <typename T>        constexpr bool TIsFormattedStringBuilderType_V<const volatile T> = TIsFormattedStringBuilderType_V<T>;
+
+// Formatted Append Operators
+
+template <
+	typename CharType,
+	typename T
+	UE_REQUIRES(TIsFormattedStringBuilderType_V<T>)
+>
+inline TStringBuilderBase<CharType>& operator<<(TStringBuilderBase<CharType>& Builder, const T& Value)
+{
+	// std::remove_cv_t to remove potential volatile decorations. Removing const is pointless, but harmless because it's specified in the param declaration.
+	return Builder.Appendf(TFormatSpecifier<std::remove_cv_t<T>>::template GetFormatSpecifier<CharType>(), Value);
+}
 
 template <typename CharType, int32 BufferSize>
 class UE_DEPRECATED(5.3, "Use WriteToString<N>(...) or TStringBuilder<N>(InPlace, ...).") TWriteToString : public TStringBuilderWithBuffer<CharType, BufferSize>
diff --git a/Engine/Source/Runtime/Core/Public/Templates/UnrealTypeTraits.h b/Engine/Source/Runtime/Core/Public/Templates/UnrealTypeTraits.h
--- a/Engine/Source/Runtime/Core/Public/Templates/UnrealTypeTraits.h
+++ b/Engine/Source/Runtime/Core/Public/Templates/UnrealTypeTraits.h
@@ -117,7 +117,8 @@ struct TNthTypeFromParameterPack<0, T, OtherTypes...>
  */
 template<typename T> 
 struct TFormatSpecifier
-{ 
+{
+	template <typename CharType = TCHAR>
 	FORCEINLINE static TCHAR const* GetFormatSpecifier()
 	{
 		// Force the template instantiation to be dependent upon T so the compiler cannot automatically decide that this template can never be instantiated.
@@ -132,9 +133,25 @@ struct TFormatSpecifier
 template<> \
 struct TFormatSpecifier<type> \
 {  \
+	template <typename CharType = TCHAR> \
 	FORCEINLINE static decltype(auto) GetFormatSpecifier() \
 	{ \
-		return TEXT(format); \
+        if constexpr (std::is_same_v<CharType, WIDECHAR>) \
+        { \
+            return WIDETEXT(format); \
+        } \
+        else if constexpr (std::is_same_v<CharType, UTF8CHAR>) \
+        { \
+            return UTF8TEXT(format); \
+        } \
+		else if constexpr (std::is_same_v<CharType, ANSICHAR>) \
+        { \
+            return format; \
+        } \
+        else \
+        { \
+            static_assert(sizeof(CharType) == 0, "Unsupported char type"); \
+        } \
 	} \
 };
 
