disable reserve scheduler by default as it seems to be causing deadlocks on low-core windows

diff --git a/Engine/Programs/UnrealPak/Config/DefaultEngine.ini b/Engine/Programs/UnrealPak/Config/DefaultEngine.ini
--- a/Engine/Programs/UnrealPak/Config/DefaultEngine.ini
+++ b/Engine/Programs/UnrealPak/Config/DefaultEngine.ini
@@ -1,2 +0,0 @@
-[TaskGraph]
-EnableReserveWorkers=false
diff --git a/Engine/Programs/ZenDashboard/Config/DefaultEngine.ini b/Engine/Programs/ZenDashboard/Config/DefaultEngine.ini
--- a/Engine/Programs/ZenDashboard/Config/DefaultEngine.ini
+++ b/Engine/Programs/ZenDashboard/Config/DefaultEngine.ini
@@ -1,2 +0,0 @@
-[TaskGraph]
-EnableReserveWorkers=false
diff --git a/Engine/Source/Runtime/Core/Private/Async/TaskGraph.cpp b/Engine/Source/Runtime/Core/Private/Async/TaskGraph.cpp
--- a/Engine/Source/Runtime/Core/Private/Async/TaskGraph.cpp
+++ b/Engine/Source/Runtime/Core/Private/Async/TaskGraph.cpp
@@ -2319,7 +2319,7 @@ private:
 
 		check(GConfig);
 		
-		bool bEnableReserveWorkers = true; // by default
+		bool bEnableReserveWorkers = false; // by default
 		GConfig->GetBool(TEXT("TaskGraph"), TEXT("EnableReserveWorkers"), bEnableReserveWorkers, GEngineIni);
 
 		if (bEnableReserveWorkers)
