fix race condition where std::condition_variable is calling notify on thread A and is being destroyed at the same time on thread B

#rnx

diff --git a/Engine/Source/Runtime/Core/Private/GenericPlatform/GenericPlatformManualResetEvent.h b/Engine/Source/Runtime/Core/Private/GenericPlatform/GenericPlatformManualResetEvent.h
--- a/Engine/Source/Runtime/Core/Private/GenericPlatform/GenericPlatformManualResetEvent.h
+++ b/Engine/Source/Runtime/Core/Private/GenericPlatform/GenericPlatformManualResetEvent.h
@@ -79,10 +79,8 @@ public:
 	 */
 	void Notify()
 	{
-		{
-			std::unique_lock SelfLock(Lock);
-			bWait = false;
-		}
+		std::unique_lock SelfLock(Lock);
+		bWait = false;
 		Condition.notify_one();
 	}
 
