Correctly report Windows long path support under Wine

diff --git a/Engine/Source/Runtime/Core/Private/Windows/WindowsPlatformMisc.cpp b/Engine/Source/Runtime/Core/Private/Windows/WindowsPlatformMisc.cpp
--- a/Engine/Source/Runtime/Core/Private/Windows/WindowsPlatformMisc.cpp
+++ b/Engine/Source/Runtime/Core/Private/Windows/WindowsPlatformMisc.cpp
@@ -770,7 +770,15 @@ int32 FWindowsPlatformMisc::GetMaxPathLength()
 			{
 				typedef BOOLEAN(NTAPI *RtlAreLongPathsEnabledFunc)();
 				RtlAreLongPathsEnabledFunc RtlAreLongPathsEnabled = (RtlAreLongPathsEnabledFunc)(void*)GetProcAddress(Handle, "RtlAreLongPathsEnabled");
-				bValue = (RtlAreLongPathsEnabled != NULL && RtlAreLongPathsEnabled());
+				if (RtlAreLongPathsEnabled != NULL)
+				{
+					bValue = RtlAreLongPathsEnabled();
+				}
+				else
+				{
+					// Long paths are always supported under Wine
+					bValue = FWindowsPlatformMisc::IsWine();
+				}
 			}
 		}
 	};
@@ -1992,6 +2000,31 @@ bool FWindowsPlatformMisc::VerifyWindowsVersion(uint32 MajorVersion, uint32 Mino
 	return !!VerifyVersionInfo(&Version, VER_MAJORVERSION | VER_MINORVERSION | VER_BUILDNUMBER, ConditionMask);
 }
 
+bool FWindowsPlatformMisc::IsWine()
+{
+	struct FWineDetected
+	{
+		bool bValue;
+
+		FWineDetected()
+		{
+			HMODULE Handle = GetModuleHandle(TEXT("ntdll.dll"));
+			if (Handle == NULL)
+			{
+				bValue = false;
+			}
+			else
+			{
+				void* WineGetVersion = (void*)GetProcAddress(Handle, "wine_get_version");
+				bValue = (WineGetVersion != NULL);
+			}
+		}
+	};
+
+	static FWineDetected WineDetected;
+	return WineDetected.bValue;
+}
+
 bool FWindowsPlatformMisc::IsValidAbsolutePathFormat(const FString& Path)
 {
 	bool bIsValid = true;
diff --git a/Engine/Source/Runtime/Core/Public/Windows/WindowsPlatformMisc.h b/Engine/Source/Runtime/Core/Public/Windows/WindowsPlatformMisc.h
--- a/Engine/Source/Runtime/Core/Public/Windows/WindowsPlatformMisc.h
+++ b/Engine/Source/Runtime/Core/Public/Windows/WindowsPlatformMisc.h
@@ -161,6 +161,13 @@ struct FWindowsPlatformMisc
 	 */
 	static CORE_API bool VerifyWindowsVersion(uint32 MajorVersion, uint32 MinorVersion, uint32 BuildNumber = 0);
 
+	/** 
+	 * Determines if we are running under Wine rather than a real version of Windows
+	 *
+	 * @return	Returns true if the current runtime environment is Wine
+	 */
+	static bool IsWine();
+
 #if !UE_BUILD_SHIPPING
 	static CORE_API void PromptForRemoteDebugging(bool bIsEnsure);
 #endif	//#if !UE_BUILD_SHIPPING
