Implement operator<< for TStringBuilderBase<CharType> and FNumericUnit<T>

diff --git a/Engine/Source/Runtime/Core/Public/Math/UnitConversion.h b/Engine/Source/Runtime/Core/Public/Math/UnitConversion.h
--- a/Engine/Source/Runtime/Core/Public/Math/UnitConversion.h
+++ b/Engine/Source/Runtime/Core/Public/Math/UnitConversion.h
@@ -234,6 +234,9 @@ private:
 	}
 };
 
+template <typename CharType, typename T>
+TStringBuilderBase<CharType>& operator<<(TStringBuilderBase<CharType>& Builder, const FNumericUnit<T>& NumericUnit);
+
 template<typename T>
 FString LexToString(const FNumericUnit<T>& NumericUnit);
 
@@ -246,6 +249,5 @@ void LexFromString(FNumericUnit<T>& OutValue, const TCHAR* String);
 template<typename T>
 bool LexTryParseString(FNumericUnit<T>& OutValue, const TCHAR* String);
 
-
 // Include template definitions
 #include "Math/UnitConversion.inl" // IWYU pragma: export
diff --git a/Engine/Source/Runtime/Core/Public/Math/UnitConversion.inl b/Engine/Source/Runtime/Core/Public/Math/UnitConversion.inl
--- a/Engine/Source/Runtime/Core/Public/Math/UnitConversion.inl
+++ b/Engine/Source/Runtime/Core/Public/Math/UnitConversion.inl
@@ -396,14 +396,21 @@ template <typename NumericType>
 struct TNumericLimits<FNumericUnit<NumericType>> : public TNumericLimits<NumericType>
 { };
 
+template <typename CharType, typename T>
+TStringBuilderBase<CharType>& operator<<(TStringBuilderBase<CharType>& Builder, const FNumericUnit<T>& NumericUnit)
+{
+	Builder << NumericUnit.Value;
+	Builder << ANSITEXTVIEW(" ");
+	Builder << FUnitConversion::GetUnitDisplayString(NumericUnit.Units);
+	return Builder;
+}
+
 template<typename T>
 FString LexToString(const FNumericUnit<T>& NumericUnit)
 {
-	FString String = LexToString(NumericUnit.Value);
-	String += TEXT(" ");
-	String += FUnitConversion::GetUnitDisplayString(NumericUnit.Units);
-
-	return String;
+	TStringBuilder<128> Builder;
+	Builder << NumericUnit;
+	return FString(Builder);
 }
 
 template<typename T>
