Do not use sequencial consistency ordering on an FGenericPlatformManualResetEvent::bWait

diff --git a/Engine/Source/Runtime/Core/Private/GenericPlatform/GenericPlatformManualResetEvent.h b/Engine/Source/Runtime/Core/Private/GenericPlatform/GenericPlatformManualResetEvent.h
--- a/Engine/Source/Runtime/Core/Private/GenericPlatform/GenericPlatformManualResetEvent.h
+++ b/Engine/Source/Runtime/Core/Private/GenericPlatform/GenericPlatformManualResetEvent.h
@@ -34,7 +34,7 @@ public:
 	 */
 	void Reset()
 	{
-		bWait = true;
+		bWait.store(true, std::memory_order_release);
 	}
 
 	/**
@@ -60,15 +60,15 @@ public:
 		std::unique_lock SelfLock(Lock);
 		if (WaitTime.IsInfinity())
 		{
-			Condition.wait(SelfLock, [this] { return !bWait; });
+			Condition.wait(SelfLock, [this] { return !bWait.load(std::memory_order_acquire); });
 			return true;
 		}
 		if (FMonotonicTimeSpan WaitSpan = WaitTime - FMonotonicTimePoint::Now(); WaitSpan > FMonotonicTimeSpan::Zero())
 		{
 			const int64 WaitMs = FPlatformMath::CeilToInt64(WaitSpan.ToMilliseconds());
-			return Condition.wait_for(SelfLock, std::chrono::milliseconds(WaitMs), [this] { return !bWait; });
+			return Condition.wait_for(SelfLock, std::chrono::milliseconds(WaitMs), [this] { return !bWait.load(std::memory_order_acquire); });
 		}
-		return !bWait;
+		return !bWait.load(std::memory_order_acquire);
 	}
 
 	/**
@@ -79,8 +79,10 @@ public:
 	 */
 	void Notify()
 	{
+		// We need this lock to ensure wait_for(SelfLock, std::chrono::milliseconds(WaitMs) does not return until Condition.notify_one() is fully called.
+		// Otherwise it's possible to destroy the Event right after Notify() call and while some other thread is still waiting on it
 		std::unique_lock SelfLock(Lock);
-		bWait = false;
+		bWait.store(false, std::memory_order_relaxed);
 		Condition.notify_one();
 	}
 
