Fix checkf print param errors in checkSlow and Mac compils

#rnx

diff --git a/Engine/Source/Runtime/ApplicationCore/Private/Mac/Accessibility/MacAccessibilityManager.cpp b/Engine/Source/Runtime/ApplicationCore/Private/Mac/Accessibility/MacAccessibilityManager.cpp
--- a/Engine/Source/Runtime/ApplicationCore/Private/Mac/Accessibility/MacAccessibilityManager.cpp
+++ b/Engine/Source/Runtime/ApplicationCore/Private/Mac/Accessibility/MacAccessibilityManager.cpp
@@ -151,7 +151,7 @@ static void DumpAccessibilityStatsForwarder()
 {
 	//This should only be called in ~FMacApplication() in a Main Thread Call
 	check([NSThread isMainThread]);
-	checkf([self IsAccessibilityCacheEmpty], TEXT("Accessibility Manager is being torn down and still has %yu elements!"), [self GetAccessibilityCacheSize]);
+	checkf([self IsAccessibilityCacheEmpty], TEXT("Accessibility Manager is being torn down and still has %u elements!"), [self GetAccessibilityCacheSize]);
 	[Cache release];
 	Cache = Nil;
 	//unregister KVO for Voiceover changes
diff --git a/Engine/Source/Runtime/Core/Private/Misc/ConfigContext.cpp b/Engine/Source/Runtime/Core/Private/Misc/ConfigContext.cpp
--- a/Engine/Source/Runtime/Core/Private/Misc/ConfigContext.cpp
+++ b/Engine/Source/Runtime/Core/Private/Misc/ConfigContext.cpp
@@ -550,7 +550,7 @@ void FConfigContext::AddStaticLayersToHierarchy()
 		if (!EnumHasAnyFlags(Layer.Flag, EConfigLayerFlags::NoExpand))
 		{
 			// we assume none of the more special tags in expanded ones
-			checkfSlow(FCString::Strstr(Layer.Path, TEXT("{USERSETTINGS}")) == nullptr && FCString::Strstr(Layer.Path, TEXT("{USER}")) == nullptr, TEXT("Expanded config %s shouldn't have a {USER*} tags in it"), *Layer.Path);
+			checkfSlow(FCString::Strstr(Layer.Path, TEXT("{USERSETTINGS}")) == nullptr && FCString::Strstr(Layer.Path, TEXT("{USER}")) == nullptr, TEXT("Expanded config %s shouldn't have a {USER*} tags in it"), Layer.Path);
 
 			// loop over all the possible expansions
 			for (int32 ExpansionIndex = 0; ExpansionIndex < UE_ARRAY_COUNT(GConfigExpansions); ExpansionIndex++)
@@ -621,7 +621,7 @@ void FConfigContext::AddStaticLayersToHierarchy()
 		// if no expansion, just process the special tags (assume no PLATFORM tags)
 		else
 		{
-			checkfSlow(!bHasPlatformTag, TEXT("Non-expanded config %s shouldn't have a PLATFORM in it"), *Layer.Path);
+			checkfSlow(!bHasPlatformTag, TEXT("Non-expanded config %s shouldn't have a PLATFORM in it"), Layer.Path);
 			checkfSlow(!EnumHasAnyFlags(Layer.Flag, EConfigLayerFlags::AllowCommandLineOverride), TEXT("Non-expanded config can't have a EConfigLayerFlags::AllowCommandLineOverride"));
 
 			FString FinalPath = PerformFinalExpansions(LayerPath, TEXT(""));
