IAS - verify chunk block(s)

#rnx

diff --git a/Engine/Source/Runtime/Core/Private/IO/IoChunkEncoding.cpp b/Engine/Source/Runtime/Core/Private/IO/IoChunkEncoding.cpp
--- a/Engine/Source/Runtime/Core/Private/IO/IoChunkEncoding.cpp
+++ b/Engine/Source/Runtime/Core/Private/IO/IoChunkEncoding.cpp
@@ -245,6 +245,7 @@ bool FIoChunkEncoding::Decode(const FIoChunkDecodingParams& Params, FMemoryView
 	}
 
 	const TConstArrayView<uint32>& EncodedBlockSize = Params.EncodedBlockSize;
+	const TConstArrayView<uint32>& BlockHash = Params.BlockHash;
 	const uint32 BlockSize = Params.BlockSize;
 	const uint64 RawOffset = Params.RawOffset;
 	const uint32 BlockCount = EncodedBlockSize.Num();
@@ -262,6 +263,7 @@ bool FIoChunkEncoding::Decode(const FIoChunkDecodingParams& Params, FMemoryView
 	// Subtract the encoded offset if the encoded blocks is a partial range of all encoded block(s)
 	EncodedBlocks += (EncodedOffset - Params.EncodedOffset);
 
+	const bool bVerifyBlocks = BlockHash.IsEmpty() == false;
 	for (uint32 BlockIndex = FirstBlockIndex; BlockIndex <= LastBlockIndex; BlockIndex++)
 	{
 		const uint64 RawBlockSize = BlockIndex == BlockCount - 1 ? LastRawBlockSize : BlockSize;
@@ -271,6 +273,17 @@ bool FIoChunkEncoding::Decode(const FIoChunkDecodingParams& Params, FMemoryView
 
 		FIoBuffer Tmp(AlignedBlockSize);
 		Tmp.GetMutableView().CopyFrom(EncodedBlocks.Left(AlignedBlockSize));
+
+		if (bVerifyBlocks)
+		{
+			check(BlockHash.IsEmpty() == false);
+			FIoBlockHash Hash = HashBlock(Tmp.GetView());
+			if (Hash != BlockHash[BlockIndex])
+			{
+				return false;
+			}
+		}
+
 		if (AESKey.IsValid())
 		{
 			FAES::DecryptData(Tmp.GetData(), uint32(Tmp.GetSize()), AESKey);
@@ -362,24 +375,10 @@ uint64 FIoChunkEncoding::GetTotalEncodedSize(TConstArrayView<uint32> EncodedBloc
 	return TotalEncodedSize;
 }
 
-FIoStatus FIoChunkEncoding::HashBlocks(const FIoChunkEncoding::FHeader& Header, FMemoryView EncodedData, TArray<FIoHash>& OutHashes)
+FIoBlockHash FIoChunkEncoding::HashBlock(FMemoryView Block)
 {
-	TConstArrayView<uint32> Blocks = Header.GetBlocks();
-	if (Blocks.IsEmpty())
-	{
-		return FIoStatus(EIoErrorCode::InvalidParameter);
-	}
-
-	OutHashes.Reserve(Blocks.Num());
-
-	for (uint32 BlockSize : Blocks)
-	{
-		check(!EncodedData.IsEmpty());
-		OutHashes.Add(FIoHash::HashBuffer(EncodedData.Left(BlockSize)));
-		EncodedData += BlockSize;
-	}
-
-	check(EncodedData.IsEmpty());
-	
-	return FIoStatus::Ok;
+	const FIoHash Hash = FIoHash::HashBuffer(Block);
+	FIoBlockHash BlockHash;
+	FMemory::Memcpy(&BlockHash, &Hash, sizeof(FIoBlockHash));
+	return BlockHash;
 }
diff --git a/Engine/Source/Runtime/Core/Public/IO/IoChunkEncoding.h b/Engine/Source/Runtime/Core/Public/IO/IoChunkEncoding.h
--- a/Engine/Source/Runtime/Core/Public/IO/IoChunkEncoding.h
+++ b/Engine/Source/Runtime/Core/Public/IO/IoChunkEncoding.h
@@ -10,6 +10,8 @@
 class FIoBuffer;
 struct FIoHash;
 
+using FIoBlockHash = uint32;
+
 /** I/O chunk encryption method. */
 enum class EIoEncryptionMethod : uint8
 {
@@ -33,6 +35,7 @@ struct FIoChunkDecodingParams
 	uint64 RawOffset = 0;
 	uint64 EncodedOffset = 0;
 	TConstArrayView<uint32> EncodedBlockSize;
+	TConstArrayView<FIoBlockHash> BlockHash;
 };
 
 /**
@@ -77,5 +80,5 @@ public:
 	static CORE_API TIoStatusOr<FIoOffsetAndLength> GetChunkRange(uint64 TotalRawSize, uint32 RawBlockSize, TConstArrayView<uint32> EncodedBlockSize, uint64 RawOffset, uint64 RawSize);
 	static CORE_API TIoStatusOr<FIoOffsetAndLength> GetChunkRange(const FIoChunkDecodingParams& Params, uint64 RawSize);
 	static CORE_API uint64 GetTotalEncodedSize(TConstArrayView<uint32> EncodedBlockSize);
-	static CORE_API FIoStatus HashBlocks(const FHeader& Header, FMemoryView EncodedData, TArray<FIoHash>& OutHashes);
+	static CORE_API FIoBlockHash HashBlock(FMemoryView Block);
 };
diff --git a/Engine/Source/Runtime/Experimental/IoStoreOnDemand/Private/IoStoreOnDemand.cpp b/Engine/Source/Runtime/Experimental/IoStoreOnDemand/Private/IoStoreOnDemand.cpp
--- a/Engine/Source/Runtime/Experimental/IoStoreOnDemand/Private/IoStoreOnDemand.cpp
+++ b/Engine/Source/Runtime/Experimental/IoStoreOnDemand/Private/IoStoreOnDemand.cpp
@@ -7,6 +7,7 @@
 #include "HAL/LowLevelMemTracker.h"
 #include "HAL/PlatformMisc.h"
 #include "HAL/PlatformTime.h"
+#include "IO/IoChunkEncoding.h"
 #include "IasCache.h"
 #include "LatencyInjector.h"
 #include "Misc/Base64.h"
@@ -973,8 +974,7 @@ TIoStatusOr<FIoStoreUploadResult> UploadContainerFiles(
 
 				FMemoryView EncodedBlock = EncodedBlocks.Left(EncodedBlockSize);
 				EncodedBlocks += EncodedBlock.GetSize();
-				const FIoHash BlockHash = FIoHash::HashBuffer(EncodedBlock);
-				ContainerEntry.BlockHashes.Add(*reinterpret_cast<const uint32*>(&BlockHash));
+				ContainerEntry.BlockHashes.Add(FIoChunkEncoding::HashBlock(EncodedBlock));
 
 				EncodedChunkSize += EncodedBlockSize;
 				RawChunkSize += BlockInfo.UncompressedSize;
diff --git a/Engine/Source/Runtime/Experimental/IoStoreOnDemand/Private/OnDemandIoDispatcherBackend.cpp b/Engine/Source/Runtime/Experimental/IoStoreOnDemand/Private/OnDemandIoDispatcherBackend.cpp
--- a/Engine/Source/Runtime/Experimental/IoStoreOnDemand/Private/OnDemandIoDispatcherBackend.cpp
+++ b/Engine/Source/Runtime/Experimental/IoStoreOnDemand/Private/OnDemandIoDispatcherBackend.cpp
@@ -193,6 +193,7 @@ public:
 
 		TMap<FIoChunkId, FTocEntry> TocEntries;
 		TArray<uint32> BlockSizes;
+		TArray<FIoBlockHash> BlockHashes;
 	};
 
 	struct FToc
@@ -216,6 +217,14 @@ public:
 			check(Container != nullptr && Entry != nullptr);
 			return TConstArrayView<uint32>(Container->BlockSizes.GetData() + Entry->BlockOffset, Entry->BlockCount);
 		}
+		
+		TConstArrayView<FIoBlockHash> GetBlockHashes() const
+		{
+			check(Container != nullptr && Entry != nullptr);
+			return Container->BlockHashes.IsEmpty()
+				? TConstArrayView<FIoBlockHash>()
+				: TConstArrayView<FIoBlockHash>(Container->BlockHashes.GetData() + Entry->BlockOffset, Entry->BlockCount);
+		}
 	};
 
 	FOnDemandIoStore();
@@ -299,6 +308,7 @@ void FOnDemandIoStore::AddToc(const FString& TocPath, FOnDemandToc&& Toc)
 			}
 
 			NewContainer->BlockSizes = MoveTemp(Container.BlockSizes);
+			NewContainer->BlockHashes = MoveTemp(Container.BlockHashes);
 
 			DeferredContainers.Add(NewContainer);
 		}
@@ -583,6 +593,7 @@ struct FChunkRequestParams
 		Params.BlockSize = ChunkInfo.Container->BlockSize;
 		Params.TotalRawSize = ChunkInfo.Entry->RawSize;
 		Params.EncodedBlockSize = ChunkInfo.GetBlocks(); 
+		Params.BlockHash = ChunkInfo.GetBlockHashes(); 
 		Params.EncodedOffset = ChunkRange.GetOffset();
 
 		return Params;
diff --git a/Engine/Source/Runtime/Experimental/IoStoreOnDemand/Public/IO/IoStoreOnDemand.h b/Engine/Source/Runtime/Experimental/IoStoreOnDemand/Public/IO/IoStoreOnDemand.h
--- a/Engine/Source/Runtime/Experimental/IoStoreOnDemand/Public/IO/IoStoreOnDemand.h
+++ b/Engine/Source/Runtime/Experimental/IoStoreOnDemand/Public/IO/IoStoreOnDemand.h
@@ -26,6 +26,7 @@ struct FAnalyticsEventAttribute;
 struct FIoContainerSettings;
 struct FIoStoreWriterSettings;
 namespace UE::IO::IAS { struct FOnDemandEndpoint; }
+using FIoBlockHash = uint32;
 
 UE_API DECLARE_LOG_CATEGORY_EXTERN(LogIas, VeryVerbose, All);
 
@@ -109,7 +110,7 @@ struct FOnDemandTocContainerEntry
 	FString EncryptionKeyGuid;
 	TArray<FOnDemandTocEntry> Entries;
 	TArray<uint32> BlockSizes;
-	TArray<uint32> BlockHashes;
+	TArray<FIoBlockHash> BlockHashes;
 	FIoHash UTocHash;
 
 	UE_API friend FArchive& operator<<(FArchive& Ar, FOnDemandTocContainerEntry& ContainerEntry);
