checkf and ensureMsgf will now validate the printf parameters when UE_VALIDATE_FORMAT_STRINGS is defined.
FormatStringSan::bIsIntegralEnum now properly detects TEnumAsByte as enum types.

diff --git a/Engine/Source/Runtime/Core/Public/Logging/LogMacros.h b/Engine/Source/Runtime/Core/Public/Logging/LogMacros.h
--- a/Engine/Source/Runtime/Core/Public/Logging/LogMacros.h
+++ b/Engine/Source/Runtime/Core/Public/Logging/LogMacros.h
@@ -230,12 +230,6 @@ CORE_API void BasicFatalLog(const FLogCategoryBase& Category, const FStaticBasic
 	#define UE_SET_LOG_VERBOSITY(CategoryName, Verbosity) \
 		CategoryName.SetVerbosity(ELogVerbosity::Verbosity);
 
-	#if UE_VALIDATE_FORMAT_STRINGS
-		#define UE_VALIDATE_FORMAT_STRING UE_CHECK_FORMAT_STRING
-	#else
-		#define UE_VALIDATE_FORMAT_STRING(Format, ...)
-	#endif
-
 	/**
 	 * A macro that logs a formatted message if the log category is active at the requested verbosity level.
 	 *
diff --git a/Engine/Source/Runtime/Core/Public/Misc/AssertionMacros.h b/Engine/Source/Runtime/Core/Public/Misc/AssertionMacros.h
--- a/Engine/Source/Runtime/Core/Public/Misc/AssertionMacros.h
+++ b/Engine/Source/Runtime/Core/Public/Misc/AssertionMacros.h
@@ -253,6 +253,7 @@ RetType FORCENOINLINE UE_DEBUG_SECTION DispatchCheckVerify(InnerType&& Inner, Ar
 
 	#define UE_CHECK_F_IMPL(expr, format, ...) \
 		{ \
+			UE_VALIDATE_FORMAT_STRING(format, ##__VA_ARGS__); \
 			if(UNLIKELY(!(expr))) \
 			{ \
 				if (FDebug::CheckVerifyFailedImpl(#expr, __FILE__, __LINE__, PLATFORM_RETURN_ADDRESS(), format, ##__VA_ARGS__)) \
@@ -371,12 +372,12 @@ RetType FORCENOINLINE UE_DEBUG_SECTION DispatchCheckVerify(InnerType&& Inner, Ar
 			return CheckVerifyImpl(bExecuted, Always, __FILE__, __LINE__, PLATFORM_RETURN_ADDRESS(), #InExpression, InFormat); \
 		}) && [] () { PLATFORM_BREAK(); return false; } ()))
 
-	#define UE_ENSURE_IMPL2(Capture, Always, InExpression, ...) \
+	#define UE_ENSURE_IMPL2(Capture, Always, InExpression, InFormat, ...) \
 		(LIKELY(!!(InExpression)) || (DispatchCheckVerify<bool>([Capture] () UE_DEBUG_SECTION \
 		{ \
 			static bool bExecuted = false; \
-			FValidateArgsInternal(__VA_ARGS__); \
-			return CheckVerifyImpl(bExecuted, Always, __FILE__, __LINE__, PLATFORM_RETURN_ADDRESS(), #InExpression, ##__VA_ARGS__); \
+			UE_VALIDATE_FORMAT_STRING(InFormat, ##__VA_ARGS__); \
+			return CheckVerifyImpl(bExecuted, Always, __FILE__, __LINE__, PLATFORM_RETURN_ADDRESS(), #InExpression, InFormat, ##__VA_ARGS__); \
 		}) && [] () { PLATFORM_BREAK(); return false; } ()))
 
 	#define ensure(           InExpression                ) UE_ENSURE_IMPL ( , false, InExpression, TEXT(""))
diff --git a/Engine/Source/Runtime/Core/Public/String/FormatStringSan.h b/Engine/Source/Runtime/Core/Public/String/FormatStringSan.h
--- a/Engine/Source/Runtime/Core/Public/String/FormatStringSan.h
+++ b/Engine/Source/Runtime/Core/Public/String/FormatStringSan.h
@@ -15,6 +15,16 @@
 #define UE_CHECK_FORMAT_STRING_ERR(Err, Fmt, ...)												\
 	(::UE::Core::Private::FormatStringSan::TCheckFormatString<decltype(::UE::Core::Private::FormatStringSan::GetFmtArgTypes(__VA_ARGS__))>::Check(false, 0, Fmt).Status == Err)
 
+
+#if UE_VALIDATE_FORMAT_STRINGS
+	#define UE_VALIDATE_FORMAT_STRING UE_CHECK_FORMAT_STRING
+#else
+	#define UE_VALIDATE_FORMAT_STRING(Format, ...)
+#endif
+
+// Forward declaration
+template <typename T> class TEnumAsByte;
+
 /// implementation
 namespace UE::Core::Private
 {
@@ -108,8 +118,8 @@ namespace UE::Core::Private
 		FMT_STR_ERR(StatusI64BadSpec, "'%I' must appear as '%I64' with an integral suffix (eg. '%I64d', '%I64u', etc.)");
 		FMT_STR_ERR(StatusI64NeedsIntegerArg, "'%I64[ ]' expects integral arg (eg. `char`, `int`, `long`, etc.)");
 		FMT_STR_ERR(StatusDynamicLengthSpecNeedsIntegerArg, "dynamic field width specifier '*' expects integral arg (eg. `char`, `int`, `long`, etc.)");
-		FMT_STR_ERR(StatusUTF8NeedsStringCastAndS, "Pass UTF8 strings to StringCast<TCHAR> and use %s.");
-		FMT_STR_ERR(StatusUTF8NeedsStringCast, "'%s' expects `TCHAR*`; pass UTF8 strings to StringCast<TCHAR> first.");
+		FMT_STR_ERR(StatusUTF8NeedsStringCastAndS, "Pass UTF8 strings to WriteToString and use %s.");
+		FMT_STR_ERR(StatusUTF8NeedsStringCast, "'%s' expects `TCHAR*`; pass UTF8 strings to WriteToString first.");
 
 #undef FMT_STR_ERR
 
@@ -122,8 +132,16 @@ namespace UE::Core::Private
 		inline constexpr bool bIsTCharPtr = bIsPointerTo<T, TCHAR>;
 		template <typename T>
 		inline constexpr bool bIsFloatOrDouble = std::is_same_v<float, T> || std::is_same_v<double, T>;
+
+		// Returns true only when the type is TEnumAsByte
+		template <typename T>
+		struct TIsEnumAsByte { static constexpr bool IsEnumAsByte = false; };
+		template <typename T>
+		struct TIsEnumAsByte<TEnumAsByte<T>> { static constexpr bool IsEnumAsByte = true; };
+		
+
 		template <typename T>
-		inline constexpr bool bIsIntegralEnum = std::is_enum_v<T>;
+		inline constexpr bool bIsIntegralEnum = std::is_enum_v<T> || TIsEnumAsByte<T>::IsEnumAsByte;
 
 		template <typename... Ts>
 		struct TFmtArgTypes {};
diff --git a/Engine/Source/Runtime/Core/Tests/String/FormatStringSanTest.cpp b/Engine/Source/Runtime/Core/Tests/String/FormatStringSanTest.cpp
--- a/Engine/Source/Runtime/Core/Tests/String/FormatStringSanTest.cpp
+++ b/Engine/Source/Runtime/Core/Tests/String/FormatStringSanTest.cpp
@@ -92,6 +92,12 @@ TEST_CASE_NAMED(FFormatStringValidatorTest, "Core::String::FormatStringSan", "[C
 			enum class MyEnum { Value };
 			STATIC_CHECK(UE_CHECK_FORMAT_STRING_ERR(FormatStringSan::StatusOk, TEXT("enum class %d value"), MyEnum::Value));
 		}
+
+		{
+			enum ETestEnumAsByte { ETestEnumAsByte_Zero = 0 };
+			TEnumAsByte<ETestEnumAsByte> EnumAsByteParam = ETestEnumAsByte_Zero;
+			STATIC_CHECK(UE_CHECK_FORMAT_STRING_ERR(FormatStringSan::StatusOk, TEXT("%d"), EnumAsByteParam));
+		}
 	}
 }
 
