FString prep for UTF-8:

- TCHAR to ElementType.
- Implementation functions taking TCHARs to being templated taking a CharType.
- FCString to TCString<CharType>.
- FChar to TChar<CharType>.
- TEXT(...) to UE_STRING_TEXT(...).

diff --git a/Engine/Source/Runtime/Core/Private/Containers/String.cpp b/Engine/Source/Runtime/Core/Private/Containers/String.cpp
--- a/Engine/Source/Runtime/Core/Private/Containers/String.cpp
+++ b/Engine/Source/Runtime/Core/Private/Containers/String.cpp
@@ -24,29 +24,34 @@
 #include "Templates/UnrealTemplate.h"
 #include "UObject/NameTypes.h"
 
+#define UE_STRING_COMPILING_UTF8 0
+#define UE_STRING_TEXT(x) TEXT(x)
+
 /* FString implementation
  *****************************************************************************/
 
 namespace UE::Core::Private
 {
-	struct FCompareCharsCaseSensitive
+	template <typename CharType>
+	struct TCompareCharsCaseSensitive
 	{
-		static FORCEINLINE bool Compare(TCHAR Lhs, TCHAR Rhs)
+		static FORCEINLINE bool Compare(CharType Lhs, CharType Rhs)
 		{
 			return Lhs == Rhs;
 		}
 	};
 
-	struct FCompareCharsCaseInsensitive
+	template <typename CharType>
+	struct TCompareCharsCaseInsensitive
 	{
-		static FORCEINLINE bool Compare(TCHAR Lhs, TCHAR Rhs)
+		static FORCEINLINE bool Compare(CharType Lhs, CharType Rhs)
 		{
-			return FChar::ToLower(Lhs) == FChar::ToLower(Rhs);
+			return TChar<CharType>::ToLower(Lhs) == TChar<CharType>::ToLower(Rhs);
 		}
 	};
 
-	template <typename CompareType>
-	bool MatchesWildcardRecursive(const TCHAR* Target, int32 TargetLength, const TCHAR* Wildcard, int32 WildcardLength)
+	template <typename CompareType, typename CharType>
+	bool MatchesWildcardRecursive(const CharType* Target, int32 TargetLength, const CharType* Wildcard, int32 WildcardLength)
 	{
 		// Skip over common initial non-wildcard-char sequence of Target and Wildcard
 		for (;;)
@@ -56,8 +61,8 @@ namespace UE::Core::Private
 				return TargetLength == 0;
 			}
 
-			TCHAR WCh = *Wildcard;
-			if (WCh == TEXT('*') || WCh == TEXT('?'))
+			CharType WCh = *Wildcard;
+			if (WCh == CharType('*') || WCh == CharType('?'))
 			{
 				break;
 			}
@@ -74,15 +79,15 @@ namespace UE::Core::Private
 		}
 
 		// Test for common suffix
-		const TCHAR* TPtr = Target   + TargetLength;
-		const TCHAR* WPtr = Wildcard + WildcardLength;
+		const CharType* TPtr = Target   + TargetLength;
+		const CharType* WPtr = Wildcard + WildcardLength;
 		for (;;)
 		{
 			--TPtr;
 			--WPtr;
 
-			TCHAR WCh = *WPtr;
-			if (WCh == TEXT('*') || WCh == TEXT('?'))
+			CharType WCh = *WPtr;
+			if (WCh == CharType('*') || WCh == CharType('?'))
 			{
 				break;
 			}
@@ -102,8 +107,8 @@ namespace UE::Core::Private
 		}
 
 		// Match * against anything and ? against single (and zero?) chars
-		TCHAR FirstWild = *Wildcard;
-		if (WildcardLength == 1 && (FirstWild == TEXT('*') || TargetLength < 2))
+		CharType FirstWild = *Wildcard;
+		if (WildcardLength == 1 && (FirstWild == CharType('*') || TargetLength < 2))
 		{
 			return true;
 		}
@@ -112,7 +117,7 @@ namespace UE::Core::Private
 
 		// This routine is very slow, though it does ok with one wildcard
 		int32 MaxNum = TargetLength;
-		if (FirstWild == TEXT('?') && MaxNum > 1)
+		if (FirstWild == CharType('?') && MaxNum > 1)
 		{
 			MaxNum = 1;
 		}
@@ -127,8 +132,8 @@ namespace UE::Core::Private
 		return false;
 	}
 
-	template<typename CharType>
-	void AppendCharacters(TArray<TCHAR>& Out, const CharType* Str, int32 Count)
+	template <typename DestCharType, typename SrcCharType>
+	void AppendCharacters(TArray<DestCharType>& Out, const SrcCharType* Str, int32 Count)
 	{
 		check(Count >= 0);
 
@@ -146,15 +151,15 @@ namespace UE::Core::Private
 		Out.AddUninitialized(Count + (OldEnd ? 0 : 1));
 		OldEnd -= OldEnd ? 1 : 0;
 
-		TCHAR* Dest = Out.GetData() + OldEnd;
+		DestCharType* Dest = Out.GetData() + OldEnd;
 
 		// Try copying characters to end of string, overwriting null terminator if we already have one
-		TCHAR* NewEnd = FPlatformString::Convert(Dest, Count, Str, Count);
+		DestCharType* NewEnd = FPlatformString::Convert(Dest, Count, Str, Count);
 		if (!NewEnd)
 		{
 			// If that failed, it will have meant that conversion likely contained multi-code unit characters
 			// and so the buffer wasn't long enough, so calculate it properly.
-			int32 Length = FPlatformString::ConvertedLength<TCHAR>(Str, Count);
+			int32 Length = FPlatformString::ConvertedLength<DestCharType>(Str, Count);
 
 			// Add the extra bytes that we need
 			Out.AddUninitialized(Length - Count);
@@ -178,13 +183,13 @@ namespace UE::Core::Private
 		*NewEnd = '\0';
 	}
 
-	template<typename CharType>
-	FORCEINLINE void ConstructFromCString(/* Out */ TArray<TCHAR>& Data, const CharType* Src)
+	template <typename DestCharType, typename SrcCharType>
+	FORCEINLINE void ConstructFromCString(/* Out */ TArray<DestCharType>& Data, const SrcCharType* Src)
 	{
 		if (Src && *Src)
 		{
-			int32 SrcLen  = TCString<CharType>::Strlen(Src) + 1;
-			int32 DestLen = FPlatformString::ConvertedLength<TCHAR>(Src, SrcLen);
+			int32 SrcLen  = TCString<SrcCharType>::Strlen(Src) + 1;
+			int32 DestLen = FPlatformString::ConvertedLength<DestCharType>(Src, SrcLen);
 			Data.Reserve(DestLen);
 			Data.AddUninitialized(DestLen);
 
@@ -192,30 +197,30 @@ namespace UE::Core::Private
 		}
 	}
 
-	template<typename CharType>
-	FORCEINLINE void ConstructWithLength(/* Out */ TArray<TCHAR>& Data, int32 InCount, const CharType* InSrc)
+	template <typename DestCharType, typename SrcCharType>
+	FORCEINLINE void ConstructWithLength(/* Out */ TArray<DestCharType>& Data, int32 InCount, const SrcCharType* InSrc)
 	{
 		if (InSrc)
 		{
-			int32 DestLen = FPlatformString::ConvertedLength<TCHAR>(InSrc, InCount);
+			int32 DestLen = FPlatformString::ConvertedLength<DestCharType>(InSrc, InCount);
 			if (DestLen > 0 && *InSrc)
 			{
 				Data.Reserve(DestLen + 1);
 				Data.AddUninitialized(DestLen + 1);
 
 				FPlatformString::Convert(Data.GetData(), DestLen, InSrc, InCount);
-				*(Data.GetData() + Data.Num() - 1) = TEXT('\0');
+				*(Data.GetData() + Data.Num() - 1) = DestCharType('\0');
 			}
 		}
 	}
 
-	template<typename CharType>
-	FORCEINLINE void ConstructWithSlack(/* Out */ TArray<TCHAR>& Data, const CharType* Src, int32 ExtraSlack)
+	template <typename DestCharType, typename SrcCharType>
+	FORCEINLINE void ConstructWithSlack(/* Out */ TArray<DestCharType>& Data, const SrcCharType* Src, int32 ExtraSlack)
 	{
 		if (Src && *Src)
 		{
-			int32 SrcLen = TCString<CharType>::Strlen(Src) + 1;
-			int32 DestLen = FPlatformString::ConvertedLength<TCHAR>(Src, SrcLen);
+			int32 SrcLen = TCString<SrcCharType>::Strlen(Src) + 1;
+			int32 DestLen = FPlatformString::ConvertedLength<DestCharType>(Src, SrcLen);
 			Data.Reserve(DestLen + ExtraSlack);
 			Data.AddUninitialized(DestLen);
 
@@ -241,24 +246,24 @@ FString::FString(const WIDECHAR* Str, int32 ExtraSlack)				{ UE::Core::Private::
 FString::FString(const UTF8CHAR* Str, int32 ExtraSlack)				{ UE::Core::Private::ConstructWithSlack(Data, Str, ExtraSlack); }
 FString::FString(const UCS2CHAR* Str, int32 ExtraSlack)				{ UE::Core::Private::ConstructWithSlack(Data, Str, ExtraSlack); }
 
-FString& FString::operator=( const TCHAR* Other )
+FString& FString::operator=( const ElementType* Other )
 {
 	if (Data.GetData() != Other)
 	{
-		int32 Len = (Other && *Other) ? FCString::Strlen(Other)+1 : 0;
+		int32 Len = (Other && *Other) ? TCString<ElementType>::Strlen(Other)+1 : 0;
 		Data.Empty(Len);
 		Data.AddUninitialized(Len);
 			
 		if( Len )
 		{
-			FMemory::Memcpy( Data.GetData(), Other, Len * sizeof(TCHAR) );
+			FMemory::Memcpy( Data.GetData(), Other, Len * sizeof(ElementType) );
 		}
 	}
 	return *this;
 }
 
 
-void FString::AssignRange(const TCHAR* OtherData, int32 OtherLen)
+void FString::AssignRange(const ElementType* OtherData, int32 OtherLen)
 {
 	if (OtherLen == 0)
 	{
@@ -270,18 +275,18 @@ void FString::AssignRange(const TCHAR* OtherData, int32 OtherLen)
 		if (OtherLen <= ThisLen)
 		{
 			// Unless the input is longer, this might be assigned from a view of itself.
-			TCHAR* DataPtr = Data.GetData();
-			FMemory::Memmove(DataPtr, OtherData, OtherLen * sizeof(TCHAR));
-			DataPtr[OtherLen] = TEXT('\0');
+			ElementType* DataPtr = Data.GetData();
+			FMemory::Memmove(DataPtr, OtherData, OtherLen * sizeof(ElementType));
+			DataPtr[OtherLen] = UE_STRING_TEXT('\0');
 			Data.RemoveAt(OtherLen + 1, ThisLen - OtherLen);
 		}
 		else
 		{
 			Data.Empty(OtherLen + 1);
 			Data.AddUninitialized(OtherLen + 1);
-			TCHAR* DataPtr = Data.GetData();
-			FMemory::Memcpy(DataPtr, OtherData, OtherLen * sizeof(TCHAR));
-			DataPtr[OtherLen] = TEXT('\0');
+			ElementType* DataPtr = Data.GetData();
+			FMemory::Memcpy(DataPtr, OtherData, OtherLen * sizeof(ElementType));
+			DataPtr[OtherLen] = UE_STRING_TEXT('\0');
 		}
 	}
 }
@@ -309,9 +314,9 @@ void FString::Reset(int32 NewReservedSize)
 {
 	const int32 NewSizeIncludingTerminator = (NewReservedSize > 0) ? (NewReservedSize + 1) : 0;
 	Data.Reset(NewSizeIncludingTerminator);
-	if (TCHAR* DataPtr = Data.GetData())
+	if (ElementType* DataPtr = Data.GetData())
 	{
-		*DataPtr = TEXT('\0');
+		*DataPtr = UE_STRING_TEXT('\0');
 	}
 }
 
@@ -325,14 +330,14 @@ void FString::Shrink()
 NSString* FString::GetNSString() const
 {
 #if PLATFORM_TCHAR_IS_4_BYTES
-    return [[[NSString alloc] initWithBytes:Data.GetData() length:Len() * sizeof(TCHAR) encoding:NSUTF32LittleEndianStringEncoding] autorelease];
+    return [[[NSString alloc] initWithBytes:Data.GetData() length:Len() * sizeof(ElementType) encoding:NSUTF32LittleEndianStringEncoding] autorelease];
 #else
-    return [[[NSString alloc] initWithBytes:Data.GetData() length:Len() * sizeof(TCHAR) encoding:NSUTF16LittleEndianStringEncoding] autorelease];
+    return [[[NSString alloc] initWithBytes:Data.GetData() length:Len() * sizeof(ElementType) encoding:NSUTF16LittleEndianStringEncoding] autorelease];
 #endif
 }
 #endif
 
-FString& FString::AppendChar(TCHAR InChar)
+FString& FString::AppendChar(ElementType InChar)
 {
 	CheckInvariants();
 
@@ -348,7 +353,7 @@ FString& FString::AppendChar(TCHAR InChar)
 
 		Data.AddUninitialized(InsertCount);
 		Data[InsertIndex] = InChar;
-		Data[InsertIndex+1] = TEXT('\0');
+		Data[InsertIndex+1] = UE_STRING_TEXT('\0');
 	}
 	return *this;
 }
@@ -381,7 +386,7 @@ void FString::TrimToNullTerminator()
 {
 	if( Data.Num() )
 	{
-		int32 DataLen = FCString::Strlen(Data.GetData());
+		int32 DataLen = TCString<ElementType>::Strlen(Data.GetData());
 		check(DataLen == 0 || DataLen < Data.Num());
 		int32 Len = DataLen > 0 ? DataLen+1 : 0;
 
@@ -391,23 +396,23 @@ void FString::TrimToNullTerminator()
 }
 
 
-int32 FString::Find(const TCHAR* SubStr, int32 SubStrLen, ESearchCase::Type SearchCase, ESearchDir::Type SearchDir, int32 StartPosition) const
+int32 FString::Find(const ElementType* SubStr, int32 SubStrLen, ESearchCase::Type SearchCase, ESearchDir::Type SearchDir, int32 StartPosition) const
 {
 	checkf(SubStrLen >= 0, TEXT("Invalid SubStrLen: %d"), SubStrLen);
 
 	if (SearchDir == ESearchDir::FromStart)
 	{
-		const TCHAR* Start = **this;
+		const ElementType* Start = **this;
 		int32 RemainingLength = Len();
 		if (StartPosition != INDEX_NONE && RemainingLength > 0)
 		{
-			const TCHAR* End = Start + RemainingLength;
+			const ElementType* End = Start + RemainingLength;
 			Start += FMath::Clamp(StartPosition, 0, RemainingLength - 1);
 			RemainingLength = UE_PTRDIFF_TO_INT32(End - Start);
 		}
-		const TCHAR* Tmp = SearchCase == ESearchCase::IgnoreCase
-			? FCString::Strnistr(Start, RemainingLength, SubStr, SubStrLen)
-			: FCString::Strnstr(Start, RemainingLength, SubStr, SubStrLen);
+		const ElementType* Tmp = SearchCase == ESearchCase::IgnoreCase
+			? TCString<ElementType>::Strnistr(Start, RemainingLength, SubStr, SubStrLen)
+			: TCString<ElementType>::Strnstr(Start, RemainingLength, SubStr, SubStrLen);
 
 		return Tmp ? UE_PTRDIFF_TO_INT32(Tmp-**this) : INDEX_NONE;
 	}
@@ -500,10 +505,10 @@ FString FString::ToUpper() &&
 void FString::ToUpperInline()
 {
 	const int32 StringLength = Len();
-	TCHAR* RawData = Data.GetData();
+	ElementType* RawData = Data.GetData();
 	for (int32 i = 0; i < StringLength; ++i)
 	{
-		RawData[i] = FChar::ToUpper(RawData[i]);
+		RawData[i] = TChar<ElementType>::ToUpper(RawData[i]);
 	}
 }
 
@@ -524,10 +529,10 @@ FString FString::ToLower() &&
 void FString::ToLowerInline()
 {
 	const int32 StringLength = Len();
-	TCHAR* RawData = Data.GetData();
+	ElementType* RawData = Data.GetData();
 	for (int32 i = 0; i < StringLength; ++i)
 	{
-		RawData[i] = FChar::ToLower(RawData[i]);
+		RawData[i] = TChar<ElementType>::ToLower(RawData[i]);
 	}
 }
 
@@ -539,7 +544,7 @@ void FString::RemoveSpacesInline()
 		return;
 	}
 
-	TCHAR* RawData = Data.GetData();
+	ElementType* RawData = Data.GetData();
 	int32 CopyToIndex = 0;
 	for (int32 CopyFromIndex = 0; CopyFromIndex < StringLength; ++CopyFromIndex)
 	{
@@ -553,40 +558,40 @@ void FString::RemoveSpacesInline()
 	// Copy null-terminating character.
 	if (CopyToIndex <= StringLength)
 	{
-		RawData[CopyToIndex] = TEXT('\0');
+		RawData[CopyToIndex] = UE_STRING_TEXT('\0');
 		Data.SetNum(CopyToIndex + 1, false);
 	}
 }
 
-bool FString::StartsWith(const TCHAR* InPrefix, int32 InPrefixLen, ESearchCase::Type SearchCase) const
+bool FString::StartsWith(const ElementType* InPrefix, int32 InPrefixLen, ESearchCase::Type SearchCase) const
 {
 	if (SearchCase == ESearchCase::IgnoreCase)
 	{
-		return InPrefixLen > 0 && !FCString::Strnicmp(**this, InPrefix, InPrefixLen);
+		return InPrefixLen > 0 && !TCString<ElementType>::Strnicmp(**this, InPrefix, InPrefixLen);
 	}
 	else
 	{
-		return InPrefixLen > 0 && !FCString::Strncmp(**this, InPrefix, InPrefixLen);
+		return InPrefixLen > 0 && !TCString<ElementType>::Strncmp(**this, InPrefix, InPrefixLen);
 	}
 }
 
-bool FString::EndsWith(const TCHAR* InSuffix, int32 InSuffixLen, ESearchCase::Type SearchCase ) const
+bool FString::EndsWith(const ElementType* InSuffix, int32 InSuffixLen, ESearchCase::Type SearchCase ) const
 {
 	if (SearchCase == ESearchCase::IgnoreCase)
 	{
 		return InSuffixLen > 0 &&
 			Len() >= InSuffixLen &&
-			!FCString::Stricmp( &(*this)[ Len() - InSuffixLen ], InSuffix );
+			!TCString<ElementType>::Stricmp( &(*this)[ Len() - InSuffixLen ], InSuffix );
 	}
 	else
 	{
 		return InSuffixLen > 0 &&
 			Len() >= InSuffixLen &&
-			!FCString::Strcmp( &(*this)[ Len() - InSuffixLen ], InSuffix );
+			!TCString<ElementType>::Strcmp( &(*this)[ Len() - InSuffixLen ], InSuffix );
 	}
 }
 
-void FString::InsertAt(int32 Index, TCHAR Character)
+void FString::InsertAt(int32 Index, ElementType Character)
 {
 	if (Character != 0)
 	{
@@ -621,7 +626,7 @@ void FString::RemoveAt(int32 Index, int32 Count, bool bAllowShrinking)
 	Data.RemoveAt(Index, FMath::Clamp(Count, 0, Len()-Index), bAllowShrinking);
 }
 
-bool FString::RemoveFromStart(const TCHAR* InPrefix, int32 InPrefixLen, ESearchCase::Type SearchCase)
+bool FString::RemoveFromStart(const ElementType* InPrefix, int32 InPrefixLen, ESearchCase::Type SearchCase)
 {
 	if (InPrefixLen == 0 )
 	{
@@ -637,7 +642,7 @@ bool FString::RemoveFromStart(const TCHAR* InPrefix, int32 InPrefixLen, ESearchC
 	return false;
 }
 
-bool FString::RemoveFromEnd(const TCHAR* InSuffix, int32 InSuffixLen, ESearchCase::Type SearchCase)
+bool FString::RemoveFromEnd(const ElementType* InSuffix, int32 InSuffixLen, ESearchCase::Type SearchCase)
 {
 	if (InSuffixLen == 0)
 	{
@@ -675,9 +680,11 @@ UE_NODISCARD FORCEINLINE FString ConcatFStrings(LhsType&& Lhs, RhsType&& Rhs)
 	return Result;
 }
 
-template <typename RhsType>
-UE_NODISCARD FORCEINLINE FString ConcatRangeFString(const TCHAR* Lhs, int32 LhsLen, RhsType&& Rhs)
+template <typename LhsCharType, typename RhsType>
+UE_NODISCARD FORCEINLINE FString ConcatRangeFString(const LhsCharType* Lhs, int32 LhsLen, RhsType&& Rhs)
 {
+	using ElementType = FString::ElementType;
+
 	checkSlow(LhsLen >= 0);
 	Rhs.CheckInvariants();
 	if (LhsLen == 0)
@@ -694,16 +701,16 @@ UE_NODISCARD FORCEINLINE FString ConcatRangeFString(const TCHAR* Lhs, int32 LhsL
 	Result.GetCharArray().Reserve(LhsLen + RhsLen + 1);
 	Result.GetCharArray().AddUninitialized(LhsLen + RhsLen + 1);
 
-	TCHAR* ResultData = Result.GetCharArray().GetData();
+	ElementType* ResultData = Result.GetCharArray().GetData();
 	CopyAssignItems(ResultData, Lhs, LhsLen);
 	CopyAssignItems(ResultData + LhsLen, Rhs.GetCharArray().GetData(), RhsLen);
-	*(ResultData + LhsLen + RhsLen) = TEXT('\0');
-		
+	*(ResultData + LhsLen + RhsLen) = UE_STRING_TEXT('\0');
+
 	return Result;
 }
 
-template <typename LhsType>
-UE_NODISCARD FORCEINLINE FString ConcatFStringRange(LhsType&& Lhs, const TCHAR* Rhs, int32 RhsLen)
+template <typename LhsType, typename RhsCharType>
+UE_NODISCARD FORCEINLINE FString ConcatFStringRange(LhsType&& Lhs, const RhsCharType* Rhs, int32 RhsLen)
 {
 	Lhs.CheckInvariants();
 	checkSlow(RhsLen >= 0);
@@ -718,8 +725,8 @@ UE_NODISCARD FORCEINLINE FString ConcatFStringRange(LhsType&& Lhs, const TCHAR*
 	return Result;
 }
 
-template <typename RhsType>
-UE_NODISCARD FORCEINLINE FString ConcatCStringFString(const TCHAR* Lhs, RhsType&& Rhs)
+template <typename LhsCharType, typename RhsType>
+UE_NODISCARD FORCEINLINE FString ConcatCStringFString(const LhsCharType* Lhs, RhsType&& Rhs)
 {
 	checkSlow(Lhs);
 	if (!Lhs)
@@ -727,34 +734,34 @@ UE_NODISCARD FORCEINLINE FString ConcatCStringFString(const TCHAR* Lhs, RhsType&
 		return Forward<RhsType>(Rhs);
 	}
 
-	return ConcatRangeFString(Lhs, FCString::Strlen(Lhs), Forward<RhsType>(Rhs));
+	return ConcatRangeFString(Lhs, TCString<LhsCharType>::Strlen(Lhs), Forward<RhsType>(Rhs));
 }
 
-template <typename LhsType>
-UE_NODISCARD FORCEINLINE FString ConcatFStringCString(LhsType&& Lhs, const TCHAR* Rhs)
+template <typename LhsType, typename RhsCharType>
+UE_NODISCARD FORCEINLINE FString ConcatFStringCString(LhsType&& Lhs, const RhsCharType* Rhs)
 {
 	checkSlow(Rhs);
 	if (!Rhs)
 	{
 		return Forward<LhsType>(Lhs);
 	}
-	return ConcatFStringRange(Forward<LhsType>(Lhs), Rhs, FCString::Strlen(Rhs));
+	return ConcatFStringRange(Forward<LhsType>(Lhs), Rhs, TCString<RhsCharType>::Strlen(Rhs));
 }
 
 } // namespace UE::String::Private
 
-FString FString::ConcatFF(const FString& Lhs, const FString& Rhs)				{ return UE::String::Private::ConcatFStrings(Lhs, Rhs); }
-FString FString::ConcatFF(FString&& Lhs, const FString& Rhs)					{ return UE::String::Private::ConcatFStrings(MoveTemp(Lhs), Rhs); }
-FString FString::ConcatFF(const FString& Lhs, FString&& Rhs)					{ return UE::String::Private::ConcatFStrings(Lhs, MoveTemp(Rhs)); }
-FString FString::ConcatFF(FString&& Lhs, FString&& Rhs)							{ return UE::String::Private::ConcatFStrings(MoveTemp(Lhs), MoveTemp(Rhs)); }
-FString FString::ConcatFC(const FString& Lhs, const TCHAR* Rhs)					{ return UE::String::Private::ConcatFStringCString(Lhs, Rhs); }
-FString FString::ConcatFC(FString&& Lhs, const TCHAR* Rhs)						{ return UE::String::Private::ConcatFStringCString(MoveTemp(Lhs), Rhs); }
-FString FString::ConcatCF(const TCHAR* Lhs,	const FString& Rhs)					{ return UE::String::Private::ConcatCStringFString(Lhs, Rhs); }
-FString FString::ConcatCF(const TCHAR* Lhs,	FString&& Rhs)						{ return UE::String::Private::ConcatCStringFString(Lhs, MoveTemp(Rhs)); }
-FString FString::ConcatFR(const FString& Lhs, const TCHAR* Rhs, int32 RhsLen)	{ return UE::String::Private::ConcatFStringRange(Lhs, Rhs, RhsLen); }
-FString FString::ConcatFR(FString&& Lhs, const TCHAR* Rhs, int32 RhsLen)		{ return UE::String::Private::ConcatFStringRange(MoveTemp(Lhs), Rhs, RhsLen); }
-FString FString::ConcatRF(const TCHAR* Lhs, int32 LhsLen, const FString& Rhs)	{ return UE::String::Private::ConcatRangeFString(Lhs, LhsLen, Rhs); }
-FString FString::ConcatRF(const TCHAR* Lhs, int32 LhsLen, FString&& Rhs)		{ return UE::String::Private::ConcatRangeFString(Lhs, LhsLen, MoveTemp(Rhs)); }
+FString FString::ConcatFF(const FString& Lhs, const FString& Rhs)						{ return UE::String::Private::ConcatFStrings(Lhs, Rhs); }
+FString FString::ConcatFF(FString&& Lhs, const FString& Rhs)							{ return UE::String::Private::ConcatFStrings(MoveTemp(Lhs), Rhs); }
+FString FString::ConcatFF(const FString& Lhs, FString&& Rhs)							{ return UE::String::Private::ConcatFStrings(Lhs, MoveTemp(Rhs)); }
+FString FString::ConcatFF(FString&& Lhs, FString&& Rhs)									{ return UE::String::Private::ConcatFStrings(MoveTemp(Lhs), MoveTemp(Rhs)); }
+FString FString::ConcatFC(const FString& Lhs, const ElementType* Rhs)					{ return UE::String::Private::ConcatFStringCString(Lhs, Rhs); }
+FString FString::ConcatFC(FString&& Lhs, const ElementType* Rhs)						{ return UE::String::Private::ConcatFStringCString(MoveTemp(Lhs), Rhs); }
+FString FString::ConcatCF(const ElementType* Lhs,	const FString& Rhs)					{ return UE::String::Private::ConcatCStringFString(Lhs, Rhs); }
+FString FString::ConcatCF(const ElementType* Lhs,	FString&& Rhs)						{ return UE::String::Private::ConcatCStringFString(Lhs, MoveTemp(Rhs)); }
+FString FString::ConcatFR(const FString& Lhs, const ElementType* Rhs, int32 RhsLen)		{ return UE::String::Private::ConcatFStringRange(Lhs, Rhs, RhsLen); }
+FString FString::ConcatFR(FString&& Lhs, const ElementType* Rhs, int32 RhsLen)			{ return UE::String::Private::ConcatFStringRange(MoveTemp(Lhs), Rhs, RhsLen); }
+FString FString::ConcatRF(const ElementType* Lhs, int32 LhsLen, const FString& Rhs)		{ return UE::String::Private::ConcatRangeFString(Lhs, LhsLen, Rhs); }
+FString FString::ConcatRF(const ElementType* Lhs, int32 LhsLen, FString&& Rhs)			{ return UE::String::Private::ConcatRangeFString(Lhs, LhsLen, MoveTemp(Rhs)); }
 
 /**
  * Concatenate this path with given path ensuring the / character is used between them
@@ -762,24 +769,24 @@ FString FString::ConcatRF(const TCHAR* Lhs, int32 LhsLen, FString&& Rhs)		{ retu
  * @param Str       Pointer to an array of TCHARs (not necessarily null-terminated) to be concatenated onto the end of this.
  * @param StrLength Exact number of characters from Str to append.
  */
-void FString::PathAppend(const TCHAR* Str, int32 StrLength)
+void FString::PathAppend(const ElementType* Str, int32 StrLength)
 {
 	int32 DataNum = Data.Num();
 	if (StrLength == 0)
 	{
-		if (DataNum > 1 && Data[DataNum - 2] != TEXT('/') && Data[DataNum - 2] != TEXT('\\'))
+		if (DataNum > 1 && Data[DataNum - 2] != UE_STRING_TEXT('/') && Data[DataNum - 2] != UE_STRING_TEXT('\\'))
 		{
-			Data[DataNum - 1] = TEXT('/');
-			Data.Add(TEXT('\0'));
+			Data[DataNum - 1] = UE_STRING_TEXT('/');
+			Data.Add(UE_STRING_TEXT('\0'));
 		}
 	}
 	else
 	{
 		if (DataNum > 0)
 		{
-			if (DataNum > 1 && Data[DataNum - 2] != TEXT('/') && Data[DataNum - 2] != TEXT('\\') && *Str != TEXT('/'))
+			if (DataNum > 1 && Data[DataNum - 2] != UE_STRING_TEXT('/') && Data[DataNum - 2] != UE_STRING_TEXT('\\') && *Str != UE_STRING_TEXT('/'))
 			{
-				Data[DataNum - 1] = TEXT('/');
+				Data[DataNum - 1] = UE_STRING_TEXT('/');
 			}
 			else
 			{
@@ -790,7 +797,7 @@ void FString::PathAppend(const TCHAR* Str, int32 StrLength)
 
 		Reserve(DataNum + StrLength);
 		Data.Append(Str, StrLength);
-		Data.Add(TEXT('\0'));
+		Data.Add(UE_STRING_TEXT('\0'));
 	}
 }
 
@@ -821,17 +828,17 @@ FString FString::Mid(int32 Start, int32 Count) &&
 	return MoveTemp(*this);
 }
 
-void FString::ReplaceCharInlineCaseSensitive(const TCHAR SearchChar, const TCHAR ReplacementChar)
+void FString::ReplaceCharInlineCaseSensitive(const ElementType SearchChar, const ElementType ReplacementChar)
 {
-	for (TCHAR& Character : Data)
+	for (ElementType& Character : Data)
 	{
 		Character = Character == SearchChar ? ReplacementChar : Character;
 	}
 }
 
-void FString::ReplaceCharInlineIgnoreCase(const TCHAR SearchChar, const TCHAR ReplacementChar)
+void FString::ReplaceCharInlineIgnoreCase(const ElementType SearchChar, const ElementType ReplacementChar)
 {
-	TCHAR OtherCaseSearchChar = TChar<TCHAR>::IsUpper(SearchChar) ? TChar<TCHAR>::ToLower(SearchChar) : TChar<TCHAR>::ToUpper(SearchChar);
+	ElementType OtherCaseSearchChar = TChar<ElementType>::IsUpper(SearchChar) ? TChar<ElementType>::ToLower(SearchChar) : TChar<ElementType>::ToUpper(SearchChar);
 	ReplaceCharInlineCaseSensitive(OtherCaseSearchChar, ReplacementChar);
 	ReplaceCharInlineCaseSensitive(SearchChar, ReplacementChar);
 }
@@ -858,7 +865,7 @@ FString FString::TrimStartAndEnd() &&
 void FString::TrimStartInline()
 {
 	int32 Pos = 0;
-	while(Pos < Len() && FChar::IsWhitespace((*this)[Pos]))
+	while(Pos < Len() && TChar<ElementType>::IsWhitespace((*this)[Pos]))
 	{
 		Pos++;
 	}
@@ -881,7 +888,7 @@ FString FString::TrimStart() &&
 void FString::TrimEndInline()
 {
 	int32 End = Len();
-	while(End > 0 && FChar::IsWhitespace((*this)[End - 1]))
+	while(End > 0 && TChar<ElementType>::IsWhitespace((*this)[End - 1]))
 	{
 		End--;
 	}
@@ -901,7 +908,7 @@ FString FString::TrimEnd() &&
 	return MoveTemp(*this);
 }
 
-void FString::TrimCharInline(const TCHAR CharacterToTrim, bool* bCharRemoved)
+void FString::TrimCharInline(const ElementType CharacterToTrim, bool* bCharRemoved)
 {
 	bool bQuotesWereRemoved=false;
 	int32 Start = 0, Count = Len();
@@ -930,7 +937,7 @@ void FString::TrimCharInline(const TCHAR CharacterToTrim, bool* bCharRemoved)
 
 void FString::TrimQuotesInline(bool* bQuotesRemoved)
 {
-	TrimCharInline(TCHAR('"'), bQuotesRemoved);
+	TrimCharInline(ElementType('"'), bQuotesRemoved);
 }
 
 FString FString::TrimQuotes(bool* bQuotesRemoved) const &
@@ -946,14 +953,14 @@ FString FString::TrimQuotes(bool* bQuotesRemoved) &&
 	return MoveTemp(*this);
 }
 
-FString FString::TrimChar(const TCHAR CharacterToTrim, bool* bCharRemoved) const &
+FString FString::TrimChar(const ElementType CharacterToTrim, bool* bCharRemoved) const &
 {
 	FString Result(*this);
 	Result.TrimCharInline(CharacterToTrim, bCharRemoved);
 	return Result;
 }
 
-FString FString::TrimChar(const TCHAR CharacterToTrim, bool* bCharRemoved) &&
+FString FString::TrimChar(const ElementType CharacterToTrim, bool* bCharRemoved) &&
 {
 	TrimCharInline(CharacterToTrim, bCharRemoved);
 	return MoveTemp(*this);
@@ -984,9 +991,9 @@ void FString::ReverseString()
 {
 	if ( Len() > 0 )
 	{
-		TCHAR* StartChar = &(*this)[0];
-		TCHAR* EndChar = &(*this)[Len()-1];
-		TCHAR TempChar;
+		ElementType* StartChar = &(*this)[0];
+		ElementType* EndChar = &(*this)[Len()-1];
+		ElementType TempChar;
 		do 
 		{
 			TempChar = *StartChar;	// store the current value of StartChar
@@ -1012,7 +1019,7 @@ FString FString::FormatAsNumber( int32 InNumber )
 		dec++;
 		if( dec == 3 && x > 0 )
 		{
-			Result += TEXT(",");
+			Result += UE_STRING_TEXT(',');
 			dec = 0;
 		}
 	}
@@ -1048,11 +1055,11 @@ void FString::SerializeAsANSICharArray( FArchive& Ar, int32 MinCharacters ) cons
 
 void FString::AppendInt( int32 Num )
 {
-	const TCHAR* DigitToChar		= TEXT("9876543210123456789");
+	const ElementType* DigitToChar	= UE_STRING_TEXT("9876543210123456789");
 	constexpr int32 ZeroDigitIndex	= 9;
 	bool bIsNumberNegative			= Num < 0;
 	const int32 TempBufferSize		= 16; // 16 is big enough
-	TCHAR TempNum[TempBufferSize];				
+	ElementType TempNum[TempBufferSize];
 	int32 TempAt					= TempBufferSize; // fill the temp string from the top down.
 
 	// Convert to string assuming base ten.
@@ -1064,10 +1071,10 @@ void FString::AppendInt( int32 Num )
 
 	if( bIsNumberNegative )
 	{
-		TempNum[--TempAt] = TEXT('-');
+		TempNum[--TempAt] = UE_STRING_TEXT('-');
 	}
 
-	const TCHAR* CharPtr = TempNum + TempAt;
+	const ElementType* CharPtr = TempNum + TempAt;
 	const int32 NumChars = TempBufferSize - TempAt;
 	Append(CharPtr, NumChars);
 }
@@ -1075,7 +1082,7 @@ void FString::AppendInt( int32 Num )
 
 bool FString::ToBool() const
 {
-	return FCString::ToBool(**this);
+	return TCString<ElementType>::ToBool(**this);
 }
 
 FString FString::FromBlob(const uint8* SrcBuffer,const uint32 SrcSize)
@@ -1085,7 +1092,7 @@ FString FString::FromBlob(const uint8* SrcBuffer,const uint32 SrcSize)
 	// Convert and append each byte in the buffer
 	for (uint32 Count = 0; Count < SrcSize; Count++)
 	{
-		Result += FString::Printf(TEXT("%03d"),(uint8)SrcBuffer[Count]);
+		Result += FString::Printf(UE_STRING_TEXT("%03d"),(uint8)SrcBuffer[Count]);
 	}
 	return Result;
 }
@@ -1097,8 +1104,8 @@ bool FString::ToBlob(const FString& Source,uint8* DestBuffer,const uint32 DestSi
 	if (DestSize >= (uint32)(Source.Len() / 3) &&
 		(Source.Len() % 3) == 0)
 	{
-		TCHAR ConvBuffer[4];
-		ConvBuffer[3] = TEXT('\0');
+		ElementType ConvBuffer[4];
+		ConvBuffer[3] = UE_STRING_TEXT('\0');
 		int32 WriteIndex = 0;
 		// Walk the string 3 chars at a time
 		for (int32 Index = 0; Index < Source.Len(); Index += 3, WriteIndex++)
@@ -1106,7 +1113,7 @@ bool FString::ToBlob(const FString& Source,uint8* DestBuffer,const uint32 DestSi
 			ConvBuffer[0] = Source[Index];
 			ConvBuffer[1] = Source[Index + 1];
 			ConvBuffer[2] = Source[Index + 2];
-			DestBuffer[WriteIndex] = (uint8)FCString::Atoi(ConvBuffer);
+			DestBuffer[WriteIndex] = (uint8)TCString<ElementType>::Atoi(ConvBuffer);
 		}
 		return true;
 	}
@@ -1120,7 +1127,7 @@ FString FString::FromHexBlob( const uint8* SrcBuffer, const uint32 SrcSize )
 	// Convert and append each byte in the buffer
 	for (uint32 Count = 0; Count < SrcSize; Count++)
 	{
-		Result += FString::Printf( TEXT( "%02X" ), (uint8)SrcBuffer[Count] );
+		Result += FString::Printf( UE_STRING_TEXT( "%02X" ), (uint8)SrcBuffer[Count] );
 	}
 	return Result;
 }
@@ -1132,16 +1139,16 @@ bool FString::ToHexBlob( const FString& Source, uint8* DestBuffer, const uint32
 	if (DestSize >= (uint32)(Source.Len() / 2) &&
 		 (Source.Len() % 2) == 0)
 	{
-		TCHAR ConvBuffer[3];
-		ConvBuffer[2] = TEXT( '\0' );
+		ElementType ConvBuffer[3];
+		ConvBuffer[2] = UE_STRING_TEXT( '\0' );
 		int32 WriteIndex = 0;
 		// Walk the string 2 chars at a time
-		TCHAR* End = nullptr;
+		ElementType* End = nullptr;
 		for (int32 Index = 0; Index < Source.Len(); Index += 2, WriteIndex++)
 		{
 			ConvBuffer[0] = Source[Index];
 			ConvBuffer[1] = Source[Index + 1];
-			DestBuffer[WriteIndex] = (uint8)FCString::Strtoi( ConvBuffer, &End, 16 );
+			DestBuffer[WriteIndex] = (uint8)TCString<ElementType>::Strtoi( ConvBuffer, &End, 16 );
 		}
 		return true;
 	}
@@ -1164,7 +1171,7 @@ FString FString::SanitizeFloat( double InFloat, const int32 InMinFractionalDigit
 	StripNegativeZero(InFloat);
 
 	// First create the string
-	FString TempString = FString::Printf(TEXT("%f"), InFloat);
+	FString TempString = FString::Printf(UE_STRING_TEXT("%f"), InFloat);
 	if (!TempString.IsNumeric())
 	{
 		// String did not format as a valid decimal number so avoid messing with it
@@ -1176,14 +1183,14 @@ FString FString::SanitizeFloat( double InFloat, const int32 InMinFractionalDigit
 	int32 DecimalSeparatorIndex = INDEX_NONE;
 	for (int32 CharIndex = TempString.Len() - 1; CharIndex >= 0; --CharIndex)
 	{
-		const TCHAR Char = TempString[CharIndex];
-		if (Char == TEXT('.'))
+		const ElementType Char = TempString[CharIndex];
+		if (Char == UE_STRING_TEXT('.'))
 		{
 			DecimalSeparatorIndex = CharIndex;
 			TrimIndex = FMath::Max(TrimIndex, DecimalSeparatorIndex);
 			break;
 		}
-		if (TrimIndex == INDEX_NONE && Char != TEXT('0'))
+		if (TrimIndex == INDEX_NONE && Char != UE_STRING_TEXT('0'))
 		{
 			TrimIndex = CharIndex + 1;
 		}
@@ -1197,7 +1204,7 @@ FString FString::SanitizeFloat( double InFloat, const int32 InMinFractionalDigit
 		if (TrimIndex == DecimalSeparatorIndex)
 		{
 			// Re-add the decimal separator
-			TempString.AppendChar(TEXT('.'));
+			TempString.AppendChar(UE_STRING_TEXT('.'));
 		}
 
 		const int32 NumFractionalDigits = (TempString.Len() - DecimalSeparatorIndex) - 1;
@@ -1207,7 +1214,7 @@ FString FString::SanitizeFloat( double InFloat, const int32 InMinFractionalDigit
 			TempString.Reserve(TempString.Len() + FractionalDigitsToPad);
 			for (int32 Cx = 0; Cx < FractionalDigitsToPad; ++Cx)
 			{
-				TempString.AppendChar(TEXT('0'));
+				TempString.AppendChar(UE_STRING_TEXT('0'));
 			}
 		}
 	}
@@ -1215,14 +1222,14 @@ FString FString::SanitizeFloat( double InFloat, const int32 InMinFractionalDigit
 	return TempString;
 }
 
-FString FString::Chr( TCHAR Ch )
+FString FString::Chr(ElementType Ch)
 {
-	TCHAR Temp[2]= { Ch, TEXT('\0') };
+	ElementType Temp[2]= { Ch, UE_STRING_TEXT('\0') };
 	return FString(Temp);
 }
 
 
-FString FString::ChrN( int32 NumCharacters, TCHAR Char )
+FString FString::ChrN( int32 NumCharacters, ElementType Char )
 {
 	check( NumCharacters >= 0 );
 
@@ -1232,7 +1239,7 @@ FString FString::ChrN( int32 NumCharacters, TCHAR Char )
 	{
 		Temp[Cx] = Char;
 	}
-	Temp.Data[NumCharacters] = TEXT('\0');
+	Temp.Data[NumCharacters] = UE_STRING_TEXT('\0');
 	return Temp;
 }
 
@@ -1242,7 +1249,7 @@ FString FString::LeftPad( int32 ChCount ) const
 
 	if (Pad > 0)
 	{
-		return ChrN(Pad, TEXT(' ')) + *this;
+		return ChrN(Pad, UE_STRING_TEXT(' ')) + *this;
 	}
 	else
 	{
@@ -1255,7 +1262,7 @@ FString FString::RightPad( int32 ChCount ) const
 
 	if (Pad > 0)
 	{
-		return *this + ChrN(Pad, TEXT(' '));
+		return *this + ChrN(Pad, UE_STRING_TEXT(' '));
 	}
 	else
 	{
@@ -1270,10 +1277,10 @@ bool FString::IsNumeric() const
 		return 0;
 	}
 
-	return FCString::IsNumeric(Data.GetData());
+	return TCString<ElementType>::IsNumeric(Data.GetData());
 }
 
-int32 FString::ParseIntoArray( TArray<FString>& OutArray, const TCHAR* pchDelim, const bool InCullEmpty ) const
+int32 FString::ParseIntoArray( TArray<FString>& OutArray, const ElementType* pchDelim, const bool InCullEmpty ) const
 {
 	check(pchDelim);
 	OutArray.Reset();
@@ -1289,34 +1296,34 @@ int32 FString::ParseIntoArray( TArray<FString>& OutArray, const TCHAR* pchDelim,
 	return OutArray.Num();
 }
 
-bool FString::MatchesWildcard(const TCHAR* InWildcard, int32 InWildcardLen, ESearchCase::Type SearchCase) const
+bool FString::MatchesWildcard(const ElementType* InWildcard, int32 InWildcardLen, ESearchCase::Type SearchCase) const
 {
-	const TCHAR* Target = **this;
+	const ElementType* Target = **this;
 	int32        TargetLength = Len();
 
 	if (SearchCase == ESearchCase::CaseSensitive)
 	{
-		return UE::Core::Private::MatchesWildcardRecursive<UE::Core::Private::FCompareCharsCaseSensitive>(Target, TargetLength, InWildcard, InWildcardLen);
+		return UE::Core::Private::MatchesWildcardRecursive<UE::Core::Private::TCompareCharsCaseSensitive<ElementType>>(Target, TargetLength, InWildcard, InWildcardLen);
 	}
 	else
 	{
-		return UE::Core::Private::MatchesWildcardRecursive<UE::Core::Private::FCompareCharsCaseInsensitive>(Target, TargetLength, InWildcard, InWildcardLen);
+		return UE::Core::Private::MatchesWildcardRecursive<UE::Core::Private::TCompareCharsCaseInsensitive<ElementType>>(Target, TargetLength, InWildcard, InWildcardLen);
 	}
 }
 
 
 /** Caution!! this routine is O(N^2) allocations...use it for parsing very short text or not at all */
-int32 FString::ParseIntoArrayWS( TArray<FString>& OutArray, const TCHAR* pchExtraDelim, bool InCullEmpty ) const
+int32 FString::ParseIntoArrayWS( TArray<FString>& OutArray, const ElementType* pchExtraDelim, bool InCullEmpty ) const
 {
 	// default array of White Spaces, the last entry can be replaced with the optional pchExtraDelim string
 	// (if you want to split on white space and another character)
-	const TCHAR* WhiteSpace[] = 
+	const ElementType* WhiteSpace[] =
 	{
-		TEXT(" "),
-		TEXT("\t"),
-		TEXT("\r"),
-		TEXT("\n"),
-		TEXT(""),
+		UE_STRING_TEXT(" "),
+		UE_STRING_TEXT("\t"),
+		UE_STRING_TEXT("\r"),
+		UE_STRING_TEXT("\n"),
+		UE_STRING_TEXT(""),
 	};
 
 	// start with just the standard whitespaces
@@ -1333,11 +1340,11 @@ int32 FString::ParseIntoArrayWS( TArray<FString>& OutArray, const TCHAR* pchExtr
 int32 FString::ParseIntoArrayLines(TArray<FString>& OutArray, bool InCullEmpty) const
 {
 	// default array of LineEndings
-	static const TCHAR* LineEndings[] =
-	{				
-		TEXT("\r\n"),
-		TEXT("\r"),
-		TEXT("\n"),	
+	static const ElementType* LineEndings[] =
+	{
+		UE_STRING_TEXT("\r\n"),
+		UE_STRING_TEXT("\r"),
+		UE_STRING_TEXT("\n"),	
 	};
 
 	// start with just the standard line endings
@@ -1345,12 +1352,12 @@ int32 FString::ParseIntoArrayLines(TArray<FString>& OutArray, bool InCullEmpty)
 	return ParseIntoArray(OutArray, LineEndings, NumLineEndings, InCullEmpty);
 }
 
-int32 FString::ParseIntoArray(TArray<FString>& OutArray, const TCHAR* const * DelimArray, int32 NumDelims, bool InCullEmpty) const
+int32 FString::ParseIntoArray(TArray<FString>& OutArray, const ElementType* const * DelimArray, int32 NumDelims, bool InCullEmpty) const
 {
 	// Make sure the delimit string is not null or empty
 	check(DelimArray);
 	OutArray.Reset();
-	const TCHAR *Start = Data.GetData();
+	const ElementType* Start = Data.GetData();
 	const int32 Length = Len();
 	if (Start)
 	{
@@ -1365,10 +1372,10 @@ int32 FString::ParseIntoArray(TArray<FString>& OutArray, const TCHAR* const * De
 			// Attempt each delimiter.
 			for(int32 DelimIndex = 0; DelimIndex < NumDelims; ++DelimIndex)
 			{
-				DelimiterLength = FCString::Strlen(DelimArray[DelimIndex]);
+				DelimiterLength = TCString<ElementType>::Strlen(DelimArray[DelimIndex]);
 
 				// If we found a delimiter...
-				if (FCString::Strncmp(Start + i, DelimArray[DelimIndex], DelimiterLength) == 0)
+				if (TCString<ElementType>::Strncmp(Start + i, DelimArray[DelimIndex], DelimiterLength) == 0)
 				{
 					// Mark the end of the substring.
 					SubstringEndIndex = i;
@@ -1408,7 +1415,7 @@ int32 FString::ParseIntoArray(TArray<FString>& OutArray, const TCHAR* const * De
 	return OutArray.Num();
 }
 
-FString FString::Replace(const TCHAR* From, const TCHAR* To, ESearchCase::Type SearchCase) const &
+FString FString::Replace(const ElementType* From, const ElementType* To, ESearchCase::Type SearchCase) const &
 {
 	// Previous code used to accidentally accept a nullptr replacement string - this is no longer accepted.
 	check(To);
@@ -1419,17 +1426,17 @@ FString FString::Replace(const TCHAR* From, const TCHAR* To, ESearchCase::Type S
 	}
 
 	// get a pointer into the character data
-	const TCHAR* Travel = Data.GetData();
+	const ElementType* Travel = Data.GetData();
 
 	// precalc the lengths of the replacement strings
-	int32 FromLength = FCString::Strlen(From);
-	int32 ToLength   = FCString::Strlen(To);
+	int32 FromLength = TCString<ElementType>::Strlen(From);
+	int32 ToLength   = TCString<ElementType>::Strlen(To);
 
 	FString Result;
 	while (true)
 	{
 		// look for From in the remaining string
-		const TCHAR* FromLocation = SearchCase == ESearchCase::IgnoreCase ? FCString::Stristr(Travel, From) : FCString::Strstr(Travel, From);
+		const ElementType* FromLocation = SearchCase == ESearchCase::IgnoreCase ? TCString<ElementType>::Stristr(Travel, From) : TCString<ElementType>::Strstr(Travel, From);
 		if (!FromLocation)
 		{
 			break;
@@ -1450,31 +1457,31 @@ FString FString::Replace(const TCHAR* From, const TCHAR* To, ESearchCase::Type S
 	return Result;
 }
 
-FString FString::Replace(const TCHAR* From, const TCHAR* To, ESearchCase::Type SearchCase) &&
+FString FString::Replace(const ElementType* From, const ElementType* To, ESearchCase::Type SearchCase) &&
 {
 	ReplaceInline(From, To, SearchCase);
 	return MoveTemp(*this);
 }
 
-int32 FString::ReplaceInline(const TCHAR* SearchText, const TCHAR* ReplacementText, ESearchCase::Type SearchCase)
+int32 FString::ReplaceInline(const ElementType* SearchText, const ElementType* ReplacementText, ESearchCase::Type SearchCase)
 {
 	int32 ReplacementCount = 0;
 
 	if (Len() > 0
 		&& SearchText != nullptr && *SearchText != 0
-		&& ReplacementText != nullptr && (SearchCase == ESearchCase::IgnoreCase || FCString::Strcmp(SearchText, ReplacementText) != 0))
+		&& ReplacementText != nullptr && (SearchCase == ESearchCase::IgnoreCase || TCString<ElementType>::Strcmp(SearchText, ReplacementText) != 0))
 	{
-		const int32 NumCharsToReplace = FCString::Strlen(SearchText);
-		const int32 NumCharsToInsert = FCString::Strlen(ReplacementText);
+		const int32 NumCharsToReplace = TCString<ElementType>::Strlen(SearchText);
+		const int32 NumCharsToInsert = TCString<ElementType>::Strlen(ReplacementText);
 
 		if (NumCharsToInsert == NumCharsToReplace)
 		{
-			TCHAR* Pos = SearchCase == ESearchCase::IgnoreCase ? FCString::Stristr(&(*this)[0], SearchText) : FCString::Strstr(&(*this)[0], SearchText);
+			ElementType* Pos = SearchCase == ESearchCase::IgnoreCase ? TCString<ElementType>::Stristr(&(*this)[0], SearchText) : TCString<ElementType>::Strstr(&(*this)[0], SearchText);
 			while (Pos != nullptr)
 			{
 				ReplacementCount++;
 
-				// FCString::Strcpy now inserts a terminating zero so can't use that
+				// TCString<ElementType>::Strcpy now inserts a terminating zero so can't use that
 				for (int32 i = 0; i < NumCharsToInsert; i++)
 				{
 					Pos[i] = ReplacementText[i];
@@ -1482,7 +1489,7 @@ int32 FString::ReplaceInline(const TCHAR* SearchText, const TCHAR* ReplacementTe
 
 				if (Pos + NumCharsToReplace - **this < Len())
 				{
-					Pos = SearchCase == ESearchCase::IgnoreCase ? FCString::Stristr(Pos + NumCharsToReplace, SearchText) : FCString::Strstr(Pos + NumCharsToReplace, SearchText);
+					Pos = SearchCase == ESearchCase::IgnoreCase ? TCString<ElementType>::Stristr(Pos + NumCharsToReplace, SearchText) : TCString<ElementType>::Strstr(Pos + NumCharsToReplace, SearchText);
 				}
 				else
 				{
@@ -1495,15 +1502,15 @@ int32 FString::ReplaceInline(const TCHAR* SearchText, const TCHAR* ReplacementTe
 			FString Copy(MoveTemp(*this));
 
 			// get a pointer into the character data
-			TCHAR* WritePosition = (TCHAR*)Copy.Data.GetData();
+			ElementType* WritePosition = (ElementType*)Copy.Data.GetData();
 			// look for From in the remaining string
-			TCHAR* SearchPosition = SearchCase == ESearchCase::IgnoreCase ? FCString::Stristr(WritePosition, SearchText) : FCString::Strstr(WritePosition, SearchText);
+			ElementType* SearchPosition = SearchCase == ESearchCase::IgnoreCase ? TCString<ElementType>::Stristr(WritePosition, SearchText) : TCString<ElementType>::Strstr(WritePosition, SearchText);
 			while (SearchPosition != nullptr)
 			{
 				ReplacementCount++;
 
 				// replace the first letter of the From with 0 so we can do a strcpy (FString +=)
-				*SearchPosition = TEXT('\0');
+				*SearchPosition = UE_STRING_TEXT('\0');
 
 				// copy everything up to the SearchPosition
 				(*this) += WritePosition;
@@ -1515,7 +1522,7 @@ int32 FString::ReplaceInline(const TCHAR* SearchText, const TCHAR* ReplacementTe
 				*SearchPosition = *SearchText;
 
 				WritePosition = SearchPosition + NumCharsToReplace;
-				SearchPosition = SearchCase == ESearchCase::IgnoreCase ? FCString::Stristr(WritePosition, SearchText) : FCString::Strstr(WritePosition, SearchText);
+				SearchPosition = SearchCase == ESearchCase::IgnoreCase ? TCString<ElementType>::Stristr(WritePosition, SearchText) : TCString<ElementType>::Strstr(WritePosition, SearchText);
 			}
 
 			// copy anything left over
@@ -1532,11 +1539,11 @@ int32 FString::ReplaceInline(const TCHAR* SearchText, const TCHAR* ReplacementTe
  */
 FString FString::ReplaceQuotesWithEscapedQuotes() &&
 {
-	if (Contains(TEXT("\""), ESearchCase::CaseSensitive))
+	if (Contains(UE_STRING_TEXT("\""), ESearchCase::CaseSensitive))
 	{
 		FString Copy(MoveTemp(*this));
 
-		const TCHAR* pChar = *Copy;
+		const ElementType* pChar = *Copy;
 
 		bool bEscaped = false;
 		while ( *pChar != 0 )
@@ -1545,13 +1552,13 @@ FString FString::ReplaceQuotesWithEscapedQuotes() &&
 			{
 				bEscaped = false;
 			}
-			else if ( *pChar == TCHAR('\\') )
+			else if ( *pChar == ElementType('\\') )
 			{
 				bEscaped = true;
 			}
-			else if ( *pChar == TCHAR('"') )
+			else if ( *pChar == ElementType('"') )
 			{
-				*this += TCHAR('\\');
+				*this += ElementType('\\');
 			}
 
 			*this += *pChar++;
@@ -1561,45 +1568,49 @@ FString FString::ReplaceQuotesWithEscapedQuotes() &&
 	return MoveTemp(*this);
 }
 
-static const TCHAR* CharToEscapeSeqMap[][2] =
+template <typename CharType>
+static const CharType* CharToEscapeSeqMap[6][2] =
 {
 	// Always replace \\ first to avoid double-escaping characters
-	{ TEXT("\\"), TEXT("\\\\") },
-	{ TEXT("\n"), TEXT("\\n")  },
-	{ TEXT("\r"), TEXT("\\r")  },
-	{ TEXT("\t"), TEXT("\\t")  },
-	{ TEXT("\'"), TEXT("\\'")  },
-	{ TEXT("\""), TEXT("\\\"") }
+	{ UE_STRING_TEXT("\\"), UE_STRING_TEXT("\\\\") },
+	{ UE_STRING_TEXT("\n"), UE_STRING_TEXT("\\n")  },
+	{ UE_STRING_TEXT("\r"), UE_STRING_TEXT("\\r")  },
+	{ UE_STRING_TEXT("\t"), UE_STRING_TEXT("\\t")  },
+	{ UE_STRING_TEXT("\'"), UE_STRING_TEXT("\\'")  },
+	{ UE_STRING_TEXT("\""), UE_STRING_TEXT("\\\"") }
 };
 
-static const uint32 MaxSupportedEscapeChars = UE_ARRAY_COUNT(CharToEscapeSeqMap);
+template <typename CharType>
+static const uint32 MaxSupportedEscapeChars = UE_ARRAY_COUNT(CharToEscapeSeqMap<CharType>);
 
-void FString::ReplaceCharWithEscapedCharInline(const TArray<TCHAR>* Chars/*=nullptr*/)
+void FString::ReplaceCharWithEscapedCharInline(const TArray<ElementType>* Chars/*=nullptr*/)
 {
 	if ( Len() > 0 && (Chars == nullptr || Chars->Num() > 0) )
 	{
-		for ( int32 ChIdx = 0; ChIdx < MaxSupportedEscapeChars; ChIdx++ )
+		for ( uint32 ChIdx = 0; ChIdx < MaxSupportedEscapeChars<ElementType>; ChIdx++ )
 		{
-			if ( Chars == nullptr || Chars->Contains(*(CharToEscapeSeqMap[ChIdx][0])) )
+			if ( Chars == nullptr || Chars->Contains(*(CharToEscapeSeqMap<ElementType>[ChIdx][0])) )
 			{
 				// use ReplaceInline as that won't create a copy of the string if the character isn't found
-				ReplaceInline(CharToEscapeSeqMap[ChIdx][0], CharToEscapeSeqMap[ChIdx][1]);
+				ReplaceInline(CharToEscapeSeqMap<ElementType>[ChIdx][0], CharToEscapeSeqMap<ElementType>[ChIdx][1]);
 			}
 		}
 	}
 }
 
-void FString::ReplaceEscapedCharWithCharInline(const TArray<TCHAR>* Chars/*=nullptr*/)
+void FString::ReplaceEscapedCharWithCharInline(const TArray<ElementType>* Chars/*=nullptr*/)
 {
 	if ( Len() > 0 && (Chars == nullptr || Chars->Num() > 0) )
 	{
 		// Spin CharToEscapeSeqMap backwards to ensure we're doing the inverse of ReplaceCharWithEscapedChar
-		for ( int32 ChIdx = MaxSupportedEscapeChars - 1; ChIdx >= 0; ChIdx-- )
+		for ( int32 ChIdx = MaxSupportedEscapeChars<ElementType>; ChIdx > 0; )
 		{
-			if ( Chars == nullptr || Chars->Contains(*(CharToEscapeSeqMap[ChIdx][0])) )
+			--ChIdx;
+
+			if ( Chars == nullptr || Chars->Contains(*(CharToEscapeSeqMap<ElementType>[ChIdx][0])) )
 			{
 				// use ReplaceInline as that won't create a copy of the string if the character isn't found
-				ReplaceInline(CharToEscapeSeqMap[ChIdx][1], CharToEscapeSeqMap[ChIdx][0]);
+				ReplaceInline(CharToEscapeSeqMap<ElementType>[ChIdx][1], CharToEscapeSeqMap<ElementType>[ChIdx][0]);
 			}
 		}
 	}
@@ -1615,13 +1626,13 @@ void FString::ConvertTabsToSpacesInline(const int32 InSpacesPerTab)
 	check(InSpacesPerTab > 0);
 
 	int32 TabIndex;
-	while ((TabIndex = Find(TEXT("\t"), ESearchCase::CaseSensitive)) != INDEX_NONE )
+	while ((TabIndex = Find(UE_STRING_TEXT("\t"), ESearchCase::CaseSensitive)) != INDEX_NONE )
 	{
 		FString RightSide = Mid(TabIndex+1);
 		LeftInline(TabIndex, false);
 
 		//for a tab size of 4, 
-		int32 LineBegin = Find(TEXT("\n"), ESearchCase::CaseSensitive, ESearchDir::FromEnd, TabIndex);
+		int32 LineBegin = Find(UE_STRING_TEXT("\n"), ESearchCase::CaseSensitive, ESearchDir::FromEnd, TabIndex);
 		if (LineBegin == INDEX_NONE)
 		{
 			LineBegin = 0;
@@ -1631,7 +1642,7 @@ void FString::ConvertTabsToSpacesInline(const int32 InSpacesPerTab)
 		int32 NumSpacesForTab = InSpacesPerTab - (CharactersOnLine % InSpacesPerTab);
 		for (int32 i = 0; i < NumSpacesForTab; ++i)
 		{
-			AppendChar(TEXT(' '));
+			AppendChar(UE_STRING_TEXT(' '));
 		}
 		Append(RightSide);
 	}
@@ -1640,11 +1651,11 @@ void FString::ConvertTabsToSpacesInline(const int32 InSpacesPerTab)
 // This starting size catches 99.97% of printf calls - there are about 700k printf calls per level
 #define STARTING_BUFFER_SIZE		512
 
-FString FString::PrintfImpl(const TCHAR* Fmt, ...)
+FString FString::PrintfImpl(const ElementType* Fmt, ...)
 {
 	int32		BufferSize	= STARTING_BUFFER_SIZE;
-	TCHAR	StartingBuffer[STARTING_BUFFER_SIZE];
-	TCHAR*	Buffer		= StartingBuffer;
+	ElementType	StartingBuffer[STARTING_BUFFER_SIZE];
+	ElementType*	Buffer		= StartingBuffer;
 	int32		Result		= -1;
 
 	// First try to print to a stack allocated location 
@@ -1657,12 +1668,12 @@ FString FString::PrintfImpl(const TCHAR* Fmt, ...)
 		while(Result == -1)
 		{
 			BufferSize *= 2;
-			Buffer = (TCHAR*) FMemory::Realloc( Buffer, BufferSize * sizeof(TCHAR) );
+			Buffer = (ElementType*) FMemory::Realloc( Buffer, BufferSize * sizeof(ElementType) );
 			GET_VARARGS_RESULT( Buffer, BufferSize, BufferSize-1, Fmt, Fmt, Result );
 		};
 	}
 
-	Buffer[Result] = TEXT('\0');
+	Buffer[Result] = UE_STRING_TEXT('\0');
 
 	FString ResultString(Buffer);
 
@@ -1674,11 +1685,11 @@ FString FString::PrintfImpl(const TCHAR* Fmt, ...)
 	return ResultString;
 }
 
-void FString::AppendfImpl(FString& AppendToMe, const TCHAR* Fmt, ...)
+void FString::AppendfImpl(FString& AppendToMe, const ElementType* Fmt, ...)
 {
 	int32		BufferSize = STARTING_BUFFER_SIZE;
-	TCHAR	StartingBuffer[STARTING_BUFFER_SIZE];
-	TCHAR*	Buffer = StartingBuffer;
+	ElementType	StartingBuffer[STARTING_BUFFER_SIZE];
+	ElementType*	Buffer = StartingBuffer;
 	int32		Result = -1;
 
 	// First try to print to a stack allocated location 
@@ -1691,12 +1702,12 @@ void FString::AppendfImpl(FString& AppendToMe, const TCHAR* Fmt, ...)
 		while (Result == -1)
 		{
 			BufferSize *= 2;
-			Buffer = (TCHAR*)FMemory::Realloc(Buffer, BufferSize * sizeof(TCHAR));
+			Buffer = (ElementType*)FMemory::Realloc(Buffer, BufferSize * sizeof(ElementType));
 			GET_VARARGS_RESULT(Buffer, BufferSize, BufferSize - 1, Fmt, Fmt, Result);
 		};
 	}
 
-	Buffer[Result] = TEXT('\0');
+	Buffer[Result] = UE_STRING_TEXT('\0');
 
 	AppendToMe += Buffer;
 
@@ -1710,130 +1721,140 @@ static_assert(PLATFORM_LITTLE_ENDIAN, "FString serialization needs updating to s
 
 FArchive& operator<<( FArchive& Ar, FString& A )
 {
-	// > 0 for ANSICHAR, < 0 for UTF16CHAR serialization
-	static_assert(sizeof(UTF16CHAR) == sizeof(UCS2CHAR), "UTF16CHAR and UCS2CHAR are assumed to be the same size!");
+	using ElementType = FString::ElementType;
 
-	if (Ar.IsLoading())
+	if constexpr (std::is_same_v<ElementType, TCHAR>)
 	{
-		int32 SaveNum = 0;
-		Ar << SaveNum;
+		// > 0 for ANSICHAR, < 0 for UTF16CHAR serialization
+		static_assert(sizeof(UTF16CHAR) == sizeof(UCS2CHAR), "UTF16CHAR and UCS2CHAR are assumed to be the same size!");
 
-		bool bLoadUnicodeChar = SaveNum < 0;
-		if (bLoadUnicodeChar)
+		if (Ar.IsLoading())
 		{
-			// If SaveNum cannot be negated due to integer overflow, Ar is corrupted.
-			if (SaveNum == MIN_int32)
+			int32 SaveNum = 0;
+			Ar << SaveNum;
+
+			bool bLoadUnicodeChar = SaveNum < 0;
+			if (bLoadUnicodeChar)
+			{
+				// If SaveNum cannot be negated due to integer overflow, Ar is corrupted.
+				if (SaveNum == MIN_int32)
+				{
+					Ar.SetCriticalError();
+					UE_LOG(LogCore, Error, TEXT("Archive is corrupted"));
+					return Ar;
+				}
+
+				SaveNum = -SaveNum;
+			}
+
+			int64 MaxSerializeSize = Ar.GetMaxSerializeSize();
+			// Protect against network packets allocating too much memory
+			if ((MaxSerializeSize > 0) && (SaveNum > MaxSerializeSize))
 			{
 				Ar.SetCriticalError();
-				UE_LOG(LogCore, Error, TEXT("Archive is corrupted"));
+				UE_LOG(LogCore, Error, TEXT("String is too large (Size: %i, Max: %i)"), SaveNum, MaxSerializeSize);
 				return Ar;
 			}
 
-			SaveNum = -SaveNum;
-		}
+			// Resize the array only if it passes the above tests to prevent rogue packets from crashing
+			A.Data.Empty(SaveNum);
+			A.Data.AddUninitialized(SaveNum);
 
-		int64 MaxSerializeSize = Ar.GetMaxSerializeSize();
-		// Protect against network packets allocating too much memory
-		if ((MaxSerializeSize > 0) && (SaveNum > MaxSerializeSize))
-		{
-			Ar.SetCriticalError();
-			UE_LOG(LogCore, Error, TEXT("String is too large (Size: %i, Max: %i)"), SaveNum, MaxSerializeSize);
-			return Ar;
-		}
-
-		// Resize the array only if it passes the above tests to prevent rogue packets from crashing
-		A.Data.Empty           (SaveNum);
-		A.Data.AddUninitialized(SaveNum);
-
-		if (SaveNum)
-		{
-			if (bLoadUnicodeChar)
+			if (SaveNum)
 			{
-				// read in the unicode string
-				auto Passthru = StringMemoryPassthru<UCS2CHAR>(A.Data.GetData(), SaveNum, SaveNum);
-				Ar.Serialize(Passthru.Get(), SaveNum * sizeof(UCS2CHAR));
-				if (Ar.IsByteSwapping())
+				if (bLoadUnicodeChar)
 				{
-					for (int32 CharIndex = 0; CharIndex < SaveNum; ++CharIndex)
+					// read in the unicode string
+					auto Passthru = StringMemoryPassthru<UCS2CHAR>(A.Data.GetData(), SaveNum, SaveNum);
+					Ar.Serialize(Passthru.Get(), SaveNum * sizeof(UCS2CHAR));
+					if (Ar.IsByteSwapping())
+					{
+						for (int32 CharIndex = 0; CharIndex < SaveNum; ++CharIndex)
+						{
+							Passthru.Get()[CharIndex] = ByteSwap(Passthru.Get()[CharIndex]);
+						}
+					}
+					// Ensure the string has a null terminator
+					Passthru.Get()[SaveNum - 1] = '\0';
+					Passthru.Apply();
+
+					// Inline combine any surrogate pairs in the data when loading into a UTF-32 string
+					StringConv::InlineCombineSurrogates(A);
+
+					// Since Microsoft's vsnwprintf implementation raises an invalid parameter warning
+					// with a character of 0xffff, scan for it and terminate the string there.
+					// 0xffff isn't an actual Unicode character anyway.
+					int Index = 0;
+					if (A.FindChar(0xffff, Index))
 					{
-						Passthru.Get()[CharIndex] = ByteSwap(Passthru.Get()[CharIndex]);
+						A[Index] = UE_STRING_TEXT('\0');
+						A.TrimToNullTerminator();
 					}
 				}
-				// Ensure the string has a null terminator
-				Passthru.Get()[SaveNum-1] = '\0';
-				Passthru.Apply();
-
-				// Inline combine any surrogate pairs in the data when loading into a UTF-32 string
-				StringConv::InlineCombineSurrogates(A);
-
-				// Since Microsoft's vsnwprintf implementation raises an invalid parameter warning
-				// with a character of 0xffff, scan for it and terminate the string there.
-				// 0xffff isn't an actual Unicode character anyway.
-				int Index = 0;
-				if(A.FindChar(0xffff, Index))
+				else
 				{
-					A[Index] = TEXT('\0');
-					A.TrimToNullTerminator();
+					auto Passthru = StringMemoryPassthru<ANSICHAR>(A.Data.GetData(), SaveNum, SaveNum);
+					Ar.Serialize(Passthru.Get(), SaveNum * sizeof(ANSICHAR));
+					// Ensure the string has a null terminator
+					Passthru.Get()[SaveNum - 1] = '\0';
+					Passthru.Apply();
 				}
-			}
-			else
-			{
-				auto Passthru = StringMemoryPassthru<ANSICHAR>(A.Data.GetData(), SaveNum, SaveNum);
-				Ar.Serialize(Passthru.Get(), SaveNum * sizeof(ANSICHAR));
-				// Ensure the string has a null terminator
-				Passthru.Get()[SaveNum-1] = '\0';
-				Passthru.Apply();
-			}
 
-			// Throw away empty string.
-			if (SaveNum == 1)
-			{
-				A.Data.Empty();
+				// Throw away empty string.
+				if (SaveNum == 1)
+				{
+					A.Data.Empty();
+				}
 			}
 		}
-	}
-	else
-	{
-		A.Data.CountBytes(Ar);
-
-		const bool bSaveUnicodeChar = Ar.IsForcingUnicode() || !FCString::IsPureAnsi(*A);
-		if (bSaveUnicodeChar)
+		else
 		{
-			// Note: This is a no-op on platforms that are using a 16-bit TCHAR
- 			FTCHARToUTF16 UTF16String(*A, A.Len() + 1); // include the null terminator
-			int32 Num = UTF16String.Length() + 1; // include the null terminator
-
-			int32 SaveNum = -Num;
-			Ar << SaveNum;
+			A.Data.CountBytes(Ar);
 
-			if (Num)
+			const bool bSaveUnicodeChar = Ar.IsForcingUnicode() || !TCString<ElementType>::IsPureAnsi(*A);
+			if (bSaveUnicodeChar)
 			{
-				if (!Ar.IsByteSwapping())
-				{
-					Ar.Serialize((void*)UTF16String.Get(), sizeof(UTF16CHAR) * Num);
-				}
-				else
+				// This preprocessor block should not be necessary when the StringCast above understands UTF16CHAR.
+				FTCHARToUTF16 UTF16String(*A, A.Len() + 1); // include the null terminator
+				int32 Num = UTF16String.Length() + 1; // include the null terminator
+
+				int32 SaveNum = -Num;
+				Ar << SaveNum;
+
+				if (Num)
 				{
-					TArray<UTF16CHAR> Swapped(UTF16String.Get(), Num);
-					for (int32 CharIndex = 0; CharIndex < Num; ++CharIndex)
+					if (!Ar.IsByteSwapping())
+					{
+						Ar.Serialize((void*)UTF16String.Get(), sizeof(UTF16CHAR) * Num);
+					}
+					else
 					{
-						Swapped[CharIndex] = ByteSwap(Swapped[CharIndex]);
+						TArray<UTF16CHAR> Swapped(UTF16String.Get(), Num);
+						for (int32 CharIndex = 0; CharIndex < Num; ++CharIndex)
+						{
+							Swapped[CharIndex] = ByteSwap(Swapped[CharIndex]);
+						}
+						Ar.Serialize((void*)Swapped.GetData(), sizeof(UTF16CHAR) * Num);
 					}
-					Ar.Serialize((void*)Swapped.GetData(), sizeof(UTF16CHAR) * Num);
 				}
 			}
-		}
-		else
-		{
-			int32 Num = A.Data.Num();
-			Ar << Num;
-
-			if (Num)
+			else
 			{
-				Ar.Serialize((void*)StringCast<ANSICHAR>(A.Data.GetData(), Num).Get(), sizeof(ANSICHAR) * Num);
+				int32 Num = A.Data.Num();
+				Ar << Num;
+
+				if (Num)
+				{
+					Ar.Serialize((void*)StringCast<ANSICHAR>(A.Data.GetData(), Num).Get(), sizeof(ANSICHAR) * Num);
+				}
 			}
 		}
 	}
+	else
+	{
+		// Can just serialize as UTF-8 always
+		check(false); // TODO
+	}
 
 	return Ar;
 }
@@ -1845,31 +1866,33 @@ int32 HexToBytes(const FString& HexString, uint8* OutBytes)
 
 int32 FindMatchingClosingParenthesis(const FString& TargetString, const int32 StartSearch)
 {
+	using ElementType = FString::ElementType;
+
 	check(StartSearch >= 0 && StartSearch <= TargetString.Len());// Check for usage, we do not accept INDEX_NONE like other string functions
 
-	const TCHAR* const StartPosition = (*TargetString) + StartSearch;
-	const TCHAR* CurrPosition = StartPosition;
+	const ElementType* const StartPosition = (*TargetString) + StartSearch;
+	const ElementType* CurrPosition = StartPosition;
 	int32 ParenthesisCount = 0;
 
 	// Move to first open parenthesis
-	while (*CurrPosition != 0 && *CurrPosition != TEXT('('))
+	while (*CurrPosition != 0 && *CurrPosition != UE_STRING_TEXT('('))
 	{
 		++CurrPosition;
 	}
 
 	// Did we find the open parenthesis
-	if (*CurrPosition == TEXT('('))
+	if (*CurrPosition == UE_STRING_TEXT('('))
 	{
 		++ParenthesisCount;
 		++CurrPosition;
 
 		while (*CurrPosition != 0 && ParenthesisCount > 0)
 		{
-			if (*CurrPosition == TEXT('('))
+			if (*CurrPosition == UE_STRING_TEXT('('))
 			{
 				++ParenthesisCount;
 			}
-			else if (*CurrPosition == TEXT(')'))
+			else if (*CurrPosition == UE_STRING_TEXT(')'))
 			{
 				--ParenthesisCount;
 			}
@@ -1877,7 +1900,7 @@ int32 FindMatchingClosingParenthesis(const FString& TargetString, const int32 St
 		}
 
 		// Did we find the matching close parenthesis
-		if (ParenthesisCount == 0 && *(CurrPosition - 1) == TEXT(')'))
+		if (ParenthesisCount == 0 && *(CurrPosition - 1) == UE_STRING_TEXT(')'))
 		{
 			return StartSearch + UE_PTRDIFF_TO_INT32((CurrPosition - 1) - StartPosition);
 		}
@@ -1886,8 +1909,10 @@ int32 FindMatchingClosingParenthesis(const FString& TargetString, const int32 St
 	return INDEX_NONE;
 }
 
-FString SlugStringForValidName(const FString& DisplayString, const TCHAR* ReplaceWith /*= TEXT("")*/)
+FString SlugStringForValidName(const FString& DisplayString, const FString::ElementType* ReplaceWith /*= UE_STRING_TEXT("")*/)
 {
+	using ElementType = FString::ElementType;
+
 	FString GeneratedName = DisplayString;
 
 	// Convert the display label, which may consist of just about any possible character, into a
@@ -1895,7 +1920,7 @@ FString SlugStringForValidName(const FString& DisplayString, const TCHAR* Replac
 	{
 		for ( int32 BadCharacterIndex = 0; BadCharacterIndex < UE_ARRAY_COUNT(INVALID_OBJECTNAME_CHARACTERS) - 1; ++BadCharacterIndex )
 		{
-			const TCHAR TestChar[2] = { INVALID_OBJECTNAME_CHARACTERS[BadCharacterIndex], TEXT('\0') };
+			const ElementType TestChar[2] = { INVALID_OBJECTNAME_CHARACTERS[BadCharacterIndex], UE_STRING_TEXT('\0') };
 			const int32 NumReplacedChars = GeneratedName.ReplaceInline(TestChar, ReplaceWith);
 		}
 	}
@@ -1905,15 +1930,17 @@ FString SlugStringForValidName(const FString& DisplayString, const TCHAR* Replac
 
 void FTextRange::CalculateLineRangesFromString(const FString& Input, TArray<FTextRange>& LineRanges)
 {
+	using ElementType = FString::ElementType;
+
 	int32 LineBeginIndex = 0;
 
 	// Loop through splitting at new-lines
-	const TCHAR* const InputStart = *Input;
-	for (const TCHAR* CurrentChar = InputStart; CurrentChar && *CurrentChar; ++CurrentChar)
+	const ElementType* const InputStart = *Input;
+	for (const ElementType* CurrentChar = InputStart; CurrentChar && *CurrentChar; ++CurrentChar)
 	{
-		// Handle a chain of \r\n slightly differently to stop the FChar::IsLinebreak adding two separate new-lines
+		// Handle a chain of \r\n slightly differently to stop the TChar<ElementType>::IsLinebreak adding two separate new-lines
 		const bool bIsWindowsNewLine = (*CurrentChar == '\r' && *(CurrentChar + 1) == '\n');
-		if (bIsWindowsNewLine || FChar::IsLinebreak(*CurrentChar))
+		if (bIsWindowsNewLine || TChar<ElementType>::IsLinebreak(*CurrentChar))
 		{
 			const int32 LineEndIndex = UE_PTRDIFF_TO_INT32(CurrentChar - InputStart);
 			check(LineEndIndex >= LineBeginIndex);
@@ -1938,3 +1965,6 @@ void StringConv::InlineCombineSurrogates(FString& Str)
 {
 	InlineCombineSurrogates_Array(Str.GetCharArray());
 }
+
+#undef UE_STRING_TEXT
+#undef UE_STRING_COMPILING_UTF8
\ No newline at end of file
diff --git a/Engine/Source/Runtime/Core/Public/Containers/UnrealString.h b/Engine/Source/Runtime/Core/Public/Containers/UnrealString.h
--- a/Engine/Source/Runtime/Core/Public/Containers/UnrealString.h
+++ b/Engine/Source/Runtime/Core/Public/Containers/UnrealString.h
@@ -32,6 +32,9 @@
 
 #include <type_traits>
 
+#define UE_STRING_COMPILING_UTF8 0
+#define UE_STRING_TEXT(x) TEXT(x)
+
 struct FStringFormatArg;
 template<typename InKeyType,typename InValueType,typename SetAllocator ,typename KeyFuncs > class TMap;
 
@@ -56,10 +59,11 @@ class FString
 {
 public:
 	using AllocatorType = TSizedDefaultAllocator<32>;
+	using ElementType   = TCHAR;
 
 private:
 	/** Array holding the character data */
-	typedef TArray<TCHAR, AllocatorType> DataType;
+	typedef TArray<ElementType, AllocatorType> DataType;
 	DataType Data;
 
 	template <typename RangeType>
@@ -71,9 +75,9 @@ private:
 	};
 
 	template <typename CharRangeType>
-	struct TIsRangeOfTCHAR
+	struct TIsRangeOfElementType
 	{
-		enum { Value = std::is_same_v<TCHAR, TRangeElementType<CharRangeType>> };
+		enum { Value = std::is_same_v<ElementType, TRangeElementType<CharRangeType>> };
 	};
 
 	/** Trait testing whether a type is a contiguous range of characters, and not CharType[]. */
@@ -83,28 +87,26 @@ private:
 		TNot<TIsArray<typename TRemoveReference<CharRangeType>::Type>>,
 		TIsRangeOfCharType<CharRangeType>>;
 
-	/** Trait testing whether a type is a contiguous range of characters, and not CharType[] and not FString. */
+	/** Trait testing whether a type is a contiguous range of characters, and not CharType[] and not a string class type. */
 	template <typename CharRangeType>
-	using TIsCharRangeNotCArrayNotFString = TAnd<
+	using TIsCharRangeNotCArrayNotStringClass = TAnd<
 		TIsCharRangeNotCArray<CharRangeType>,
 		TNot<TIsDerivedFrom<typename TDecay<CharRangeType>::Type, FString>>>;
 
-	/** Trait testing whether a type is a contiguous range of TCHAR, and not TCHAR[]. */
+	/** Trait testing whether a type is a contiguous range of ElementType, and not CharType[]. */
 	template <typename CharRangeType>
 	using TIsTCharRangeNotCArray = TAnd<
 		TIsContiguousContainer<CharRangeType>,
 		TNot<TIsArray<typename TRemoveReference<CharRangeType>::Type>>,
-		TIsRangeOfTCHAR<CharRangeType>>;
+		TIsRangeOfElementType<CharRangeType>>;
 
-	/** Trait testing whether a type is a contiguous range of TCHAR, and not TCHAR[] and not FString. */
+	/** Trait testing whether a type is a contiguous range of ElementType, and not CharType[] and not a string class type. */
 	template <typename CharRangeType>
-	using TIsTCharRangeNotCArrayNotFString = TAnd<
+	using TIsTCharRangeNotCArrayNotStringClass = TAnd<
 		TIsTCharRangeNotCArray<CharRangeType>,
 		TNot<TIsDerivedFrom<typename TDecay<CharRangeType>::Type, FString>>>;
 
 public:
-	using ElementType = TCHAR;
-
 	FString() = default;
 	FString(FString&&) = default;
 	FString(const FString&) = default;
@@ -151,14 +153,14 @@ public:
 	CORE_API FString(const UTF8CHAR* Str, int32 ExtraSlack);
 	CORE_API FString(const UCS2CHAR* Str, int32 ExtraSlack);
 
-	/** Construct from contiguous range of characters such as FStringView or FStringBuilderBase */
-	template <typename CharRangeType, typename TEnableIf<TIsCharRangeNotCArrayNotFString<CharRangeType>::Value>::Type* = nullptr>
+	/** Construct from contiguous range of characters such as a string view or string builder */
+	template <typename CharRangeType, typename TEnableIf<TIsCharRangeNotCArrayNotStringClass<CharRangeType>::Value>::Type* = nullptr>
 	FORCEINLINE explicit FString(CharRangeType&& Str) : FString(GetNum(Str), GetData(Forward<CharRangeType>(Str)))
 	{
 	}
 
 	/** Construct from contiguous range of characters with extra slack on top of original string length */
-	template <typename CharRangeType, typename TEnableIf<TIsCharRangeNotCArrayNotFString<CharRangeType>::Value>::Type* = nullptr>
+	template <typename CharRangeType, typename TEnableIf<TIsCharRangeNotCArrayNotStringClass<CharRangeType>::Value>::Type* = nullptr>
 	explicit FString(CharRangeType&& Str, int32 ExtraSlack)
 	{
 		uint32 InLen = GetNum(Str);
@@ -195,7 +197,7 @@ public:
 	}
 #endif
 
-	CORE_API FString& operator=(const TCHAR* Str);
+	CORE_API FString& operator=(const ElementType* Str);
 
 	template <typename CharRangeType, typename TEnableIf<TIsTCharRangeNotCArray<CharRangeType>::Value>::Type* = nullptr>
 	FORCEINLINE FString& operator=(CharRangeType&& Range)
@@ -205,7 +207,7 @@ public:
 	}
 
 private:
-	CORE_API void AssignRange(const TCHAR* Str, int32 Len);
+	CORE_API void AssignRange(const ElementType* Str, int32 Len);
 
 public:
 	/**
@@ -214,7 +216,7 @@ public:
 	 * @param Index into string
 	 * @return Character at Index
 	 */
-	FORCEINLINE TCHAR& operator[]( int32 Index ) UE_LIFETIMEBOUND
+	FORCEINLINE ElementType& operator[]( int32 Index ) UE_LIFETIMEBOUND
 	{
 		checkf(IsValidIndex(Index), TEXT("String index out of bounds: Index %i from a string with a length of %i"), Index, Len());
 		return Data.GetData()[Index];
@@ -226,7 +228,7 @@ public:
 	 * @param Index into string
 	 * @return const Character at Index
 	 */
-	FORCEINLINE const TCHAR& operator[]( int32 Index ) const UE_LIFETIMEBOUND
+	FORCEINLINE const ElementType& operator[]( int32 Index ) const UE_LIFETIMEBOUND
 	{
 		checkf(IsValidIndex(Index), TEXT("String index out of bounds: Index %i from a string with a length of %i"), Index, Len());
 		return Data.GetData()[Index];
@@ -324,11 +326,11 @@ public:
 	/**
 	 * Get pointer to the string
 	 *
-	 * @Return Pointer to Array of TCHAR if Num, otherwise the empty string
+	 * @Return Pointer to Array of ElementType if Num, otherwise the empty string
 	 */
-	UE_NODISCARD FORCEINLINE const TCHAR* operator*() const UE_LIFETIMEBOUND
+	UE_NODISCARD FORCEINLINE const ElementType* operator*() const UE_LIFETIMEBOUND
 	{
-		return Data.Num() ? Data.GetData() : TEXT("");
+		return Data.Num() ? Data.GetData() : UE_STRING_TEXT("");
 	}
 
 	/** 
@@ -349,7 +351,7 @@ public:
 	}
 
 #ifdef __OBJC__
-	/** Convert FString to Objective-C NSString */
+	/** Convert the string to Objective-C NSString */
     CORE_API NSString* GetNSString() const;
 #endif
 
@@ -395,7 +397,7 @@ public:
 	}
 
 	/** Append a single character and return a reference to this */
-	CORE_API FString& AppendChar(TCHAR InChar);
+	CORE_API FString& AppendChar(ElementType InChar);
 
 	/** Append a string and return a reference to this */
 	template <typename StrType>
@@ -411,9 +413,9 @@ public:
 	>
 	FORCEINLINE FString& operator+=(AppendedCharType Char)
 	{
-		if constexpr (TIsCharEncodingSimplyConvertibleTo_V<AppendedCharType, TCHAR>)
+		if constexpr (TIsCharEncodingSimplyConvertibleTo_V<AppendedCharType, ElementType>)
 		{
-			return AppendChar((TCHAR)Char);
+			return AppendChar((ElementType)Char);
 		}
 		else
 		{
@@ -422,7 +424,7 @@ public:
 		}
 	}
 
-	CORE_API void InsertAt(int32 Index, TCHAR Character);
+	CORE_API void InsertAt(int32 Index, ElementType Character);
 	CORE_API void InsertAt(int32 Index, const FString& Characters);
 
 	/**
@@ -440,10 +442,10 @@ public:
 	 * @param InPrefix the prefix to search for at the start of the string to remove.
 	 * @return true if the prefix was removed, otherwise false.
 	 */
-	template <typename TCharRangeType, std::enable_if_t<TIsCharRangeNotCArrayNotFString<TCharRangeType>::Value>* = nullptr>
+	template <typename TCharRangeType, std::enable_if_t<TIsCharRangeNotCArrayNotStringClass<TCharRangeType>::Value>* = nullptr>
 	bool RemoveFromStart(TCharRangeType&& InPrefix, ESearchCase::Type SearchCase = ESearchCase::IgnoreCase)
 	{
-		static_assert(std::is_same_v<typename TElementType<TCharRangeType>::Type, TCHAR>, "Expected a range of TCHAR");
+		static_assert(std::is_same_v<typename TElementType<TCharRangeType>::Type, ElementType>, "Expected a range of ElementType");
 		return RemoveFromStart(GetData(InPrefix), GetNum(InPrefix), SearchCase);
 	}
 
@@ -453,7 +455,7 @@ public:
 	 * @param InPrefix the prefix to search for at the start of the string to remove.
 	 * @return true if the prefix was removed, otherwise false.
 	 */
-	bool RemoveFromStart(const TCHAR* InPrefix, ESearchCase::Type SearchCase = ESearchCase::IgnoreCase)
+	bool RemoveFromStart(const ElementType* InPrefix, ESearchCase::Type SearchCase = ESearchCase::IgnoreCase)
 	{
 		return RemoveFromStart(InPrefix, InPrefix ? FCString::Strlen(InPrefix) : 0, SearchCase);
 	}
@@ -476,7 +478,7 @@ public:
 	 * @param InPrefixLen length of InPrefix
 	 * @return true if the prefix was removed, otherwise false.
 	 */
-	CORE_API bool RemoveFromStart(const TCHAR* InPrefix, int32 InPrefixLen, ESearchCase::Type SearchCase = ESearchCase::IgnoreCase);
+	CORE_API bool RemoveFromStart(const ElementType* InPrefix, int32 InPrefixLen, ESearchCase::Type SearchCase = ESearchCase::IgnoreCase);
 
 	/**
 	 * Removes the text from the end of the string if it exists.
@@ -484,10 +486,10 @@ public:
 	 * @param InSuffix the suffix to search for at the end of the string to remove.
 	 * @return true if the suffix was removed, otherwise false.
 	 */
-	template <typename TCharRangeType, std::enable_if_t<TIsCharRangeNotCArrayNotFString<TCharRangeType>::Value>* = nullptr>
+	template <typename TCharRangeType, std::enable_if_t<TIsCharRangeNotCArrayNotStringClass<TCharRangeType>::Value>* = nullptr>
 	bool RemoveFromEnd(TCharRangeType&& InSuffix, ESearchCase::Type SearchCase = ESearchCase::IgnoreCase)
 	{
-		static_assert(std::is_same_v<typename TElementType<TCharRangeType>::Type, TCHAR>, "Expected a range of TCHAR");
+		static_assert(std::is_same_v<typename TElementType<TCharRangeType>::Type, ElementType>, "Expected a range of ElementType");
 		return RemoveFromEnd(GetData(InSuffix), GetNum(InSuffix), SearchCase);
 	}
 
@@ -497,7 +499,7 @@ public:
 	 * @param InSuffix the suffix to search for at the end of the string to remove.
 	 * @return true if the suffix was removed, otherwise false.
 	 */
-	bool RemoveFromEnd(const TCHAR* InSuffix, ESearchCase::Type SearchCase = ESearchCase::IgnoreCase)
+	bool RemoveFromEnd(const ElementType* InSuffix, ESearchCase::Type SearchCase = ESearchCase::IgnoreCase)
 	{
 		return RemoveFromEnd(InSuffix, InSuffix ? FCString::Strlen(InSuffix) : 0, SearchCase);
 	}
@@ -520,7 +522,7 @@ public:
 	 * @param InSuffixLen length of InSuffix
 	 * @return true if the suffix was removed, otherwise false.
 	 */
-	CORE_API bool RemoveFromEnd(const TCHAR* InSuffix, int32 InSuffixLen, ESearchCase::Type SearchCase = ESearchCase::IgnoreCase);
+	CORE_API bool RemoveFromEnd(const ElementType* InSuffix, int32 InSuffixLen, ESearchCase::Type SearchCase = ESearchCase::IgnoreCase);
 
 	/**
 	 * Concatenate this path with given path ensuring the / character is used between them
@@ -528,12 +530,12 @@ public:
 	 * @param Str       Pointer to an array of TCHARs (not necessarily null-terminated) to be concatenated onto the end of this.
 	 * @param StrLength Exact number of characters from Str to append.
 	 */
-	CORE_API void PathAppend(const TCHAR* Str, int32 StrLength);
+	CORE_API void PathAppend(const ElementType* Str, int32 StrLength);
 
 	/**
-	 * Concatenates an FString with a TCHAR.
+	 * Concatenates a string with a character.
 	 * 
-	 * @param Lhs The FString on the left-hand-side of the expression.
+	 * @param Lhs The string on the left-hand-side of the expression.
 	 * @param Rhs The char on the right-hand-side of the expression.
 	 *
 	 * @return The concatenated string.
@@ -550,9 +552,9 @@ public:
 	}
 
 	/**
-	 * Concatenates an FString with a TCHAR.
+	 * Concatenates a string with a character.
 	 * 
-	 * @param Lhs The FString on the left-hand-side of the expression.
+	 * @param Lhs The string on the left-hand-side of the expression.
 	 * @param Rhs The char on the right-hand-side of the expression.
 	 *
 	 * @return The concatenated string.
@@ -572,42 +574,42 @@ private:
 	UE_NODISCARD static CORE_API FString ConcatFF(const FString& Lhs, const FString& Rhs);
 	UE_NODISCARD static CORE_API FString ConcatFF(FString&& Lhs, const FString& Rhs);
 	UE_NODISCARD static CORE_API FString ConcatFF(const FString& Lhs, FString&& Rhs);
-	UE_NODISCARD static CORE_API FString ConcatFF(FString&& Lhs,	FString&& Rhs);
-	UE_NODISCARD static CORE_API FString ConcatFC(const FString& Lhs, const TCHAR* Rhs);
-	UE_NODISCARD static CORE_API FString ConcatFC(FString&& Lhs,	const TCHAR* Rhs);
-	UE_NODISCARD static CORE_API FString ConcatCF(const TCHAR* Lhs, const FString& Rhs);
-	UE_NODISCARD static CORE_API FString ConcatCF(const TCHAR* Lhs, FString&& Rhs);
-	UE_NODISCARD static CORE_API FString ConcatFR(const FString& Lhs, const TCHAR* Rhs, int32 RhsLen);
-	UE_NODISCARD static CORE_API FString ConcatFR(FString&& Lhs,	const TCHAR* Rhs, int32 RhsLen);
-	UE_NODISCARD static CORE_API FString ConcatRF(const TCHAR* Lhs, int32 LhsLen, const FString& Rhs);
-	UE_NODISCARD static CORE_API FString ConcatRF(const TCHAR* Lhs, int32 LhsLen, FString&& Rhs);
+	UE_NODISCARD static CORE_API FString ConcatFF(FString&& Lhs, FString&& Rhs);
+	UE_NODISCARD static CORE_API FString ConcatFC(const FString& Lhs, const ElementType* Rhs);
+	UE_NODISCARD static CORE_API FString ConcatFC(FString&& Lhs,	const ElementType* Rhs);
+	UE_NODISCARD static CORE_API FString ConcatCF(const ElementType* Lhs, const FString& Rhs);
+	UE_NODISCARD static CORE_API FString ConcatCF(const ElementType* Lhs, FString&& Rhs);
+	UE_NODISCARD static CORE_API FString ConcatFR(const FString& Lhs, const ElementType* Rhs, int32 RhsLen);
+	UE_NODISCARD static CORE_API FString ConcatFR(FString&& Lhs,	const ElementType* Rhs, int32 RhsLen);
+	UE_NODISCARD static CORE_API FString ConcatRF(const ElementType* Lhs, int32 LhsLen, const FString& Rhs);
+	UE_NODISCARD static CORE_API FString ConcatRF(const ElementType* Lhs, int32 LhsLen, FString&& Rhs);
 
 public:
-	UE_NODISCARD FORCEINLINE friend FString operator+(const FString& Lhs, const FString& Rhs)	{ return ConcatFF(Lhs, Rhs); }
-	UE_NODISCARD FORCEINLINE friend FString operator+(FString&& Lhs, const FString& Rhs)		{ return ConcatFF(MoveTemp(Lhs), Rhs); }
-	UE_NODISCARD FORCEINLINE friend FString operator+(const FString& Lhs, FString&& Rhs)		{ return ConcatFF(Lhs,MoveTemp(Rhs)); }
-	UE_NODISCARD FORCEINLINE friend FString operator+(FString&& Lhs, FString&& Rhs)				{ return ConcatFF(MoveTemp(Lhs), MoveTemp(Rhs)); }
-	UE_NODISCARD FORCEINLINE friend FString operator+(const TCHAR* Lhs, const FString& Rhs)		{ return ConcatCF(Lhs, Rhs); }
-	UE_NODISCARD FORCEINLINE friend FString operator+(const TCHAR* Lhs, FString&& Rhs)			{ return ConcatCF(Lhs, MoveTemp(Rhs)); }
-	UE_NODISCARD FORCEINLINE friend FString operator+(const FString& Lhs, const TCHAR* Rhs)		{ return ConcatFC(Lhs, Rhs); }
-	UE_NODISCARD FORCEINLINE friend FString operator+(FString&& Lhs, const TCHAR* Rhs)			{ return ConcatFC(MoveTemp(Lhs), Rhs); }
-
-	template <typename T, typename TEnableIf<TIsTCharRangeNotCArrayNotFString<T>::Value>::Type* = nullptr>
-	UE_NODISCARD FORCEINLINE friend FString operator+(T&& Lhs, const FString& Rhs)				{ return ConcatRF(GetData(Lhs), GetNum(Lhs), Rhs); }
-	template <typename T, typename TEnableIf<TIsTCharRangeNotCArrayNotFString<T>::Value>::Type* = nullptr>
-	UE_NODISCARD FORCEINLINE friend FString operator+(T&& Lhs, FString&& Rhs)					{ return ConcatRF(GetData(Lhs), GetNum(Lhs), MoveTemp(Rhs)); }
-	template <typename T, typename TEnableIf<TIsTCharRangeNotCArrayNotFString<T>::Value>::Type* = nullptr>
-	UE_NODISCARD FORCEINLINE friend FString operator+(const FString& Lhs, T&& Rhs)				{ return ConcatFR(Lhs, GetData(Rhs), GetNum(Rhs)); }
-	template <typename T, typename TEnableIf<TIsTCharRangeNotCArrayNotFString<T>::Value>::Type* = nullptr>
-	UE_NODISCARD FORCEINLINE friend FString operator+(FString&& Lhs, T&& Rhs)					{ return ConcatFR(MoveTemp(Lhs), GetData(Rhs), GetNum(Rhs)); }
+	UE_NODISCARD FORCEINLINE friend FString operator+(const FString& Lhs, const FString& Rhs)			{ return ConcatFF(Lhs, Rhs); }
+	UE_NODISCARD FORCEINLINE friend FString operator+(FString&& Lhs, const FString& Rhs)				{ return ConcatFF(MoveTemp(Lhs), Rhs); }
+	UE_NODISCARD FORCEINLINE friend FString operator+(const FString& Lhs, FString&& Rhs)				{ return ConcatFF(Lhs,MoveTemp(Rhs)); }
+	UE_NODISCARD FORCEINLINE friend FString operator+(FString&& Lhs, FString&& Rhs)						{ return ConcatFF(MoveTemp(Lhs), MoveTemp(Rhs)); }
+	UE_NODISCARD FORCEINLINE friend FString operator+(const ElementType* Lhs, const FString& Rhs)		{ return ConcatCF(Lhs, Rhs); }
+	UE_NODISCARD FORCEINLINE friend FString operator+(const ElementType* Lhs, FString&& Rhs)			{ return ConcatCF(Lhs, MoveTemp(Rhs)); }
+	UE_NODISCARD FORCEINLINE friend FString operator+(const FString& Lhs, const ElementType* Rhs)		{ return ConcatFC(Lhs, Rhs); }
+	UE_NODISCARD FORCEINLINE friend FString operator+(FString&& Lhs, const ElementType* Rhs)			{ return ConcatFC(MoveTemp(Lhs), Rhs); }
+
+	template <typename T, typename TEnableIf<TIsTCharRangeNotCArrayNotStringClass<T>::Value>::Type* = nullptr>
+	UE_NODISCARD FORCEINLINE friend FString operator+(T&& Lhs, const FString& Rhs)						{ return ConcatRF(GetData(Lhs), GetNum(Lhs), Rhs); }
+	template <typename T, typename TEnableIf<TIsTCharRangeNotCArrayNotStringClass<T>::Value>::Type* = nullptr>
+	UE_NODISCARD FORCEINLINE friend FString operator+(T&& Lhs, FString&& Rhs)							{ return ConcatRF(GetData(Lhs), GetNum(Lhs), MoveTemp(Rhs)); }
+	template <typename T, typename TEnableIf<TIsTCharRangeNotCArrayNotStringClass<T>::Value>::Type* = nullptr>
+	UE_NODISCARD FORCEINLINE friend FString operator+(const FString& Lhs, T&& Rhs)						{ return ConcatFR(Lhs, GetData(Rhs), GetNum(Rhs)); }
+	template <typename T, typename TEnableIf<TIsTCharRangeNotCArrayNotStringClass<T>::Value>::Type* = nullptr>
+	UE_NODISCARD FORCEINLINE friend FString operator+(FString&& Lhs, T&& Rhs)							{ return ConcatFR(MoveTemp(Lhs), GetData(Rhs), GetNum(Rhs)); }
 
 	/**
 	 * Concatenate this path with given path ensuring the / character is used between them
 	 * 
-	 * @param Str path array of TCHAR to be concatenated onto the end of this
+	 * @param Str path array of characters to be concatenated onto the end of this
 	 * @return reference to path
 	 */
-	FORCEINLINE FString& operator/=( const TCHAR* Str )
+	FORCEINLINE FString& operator/=( const ElementType* Str )
 	{
 		checkSlow(Str);
 
@@ -618,7 +620,7 @@ public:
 	/**
 	* Concatenate this path with given path ensuring the / character is used between them
 	* 
-	* @param Str path CharRangeType (FString/FStringView/TStringBuilder etc) to be concatenated onto the end of this
+	* @param Str path CharRangeType (string class/string view/string builder etc) to be concatenated onto the end of this
 	* @return reference to path
 	*/
 	template <typename CharRangeType, typename TEnableIf<TIsTCharRangeNotCArray <CharRangeType>::Value>::Type* = nullptr>
@@ -647,9 +649,9 @@ public:
 	 *
 	 * @param Lhs Path to concatenate onto.
 	 * @param Rhs Path to concatenate.
-	 * @return new FString of the path
+	 * @return The new concatenated path
 	 */
-	UE_NODISCARD FORCEINLINE friend FString operator/(const FString& Lhs, const TCHAR* Rhs)
+	UE_NODISCARD FORCEINLINE friend FString operator/(const FString& Lhs, const ElementType* Rhs)
 	{
 		checkSlow(Rhs);
 
@@ -665,9 +667,9 @@ public:
 	 *
 	 * @param Lhs Path to concatenate onto.
 	 * @param Rhs Path to concatenate.
-	 * @return new FString of the path
+	 * @return The new concatenated path
 	 */
-	UE_NODISCARD FORCEINLINE friend FString operator/(FString&& Lhs, const TCHAR* Rhs)
+	UE_NODISCARD FORCEINLINE friend FString operator/(FString&& Lhs, const ElementType* Rhs)
 	{
 		checkSlow(Rhs);
 
@@ -683,7 +685,7 @@ public:
 	 *
 	 * @param Lhs Path to concatenate onto.
 	 * @param Rhs Path to concatenate.
-	 * @return new FString of the path
+	 * @return The new concatenated path
 	 */
 	UE_NODISCARD FORCEINLINE friend FString operator/(const FString& Lhs, const FString& Rhs)
 	{
@@ -699,7 +701,7 @@ public:
 	 *
 	 * @param Lhs Path to concatenate onto.
 	 * @param Rhs Path to concatenate.
-	 * @return new FString of the path
+	 * @return The new concatenated path
 	 */
 	UE_NODISCARD FORCEINLINE friend FString operator/(FString&& Lhs, const FString& Rhs)
 	{
@@ -715,9 +717,9 @@ public:
 	 *
 	 * @param Lhs Path to concatenate onto.
 	 * @param Rhs Path to concatenate.
-	 * @return new FString of the path
+	 * @return new string of the path
 	 */
-	UE_NODISCARD FORCEINLINE friend FString operator/(const TCHAR* Lhs, const FString& Rhs)
+	UE_NODISCARD FORCEINLINE friend FString operator/(const ElementType* Lhs, const FString& Rhs)
 	{
 		int32 StrLength = Rhs.Len();
 
@@ -1058,7 +1060,7 @@ public:
 	UE_NODISCARD CORE_API FString Mid(int32 Start, int32 Count) &&;
 
 	/** Returns the substring from Start position to the end */
-	UE_NODISCARD CORE_API FORCEINLINE FString Mid(int32 Start) const & { return RightChop(Start); }
+	UE_NODISCARD FORCEINLINE FString Mid(int32 Start) const & { return RightChop(Start); }
 	UE_NODISCARD FORCEINLINE FString Mid(int32 Start) && { return ((FString&&)*this).RightChop(Start); }
 
 	/** Modifies the string such that it is now the substring from Start position for Count characters. */
@@ -1087,11 +1089,11 @@ public:
 	 *
 	 *        Consider using UE::String::FindLast() as an alternative.
 	 */
-	template <typename TCharRangeType, std::enable_if_t<TIsCharRangeNotCArrayNotFString<TCharRangeType>::Value>* = nullptr>
+	template <typename TCharRangeType, std::enable_if_t<TIsCharRangeNotCArrayNotStringClass<TCharRangeType>::Value>* = nullptr>
 	UE_NODISCARD int32 Find(TCharRangeType&& SubStr, ESearchCase::Type SearchCase = ESearchCase::IgnoreCase,
 		ESearchDir::Type SearchDir = ESearchDir::FromStart, int32 StartPosition = INDEX_NONE) const
 	{
-		static_assert(std::is_same_v<typename TElementType<TCharRangeType>::Type, TCHAR>, "Expected a range of TCHAR");
+		static_assert(std::is_same_v<typename TElementType<TCharRangeType>::Type, ElementType>, "Expected a range of ElementType");
 		return Find(GetData(SubStr), GetNum(SubStr), SearchCase, SearchDir, StartPosition);
 	}
 
@@ -1099,12 +1101,12 @@ public:
 	 * Searches the string for a substring, and returns index into this string of the first found instance. Can search
 	 * from beginning or end, and ignore case or not. If substring is empty, returns clamped StartPosition.
 	 *
-	 * @param SubStr			The string array of TCHAR to search for
+	 * @param SubStr			The string array of characters to search for
 	 * @param StartPosition		The start character position to search from.  See note below.
 	 * @param SearchCase		Indicates whether the search is case sensitive or not
 	 * @param SearchDir			Indicates whether the search starts at the beginning or at the end.
 	 */
-	UE_NODISCARD int32 Find(const TCHAR* SubStr, ESearchCase::Type SearchCase = ESearchCase::IgnoreCase,
+	UE_NODISCARD int32 Find(const ElementType* SubStr, ESearchCase::Type SearchCase = ESearchCase::IgnoreCase,
 		ESearchDir::Type SearchDir = ESearchDir::FromStart, int32 StartPosition = INDEX_NONE) const
 	{
 		return SubStr ? Find(SubStr, FCString::Strlen(SubStr), SearchCase, SearchDir, StartPosition) : INDEX_NONE;
@@ -1114,7 +1116,7 @@ public:
 	 * Searches the string for a substring, and returns index into this string of the first found instance. Can search
 	 * from beginning or end, and ignore case or not. If substring is empty, returns clamped StartPosition.
 	 *
-	 * @param SubStr			The string array of TCHAR to search for
+	 * @param SubStr			The string array of characters to search for
 	 * @param StartPosition		The start character position to search from.  See note below.
 	 * @param SearchCase		Indicates whether the search is case sensitive or not
 	 * @param SearchDir			Indicates whether the search starts at the beginning or at the end.
@@ -1136,7 +1138,7 @@ public:
 	 * Searches the string for a substring, and returns index into this string of the first found instance. Can search
 	 * from beginning or end, and ignore case or not. If substring is empty, returns clamped StartPosition.
 	 *
-	 * @param SubStr			The string array of TCHAR to search for
+	 * @param SubStr			The string array of characters to search for
 	 * @param SubStrLen			The length of the SubStr array
 	 * @param StartPosition		The start character position to search from.  See note below.
 	 * @param SearchCase		Indicates whether the search is case sensitive or not
@@ -1149,7 +1151,7 @@ public:
 	 *
 	 *        Consider using UE::String::FindLast() as an alternative.
 	 */
-	UE_NODISCARD CORE_API int32 Find(const TCHAR* SubStr, int32 InSubStrLen, ESearchCase::Type SearchCase = ESearchCase::IgnoreCase,
+	UE_NODISCARD CORE_API int32 Find(const ElementType* SubStr, int32 InSubStrLen, ESearchCase::Type SearchCase = ESearchCase::IgnoreCase,
 		ESearchDir::Type SearchDir = ESearchDir::FromStart, int32 StartPosition = INDEX_NONE) const;
 
 
@@ -1161,11 +1163,11 @@ public:
 	 * @param SearchDir			Indicates whether the search starts at the beginning or at the end ( defaults to ESearchDir::FromStart )
 	 * @return					Returns whether the string contains the substring. If the substring is empty, returns true.
 	 **/
-	template <typename TCharRangeType, std::enable_if_t<TIsCharRangeNotCArrayNotFString<TCharRangeType>::Value>* = nullptr>
+	template <typename TCharRangeType, std::enable_if_t<TIsCharRangeNotCArrayNotStringClass<TCharRangeType>::Value>* = nullptr>
 	UE_NODISCARD FORCEINLINE bool Contains(TCharRangeType&& SubStr, ESearchCase::Type SearchCase = ESearchCase::IgnoreCase,
 		ESearchDir::Type SearchDir = ESearchDir::FromStart) const
 	{
-		static_assert(std::is_same_v<typename TElementType<TCharRangeType>::Type, TCHAR>, "Expected a range of TCHAR");
+		static_assert(std::is_same_v<typename TElementType<TCharRangeType>::Type, ElementType>, "Expected a range of characters");
 		return Find(Forward<TCharRangeType>(SubStr), SearchCase, SearchDir) != INDEX_NONE;
 	}
 
@@ -1177,7 +1179,7 @@ public:
 	 * @param SearchDir			Indicates whether the search starts at the beginning or at the end ( defaults to ESearchDir::FromStart )
 	 * @return					Returns whether the string contains the substring. If the substring is empty, returns true.
 	 **/
-	UE_NODISCARD FORCEINLINE bool Contains(const TCHAR* SubStr, ESearchCase::Type SearchCase = ESearchCase::IgnoreCase,
+	UE_NODISCARD FORCEINLINE bool Contains(const ElementType* SubStr, ESearchCase::Type SearchCase = ESearchCase::IgnoreCase,
 		ESearchDir::Type SearchDir = ESearchDir::FromStart) const
 	{
 		return Find(SubStr, SearchCase, SearchDir) != INDEX_NONE;
@@ -1191,7 +1193,7 @@ public:
 	 * @param SearchDir			Indicates whether the search starts at the beginning or at the end ( defaults to ESearchDir::FromStart )
 	 * @return					Returns whether the string contains the substring. If the substring is empty, returns true.
 	 **/
-	UE_NODISCARD FORCEINLINE bool Contains(const FString& SubStr, ESearchCase::Type SearchCase = ESearchCase::IgnoreCase, 
+	UE_NODISCARD FORCEINLINE bool Contains(const FString& SubStr, ESearchCase::Type SearchCase = ESearchCase::IgnoreCase,
 							  ESearchDir::Type SearchDir = ESearchDir::FromStart ) const
 	{
 		return Find(*SubStr, SearchCase, SearchDir) != INDEX_NONE;
@@ -1206,7 +1208,7 @@ public:
 	 * @param SearchDir			Indicates whether the search starts at the beginning or at the end ( defaults to ESearchDir::FromStart )
 	 * @return					Returns whether the string contains the substring. If the substring is empty, returns true.
 	 **/
-	UE_NODISCARD FORCEINLINE bool Contains(const TCHAR* SubStr, int32 SubStrLen,
+	UE_NODISCARD FORCEINLINE bool Contains(const ElementType* SubStr, int32 SubStrLen,
 		ESearchCase::Type SearchCase = ESearchCase::IgnoreCase, ESearchDir::Type SearchDir = ESearchDir::FromStart) const
 	{
 		return Find(SubStr, SubStrLen, SearchCase, SearchDir) != INDEX_NONE;
@@ -1219,7 +1221,7 @@ public:
 	 * @param Index out the position the character was found at, INDEX_NONE if return is false
 	 * @return true if character was found in this string, otherwise false
 	 */
-	FORCEINLINE bool FindChar( TCHAR InChar, int32& Index ) const
+	FORCEINLINE bool FindChar(ElementType InChar, int32& Index ) const
 	{
 		return Data.Find(InChar, Index);
 	}
@@ -1231,7 +1233,7 @@ public:
 	 * @param Index out the position the character was found at, INDEX_NONE if return is false
 	 * @return true if character was found in this string, otherwise false
 	 */
-	FORCEINLINE bool FindLastChar( TCHAR InChar, int32& Index ) const
+	FORCEINLINE bool FindLastChar( ElementType InChar, int32& Index ) const
 	{
 		return Data.FindLast(InChar, Index);
 	}
@@ -1239,10 +1241,10 @@ public:
 	/**
 	 * Searches an initial substring for the last occurrence of a character which matches the specified predicate.
 	 *
-	 * @param Pred Predicate that takes TCHAR and returns true if TCHAR matches search criteria, false otherwise.
+	 * @param Pred Predicate that takes a character and returns true if it matches search criteria, false otherwise.
 	 * @param Count The number of characters from the front of the string through which to search.
 	 *
-	 * @return Index of found TCHAR, INDEX_NONE otherwise.
+	 * @return Index of the found character, INDEX_NONE otherwise.
 	 */
 	template <typename Predicate>
 	FORCEINLINE int32 FindLastCharByPredicate(Predicate Pred, int32 Count) const
@@ -1254,10 +1256,10 @@ public:
 	/**
 	 * Searches the string for the last occurrence of a character which matches the specified predicate.
 	 *
-	 * @param Pred Predicate that takes TCHAR and returns true if TCHAR matches search criteria, false otherwise.
-	 * @param StartIndex Index of element from which to start searching. Defaults to last TCHAR in string.
+	 * @param Pred Predicate that takes a character and returns true if it matches search criteria, false otherwise.
+	 * @param StartIndex Index of element from which to start searching. Defaults to the last character in string.
 	 *
-	 * @return Index of found TCHAR, INDEX_NONE otherwise.
+	 * @return Index of the found character, INDEX_NONE otherwise.
 	 */
 	template <typename Predicate>
 	FORCEINLINE int32 FindLastCharByPredicate(Predicate Pred) const
@@ -1369,39 +1371,39 @@ public:
 	CORE_API void RemoveSpacesInline();
 
 	/**
-	 * Constructs FString object similarly to how classic sprintf works.
+	 * Constructs a string similarly to how classic sprintf works.
 	 *
-	 * @param Format	Format string that specifies how FString should be built optionally using additional args. Refer to standard printf format.
+	 * @param Format	Format string that specifies how the string should be built optionally using additional args. Refer to standard printf format.
 	 * @param ...		Depending on format function may require additional arguments to build output object.
 	 *
-	 * @returns FString object that was constructed using format and additional parameters.
+	 * @returns A string that was constructed using format and additional parameters.
 	 */
 	template <typename FmtType, typename... Types>
 	UE_NODISCARD static FString Printf(const FmtType& Fmt, Types... Args)
 	{
-		static_assert(TIsArrayOrRefOfTypeByPredicate<FmtType, TIsCharEncodingCompatibleWithTCHAR>::Value, "Formatting string must be a TCHAR array.");
-		static_assert(TAnd<TIsValidVariadicFunctionArg<Types>...>::Value, "Invalid argument(s) passed to FString::Printf");
+		static_assert(TIsArrayOrRefOfTypeByPredicate<FmtType, TIsCharEncodingCompatibleWithTCHAR>::Value, "Formatting string must be a character array.");
+		static_assert(TAnd<TIsValidVariadicFunctionArg<Types>...>::Value, "Invalid argument(s) passed to Printf");
 
-		return PrintfImpl((const TCHAR*)Fmt, Args...);
+		return PrintfImpl((const ElementType*)Fmt, Args...);
 	}
 
 	/**
-	 * Just like Printf, but appends the formatted text to the existing FString instead.
+	 * Just like Printf, but appends the formatted text to the existing string instead.
 	 * @return a reference to the modified string, so that it can be chained
 	 */
 	template <typename FmtType, typename... Types>
 	FString& Appendf(const FmtType& Fmt, Types... Args)
 	{
-		static_assert(TIsArrayOrRefOfTypeByPredicate<FmtType, TIsCharEncodingCompatibleWithTCHAR>::Value, "Formatting string must be a TCHAR array.");
-		static_assert(TAnd<TIsValidVariadicFunctionArg<Types>...>::Value, "Invalid argument(s) passed to FString::Appendf");
+		static_assert(TIsArrayOrRefOfTypeByPredicate<FmtType, TIsCharEncodingCompatibleWithTCHAR>::Value, "Formatting string must be a character array.");
+		static_assert(TAnd<TIsValidVariadicFunctionArg<Types>...>::Value, "Invalid argument(s) passed to TString::Appendf");
 
-		AppendfImpl(*this, (const TCHAR*)Fmt, Args...);
+		AppendfImpl(*this, (const ElementType*)Fmt, Args...);
 		return *this;
 	}
 
 private:
-	static CORE_API FString VARARGS PrintfImpl(const TCHAR* Fmt, ...);
-	static CORE_API void VARARGS AppendfImpl(FString& AppendToMe, const TCHAR* Fmt, ...);
+	static CORE_API FString VARARGS PrintfImpl(const ElementType* Fmt, ...);
+	static CORE_API void VARARGS AppendfImpl(FString& AppendToMe, const ElementType* Fmt, ...);
 public:
 
 	/**
@@ -1410,7 +1412,7 @@ public:
 	 * @param InNamedArguments		A map of named arguments that match the tokens specified in InExpression
 	 * @return A string containing the formatted text
 	 */
-	UE_NODISCARD static CORE_API FString Format(const TCHAR* InFormatString, const FStringFormatNamedArguments& InNamedArguments);
+	UE_NODISCARD static CORE_API FString Format(const ElementType* InFormatString, const FStringFormatNamedArguments& InNamedArguments);
 
 	/**
 	 * Format the specified string using the specified arguments. Replaces instances of {0} with indices from the given array matching the index specified in the token
@@ -1418,10 +1420,10 @@ public:
 	 * @param InOrderedArguments	An array of ordered arguments that match the tokens specified in InExpression
 	 * @return A string containing the formatted text
 	 */
-	UE_NODISCARD static CORE_API FString Format(const TCHAR* InFormatString, const FStringFormatOrderedArguments& InOrderedArguments);
+	UE_NODISCARD static CORE_API FString Format(const ElementType* InFormatString, const FStringFormatOrderedArguments& InOrderedArguments);
 
 	/** Returns a string containing only the Ch character */
-	UE_NODISCARD static CORE_API FString Chr( TCHAR Ch );
+	UE_NODISCARD static CORE_API FString Chr( ElementType Ch );
 
 	/**
 	 * Returns a string that is full of a variable number of characters
@@ -1431,7 +1433,7 @@ public:
 	 * 
 	 * @return The string of NumCharacters characters.
 	 */
-	UE_NODISCARD static CORE_API FString ChrN( int32 NumCharacters, TCHAR Char );
+	UE_NODISCARD static CORE_API FString ChrN( int32 NumCharacters, ElementType Char );
 
 	/**
 	 * Serializes the string.
@@ -1450,10 +1452,10 @@ public:
 	 * @param SearchCase		Indicates whether the search is case sensitive or not ( defaults to ESearchCase::IgnoreCase )
 	 * @return true if this string begins with specified text, false otherwise
 	 */
-	template <typename TCharRangeType, std::enable_if_t<TIsCharRangeNotCArrayNotFString<TCharRangeType>::Value>* = nullptr>
+	template <typename TCharRangeType, std::enable_if_t<TIsCharRangeNotCArrayNotStringClass<TCharRangeType>::Value>* = nullptr>
 	UE_NODISCARD bool StartsWith(TCharRangeType&& InPrefix, ESearchCase::Type SearchCase = ESearchCase::IgnoreCase) const
 	{
-		static_assert(std::is_same_v<typename TElementType<TCharRangeType>::Type, TCHAR>, "Expected a range of TCHAR");
+		static_assert(std::is_same_v<typename TElementType<TCharRangeType>::Type, ElementType>, "Expected a range of characters");
 		return StartsWith(GetData(InPrefix), GetNum(InPrefix), SearchCase);
 	}
 
@@ -1463,7 +1465,7 @@ public:
 	 * @param SearchCase		Indicates whether the search is case sensitive or not ( defaults to ESearchCase::IgnoreCase )
 	 * @return true if this string begins with specified text, false otherwise
 	 */
-	UE_NODISCARD bool StartsWith(const TCHAR* InPrefix, ESearchCase::Type SearchCase = ESearchCase::IgnoreCase) const
+	UE_NODISCARD bool StartsWith(const ElementType* InPrefix, ESearchCase::Type SearchCase = ESearchCase::IgnoreCase) const
 	{
 		return StartsWith(InPrefix, InPrefix ? FCString::Strlen(InPrefix) : 0, SearchCase);
 	}
@@ -1485,7 +1487,7 @@ public:
 	 * @param SearchCase		Indicates whether the search is case sensitive or not ( defaults to ESearchCase::IgnoreCase )
 	 * @return true if this string begins with specified text, false otherwise
 	 */
-	UE_NODISCARD CORE_API bool StartsWith(const TCHAR* InPrefix, int32 InPrefixLen, ESearchCase::Type SearchCase = ESearchCase::IgnoreCase) const;
+	UE_NODISCARD CORE_API bool StartsWith(const ElementType* InPrefix, int32 InPrefixLen, ESearchCase::Type SearchCase = ESearchCase::IgnoreCase) const;
 
 	/**
 	 * Test whether this string ends with given suffix.
@@ -1493,10 +1495,10 @@ public:
 	 * @param SearchCase		Indicates whether the search is case sensitive or not ( defaults to ESearchCase::IgnoreCase )
 	 * @return true if this string ends with specified text, false otherwise
 	 */
-	template <typename TCharRangeType, std::enable_if_t<TIsCharRangeNotCArrayNotFString<TCharRangeType>::Value>* = nullptr>
+	template <typename TCharRangeType, std::enable_if_t<TIsCharRangeNotCArrayNotStringClass<TCharRangeType>::Value>* = nullptr>
 	UE_NODISCARD bool EndsWith(TCharRangeType&& InSuffix, ESearchCase::Type SearchCase = ESearchCase::IgnoreCase) const
 	{
-		static_assert(std::is_same_v<typename TElementType<TCharRangeType>::Type, TCHAR>, "Expected a range of TCHAR");
+		static_assert(std::is_same_v<typename TElementType<TCharRangeType>::Type, ElementType>, "Expected a range of characters");
 		return EndsWith(GetData(InSuffix), GetNum(InSuffix), SearchCase);
 	}
 
@@ -1506,7 +1508,7 @@ public:
 	 * @param SearchCase		Indicates whether the search is case sensitive or not ( defaults to ESearchCase::IgnoreCase )
 	 * @return true if this string ends with specified text, false otherwise
 	 */
-	UE_NODISCARD bool EndsWith(const TCHAR* InSuffix, ESearchCase::Type SearchCase = ESearchCase::IgnoreCase) const
+	UE_NODISCARD bool EndsWith(const ElementType* InSuffix, ESearchCase::Type SearchCase = ESearchCase::IgnoreCase) const
 	{
 		return EndsWith(InSuffix, InSuffix ? FCString::Strlen(InSuffix) : 0, SearchCase);
 	}
@@ -1528,7 +1530,7 @@ public:
 	 * @param SearchCase		Indicates whether the search is case sensitive or not ( defaults to ESearchCase::IgnoreCase )
 	 * @return true if this string ends with specified text, false otherwise
 	 */
-	UE_NODISCARD CORE_API bool EndsWith(const TCHAR* InSuffix, int32 InSuffixLen, ESearchCase::Type SearchCase = ESearchCase::IgnoreCase ) const;
+	UE_NODISCARD CORE_API bool EndsWith(const ElementType* InSuffix, int32 InSuffixLen, ESearchCase::Type SearchCase = ESearchCase::IgnoreCase ) const;
 
 	/**
 	 * Searches this string for a given wild card
@@ -1538,10 +1540,10 @@ public:
 	 * @return true if this string matches the *?-type wildcard given. 
 	 * @warning This is a simple, SLOW routine. Use with caution
 	 */
-	template <typename TCharRangeType, std::enable_if_t<TIsCharRangeNotCArrayNotFString<TCharRangeType>::Value>* = nullptr>
+	template <typename TCharRangeType, std::enable_if_t<TIsCharRangeNotCArrayNotStringClass<TCharRangeType>::Value>* = nullptr>
 	UE_NODISCARD bool MatchesWildcard(TCharRangeType&& Wildcard, ESearchCase::Type SearchCase = ESearchCase::IgnoreCase) const
 	{
-		static_assert(std::is_same_v<typename TElementType<TCharRangeType>::Type, TCHAR>, "Expected a range of TCHAR");
+		static_assert(std::is_same_v<typename TElementType<TCharRangeType>::Type, ElementType>, "Expected a range of characters");
 		return MatchesWildcard(GetData(Wildcard), GetNum(Wildcard), SearchCase);
 	}
 
@@ -1553,7 +1555,7 @@ public:
 	 * @return true if this string matches the *?-type wildcard given.
 	 * @warning This is a simple, SLOW routine. Use with caution
 	 */
-	UE_NODISCARD bool MatchesWildcard(const TCHAR* Wildcard, ESearchCase::Type SearchCase = ESearchCase::IgnoreCase) const
+	UE_NODISCARD bool MatchesWildcard(const ElementType* Wildcard, ESearchCase::Type SearchCase = ESearchCase::IgnoreCase) const
 	{
 		return MatchesWildcard(Wildcard, Wildcard ? FCString::Strlen(Wildcard) : 0, SearchCase);
 	}
@@ -1579,7 +1581,7 @@ public:
 	 * @return true if this string matches the *?-type wildcard given.
 	 * @warning This is a simple, SLOW routine. Use with caution
 	 */
-	UE_NODISCARD CORE_API bool MatchesWildcard(const TCHAR* Wildcard, int32 WildcardLen, ESearchCase::Type SearchCase = ESearchCase::IgnoreCase) const;
+	UE_NODISCARD CORE_API bool MatchesWildcard(const ElementType* Wildcard, int32 WildcardLen, ESearchCase::Type SearchCase = ESearchCase::IgnoreCase) const;
 
 	/**
 	 * Removes whitespace characters from the start and end of this string. Modifies the string in-place.
@@ -1647,7 +1649,7 @@ public:
 	* Trims a single character from the start and end of the string (removes at max one instance in the beginning and end of the string).
 	* @see TrimChar for a variant that returns a modified copy of the string
 	*/
-	CORE_API void TrimCharInline(const TCHAR CharacterToTrim, bool* bCharRemoved);
+	CORE_API void TrimCharInline(ElementType CharacterToTrim, bool* bCharRemoved);
 	
 	/**
 	 * Returns a copy of this string with wrapping quotation marks removed.
@@ -1663,12 +1665,12 @@ public:
 	* Returns a copy of this string with wrapping CharacterToTrim removed (removes at max one instance in the beginning and end of the string).
 	* @see TrimCharInline for an inline variant
 	*/
-	UE_NODISCARD CORE_API FString TrimChar(const TCHAR CharacterToTrim, bool* bCharRemoved = nullptr ) const &;
+	UE_NODISCARD CORE_API FString TrimChar(ElementType CharacterToTrim, bool* bCharRemoved = nullptr ) const &;
 
 	/**
 	* Returns a copy of this string with wrapping CharacterToTrim removed (removes at max one instance in the beginning and end of the string).
 	*/
-	UE_NODISCARD CORE_API FString TrimChar(const TCHAR CharacterToTrim, bool* bCharRemoved = nullptr) &&;
+	UE_NODISCARD CORE_API FString TrimChar(ElementType CharacterToTrim, bool* bCharRemoved = nullptr) &&;
 	
 	/**
 	 * Breaks up a delimited string into elements of a string array.
@@ -1679,7 +1681,7 @@ public:
 	 *
 	 * @return	The number of elements in InArray
 	 */
-	CORE_API int32 ParseIntoArray( TArray<FString>& OutArray, const TCHAR* pchDelim, bool InCullEmpty = true ) const;
+	CORE_API int32 ParseIntoArray( TArray<FString>& OutArray, const ElementType* pchDelim, bool InCullEmpty = true ) const;
 
 	/**
 	 * Breaks up a delimited string into elements of a string array, using any whitespace and an 
@@ -1691,7 +1693,7 @@ public:
 	 *
 	 * @return	The number of elements in InArray
 	 */
-	CORE_API int32 ParseIntoArrayWS( TArray<FString>& OutArray, const TCHAR* pchExtraDelim = nullptr, bool InCullEmpty = true ) const;
+	CORE_API int32 ParseIntoArrayWS( TArray<FString>& OutArray, const ElementType* pchExtraDelim = nullptr, bool InCullEmpty = true ) const;
 
 	/**
 	* Breaks up a delimited string into elements of a string array, using line ending characters
@@ -1713,7 +1715,7 @@ public:
 	*
 	* @return	The number of elements in InArray
 	*/
-	CORE_API int32 ParseIntoArray(TArray<FString>& OutArray, const TCHAR*const* DelimArray, int32 NumDelims, bool InCullEmpty = true) const;
+	CORE_API int32 ParseIntoArray(TArray<FString>& OutArray, const ElementType*const* DelimArray, int32 NumDelims, bool InCullEmpty = true) const;
 
 	/**
 	 * Takes an array of strings and removes any zero length entries.
@@ -1747,7 +1749,7 @@ public:
 	 * @param SearchCase	Indicates whether the search is case sensitive or not ( defaults to ESearchCase::IgnoreCase )
 	 * @return a copy of this string with the replacement made
 	 */
-	UE_NODISCARD CORE_API FString Replace(const TCHAR* From, const TCHAR* To, ESearchCase::Type SearchCase = ESearchCase::IgnoreCase) const &;
+	UE_NODISCARD CORE_API FString Replace(const ElementType* From, const ElementType* To, ESearchCase::Type SearchCase = ESearchCase::IgnoreCase) const &;
 
 	/**
 	 * Replace all occurrences of a substring in this string
@@ -1757,7 +1759,7 @@ public:
 	 * @param SearchCase	Indicates whether the search is case sensitive or not ( defaults to ESearchCase::IgnoreCase )
 	 * @return a copy of this string with the replacement made
 	 */
-	UE_NODISCARD CORE_API FString Replace(const TCHAR* From, const TCHAR* To, ESearchCase::Type SearchCase = ESearchCase::IgnoreCase) &&;
+	UE_NODISCARD CORE_API FString Replace(const ElementType* From, const ElementType* To, ESearchCase::Type SearchCase = ESearchCase::IgnoreCase) &&;
 
 	/**
 	 * Replace all occurrences of SearchText with ReplacementText in this string.
@@ -1768,19 +1770,19 @@ public:
 	 *
 	 * @return	the number of occurrences of SearchText that were replaced.
 	 */
-	CORE_API int32 ReplaceInline( const TCHAR* SearchText, const TCHAR* ReplacementText, ESearchCase::Type SearchCase = ESearchCase::IgnoreCase );
+	CORE_API int32 ReplaceInline( const ElementType* SearchText, const ElementType* ReplacementText, ESearchCase::Type SearchCase = ESearchCase::IgnoreCase );
 
 	/**
 	 * Replace all occurrences of a character with another.
 	 *
-	 * @param SearchChar      Character to remove from this FString
+	 * @param SearchChar      Character to remove from this string
 	 * @param ReplacementChar Replacement character
 	 * @param SearchCase      Indicates whether the search is case sensitive or not ( defaults to ESearchCase::IgnoreCase )
 	 * @note no dynamic allocation
 	 */
-	void ReplaceCharInline(const TCHAR SearchChar, const TCHAR ReplacementChar, ESearchCase::Type SearchCase = ESearchCase::IgnoreCase)
+	void ReplaceCharInline(ElementType SearchChar, ElementType ReplacementChar, ESearchCase::Type SearchCase = ESearchCase::IgnoreCase)
 	{
-		if (SearchCase == ESearchCase::IgnoreCase && TChar<TCHAR>::IsAlpha(SearchChar))
+		if (SearchCase == ESearchCase::IgnoreCase && TChar<ElementType>::IsAlpha(SearchChar))
 		{
 			ReplaceCharInlineIgnoreCase(SearchChar, ReplacementChar);
 		}
@@ -1792,8 +1794,8 @@ public:
 
 private:
 	
-	CORE_API void ReplaceCharInlineCaseSensitive(const TCHAR SearchChar, const TCHAR ReplacementChar);
-	CORE_API void ReplaceCharInlineIgnoreCase(const TCHAR SearchChar, const TCHAR ReplacementChar);
+	CORE_API void ReplaceCharInlineCaseSensitive(ElementType SearchChar, ElementType ReplacementChar);
+	CORE_API void ReplaceCharInlineIgnoreCase(ElementType SearchChar, ElementType ReplacementChar);
 
 public:
 
@@ -1817,7 +1819,7 @@ public:
 	 *
 	 * @param	Chars	by default, replaces all supported characters; this parameter allows you to limit the replacement to a subset.
 	 */
-	CORE_API void ReplaceCharWithEscapedCharInline( const TArray<TCHAR>* Chars = nullptr );
+	CORE_API void ReplaceCharWithEscapedCharInline( const TArray<ElementType>* Chars = nullptr );
 
 	/**
 	 * Replaces certain characters with the "escaped" version of that character (i.e. replaces "\n" with "\\n").
@@ -1827,7 +1829,7 @@ public:
 	 *
 	 * @return	a string with all control characters replaced by the escaped version.
 	 */
-	UE_NODISCARD FString ReplaceCharWithEscapedChar( const TArray<TCHAR>* Chars = nullptr ) const &
+	UE_NODISCARD FString ReplaceCharWithEscapedChar( const TArray<ElementType>* Chars = nullptr ) const &
 	{
 		FString Result(*this);
 		Result.ReplaceCharWithEscapedCharInline(Chars);
@@ -1842,7 +1844,7 @@ public:
 	 *
 	 * @return	a string with all control characters replaced by the escaped version.
 	 */
-	UE_NODISCARD FString ReplaceCharWithEscapedChar( const TArray<TCHAR>* Chars = nullptr ) &&
+	UE_NODISCARD FString ReplaceCharWithEscapedChar( const TArray<ElementType>* Chars = nullptr ) &&
 	{
 		ReplaceCharWithEscapedCharInline(Chars);
 		return MoveTemp(*this);
@@ -1852,14 +1854,14 @@ public:
 	 * Removes the escape backslash for all supported characters, replacing the escape and character with the non-escaped version.  (i.e.
 	 * replaces "\\n" with "\n".  Counterpart to ReplaceCharWithEscapedCharInline().
 	 */
-	CORE_API void ReplaceEscapedCharWithCharInline( const TArray<TCHAR>* Chars = nullptr );
+	CORE_API void ReplaceEscapedCharWithCharInline( const TArray<ElementType>* Chars = nullptr );
 
 	/**
 	 * Removes the escape backslash for all supported characters, replacing the escape and character with the non-escaped version.  (i.e.
 	 * replaces "\\n" with "\n".  Counterpart to ReplaceCharWithEscapedChar().
 	 * @return copy of this string with replacement made
 	 */
-	UE_NODISCARD FString ReplaceEscapedCharWithChar( const TArray<TCHAR>* Chars = nullptr ) const &
+	UE_NODISCARD FString ReplaceEscapedCharWithChar( const TArray<ElementType>* Chars = nullptr ) const &
 	{
 		FString Result(*this);
 		Result.ReplaceEscapedCharWithCharInline(Chars);
@@ -1871,7 +1873,7 @@ public:
 	 * replaces "\\n" with "\n".  Counterpart to ReplaceCharWithEscapedChar().
 	 * @return copy of this string with replacement made
 	 */
-	UE_NODISCARD FString ReplaceEscapedCharWithChar( const TArray<TCHAR>* Chars = nullptr ) &&
+	UE_NODISCARD FString ReplaceEscapedCharWithChar( const TArray<ElementType>* Chars = nullptr ) &&
 	{
 		ReplaceEscapedCharWithCharInline(Chars);
 		return MoveTemp(*this);
@@ -1925,7 +1927,7 @@ public:
 	/** Converts an integer to a string. */
 	UE_NODISCARD static FORCEINLINE FString FromInt( int32 Num )
 	{
-		FString Ret; 
+		FString Ret;
 		Ret.AppendInt(Num); 
 		return Ret;
 	}
@@ -2002,7 +2004,7 @@ public:
 	 * @return	The final, joined, separated string.
 	 */
 	template <typename RangeType>
-	UE_NODISCARD static FString Join(const RangeType& Range, const TCHAR* Separator)
+	UE_NODISCARD static FString Join(const RangeType& Range, const ElementType* Separator)
 	{
 		FString Result;
 		bool    First = true;
@@ -2033,7 +2035,7 @@ public:
 	 * @return	The final, joined, separated string.
 	 */
 	template <typename RangeType, typename ProjectionType>
-	UE_NODISCARD static FString JoinBy(const RangeType& Range, const TCHAR* Separator, ProjectionType Proj)
+	UE_NODISCARD static FString JoinBy(const RangeType& Range, const ElementType* Separator, ProjectionType Proj)
 	{
 		FString Result;
 		bool    First = true;
@@ -2070,12 +2072,12 @@ public:
 template<> struct TIsZeroConstructType<FString> { enum { Value = true }; };
 Expose_TNameOf(FString)
 
-inline TCHAR* GetData(FString& String)
+inline FString::ElementType* GetData(FString& String)
 {
 	return String.GetCharArray().GetData();
 }
 
-inline const TCHAR* GetData(const FString& String)
+inline const FString::ElementType* GetData(const FString& String)
 {
 	return String.GetCharArray().GetData();
 }
@@ -2086,7 +2088,7 @@ inline int32 GetNum(const FString& String)
 }
 
 /** 
- * Convert an array of bytes to a TCHAR
+ * Convert an array of bytes to a string
  * @param In byte array values to convert
  * @param Count number of bytes to convert
  * @return Valid string representing bytes.
@@ -2102,7 +2104,7 @@ UE_NODISCARD inline FString BytesToString(const uint8* In, int32 Count)
 		int16 Value = *In;
 		Value += 1;
 
-		Result += TCHAR(Value);
+		Result += FString::ElementType(Value);
 
 		++In;
 		Count--;
@@ -2120,7 +2122,7 @@ UE_NODISCARD inline FString BytesToString(const uint8* In, int32 Count)
 inline int32 StringToBytes( const FString& String, uint8* OutBytes, int32 MaxBufferSize )
 {
 	int32 NumBytes = 0;
-	const TCHAR* CharPos = *String;
+	const FString::ElementType* CharPos = *String;
 
 	while( *CharPos && NumBytes < MaxBufferSize)
 	{
@@ -2593,3 +2595,6 @@ struct FTextRange
 };
 
 #include "Misc/StringFormatArg.h"
+
+#undef UE_STRING_TEXT
+#undef UE_STRING_COMPILING_UTF8
\ No newline at end of file
