- Fix NaN's poisoning the VVM constant table
- Comparing A==B will return true for NaN depending on the FP mode, so we now compare if both values are finite before comparing for equality.
- This does not fix the source of the error, but will stop all subsequent float values being set to NaN

diff --git a/Engine/Source/Developer/ShaderFormatVectorVM/Private/ir_vm_gen_bytecode_visitor.cpp b/Engine/Source/Developer/ShaderFormatVectorVM/Private/ir_vm_gen_bytecode_visitor.cpp
--- a/Engine/Source/Developer/ShaderFormatVectorVM/Private/ir_vm_gen_bytecode_visitor.cpp
+++ b/Engine/Source/Developer/ShaderFormatVectorVM/Private/ir_vm_gen_bytecode_visitor.cpp
@@ -1688,14 +1688,11 @@ class ir_gen_vvm_visitor : public ir_hierarchical_visitor
 				{			
 					float val = constant->value.f[0];
 					float otherval = other->value.f[0];
-					if (!FMath::IsNaN(val) && FMath::IsFinite(val))
+					if ( FMath::IsFinite(val) && FMath::IsFinite(otherval) )
 					{
 						return val == otherval;
 					}
-					else
-					{
-						return FMath::IsNaN(val) == FMath::IsNaN(otherval) && FMath::IsFinite(val) == FMath::IsFinite(otherval);
-					}
+					return FMath::IsFinite(val) == FMath::IsFinite(otherval) && FMath::IsNaN(val) == FMath::IsNaN(otherval);
 				}
 				case GLSL_TYPE_INT: return constant->value.i[0] == other->value.i[0];
 				case GLSL_TYPE_BOOL: return constant->value.b[0] == other->value.b[0];
diff --git a/Engine/Source/Runtime/Core/Private/UObject/FortniteMainBranchObjectVersion.cpp b/Engine/Source/Runtime/Core/Private/UObject/FortniteMainBranchObjectVersion.cpp
--- a/Engine/Source/Runtime/Core/Private/UObject/FortniteMainBranchObjectVersion.cpp
+++ b/Engine/Source/Runtime/Core/Private/UObject/FortniteMainBranchObjectVersion.cpp
@@ -11,7 +11,7 @@ TMap<FGuid, FGuid> FFortniteMainBranchObjectVersion::GetSystemGuids()
 	SystemGuids.Add(DevGuids.MATERIALSHADERMAP_DERIVEDDATA_VER, FGuid("CE7776E63E9E43FBB8C103EE1AB40B5D"));
 	SystemGuids.Add(DevGuids.NANITE_DERIVEDDATA_VER, FGuid("D28C003F825F4C178F698938FB0E47F7"));
 	SystemGuids.Add(DevGuids.NIAGARASHADERMAP_DERIVEDDATA_VER, FGuid("D0DF49D408F34E38A98FA097E0280EC7"));
-	SystemGuids.Add(DevGuids.Niagara_LatestScriptCompileVersion, FGuid("860B84D1960440A0AB802924893D3DA1"));
+	SystemGuids.Add(DevGuids.Niagara_LatestScriptCompileVersion, FGuid("3C56691355BF1941B5D7A4DB0AFE5F95"));
 	SystemGuids.Add(DevGuids.SkeletalMeshDerivedDataVersion, FGuid("518D7C51D31D9BE4A9F4B81542AC8119"));
 	SystemGuids.Add(DevGuids.STATICMESH_DERIVEDDATA_VER, FGuid("543843B103794E0E9BA4BE80FB602F79"));
 	return SystemGuids;
