Add a cascading abort to AutoRTFM to abort an entire nest of transactions.

diff --git a/Engine/Source/Programs/AutoRTFMTests/Private/Abort.cpp b/Engine/Source/Programs/AutoRTFMTests/Private/Abort.cpp
--- a/Engine/Source/Programs/AutoRTFMTests/Private/Abort.cpp
+++ b/Engine/Source/Programs/AutoRTFMTests/Private/Abort.cpp
@@ -126,3 +126,77 @@ TEST_CASE("Abort.TransactionInOpenAbort")
 	REQUIRE(AutoRTFM::ETransactionResult::AbortedByTransactInOpenAbort == InnerResult);
 	REQUIRE(AutoRTFM::ETransactionResult::AbortedByRequest == Result);
 }
+
+TEST_CASE("Abort.Cascade")
+{
+	bool bTouched = false;
+
+	const AutoRTFM::ETransactionResult Result = AutoRTFM::Transact([&]
+	{
+		bTouched = true;
+		AutoRTFM::Transact([&]
+		{
+			AutoRTFM::CascadingAbortTransaction();
+		});
+	});
+
+	REQUIRE(AutoRTFM::ETransactionResult::AbortedByCascade == Result);
+	REQUIRE(false == bTouched);
+}
+
+TEST_CASE("Abort.CascadeThroughOpen")
+{
+	bool bTouched = false;
+
+	const AutoRTFM::ETransactionResult Result = AutoRTFM::Transact([&]
+	{
+		bTouched = true;
+
+		AutoRTFM::Open([&]
+		{
+			const AutoRTFM::EContextStatus Status = AutoRTFM::Close([&]
+			{
+				AutoRTFM::Transact([&]
+				{
+					AutoRTFM::CascadingAbortTransaction();
+				});
+			});
+
+			REQUIRE(AutoRTFM::EContextStatus::AbortedByCascade == Status);
+		});
+	});
+
+	REQUIRE(AutoRTFM::ETransactionResult::AbortedByCascade == Result);
+	REQUIRE(false == bTouched);
+}
+
+TEST_CASE("Abort.CascadeThroughManualTransaction")
+{
+	bool bTouched = false;
+
+	const AutoRTFM::ETransactionResult Result = AutoRTFM::Transact([&]
+	{
+		bTouched = true;
+
+		AutoRTFM::Open([&]
+		{
+			REQUIRE(true == AutoRTFM::StartTransaction());
+
+			const AutoRTFM::EContextStatus Status = AutoRTFM::Close([&]
+			{
+				AutoRTFM::CascadingAbortTransaction();
+			});
+
+			REQUIRE(AutoRTFM::EContextStatus::AbortedByCascade == Status);
+
+			// We need to clear the status ourselves.
+			AutoRTFM::ClearTransactionStatus();
+
+			// Before manually starting the cascade again.
+			AutoRTFM::CascadingAbortTransaction();
+		});
+	});
+
+	REQUIRE(AutoRTFM::ETransactionResult::AbortedByCascade == Result);
+	REQUIRE(false == bTouched);
+}
diff --git a/Engine/Source/Runtime/Core/Private/AutoRTFM/API.cpp b/Engine/Source/Runtime/Core/Private/AutoRTFM/API.cpp
--- a/Engine/Source/Runtime/Core/Private/AutoRTFM/API.cpp
+++ b/Engine/Source/Runtime/Core/Private/AutoRTFM/API.cpp
@@ -106,7 +106,15 @@ extern "C" UE_AUTORTFM_AUTORTFM("RTFM_autortfm_commit_transaction") autortfm_res
 extern "C" UE_AUTORTFM_AUTORTFM("RTFM_autortfm_abort_transaction") autortfm_result autortfm_abort_transaction()
 {
 	UE_CLOG(!FContext::IsTransactional(), LogAutoRTFM, Fatal, TEXT("The function `autortfm_abort_transaction` was called from outside a transact."));
-	return static_cast<autortfm_result>(FContext::Get()->AbortTransaction(false));
+	FContext* const Context = FContext::Get();
+	return static_cast<autortfm_result>(Context->AbortTransaction(false, false));
+}
+
+extern "C" UE_AUTORTFM_AUTORTFM("RTFM_autortfm_cascading_abort_transaction") autortfm_result autortfm_cascading_abort_transaction()
+{
+	UE_CLOG(!FContext::IsTransactional(), LogAutoRTFM, Fatal, TEXT("The function `autortfm_cascading_abort_transaction` was called from outside a transact."));
+	FContext* const Context = FContext::Get();
+	return static_cast<autortfm_result>(Context->AbortTransaction(false, true));
 }
 
 extern "C" UE_AUTORTFM_AUTORTFM("RTFM_autortfm_clear_transaction_status") void autortfm_clear_transaction_status()
@@ -333,14 +341,20 @@ extern "C" UE_AUTORTFM_NOAUTORTFM void RTFM_autortfm_commit_transaction()
 
 extern "C" UE_AUTORTFM_NOAUTORTFM autortfm_result RTFM_autortfm_abort_transaction()
 {
-	UE_CLOG(!FContext::IsTransactional(), LogAutoRTFM, Fatal, TEXT("The function `autortfm_abort_transaction` was called from outside a transaction"));
-	FContext* Context = FContext::Get();
-	return static_cast<autortfm_result>(Context->AbortTransaction(true));
+	FContext* const Context = FContext::Get();
+	return static_cast<autortfm_result>(Context->AbortTransaction(true, false));
+}
+
+extern "C" UE_AUTORTFM_NOAUTORTFM autortfm_result RTFM_autortfm_cascading_abort_transaction()
+{
+	FContext* const Context = FContext::Get();
+	return static_cast<autortfm_result>(Context->AbortTransaction(true, true));
 }
 
 extern "C" UE_AUTORTFM_NOAUTORTFM void RTFM_autortfm_clear_transaction_status()
 {
 	UE_LOG(LogAutoRTFM, Fatal, TEXT("The function `autortfm_clear_transaction_status` was called from closed code."));
+	AutoRTFM::Unreachable();
 }
 
 extern "C" UE_AUTORTFM_NOAUTORTFM void RTFM_autortfm_abort_if_transactional()
diff --git a/Engine/Source/Runtime/Core/Private/AutoRTFM/Context.cpp b/Engine/Source/Runtime/Core/Private/AutoRTFM/Context.cpp
--- a/Engine/Source/Runtime/Core/Private/AutoRTFM/Context.cpp
+++ b/Engine/Source/Runtime/Core/Private/AutoRTFM/Context.cpp
@@ -118,13 +118,13 @@ ETransactionResult FContext::CommitTransaction()
 	return Result;
 }
 
-ETransactionResult FContext::AbortTransaction(bool bIsClosed)
+ETransactionResult FContext::AbortTransaction(bool bIsClosed, bool bIsCascading)
 {
 	ETransactionResult Result = ETransactionResult::AbortedByRequest;
 	ASSERT(Status == EContextStatus::OnTrack);
-	Status = EContextStatus::AbortedByRequest;
+	Status = bIsCascading ? EContextStatus::AbortedByCascade : EContextStatus::AbortedByRequest;
 
-	ASSERT(CurrentTransaction != nullptr);
+	ASSERT(nullptr != CurrentTransaction);
 
 	// Sort out how aborts work
 	CurrentTransaction->AbortWithoutThrowing();
@@ -224,8 +224,11 @@ void FContext::ClearTransactionStatus()
 	case EContextStatus::AbortedByRequest:
 		Status = EContextStatus::OnTrack;
 		break;
+	case EContextStatus::AbortedByCascade:
+		Status = EContextStatus::OnTrack;
+		break;
 	default:
-		ASSERT(!"Should not be reached");
+		AutoRTFM::Unreachable();
 	}
 }
 
@@ -250,9 +253,10 @@ ETransactionResult FContext::ResolveNestedTransaction(FTransaction* NewTransacti
 		return ETransactionResult::AbortedByRequest;
 	case EContextStatus::AbortedByLanguage:
 		return ETransactionResult::AbortedByLanguage;
+	case EContextStatus::AbortedByCascade:
+		return ETransactionResult::AbortedByCascade;
 	default:
-		ASSERT(!"Should not be reached");
-		return ETransactionResult::AbortedByLanguage;
+		AutoRTFM::Unreachable();
 	}
 }
 
@@ -339,6 +343,12 @@ ETransactionResult FContext::Transact(void (*Function)(void* Arg), void* Arg)
                 break;
             }
 
+            if (Status == EContextStatus::AbortedByCascade)
+            {
+                Result = ETransactionResult::AbortedByCascade;
+                break;
+            }
+
             ASSERT(Status == EContextStatus::AbortedByFailedLockAcquisition);
         }
 
@@ -367,10 +377,17 @@ ETransactionResult FContext::Transact(void (*Function)(void* Arg), void* Arg)
 		
 		PopCallNest();
 		PopTransaction();
-		ClearTransactionStatus();
 
 		ASSERT(CurrentNest != nullptr);
 		ASSERT(CurrentTransaction != nullptr);
+
+		// A cascading abort should cause all transactions to abort!
+		if (ETransactionResult::AbortedByCascade == Result)
+		{
+			CurrentTransaction->AbortAndThrow();
+		}
+
+		ClearTransactionStatus();
 	}
 
 	return Result;
diff --git a/Engine/Source/Runtime/Core/Private/AutoRTFM/Context.h b/Engine/Source/Runtime/Core/Private/AutoRTFM/Context.h
--- a/Engine/Source/Runtime/Core/Private/AutoRTFM/Context.h
+++ b/Engine/Source/Runtime/Core/Private/AutoRTFM/Context.h
@@ -31,7 +31,7 @@ public:
 	bool StartTransaction();
 
 	ETransactionResult CommitTransaction();
-	ETransactionResult AbortTransaction(bool bIsClosed);
+	ETransactionResult AbortTransaction(bool bIsClosed, bool bIsCascading);
 	void ClearTransactionStatus();
 	bool IsAborting() const;
 
diff --git a/Engine/Source/Runtime/Core/Private/AutoRTFM/Transaction.cpp b/Engine/Source/Runtime/Core/Private/AutoRTFM/Transaction.cpp
--- a/Engine/Source/Runtime/Core/Private/AutoRTFM/Transaction.cpp
+++ b/Engine/Source/Runtime/Core/Private/AutoRTFM/Transaction.cpp
@@ -28,9 +28,7 @@ void FTransaction::AbortWithoutThrowing()
 {
 	UE_LOG(LogAutoRTFM, Verbose, TEXT("Aborting '%hs'!"), GetContextStatusName(Context->GetStatus()));
 
-    ASSERT(Context->GetStatus() == EContextStatus::AbortedByFailedLockAcquisition
-           || Context->GetStatus() == EContextStatus::AbortedByLanguage
-           || Context->GetStatus() == EContextStatus::AbortedByRequest);
+    ASSERT(Context->IsAborting());
     ASSERT(Context->GetCurrentTransaction() == this);
 
     Stats.Collect<EStatsKind::Abort>();
@@ -148,17 +146,7 @@ void FTransaction::AbortOuterNest()
 
     AbortTasks.ForEachBackward([] (const TFunction<void()>& Task) -> bool { Task(); return true; });
 
-    switch (Context->GetStatus())
-    {
-    case EContextStatus::AbortedByFailedLockAcquisition:
-        break;
-    case EContextStatus::AbortedByRequest:
-    case EContextStatus::AbortedByLanguage:
-        break;
-    default:
-        ASSERT(!"Should not be reached");
-        break;
-    }
+	ASSERT(Context->IsAborting());
 }
 
 void FTransaction::CommitNested()
diff --git a/Engine/Source/Runtime/Core/Public/AutoRTFM/AutoRTFM.h b/Engine/Source/Runtime/Core/Public/AutoRTFM/AutoRTFM.h
--- a/Engine/Source/Runtime/Core/Public/AutoRTFM/AutoRTFM.h
+++ b/Engine/Source/Runtime/Core/Public/AutoRTFM/AutoRTFM.h
@@ -71,7 +71,8 @@ typedef enum
     autortfm_aborted_by_language,
     autortfm_committed,
 	autortfm_aborted_by_transact_in_open_commit,
-	autortfm_aborted_by_transact_in_open_abort
+	autortfm_aborted_by_transact_in_open_abort,
+	autortfm_aborted_by_cascade
 } autortfm_result;
 
 // This must match AutoRTFM::EContextStatus.
@@ -82,7 +83,8 @@ typedef enum
 	autortfm_status_aborted_by_failed_lock_aquisition,
 	autortfm_status_aborted_by_language,
 	autortfm_status_aborted_by_request,
-	autortfm_status_committing
+	autortfm_status_committing,
+	autortfm_status_aborted_by_cascade
 } autortfm_status;
 
 // Tells if we are currently running in a transaction. This will return true in an
@@ -182,6 +184,17 @@ UE_AUTORTFM_FORCEINLINE autortfm_result autortfm_abort_transaction()
 }
 #endif
 
+// End a transaction, discard all changes, and cause all parent transactions of
+// the current transaction to do the same.
+#if UE_AUTORTFM
+UE_AUTORTFM_API autortfm_result autortfm_cascading_abort_transaction();
+#else
+UE_AUTORTFM_FORCEINLINE autortfm_result autortfm_cascading_abort_transaction()
+{
+	return autortfm_aborted_by_cascade;
+}
+#endif
+
 // Clear the status of a transaction that was aborted in the open
 #if UE_AUTORTFM
 UE_AUTORTFM_API void autortfm_clear_transaction_status();
@@ -375,7 +388,8 @@ enum class ETransactionResult
     AbortedByLanguage = autortfm_aborted_by_language,
     Committed = autortfm_committed,
 	AbortedByTransactInOpenCommit = autortfm_aborted_by_transact_in_open_commit,
-	AbortedByTransactInOpenAbort = autortfm_aborted_by_transact_in_open_abort
+	AbortedByTransactInOpenAbort = autortfm_aborted_by_transact_in_open_abort,
+	AbortedByCascade = autortfm_aborted_by_cascade
 };
 
 enum class EContextStatus
@@ -386,6 +400,7 @@ enum class EContextStatus
 	AbortedByLanguage = autortfm_status_aborted_by_language,
 	AbortedByRequest = autortfm_status_aborted_by_request,
 	Committing = autortfm_status_committing,
+	AbortedByCascade = autortfm_status_aborted_by_cascade
 };
 
 UE_AUTORTFM_FORCEINLINE bool IsTransactional() { return autortfm_is_transactional(); }
@@ -434,6 +449,11 @@ UE_AUTORTFM_FORCEINLINE ETransactionResult AbortTransaction()
 	return static_cast<ETransactionResult>(autortfm_abort_transaction());
 }
 
+UE_AUTORTFM_FORCEINLINE ETransactionResult CascadingAbortTransaction()
+{
+	return static_cast<ETransactionResult>(autortfm_cascading_abort_transaction());
+}
+
 UE_AUTORTFM_FORCEINLINE void ClearTransactionStatus()
 {
 	autortfm_clear_transaction_status();
