Optimization when decoding IoStore data

#rnx

- Only create a copy of each block for decryption if we actually need to decrypt (ie there is a valid AES key)
- If we do need to decrypt the blocks then only create a single buffer and reuse it for each block rather than creating a new buffer and copying to it for each block.

diff --git a/Engine/Source/Runtime/Core/Private/IO/IoChunkEncoding.cpp b/Engine/Source/Runtime/Core/Private/IO/IoChunkEncoding.cpp
--- a/Engine/Source/Runtime/Core/Private/IO/IoChunkEncoding.cpp
+++ b/Engine/Source/Runtime/Core/Private/IO/IoChunkEncoding.cpp
@@ -232,6 +232,8 @@ bool FIoChunkEncoding::Decode(
 
 bool FIoChunkEncoding::Decode(const FIoChunkDecodingParams& Params, FMemoryView EncodedBlocks, FMutableMemoryView OutRawData)
 {
+	TRACE_CPUPROFILER_EVENT_SCOPE(FIoChunkEncoding::Decode);
+
 	if (Params.TotalRawSize < Params.RawOffset + OutRawData.GetSize())
 	{
 		return false;
@@ -254,7 +256,7 @@ bool FIoChunkEncoding::Decode(const FIoChunkDecodingParams& Params, FMemoryView
 	const uint32 LastBlockIndex = uint32((RawOffset + OutRawData.GetSize() - 1) / BlockSize);
 
 	uint64 RawBlockOffset = RawOffset % BlockSize;
-	
+
 	uint64 EncodedOffset = 0;
 	for (uint32 BlockIndex = 0; BlockIndex < FirstBlockIndex; ++BlockIndex)
 	{
@@ -263,6 +265,21 @@ bool FIoChunkEncoding::Decode(const FIoChunkDecodingParams& Params, FMemoryView
 	// Subtract the encoded offset if the encoded blocks is a partial range of all encoded block(s)
 	EncodedBlocks += (EncodedOffset - Params.EncodedOffset);
 
+	// Buffer used to decrypt blocks into, only allocated if a valid ASE key is found
+	TUniquePtr<uint8[]> DecryptedBlockBuffer;
+	uint64 DecryptedScratchBufferSize = 0;
+
+	if (AESKey.IsValid())
+	{
+		for (uint32 BlockIndex = FirstBlockIndex; BlockIndex <= LastBlockIndex; BlockIndex++)
+		{
+			const uint64 AlignedBlockSize = Align(EncodedBlockSize[BlockIndex], FAES::AESBlockSize);
+			DecryptedScratchBufferSize = FMath::Max(AlignedBlockSize, DecryptedScratchBufferSize);
+		}
+
+		DecryptedBlockBuffer = MakeUnique<uint8[]>(DecryptedScratchBufferSize);
+	}
+
 	const bool bVerifyBlocks = BlockHash.IsEmpty() == false;
 	for (uint32 BlockIndex = FirstBlockIndex; BlockIndex <= LastBlockIndex; BlockIndex++)
 	{
@@ -271,13 +288,12 @@ bool FIoChunkEncoding::Decode(const FIoChunkDecodingParams& Params, FMemoryView
 		const uint32 CompressedBlockSize = EncodedBlockSize[BlockIndex];
 		const uint32 AlignedBlockSize = Align(CompressedBlockSize, FAES::AESBlockSize);
 
-		FIoBuffer Tmp(AlignedBlockSize);
-		Tmp.GetMutableView().CopyFrom(EncodedBlocks.Left(AlignedBlockSize));
+		FMemoryView BlockView = EncodedBlocks.Left(AlignedBlockSize);
 
 		if (bVerifyBlocks)
 		{
 			check(BlockHash.IsEmpty() == false);
-			FIoBlockHash Hash = HashBlock(Tmp.GetView());
+			FIoBlockHash Hash = HashBlock(BlockView);
 			if (Hash != BlockHash[BlockIndex])
 			{
 				return false;
@@ -286,14 +302,19 @@ bool FIoChunkEncoding::Decode(const FIoChunkDecodingParams& Params, FMemoryView
 
 		if (AESKey.IsValid())
 		{
-			FAES::DecryptData(Tmp.GetData(), uint32(Tmp.GetSize()), AESKey);
+			check(BlockView.GetSize() <= DecryptedScratchBufferSize);
+			FMemory::Memcpy(DecryptedBlockBuffer.Get(), BlockView.GetData(), BlockView.GetSize());
+			
+			FAES::DecryptData(DecryptedBlockBuffer.Get(), BlockView.GetSize(), AESKey);
+
+			BlockView = MakeMemoryView(DecryptedBlockBuffer.Get(), BlockView.GetSize());
 		}
 
 		if (CompressedBlockSize < RawBlockSize)
 		{
 			if (RawBlockReadSize == RawBlockSize)
 			{
-				if (!FCompression::UncompressMemory(Params.CompressionFormat, OutRawData.GetData(), int32(RawBlockReadSize), Tmp.GetData(), CompressedBlockSize))
+				if (!FCompression::UncompressMemory(Params.CompressionFormat, OutRawData.GetData(), int32(RawBlockReadSize), BlockView.GetData(), CompressedBlockSize))
 				{
 					return false;
 				}
@@ -301,7 +322,7 @@ bool FIoChunkEncoding::Decode(const FIoChunkDecodingParams& Params, FMemoryView
 			else
 			{
 				FIoBuffer RawBlockTmp = FIoBuffer(RawBlockSize);
-				if (!FCompression::UncompressMemory(Params.CompressionFormat, RawBlockTmp.GetData(), int32(RawBlockSize), Tmp.GetData(), CompressedBlockSize))
+				if (!FCompression::UncompressMemory(Params.CompressionFormat, RawBlockTmp.GetData(), int32(RawBlockSize), BlockView.GetData(), CompressedBlockSize))
 				{
 					return false;
 				}
@@ -310,7 +331,7 @@ bool FIoChunkEncoding::Decode(const FIoChunkDecodingParams& Params, FMemoryView
 		}
 		else
 		{
-			OutRawData.CopyFrom(Tmp.GetView().Mid(RawBlockOffset, RawBlockReadSize));
+			OutRawData.CopyFrom(BlockView.Mid(RawBlockOffset, RawBlockReadSize));
 		}
 
 		RawBlockOffset = 0;
