added more debug info to help track down cause of bugs

diff --git a/Engine/Source/Runtime/Core/Public/Templates/UnrealTemplate.h b/Engine/Source/Runtime/Core/Public/Templates/UnrealTemplate.h
--- a/Engine/Source/Runtime/Core/Public/Templates/UnrealTemplate.h
+++ b/Engine/Source/Runtime/Core/Public/Templates/UnrealTemplate.h
@@ -160,7 +160,14 @@ constexpr bool IntFitsIn(InType In)
 template<typename OutType, typename InType>
 OutType IntCastChecked(InType In)
 {
-	checkf(IntFitsIn<OutType>(In), TEXT("Loss of data caused by narrowing conversion"));
+	if constexpr (std::is_signed_v<InType>)
+	{
+		checkf(IntFitsIn<OutType>(In), TEXT("Loss of data caused by narrowing conversion, In = %" INT64_FMT), (int64)In);
+	}
+	else
+	{
+		checkf(IntFitsIn<OutType>(In), TEXT("Loss of data caused by narrowing conversion, In = %" UINT64_FMT), (uint64)In);
+	}
 	return static_cast<OutType>(In);
 }
 
