Fixed TStaticBitArray::operator== const incorrectness, which wasn't being called by operator!= but was instead causing an incorrect path to be called via implicit bool conversions.
Disabled implicit bool construction, conversions and comparisons.
Added HasAnyBitsSet() instead.
Added tests.

diff --git a/Engine/Source/Runtime/Core/Private/Containers/StaticBitArrayTests.cpp b/Engine/Source/Runtime/Core/Private/Containers/StaticBitArrayTests.cpp
new file mode 100644
--- /dev/null
+++ b/Engine/Source/Runtime/Core/Private/Containers/StaticBitArrayTests.cpp
@@ -0,0 +1,49 @@
+// Copyright Epic Games, Inc. All Rights Reserved.
+
+#if WITH_LOW_LEVEL_TESTS
+
+#include "CoreMinimal.h"
+#include "Misc/AutomationTest.h"
+#include "Containers/StaticBitArray.h"
+#include "TestHarness.h"
+
+#include <catch2/generators/catch_generators.hpp>
+
+TEST_CASE("System::Core::Containers::TStaticBitArray", "[Containers][Smoke]")
+{
+	SECTION("Empty bit array")
+	{
+		TStaticBitArray<128> Bits;
+		CHECK(!Bits.HasAnyBitsSet());
+		CHECK(!Bits);
+		CHECK(Bits.Num() == 128);
+
+		TStaticBitArray<128> BitsEmpty;
+		CHECK(Bits == BitsEmpty);
+		CHECK(!(Bits != BitsEmpty));
+	}
+
+	SECTION("One bit set")
+	{
+		TStaticBitArray<128> Bits;
+		Bits[88] = true;
+		CHECK(Bits.HasAnyBitsSet());
+		CHECK(Bits);
+
+		TStaticBitArray<128> BitsEmpty;
+		CHECK(Bits != BitsEmpty);
+		CHECK(!(Bits == BitsEmpty));
+
+		TStaticBitArray<128> BitsSame;
+		BitsSame[88] = true;
+		CHECK(Bits == BitsSame);
+		CHECK(!(Bits != BitsSame));
+
+		TStaticBitArray<128> BitsDifferent;
+		BitsDifferent[44] = true;
+		CHECK(Bits != BitsDifferent);
+		CHECK(!(Bits == BitsDifferent));
+	}
+}
+
+#endif
diff --git a/Engine/Source/Runtime/Core/Public/Containers/StaticBitArray.h b/Engine/Source/Runtime/Core/Public/Containers/StaticBitArray.h
--- a/Engine/Source/Runtime/Core/Public/Containers/StaticBitArray.h
+++ b/Engine/Source/Runtime/Core/Public/Containers/StaticBitArray.h
@@ -86,6 +86,7 @@ public:
 	}
 
 	/** Constructor that allows initializing by assignment from 0 */
+	UE_DEPRECATED(5.4, "Implicitly constructing a TStaticBitArray from 0 has been deprecated - please use the default constructor instead")
 	FORCEINLINE TStaticBitArray(UnspecifiedZeroType)
 	{
 		Clear_();
@@ -145,8 +146,7 @@ public:
 		}
 	}
 
-	// Conversion to bool
-	FORCEINLINE operator UnspecifiedBoolType() const
+	FORCEINLINE bool HasAnyBitsSet() const
 	{
 		WordType And = 0;
 		for(int32 Index = 0; Index < NumWords; ++Index)
@@ -154,7 +154,13 @@ public:
 			And |= Words[Index];
 		}
 
-		return And ? &FBoolType::Valid : NULL;
+		return And != 0;
+	}
+
+	// Explicit conversion to bool
+	FORCEINLINE explicit operator bool() const
+	{
+		return this->HasAnyBitsSet();
 	}
 
 	// Accessors.
@@ -247,7 +253,7 @@ public:
 		Results ^= B;
 		return Results;
 	}
-	FORCEINLINE bool operator==(const TStaticBitArray<NumBits>& B)
+	FORCEINLINE bool operator==(const TStaticBitArray<NumBits>& B) const
 	{
 		for(int32 Index = 0; Index < NumWords; ++Index)
 		{
@@ -259,6 +265,7 @@ public:
 		return true;
 	}
 	/** This operator only exists to disambiguate == in statements of the form (flags == 0) */
+	UE_DEPRECATED(5.4, "Comparing a TStaticBitArray to zero has been deprecated - please use !BitArray.HasAnyBitsSet() instead")
 	friend FORCEINLINE bool operator==(const TStaticBitArray<NumBits>& A, UnspecifiedBoolType Value)
 	{
 		return (UnspecifiedBoolType)A == Value;
@@ -269,6 +276,7 @@ public:
 		return !(*this == B);
 	}
 	/** != simple maps to == */
+	UE_DEPRECATED(5.4, "Comparing a TStaticBitArray to zero has been deprecated - please use BitArray.HasAnyBitsSet() instead")
 	friend FORCEINLINE bool operator!=(const TStaticBitArray<NumBits>& A, UnspecifiedBoolType Value)
 	{
 		return !(A == Value);
