[Switch] Fix for no callback on exceptions in retail kits. Make some basic calls so we get the info we need when the thread hang detector triggers.

diff --git a/Engine/Source/Runtime/Core/Private/HAL/ThreadHeartBeat.cpp b/Engine/Source/Runtime/Core/Private/HAL/ThreadHeartBeat.cpp
--- a/Engine/Source/Runtime/Core/Private/HAL/ThreadHeartBeat.cpp
+++ b/Engine/Source/Runtime/Core/Private/HAL/ThreadHeartBeat.cpp
@@ -195,10 +195,8 @@ void FORCENOINLINE FThreadHeartBeat::OnHang(double HangDuration, uint32 ThreadTh
 	LastHungThreadId = ThreadThatHung;
 	FGenericCrashContext::SetEngineData(TEXT("HungThread"), LexToString(ThreadThatHung));
 
-#if !UE_BUILD_SHIPPING
 	// Delegate implementation will be called from the hang detector thread and not from the hung thread
 	OnHangDelegate.ExecuteIfBound(ThreadThatHung);
-#endif
 
 	// We want to avoid all memory allocations if a hang is detected.
 	// Force a crash in a way that will generate a crash report.
diff --git a/Engine/Source/Runtime/Core/Public/HAL/ThreadHeartBeat.h b/Engine/Source/Runtime/Core/Public/HAL/ThreadHeartBeat.h
--- a/Engine/Source/Runtime/Core/Public/HAL/ThreadHeartBeat.h
+++ b/Engine/Source/Runtime/Core/Public/HAL/ThreadHeartBeat.h
@@ -18,10 +18,7 @@
 
 DECLARE_DELEGATE_OneParam(FOnThreadStuck, uint32);
 DECLARE_DELEGATE_OneParam(FOnThreadUnstuck, uint32);
-
-#if !UE_BUILD_SHIPPING
 DECLARE_DELEGATE_OneParam(FOnHangDelegate, uint32);
-#endif
 
 /**
  * Our own local clock.
@@ -146,9 +143,7 @@ class FThreadHeartBeat : public FRunnable
 
 	FOnThreadStuck OnStuck;
 	FOnThreadUnstuck OnUnstuck;
-#if !UE_BUILD_SHIPPING
 	FOnHangDelegate OnHangDelegate;
-#endif
 
 	CORE_API FThreadHeartBeat();
 	CORE_API virtual ~FThreadHeartBeat();
@@ -246,11 +241,8 @@ public:
 	/*
 	* Get delegate for callback on hang.
 	* Delegate implementation will be called from the hang detector thread and not from the hung thread
-	* Disabled in shipping build
 	*/
-#if !UE_BUILD_SHIPPING
 	FOnHangDelegate& GetOnHangDelegate() { return OnHangDelegate; }
-#endif
 
 	/*
 	*  Get hang duration threshold.
