Fix brief period during engine init in which backgrounding could change the android hardware window before the RHI is initialized.
 - The Android window is locked during preinit to prevent app backgrounding destroying it during RHIInit. Once the RHI is initialized it is able to handle backgrounding.

diff --git a/Engine/Source/Runtime/Core/Public/Android/AndroidPlatformMisc.h b/Engine/Source/Runtime/Core/Public/Android/AndroidPlatformMisc.h
--- a/Engine/Source/Runtime/Core/Public/Android/AndroidPlatformMisc.h
+++ b/Engine/Source/Runtime/Core/Public/Android/AndroidPlatformMisc.h
@@ -286,7 +286,7 @@ public:
 
 	static void SaveDeviceOrientation(EDeviceScreenOrientation NewDeviceOrentation) { DeviceOrientation = NewDeviceOrentation; }
 
-	// Window access is locked by the game thread before preinit and unlocked here after RHIInit (PlatformCreateDynamicRHI). 
+	// Window access is locked by the game thread before preinit and unlocked here after RHIInit
 	static CORE_API void UnlockAndroidWindow();
 	
 	static CORE_API TArray<int32> GetSupportedNativeDisplayRefreshRates();
diff --git a/Engine/Source/Runtime/RHI/Private/Android/AndroidDynamicRHI.cpp b/Engine/Source/Runtime/RHI/Private/Android/AndroidDynamicRHI.cpp
--- a/Engine/Source/Runtime/RHI/Private/Android/AndroidDynamicRHI.cpp
+++ b/Engine/Source/Runtime/RHI/Private/Android/AndroidDynamicRHI.cpp
@@ -52,7 +52,5 @@ FDynamicRHI* PlatformCreateDynamicRHI()
 		DynamicRHI = DynamicRHIModule->CreateRHI(RequestedFeatureLevel);
 	}
 
-	FPlatformMisc::UnlockAndroidWindow();
-
 	return DynamicRHI;
 }
diff --git a/Engine/Source/Runtime/RHI/Private/DynamicRHI.cpp b/Engine/Source/Runtime/RHI/Private/DynamicRHI.cpp
--- a/Engine/Source/Runtime/RHI/Private/DynamicRHI.cpp
+++ b/Engine/Source/Runtime/RHI/Private/DynamicRHI.cpp
@@ -440,6 +440,12 @@ void RHIPostInit(const TArray<uint32>& InPixelFormatByteWidth)
 	check(GDynamicRHI);
 	GDynamicRHI->InitPixelFormatInfo(InPixelFormatByteWidth);
 	GDynamicRHI->PostInit();
+
+#if PLATFORM_ANDROID
+	// The Android HW window is locked during init to prevent it being destroyed asynch during app backgrounding
+	// It is unlocked after the RHI is initialized as it is able to handle backgrounding.
+	FAndroidMisc::UnlockAndroidWindow();
+#endif
 }
 
 void RHIExit()
