MB2 Allocator optimizations.
- Removed extra page allocation for VMA pool header.
- Added VMA pool stats reporting.
- Removed extra lock per pool (was not used).
- Scale factor parameter is configurable via compile-time definition.

#tests replays

diff --git a/Engine/Source/Runtime/Core/Private/Android/AndroidPlatformMemory.cpp b/Engine/Source/Runtime/Core/Private/Android/AndroidPlatformMemory.cpp
--- a/Engine/Source/Runtime/Core/Private/Android/AndroidPlatformMemory.cpp
+++ b/Engine/Source/Runtime/Core/Private/Android/AndroidPlatformMemory.cpp
@@ -134,6 +134,9 @@ static int64 GetNativeHeapAllocatedSize()
 	return AllocatedSize;
 }
 
+/** Controls growth of pools - see PooledVirtualMemoryAllocator.cpp */
+extern float GVMAPoolScale;
+
 void FAndroidPlatformMemory::Init()
 {
 	// Only allow this method to be called once
@@ -154,6 +157,7 @@ void FAndroidPlatformMemory::Init()
 		float((double)MemoryStats.AvailablePhysical/1024.0/1024.0),
 		float((double)MemoryConstants.PageSize/1024.0)
 		);
+	UE_LOG(LogInit, Log, TEXT(" - VirtualMemoryAllocator pools will grow at scale %g"), GVMAPoolScale);
 
 #if HAS_ANDROID_MEMORY_ADVICE
 	CVarAndroidUseMemoryAdvisor->SetOnChangedCallback(FConsoleVariableDelegate::CreateStatic(&OnCVarAndroidUseMemoryAdvisorChanged));
diff --git a/Engine/Source/Runtime/Core/Private/HAL/Allocators/PooledVirtualMemoryAllocator.cpp b/Engine/Source/Runtime/Core/Private/HAL/Allocators/PooledVirtualMemoryAllocator.cpp
--- a/Engine/Source/Runtime/Core/Private/HAL/Allocators/PooledVirtualMemoryAllocator.cpp
+++ b/Engine/Source/Runtime/Core/Private/HAL/Allocators/PooledVirtualMemoryAllocator.cpp
@@ -14,14 +14,15 @@
 
 using T64KBAlignedPool = TMemoryPool<65536>;
 
+#ifndef UE_VMA_POOL_SCALE
+#define UE_VMA_POOL_SCALE 1.4f
+#endif
+
 /** Scale parameter used when growing the pools on allocation (and scaling them back), configurable from the commandline */
-float GVMAPoolScale = 1.4f;
+float GVMAPoolScale = UE_VMA_POOL_SCALE;
 
 struct FPoolDescriptor : public FPooledVirtualMemoryAllocator::FPoolDescriptorBase
 {
-	/** Lock on modifying the pool - temporary, the class can be made lock-less */
-	FCriticalSection PoolAccessLock;
-
 	/** Pool itself */
 	T64KBAlignedPool* Pool;
 };
@@ -235,6 +236,18 @@ FPooledVirtualMemoryAllocator::FPoolDescriptorBase* FPooledVirtualMemoryAllocato
 	Ptr->Pool = new (PointerToPool) T64KBAlignedPool(AllocationSize, reinterpret_cast<SIZE_T>(AlignedMemoryForThePool), NumPooledAllocations, 
 		PointerToBookkeepingMemory, VMBlock);
 
+	// mmap can return allocation that is not 64k aligned, so we had to add 64k padding after the header.
+	// But we can trim excess pages after the fact.
+	uint8* AlignedMemoryAfterThePool = Align(AlignedMemoryForThePool + AllocationSize * static_cast<SIZE_T>(NumPooledAllocations), FPlatformMemory::FPlatformVirtualMemoryBlock::GetVirtualSizeAlignment());
+	SIZE_T PagesUsed = (AlignedMemoryAfterThePool - RawPtr) / FPlatformMemory::FPlatformVirtualMemoryBlock::GetVirtualSizeAlignment();
+	SIZE_T PagesLeft = Ptr->VMSizeDivVirtualSizeAlignment - PagesUsed;
+	if (PagesLeft > 0)
+	{
+		checkf(Ptr->VMSizeDivVirtualSizeAlignment > PagesLeft, TEXT("Arithmetic error calculating excess pages"));
+		Ptr->VMSizeDivVirtualSizeAlignment -= PagesLeft;
+		FPlatformMemory::FPlatformVirtualMemoryBlock(AlignedMemoryAfterThePool, (uint32)PagesLeft).FreeVirtual();
+	}
+
 	return Ptr;
 }
 
@@ -281,4 +294,28 @@ uint64 FPooledVirtualMemoryAllocator::GetCachedFreeTotal()
 
 	return TotalFree;
 }
+
+void FPooledVirtualMemoryAllocator::DumpAllocatorStats(FOutputDevice& Ar)
+{
+	for (int32 IdxSizeClass = 0; IdxSizeClass < Limits::NumAllocationSizeClasses; ++IdxSizeClass)
+	{
+		FScopeLock Lock(&ClassesLocks[IdxSizeClass]);
+		SIZE_T AllocationSizeForClass = CalculateAllocationSizeFromClass(IdxSizeClass);
+		Ar.Logf(TEXT("PooledVirtualMemoryAllocator Index: %d, SizeClass: %.2fKB"),
+			IdxSizeClass,
+			(double)AllocationSizeForClass / 1024.0);
+		for(FPoolDescriptorBase* BaseDesc = ClassesListHeads[IdxSizeClass]; BaseDesc; BaseDesc = BaseDesc->Next)
+		{
+			FPoolDescriptor& Desc = static_cast<FPoolDescriptor&>(*BaseDesc);
+
+			Ar.Logf(TEXT("Pool[%3d]: %10.2fKB allocatable, %10.2fKB overhead"),
+				IdxSizeClass,
+				(double)Desc.Pool->GetAllocatableMemorySize() / 1024.0,
+				(double)Desc.Pool->GetOverheadSize() / 1024.0);
+		}
+	}
+
+	FScopeLock Lock(&OsAllocatorCacheLock);
+	OsAllocatorCache.DumpAllocatorStats(Ar);
+}
 #endif
diff --git a/Engine/Source/Runtime/Core/Private/HAL/MallocBinned2.cpp b/Engine/Source/Runtime/Core/Private/HAL/MallocBinned2.cpp
--- a/Engine/Source/Runtime/Core/Private/HAL/MallocBinned2.cpp
+++ b/Engine/Source/Runtime/Core/Private/HAL/MallocBinned2.cpp
@@ -1509,6 +1509,13 @@ void FMallocBinned2::DumpAllocatorStats(class FOutputDevice& Ar)
 			AllocatedOSSmallPoolMemory + AllocatedLargePoolMemoryWAlignment + Binned2PoolInfoMemory + Binned2HashMemory + Binned2TLSMemory
 			) / (1024.0f * 1024.0f));
 	Ar.Logf(TEXT("Cached free OS pages: %fmb"), ((double)OSPageAllocatorCachedFreeSize) / (1024.0f * 1024.0f));
+
+#if !PLATFORM_UNIX && !PLATFORM_ANDROID
+	// Doesn't implemented
+#else
+	CachedOSPageAllocator.DumpAllocatorStats(Ar);
+#endif
+
 #else
 	Ar.Logf(TEXT("Allocator Stats for binned2 are not in this build set BINNED2_ALLOCATOR_STATS 1 in MallocBinned2.cpp"));
 #endif
diff --git a/Engine/Source/Runtime/Core/Public/HAL/Allocators/CachedOSPageAllocator.h b/Engine/Source/Runtime/Core/Public/HAL/Allocators/CachedOSPageAllocator.h
--- a/Engine/Source/Runtime/Core/Public/HAL/Allocators/CachedOSPageAllocator.h
+++ b/Engine/Source/Runtime/Core/Public/HAL/Allocators/CachedOSPageAllocator.h
@@ -72,6 +72,11 @@ struct TCachedOSPageAllocator : private FCachedOSPageAllocator
 		return FCachedOSPageAllocator::IsOSAllocation(Size, CachedByteLimit);
 	}
 
+	void DumpAllocatorStats(class FOutputDevice& Ar)
+	{
+		Ar.Logf(TEXT("CachedOSPageAllocator = %fkb"), (double)GetCachedFreeTotal() / 1024.0);
+	}
+
 private:
 	FFreePageBlock FreedPageBlocks[NumCacheBlocks*2];
 	SIZE_T         CachedTotal;
diff --git a/Engine/Source/Runtime/Core/Public/HAL/Allocators/PooledVirtualMemoryAllocator.h b/Engine/Source/Runtime/Core/Public/HAL/Allocators/PooledVirtualMemoryAllocator.h
--- a/Engine/Source/Runtime/Core/Public/HAL/Allocators/PooledVirtualMemoryAllocator.h
+++ b/Engine/Source/Runtime/Core/Public/HAL/Allocators/PooledVirtualMemoryAllocator.h
@@ -84,6 +84,8 @@ struct FPooledVirtualMemoryAllocator
 	/** Update memory stats (does nothing in that implementation) */
 	void UpdateStats() {}
 
+	void DumpAllocatorStats(class FOutputDevice& Ar);
+
 private:
 
 	enum Limits
