Fix the issue the boot hotfix file doesn't get loaded in shipping version;
Only load from boot hotfix file for CVars which has flag ECVF_SaveForNextBoot.

Reason: Engine has more restriction on loading .ini file.

#test Tried in Shipping version on console

diff --git a/Engine/Source/Runtime/Core/Private/Misc/ConfigUtilities.cpp b/Engine/Source/Runtime/Core/Private/Misc/ConfigUtilities.cpp
--- a/Engine/Source/Runtime/Core/Private/Misc/ConfigUtilities.cpp
+++ b/Engine/Source/Runtime/Core/Private/Misc/ConfigUtilities.cpp
@@ -7,7 +7,7 @@
 #include "Misc/CoreDelegates.h"
 #include "Misc/FileHelper.h"
 
-#define UE_HOTFIX_FOR_NEXT_BOOT_FILENAME TEXT("HotfixForNextBoot.ini")
+#define UE_HOTFIX_FOR_NEXT_BOOT_FILENAME TEXT("HotfixForNextBoot.txt")
 
 namespace UE::ConfigUtilities
 {
@@ -49,16 +49,26 @@ void LoadCVarsFromFileForNextBoot(TMap<FString, FString>& OutCVars)
 
 	if (!FileManager.FileExists(*FullPath))
 	{
+		UE_LOG(LogConfig, Log, TEXT("No local boot hotfix file found at: [%s]"), *FullPath);
 		return;
 	}
 
 	FString Content;
-	FFileHelper::LoadFileToString(Content, *FullPath);
+	if (!FFileHelper::LoadFileToString(Content, *FullPath))
+	{
+		UE_LOG(LogConfig, Error, TEXT("Failed to load local boot hotfix file: [%s]"), *FullPath);
+		return;
+	}
 
 	// Delete it so that we don't worry about it when write.
 	// Also if for some reason the switch don't work well when boot even before getting latest hotfix, 
 	// the next boot will likely succeed by default like before without this file
-	FileManager.Delete(*FullPath);
+	if (!FileManager.Delete(*FullPath, true/*RequireExists*/))
+	{
+		UE_LOG(LogConfig, Error, TEXT("Failed to delete local boot hotfix file [%s]"), *FullPath);
+	}
+
+	UE_LOG(LogConfig, Log, TEXT("Local boot hotfix file [%s] loaded and deleted"), *FullPath);
 
 	TArray<FString> Lines;
 	Content.ParseIntoArrayLines(Lines);
@@ -80,6 +90,7 @@ void SaveCVarForNextBoot(const TCHAR* Key, const TCHAR* Value)
 #if !UE_SERVER
 	if (!FPaths::HasProjectPersistentDownloadDir())
 	{
+		UE_LOG(LogConfig, Log, TEXT("No persistent download dir, ignoring CVar %s hotfix for next boot"), Key);
 		return;
 	}
 
@@ -98,6 +109,8 @@ void SaveCVarForNextBoot(const TCHAR* Key, const TCHAR* Value)
 
 	const FString FullPath = FPaths::ProjectPersistentDownloadDir() / UE_HOTFIX_FOR_NEXT_BOOT_FILENAME;
 	FFileHelper::SaveStringToFile(ContentToSave, *FullPath);
+
+	UE_LOG(LogConfig, Log, TEXT("Local boot hotfix file [%s] saved with hotfixed CVar: %s=%s"), *FullPath, Key, Value);
 #endif // !UE_SERVER
 }
 
@@ -110,7 +123,7 @@ void ApplyCVarsFromBootHotfix()
 	for (const TPair<FString, FString>& CVarPair : CVarsToApply)
 	{
 		IConsoleVariable* CVar = IConsoleManager::Get().FindConsoleVariable(*CVarPair.Key); 
-		if (CVar)
+		if (CVar && CVar->TestFlags(ECVF_SaveForNextBoot))
 		{
 			CVar->Set(*CVarPair.Value, ECVF_SetByHotfix);
 		}
@@ -200,6 +213,7 @@ void OnSetCVarFromIniEntry(const TCHAR *IniFile, const TCHAR *Key, const TCHAR*
 
 		if (CVar->TestFlags(ECVF_SaveForNextBoot) && (SetBy == ECVF_SetByHotfix))
 		{
+			UE_LOG(LogConfig, Log, TEXT("Saving %s for boot hotfix"), Key);
 			SaveCVarForNextBoot(Key, Value);
 		}
 	}
