Logging: Made SerializeForLog exist only if the needed operator<< exists

Added CallSerializeForLog and CSerializableForLog to assist in implementing SerializeForLog overloads.

#rnx

diff --git a/Engine/Source/Runtime/Core/Public/Logging/StructuredLog.h b/Engine/Source/Runtime/Core/Public/Logging/StructuredLog.h
--- a/Engine/Source/Runtime/Core/Public/Logging/StructuredLog.h
+++ b/Engine/Source/Runtime/Core/Public/Logging/StructuredLog.h
@@ -2,6 +2,7 @@
 
 #pragma once
 
+#include "Concepts/Insertable.h"
 #include "Containers/StringFwd.h"
 #include "Logging/LogCategory.h"
 #include "Logging/LogMacros.h"
@@ -9,6 +10,7 @@
 #include "Serialization/CompactBinary.h"
 #include "Serialization/CompactBinaryWriter.h"
 #include "Templates/IsArrayOrRefOfType.h"
+#include "Templates/Models.h"
 #include "Templates/UnrealTypeTraits.h"
 
 #include <atomic>
@@ -205,12 +207,33 @@ private:
 	const TCHAR* TextKey = nullptr;
 };
 
-template <typename ValueType>
+/**
+ * Serializes the value to be used in a log message.
+ *
+ * Overload this when the log behavior needs to differ from general serialization to compact binary.
+ */
+template <typename ValueType UE_REQUIRES(TModels_V<CInsertable<FCbWriter&>, ValueType>)>
 inline void SerializeForLog(FCbWriter& Writer, ValueType&& Value)
 {
 	Writer << (ValueType&&)Value;
 }
 
+/** Describes a type that can be serialized for use in a log message. */
+struct CSerializableForLog
+{
+	template <typename ValueType>
+	auto Requires(FCbWriter& Writer, ValueType&& Value) -> decltype(
+		SerializeForLog(Writer, (ValueType&&)Value)
+	);
+};
+
+/** Wrapper to support calling SerializeForLog with ADL from within an overload of SerializeForLog. */
+template <typename ValueType UE_REQUIRES(TModels_V<CSerializableForLog, ValueType>)>
+inline void CallSerializeForLog(FCbWriter& Writer, ValueType&& Value)
+{
+	SerializeForLog(Writer, (ValueType&&)Value);
+}
+
 } // UE
 
 namespace UE::Logging::Private
