Reduce default file cache size used by virtual texturing. Add profiling to track file cache usage.

diff --git a/Engine/Source/Runtime/Core/Private/FileCache/FileCache.cpp b/Engine/Source/Runtime/Core/Private/FileCache/FileCache.cpp
--- a/Engine/Source/Runtime/Core/Private/FileCache/FileCache.cpp
+++ b/Engine/Source/Runtime/Core/Private/FileCache/FileCache.cpp
@@ -5,6 +5,7 @@
 #include "Containers/Queue.h"
 #include "Containers/LockFreeList.h"
 #include "Templates/TypeHash.h"
+#include "Misc/CoreDelegates.h"
 #include "Misc/ScopeLock.h"
 #include "Async/AsyncFileHandle.h"
 #include "Async/TaskGraphInterfaces.h"
@@ -12,6 +13,7 @@
 #include "HAL/PlatformFile.h"
 #include "HAL/PlatformFileManager.h"
 #include "HAL/IConsoleManager.h"
+#include "ProfilingDebugging/CsvProfiler.h"
 
 DECLARE_STATS_GROUP(TEXT("Streaming File Cache"), STATGROUP_SFC, STATCAT_Advanced);
 
@@ -22,6 +24,8 @@ DECLARE_CYCLE_STAT(TEXT("EvictAll"), STAT_SFC_EvictAll, STATGROUP_SFC);
 // These below are pretty high throughput and probably should be removed once the system gets more mature
 DECLARE_CYCLE_STAT(TEXT("Find Eviction Candidate"), STAT_SFC_FindEvictionCandidate, STATGROUP_SFC);
 
+CSV_DEFINE_CATEGORY(FileCache, true);
+
 DEFINE_LOG_CATEGORY_STATIC(LogStreamingFileCache, Log, All);
 
 static int32 GFileCacheBlockSizeKB = 64;
@@ -33,7 +37,7 @@ static FAutoConsoleVariableRef CVarFileCacheBlockSize(
 	ECVF_ReadOnly
 );
 
-static int32 GNumFileCacheBlocks = 256;
+static int32 GNumFileCacheBlocks = 64;
 static FAutoConsoleVariableRef CVarNumFileCacheBlocks(
 	TEXT("fc.NumBlocks"),
 	GNumFileCacheBlocks,
@@ -197,6 +201,7 @@ public:
 	uint8* Memory;
 	int64 SizeInBytes;
 	int32 NumFreeSlots;
+	int32 MinNumFreeSlots;
 };
 
 static FFileCache &GetCache()
@@ -334,6 +339,7 @@ FFileCache::FFileCache(int32 NumSlots)
 		FConsoleCommandDelegate::CreateRaw(this, &FFileCache::EvictFileCacheFromConsole))
 	, SizeInBytes(NumSlots * CacheSlotID::GetSize())
 	, NumFreeSlots(NumSlots)
+	, MinNumFreeSlots(NumSlots)
 {
 	LLM_SCOPE(ELLMTag::FileSystem);
 
@@ -378,12 +384,18 @@ FFileCache::FFileCache(int32 NumSlots)
 #if !UE_BUILD_SHIPPING
 	bFileCacheInitialized.store(true);
 #endif
+
+	FCoreDelegates::OnBeginFrameRT.AddLambda([this]() { 
+		CSV_CUSTOM_STAT(FileCache, NumFreeSlots, NumFreeSlots, ECsvCustomStatOp::Set); 
+		CSV_CUSTOM_STAT(FileCache, MinNumFreeSlots, MinNumFreeSlots, ECsvCustomStatOp::Set);
+		});
 }
 
 CacheSlotID FFileCache::AcquireAndLockSlot(FFileCacheHandle* InHandle, CacheLineID InLineID)
 {
 	check(NumFreeSlots > 0);
 	--NumFreeSlots;
+	MinNumFreeSlots = FMath::Min(MinNumFreeSlots, NumFreeSlots);
 
 	const int32 SlotIndex = SlotInfo[0].NextSlotIndex;
 	check(SlotIndex != 0);
@@ -419,6 +431,7 @@ void FFileCache::LockSlot(CacheSlotID InSlotID)
 	{
 		check(NumFreeSlots > 0);
 		--NumFreeSlots;
+		MinNumFreeSlots = FMath::Min(MinNumFreeSlots, NumFreeSlots);
 		UnlinkSlot(SlotIndex);
 	}
 	Info.LockCount = PrevLockCount + 1;
