Remove some redundant check()s from TArrayView constructors, in particular when copying from other TArrayViews.

diff --git a/Engine/Source/Runtime/Core/Public/Containers/ArrayView.h b/Engine/Source/Runtime/Core/Public/Containers/ArrayView.h
--- a/Engine/Source/Runtime/Core/Public/Containers/ArrayView.h
+++ b/Engine/Source/Runtime/Core/Public/Containers/ArrayView.h
@@ -140,6 +140,11 @@ public:
 
 	static_assert(TIsSigned<SizeType>::Value, "TArrayView only supports signed index types");
 
+	// Defaulted object behavior - we want compiler-generated functions rather than going through the generic range constructor.
+	TArrayView(const TArrayView&) = default;
+	TArrayView& operator=(const TArrayView&) = default;
+	~TArrayView() = default;
+
 	/**
 	 * Constructor.
 	 */
@@ -177,7 +182,7 @@ public:
 				>
 			>::Value
 		>::Type,
-		std::enable_if_t<TIsTArrayView<std::decay_t<OtherRangeType>>::Value>* = nullptr
+		std::enable_if_t<TIsTArrayView<std::decay_t<OtherRangeType>>::Value && !std::is_same_v<std::decay_t<OtherRangeType>, TArrayView>>* = nullptr
 	>
 	FORCEINLINE TArrayView(OtherRangeType&& Other)
 		: DataPtr(TChooseClass<
@@ -187,20 +192,27 @@ public:
 					>::Result::GetData(Forward<OtherRangeType>(Other)))
 	{
 		const auto InCount = GetNum(Forward<OtherRangeType>(Other));
-		check((InCount >= 0) && ((sizeof(InCount) < sizeof(SizeType)) || (InCount <= static_cast<decltype(InCount)>(TNumericLimits<SizeType>::Max()))));
+		using InCountType = decltype(InCount);
+
+		// Unlike the other constructor, we don't need to check(InCount >= 0), because it's coming from a TArrayView which guarantees that
+		if constexpr (sizeof(InCountType) > sizeof(SizeType) || (sizeof(InCountType) == sizeof(SizeType) && std::is_unsigned_v<InCountType>))
+		{
+			check(InCount <= static_cast<InCountType>(TNumericLimits<SizeType>::Max()));
+		}
+
 		ArrayNum = (SizeType)InCount;
 	}
 	template <
 		typename OtherRangeType,
 		typename CVUnqualifiedOtherRangeType = std::remove_cv_t<typename TRemoveReference<OtherRangeType>::Type>,
 		typename = typename TEnableIf<
-		TAnd<
-		TIsContiguousContainer<CVUnqualifiedOtherRangeType>,
-		TOr<
-		TIsCompatibleRangeType<OtherRangeType>,
-		TIsReinterpretableRangeType<OtherRangeType>
-		>
-		>::Value
+			TAnd<
+				TIsContiguousContainer<CVUnqualifiedOtherRangeType>,
+				TOr<
+					TIsCompatibleRangeType<OtherRangeType>,
+					TIsReinterpretableRangeType<OtherRangeType>
+				>
+			>::Value
 		>::Type,
 		std::enable_if_t<!TIsTArrayView<std::decay_t<OtherRangeType>>::Value>* = nullptr
 	>
@@ -212,7 +224,15 @@ public:
 		>::Result::GetData(Forward<OtherRangeType>(Other)))
 	{
 		const auto InCount = GetNum(Forward<OtherRangeType>(Other));
-		check((InCount >= 0) && ((sizeof(InCount) < sizeof(SizeType)) || (InCount <= static_cast<decltype(InCount)>(TNumericLimits<SizeType>::Max()))));
+		using InCountType = decltype(InCount);
+		if constexpr (sizeof(InCountType) > sizeof(SizeType) || (sizeof(InCountType) == sizeof(SizeType) && std::is_unsigned_v<InCountType>))
+		{
+			check(InCount >= 0 && InCount <= static_cast<InCountType>(TNumericLimits<SizeType>::Max()));
+		}
+		else
+		{
+			check(InCount >= 0);
+		}
 		ArrayNum = (SizeType)InCount;
 	}
 
