Added and used TIsTArrayView_V in addition to TIsTArrayView.

diff --git a/Engine/Source/Runtime/Core/Public/Containers/ArrayView.h b/Engine/Source/Runtime/Core/Public/Containers/ArrayView.h
--- a/Engine/Source/Runtime/Core/Public/Containers/ArrayView.h
+++ b/Engine/Source/Runtime/Core/Public/Containers/ArrayView.h
@@ -95,11 +95,18 @@ namespace UE::Core::ArrayView::Private
 	};
 }
 
-template <typename T>                                  struct TIsTArrayView                                                       { static constexpr bool Value = false; };
-template <typename InElementType, typename InSizeType> struct TIsTArrayView<               TArrayView<InElementType, InSizeType>> { static constexpr bool Value = true;  };
-template <typename InElementType, typename InSizeType> struct TIsTArrayView<      volatile TArrayView<InElementType, InSizeType>> { static constexpr bool Value = true;  };
-template <typename InElementType, typename InSizeType> struct TIsTArrayView<const          TArrayView<InElementType, InSizeType>> { static constexpr bool Value = true;  };
-template <typename InElementType, typename InSizeType> struct TIsTArrayView<const volatile TArrayView<InElementType, InSizeType>> { static constexpr bool Value = true;  };
+template <typename T>                                  constexpr bool TIsTArrayView_V                                                       = false;
+template <typename InElementType, typename InSizeType> constexpr bool TIsTArrayView_V<               TArrayView<InElementType, InSizeType>> = true;
+template <typename InElementType, typename InSizeType> constexpr bool TIsTArrayView_V<      volatile TArrayView<InElementType, InSizeType>> = true;
+template <typename InElementType, typename InSizeType> constexpr bool TIsTArrayView_V<const          TArrayView<InElementType, InSizeType>> = true;
+template <typename InElementType, typename InSizeType> constexpr bool TIsTArrayView_V<const volatile TArrayView<InElementType, InSizeType>> = true;
+
+template <typename T>
+struct TIsTArrayView
+{
+	static constexpr bool Value = TIsTArrayView_V<T>;
+	static constexpr bool value = TIsTArrayView_V<T>;
+};
 
 /**
  * Templated fixed-size view of another array
@@ -182,7 +189,7 @@ public:
 				>
 			>::Value
 		>::Type,
-		std::enable_if_t<TIsTArrayView<std::decay_t<OtherRangeType>>::Value && !std::is_same_v<std::decay_t<OtherRangeType>, TArrayView>>* = nullptr
+		std::enable_if_t<TIsTArrayView_V<std::decay_t<OtherRangeType>> && !std::is_same_v<std::decay_t<OtherRangeType>, TArrayView>>* = nullptr
 	>
 	FORCEINLINE TArrayView(OtherRangeType&& Other)
 		: DataPtr(TChooseClass<
@@ -214,7 +221,7 @@ public:
 				>
 			>::Value
 		>::Type,
-		std::enable_if_t<!TIsTArrayView<std::decay_t<OtherRangeType>>::Value>* = nullptr
+		std::enable_if_t<!TIsTArrayView_V<std::decay_t<OtherRangeType>>>* = nullptr
 	>
 	FORCEINLINE TArrayView(OtherRangeType&& Other UE_LIFETIMEBOUND)
 		: DataPtr(TChooseClass<
@@ -799,7 +806,7 @@ template <
 	typename OtherRangeType,
 	typename CVUnqualifiedOtherRangeType = std::remove_cv_t<typename TRemoveReference<OtherRangeType>::Type>,
 	typename = typename TEnableIf<TIsContiguousContainer<CVUnqualifiedOtherRangeType>::Value>::Type,
-	std::enable_if_t<TIsTArrayView<std::decay_t<OtherRangeType>>::Value>* = nullptr
+	std::enable_if_t<TIsTArrayView_V<std::decay_t<OtherRangeType>>>* = nullptr
 >
 auto MakeArrayView(OtherRangeType&& Other)
 {
@@ -809,7 +816,7 @@ template <
 	typename OtherRangeType,
 	typename CVUnqualifiedOtherRangeType = std::remove_cv_t<typename TRemoveReference<OtherRangeType>::Type>,
 	typename = typename TEnableIf<TIsContiguousContainer<CVUnqualifiedOtherRangeType>::Value>::Type,
-	std::enable_if_t<!TIsTArrayView<std::decay_t<OtherRangeType>>::Value>* = nullptr
+	std::enable_if_t<!TIsTArrayView_V<std::decay_t<OtherRangeType>>>* = nullptr
 >
 auto MakeArrayView(OtherRangeType&& Other UE_LIFETIMEBOUND)
 {
