Replaced legacy traits with std:: traits in TArrayView.

diff --git a/Engine/Source/Runtime/Core/Public/Containers/ArrayView.h b/Engine/Source/Runtime/Core/Public/Containers/ArrayView.h
--- a/Engine/Source/Runtime/Core/Public/Containers/ArrayView.h
+++ b/Engine/Source/Runtime/Core/Public/Containers/ArrayView.h
@@ -43,7 +43,7 @@ namespace UE::Core::ArrayView::Private
 	FORCEINLINE decltype(auto) GetReinterpretedDataHelper(T&& Arg)
 	{
 		auto NaturalPtr = GetData(Forward<T>(Arg));
-		using NaturalElementType = typename TRemovePointer<decltype(NaturalPtr)>::Type;
+		using NaturalElementType = std::remove_pointer_t<decltype(NaturalPtr)>;
 
 		auto Size = GetNum(Arg);
 		auto EndPtr = NaturalPtr + Size;
@@ -58,7 +58,7 @@ namespace UE::Core::ArrayView::Private
 	template <typename RangeType, typename ElementType>
 	struct TIsCompatibleRangeType
 	{
-		static constexpr bool Value = TIsCompatibleElementType<typename TRemovePointer<decltype(GetData(DeclVal<RangeType&>()))>::Type, ElementType>::Value;
+		static constexpr bool Value = TIsCompatibleElementType<std::remove_pointer_t<decltype(GetData(DeclVal<RangeType&>()))>, ElementType>::Value;
 
 		template <typename T>
 		static decltype(auto) GetData(T&& Arg)
@@ -74,7 +74,7 @@ namespace UE::Core::ArrayView::Private
 	struct TIsReinterpretableRangeType
 	{
 	private:
-		using NaturalElementType = typename TRemovePointer<decltype(GetData(DeclVal<RangeType&>()))>::Type;
+		using NaturalElementType = std::remove_pointer_t<decltype(GetData(DeclVal<RangeType&>()))>;
 		using TypeCompat = TContainerElementTypeCompatibility<NaturalElementType>;
 
 	public:
@@ -179,7 +179,7 @@ public:
 	 */
 	template <
 		typename OtherRangeType,
-		typename CVUnqualifiedOtherRangeType = std::remove_cv_t<typename TRemoveReference<OtherRangeType>::Type>,
+		typename CVUnqualifiedOtherRangeType = std::remove_cv_t<std::remove_reference_t<OtherRangeType>>,
 		typename = typename TEnableIf<
 			TAnd<
 				TIsContiguousContainer<CVUnqualifiedOtherRangeType>,
@@ -211,7 +211,7 @@ public:
 	}
 	template <
 		typename OtherRangeType,
-		typename CVUnqualifiedOtherRangeType = std::remove_cv_t<typename TRemoveReference<OtherRangeType>::Type>,
+		typename CVUnqualifiedOtherRangeType = std::remove_cv_t<std::remove_reference_t<OtherRangeType>>,
 		typename = typename TEnableIf<
 			TAnd<
 				TIsContiguousContainer<CVUnqualifiedOtherRangeType>,
@@ -804,23 +804,23 @@ struct TIsContiguousContainer<TArrayView<T, SizeType>>
 
 template <
 	typename OtherRangeType,
-	typename CVUnqualifiedOtherRangeType = std::remove_cv_t<typename TRemoveReference<OtherRangeType>::Type>,
+	typename CVUnqualifiedOtherRangeType = std::remove_cv_t<std::remove_reference_t<OtherRangeType>>,
 	typename = typename TEnableIf<TIsContiguousContainer<CVUnqualifiedOtherRangeType>::Value>::Type,
 	std::enable_if_t<TIsTArrayView_V<std::decay_t<OtherRangeType>>>* = nullptr
 >
 auto MakeArrayView(OtherRangeType&& Other)
 {
-	return TArrayView<typename TRemovePointer<decltype(GetData(DeclVal<OtherRangeType&>()))>::Type>(Forward<OtherRangeType>(Other));
+	return TArrayView<std::remove_pointer_t<decltype(GetData(DeclVal<OtherRangeType&>()))>>(Forward<OtherRangeType>(Other));
 }
 template <
 	typename OtherRangeType,
-	typename CVUnqualifiedOtherRangeType = std::remove_cv_t<typename TRemoveReference<OtherRangeType>::Type>,
+	typename CVUnqualifiedOtherRangeType = std::remove_cv_t<std::remove_reference_t<OtherRangeType>>,
 	typename = typename TEnableIf<TIsContiguousContainer<CVUnqualifiedOtherRangeType>::Value>::Type,
 	std::enable_if_t<!TIsTArrayView_V<std::decay_t<OtherRangeType>>>* = nullptr
 >
 auto MakeArrayView(OtherRangeType&& Other UE_LIFETIMEBOUND)
 {
-	return TArrayView<typename TRemovePointer<decltype(GetData(DeclVal<OtherRangeType&>()))>::Type>(Forward<OtherRangeType>(Other));
+	return TArrayView<std::remove_pointer_t<decltype(GetData(DeclVal<OtherRangeType&>()))>>(Forward<OtherRangeType>(Other));
 }
 
 template<typename ElementType>
