Added TChar implementations for UTF8CHAR - simply mirroring ANSICHAR right now.

diff --git a/Engine/Source/Runtime/Core/Public/Misc/Char.h b/Engine/Source/Runtime/Core/Public/Misc/Char.h
--- a/Engine/Source/Runtime/Core/Public/Misc/Char.h
+++ b/Engine/Source/Runtime/Core/Public/Misc/Char.h
@@ -171,4 +171,18 @@ template <> inline bool TChar<ANSICHAR>::IsHexDigit(ANSICHAR Char) { return ::is
 template <> inline bool TChar<ANSICHAR>::IsWhitespace(ANSICHAR Char) { return ::isspace((unsigned char)Char) != 0; }
 template <> inline bool TChar<ANSICHAR>::IsControl(ANSICHAR Char) { return ::iscntrl((unsigned char)Char) != 0; }
 
+
+/*-----------------------------------------------------------------------------
+	UTF8CHAR specialized functions
+-----------------------------------------------------------------------------*/
+template <> inline bool TChar<UTF8CHAR>::IsUpper(UTF8CHAR Char) { return ::isupper((unsigned char)Char) != 0; }
+template <> inline bool TChar<UTF8CHAR>::IsLower(UTF8CHAR Char) { return ::islower((unsigned char)Char) != 0; }
+template <> inline bool TChar<UTF8CHAR>::IsAlpha(UTF8CHAR Char) { return ::isalpha((unsigned char)Char) != 0; }
+template <> inline bool TChar<UTF8CHAR>::IsGraph(UTF8CHAR Char) { return ::isgraph((unsigned char)Char) != 0; }
+template <> inline bool TChar<UTF8CHAR>::IsPrint(UTF8CHAR Char) { return ::isprint((unsigned char)Char) != 0; }
+template <> inline bool TChar<UTF8CHAR>::IsPunct(UTF8CHAR Char) { return ::ispunct((unsigned char)Char) != 0; }
+template <> inline bool TChar<UTF8CHAR>::IsAlnum(UTF8CHAR Char) { return ::isalnum((unsigned char)Char) != 0; }
+template <> inline bool TChar<UTF8CHAR>::IsDigit(UTF8CHAR Char) { return ::isdigit((unsigned char)Char) != 0; }
+template <> inline bool TChar<UTF8CHAR>::IsHexDigit(UTF8CHAR Char) { return ::isxdigit((unsigned char)Char) != 0; }
 template <> inline bool TChar<UTF8CHAR>::IsWhitespace(UTF8CHAR Char) { return ::isspace((unsigned char)Char) != 0; }
+template <> inline bool TChar<UTF8CHAR>::IsControl(UTF8CHAR Char) { return ::iscntrl((unsigned char)Char) != 0; }
