CHARTEXT(CharType, x) which expands out to x, TEXT(x) or UTF8TEXT(x) depending on CharType.

diff --git a/Engine/Source/Runtime/Core/Private/Containers/String.cpp b/Engine/Source/Runtime/Core/Private/Containers/String.cpp
--- a/Engine/Source/Runtime/Core/Private/Containers/String.cpp
+++ b/Engine/Source/Runtime/Core/Private/Containers/String.cpp
@@ -24,9 +24,6 @@
 #include "Templates/UnrealTemplate.h"
 #include "UObject/NameTypes.h"
 
-#define UE_STRING_COMPILING_UTF8 0
-#define UE_STRING_TEXT(x) TEXT(x)
-
 /* FString implementation
  *****************************************************************************/
 
@@ -277,7 +274,7 @@ void FString::AssignRange(const ElementType* OtherData, int32 OtherLen)
 			// Unless the input is longer, this might be assigned from a view of itself.
 			ElementType* DataPtr = Data.GetData();
 			FMemory::Memmove(DataPtr, OtherData, OtherLen * sizeof(ElementType));
-			DataPtr[OtherLen] = UE_STRING_TEXT('\0');
+			DataPtr[OtherLen] = CHARTEXT(ElementType, '\0');
 			Data.RemoveAt(OtherLen + 1, ThisLen - OtherLen);
 		}
 		else
@@ -286,7 +283,7 @@ void FString::AssignRange(const ElementType* OtherData, int32 OtherLen)
 			Data.AddUninitialized(OtherLen + 1);
 			ElementType* DataPtr = Data.GetData();
 			FMemory::Memcpy(DataPtr, OtherData, OtherLen * sizeof(ElementType));
-			DataPtr[OtherLen] = UE_STRING_TEXT('\0');
+			DataPtr[OtherLen] = CHARTEXT(ElementType, '\0');
 		}
 	}
 }
@@ -316,7 +313,7 @@ void FString::Reset(int32 NewReservedSize)
 	Data.Reset(NewSizeIncludingTerminator);
 	if (ElementType* DataPtr = Data.GetData())
 	{
-		*DataPtr = UE_STRING_TEXT('\0');
+		*DataPtr = CHARTEXT(ElementType, '\0');
 	}
 }
 
@@ -353,7 +350,7 @@ FString& FString::AppendChar(ElementType InChar)
 
 		Data.AddUninitialized(InsertCount);
 		Data[InsertIndex] = InChar;
-		Data[InsertIndex+1] = UE_STRING_TEXT('\0');
+		Data[InsertIndex+1] = CHARTEXT(ElementType, '\0');
 	}
 	return *this;
 }
@@ -558,7 +555,7 @@ void FString::RemoveSpacesInline()
 	// Copy null-terminating character.
 	if (CopyToIndex <= StringLength)
 	{
-		RawData[CopyToIndex] = UE_STRING_TEXT('\0');
+		RawData[CopyToIndex] = CHARTEXT(ElementType, '\0');
 		Data.SetNum(CopyToIndex + 1, false);
 	}
 }
@@ -704,7 +701,7 @@ UE_NODISCARD FORCEINLINE FString ConcatRangeFString(const LhsCharType* Lhs, int3
 	ElementType* ResultData = Result.GetCharArray().GetData();
 	CopyAssignItems(ResultData, Lhs, LhsLen);
 	CopyAssignItems(ResultData + LhsLen, Rhs.GetCharArray().GetData(), RhsLen);
-	*(ResultData + LhsLen + RhsLen) = UE_STRING_TEXT('\0');
+	*(ResultData + LhsLen + RhsLen) = CHARTEXT(ElementType, '\0');
 
 	return Result;
 }
@@ -774,19 +771,19 @@ void FString::PathAppend(const ElementType* Str, int32 StrLength)
 	int32 DataNum = Data.Num();
 	if (StrLength == 0)
 	{
-		if (DataNum > 1 && Data[DataNum - 2] != UE_STRING_TEXT('/') && Data[DataNum - 2] != UE_STRING_TEXT('\\'))
+		if (DataNum > 1 && Data[DataNum - 2] != CHARTEXT(ElementType, '/') && Data[DataNum - 2] != CHARTEXT(ElementType, '\\'))
 		{
-			Data[DataNum - 1] = UE_STRING_TEXT('/');
-			Data.Add(UE_STRING_TEXT('\0'));
+			Data[DataNum - 1] = CHARTEXT(ElementType, '/');
+			Data.Add(CHARTEXT(ElementType, '\0'));
 		}
 	}
 	else
 	{
 		if (DataNum > 0)
 		{
-			if (DataNum > 1 && Data[DataNum - 2] != UE_STRING_TEXT('/') && Data[DataNum - 2] != UE_STRING_TEXT('\\') && *Str != UE_STRING_TEXT('/'))
+			if (DataNum > 1 && Data[DataNum - 2] != CHARTEXT(ElementType, '/') && Data[DataNum - 2] != CHARTEXT(ElementType, '\\') && *Str != CHARTEXT(ElementType, '/'))
 			{
-				Data[DataNum - 1] = UE_STRING_TEXT('/');
+				Data[DataNum - 1] = CHARTEXT(ElementType, '/');
 			}
 			else
 			{
@@ -797,7 +794,7 @@ void FString::PathAppend(const ElementType* Str, int32 StrLength)
 
 		Reserve(DataNum + StrLength);
 		Data.Append(Str, StrLength);
-		Data.Add(UE_STRING_TEXT('\0'));
+		Data.Add(CHARTEXT(ElementType, '\0'));
 	}
 }
 
@@ -1019,7 +1016,7 @@ FString FString::FormatAsNumber( int32 InNumber )
 		dec++;
 		if( dec == 3 && x > 0 )
 		{
-			Result += UE_STRING_TEXT(',');
+			Result += CHARTEXT(ElementType, ',');
 			dec = 0;
 		}
 	}
@@ -1055,7 +1052,7 @@ void FString::SerializeAsANSICharArray( FArchive& Ar, int32 MinCharacters ) cons
 
 void FString::AppendInt( int32 Num )
 {
-	const ElementType* DigitToChar	= UE_STRING_TEXT("9876543210123456789");
+	const ElementType* DigitToChar	= CHARTEXT(ElementType, "9876543210123456789");
 	constexpr int32 ZeroDigitIndex	= 9;
 	bool bIsNumberNegative			= Num < 0;
 	const int32 TempBufferSize		= 16; // 16 is big enough
@@ -1071,7 +1068,7 @@ void FString::AppendInt( int32 Num )
 
 	if( bIsNumberNegative )
 	{
-		TempNum[--TempAt] = UE_STRING_TEXT('-');
+		TempNum[--TempAt] = CHARTEXT(ElementType, '-');
 	}
 
 	const ElementType* CharPtr = TempNum + TempAt;
@@ -1092,7 +1089,7 @@ FString FString::FromBlob(const uint8* SrcBuffer,const uint32 SrcSize)
 	// Convert and append each byte in the buffer
 	for (uint32 Count = 0; Count < SrcSize; Count++)
 	{
-		Result += FString::Printf(UE_STRING_TEXT("%03d"),(uint8)SrcBuffer[Count]);
+		Result += FString::Printf(CHARTEXT(ElementType, "%03d"),(uint8)SrcBuffer[Count]);
 	}
 	return Result;
 }
@@ -1105,7 +1102,7 @@ bool FString::ToBlob(const FString& Source,uint8* DestBuffer,const uint32 DestSi
 		(Source.Len() % 3) == 0)
 	{
 		ElementType ConvBuffer[4];
-		ConvBuffer[3] = UE_STRING_TEXT('\0');
+		ConvBuffer[3] = CHARTEXT(ElementType, '\0');
 		int32 WriteIndex = 0;
 		// Walk the string 3 chars at a time
 		for (int32 Index = 0; Index < Source.Len(); Index += 3, WriteIndex++)
@@ -1127,7 +1124,7 @@ FString FString::FromHexBlob( const uint8* SrcBuffer, const uint32 SrcSize )
 	// Convert and append each byte in the buffer
 	for (uint32 Count = 0; Count < SrcSize; Count++)
 	{
-		Result += FString::Printf( UE_STRING_TEXT( "%02X" ), (uint8)SrcBuffer[Count] );
+		Result += FString::Printf( CHARTEXT(ElementType,  "%02X" ), (uint8)SrcBuffer[Count] );
 	}
 	return Result;
 }
@@ -1140,7 +1137,7 @@ bool FString::ToHexBlob( const FString& Source, uint8* DestBuffer, const uint32
 		 (Source.Len() % 2) == 0)
 	{
 		ElementType ConvBuffer[3];
-		ConvBuffer[2] = UE_STRING_TEXT( '\0' );
+		ConvBuffer[2] = CHARTEXT(ElementType, '\0' );
 		int32 WriteIndex = 0;
 		// Walk the string 2 chars at a time
 		ElementType* End = nullptr;
@@ -1171,7 +1168,7 @@ FString FString::SanitizeFloat( double InFloat, const int32 InMinFractionalDigit
 	StripNegativeZero(InFloat);
 
 	// First create the string
-	FString TempString = FString::Printf(UE_STRING_TEXT("%f"), InFloat);
+	FString TempString = FString::Printf(CHARTEXT(ElementType, "%f"), InFloat);
 	if (!TempString.IsNumeric())
 	{
 		// String did not format as a valid decimal number so avoid messing with it
@@ -1184,13 +1181,13 @@ FString FString::SanitizeFloat( double InFloat, const int32 InMinFractionalDigit
 	for (int32 CharIndex = TempString.Len() - 1; CharIndex >= 0; --CharIndex)
 	{
 		const ElementType Char = TempString[CharIndex];
-		if (Char == UE_STRING_TEXT('.'))
+		if (Char == CHARTEXT(ElementType, '.'))
 		{
 			DecimalSeparatorIndex = CharIndex;
 			TrimIndex = FMath::Max(TrimIndex, DecimalSeparatorIndex);
 			break;
 		}
-		if (TrimIndex == INDEX_NONE && Char != UE_STRING_TEXT('0'))
+		if (TrimIndex == INDEX_NONE && Char != CHARTEXT(ElementType, '0'))
 		{
 			TrimIndex = CharIndex + 1;
 		}
@@ -1204,7 +1201,7 @@ FString FString::SanitizeFloat( double InFloat, const int32 InMinFractionalDigit
 		if (TrimIndex == DecimalSeparatorIndex)
 		{
 			// Re-add the decimal separator
-			TempString.AppendChar(UE_STRING_TEXT('.'));
+			TempString.AppendChar(CHARTEXT(ElementType, '.'));
 		}
 
 		const int32 NumFractionalDigits = (TempString.Len() - DecimalSeparatorIndex) - 1;
@@ -1214,7 +1211,7 @@ FString FString::SanitizeFloat( double InFloat, const int32 InMinFractionalDigit
 			TempString.Reserve(TempString.Len() + FractionalDigitsToPad);
 			for (int32 Cx = 0; Cx < FractionalDigitsToPad; ++Cx)
 			{
-				TempString.AppendChar(UE_STRING_TEXT('0'));
+				TempString.AppendChar(CHARTEXT(ElementType, '0'));
 			}
 		}
 	}
@@ -1224,7 +1221,7 @@ FString FString::SanitizeFloat( double InFloat, const int32 InMinFractionalDigit
 
 FString FString::Chr(ElementType Ch)
 {
-	ElementType Temp[2]= { Ch, UE_STRING_TEXT('\0') };
+	ElementType Temp[2]= { Ch, CHARTEXT(ElementType, '\0') };
 	return FString(Temp);
 }
 
@@ -1239,7 +1236,7 @@ FString FString::ChrN( int32 NumCharacters, ElementType Char )
 	{
 		Temp[Cx] = Char;
 	}
-	Temp.Data[NumCharacters] = UE_STRING_TEXT('\0');
+	Temp.Data[NumCharacters] = CHARTEXT(ElementType, '\0');
 	return Temp;
 }
 
@@ -1249,7 +1246,7 @@ FString FString::LeftPad( int32 ChCount ) const
 
 	if (Pad > 0)
 	{
-		return ChrN(Pad, UE_STRING_TEXT(' ')) + *this;
+		return ChrN(Pad, CHARTEXT(ElementType, ' ')) + *this;
 	}
 	else
 	{
@@ -1262,7 +1259,7 @@ FString FString::RightPad( int32 ChCount ) const
 
 	if (Pad > 0)
 	{
-		return *this + ChrN(Pad, UE_STRING_TEXT(' '));
+		return *this + ChrN(Pad, CHARTEXT(ElementType, ' '));
 	}
 	else
 	{
@@ -1319,11 +1316,11 @@ int32 FString::ParseIntoArrayWS( TArray<FString>& OutArray, const ElementType* p
 	// (if you want to split on white space and another character)
 	const ElementType* WhiteSpace[] =
 	{
-		UE_STRING_TEXT(" "),
-		UE_STRING_TEXT("\t"),
-		UE_STRING_TEXT("\r"),
-		UE_STRING_TEXT("\n"),
-		UE_STRING_TEXT(""),
+		CHARTEXT(ElementType, " "),
+		CHARTEXT(ElementType, "\t"),
+		CHARTEXT(ElementType, "\r"),
+		CHARTEXT(ElementType, "\n"),
+		CHARTEXT(ElementType, ""),
 	};
 
 	// start with just the standard whitespaces
@@ -1342,9 +1339,9 @@ int32 FString::ParseIntoArrayLines(TArray<FString>& OutArray, bool InCullEmpty)
 	// default array of LineEndings
 	static const ElementType* LineEndings[] =
 	{
-		UE_STRING_TEXT("\r\n"),
-		UE_STRING_TEXT("\r"),
-		UE_STRING_TEXT("\n"),	
+		CHARTEXT(ElementType, "\r\n"),
+		CHARTEXT(ElementType, "\r"),
+		CHARTEXT(ElementType, "\n"),
 	};
 
 	// start with just the standard line endings
@@ -1510,7 +1507,7 @@ int32 FString::ReplaceInline(const ElementType* SearchText, const ElementType* R
 				ReplacementCount++;
 
 				// replace the first letter of the From with 0 so we can do a strcpy (FString +=)
-				*SearchPosition = UE_STRING_TEXT('\0');
+				*SearchPosition = CHARTEXT(ElementType, '\0');
 
 				// copy everything up to the SearchPosition
 				(*this) += WritePosition;
@@ -1539,7 +1536,7 @@ int32 FString::ReplaceInline(const ElementType* SearchText, const ElementType* R
  */
 FString FString::ReplaceQuotesWithEscapedQuotes() &&
 {
-	if (Contains(UE_STRING_TEXT("\""), ESearchCase::CaseSensitive))
+	if (Contains(CHARTEXT(ElementType, "\""), ESearchCase::CaseSensitive))
 	{
 		FString Copy(MoveTemp(*this));
 
@@ -1572,12 +1569,12 @@ template <typename CharType>
 static const CharType* CharToEscapeSeqMap[6][2] =
 {
 	// Always replace \\ first to avoid double-escaping characters
-	{ UE_STRING_TEXT("\\"), UE_STRING_TEXT("\\\\") },
-	{ UE_STRING_TEXT("\n"), UE_STRING_TEXT("\\n")  },
-	{ UE_STRING_TEXT("\r"), UE_STRING_TEXT("\\r")  },
-	{ UE_STRING_TEXT("\t"), UE_STRING_TEXT("\\t")  },
-	{ UE_STRING_TEXT("\'"), UE_STRING_TEXT("\\'")  },
-	{ UE_STRING_TEXT("\""), UE_STRING_TEXT("\\\"") }
+	{ CHARTEXT(CharType, "\\"), CHARTEXT(CharType, "\\\\") },
+	{ CHARTEXT(CharType, "\n"), CHARTEXT(CharType, "\\n")  },
+	{ CHARTEXT(CharType, "\r"), CHARTEXT(CharType, "\\r")  },
+	{ CHARTEXT(CharType, "\t"), CHARTEXT(CharType, "\\t")  },
+	{ CHARTEXT(CharType, "\'"), CHARTEXT(CharType, "\\'")  },
+	{ CHARTEXT(CharType, "\""), CHARTEXT(CharType, "\\\"") }
 };
 
 template <typename CharType>
@@ -1626,13 +1623,13 @@ void FString::ConvertTabsToSpacesInline(const int32 InSpacesPerTab)
 	check(InSpacesPerTab > 0);
 
 	int32 TabIndex;
-	while ((TabIndex = Find(UE_STRING_TEXT("\t"), ESearchCase::CaseSensitive)) != INDEX_NONE )
+	while ((TabIndex = Find(CHARTEXT(ElementType, "\t"), ESearchCase::CaseSensitive)) != INDEX_NONE )
 	{
 		FString RightSide = Mid(TabIndex+1);
 		LeftInline(TabIndex, false);
 
 		//for a tab size of 4, 
-		int32 LineBegin = Find(UE_STRING_TEXT("\n"), ESearchCase::CaseSensitive, ESearchDir::FromEnd, TabIndex);
+		int32 LineBegin = Find(CHARTEXT(ElementType, "\n"), ESearchCase::CaseSensitive, ESearchDir::FromEnd, TabIndex);
 		if (LineBegin == INDEX_NONE)
 		{
 			LineBegin = 0;
@@ -1642,7 +1639,7 @@ void FString::ConvertTabsToSpacesInline(const int32 InSpacesPerTab)
 		int32 NumSpacesForTab = InSpacesPerTab - (CharactersOnLine % InSpacesPerTab);
 		for (int32 i = 0; i < NumSpacesForTab; ++i)
 		{
-			AppendChar(UE_STRING_TEXT(' '));
+			AppendChar(CHARTEXT(ElementType, ' '));
 		}
 		Append(RightSide);
 	}
@@ -1673,7 +1670,7 @@ FString FString::PrintfImpl(const ElementType* Fmt, ...)
 		};
 	}
 
-	Buffer[Result] = UE_STRING_TEXT('\0');
+	Buffer[Result] = CHARTEXT(ElementType, '\0');
 
 	FString ResultString(Buffer);
 
@@ -1707,7 +1704,7 @@ void FString::AppendfImpl(FString& AppendToMe, const ElementType* Fmt, ...)
 		};
 	}
 
-	Buffer[Result] = UE_STRING_TEXT('\0');
+	Buffer[Result] = CHARTEXT(ElementType, '\0');
 
 	AppendToMe += Buffer;
 
@@ -1740,7 +1737,7 @@ FArchive& operator<<( FArchive& Ar, FString& A )
 				if (SaveNum == MIN_int32)
 				{
 					Ar.SetCriticalError();
-					UE_LOG(LogCore, Error, TEXT("Archive is corrupted"));
+					UE_LOG(LogCore, Error, CHARTEXT(ElementType, "Archive is corrupted"));
 					return Ar;
 				}
 
@@ -1752,7 +1749,7 @@ FArchive& operator<<( FArchive& Ar, FString& A )
 			if ((MaxSerializeSize > 0) && (SaveNum > MaxSerializeSize))
 			{
 				Ar.SetCriticalError();
-				UE_LOG(LogCore, Error, TEXT("String is too large (Size: %i, Max: %i)"), SaveNum, MaxSerializeSize);
+				UE_LOG(LogCore, Error, CHARTEXT(ElementType, "String is too large (Size: %i, Max: %i)"), SaveNum, MaxSerializeSize);
 				return Ar;
 			}
 
@@ -1787,7 +1784,7 @@ FArchive& operator<<( FArchive& Ar, FString& A )
 					int Index = 0;
 					if (A.FindChar(0xffff, Index))
 					{
-						A[Index] = UE_STRING_TEXT('\0');
+						A[Index] = CHARTEXT(ElementType, '\0');
 						A.TrimToNullTerminator();
 					}
 				}
@@ -1875,24 +1872,24 @@ int32 FindMatchingClosingParenthesis(const FString& TargetString, const int32 St
 	int32 ParenthesisCount = 0;
 
 	// Move to first open parenthesis
-	while (*CurrPosition != 0 && *CurrPosition != UE_STRING_TEXT('('))
+	while (*CurrPosition != 0 && *CurrPosition != CHARTEXT(ElementType, '('))
 	{
 		++CurrPosition;
 	}
 
 	// Did we find the open parenthesis
-	if (*CurrPosition == UE_STRING_TEXT('('))
+	if (*CurrPosition == CHARTEXT(ElementType, '('))
 	{
 		++ParenthesisCount;
 		++CurrPosition;
 
 		while (*CurrPosition != 0 && ParenthesisCount > 0)
 		{
-			if (*CurrPosition == UE_STRING_TEXT('('))
+			if (*CurrPosition == CHARTEXT(ElementType, '('))
 			{
 				++ParenthesisCount;
 			}
-			else if (*CurrPosition == UE_STRING_TEXT(')'))
+			else if (*CurrPosition == CHARTEXT(ElementType, ')'))
 			{
 				--ParenthesisCount;
 			}
@@ -1900,7 +1897,7 @@ int32 FindMatchingClosingParenthesis(const FString& TargetString, const int32 St
 		}
 
 		// Did we find the matching close parenthesis
-		if (ParenthesisCount == 0 && *(CurrPosition - 1) == UE_STRING_TEXT(')'))
+		if (ParenthesisCount == 0 && *(CurrPosition - 1) == CHARTEXT(ElementType, ')'))
 		{
 			return StartSearch + UE_PTRDIFF_TO_INT32((CurrPosition - 1) - StartPosition);
 		}
@@ -1909,7 +1906,7 @@ int32 FindMatchingClosingParenthesis(const FString& TargetString, const int32 St
 	return INDEX_NONE;
 }
 
-FString SlugStringForValidName(const FString& DisplayString, const FString::ElementType* ReplaceWith /*= UE_STRING_TEXT("")*/)
+FString SlugStringForValidName(const FString& DisplayString, const FString::ElementType* ReplaceWith /*= CHARTEXT(ElementType, "")*/)
 {
 	using ElementType = FString::ElementType;
 
@@ -1920,7 +1917,7 @@ FString SlugStringForValidName(const FString& DisplayString, const FString::Elem
 	{
 		for ( int32 BadCharacterIndex = 0; BadCharacterIndex < UE_ARRAY_COUNT(INVALID_OBJECTNAME_CHARACTERS) - 1; ++BadCharacterIndex )
 		{
-			const ElementType TestChar[2] = { INVALID_OBJECTNAME_CHARACTERS[BadCharacterIndex], UE_STRING_TEXT('\0') };
+			const ElementType TestChar[2] = { INVALID_OBJECTNAME_CHARACTERS[BadCharacterIndex], CHARTEXT(ElementType, '\0') };
 			const int32 NumReplacedChars = GeneratedName.ReplaceInline(TestChar, ReplaceWith);
 		}
 	}
@@ -1965,6 +1962,3 @@ void StringConv::InlineCombineSurrogates(FString& Str)
 {
 	InlineCombineSurrogates_Array(Str.GetCharArray());
 }
-
-#undef UE_STRING_TEXT
-#undef UE_STRING_COMPILING_UTF8
\ No newline at end of file
diff --git a/Engine/Source/Runtime/Core/Public/Containers/UnrealString.h b/Engine/Source/Runtime/Core/Public/Containers/UnrealString.h
--- a/Engine/Source/Runtime/Core/Public/Containers/UnrealString.h
+++ b/Engine/Source/Runtime/Core/Public/Containers/UnrealString.h
@@ -32,9 +32,6 @@
 
 #include <type_traits>
 
-#define UE_STRING_COMPILING_UTF8 0
-#define UE_STRING_TEXT(x) TEXT(x)
-
 struct FStringFormatArg;
 template<typename InKeyType,typename InValueType,typename SetAllocator ,typename KeyFuncs > class TMap;
 
@@ -330,7 +327,7 @@ public:
 	 */
 	UE_NODISCARD FORCEINLINE const ElementType* operator*() const UE_LIFETIMEBOUND
 	{
-		return Data.Num() ? Data.GetData() : UE_STRING_TEXT("");
+		return Data.Num() ? Data.GetData() : CHARTEXT(ElementType, "");
 	}
 
 	/** 
@@ -2595,6 +2592,3 @@ struct FTextRange
 };
 
 #include "Misc/StringFormatArg.h"
-
-#undef UE_STRING_TEXT
-#undef UE_STRING_COMPILING_UTF8
\ No newline at end of file
diff --git a/Engine/Source/Runtime/Core/Public/HAL/Platform.h b/Engine/Source/Runtime/Core/Public/HAL/Platform.h
--- a/Engine/Source/Runtime/Core/Public/HAL/Platform.h
+++ b/Engine/Source/Runtime/Core/Public/HAL/Platform.h
@@ -1140,10 +1140,44 @@ namespace UE::Core::Private
 	{
 		return (UTF8CHAR)Ch;
 	}
+
+	template <typename CharType>
+	void CharTextStaticAssert()
+	{
+		static_assert(sizeof(CharType) == 0, "Unsupported char type passed to CHARTEXT");
+	}
 }
 
 #define UTF8TEXT(x) (UE::Core::Private::ToUTF8Literal(UTF8TEXT_PASTE(x)))
 #define UTF16TEXT(x) UTF16TEXT_PASTE(x)
 #define WIDETEXT(str) WIDETEXT_PASTE(str)
 
+// Expands out to x, TEXT(x) or UTF8TEXT(x) depending on CharType
+#define CHARTEXT(CharType, x) \
+	( \
+		[]() -> decltype(auto) \
+		{ \
+			/* We expect <type_traits> to already have been included any time the CHARTEXT macro is used */ \
+			using UnqualifiedCharType = std::remove_cv_t<CharType>; \
+			if constexpr (std::is_same_v<UnqualifiedCharType, ANSICHAR>) \
+			{ \
+				return x; \
+			} \
+			else if constexpr (std::is_same_v<UnqualifiedCharType, TCHAR>) \
+			{ \
+				return TEXT(x); \
+			} \
+			else if constexpr (std::is_same_v<UnqualifiedCharType, UTF8CHAR>) \
+			{ \
+				return UTF8TEXT(x); \
+			} \
+			else \
+			{ \
+				/* We want a compile error, so forward to a templated function because of the static_assert(false) problem */ \
+				UE::Core::Private::CharTextStaticAssert<CharType>(); \
+				return x; \
+			} \
+		}() \
+	)
+
 // IWYU pragma: end_exports
