Removed ChooseClass.h includes.

diff --git a/Engine/Plugins/Experimental/LearningAgents/Source/Learning/Public/MultiArray.h b/Engine/Plugins/Experimental/LearningAgents/Source/Learning/Public/MultiArray.h
--- a/Engine/Plugins/Experimental/LearningAgents/Source/Learning/Public/MultiArray.h
+++ b/Engine/Plugins/Experimental/LearningAgents/Source/Learning/Public/MultiArray.h
@@ -15,7 +15,6 @@
 #include "Templates/IdentityFunctor.h"
 #include "Templates/Invoke.h"
 #include "Templates/Less.h"
-#include "Templates/ChooseClass.h"
 #include "Templates/Sorting.h"
 #include "Templates/AlignmentTemplates.h"
 #include "Templates/IsConstructible.h"
diff --git a/Engine/Source/Editor/AnimGraph/Private/AnimStateAliasNode.cpp b/Engine/Source/Editor/AnimGraph/Private/AnimStateAliasNode.cpp
--- a/Engine/Source/Editor/AnimGraph/Private/AnimStateAliasNode.cpp
+++ b/Engine/Source/Editor/AnimGraph/Private/AnimStateAliasNode.cpp
@@ -13,7 +13,6 @@
 #include "Kismet2/CompilerResultsLog.h"
 #include "Kismet2/Kismet2NameValidators.h"
 #include "Templates/Casts.h"
-#include "Templates/ChooseClass.h"
 #include "Templates/SharedPointer.h"
 #include "UObject/Object.h"
 #include "UObject/WeakObjectPtr.h"
diff --git a/Engine/Source/Editor/AnimationBlueprintLibrary/Private/AnimationBlueprintLibrary.cpp b/Engine/Source/Editor/AnimationBlueprintLibrary/Private/AnimationBlueprintLibrary.cpp
--- a/Engine/Source/Editor/AnimationBlueprintLibrary/Private/AnimationBlueprintLibrary.cpp
+++ b/Engine/Source/Editor/AnimationBlueprintLibrary/Private/AnimationBlueprintLibrary.cpp
@@ -55,7 +55,6 @@
 #include "ReferenceSkeleton.h"
 #include "Serialization/StructuredArchiveAdapters.h"
 #include "Templates/Casts.h"
-#include "Templates/ChooseClass.h"
 #include "Templates/UnrealTemplate.h"
 #include "Trace/Detail/Channel.h"
 #include "UObject/Class.h"
diff --git a/Engine/Source/Editor/AnimationModifiers/Private/SAnimationModifiersTab.cpp b/Engine/Source/Editor/AnimationModifiers/Private/SAnimationModifiersTab.cpp
--- a/Engine/Source/Editor/AnimationModifiers/Private/SAnimationModifiersTab.cpp
+++ b/Engine/Source/Editor/AnimationModifiers/Private/SAnimationModifiersTab.cpp
@@ -38,7 +38,6 @@
 #include "Styling/AppStyle.h"
 #include "Subsystems/AssetEditorSubsystem.h"
 #include "Templates/Casts.h"
-#include "Templates/ChooseClass.h"
 #include "Templates/SubclassOf.h"
 #include "Toolkits/AssetEditorToolkit.h"
 #include "Types/SlateEnums.h"
diff --git a/Engine/Source/Editor/AnimationModifiers/Private/SModifierItemRow.cpp b/Engine/Source/Editor/AnimationModifiers/Private/SModifierItemRow.cpp
--- a/Engine/Source/Editor/AnimationModifiers/Private/SModifierItemRow.cpp
+++ b/Engine/Source/Editor/AnimationModifiers/Private/SModifierItemRow.cpp
@@ -10,7 +10,6 @@
 #include "Misc/Attribute.h"
 #include "SlotBase.h"
 #include "Styling/AppStyle.h"
-#include "Templates/ChooseClass.h"
 #include "Templates/SubclassOf.h"
 #include "UObject/WeakObjectPtrTemplates.h"
 #include "Widgets/Images/SImage.h"
diff --git a/Engine/Source/Editor/AudioEditor/Private/AssetTypeActions/AssetDefinition_SoundCue.cpp b/Engine/Source/Editor/AudioEditor/Private/AssetTypeActions/AssetDefinition_SoundCue.cpp
--- a/Engine/Source/Editor/AudioEditor/Private/AssetTypeActions/AssetDefinition_SoundCue.cpp
+++ b/Engine/Source/Editor/AudioEditor/Private/AssetTypeActions/AssetDefinition_SoundCue.cpp
@@ -20,7 +20,6 @@
 #include "Sound/SoundAttenuation.h"
 #include "Sound/SoundCue.h"
 #include "Templates/Casts.h"
-#include "Templates/ChooseClass.h"
 #include "Textures/SlateIcon.h"
 #include "ToolMenuSection.h"
 #include "Toolkits/IToolkit.h"
diff --git a/Engine/Source/Editor/AudioEditor/Private/AssetTypeActions/AssetTypeActions_SoundEffectPreset.cpp b/Engine/Source/Editor/AudioEditor/Private/AssetTypeActions/AssetTypeActions_SoundEffectPreset.cpp
--- a/Engine/Source/Editor/AudioEditor/Private/AssetTypeActions/AssetTypeActions_SoundEffectPreset.cpp
+++ b/Engine/Source/Editor/AudioEditor/Private/AssetTypeActions/AssetTypeActions_SoundEffectPreset.cpp
@@ -13,7 +13,6 @@
 #include "Engine/Engine.h"
 #include "Misc/AssertionMacros.h"
 #include "Templates/Casts.h"
-#include "Templates/ChooseClass.h"
 #include "Templates/SubclassOf.h"
 #include "Toolkits/IToolkit.h"
 #include "UObject/Class.h"
diff --git a/Engine/Source/Editor/BlueprintEditorLibrary/Private/BlueprintEditorLibrary.cpp b/Engine/Source/Editor/BlueprintEditorLibrary/Private/BlueprintEditorLibrary.cpp
--- a/Engine/Source/Editor/BlueprintEditorLibrary/Private/BlueprintEditorLibrary.cpp
+++ b/Engine/Source/Editor/BlueprintEditorLibrary/Private/BlueprintEditorLibrary.cpp
@@ -33,7 +33,6 @@
 #include "Misc/AssertionMacros.h"
 #include "Subsystems/AssetEditorSubsystem.h"
 #include "Templates/Casts.h"
-#include "Templates/ChooseClass.h"
 #include "Templates/SubclassOf.h"
 #include "Trace/Detail/Channel.h"
 #include "UObject/Class.h"
diff --git a/Engine/Source/Editor/BlueprintGraph/Private/BlueprintActionDatabase.cpp b/Engine/Source/Editor/BlueprintGraph/Private/BlueprintActionDatabase.cpp
--- a/Engine/Source/Editor/BlueprintGraph/Private/BlueprintActionDatabase.cpp
+++ b/Engine/Source/Editor/BlueprintGraph/Private/BlueprintActionDatabase.cpp
@@ -87,7 +87,6 @@
 #include "ProfilingDebugging/CpuProfilerTrace.h"
 #include "PropertyPermissionList.h"
 #include "Templates/Casts.h"
-#include "Templates/ChooseClass.h"
 #include "Templates/SubclassOf.h"
 #include "Templates/Tuple.h"
 #include "Templates/UnrealTemplate.h"
diff --git a/Engine/Source/Editor/BlueprintGraph/Private/BlueprintAssetNodeSpawner.cpp b/Engine/Source/Editor/BlueprintGraph/Private/BlueprintAssetNodeSpawner.cpp
--- a/Engine/Source/Editor/BlueprintGraph/Private/BlueprintAssetNodeSpawner.cpp
+++ b/Engine/Source/Editor/BlueprintGraph/Private/BlueprintAssetNodeSpawner.cpp
@@ -4,7 +4,6 @@
 
 #include "EdGraph/EdGraphNode.h"
 #include "Misc/AssertionMacros.h"
-#include "Templates/ChooseClass.h"
 #include "UObject/Package.h"
 
 class UObject;
diff --git a/Engine/Source/Editor/BlueprintGraph/Private/BlueprintBoundEventNodeSpawner.cpp b/Engine/Source/Editor/BlueprintGraph/Private/BlueprintBoundEventNodeSpawner.cpp
--- a/Engine/Source/Editor/BlueprintGraph/Private/BlueprintBoundEventNodeSpawner.cpp
+++ b/Engine/Source/Editor/BlueprintGraph/Private/BlueprintBoundEventNodeSpawner.cpp
@@ -22,7 +22,6 @@
 #include "Settings/EditorStyleSettings.h"
 #include "Styling/AppStyle.h"
 #include "Templates/Casts.h"
-#include "Templates/ChooseClass.h"
 #include "Textures/SlateIcon.h"
 #include "UObject/Class.h"
 #include "UObject/Object.h"
diff --git a/Engine/Source/Editor/BlueprintGraph/Private/BlueprintComponentNodeSpawner.cpp b/Engine/Source/Editor/BlueprintGraph/Private/BlueprintComponentNodeSpawner.cpp
--- a/Engine/Source/Editor/BlueprintGraph/Private/BlueprintComponentNodeSpawner.cpp
+++ b/Engine/Source/Editor/BlueprintGraph/Private/BlueprintComponentNodeSpawner.cpp
@@ -28,7 +28,6 @@
 #include "Misc/PackageName.h"
 #include "Styling/SlateIconFinder.h"
 #include "Templates/Casts.h"
-#include "Templates/ChooseClass.h"
 #include "Templates/SharedPointer.h"
 #include "Textures/SlateIcon.h"
 #include "UObject/Class.h"
diff --git a/Engine/Source/Editor/BlueprintGraph/Private/BlueprintFunctionNodeSpawner.cpp b/Engine/Source/Editor/BlueprintGraph/Private/BlueprintFunctionNodeSpawner.cpp
--- a/Engine/Source/Editor/BlueprintGraph/Private/BlueprintFunctionNodeSpawner.cpp
+++ b/Engine/Source/Editor/BlueprintGraph/Private/BlueprintFunctionNodeSpawner.cpp
@@ -36,7 +36,6 @@
 #include "ObjectEditorUtils.h"
 #include "SNodePanel.h"
 #include "Templates/Casts.h"
-#include "Templates/ChooseClass.h"
 #include "Textures/SlateIcon.h"
 #include "UObject/Class.h"
 #include "UObject/Field.h"
diff --git a/Engine/Source/Editor/BlueprintGraph/Private/BlueprintNodeSignature.cpp b/Engine/Source/Editor/BlueprintGraph/Private/BlueprintNodeSignature.cpp
--- a/Engine/Source/Editor/BlueprintGraph/Private/BlueprintNodeSignature.cpp
+++ b/Engine/Source/Editor/BlueprintGraph/Private/BlueprintNodeSignature.cpp
@@ -8,7 +8,6 @@
 #include "HAL/PlatformCrt.h"
 #include "Misc/CString.h"
 #include "Misc/SecureHash.h"
-#include "Templates/ChooseClass.h"
 #include "Templates/Tuple.h"
 
 /*******************************************************************************
diff --git a/Engine/Source/Editor/BlueprintGraph/Private/BlueprintNodeSpawner.cpp b/Engine/Source/Editor/BlueprintGraph/Private/BlueprintNodeSpawner.cpp
--- a/Engine/Source/Editor/BlueprintGraph/Private/BlueprintNodeSpawner.cpp
+++ b/Engine/Source/Editor/BlueprintGraph/Private/BlueprintNodeSpawner.cpp
@@ -13,7 +13,6 @@
 #include "Misc/AssertionMacros.h"
 #include "ProfilingDebugging/CpuProfilerTrace.h"
 #include "Styling/AppStyle.h"
-#include "Templates/ChooseClass.h"
 #include "UObject/Class.h"
 #include "UObject/Package.h"
 #include "UObject/UnrealType.h"
diff --git a/Engine/Source/Editor/BlueprintGraph/Private/BlueprintNodeTemplateCache.cpp b/Engine/Source/Editor/BlueprintGraph/Private/BlueprintNodeTemplateCache.cpp
--- a/Engine/Source/Editor/BlueprintGraph/Private/BlueprintNodeTemplateCache.cpp
+++ b/Engine/Source/Editor/BlueprintGraph/Private/BlueprintNodeTemplateCache.cpp
@@ -22,7 +22,6 @@
 #include "Misc/AssertionMacros.h"
 #include "ProfilingDebugging/CpuProfilerTrace.h"
 #include "Templates/Casts.h"
-#include "Templates/ChooseClass.h"
 #include "Templates/SubclassOf.h"
 #include "Templates/Tuple.h"
 #include "Trace/Detail/Channel.h"
diff --git a/Engine/Source/Editor/BlueprintGraph/Private/BlueprintVariableNodeSpawner.cpp b/Engine/Source/Editor/BlueprintGraph/Private/BlueprintVariableNodeSpawner.cpp
--- a/Engine/Source/Editor/BlueprintGraph/Private/BlueprintVariableNodeSpawner.cpp
+++ b/Engine/Source/Editor/BlueprintGraph/Private/BlueprintVariableNodeSpawner.cpp
@@ -24,7 +24,6 @@
 #include "ObjectEditorUtils.h"
 #include "Settings/EditorStyleSettings.h"
 #include "Templates/Casts.h"
-#include "Templates/ChooseClass.h"
 #include "Textures/SlateIcon.h"
 #include "UObject/Class.h"
 #include "UObject/NameTypes.h"
diff --git a/Engine/Source/Editor/BlueprintGraph/Private/DynamicCastHandler.cpp b/Engine/Source/Editor/BlueprintGraph/Private/DynamicCastHandler.cpp
--- a/Engine/Source/Editor/BlueprintGraph/Private/DynamicCastHandler.cpp
+++ b/Engine/Source/Editor/BlueprintGraph/Private/DynamicCastHandler.cpp
@@ -19,7 +19,6 @@
 #include "KismetCompiler.h"
 #include "Misc/AssertionMacros.h"
 #include "Templates/Casts.h"
-#include "Templates/ChooseClass.h"
 #include "Templates/SubclassOf.h"
 #include "UObject/Class.h"
 #include "UObject/NameTypes.h"
diff --git a/Engine/Source/Editor/BlueprintGraph/Private/K2Node_AddComponent.cpp b/Engine/Source/Editor/BlueprintGraph/Private/K2Node_AddComponent.cpp
--- a/Engine/Source/Editor/BlueprintGraph/Private/K2Node_AddComponent.cpp
+++ b/Engine/Source/Editor/BlueprintGraph/Private/K2Node_AddComponent.cpp
@@ -28,7 +28,6 @@
 #include "Math/Transform.h"
 #include "Serialization/Archive.h"
 #include "Templates/Casts.h"
-#include "Templates/ChooseClass.h"
 #include "Templates/SubclassOf.h"
 #include "UObject/BlueprintsObjectVersion.h" // for ComponentTemplateClassSupport
 #include "UObject/Class.h"
diff --git a/Engine/Source/Editor/BlueprintGraph/Private/K2Node_BaseAsyncTask.cpp b/Engine/Source/Editor/BlueprintGraph/Private/K2Node_BaseAsyncTask.cpp
--- a/Engine/Source/Editor/BlueprintGraph/Private/K2Node_BaseAsyncTask.cpp
+++ b/Engine/Source/Editor/BlueprintGraph/Private/K2Node_BaseAsyncTask.cpp
@@ -35,7 +35,6 @@
 #include "Misc/ConfigCacheIni.h"
 #include "Misc/Parse.h"
 #include "Templates/Casts.h"
-#include "Templates/ChooseClass.h"
 #include "Trace/Detail/Channel.h"
 #include "UObject/Class.h"
 #include "UObject/Field.h"
diff --git a/Engine/Source/Editor/BlueprintGraph/Private/K2Node_ClassDynamicCast.cpp b/Engine/Source/Editor/BlueprintGraph/Private/K2Node_ClassDynamicCast.cpp
--- a/Engine/Source/Editor/BlueprintGraph/Private/K2Node_ClassDynamicCast.cpp
+++ b/Engine/Source/Editor/BlueprintGraph/Private/K2Node_ClassDynamicCast.cpp
@@ -17,7 +17,6 @@
 #include "KismetCompilerMisc.h"
 #include "Misc/AssertionMacros.h"
 #include "Templates/Casts.h"
-#include "Templates/ChooseClass.h"
 #include "Templates/SubclassOf.h"
 #include "UObject/Class.h"
 #include "UObject/Interface.h"
diff --git a/Engine/Source/Editor/BlueprintGraph/Private/K2Node_CreateDelegate.cpp b/Engine/Source/Editor/BlueprintGraph/Private/K2Node_CreateDelegate.cpp
--- a/Engine/Source/Editor/BlueprintGraph/Private/K2Node_CreateDelegate.cpp
+++ b/Engine/Source/Editor/BlueprintGraph/Private/K2Node_CreateDelegate.cpp
@@ -20,7 +20,6 @@
 #include "Kismet2/CompilerResultsLog.h"
 #include "Misc/AssertionMacros.h"
 #include "Templates/Casts.h"
-#include "Templates/ChooseClass.h"
 #include "Templates/SubclassOf.h"
 #include "UObject/Class.h"
 #include "UObject/Object.h"
diff --git a/Engine/Source/Editor/BlueprintGraph/Private/K2Node_DynamicCast.cpp b/Engine/Source/Editor/BlueprintGraph/Private/K2Node_DynamicCast.cpp
--- a/Engine/Source/Editor/BlueprintGraph/Private/K2Node_DynamicCast.cpp
+++ b/Engine/Source/Editor/BlueprintGraph/Private/K2Node_DynamicCast.cpp
@@ -25,7 +25,6 @@
 #include "ScopedTransaction.h"
 #include "Styling/AppStyle.h"
 #include "Templates/Casts.h"
-#include "Templates/ChooseClass.h"
 #include "Templates/SharedPointer.h"
 #include "ToolMenu.h"
 #include "ToolMenuSection.h"
diff --git a/Engine/Source/Editor/BlueprintGraph/Private/K2Node_Event.cpp b/Engine/Source/Editor/BlueprintGraph/Private/K2Node_Event.cpp
--- a/Engine/Source/Editor/BlueprintGraph/Private/K2Node_Event.cpp
+++ b/Engine/Source/Editor/BlueprintGraph/Private/K2Node_Event.cpp
@@ -31,7 +31,6 @@
 #include "Serialization/Archive.h"
 #include "Styling/AppStyle.h"
 #include "Templates/Casts.h"
-#include "Templates/ChooseClass.h"
 #include "Templates/UnrealTemplate.h"
 #include "Trace/Detail/Channel.h"
 #include "UObject/BlueprintsObjectVersion.h"
diff --git a/Engine/Source/Editor/BlueprintGraph/Private/K2Node_GenericCreateObject.cpp b/Engine/Source/Editor/BlueprintGraph/Private/K2Node_GenericCreateObject.cpp
--- a/Engine/Source/Editor/BlueprintGraph/Private/K2Node_GenericCreateObject.cpp
+++ b/Engine/Source/Editor/BlueprintGraph/Private/K2Node_GenericCreateObject.cpp
@@ -22,7 +22,6 @@
 #include "KismetCompiler.h"
 #include "KismetCompilerMisc.h"
 #include "Misc/AssertionMacros.h"
-#include "Templates/ChooseClass.h"
 #include "Templates/SubclassOf.h"
 #include "UObject/Class.h"
 #include "UObject/NameTypes.h"
diff --git a/Engine/Source/Editor/BlueprintGraph/Private/K2Node_GetDataTableRow.cpp b/Engine/Source/Editor/BlueprintGraph/Private/K2Node_GetDataTableRow.cpp
--- a/Engine/Source/Editor/BlueprintGraph/Private/K2Node_GetDataTableRow.cpp
+++ b/Engine/Source/Editor/BlueprintGraph/Private/K2Node_GetDataTableRow.cpp
@@ -27,7 +27,6 @@
 #include "Misc/AssertionMacros.h"
 #include "Styling/AppStyle.h"
 #include "Templates/Casts.h"
-#include "Templates/ChooseClass.h"
 #include "UObject/Class.h"
 #include "UObject/NameTypes.h"
 #include "UObject/Object.h"
diff --git a/Engine/Source/Editor/BlueprintGraph/Private/K2Node_GetSubsystem.cpp b/Engine/Source/Editor/BlueprintGraph/Private/K2Node_GetSubsystem.cpp
--- a/Engine/Source/Editor/BlueprintGraph/Private/K2Node_GetSubsystem.cpp
+++ b/Engine/Source/Editor/BlueprintGraph/Private/K2Node_GetSubsystem.cpp
@@ -35,7 +35,6 @@
 #include "Subsystems/SubsystemBlueprintLibrary.h"
 #include "Subsystems/WorldSubsystem.h"
 #include "Templates/Casts.h"
-#include "Templates/ChooseClass.h"
 #include "Templates/UnrealTemplate.h"
 #include "UObject/Class.h"
 #include "UObject/NameTypes.h"
diff --git a/Engine/Source/Editor/BlueprintGraph/Private/K2Node_Literal.cpp b/Engine/Source/Editor/BlueprintGraph/Private/K2Node_Literal.cpp
--- a/Engine/Source/Editor/BlueprintGraph/Private/K2Node_Literal.cpp
+++ b/Engine/Source/Editor/BlueprintGraph/Private/K2Node_Literal.cpp
@@ -27,7 +27,6 @@
 #include "Misc/AssertionMacros.h"
 #include "Styling/SlateIconFinder.h"
 #include "Templates/Casts.h"
-#include "Templates/ChooseClass.h"
 #include "UObject/Class.h"
 #include "UObject/NameTypes.h"
 #include "UObject/Object.h"
diff --git a/Engine/Source/Editor/BlueprintGraph/Private/K2Node_MacroInstance.cpp b/Engine/Source/Editor/BlueprintGraph/Private/K2Node_MacroInstance.cpp
--- a/Engine/Source/Editor/BlueprintGraph/Private/K2Node_MacroInstance.cpp
+++ b/Engine/Source/Editor/BlueprintGraph/Private/K2Node_MacroInstance.cpp
@@ -21,7 +21,6 @@
 #include "Settings/EditorStyleSettings.h"
 #include "Styling/AppStyle.h"
 #include "Templates/Casts.h"
-#include "Templates/ChooseClass.h"
 #include "Templates/SubclassOf.h"
 #include "Templates/UnrealTemplate.h"
 #include "ToolMenu.h"
diff --git a/Engine/Source/Editor/BlueprintGraph/Private/K2Node_SpawnActor.cpp b/Engine/Source/Editor/BlueprintGraph/Private/K2Node_SpawnActor.cpp
--- a/Engine/Source/Editor/BlueprintGraph/Private/K2Node_SpawnActor.cpp
+++ b/Engine/Source/Editor/BlueprintGraph/Private/K2Node_SpawnActor.cpp
@@ -23,7 +23,6 @@
 #include "Misc/AssertionMacros.h"
 #include "Styling/AppStyle.h"
 #include "Templates/Casts.h"
-#include "Templates/ChooseClass.h"
 #include "Templates/SubclassOf.h"
 #include "Templates/UnrealTemplate.h"
 #include "UObject/Class.h"
diff --git a/Engine/Source/Editor/BlueprintGraph/Private/K2Node_SpawnActorFromClass.cpp b/Engine/Source/Editor/BlueprintGraph/Private/K2Node_SpawnActorFromClass.cpp
--- a/Engine/Source/Editor/BlueprintGraph/Private/K2Node_SpawnActorFromClass.cpp
+++ b/Engine/Source/Editor/BlueprintGraph/Private/K2Node_SpawnActorFromClass.cpp
@@ -28,7 +28,6 @@
 #include "Misc/AssertionMacros.h"
 #include "Styling/AppStyle.h"
 #include "Templates/Casts.h"
-#include "Templates/ChooseClass.h"
 #include "Templates/SubclassOf.h"
 #include "Templates/UnrealTemplate.h"
 #include "UObject/Class.h"
diff --git a/Engine/Source/Editor/BlueprintGraph/Private/K2Node_Switch.cpp b/Engine/Source/Editor/BlueprintGraph/Private/K2Node_Switch.cpp
--- a/Engine/Source/Editor/BlueprintGraph/Private/K2Node_Switch.cpp
+++ b/Engine/Source/Editor/BlueprintGraph/Private/K2Node_Switch.cpp
@@ -23,7 +23,6 @@
 #include "KismetCompilerMisc.h"
 #include "Styling/AppStyle.h"
 #include "Templates/Casts.h"
-#include "Templates/ChooseClass.h"
 #include "UObject/Class.h"
 #include "UObject/Object.h"
 #include "UObject/ObjectPtr.h"
diff --git a/Engine/Source/Editor/BlueprintGraph/Private/K2Node_Variable.cpp b/Engine/Source/Editor/BlueprintGraph/Private/K2Node_Variable.cpp
--- a/Engine/Source/Editor/BlueprintGraph/Private/K2Node_Variable.cpp
+++ b/Engine/Source/Editor/BlueprintGraph/Private/K2Node_Variable.cpp
@@ -33,7 +33,6 @@
 #include "Styling/AppStyle.h"
 #include "Styling/SlateIconFinder.h"
 #include "Templates/Casts.h"
-#include "Templates/ChooseClass.h"
 #include "Templates/UnrealTemplate.h"
 #include "UObject/Class.h"
 #include "UObject/Object.h"
diff --git a/Engine/Source/Editor/Blutility/Private/BlutilityContentBrowserExtensions.cpp b/Engine/Source/Editor/Blutility/Private/BlutilityContentBrowserExtensions.cpp
--- a/Engine/Source/Editor/Blutility/Private/BlutilityContentBrowserExtensions.cpp
+++ b/Engine/Source/Editor/Blutility/Private/BlutilityContentBrowserExtensions.cpp
@@ -24,7 +24,6 @@
 #include "Misc/NamePermissionList.h"
 #include "Modules/ModuleManager.h"
 #include "Templates/Casts.h"
-#include "Templates/ChooseClass.h"
 #include "Templates/SharedPointer.h"
 #include "Templates/SubclassOf.h"
 #include "Templates/UnrealTemplate.h"
diff --git a/Engine/Source/Editor/Blutility/Private/BlutilityModule.cpp b/Engine/Source/Editor/Blutility/Private/BlutilityModule.cpp
--- a/Engine/Source/Editor/Blutility/Private/BlutilityModule.cpp
+++ b/Engine/Source/Editor/Blutility/Private/BlutilityModule.cpp
@@ -43,7 +43,6 @@
 #include "Modules/ModuleManager.h"
 #include "Styling/AppStyle.h"
 #include "Templates/Casts.h"
-#include "Templates/ChooseClass.h"
 #include "Templates/SharedPointer.h"
 #include "Templates/SubclassOf.h"
 #include "Textures/SlateIcon.h"
diff --git a/Engine/Source/Editor/Blutility/Private/EditorUtilityBlueprintFactory.cpp b/Engine/Source/Editor/Blutility/Private/EditorUtilityBlueprintFactory.cpp
--- a/Engine/Source/Editor/Blutility/Private/EditorUtilityBlueprintFactory.cpp
+++ b/Engine/Source/Editor/Blutility/Private/EditorUtilityBlueprintFactory.cpp
@@ -23,7 +23,6 @@
 #include "Misc/AssertionMacros.h"
 #include "Misc/MessageDialog.h"
 #include "Modules/ModuleManager.h"
-#include "Templates/ChooseClass.h"
 #include "UObject/Class.h"
 #include "Widgets/SWindow.h"
 
diff --git a/Engine/Source/Editor/Blutility/Private/EditorUtilitySubsystem.cpp b/Engine/Source/Editor/Blutility/Private/EditorUtilitySubsystem.cpp
--- a/Engine/Source/Editor/Blutility/Private/EditorUtilitySubsystem.cpp
+++ b/Engine/Source/Editor/Blutility/Private/EditorUtilitySubsystem.cpp
@@ -30,7 +30,6 @@
 #include "Misc/PackageName.h"
 #include "Modules/ModuleManager.h"
 #include "Templates/Casts.h"
-#include "Templates/ChooseClass.h"
 #include "Templates/SubclassOf.h"
 #include "Templates/Tuple.h"
 #include "Trace/Detail/Channel.h"
diff --git a/Engine/Source/Editor/Blutility/Private/EditorUtilityWidgetBlueprint.cpp b/Engine/Source/Editor/Blutility/Private/EditorUtilityWidgetBlueprint.cpp
--- a/Engine/Source/Editor/Blutility/Private/EditorUtilityWidgetBlueprint.cpp
+++ b/Engine/Source/Editor/Blutility/Private/EditorUtilityWidgetBlueprint.cpp
@@ -17,7 +17,6 @@
 #include "Modules/ModuleManager.h"
 #include "SlotBase.h"
 #include "Templates/Casts.h"
-#include "Templates/ChooseClass.h"
 #include "Templates/SubclassOf.h"
 #include "Types/SlateEnums.h"
 #include "UObject/Package.h"
diff --git a/Engine/Source/Editor/Blutility/Private/EditorUtilityWidgetBlueprintFactory.cpp b/Engine/Source/Editor/Blutility/Private/EditorUtilityWidgetBlueprintFactory.cpp
--- a/Engine/Source/Editor/Blutility/Private/EditorUtilityWidgetBlueprintFactory.cpp
+++ b/Engine/Source/Editor/Blutility/Private/EditorUtilityWidgetBlueprintFactory.cpp
@@ -28,7 +28,6 @@
 #include "Misc/MessageDialog.h"
 #include "Modules/ModuleManager.h"
 #include "Templates/Casts.h"
-#include "Templates/ChooseClass.h"
 #include "Templates/SharedPointer.h"
 #include "UObject/Class.h"
 #include "UObject/Object.h"
diff --git a/Engine/Source/Editor/ClothPainter/Private/ClothingAssetExporter.cpp b/Engine/Source/Editor/ClothPainter/Private/ClothingAssetExporter.cpp
--- a/Engine/Source/Editor/ClothPainter/Private/ClothingAssetExporter.cpp
+++ b/Engine/Source/Editor/ClothPainter/Private/ClothingAssetExporter.cpp
@@ -18,7 +18,6 @@
 #include "Misc/MessageDialog.h"
 #include "Misc/PackageName.h"
 #include "Modules/ModuleManager.h"
-#include "Templates/ChooseClass.h"
 #include "Templates/Function.h"
 #include "UObject/Class.h"
 #include "UObject/Package.h"
diff --git a/Engine/Source/Editor/ClothPainter/Private/ClothingAssetListCommands.cpp b/Engine/Source/Editor/ClothPainter/Private/ClothingAssetListCommands.cpp
--- a/Engine/Source/Editor/ClothPainter/Private/ClothingAssetListCommands.cpp
+++ b/Engine/Source/Editor/ClothPainter/Private/ClothingAssetListCommands.cpp
@@ -8,7 +8,6 @@
 #include "Framework/Commands/InputChord.h"
 #include "Framework/Commands/UICommandInfo.h"
 #include "Internationalization/Text.h"
-#include "Templates/ChooseClass.h"
 #include "Templates/SubclassOf.h"
 #include "Textures/SlateIcon.h"
 #include "UObject/Class.h"
diff --git a/Engine/Source/Editor/CurveEditor/Private/CurveEditor.cpp b/Engine/Source/Editor/CurveEditor/Private/CurveEditor.cpp
--- a/Engine/Source/Editor/CurveEditor/Private/CurveEditor.cpp
+++ b/Engine/Source/Editor/CurveEditor/Private/CurveEditor.cpp
@@ -42,7 +42,6 @@
 #include "SCurveEditorView.h"
 #include "ScopedTransaction.h"
 #include "Templates/Casts.h"
-#include "Templates/ChooseClass.h"
 #include "Templates/Tuple.h"
 #include "Templates/UnrealTemplate.h"
 #include "Trace/Detail/Channel.h"
diff --git a/Engine/Source/Editor/CurveEditor/Private/SCurveEditorPanel.cpp b/Engine/Source/Editor/CurveEditor/Private/SCurveEditorPanel.cpp
--- a/Engine/Source/Editor/CurveEditor/Private/SCurveEditorPanel.cpp
+++ b/Engine/Source/Editor/CurveEditor/Private/SCurveEditorPanel.cpp
@@ -67,7 +67,6 @@
 #include "Styling/ISlateStyle.h"
 #include "Styling/SlateColor.h"
 #include "Templates/Casts.h"
-#include "Templates/ChooseClass.h"
 #include "Templates/Tuple.h"
 #include "Templates/TypeHash.h"
 #include "Types/SlateEnums.h"
diff --git a/Engine/Source/Editor/CurveEditor/Private/Views/SCurveEditorViewNormalized.cpp b/Engine/Source/Editor/CurveEditor/Private/Views/SCurveEditorViewNormalized.cpp
--- a/Engine/Source/Editor/CurveEditor/Private/Views/SCurveEditorViewNormalized.cpp
+++ b/Engine/Source/Editor/CurveEditor/Private/Views/SCurveEditorViewNormalized.cpp
@@ -23,7 +23,6 @@
 #include "SCurveEditorView.h"
 #include "Styling/CoreStyle.h"
 #include "Styling/ISlateStyle.h"
-#include "Templates/ChooseClass.h"
 #include "Templates/Tuple.h"
 #include "Templates/UniquePtr.h"
 #include "Types/SlateEnums.h"
diff --git a/Engine/Source/Editor/CurveEditor/Private/Views/SCurveEditorViewStacked.cpp b/Engine/Source/Editor/CurveEditor/Private/Views/SCurveEditorViewStacked.cpp
--- a/Engine/Source/Editor/CurveEditor/Private/Views/SCurveEditorViewStacked.cpp
+++ b/Engine/Source/Editor/CurveEditor/Private/Views/SCurveEditorViewStacked.cpp
@@ -29,7 +29,6 @@
 #include "SCurveEditorView.h"
 #include "Styling/AppStyle.h"
 #include "Styling/ISlateStyle.h"
-#include "Templates/ChooseClass.h"
 #include "Templates/Tuple.h"
 #include "Templates/UniquePtr.h"
 
diff --git a/Engine/Source/Editor/DataLayerEditor/Private/DataLayer/DataLayerMode.cpp b/Engine/Source/Editor/DataLayerEditor/Private/DataLayer/DataLayerMode.cpp
--- a/Engine/Source/Editor/DataLayerEditor/Private/DataLayer/DataLayerMode.cpp
+++ b/Engine/Source/Editor/DataLayerEditor/Private/DataLayer/DataLayerMode.cpp
@@ -65,7 +65,6 @@
 #include "Selection.h"
 #include "SlotBase.h"
 #include "Templates/Casts.h"
-#include "Templates/ChooseClass.h"
 #include "Templates/Tuple.h"
 #include "Templates/TypeHash.h"
 #include "Templates/UnrealTemplate.h"
diff --git a/Engine/Source/Editor/DetailCustomizations/Private/AnimMontageSegmentDetails.cpp b/Engine/Source/Editor/DetailCustomizations/Private/AnimMontageSegmentDetails.cpp
--- a/Engine/Source/Editor/DetailCustomizations/Private/AnimMontageSegmentDetails.cpp
+++ b/Engine/Source/Editor/DetailCustomizations/Private/AnimMontageSegmentDetails.cpp
@@ -47,7 +47,6 @@
 #include "Slate/SceneViewport.h"
 #include "SlotBase.h"
 #include "Templates/Casts.h"
-#include "Templates/ChooseClass.h"
 #include "Templates/SubclassOf.h"
 #include "UObject/Field.h"
 #include "UObject/UObjectGlobals.h"
diff --git a/Engine/Source/Editor/DetailCustomizations/Private/AnimSequenceDetails.cpp b/Engine/Source/Editor/DetailCustomizations/Private/AnimSequenceDetails.cpp
--- a/Engine/Source/Editor/DetailCustomizations/Private/AnimSequenceDetails.cpp
+++ b/Engine/Source/Editor/DetailCustomizations/Private/AnimSequenceDetails.cpp
@@ -43,7 +43,6 @@
 #include "SlateOptMacros.h"
 #include "SlotBase.h"
 #include "Templates/Casts.h"
-#include "Templates/ChooseClass.h"
 #include "Templates/SubclassOf.h"
 #include "UObject/Object.h"
 #include "UObject/ObjectPtr.h"
diff --git a/Engine/Source/Editor/DetailCustomizations/Private/AnimTransitionNodeDetails.cpp b/Engine/Source/Editor/DetailCustomizations/Private/AnimTransitionNodeDetails.cpp
--- a/Engine/Source/Editor/DetailCustomizations/Private/AnimTransitionNodeDetails.cpp
+++ b/Engine/Source/Editor/DetailCustomizations/Private/AnimTransitionNodeDetails.cpp
@@ -42,7 +42,6 @@
 #include "SlotBase.h"
 #include "Styling/AppStyle.h"
 #include "Templates/Casts.h"
-#include "Templates/ChooseClass.h"
 #include "Templates/SubclassOf.h"
 #include "Textures/SlateIcon.h"
 #include "UObject/NameTypes.h"
diff --git a/Engine/Source/Editor/DetailCustomizations/Private/AttenuationSettingsCustomizations.cpp b/Engine/Source/Editor/DetailCustomizations/Private/AttenuationSettingsCustomizations.cpp
--- a/Engine/Source/Editor/DetailCustomizations/Private/AttenuationSettingsCustomizations.cpp
+++ b/Engine/Source/Editor/DetailCustomizations/Private/AttenuationSettingsCustomizations.cpp
@@ -21,7 +21,6 @@
 #include "PropertyRestriction.h"
 #include "Sound/AudioSettings.h"
 #include "Sound/SoundAttenuation.h"
-#include "Templates/ChooseClass.h"
 #include "UObject/Class.h"
 #include "UObject/NameTypes.h"
 #include "UObject/UObjectGlobals.h"
diff --git a/Engine/Source/Editor/DetailCustomizations/Private/CameraFilmbackSettingsCustomization.cpp b/Engine/Source/Editor/DetailCustomizations/Private/CameraFilmbackSettingsCustomization.cpp
--- a/Engine/Source/Editor/DetailCustomizations/Private/CameraFilmbackSettingsCustomization.cpp
+++ b/Engine/Source/Editor/DetailCustomizations/Private/CameraFilmbackSettingsCustomization.cpp
@@ -17,7 +17,6 @@
 #include "PropertyEditorModule.h"
 #include "PropertyHandle.h"
 #include "ScopedTransaction.h"
-#include "Templates/ChooseClass.h"
 #include "UObject/NameTypes.h"
 #include "UObject/UnrealType.h"
 #include "Widgets/DeclarativeSyntaxSupport.h"
diff --git a/Engine/Source/Editor/DetailCustomizations/Private/CameraFocusSettingsCustomization.cpp b/Engine/Source/Editor/DetailCustomizations/Private/CameraFocusSettingsCustomization.cpp
--- a/Engine/Source/Editor/DetailCustomizations/Private/CameraFocusSettingsCustomization.cpp
+++ b/Engine/Source/Editor/DetailCustomizations/Private/CameraFocusSettingsCustomization.cpp
@@ -24,7 +24,6 @@
 #include "ScopedTransaction.h"
 #include "SlotBase.h"
 #include "Templates/Casts.h"
-#include "Templates/ChooseClass.h"
 #include "Types/SlateEnums.h"
 #include "UObject/NameTypes.h"
 #include "UObject/UnrealType.h"
diff --git a/Engine/Source/Editor/DetailCustomizations/Private/CameraLensSettingsCustomization.cpp b/Engine/Source/Editor/DetailCustomizations/Private/CameraLensSettingsCustomization.cpp
--- a/Engine/Source/Editor/DetailCustomizations/Private/CameraLensSettingsCustomization.cpp
+++ b/Engine/Source/Editor/DetailCustomizations/Private/CameraLensSettingsCustomization.cpp
@@ -17,7 +17,6 @@
 #include "PropertyEditorModule.h"
 #include "PropertyHandle.h"
 #include "ScopedTransaction.h"
-#include "Templates/ChooseClass.h"
 #include "UObject/NameTypes.h"
 #include "UObject/UnrealType.h"
 #include "Widgets/DeclarativeSyntaxSupport.h"
diff --git a/Engine/Source/Editor/DetailCustomizations/Private/DetailCustomizations.cpp b/Engine/Source/Editor/DetailCustomizations/Private/DetailCustomizations.cpp
--- a/Engine/Source/Editor/DetailCustomizations/Private/DetailCustomizations.cpp
+++ b/Engine/Source/Editor/DetailCustomizations/Private/DetailCustomizations.cpp
@@ -150,7 +150,6 @@
 #include "StaticMeshComponentDetails.h"
 #include "SupportedRangeTypes.h"	// StructsSupportingRangeVisibility
 #include "TemplateStringStructCustomization.h"
-#include "Templates/ChooseClass.h"
 #include "Templates/SharedPointer.h"
 #include "TextCustomization.h"
 #include "TimecodeDetailsCustomization.h"
diff --git a/Engine/Source/Editor/DetailCustomizations/Private/HierarchicalSimplificationCustomizations.cpp b/Engine/Source/Editor/DetailCustomizations/Private/HierarchicalSimplificationCustomizations.cpp
--- a/Engine/Source/Editor/DetailCustomizations/Private/HierarchicalSimplificationCustomizations.cpp
+++ b/Engine/Source/Editor/DetailCustomizations/Private/HierarchicalSimplificationCustomizations.cpp
@@ -19,7 +19,6 @@
 #include "Modules/ModuleManager.h"
 #include "PropertyHandle.h"
 #include "PropertyRestriction.h"
-#include "Templates/ChooseClass.h"
 #include "UObject/Class.h"
 #include "UObject/NameTypes.h"
 #include "UObject/UnrealNames.h"
diff --git a/Engine/Source/Editor/DetailCustomizations/Private/MaterialProxySettingsCustomizations.cpp b/Engine/Source/Editor/DetailCustomizations/Private/MaterialProxySettingsCustomizations.cpp
--- a/Engine/Source/Editor/DetailCustomizations/Private/MaterialProxySettingsCustomizations.cpp
+++ b/Engine/Source/Editor/DetailCustomizations/Private/MaterialProxySettingsCustomizations.cpp
@@ -23,7 +23,6 @@
 #include "PropertyHandle.h"
 #include "PropertyRestriction.h"
 #include "RHI.h"
-#include "Templates/ChooseClass.h"
 #include "UObject/Class.h"
 #include "UObject/NameTypes.h"
 #include "UObject/ReflectedTypeAccessors.h"
diff --git a/Engine/Source/Editor/DetailCustomizations/Private/MeshProxySettingsCustomizations.cpp b/Engine/Source/Editor/DetailCustomizations/Private/MeshProxySettingsCustomizations.cpp
--- a/Engine/Source/Editor/DetailCustomizations/Private/MeshProxySettingsCustomizations.cpp
+++ b/Engine/Source/Editor/DetailCustomizations/Private/MeshProxySettingsCustomizations.cpp
@@ -17,7 +17,6 @@
 #include "Misc/Attribute.h"
 #include "Modules/ModuleManager.h"
 #include "PropertyHandle.h"
-#include "Templates/ChooseClass.h"
 #include "UObject/NameTypes.h"
 #include "UObject/UnrealNames.h"
 #include "UObject/UnrealType.h"
diff --git a/Engine/Source/Editor/DetailCustomizations/Private/NavAgentSelectorCustomization.cpp b/Engine/Source/Editor/DetailCustomizations/Private/NavAgentSelectorCustomization.cpp
--- a/Engine/Source/Editor/DetailCustomizations/Private/NavAgentSelectorCustomization.cpp
+++ b/Engine/Source/Editor/DetailCustomizations/Private/NavAgentSelectorCustomization.cpp
@@ -17,7 +17,6 @@
 #include "NavigationSystem.h"
 #include "PropertyEditorModule.h"
 #include "PropertyHandle.h"
-#include "Templates/ChooseClass.h"
 #include "Templates/SubclassOf.h"
 #include "Types/SlateEnums.h"
 #include "UObject/Class.h"
diff --git a/Engine/Source/Editor/DetailCustomizations/Private/PoseAssetDetails.cpp b/Engine/Source/Editor/DetailCustomizations/Private/PoseAssetDetails.cpp
--- a/Engine/Source/Editor/DetailCustomizations/Private/PoseAssetDetails.cpp
+++ b/Engine/Source/Editor/DetailCustomizations/Private/PoseAssetDetails.cpp
@@ -26,7 +26,6 @@
 #include "ScopedTransaction.h"
 #include "SlotBase.h"
 #include "Templates/Casts.h"
-#include "Templates/ChooseClass.h"
 #include "UObject/Object.h"
 #include "UObject/ObjectPtr.h"
 #include "UObject/SoftObjectPtr.h"
diff --git a/Engine/Source/Editor/DetailCustomizations/Private/SkeletalMeshReductionSettingsDetails.cpp b/Engine/Source/Editor/DetailCustomizations/Private/SkeletalMeshReductionSettingsDetails.cpp
--- a/Engine/Source/Editor/DetailCustomizations/Private/SkeletalMeshReductionSettingsDetails.cpp
+++ b/Engine/Source/Editor/DetailCustomizations/Private/SkeletalMeshReductionSettingsDetails.cpp
@@ -26,7 +26,6 @@
 #include "SkeletalMeshReductionSettings.h"
 #include "SkeletalRenderPublic.h"
 #include "Templates/Casts.h"
-#include "Templates/ChooseClass.h"
 #include "UObject/NameTypes.h"
 #include "UObject/UnrealType.h"
 #include "Widgets/DeclarativeSyntaxSupport.h"
diff --git a/Engine/Source/Editor/DetailCustomizations/Public/Customizations/CurveTableCustomization.h b/Engine/Source/Editor/DetailCustomizations/Public/Customizations/CurveTableCustomization.h
--- a/Engine/Source/Editor/DetailCustomizations/Public/Customizations/CurveTableCustomization.h
+++ b/Engine/Source/Editor/DetailCustomizations/Public/Customizations/CurveTableCustomization.h
@@ -27,7 +27,6 @@
 #include "Misc/Optional.h"
 #include "PropertyEditorModule.h"
 #include "PropertyHandle.h"
-#include "Templates/ChooseClass.h"
 #include "Templates/SharedPointer.h"
 #include "Templates/TypeHash.h"
 #include "Templates/UnrealTemplate.h"
diff --git a/Engine/Source/Editor/EditorWidgets/Private/SMetaDataView.cpp b/Engine/Source/Editor/EditorWidgets/Private/SMetaDataView.cpp
--- a/Engine/Source/Editor/EditorWidgets/Private/SMetaDataView.cpp
+++ b/Engine/Source/Editor/EditorWidgets/Private/SMetaDataView.cpp
@@ -19,7 +19,6 @@
 #include "Styling/CoreStyle.h"
 #include "Styling/ISlateStyle.h"
 #include "Styling/SlateTypes.h"
-#include "Templates/ChooseClass.h"
 #include "Templates/UnrealTemplate.h"
 #include "UObject/NameTypes.h"
 #include "Widgets/SBoxPanel.h"
diff --git a/Engine/Source/Editor/FoliageEdit/Private/FoliageType_ActorThumbnailRenderer.cpp b/Engine/Source/Editor/FoliageEdit/Private/FoliageType_ActorThumbnailRenderer.cpp
--- a/Engine/Source/Editor/FoliageEdit/Private/FoliageType_ActorThumbnailRenderer.cpp
+++ b/Engine/Source/Editor/FoliageEdit/Private/FoliageType_ActorThumbnailRenderer.cpp
@@ -4,7 +4,6 @@
 
 #include "FoliageType_Actor.h"
 #include "Templates/Casts.h"
-#include "Templates/ChooseClass.h"
 #include "Templates/SubclassOf.h"
 #include "UObject/Object.h"
 
diff --git a/Engine/Source/Editor/GraphEditor/Private/DragConnection.cpp b/Engine/Source/Editor/GraphEditor/Private/DragConnection.cpp
--- a/Engine/Source/Editor/GraphEditor/Private/DragConnection.cpp
+++ b/Engine/Source/Editor/GraphEditor/Private/DragConnection.cpp
@@ -23,7 +23,6 @@
 #include "ScopedTransaction.h"
 #include "SlotBase.h"
 #include "Styling/AppStyle.h"
-#include "Templates/ChooseClass.h"
 #include "Types/SlateEnums.h"
 #include "UObject/NameTypes.h"
 #include "Widgets/DeclarativeSyntaxSupport.h"
diff --git a/Engine/Source/Editor/GraphEditor/Private/SNodePanel.cpp b/Engine/Source/Editor/GraphEditor/Private/SNodePanel.cpp
--- a/Engine/Source/Editor/GraphEditor/Private/SNodePanel.cpp
+++ b/Engine/Source/Editor/GraphEditor/Private/SNodePanel.cpp
@@ -25,7 +25,6 @@
 #include "Settings/EditorStyleSettings.h"
 #include "Settings/LevelEditorViewportSettings.h"
 #include "Styling/SlateBrush.h"
-#include "Templates/ChooseClass.h"
 #include "Templates/RemoveReference.h"
 #include "Types/WidgetActiveTimerDelegate.h"
 #include "UObject/Object.h"
diff --git a/Engine/Source/Editor/Kismet/Private/BPDelegateDragDropAction.cpp b/Engine/Source/Editor/Kismet/Private/BPDelegateDragDropAction.cpp
--- a/Engine/Source/Editor/Kismet/Private/BPDelegateDragDropAction.cpp
+++ b/Engine/Source/Editor/Kismet/Private/BPDelegateDragDropAction.cpp
@@ -24,7 +24,6 @@
 #include "Layout/WidgetPath.h"
 #include "ScopedTransaction.h"
 #include "Templates/Casts.h"
-#include "Templates/ChooseClass.h"
 #include "Templates/SubclassOf.h"
 #include "Textures/SlateIcon.h"
 #include "UObject/Field.h"
diff --git a/Engine/Source/Editor/Kismet/Private/BlueprintActionMenuItem.cpp b/Engine/Source/Editor/Kismet/Private/BlueprintActionMenuItem.cpp
--- a/Engine/Source/Editor/Kismet/Private/BlueprintActionMenuItem.cpp
+++ b/Engine/Source/Editor/Kismet/Private/BlueprintActionMenuItem.cpp
@@ -21,7 +21,6 @@
 #include "SNodePanel.h"
 #include "ScopedTransaction.h"
 #include "Templates/Casts.h"
-#include "Templates/ChooseClass.h"
 #include "Templates/SharedPointer.h"
 #include "Templates/UnrealTemplate.h"
 #include "Textures/SlateIcon.h"
diff --git a/Engine/Source/Editor/Kismet/Private/BlueprintDetailsCustomization.cpp b/Engine/Source/Editor/Kismet/Private/BlueprintDetailsCustomization.cpp
--- a/Engine/Source/Editor/Kismet/Private/BlueprintDetailsCustomization.cpp
+++ b/Engine/Source/Editor/Kismet/Private/BlueprintDetailsCustomization.cpp
@@ -117,7 +117,6 @@
 #include "SubobjectDataSubsystem.h"
 #include "SupportedRangeTypes.h"	// StructsSupportingRangeVisibility
 #include "Templates/Casts.h"
-#include "Templates/ChooseClass.h"
 #include "Templates/SubclassOf.h"
 #include "Textures/SlateIcon.h"
 #include "Tools/LegacyEdModeWidgetHelpers.h"
diff --git a/Engine/Source/Editor/Kismet/Private/BlueprintEditorCommands.cpp b/Engine/Source/Editor/Kismet/Private/BlueprintEditorCommands.cpp
--- a/Engine/Source/Editor/Kismet/Private/BlueprintEditorCommands.cpp
+++ b/Engine/Source/Editor/Kismet/Private/BlueprintEditorCommands.cpp
@@ -36,7 +36,6 @@
 #include "Misc/Parse.h"
 #include "Selection.h"
 #include "Templates/Casts.h"
-#include "Templates/ChooseClass.h"
 #include "Templates/SubclassOf.h"
 #include "Textures/SlateIcon.h"
 #include "Trace/Detail/Channel.h"
diff --git a/Engine/Source/Editor/Kismet/Private/BlueprintEditorModule.cpp b/Engine/Source/Editor/Kismet/Private/BlueprintEditorModule.cpp
--- a/Engine/Source/Editor/Kismet/Private/BlueprintEditorModule.cpp
+++ b/Engine/Source/Editor/Kismet/Private/BlueprintEditorModule.cpp
@@ -41,7 +41,6 @@
 #include "SPinValueInspector.h"
 #include "Serialization/ArchiveReplaceObjectRef.h"
 #include "Templates/Casts.h"
-#include "Templates/ChooseClass.h"
 #include "Templates/TypeHash.h"
 #include "UObject/Class.h"
 #include "UObject/Field.h"
diff --git a/Engine/Source/Editor/Kismet/Private/BlueprintEditorTabFactories.cpp b/Engine/Source/Editor/Kismet/Private/BlueprintEditorTabFactories.cpp
--- a/Engine/Source/Editor/Kismet/Private/BlueprintEditorTabFactories.cpp
+++ b/Engine/Source/Editor/Kismet/Private/BlueprintEditorTabFactories.cpp
@@ -32,7 +32,6 @@
 #include "Styling/SlateColor.h"
 #include "Subsystems/AssetEditorSubsystem.h"
 #include "Templates/Casts.h"
-#include "Templates/ChooseClass.h"
 #include "Templates/SubclassOf.h"
 #include "Textures/SlateIcon.h"
 #include "Types/SlateEnums.h"
diff --git a/Engine/Source/Editor/Kismet/Private/BlueprintMemberReferenceCustomization.cpp b/Engine/Source/Editor/Kismet/Private/BlueprintMemberReferenceCustomization.cpp
--- a/Engine/Source/Editor/Kismet/Private/BlueprintMemberReferenceCustomization.cpp
+++ b/Engine/Source/Editor/Kismet/Private/BlueprintMemberReferenceCustomization.cpp
@@ -33,7 +33,6 @@
 #include "Styling/SlateColor.h"
 #include "Subsystems/AssetEditorSubsystem.h"
 #include "Templates/Casts.h"
-#include "Templates/ChooseClass.h"
 #include "Templates/SubclassOf.h"
 #include "Types/SlateStructs.h"
 #include "UObject/Class.h"
diff --git a/Engine/Source/Editor/Kismet/Private/FindInBlueprints.cpp b/Engine/Source/Editor/Kismet/Private/FindInBlueprints.cpp
--- a/Engine/Source/Editor/Kismet/Private/FindInBlueprints.cpp
+++ b/Engine/Source/Editor/Kismet/Private/FindInBlueprints.cpp
@@ -60,7 +60,6 @@
 #include "Styling/SlateTypes.h"
 #include "Styling/StyleColors.h"
 #include "Templates/Casts.h"
-#include "Templates/ChooseClass.h"
 #include "Templates/SubclassOf.h"
 #include "Trace/Detail/Channel.h"
 #include "Types/SlateConstants.h"
diff --git a/Engine/Source/Editor/Kismet/Private/SBlueprintActionMenu.cpp b/Engine/Source/Editor/Kismet/Private/SBlueprintActionMenu.cpp
--- a/Engine/Source/Editor/Kismet/Private/SBlueprintActionMenu.cpp
+++ b/Engine/Source/Editor/Kismet/Private/SBlueprintActionMenu.cpp
@@ -40,7 +40,6 @@
 #include "Styling/AppStyle.h"
 #include "Styling/SlateColor.h"
 #include "Templates/Casts.h"
-#include "Templates/ChooseClass.h"
 #include "Templates/SubclassOf.h"
 #include "Templates/UnrealTemplate.h"
 #include "Types/SlateStructs.h"
diff --git a/Engine/Source/Editor/Kismet/Private/SBlueprintPalette.cpp b/Engine/Source/Editor/Kismet/Private/SBlueprintPalette.cpp
--- a/Engine/Source/Editor/Kismet/Private/SBlueprintPalette.cpp
+++ b/Engine/Source/Editor/Kismet/Private/SBlueprintPalette.cpp
@@ -85,7 +85,6 @@
 #include "Styling/SlateTypes.h"
 #include "Styling/StyleDefaults.h"
 #include "Templates/Casts.h"
-#include "Templates/ChooseClass.h"
 #include "Templates/SubclassOf.h"
 #include "Textures/SlateIcon.h"
 #include "TutorialMetaData.h"
diff --git a/Engine/Source/Editor/Kismet/Private/SCSDiff.cpp b/Engine/Source/Editor/Kismet/Private/SCSDiff.cpp
--- a/Engine/Source/Editor/Kismet/Private/SCSDiff.cpp
+++ b/Engine/Source/Editor/Kismet/Private/SCSDiff.cpp
@@ -23,7 +23,6 @@
 #include "SubobjectData.h"
 #include "SubobjectDataHandle.h"
 #include "Templates/Casts.h"
-#include "Templates/ChooseClass.h"
 #include "Templates/SubclassOf.h"
 #include "Types/SlateEnums.h"
 #include "UObject/Object.h"
diff --git a/Engine/Source/Editor/Kismet/Private/SFixupSelfContextDlg.cpp b/Engine/Source/Editor/Kismet/Private/SFixupSelfContextDlg.cpp
--- a/Engine/Source/Editor/Kismet/Private/SFixupSelfContextDlg.cpp
+++ b/Engine/Source/Editor/Kismet/Private/SFixupSelfContextDlg.cpp
@@ -30,7 +30,6 @@
 #include "SlotBase.h"
 #include "Styling/AppStyle.h"
 #include "Styling/SlateColor.h"
-#include "Templates/ChooseClass.h"
 #include "Templates/SubclassOf.h"
 #include "Types/SlateEnums.h"
 #include "Types/SlateStructs.h"
diff --git a/Engine/Source/Editor/Kismet/Private/SKismetInspector.cpp b/Engine/Source/Editor/Kismet/Private/SKismetInspector.cpp
--- a/Engine/Source/Editor/Kismet/Private/SKismetInspector.cpp
+++ b/Engine/Source/Editor/Kismet/Private/SKismetInspector.cpp
@@ -58,7 +58,6 @@
 #include "Styling/AppStyle.h"
 #include "Styling/ISlateStyle.h"
 #include "Templates/Casts.h"
-#include "Templates/ChooseClass.h"
 #include "Templates/SubclassOf.h"
 #include "Templates/UnrealTemplate.h"
 #include "Types/ISlateMetaData.h"
diff --git a/Engine/Source/Editor/Kismet/Private/SMyBlueprint.cpp b/Engine/Source/Editor/Kismet/Private/SMyBlueprint.cpp
--- a/Engine/Source/Editor/Kismet/Private/SMyBlueprint.cpp
+++ b/Engine/Source/Editor/Kismet/Private/SMyBlueprint.cpp
@@ -87,7 +87,6 @@
 #include "Styling/SlateTypes.h"
 #include "Subsystems/AssetEditorSubsystem.h"
 #include "Templates/Casts.h"
-#include "Templates/ChooseClass.h"
 #include "Templates/Less.h"
 #include "Templates/SubclassOf.h"
 #include "Templates/Tuple.h"
diff --git a/Engine/Source/Editor/Kismet/Private/SSCSEditor.cpp b/Engine/Source/Editor/Kismet/Private/SSCSEditor.cpp
--- a/Engine/Source/Editor/Kismet/Private/SSCSEditor.cpp
+++ b/Engine/Source/Editor/Kismet/Private/SSCSEditor.cpp
@@ -94,7 +94,6 @@
 #include "Styling/SlateIconFinder.h"
 #include "Subsystems/AssetEditorSubsystem.h"
 #include "Subsystems/PanelExtensionSubsystem.h"
-#include "Templates/ChooseClass.h"
 #include "Templates/Function.h"
 #include "Textures/SlateIcon.h"
 #include "ThumbnailRendering/ThumbnailManager.h"
diff --git a/Engine/Source/Editor/KismetCompiler/Private/EdGraphCompiler.cpp b/Engine/Source/Editor/KismetCompiler/Private/EdGraphCompiler.cpp
--- a/Engine/Source/Editor/KismetCompiler/Private/EdGraphCompiler.cpp
+++ b/Engine/Source/Editor/KismetCompiler/Private/EdGraphCompiler.cpp
@@ -19,7 +19,6 @@
 #include "Internationalization/Internationalization.h"
 #include "Internationalization/Text.h"
 #include "Kismet2/CompilerResultsLog.h"
-#include "Templates/ChooseClass.h"
 #include "Templates/SubclassOf.h"
 #include "UObject/ObjectPtr.h"
 
diff --git a/Engine/Source/Editor/MovieSceneTools/Private/Channels/SCurveEditorEventChannelView.cpp b/Engine/Source/Editor/MovieSceneTools/Private/Channels/SCurveEditorEventChannelView.cpp
--- a/Engine/Source/Editor/MovieSceneTools/Private/Channels/SCurveEditorEventChannelView.cpp
+++ b/Engine/Source/Editor/MovieSceneTools/Private/Channels/SCurveEditorEventChannelView.cpp
@@ -23,7 +23,6 @@
 #include "Styling/AppStyle.h"
 #include "Styling/CoreStyle.h"
 #include "Styling/ISlateStyle.h"
-#include "Templates/ChooseClass.h"
 
 class FPaintArgs;
 class FSlateRect;
diff --git a/Engine/Source/Editor/MovieSceneTools/Private/Channels/SCurveEditorKeyBarView.cpp b/Engine/Source/Editor/MovieSceneTools/Private/Channels/SCurveEditorKeyBarView.cpp
--- a/Engine/Source/Editor/MovieSceneTools/Private/Channels/SCurveEditorKeyBarView.cpp
+++ b/Engine/Source/Editor/MovieSceneTools/Private/Channels/SCurveEditorKeyBarView.cpp
@@ -27,7 +27,6 @@
 #include "Styling/AppStyle.h"
 #include "Styling/CoreStyle.h"
 #include "Styling/ISlateStyle.h"
-#include "Templates/ChooseClass.h"
 
 class FCurveModel;
 class FMenuBuilder;
diff --git a/Engine/Source/Editor/MovieSceneTools/Private/MovieSceneEventUtils.cpp b/Engine/Source/Editor/MovieSceneTools/Private/MovieSceneEventUtils.cpp
--- a/Engine/Source/Editor/MovieSceneTools/Private/MovieSceneEventUtils.cpp
+++ b/Engine/Source/Editor/MovieSceneTools/Private/MovieSceneEventUtils.cpp
@@ -45,7 +45,6 @@
 #include "ScopedTransaction.h"
 #include "Sections/MovieSceneEventSectionBase.h"
 #include "Templates/Casts.h"
-#include "Templates/ChooseClass.h"
 #include "Templates/SubclassOf.h"
 #include "Trace/Detail/Channel.h"
 #include "Tracks/MovieSceneEventTrack.h"
diff --git a/Engine/Source/Editor/MovieSceneTools/Private/TrackEditors/CameraShakeSourceShakeTrackEditor.cpp b/Engine/Source/Editor/MovieSceneTools/Private/TrackEditors/CameraShakeSourceShakeTrackEditor.cpp
--- a/Engine/Source/Editor/MovieSceneTools/Private/TrackEditors/CameraShakeSourceShakeTrackEditor.cpp
+++ b/Engine/Source/Editor/MovieSceneTools/Private/TrackEditors/CameraShakeSourceShakeTrackEditor.cpp
@@ -65,7 +65,6 @@
 #include "Styling/SlateColor.h"
 #include "Styling/WidgetStyle.h"
 #include "Templates/Casts.h"
-#include "Templates/ChooseClass.h"
 #include "Textures/SlateIcon.h"
 #include "TimeToPixel.h"
 #include "Tracks/MovieSceneCameraShakeSourceShakeTrack.h"
diff --git a/Engine/Source/Editor/MovieSceneTools/Private/TrackEditors/CameraShakeTrackEditor.cpp b/Engine/Source/Editor/MovieSceneTools/Private/TrackEditors/CameraShakeTrackEditor.cpp
--- a/Engine/Source/Editor/MovieSceneTools/Private/TrackEditors/CameraShakeTrackEditor.cpp
+++ b/Engine/Source/Editor/MovieSceneTools/Private/TrackEditors/CameraShakeTrackEditor.cpp
@@ -44,7 +44,6 @@
 #include "SequencerUtilities.h"
 #include "SlotBase.h"
 #include "Templates/Casts.h"
-#include "Templates/ChooseClass.h"
 #include "Tracks/MovieSceneCameraShakeTrack.h"
 #include "Types/SlateEnums.h"
 #include "Types/SlateStructs.h"
diff --git a/Engine/Source/Editor/Sequencer/Private/MVVM/ObjectBindingModelStorageExtension.cpp b/Engine/Source/Editor/Sequencer/Private/MVVM/ObjectBindingModelStorageExtension.cpp
--- a/Engine/Source/Editor/Sequencer/Private/MVVM/ObjectBindingModelStorageExtension.cpp
+++ b/Engine/Source/Editor/Sequencer/Private/MVVM/ObjectBindingModelStorageExtension.cpp
@@ -15,7 +15,6 @@
 #include "MovieScene.h"
 #include "MovieSceneBinding.h"
 #include "MovieSceneSequence.h"
-#include "Templates/ChooseClass.h"
 #include "Templates/TypeHash.h"
 
 class UMovieSceneTrack;
diff --git a/Engine/Source/Editor/Sequencer/Private/MVVM/SectionModelStorageExtension.cpp b/Engine/Source/Editor/Sequencer/Private/MVVM/SectionModelStorageExtension.cpp
--- a/Engine/Source/Editor/Sequencer/Private/MVVM/SectionModelStorageExtension.cpp
+++ b/Engine/Source/Editor/Sequencer/Private/MVVM/SectionModelStorageExtension.cpp
@@ -4,7 +4,6 @@
 
 #include "MVVM/ViewModels/SectionModel.h"
 #include "MovieSceneSection.h"
-#include "Templates/ChooseClass.h"
 
 namespace UE
 {
diff --git a/Engine/Source/Editor/Sequencer/Private/MVVM/TrackRowModelStorageExtension.cpp b/Engine/Source/Editor/Sequencer/Private/MVVM/TrackRowModelStorageExtension.cpp
--- a/Engine/Source/Editor/Sequencer/Private/MVVM/TrackRowModelStorageExtension.cpp
+++ b/Engine/Source/Editor/Sequencer/Private/MVVM/TrackRowModelStorageExtension.cpp
@@ -10,7 +10,6 @@
 #include "Misc/AssertionMacros.h"
 #include "Misc/Optional.h"
 #include "MovieSceneTrack.h"
-#include "Templates/ChooseClass.h"
 #include "Templates/TypeHash.h"
 
 namespace UE
diff --git a/Engine/Source/Editor/SequencerCore/Private/MVVM/SharedViewModelData.cpp b/Engine/Source/Editor/SequencerCore/Private/MVVM/SharedViewModelData.cpp
--- a/Engine/Source/Editor/SequencerCore/Private/MVVM/SharedViewModelData.cpp
+++ b/Engine/Source/Editor/SequencerCore/Private/MVVM/SharedViewModelData.cpp
@@ -5,7 +5,6 @@
 #include "MVVM/ViewModels/ViewModel.h"
 #include "MVVM/ViewModels/ViewModelHierarchy.h"
 #include "Misc/AssertionMacros.h"
-#include "Templates/ChooseClass.h"
 
 namespace UE
 {
diff --git a/Engine/Source/Editor/UnrealEd/Private/Commandlets/DiffAssetsCommandlet.cpp b/Engine/Source/Editor/UnrealEd/Private/Commandlets/DiffAssetsCommandlet.cpp
--- a/Engine/Source/Editor/UnrealEd/Private/Commandlets/DiffAssetsCommandlet.cpp
+++ b/Engine/Source/Editor/UnrealEd/Private/Commandlets/DiffAssetsCommandlet.cpp
@@ -18,7 +18,6 @@
 #include "Misc/FileHelper.h"
 #include "Misc/PackageName.h"
 #include "Misc/Paths.h"
-#include "Templates/ChooseClass.h"
 #include "Templates/UnrealTemplate.h"
 #include "Trace/Detail/Channel.h"
 #include "UObject/NameTypes.h"
diff --git a/Engine/Source/Editor/UnrealEd/Private/Commandlets/GenerateGatherArchiveCommandlet.cpp b/Engine/Source/Editor/UnrealEd/Private/Commandlets/GenerateGatherArchiveCommandlet.cpp
--- a/Engine/Source/Editor/UnrealEd/Private/Commandlets/GenerateGatherArchiveCommandlet.cpp
+++ b/Engine/Source/Editor/UnrealEd/Private/Commandlets/GenerateGatherArchiveCommandlet.cpp
@@ -14,7 +14,6 @@
 #include "LocTextHelper.h"
 #include "Logging/LogCategory.h"
 #include "Logging/LogMacros.h"
-#include "Templates/ChooseClass.h"
 #include "Templates/SharedPointer.h"
 #include "Trace/Detail/Channel.h"
 
diff --git a/Engine/Source/Editor/UnrealEd/Private/Commandlets/ShaderPipelineCacheToolsCommandlet.cpp b/Engine/Source/Editor/UnrealEd/Private/Commandlets/ShaderPipelineCacheToolsCommandlet.cpp
--- a/Engine/Source/Editor/UnrealEd/Private/Commandlets/ShaderPipelineCacheToolsCommandlet.cpp
+++ b/Engine/Source/Editor/UnrealEd/Private/Commandlets/ShaderPipelineCacheToolsCommandlet.cpp
@@ -45,7 +45,6 @@
 #include "ShaderCodeLibrary.h"
 #include "ShaderPipelineCache.h"
 #include "Stats/Stats2.h"
-#include "Templates/ChooseClass.h"
 #include "Templates/Function.h"
 #include "Templates/Tuple.h"
 #include "Templates/UnrealTemplate.h"
diff --git a/Engine/Source/Editor/UnrealEd/Private/Commandlets/SummarizeTraceUtils.cpp b/Engine/Source/Editor/UnrealEd/Private/Commandlets/SummarizeTraceUtils.cpp
--- a/Engine/Source/Editor/UnrealEd/Private/Commandlets/SummarizeTraceUtils.cpp
+++ b/Engine/Source/Editor/UnrealEd/Private/Commandlets/SummarizeTraceUtils.cpp
@@ -20,7 +20,6 @@
 #include "Misc/Parse.h"
 #include "Misc/Timespan.h"
 #include "Model/MonotonicTimeline.h"
-#include "Templates/ChooseClass.h"
 #include "Templates/Function.h"
 #include "Templates/SharedPointer.h"
 #include "Templates/Tuple.h"
diff --git a/Engine/Source/Editor/UnrealEd/Private/DiffUtils.cpp b/Engine/Source/Editor/UnrealEd/Private/DiffUtils.cpp
--- a/Engine/Source/Editor/UnrealEd/Private/DiffUtils.cpp
+++ b/Engine/Source/Editor/UnrealEd/Private/DiffUtils.cpp
@@ -22,7 +22,6 @@
 #include "SlotBase.h"
 #include "Styling/AppStyle.h"
 #include "Styling/SlateColor.h"
-#include "Templates/ChooseClass.h"
 #include "Templates/SubclassOf.h"
 #include "Templates/UnrealTemplate.h"
 #include "Types/SlateConstants.h"
diff --git a/Engine/Source/Editor/UnrealEd/Private/UnrealEdPrivatePCH.h b/Engine/Source/Editor/UnrealEd/Private/UnrealEdPrivatePCH.h
--- a/Engine/Source/Editor/UnrealEd/Private/UnrealEdPrivatePCH.h
+++ b/Engine/Source/Editor/UnrealEd/Private/UnrealEdPrivatePCH.h
@@ -32,7 +32,6 @@
 #include "Templates/EnableIf.h"
 #include "Templates/RemoveReference.h"
 #include "Templates/TypeCompatibleBytes.h"
-#include "Templates/ChooseClass.h"
 #include "Templates/IntegralConstant.h"
 #include "Templates/IsClass.h"
 #include "Traits/IsContiguousContainer.h"
diff --git a/Engine/Source/Runtime/Core/Private/Internationalization/InternationalizationMetadata.cpp b/Engine/Source/Runtime/Core/Private/Internationalization/InternationalizationMetadata.cpp
--- a/Engine/Source/Runtime/Core/Private/Internationalization/InternationalizationMetadata.cpp
+++ b/Engine/Source/Runtime/Core/Private/Internationalization/InternationalizationMetadata.cpp
@@ -10,7 +10,6 @@
 #include "Serialization/StructuredArchiveAdapters.h"
 #include "Serialization/StructuredArchiveNameHelpers.h"
 #include "Serialization/StructuredArchiveSlots.h"
-#include "Templates/ChooseClass.h"
 #include "Templates/Less.h"
 #include "Templates/Tuple.h"
 #include "Trace/Detail/Channel.h"
diff --git a/Engine/Source/Runtime/Core/Private/Internationalization/StringTableCore.cpp b/Engine/Source/Runtime/Core/Private/Internationalization/StringTableCore.cpp
--- a/Engine/Source/Runtime/Core/Private/Internationalization/StringTableCore.cpp
+++ b/Engine/Source/Runtime/Core/Private/Internationalization/StringTableCore.cpp
@@ -17,7 +17,6 @@
 #include "Misc/ScopeLock.h"
 #include "Serialization/Archive.h"
 #include "Serialization/Csv/CsvParser.h"
-#include "Templates/ChooseClass.h"
 #include "Templates/Tuple.h"
 #include "Trace/Detail/Channel.h"
 
diff --git a/Engine/Source/Runtime/Core/Private/Misc/NamePermissionList.cpp b/Engine/Source/Runtime/Core/Private/Misc/NamePermissionList.cpp
--- a/Engine/Source/Runtime/Core/Private/Misc/NamePermissionList.cpp
+++ b/Engine/Source/Runtime/Core/Private/Misc/NamePermissionList.cpp
@@ -2,7 +2,6 @@
 
 #include "Misc/NamePermissionList.h"
 
-#include "Templates/ChooseClass.h"
 #include "Templates/Tuple.h"
 #include "Templates/UnrealTemplate.h"
 
diff --git a/Engine/Source/Runtime/Core/Private/Misc/PathPermissionList.cpp b/Engine/Source/Runtime/Core/Private/Misc/PathPermissionList.cpp
--- a/Engine/Source/Runtime/Core/Private/Misc/PathPermissionList.cpp
+++ b/Engine/Source/Runtime/Core/Private/Misc/PathPermissionList.cpp
@@ -11,7 +11,6 @@
 #include "Misc/NamePermissionList.h"
 #include "Misc/PathViews.h"
 #include "Misc/StringBuilder.h"
-#include "Templates/ChooseClass.h"
 #include "Templates/Tuple.h"
 #include "Templates/UnrealTemplate.h"
 #include "UObject/NameTypes.h"
diff --git a/Engine/Source/Runtime/Core/Public/Containers/Array.h b/Engine/Source/Runtime/Core/Public/Containers/Array.h
--- a/Engine/Source/Runtime/Core/Public/Containers/Array.h
+++ b/Engine/Source/Runtime/Core/Public/Containers/Array.h
@@ -24,7 +24,6 @@
 #include "Templates/IdentityFunctor.h"
 #include "Templates/Invoke.h"
 #include "Templates/Less.h"
-#include "Templates/ChooseClass.h"
 #include "Templates/Sorting.h"
 #include "Templates/AlignmentTemplates.h"
 #include "Templates/IsConstructible.h"
diff --git a/Engine/Source/Runtime/Core/Public/Containers/BitArray.h b/Engine/Source/Runtime/Core/Public/Containers/BitArray.h
--- a/Engine/Source/Runtime/Core/Public/Containers/BitArray.h
+++ b/Engine/Source/Runtime/Core/Public/Containers/BitArray.h
@@ -12,7 +12,6 @@
 #include "Serialization/Archive.h"
 #include "Serialization/MemoryImageWriter.h"
 #include "Serialization/MemoryLayout.h"
-#include "Templates/ChooseClass.h"
 #include "Templates/EnableIf.h"
 #include "Templates/Invoke.h"
 #include "Templates/UnrealTemplate.h"
diff --git a/Engine/Source/Runtime/Core/Public/Containers/Deque.h b/Engine/Source/Runtime/Core/Public/Containers/Deque.h
--- a/Engine/Source/Runtime/Core/Public/Containers/Deque.h
+++ b/Engine/Source/Runtime/Core/Public/Containers/Deque.h
@@ -7,7 +7,6 @@
 #include "IteratorAdapter.h"
 #include "Math/UnrealMathUtility.h"
 #include "Misc/AssertionMacros.h"
-#include "Templates/ChooseClass.h"
 #include "Templates/MemoryOps.h"
 
 #include <initializer_list>
diff --git a/Engine/Source/Runtime/Core/Public/Containers/HashTable.h b/Engine/Source/Runtime/Core/Public/Containers/HashTable.h
--- a/Engine/Source/Runtime/Core/Public/Containers/HashTable.h
+++ b/Engine/Source/Runtime/Core/Public/Containers/HashTable.h
@@ -13,7 +13,6 @@
 #include "Serialization/MemoryImageWriter.h"
 #include "Serialization/MemoryImageWriter.h"
 #include "Serialization/MemoryLayout.h"
-#include "Templates/ChooseClass.h"
 #include "Templates/UnrealTemplate.h"
 
 #include <initializer_list>
diff --git a/Engine/Source/Runtime/Core/Public/Core.h b/Engine/Source/Runtime/Core/Public/Core.h
--- a/Engine/Source/Runtime/Core/Public/Core.h
+++ b/Engine/Source/Runtime/Core/Public/Core.h
@@ -62,7 +62,6 @@ MONOLITHIC_HEADER_BOILERPLATE()
 #include "Templates/PointerIsConvertibleFromTo.h"
 #include "Templates/AlignmentTemplates.h"
 #include "Templates/RemoveReference.h"
-#include "Templates/ChooseClass.h"
 #include "Templates/IntegralConstant.h"
 #include "Templates/IsClass.h"
 #include "Templates/TypeCompatibleBytes.h"
diff --git a/Engine/Source/Runtime/Core/Public/CoreMinimal.h b/Engine/Source/Runtime/Core/Public/CoreMinimal.h
--- a/Engine/Source/Runtime/Core/Public/CoreMinimal.h
+++ b/Engine/Source/Runtime/Core/Public/CoreMinimal.h
@@ -40,7 +40,6 @@
 #include "Templates/UnrealTypeTraits.h"
 #include "Templates/EnableIf.h"
 #include "Templates/RemoveReference.h"
-#include "Templates/ChooseClass.h"
 #include "Templates/IntegralConstant.h"
 #include "Templates/IsClass.h"
 #include "Templates/TypeCompatibleBytes.h"
diff --git a/Engine/Source/Runtime/Core/Public/CoreSharedPCH.h b/Engine/Source/Runtime/Core/Public/CoreSharedPCH.h
--- a/Engine/Source/Runtime/Core/Public/CoreSharedPCH.h
+++ b/Engine/Source/Runtime/Core/Public/CoreSharedPCH.h
@@ -211,7 +211,6 @@
 #include "Stats/StatsMisc.h"
 #include "Templates/AlignmentTemplates.h"
 #include "Templates/AndOrNot.h"
-#include "Templates/ChooseClass.h"
 #include "Templates/Decay.h"
 #include "Templates/EnableIf.h"
 #include "Templates/Function.h"
diff --git a/Engine/Source/Runtime/Core/Public/Memory/MemoryView.h b/Engine/Source/Runtime/Core/Public/Memory/MemoryView.h
--- a/Engine/Source/Runtime/Core/Public/Memory/MemoryView.h
+++ b/Engine/Source/Runtime/Core/Public/Memory/MemoryView.h
@@ -8,7 +8,6 @@
 #include "Math/UnrealMathUtility.h"
 #include "Memory/MemoryFwd.h"
 #include "Templates/AndOrNot.h"
-#include "Templates/ChooseClass.h"
 #include "Templates/EnableIf.h"
 #include "Templates/Identity.h"
 #include "Templates/IsConst.h"
diff --git a/Engine/Source/Runtime/Core/Public/Misc/ConfigCacheIni.h b/Engine/Source/Runtime/Core/Public/Misc/ConfigCacheIni.h
--- a/Engine/Source/Runtime/Core/Public/Misc/ConfigCacheIni.h
+++ b/Engine/Source/Runtime/Core/Public/Misc/ConfigCacheIni.h
@@ -34,7 +34,6 @@
 #include "Serialization/StructuredArchive.h"
 #include "Serialization/StructuredArchiveAdapters.h"
 #include "Serialization/StructuredArchiveSlots.h"
-#include "Templates/ChooseClass.h"
 #include "Templates/Function.h"
 #include "Templates/UnrealTemplate.h"
 #include "UObject/NameTypes.h"
diff --git a/Engine/Source/Runtime/Core/Public/Templates/Function.h b/Engine/Source/Runtime/Core/Public/Templates/Function.h
--- a/Engine/Source/Runtime/Core/Public/Templates/Function.h
+++ b/Engine/Source/Runtime/Core/Public/Templates/Function.h
@@ -7,7 +7,6 @@
 #include "Misc/IntrusiveUnsetOptionalState.h"
 #include "HAL/UnrealMemory.h"
 #include "Templates/AndOrNot.h"
-#include "Templates/ChooseClass.h"
 #include "Templates/UnrealTypeTraits.h"
 #include "Templates/RemoveReference.h"
 #include "Templates/Decay.h"
diff --git a/Engine/Source/Runtime/Core/Tests/Math/GuardedIntTest.cpp b/Engine/Source/Runtime/Core/Tests/Math/GuardedIntTest.cpp
--- a/Engine/Source/Runtime/Core/Tests/Math/GuardedIntTest.cpp
+++ b/Engine/Source/Runtime/Core/Tests/Math/GuardedIntTest.cpp
@@ -2,7 +2,6 @@
 
 #include "Math/GuardedInt.h"
 #include "Misc/AutomationTest.h"
-#include "Templates/ChooseClass.h"
 #include <limits>
 #include <type_traits>
 
diff --git a/Engine/Source/Runtime/CoreUObject/Private/UObject/EnumProperty.cpp b/Engine/Source/Runtime/CoreUObject/Private/UObject/EnumProperty.cpp
--- a/Engine/Source/Runtime/CoreUObject/Private/UObject/EnumProperty.cpp
+++ b/Engine/Source/Runtime/CoreUObject/Private/UObject/EnumProperty.cpp
@@ -5,7 +5,6 @@
 #include "UObject/UObjectThreadContext.h"
 #include "UObject/PropertyTag.h"
 #include "UObject/UnrealTypePrivate.h"
-#include "Templates/ChooseClass.h"
 #include "Templates/IsSigned.h"
 #include "Algo/Find.h"
 #include "UObject/LinkerLoad.h"
diff --git a/Engine/Source/Runtime/CoreUObject/Public/Templates/NonNullSubclassOf.h b/Engine/Source/Runtime/CoreUObject/Public/Templates/NonNullSubclassOf.h
--- a/Engine/Source/Runtime/CoreUObject/Public/Templates/NonNullSubclassOf.h
+++ b/Engine/Source/Runtime/CoreUObject/Public/Templates/NonNullSubclassOf.h
@@ -5,7 +5,6 @@
 #include "CoreMinimal.h"
 #include "UObject/Class.h"
 #include "UObject/Field.h"
-#include "Templates/ChooseClass.h"
 #include "SubclassOf.h"
 
 // So we can construct uninitialized TNonNullSubclassOf
diff --git a/Engine/Source/Runtime/CoreUObject/Public/UObject/Class.h b/Engine/Source/Runtime/CoreUObject/Public/UObject/Class.h
--- a/Engine/Source/Runtime/CoreUObject/Public/UObject/Class.h
+++ b/Engine/Source/Runtime/CoreUObject/Public/UObject/Class.h
@@ -45,7 +45,6 @@
 #include "Serialization/StructuredArchive.h"
 #include "Serialization/StructuredArchiveAdapters.h"
 #include "Templates/AlignmentTemplates.h"
-#include "Templates/ChooseClass.h"
 #include "Templates/EnableIf.h"
 #include "Templates/IsAbstract.h"
 #include "Templates/IsEnum.h"
diff --git a/Engine/Source/Runtime/Engine/Private/EnginePrivatePCH.h b/Engine/Source/Runtime/Engine/Private/EnginePrivatePCH.h
--- a/Engine/Source/Runtime/Engine/Private/EnginePrivatePCH.h
+++ b/Engine/Source/Runtime/Engine/Private/EnginePrivatePCH.h
@@ -32,7 +32,6 @@
 #include "Templates/EnableIf.h"
 #include "Templates/RemoveReference.h"
 #include "Templates/TypeCompatibleBytes.h"
-#include "Templates/ChooseClass.h"
 #include "Templates/IntegralConstant.h"
 #include "Templates/IsClass.h"
 #include "Traits/IsContiguousContainer.h"
diff --git a/Engine/Source/Runtime/Experimental/Chaos/Public/Chaos/ImplicitObjectScaled.h b/Engine/Source/Runtime/Experimental/Chaos/Public/Chaos/ImplicitObjectScaled.h
--- a/Engine/Source/Runtime/Experimental/Chaos/Public/Chaos/ImplicitObjectScaled.h
+++ b/Engine/Source/Runtime/Experimental/Chaos/Public/Chaos/ImplicitObjectScaled.h
@@ -9,7 +9,6 @@
 #include "Chaos/Transform.h"
 #include "Chaos/Utilities.h"	// For ScaleInertia - pull that into mass utils
 #include "ChaosArchive.h"
-#include "Templates/ChooseClass.h"
 #include "Templates/EnableIf.h"
 #include "Math/NumericLimits.h"
 #include "ChaosCheck.h"
diff --git a/Engine/Source/Runtime/Experimental/Chaos/Public/Chaos/ParticleHandle.h b/Engine/Source/Runtime/Experimental/Chaos/Public/Chaos/ParticleHandle.h
--- a/Engine/Source/Runtime/Experimental/Chaos/Public/Chaos/ParticleHandle.h
+++ b/Engine/Source/Runtime/Experimental/Chaos/Public/Chaos/ParticleHandle.h
@@ -1,7 +1,6 @@
 // Copyright Epic Games, Inc. All Rights Reserved.
 #pragma once
 
-#include "Templates/ChooseClass.h"
 #include "Math/NumericLimits.h"
 #include "Chaos/ISpatialAcceleration.h"
 #include "Chaos/PBDRigidClusteredParticles.h"
diff --git a/Engine/Source/Runtime/Experimental/Iris/Core/Public/Iris/Serialization/NetSerializerBuilder.inl b/Engine/Source/Runtime/Experimental/Iris/Core/Public/Iris/Serialization/NetSerializerBuilder.inl
--- a/Engine/Source/Runtime/Experimental/Iris/Core/Public/Iris/Serialization/NetSerializerBuilder.inl
+++ b/Engine/Source/Runtime/Experimental/Iris/Core/Public/Iris/Serialization/NetSerializerBuilder.inl
@@ -2,7 +2,6 @@
 
 #pragma once
 
-#include "Templates/ChooseClass.h"
 #include "Templates/EnableIf.h"
 #include "Templates/IsPODType.h"
 #include "Iris/Serialization/NetBitStreamReader.h"
diff --git a/Engine/Source/Runtime/MeshDescription/Public/MeshAttributeArray.h b/Engine/Source/Runtime/MeshDescription/Public/MeshAttributeArray.h
--- a/Engine/Source/Runtime/MeshDescription/Public/MeshAttributeArray.h
+++ b/Engine/Source/Runtime/MeshDescription/Public/MeshAttributeArray.h
@@ -23,7 +23,6 @@
 #include "Misc/TVariant.h"
 #include "Serialization/Archive.h"
 #include "Serialization/StructuredArchiveAdapters.h"
-#include "Templates/ChooseClass.h"
 #include "Templates/CopyQualifiersFromTo.h"
 #include "Templates/EnableIf.h"
 #include "Templates/IsArray.h"
diff --git a/Engine/Source/Runtime/Net/Core/Public/Net/Core/Trace/Private/NetTraceInternal.h b/Engine/Source/Runtime/Net/Core/Public/Net/Core/Trace/Private/NetTraceInternal.h
--- a/Engine/Source/Runtime/Net/Core/Public/Net/Core/Trace/Private/NetTraceInternal.h
+++ b/Engine/Source/Runtime/Net/Core/Public/Net/Core/Trace/Private/NetTraceInternal.h
@@ -12,7 +12,6 @@
 #include "CoreTypes.h"
 #include "Misc/NetworkGuid.h"
 #include "Net/Core/Trace/NetDebugName.h"
-#include "Templates/ChooseClass.h"
 #include "UObject/NameTypes.h"
 
 class FNetTraceCollector;
diff --git a/Engine/Source/Runtime/SlateCore/Private/FastUpdate/SlateInvalidationWidgetList.cpp b/Engine/Source/Runtime/SlateCore/Private/FastUpdate/SlateInvalidationWidgetList.cpp
--- a/Engine/Source/Runtime/SlateCore/Private/FastUpdate/SlateInvalidationWidgetList.cpp
+++ b/Engine/Source/Runtime/SlateCore/Private/FastUpdate/SlateInvalidationWidgetList.cpp
@@ -10,7 +10,6 @@
 #include "Layout/Children.h"
 #include "Misc/StringBuilder.h"
 #include "Rendering/SlateLayoutTransform.h"
-#include "Templates/ChooseClass.h"
 #include "Templates/IsConst.h"
 #include "Templates/Tuple.h"
 #include "Types/ReflectionMetadata.h"
diff --git a/Engine/Source/Runtime/TypedElementFramework/Public/Elements/Framework/TypedElementRegistry.h b/Engine/Source/Runtime/TypedElementFramework/Public/Elements/Framework/TypedElementRegistry.h
--- a/Engine/Source/Runtime/TypedElementFramework/Public/Elements/Framework/TypedElementRegistry.h
+++ b/Engine/Source/Runtime/TypedElementFramework/Public/Elements/Framework/TypedElementRegistry.h
@@ -28,7 +28,6 @@
 #include "Misc/ScopeLock.h"
 #include "Misc/ScopeRWLock.h"
 #include "Templates/Casts.h"
-#include "Templates/ChooseClass.h"
 #include "Templates/SharedPointer.h"
 #include "Templates/SubclassOf.h"
 #include "Templates/Tuple.h"
