Deprecated TChooseClass.

diff --git a/Engine/Source/Runtime/Core/Public/Templates/ChooseClass.h b/Engine/Source/Runtime/Core/Public/Templates/ChooseClass.h
--- a/Engine/Source/Runtime/Core/Public/Templates/ChooseClass.h
+++ b/Engine/Source/Runtime/Core/Public/Templates/ChooseClass.h
@@ -9,14 +9,14 @@ template<bool Predicate,typename TrueClass,typename FalseClass>
 class TChooseClass;
 
 template<typename TrueClass,typename FalseClass>
-class TChooseClass<true,TrueClass,FalseClass>
+class UE_DEPRECATED(5.4, "TChooseClass has been deprecated, please use std::conditional instead.") TChooseClass<true,TrueClass,FalseClass>
 {
 public:
 	typedef TrueClass Result;
 };
 
 template<typename TrueClass,typename FalseClass>
-class TChooseClass<false,TrueClass,FalseClass>
+class UE_DEPRECATED(5.4, "TChooseClass has been deprecated, please use std::conditional instead.") TChooseClass<false,TrueClass,FalseClass>
 {
 public:
 	typedef FalseClass Result;
