Custom cross-platform dynamic TLS implementation because we hit OS TLS slot limit on many platforms. compile-time defined num slot limit.
Allocating and freeing slots can be contended, getting and setting a slot value is fast but has an additional indirection compared with OS TLS. Used only for windows editor for now, because it's a bit slower than OS TLS implementation.

diff --git a/Engine/Config/Windows/BaseWindowsEngine.ini b/Engine/Config/Windows/BaseWindowsEngine.ini
--- a/Engine/Config/Windows/BaseWindowsEngine.ini
+++ b/Engine/Config/Windows/BaseWindowsEngine.ini
@@ -25,3 +25,8 @@ BindlessSamplers=Disabled
 [SF_VULKAN_SM6]
 BindlessResources=Disabled
 BindlessSamplers=Disabled
+
+[/Script/WindowsRuntimeSettings.WindowsRuntimeSettings]
+MaxNumTlsSlots=2048
+MaxNumThreadsWithTlsSlots=512
+
diff --git a/Engine/Source/Programs/UnrealBuildTool/Platform/Windows/UEBuildWindows.cs b/Engine/Source/Programs/UnrealBuildTool/Platform/Windows/UEBuildWindows.cs
--- a/Engine/Source/Programs/UnrealBuildTool/Platform/Windows/UEBuildWindows.cs
+++ b/Engine/Source/Programs/UnrealBuildTool/Platform/Windows/UEBuildWindows.cs
@@ -485,6 +485,18 @@ namespace UnrealBuildTool
 		[ConfigFile(ConfigHierarchyType.Engine, "/Script/WindowsTargetPlatform.WindowsTargetSettings")]
 		public int DefaultStackSizeCommit;
 
+		/// <summary>
+		/// Max number of slots FWindowsPlatformTLS::AllocTlsSlot can allocate.
+		/// </summary>
+		[ConfigFile(ConfigHierarchyType.Engine, "/Script/WindowsRuntimeSettings.WindowsRuntimeSettings")]
+		public int MaxNumTlsSlots = 0;
+
+		/// <summary>
+		/// Max number threads that can use FWindowsPlatformTLS at one time.
+		/// </summary>
+		[ConfigFile(ConfigHierarchyType.Engine, "/Script/WindowsRuntimeSettings.WindowsRuntimeSettings")]
+		public int MaxNumThreadsWithTlsSlots = 0;
+
 		/// <summary>
 		/// Determines the amount of memory that the compiler allocates to construct precompiled headers (/Zm).
 		/// </summary>
@@ -822,6 +834,11 @@ namespace UnrealBuildTool
 
 		public string DirectXDllDir => Inner.DirectXDllDir;
 
+		public int MaxNumTlsSlots => Inner.MaxNumTlsSlots;
+
+		public int MaxNumThreadsWithTlsSlots => Inner.MaxNumThreadsWithTlsSlots;
+
+
 #pragma warning restore CS1591
 		#endregion
 	}
@@ -1599,6 +1616,9 @@ namespace UnrealBuildTool
 				LinkEnvironment.PGODirectory = CompileEnvironment.PGODirectory;
 				LinkEnvironment.PGOFilenamePrefix = CompileEnvironment.PGOFilenamePrefix;
 			}
+
+			CompileEnvironment.Definitions.Add("WINDOWS_MAX_NUM_TLS_SLOTS=" + Target.WindowsPlatform.MaxNumTlsSlots.ToString());
+			CompileEnvironment.Definitions.Add("WINDOWS_MAX_NUM_THREADS_WITH_TLS_SLOTS=" + Target.WindowsPlatform.MaxNumThreadsWithTlsSlots.ToString());
 		}
 
 		/// <summary>
diff --git a/Engine/Source/Runtime/Core/Private/Microsoft/MinimalWindowsApi.cpp b/Engine/Source/Runtime/Core/Private/Microsoft/MinimalWindowsApi.cpp
--- a/Engine/Source/Runtime/Core/Private/Microsoft/MinimalWindowsApi.cpp
+++ b/Engine/Source/Runtime/Core/Private/Microsoft/MinimalWindowsApi.cpp
@@ -155,6 +155,26 @@ namespace Windows
 		return ::TlsFree(dwTlsIndex);
 	}
 
+	MINIMAL_WINDOWS_API DWORD WINAPI FlsAlloc(PFLS_CALLBACK_FUNCTION lpCallback)
+	{
+		return ::FlsAlloc(lpCallback);
+	}
+
+	MINIMAL_WINDOWS_API LPVOID WINAPI FlsGetValue(DWORD dwFlsIndex)
+	{
+		return ::FlsGetValue(dwFlsIndex);
+	}
+
+	MINIMAL_WINDOWS_API BOOL WINAPI FlsSetValue(DWORD dwFlsIndex, LPVOID lpFlsValue)
+	{
+		return ::FlsSetValue(dwFlsIndex, lpFlsValue);
+	}
+
+	MINIMAL_WINDOWS_API BOOL WINAPI FlsFree(DWORD dwTlsIndex)
+	{
+		return ::FlsFree(dwTlsIndex);
+	}
+
 	MINIMAL_WINDOWS_API BOOL WINAPI IsProcessorFeaturePresent(DWORD ProcessorFeature)
 	{
 		return ::IsProcessorFeaturePresent(ProcessorFeature);
diff --git a/Engine/Source/Runtime/Core/Private/Windows/WindowsPlatformTLS.cpp b/Engine/Source/Runtime/Core/Private/Windows/WindowsPlatformTLS.cpp
new file mode 100644
--- /dev/null
+++ b/Engine/Source/Runtime/Core/Private/Windows/WindowsPlatformTLS.cpp
@@ -0,0 +1,147 @@
+// Copyright Epic Games, Inc. All Rights Reserved.
+
+#include "Windows/WindowsPlatformTLS.h"
+
+#if WITH_EDITOR
+
+#include "Templates/UnrealTemplate.h"
+#include "ThirdParty/AtomicQueue/AtomicQueue.h"
+
+// Custom cross-platform dynamic TLS implementation because we hit OS TLS slot limit on many platforms. 
+// Compile-time defined num slot and num thread limits.
+// Allocating and freeing slots can be contended, getting and setting a slot value is fast but has an additional indirection compared with OS TLS.
+
+namespace WindowsPlatformTLS_Private
+{
+	uint32 PrimarySlot = FGenericPlatformTLS::InvalidTlsSlot;
+
+	// lock-free container for indices in the range [0..Size)
+	template<uint32 Size>
+	class TIndices
+	{
+	public:
+		TIndices()
+		{
+			// initially all indices are available
+			for (uint32 i = 0; i != Size; ++i)
+			{
+				// `0` is a special value for AtomicQueue, shift away from it
+				Queue.push(i + 1);
+			}
+		}
+
+		uint32 Alloc()
+		{
+			return Queue.pop() - 1;
+		}
+
+		void Free(uint32 Value)
+		{
+			Queue.push(Value + 1);
+		}
+
+	private:
+		atomic_queue::AtomicQueue<uint32, Size> Queue;
+	};
+
+	// TLS storage for all threads and slots
+	class FStorage
+	{
+	public:
+		void** GetThreadStorage(uint32 ThreadIndex)
+		{
+			return Buffer[ThreadIndex];
+		}
+
+		uint32 GetThreadIndex(void** ThreadStorage)
+		{
+			uint32 ThreadIndex = uint32((ThreadStorage - *Buffer) / MaxSlots);
+			check(ThreadIndex < MaxThreads);
+			return ThreadIndex;
+		}
+
+		void ResetSlot(uint32 SlotIndex)
+		{
+			for (uint32 ThreadIndex = 0; ThreadIndex != MaxThreads; ++ThreadIndex)
+			{
+				Buffer[ThreadIndex][SlotIndex] = nullptr;
+			}
+		}
+
+		void ResetThread(uint32 ThreadIndex)
+		{
+			FMemory::Memzero(Buffer[ThreadIndex], MaxSlots * sizeof(void*));
+ 		}
+
+	private:
+		void* Buffer[MaxThreads][MaxSlots]; // compile-time zero initialized
+	};
+
+	FStorage Storage;
+
+	void OnThreadExit(void* TlsValue);
+
+	// uses local static var to not depend on static initialization order
+	class FSingleton
+	{
+	private:
+		FSingleton()
+		{
+			PrimarySlot = Windows::FlsAlloc(&OnThreadExit);
+		}
+
+	public:
+		static FSingleton& Get()
+		{
+			static FSingleton Singleton;
+			return Singleton;
+		}
+
+		TIndices<MaxSlots> Slots;
+		TIndices<MaxThreads> Threads;
+	};
+
+	// called on thread exit if the thread has ever set TLS value
+	// TlsValue: a value set to PrimarySlot
+	void OnThreadExit(void* TlsValue)
+	{
+		check(PrimarySlot < MaxSlots);
+
+#if DO_CHECK
+		void* TlsValueLocal = Windows::FlsGetValue(PrimarySlot);
+		checkf(TlsValue == TlsValueLocal, TEXT("%p - %p"), TlsValue, TlsValueLocal);
+#endif
+
+		Windows::FlsSetValue(PrimarySlot, nullptr); // clear the value otherwise this function will be called again
+
+		void** ThreadStorage = (void**)TlsValue;
+		uint32 ThreadIndex = Storage.GetThreadIndex(ThreadStorage);
+		Storage.ResetThread(ThreadIndex);
+		FSingleton::Get().Threads.Free(ThreadIndex);
+	}
+
+	void** AllocThreadStorage()
+	{
+		uint32 ThreadIndex = FSingleton::Get().Threads.Alloc();
+		void** ThreadStorage = Storage.GetThreadStorage(ThreadIndex);
+		Windows::FlsSetValue(PrimarySlot, ThreadStorage);
+		return ThreadStorage;
+	}
+}
+
+uint32 FWindowsPlatformTLS::AllocTlsSlot()
+{
+	return WindowsPlatformTLS_Private::FSingleton::Get().Slots.Alloc();
+}
+
+void FWindowsPlatformTLS::FreeTlsSlot(uint32 SlotIndex)
+{
+	using namespace WindowsPlatformTLS_Private;
+
+	checkf(SlotIndex < MaxSlots, TEXT("Invalid TLS slot index %u"), SlotIndex);
+
+	Storage.ResetSlot(SlotIndex);
+	FSingleton::Get().Slots.Free(SlotIndex);
+}
+
+#endif
\ No newline at end of file
diff --git a/Engine/Source/Runtime/Core/Public/Microsoft/MinimalWindowsApi.h b/Engine/Source/Runtime/Core/Public/Microsoft/MinimalWindowsApi.h
--- a/Engine/Source/Runtime/Core/Public/Microsoft/MinimalWindowsApi.h
+++ b/Engine/Source/Runtime/Core/Public/Microsoft/MinimalWindowsApi.h
@@ -94,6 +94,9 @@ namespace Windows
 	typedef _RTL_SRWLOCK RTL_SRWLOCK, *PRTL_SRWLOCK;
 	typedef RTL_SRWLOCK *PSRWLOCK;
 
+	typedef void (*PFLS_CALLBACK_FUNCTION) (void* lpFlsData);
+
+
 	// Opaque SRWLOCK structure
 	struct SRWLOCK
 	{
@@ -176,6 +179,12 @@ namespace Windows
 	MINIMAL_WINDOWS_API BOOL WINAPI TlsSetValue(DWORD dwTlsIndex, LPVOID lpTlsValue);
 	MINIMAL_WINDOWS_API BOOL WINAPI TlsFree(DWORD dwTlsIndex);
 
+	// Fiber-local storage functions
+	MINIMAL_WINDOWS_API DWORD WINAPI FlsAlloc(PFLS_CALLBACK_FUNCTION lpCallback);
+	MINIMAL_WINDOWS_API LPVOID WINAPI FlsGetValue(DWORD dwFlsIndex);
+	MINIMAL_WINDOWS_API BOOL WINAPI FlsSetValue(DWORD dwFlsIndex, LPVOID lpFlsValue);
+	MINIMAL_WINDOWS_API BOOL WINAPI FlsFree(DWORD dwTlsIndex);
+
 	// System
 	MINIMAL_WINDOWS_API BOOL WINAPI IsProcessorFeaturePresent(DWORD ProcessorFeature);
 
diff --git a/Engine/Source/Runtime/Core/Public/Windows/WindowsPlatformTLS.h b/Engine/Source/Runtime/Core/Public/Windows/WindowsPlatformTLS.h
--- a/Engine/Source/Runtime/Core/Public/Windows/WindowsPlatformTLS.h
+++ b/Engine/Source/Runtime/Core/Public/Windows/WindowsPlatformTLS.h
@@ -6,6 +6,27 @@
 #include "GenericPlatform/GenericPlatformTLS.h"
 #include "Windows/WindowsSystemIncludes.h"
 
+#if WITH_EDITOR
+
+#include "Misc/AssertionMacros.h"
+
+// Custom cross-platform dynamic TLS implementation because we hit Windows TLS slot limit (1088). compile-time defined limits for the maximum
+// number of TLS slots and the maximum number of threads that use TLS. No dynamic memory allocation. Lock-free.
+// Getting and setting a slot value is fast but has an additional indirection compared with OS TLS.
+// Memory footprint is (roughly): MaxSlots * 4B + MaxThreads * 4B + MaxSlots * MaxThreads * 8B
+
+namespace WindowsPlatformTLS_Private
+{
+	constexpr uint32 MaxSlots = WINDOWS_MAX_NUM_TLS_SLOTS; // how many slots are available
+	constexpr uint32 MaxThreads = WINDOWS_MAX_NUM_THREADS_WITH_TLS_SLOTS; // how many threads can use TLS
+
+	// a single OS TLS slot that is used to store the thread storage
+	CORE_API extern uint32 PrimarySlot;
+
+	CORE_API void** AllocThreadStorage();
+}
+
+#endif
 
 /**
  * Windows implementation of the TLS OS functions.
@@ -28,10 +49,29 @@ struct FWindowsPlatformTLS
 	 *
 	 * @return The index of the allocated slot.
 	 */
-	static FORCEINLINE uint32 AllocTlsSlot(void)
+#if WITH_EDITOR
+	CORE_API static uint32 AllocTlsSlot();
+#else
+	FORCEINLINE static uint32 AllocTlsSlot()
 	{
 		return Windows::TlsAlloc();
+}
+
+#endif
+
+	/**
+	 * Frees a previously allocated TLS slot
+	 *
+	 * @param SlotIndex the TLS index to store it in
+	 */
+#if WITH_EDITOR
+	CORE_API static void FreeTlsSlot(uint32 SlotIndex);
+#else
+	FORCEINLINE static void FreeTlsSlot(uint32 SlotIndex)
+	{
+		Windows::TlsFree(SlotIndex);
 	}
+#endif
 
 	/**
 	 * Sets a value in the specified TLS slot.
@@ -39,9 +79,22 @@ struct FWindowsPlatformTLS
 	 * @param SlotIndex the TLS index to store it in.
 	 * @param Value the value to store in the slot.
 	 */
-	static FORCEINLINE void SetTlsValue(uint32 SlotIndex,void* Value)
+	static FORCEINLINE void SetTlsValue(uint32 SlotIndex, void* Value)
 	{
-		Windows::TlsSetValue(SlotIndex,Value);
+#if WITH_EDITOR
+		using namespace WindowsPlatformTLS_Private;
+
+		checkf(SlotIndex < MaxSlots, TEXT("Invalid slot index %u"), SlotIndex);
+
+		void** ThreadStorage = (void**)Windows::FlsGetValue(PrimarySlot);
+		if (ThreadStorage == 0)
+		{
+			ThreadStorage = AllocThreadStorage();
+		}
+		ThreadStorage[SlotIndex] = Value;
+#else
+		Windows::TlsSetValue(SlotIndex, Value);
+#endif
 	}
 
 	/**
@@ -52,17 +105,16 @@ struct FWindowsPlatformTLS
 	 */
 	static FORCEINLINE void* GetTlsValue(uint32 SlotIndex)
 	{
-		return Windows::TlsGetValue(SlotIndex);
-	}
+#if WITH_EDITOR
+		using namespace WindowsPlatformTLS_Private;
 
-	/**
-	 * Frees a previously allocated TLS slot
-	 *
-	 * @param SlotIndex the TLS index to store it in
-	 */
-	static FORCEINLINE void FreeTlsSlot(uint32 SlotIndex)
-	{
-		Windows::TlsFree(SlotIndex);
+		checkf(SlotIndex < MaxSlots, TEXT("Invalid slot index %u"), SlotIndex);
+	
+		void** ThreadStorage = (void**)Windows::FlsGetValue(PrimarySlot);
+		return ThreadStorage ? ThreadStorage[SlotIndex] : nullptr;
+#else
+		return Windows::TlsGetValue(SlotIndex);
+#endif
 	}
 };
 
