Moved string code into .inl files in anticipation of the upcoming FUtf8String change.

diff --git a/Engine/Source/Runtime/Core/Private/Containers/String.cpp b/Engine/Source/Runtime/Core/Private/Containers/String.cpp
--- a/Engine/Source/Runtime/Core/Private/Containers/String.cpp
+++ b/Engine/Source/Runtime/Core/Private/Containers/String.cpp
@@ -1,1964 +1,8 @@
 // Copyright Epic Games, Inc. All Rights Reserved.
 
-#include "Containers/Array.h"
-#include "Containers/StringConv.h"
 #include "Containers/UnrealString.h"
-#include "CoreGlobals.h"
-#include "CoreTypes.h"
-#include "HAL/PlatformString.h"
-#include "HAL/UnrealMemory.h"
-#include "Logging/LogCategory.h"
-#include "Logging/LogMacros.h"
-#include "Math/NumericLimits.h"
-#include "Math/UnrealMathUtility.h"
-#include "Misc/AssertionMacros.h"
-#include "Misc/ByteSwap.h"
-#include "Misc/CString.h"
-#include "Misc/Char.h"
-#include "Misc/VarArgs.h"
-#include "Serialization/Archive.h"
-#include "String/HexToBytes.h"
-#include "String/ParseTokens.h"
-#include "Templates/MemoryOps.h"
-#include "Templates/RemoveReference.h"
-#include "Templates/UnrealTemplate.h"
-#include "UObject/NameTypes.h"
 
-/* FString implementation
- *****************************************************************************/
+// Include String.cpp.inl's includes before defining (upcoming) macros, in case the macros 'poison' other headers or there are re-entrant includes.
+#include "Containers/StringIncludes.cpp.inl"
 
-namespace UE::Core::Private
-{
-	template <typename CharType>
-	struct TCompareCharsCaseSensitive
-	{
-		static FORCEINLINE bool Compare(CharType Lhs, CharType Rhs)
-		{
-			return Lhs == Rhs;
-		}
-	};
-
-	template <typename CharType>
-	struct TCompareCharsCaseInsensitive
-	{
-		static FORCEINLINE bool Compare(CharType Lhs, CharType Rhs)
-		{
-			return TChar<CharType>::ToLower(Lhs) == TChar<CharType>::ToLower(Rhs);
-		}
-	};
-
-	template <typename CompareType, typename CharType>
-	bool MatchesWildcardRecursive(const CharType* Target, int32 TargetLength, const CharType* Wildcard, int32 WildcardLength)
-	{
-		// Skip over common initial non-wildcard-char sequence of Target and Wildcard
-		for (;;)
-		{
-			if (WildcardLength == 0)
-			{
-				return TargetLength == 0;
-			}
-
-			CharType WCh = *Wildcard;
-			if (WCh == CharType('*') || WCh == CharType('?'))
-			{
-				break;
-			}
-
-			if (!CompareType::Compare(*Target, WCh))
-			{
-				return false;
-			}
-
-			++Target;
-			++Wildcard;
-			--TargetLength;
-			--WildcardLength;
-		}
-
-		// Test for common suffix
-		const CharType* TPtr = Target   + TargetLength;
-		const CharType* WPtr = Wildcard + WildcardLength;
-		for (;;)
-		{
-			--TPtr;
-			--WPtr;
-
-			CharType WCh = *WPtr;
-			if (WCh == CharType('*') || WCh == CharType('?'))
-			{
-				break;
-			}
-
-			if (!CompareType::Compare(*TPtr, WCh))
-			{
-				return false;
-			}
-
-			--TargetLength;
-			--WildcardLength;
-
-			if (TargetLength == 0)
-			{
-				break;
-			}
-		}
-
-		// Match * against anything and ? against single (and zero?) chars
-		CharType FirstWild = *Wildcard;
-		if (WildcardLength == 1 && (FirstWild == CharType('*') || TargetLength < 2))
-		{
-			return true;
-		}
-		++Wildcard;
-		--WildcardLength;
-
-		// This routine is very slow, though it does ok with one wildcard
-		int32 MaxNum = TargetLength;
-		if (FirstWild == CharType('?') && MaxNum > 1)
-		{
-			MaxNum = 1;
-		}
-
-		for (int32 Index = 0; Index <= MaxNum; ++Index)
-		{
-			if (MatchesWildcardRecursive<CompareType>(Target + Index, TargetLength - Index, Wildcard, WildcardLength))
-			{
-				return true;
-			}
-		}
-		return false;
-	}
-
-	template <typename DestCharType, typename SrcCharType>
-	void AppendCharacters(TArray<DestCharType>& Out, const SrcCharType* Str, int32 Count)
-	{
-		check(Count >= 0);
-
-		if (!Count)
-		{
-			return;
-		}
-
-		checkSlow(Str);
-
-		int32 OldEnd = Out.Num();
-	
-		// Try to reserve enough space by guessing that the new length will be the same as the input length.
-		// Include an extra gap for a null terminator if we don't already have a string allocated
-		Out.AddUninitialized(Count + (OldEnd ? 0 : 1));
-		OldEnd -= OldEnd ? 1 : 0;
-
-		DestCharType* Dest = Out.GetData() + OldEnd;
-
-		// Try copying characters to end of string, overwriting null terminator if we already have one
-		DestCharType* NewEnd = FPlatformString::Convert(Dest, Count, Str, Count);
-		if (!NewEnd)
-		{
-			// If that failed, it will have meant that conversion likely contained multi-code unit characters
-			// and so the buffer wasn't long enough, so calculate it properly.
-			int32 Length = FPlatformString::ConvertedLength<DestCharType>(Str, Count);
-
-			// Add the extra bytes that we need
-			Out.AddUninitialized(Length - Count);
-
-			// Restablish destination pointer in case a realloc happened
-			Dest = Out.GetData() + OldEnd;
-
-			NewEnd = FPlatformString::Convert(Dest, Length, Str, Count);
-			checkSlow(NewEnd);
-		}
-		else
-		{
-			int32 NewEndIndex = (int32)(NewEnd - Dest);
-			if (NewEndIndex < Count)
-			{
-				Out.SetNumUninitialized(OldEnd + NewEndIndex + 1, /*bAllowShrinking=*/false);
-			}
-		}
-
-		// (Re-)establish the null terminator
-		*NewEnd = '\0';
-	}
-
-	template <typename DestCharType, typename SrcCharType>
-	FORCEINLINE void ConstructFromCString(/* Out */ TArray<DestCharType>& Data, const SrcCharType* Src)
-	{
-		if (Src && *Src)
-		{
-			int32 SrcLen  = TCString<SrcCharType>::Strlen(Src) + 1;
-			int32 DestLen = FPlatformString::ConvertedLength<DestCharType>(Src, SrcLen);
-			Data.Reserve(DestLen);
-			Data.AddUninitialized(DestLen);
-
-			FPlatformString::Convert(Data.GetData(), DestLen, Src, SrcLen);
-		}
-	}
-
-	template <typename DestCharType, typename SrcCharType>
-	FORCEINLINE void ConstructWithLength(/* Out */ TArray<DestCharType>& Data, int32 InCount, const SrcCharType* InSrc)
-	{
-		if (InSrc)
-		{
-			int32 DestLen = FPlatformString::ConvertedLength<DestCharType>(InSrc, InCount);
-			if (DestLen > 0 && *InSrc)
-			{
-				Data.Reserve(DestLen + 1);
-				Data.AddUninitialized(DestLen + 1);
-
-				FPlatformString::Convert(Data.GetData(), DestLen, InSrc, InCount);
-				*(Data.GetData() + Data.Num() - 1) = DestCharType('\0');
-			}
-		}
-	}
-
-	template <typename DestCharType, typename SrcCharType>
-	FORCEINLINE void ConstructWithSlack(/* Out */ TArray<DestCharType>& Data, const SrcCharType* Src, int32 ExtraSlack)
-	{
-		if (Src && *Src)
-		{
-			int32 SrcLen = TCString<SrcCharType>::Strlen(Src) + 1;
-			int32 DestLen = FPlatformString::ConvertedLength<DestCharType>(Src, SrcLen);
-			Data.Reserve(DestLen + ExtraSlack);
-			Data.AddUninitialized(DestLen);
-
-			FPlatformString::Convert(Data.GetData(), DestLen, Src, SrcLen);
-		}
-		else if (ExtraSlack > 0)
-		{
-			Data.Reserve(ExtraSlack + 1); 
-		}
-	}
-} // namespace UE::Core::Private
-
-FString::FString(const ANSICHAR* Str)								{ UE::Core::Private::ConstructFromCString(Data, Str); }
-FString::FString(const WIDECHAR* Str)								{ UE::Core::Private::ConstructFromCString(Data, Str); }
-FString::FString(const UTF8CHAR* Str)								{ UE::Core::Private::ConstructFromCString(Data, Str); }
-FString::FString(const UCS2CHAR* Str)								{ UE::Core::Private::ConstructFromCString(Data, Str); }
-FString::FString(int32 Len, const ANSICHAR* Str)					{ UE::Core::Private::ConstructWithLength(Data, Len, Str); }
-FString::FString(int32 Len, const WIDECHAR* Str)					{ UE::Core::Private::ConstructWithLength(Data, Len, Str); }
-FString::FString(int32 Len, const UTF8CHAR* Str)					{ UE::Core::Private::ConstructWithLength(Data, Len, Str); }
-FString::FString(int32 Len, const UCS2CHAR* Str)					{ UE::Core::Private::ConstructWithLength(Data, Len, Str); }
-FString::FString(const ANSICHAR* Str, int32 ExtraSlack)				{ UE::Core::Private::ConstructWithSlack(Data, Str, ExtraSlack); }
-FString::FString(const WIDECHAR* Str, int32 ExtraSlack)				{ UE::Core::Private::ConstructWithSlack(Data, Str, ExtraSlack); }
-FString::FString(const UTF8CHAR* Str, int32 ExtraSlack)				{ UE::Core::Private::ConstructWithSlack(Data, Str, ExtraSlack); }
-FString::FString(const UCS2CHAR* Str, int32 ExtraSlack)				{ UE::Core::Private::ConstructWithSlack(Data, Str, ExtraSlack); }
-
-FString& FString::operator=( const ElementType* Other )
-{
-	if (Data.GetData() != Other)
-	{
-		int32 Len = (Other && *Other) ? TCString<ElementType>::Strlen(Other)+1 : 0;
-		Data.Empty(Len);
-		Data.AddUninitialized(Len);
-			
-		if( Len )
-		{
-			FMemory::Memcpy( Data.GetData(), Other, Len * sizeof(ElementType) );
-		}
-	}
-	return *this;
-}
-
-
-void FString::AssignRange(const ElementType* OtherData, int32 OtherLen)
-{
-	if (OtherLen == 0)
-	{
-		Empty();
-	}
-	else
-	{
-		const int32 ThisLen = Len();
-		if (OtherLen <= ThisLen)
-		{
-			// Unless the input is longer, this might be assigned from a view of itself.
-			ElementType* DataPtr = Data.GetData();
-			FMemory::Memmove(DataPtr, OtherData, OtherLen * sizeof(ElementType));
-			DataPtr[OtherLen] = CHARTEXT(ElementType, '\0');
-			Data.RemoveAt(OtherLen + 1, ThisLen - OtherLen);
-		}
-		else
-		{
-			Data.Empty(OtherLen + 1);
-			Data.AddUninitialized(OtherLen + 1);
-			ElementType* DataPtr = Data.GetData();
-			FMemory::Memcpy(DataPtr, OtherData, OtherLen * sizeof(ElementType));
-			DataPtr[OtherLen] = CHARTEXT(ElementType, '\0');
-		}
-	}
-}
-
-void FString::Reserve(int32 CharacterCount)
-{
-	checkSlow(CharacterCount >= 0 && CharacterCount < MAX_int32);
-	if (CharacterCount > 0)
-	{
-		Data.Reserve(CharacterCount + 1);
-	}	
-}
-
-void FString::Empty(int32 Slack)
-{
-	Data.Empty(Slack ? Slack + 1 : 0);
-}
-
-void FString::Empty()
-{
-	Data.Empty(0);
-}
-
-void FString::Reset(int32 NewReservedSize)
-{
-	const int32 NewSizeIncludingTerminator = (NewReservedSize > 0) ? (NewReservedSize + 1) : 0;
-	Data.Reset(NewSizeIncludingTerminator);
-	if (ElementType* DataPtr = Data.GetData())
-	{
-		*DataPtr = CHARTEXT(ElementType, '\0');
-	}
-}
-
-void FString::Shrink()
-{
-	Data.Shrink();
-}
-
-#ifdef __OBJC__
-/** Convert FString to Objective-C NSString */
-NSString* FString::GetNSString() const
-{
-#if PLATFORM_TCHAR_IS_4_BYTES
-    return [[[NSString alloc] initWithBytes:Data.GetData() length:Len() * sizeof(ElementType) encoding:NSUTF32LittleEndianStringEncoding] autorelease];
-#else
-    return [[[NSString alloc] initWithBytes:Data.GetData() length:Len() * sizeof(ElementType) encoding:NSUTF16LittleEndianStringEncoding] autorelease];
-#endif
-}
-#endif
-
-FString& FString::AppendChar(ElementType InChar)
-{
-	CheckInvariants();
-
-	if ( InChar != 0 )
-	{
-		// position to insert the character.  
-		// At the end of the string if we have existing characters, otherwise at the 0 position
-		int32 InsertIndex = (Data.Num() > 0) ? Data.Num()-1 : 0;	
-
-		// number of characters to add.  If we don't have any existing characters, 
-		// we'll need to append the terminating zero as well.
-		int32 InsertCount = (Data.Num() > 0) ? 1 : 2;
-
-		Data.AddUninitialized(InsertCount);
-		Data[InsertIndex] = InChar;
-		Data[InsertIndex+1] = CHARTEXT(ElementType, '\0');
-	}
-	return *this;
-}
-
-void FString::AppendChars(const ANSICHAR* Str, int32 Count)
-{
-	CheckInvariants();
-	UE::Core::Private::AppendCharacters(Data, Str, Count);
-}
-
-void FString::AppendChars(const WIDECHAR* Str, int32 Count)
-{
-	CheckInvariants();
-	UE::Core::Private::AppendCharacters(Data, Str, Count);
-}
-
-void FString::AppendChars(const UCS2CHAR* Str, int32 Count)
-{
-	CheckInvariants();
-	UE::Core::Private::AppendCharacters(Data, Str, Count);
-}
-
-void FString::AppendChars(const UTF8CHAR* Str, int32 Count)
-{
-	CheckInvariants();
-	UE::Core::Private::AppendCharacters(Data, Str, Count);
-}
-
-void FString::TrimToNullTerminator()
-{
-	if( Data.Num() )
-	{
-		int32 DataLen = TCString<ElementType>::Strlen(Data.GetData());
-		check(DataLen == 0 || DataLen < Data.Num());
-		int32 Len = DataLen > 0 ? DataLen+1 : 0;
-
-		check(Len <= Data.Num());
-		Data.RemoveAt(Len, Data.Num()-Len);
-	}
-}
-
-
-int32 FString::Find(const ElementType* SubStr, int32 SubStrLen, ESearchCase::Type SearchCase, ESearchDir::Type SearchDir, int32 StartPosition) const
-{
-	checkf(SubStrLen >= 0, TEXT("Invalid SubStrLen: %d"), SubStrLen);
-
-	if (SearchDir == ESearchDir::FromStart)
-	{
-		const ElementType* Start = **this;
-		int32 RemainingLength = Len();
-		if (StartPosition != INDEX_NONE && RemainingLength > 0)
-		{
-			const ElementType* End = Start + RemainingLength;
-			Start += FMath::Clamp(StartPosition, 0, RemainingLength - 1);
-			RemainingLength = UE_PTRDIFF_TO_INT32(End - Start);
-		}
-		const ElementType* Tmp = SearchCase == ESearchCase::IgnoreCase
-			? TCString<ElementType>::Strnistr(Start, RemainingLength, SubStr, SubStrLen)
-			: TCString<ElementType>::Strnstr(Start, RemainingLength, SubStr, SubStrLen);
-
-		return Tmp ? UE_PTRDIFF_TO_INT32(Tmp-**this) : INDEX_NONE;
-	}
-	else
-	{
-		// if ignoring, do a onetime ToUpper on both strings, to avoid ToUppering multiple
-		// times in the loop below
-		if ( SearchCase == ESearchCase::IgnoreCase)
-		{
-			return ToUpper().Find(FString(SubStrLen, SubStr).ToUpper(), ESearchCase::CaseSensitive, SearchDir, StartPosition);
-		}
-		else
-		{
-			const int32 SearchStringLength=FMath::Max(1, SubStrLen);
-			
-			if (StartPosition == INDEX_NONE || StartPosition >= Len())
-			{
-				StartPosition = Len();
-			}
-			
-			for (int32 i = StartPosition - SearchStringLength; i >= 0; i--)
-			{
-				int32 j;
-				for (j=0; j != SubStrLen; j++)
-				{
-					if ((*this)[i+j]!=SubStr[j])
-					{
-						break;
-					}
-				}
-				
-				if (j == SubStrLen)
-				{
-					return i;
-				}
-			}
-			return INDEX_NONE;
-		}
-	}
-}
-
-bool FString::Split(const FString& InS, FString* LeftS, FString* RightS, ESearchCase::Type SearchCase, ESearchDir::Type SearchDir) const
-{
-	check(LeftS != RightS || LeftS == nullptr);
-
-	int32 InPos = Find(InS, SearchCase, SearchDir);
-
-	if (InPos < 0) { return false; }
-
-	if (LeftS)
-	{
-		if (LeftS != this)
-		{
-			*LeftS = Left(InPos);
-			if (RightS) { *RightS = Mid(InPos + InS.Len()); }
-		}
-		else
-		{
-			// we know that RightS can't be this so we can safely modify it before we deal with LeftS
-			if (RightS) { *RightS = Mid(InPos + InS.Len()); }
-			*LeftS = Left(InPos);
-		}
-	}
-	else if (RightS)
-	{
-		*RightS = Mid(InPos + InS.Len());
-	}
-
-	return true;
-}
-
-bool FString::Split(const FString& InS, FString* LeftS, FString* RightS) const
-{
-	return Split(InS, LeftS, RightS, ESearchCase::IgnoreCase);
-}
-
-FString FString::ToUpper() const &
-{
-	FString New = *this;
-	New.ToUpperInline();
-	return New;
-}
-
-FString FString::ToUpper() &&
-{
-	this->ToUpperInline();
-	return MoveTemp(*this);
-}
-
-void FString::ToUpperInline()
-{
-	const int32 StringLength = Len();
-	ElementType* RawData = Data.GetData();
-	for (int32 i = 0; i < StringLength; ++i)
-	{
-		RawData[i] = TChar<ElementType>::ToUpper(RawData[i]);
-	}
-}
-
-
-FString FString::ToLower() const &
-{
-	FString New = *this;
-	New.ToLowerInline();
-	return New;
-}
-
-FString FString::ToLower() &&
-{
-	this->ToLowerInline();
-	return MoveTemp(*this);
-}
-
-void FString::ToLowerInline()
-{
-	const int32 StringLength = Len();
-	ElementType* RawData = Data.GetData();
-	for (int32 i = 0; i < StringLength; ++i)
-	{
-		RawData[i] = TChar<ElementType>::ToLower(RawData[i]);
-	}
-}
-
-void FString::RemoveSpacesInline()
-{
-	const int32 StringLength = Len();
-	if (StringLength == 0)
-	{
-		return;
-	}
-
-	ElementType* RawData = Data.GetData();
-	int32 CopyToIndex = 0;
-	for (int32 CopyFromIndex = 0; CopyFromIndex < StringLength; ++CopyFromIndex)
-	{
-		if (RawData[CopyFromIndex] != ' ')
-		{	// Copy any character OTHER than space.
-			RawData[CopyToIndex] = RawData[CopyFromIndex];
-			++CopyToIndex;
-		}
-	}
-
-	// Copy null-terminating character.
-	if (CopyToIndex <= StringLength)
-	{
-		RawData[CopyToIndex] = CHARTEXT(ElementType, '\0');
-		Data.SetNum(CopyToIndex + 1, false);
-	}
-}
-
-bool FString::StartsWith(const ElementType* InPrefix, int32 InPrefixLen, ESearchCase::Type SearchCase) const
-{
-	if (SearchCase == ESearchCase::IgnoreCase)
-	{
-		return InPrefixLen > 0 && !TCString<ElementType>::Strnicmp(**this, InPrefix, InPrefixLen);
-	}
-	else
-	{
-		return InPrefixLen > 0 && !TCString<ElementType>::Strncmp(**this, InPrefix, InPrefixLen);
-	}
-}
-
-bool FString::EndsWith(const ElementType* InSuffix, int32 InSuffixLen, ESearchCase::Type SearchCase ) const
-{
-	if (SearchCase == ESearchCase::IgnoreCase)
-	{
-		return InSuffixLen > 0 &&
-			Len() >= InSuffixLen &&
-			!TCString<ElementType>::Stricmp( &(*this)[ Len() - InSuffixLen ], InSuffix );
-	}
-	else
-	{
-		return InSuffixLen > 0 &&
-			Len() >= InSuffixLen &&
-			!TCString<ElementType>::Strcmp( &(*this)[ Len() - InSuffixLen ], InSuffix );
-	}
-}
-
-void FString::InsertAt(int32 Index, ElementType Character)
-{
-	if (Character != 0)
-	{
-		if (Data.Num() == 0)
-		{
-			*this += Character;
-		}
-		else
-		{
-			Data.Insert(Character, Index);
-		}
-	}
-}
-
-void FString::InsertAt(int32 Index, const FString& Characters)
-{
-	if (Characters.Len())
-	{
-		if (Data.Num() == 0)
-		{
-			*this += Characters;
-		}
-		else
-		{
-			Data.Insert(Characters.Data.GetData(), Characters.Len(), Index);
-		}
-	}
-}
-
-void FString::RemoveAt(int32 Index, int32 Count, bool bAllowShrinking)
-{
-	Data.RemoveAt(Index, FMath::Clamp(Count, 0, Len()-Index), bAllowShrinking);
-}
-
-bool FString::RemoveFromStart(const ElementType* InPrefix, int32 InPrefixLen, ESearchCase::Type SearchCase)
-{
-	if (InPrefixLen == 0 )
-	{
-		return false;
-	}
-
-	if (StartsWith(InPrefix, InPrefixLen, SearchCase ))
-	{
-		RemoveAt(0, InPrefixLen);
-		return true;
-	}
-
-	return false;
-}
-
-bool FString::RemoveFromEnd(const ElementType* InSuffix, int32 InSuffixLen, ESearchCase::Type SearchCase)
-{
-	if (InSuffixLen == 0)
-	{
-		return false;
-	}
-
-	if (EndsWith(InSuffix, InSuffixLen, SearchCase))
-	{
-		RemoveAt( Len() - InSuffixLen, InSuffixLen );
-		return true;
-	}
-
-	return false;
-}
-
-namespace UE::String::Private
-{
-
-template <typename LhsType, typename RhsType>
-UE_NODISCARD FORCEINLINE FString ConcatFStrings(LhsType&& Lhs, RhsType&& Rhs)
-{
-	Lhs.CheckInvariants();
-	Rhs.CheckInvariants();
-
-	if (Lhs.IsEmpty())
-	{
-		return Forward<RhsType>(Rhs);
-	}
-
-	int32 RhsLen = Rhs.Len();
-
-	FString Result(Forward<LhsType>(Lhs), /* extra slack */ RhsLen);
-	Result.AppendChars(Rhs.GetCharArray().GetData(), RhsLen);
-		
-	return Result;
-}
-
-template <typename LhsCharType, typename RhsType>
-UE_NODISCARD FORCEINLINE FString ConcatRangeFString(const LhsCharType* Lhs, int32 LhsLen, RhsType&& Rhs)
-{
-	using ElementType = FString::ElementType;
-
-	checkSlow(LhsLen >= 0);
-	Rhs.CheckInvariants();
-	if (LhsLen == 0)
-	{
-		return Forward<RhsType>(Rhs);
-	}
-
-	int32 RhsLen = Rhs.Len();
-
-	// This is not entirely optimal, as if the Rhs is an rvalue and has enough slack space to hold Lhs, then
-	// the memory could be reused here without constructing a new object.  However, until there is proof otherwise,
-	// I believe this will be relatively rare and isn't worth making the code a lot more complex right now.
-	FString Result;
-	Result.GetCharArray().Reserve(LhsLen + RhsLen + 1);
-	Result.GetCharArray().AddUninitialized(LhsLen + RhsLen + 1);
-
-	ElementType* ResultData = Result.GetCharArray().GetData();
-	CopyAssignItems(ResultData, Lhs, LhsLen);
-	CopyAssignItems(ResultData + LhsLen, Rhs.GetCharArray().GetData(), RhsLen);
-	*(ResultData + LhsLen + RhsLen) = CHARTEXT(ElementType, '\0');
-
-	return Result;
-}
-
-template <typename LhsType, typename RhsCharType>
-UE_NODISCARD FORCEINLINE FString ConcatFStringRange(LhsType&& Lhs, const RhsCharType* Rhs, int32 RhsLen)
-{
-	Lhs.CheckInvariants();
-	checkSlow(RhsLen >= 0);
-	if (RhsLen == 0)
-	{
-		return Forward<LhsType>(Lhs);
-	}
-
-	FString Result(Forward<LhsType>(Lhs), /* extra slack */ RhsLen);
-	Result.AppendChars(Rhs, RhsLen);
-		
-	return Result;
-}
-
-template <typename LhsCharType, typename RhsType>
-UE_NODISCARD FORCEINLINE FString ConcatCStringFString(const LhsCharType* Lhs, RhsType&& Rhs)
-{
-	checkSlow(Lhs);
-	if (!Lhs)
-	{
-		return Forward<RhsType>(Rhs);
-	}
-
-	return ConcatRangeFString(Lhs, TCString<LhsCharType>::Strlen(Lhs), Forward<RhsType>(Rhs));
-}
-
-template <typename LhsType, typename RhsCharType>
-UE_NODISCARD FORCEINLINE FString ConcatFStringCString(LhsType&& Lhs, const RhsCharType* Rhs)
-{
-	checkSlow(Rhs);
-	if (!Rhs)
-	{
-		return Forward<LhsType>(Lhs);
-	}
-	return ConcatFStringRange(Forward<LhsType>(Lhs), Rhs, TCString<RhsCharType>::Strlen(Rhs));
-}
-
-} // namespace UE::String::Private
-
-FString FString::ConcatFF(const FString& Lhs, const FString& Rhs)						{ return UE::String::Private::ConcatFStrings(Lhs, Rhs); }
-FString FString::ConcatFF(FString&& Lhs, const FString& Rhs)							{ return UE::String::Private::ConcatFStrings(MoveTemp(Lhs), Rhs); }
-FString FString::ConcatFF(const FString& Lhs, FString&& Rhs)							{ return UE::String::Private::ConcatFStrings(Lhs, MoveTemp(Rhs)); }
-FString FString::ConcatFF(FString&& Lhs, FString&& Rhs)									{ return UE::String::Private::ConcatFStrings(MoveTemp(Lhs), MoveTemp(Rhs)); }
-FString FString::ConcatFC(const FString& Lhs, const ElementType* Rhs)					{ return UE::String::Private::ConcatFStringCString(Lhs, Rhs); }
-FString FString::ConcatFC(FString&& Lhs, const ElementType* Rhs)						{ return UE::String::Private::ConcatFStringCString(MoveTemp(Lhs), Rhs); }
-FString FString::ConcatCF(const ElementType* Lhs,	const FString& Rhs)					{ return UE::String::Private::ConcatCStringFString(Lhs, Rhs); }
-FString FString::ConcatCF(const ElementType* Lhs,	FString&& Rhs)						{ return UE::String::Private::ConcatCStringFString(Lhs, MoveTemp(Rhs)); }
-FString FString::ConcatFR(const FString& Lhs, const ElementType* Rhs, int32 RhsLen)		{ return UE::String::Private::ConcatFStringRange(Lhs, Rhs, RhsLen); }
-FString FString::ConcatFR(FString&& Lhs, const ElementType* Rhs, int32 RhsLen)			{ return UE::String::Private::ConcatFStringRange(MoveTemp(Lhs), Rhs, RhsLen); }
-FString FString::ConcatRF(const ElementType* Lhs, int32 LhsLen, const FString& Rhs)		{ return UE::String::Private::ConcatRangeFString(Lhs, LhsLen, Rhs); }
-FString FString::ConcatRF(const ElementType* Lhs, int32 LhsLen, FString&& Rhs)			{ return UE::String::Private::ConcatRangeFString(Lhs, LhsLen, MoveTemp(Rhs)); }
-
-/**
- * Concatenate this path with given path ensuring the / character is used between them
- *
- * @param Str       Pointer to an array of TCHARs (not necessarily null-terminated) to be concatenated onto the end of this.
- * @param StrLength Exact number of characters from Str to append.
- */
-void FString::PathAppend(const ElementType* Str, int32 StrLength)
-{
-	int32 DataNum = Data.Num();
-	if (StrLength == 0)
-	{
-		if (DataNum > 1 && Data[DataNum - 2] != CHARTEXT(ElementType, '/') && Data[DataNum - 2] != CHARTEXT(ElementType, '\\'))
-		{
-			Data[DataNum - 1] = CHARTEXT(ElementType, '/');
-			Data.Add(CHARTEXT(ElementType, '\0'));
-		}
-	}
-	else
-	{
-		if (DataNum > 0)
-		{
-			if (DataNum > 1 && Data[DataNum - 2] != CHARTEXT(ElementType, '/') && Data[DataNum - 2] != CHARTEXT(ElementType, '\\') && *Str != CHARTEXT(ElementType, '/'))
-			{
-				Data[DataNum - 1] = CHARTEXT(ElementType, '/');
-			}
-			else
-			{
-				Data.Pop(false);
-				--DataNum;
-			}
-		}
-
-		Reserve(DataNum + StrLength);
-		Data.Append(Str, StrLength);
-		Data.Add(CHARTEXT(ElementType, '\0'));
-	}
-}
-
-FString FString::RightChop(int32 Count) const &
-{
-	const int32 Length = Len();
-	const int32 Skip = FMath::Clamp(Count, 0, Length);
-	return FString(Length - Skip, **this + Skip);
-}
-
-FString FString::Mid(int32 Start, int32 Count) const &
-{
-	if (Count >= 0)
-	{
-		const int32 Length = Len();
-		const int32 RequestedStart = Start;
-		Start = FMath::Clamp(Start, 0, Length);
-		const int32 End = (int32)FMath::Clamp((int64)Count + RequestedStart, (int64)Start, (int64)Length);
-		return FString(End-Start, **this + Start);
-	}
-
-	return FString();
-}
-
-FString FString::Mid(int32 Start, int32 Count) &&
-{
-	MidInline(Start, Count, false);
-	return MoveTemp(*this);
-}
-
-void FString::ReplaceCharInlineCaseSensitive(const ElementType SearchChar, const ElementType ReplacementChar)
-{
-	for (ElementType& Character : Data)
-	{
-		Character = Character == SearchChar ? ReplacementChar : Character;
-	}
-}
-
-void FString::ReplaceCharInlineIgnoreCase(const ElementType SearchChar, const ElementType ReplacementChar)
-{
-	ElementType OtherCaseSearchChar = TChar<ElementType>::IsUpper(SearchChar) ? TChar<ElementType>::ToLower(SearchChar) : TChar<ElementType>::ToUpper(SearchChar);
-	ReplaceCharInlineCaseSensitive(OtherCaseSearchChar, ReplacementChar);
-	ReplaceCharInlineCaseSensitive(SearchChar, ReplacementChar);
-}
-
-void FString::TrimStartAndEndInline()
-{
-	TrimEndInline();
-	TrimStartInline();
-}
-
-FString FString::TrimStartAndEnd() const &
-{
-	FString Result(*this);
-	Result.TrimStartAndEndInline();
-	return Result;
-}
-
-FString FString::TrimStartAndEnd() &&
-{
-	TrimStartAndEndInline();
-	return MoveTemp(*this);
-}
-
-void FString::TrimStartInline()
-{
-	int32 Pos = 0;
-	while(Pos < Len() && TChar<ElementType>::IsWhitespace((*this)[Pos]))
-	{
-		Pos++;
-	}
-	RemoveAt(0, Pos);
-}
-
-FString FString::TrimStart() const &
-{
-	FString Result(*this);
-	Result.TrimStartInline();
-	return Result;
-}
-
-FString FString::TrimStart() &&
-{
-	TrimStartInline();
-	return MoveTemp(*this);
-}
-
-void FString::TrimEndInline()
-{
-	int32 End = Len();
-	while(End > 0 && TChar<ElementType>::IsWhitespace((*this)[End - 1]))
-	{
-		End--;
-	}
-	RemoveAt(End, Len() - End);
-}
-
-FString FString::TrimEnd() const &
-{
-	FString Result(*this);
-	Result.TrimEndInline();
-	return Result;
-}
-
-FString FString::TrimEnd() &&
-{
-	TrimEndInline();
-	return MoveTemp(*this);
-}
-
-void FString::TrimCharInline(const ElementType CharacterToTrim, bool* bCharRemoved)
-{
-	bool bQuotesWereRemoved=false;
-	int32 Start = 0, Count = Len();
-	if ( Count > 0 )
-	{
-		if ( (*this)[0] == CharacterToTrim )
-		{
-			Start++;
-			Count--;
-			bQuotesWereRemoved=true;
-		}
-
-		if ( Len() > 1 && (*this)[Len() - 1] == CharacterToTrim )
-		{
-			Count--;
-			bQuotesWereRemoved=true;
-		}
-	}
-
-	if ( bCharRemoved != nullptr )
-	{
-		*bCharRemoved = bQuotesWereRemoved;
-	}
-	MidInline(Start, Count, false);
-}
-
-void FString::TrimQuotesInline(bool* bQuotesRemoved)
-{
-	TrimCharInline(ElementType('"'), bQuotesRemoved);
-}
-
-FString FString::TrimQuotes(bool* bQuotesRemoved) const &
-{
-	FString Result(*this);
-	Result.TrimQuotesInline(bQuotesRemoved);
-	return Result;
-}
-
-FString FString::TrimQuotes(bool* bQuotesRemoved) &&
-{
-	TrimQuotesInline(bQuotesRemoved);
-	return MoveTemp(*this);
-}
-
-FString FString::TrimChar(const ElementType CharacterToTrim, bool* bCharRemoved) const &
-{
-	FString Result(*this);
-	Result.TrimCharInline(CharacterToTrim, bCharRemoved);
-	return Result;
-}
-
-FString FString::TrimChar(const ElementType CharacterToTrim, bool* bCharRemoved) &&
-{
-	TrimCharInline(CharacterToTrim, bCharRemoved);
-	return MoveTemp(*this);
-}
-
-int32 FString::CullArray( TArray<FString>* InArray )
-{
-	check(InArray);
-	FString Empty;
-	InArray->Remove(Empty);
-	return InArray->Num();
-}
-
-FString FString::Reverse() const &
-{
-	FString New(*this);
-	New.ReverseString();
-	return New;
-}
-
-FString FString::Reverse() &&
-{
-	ReverseString();
-	return MoveTemp(*this);
-}
-
-void FString::ReverseString()
-{
-	if ( Len() > 0 )
-	{
-		ElementType* StartChar = &(*this)[0];
-		ElementType* EndChar = &(*this)[Len()-1];
-		ElementType TempChar;
-		do 
-		{
-			TempChar = *StartChar;	// store the current value of StartChar
-			*StartChar = *EndChar;	// change the value of StartChar to the value of EndChar
-			*EndChar = TempChar;	// change the value of EndChar to the character that was previously at StartChar
-
-			StartChar++;
-			EndChar--;
-
-		} while( StartChar < EndChar );	// repeat until we've reached the midpoint of the string
-	}
-}
-
-FString FString::FormatAsNumber( int32 InNumber )
-{
-	FString Number = FString::FromInt( InNumber ), Result;
-
-	int32 dec = 0;
-	for( int32 x = Number.Len()-1 ; x > -1 ; --x )
-	{
-		Result += Number.Mid(x,1);
-
-		dec++;
-		if( dec == 3 && x > 0 )
-		{
-			Result += CHARTEXT(ElementType, ',');
-			dec = 0;
-		}
-	}
-
-	return Result.Reverse();
-}
-
-/**
- * Serializes a string as ANSI char array.
- *
- * @param	String			String to serialize
- * @param	Ar				Archive to serialize with
- * @param	MinCharacters	Minimum number of characters to serialize.
- */
-void FString::SerializeAsANSICharArray( FArchive& Ar, int32 MinCharacters ) const
-{
-	int32	Length = FMath::Max( Len(), MinCharacters );
-	Ar << Length;
-	
-	for( int32 CharIndex=0; CharIndex<Len(); CharIndex++ )
-	{
-		ANSICHAR AnsiChar = CharCast<ANSICHAR>( (*this)[CharIndex] );
-		Ar << AnsiChar;
-	}
-
-	// Zero pad till minimum number of characters are written.
-	for( int32 i=Len(); i<Length; i++ )
-	{
-		ANSICHAR NullChar = 0;
-		Ar << NullChar;
-	}
-}
-
-void FString::AppendInt( int32 Num )
-{
-	const ElementType* DigitToChar	= CHARTEXT(ElementType, "9876543210123456789");
-	constexpr int32 ZeroDigitIndex	= 9;
-	bool bIsNumberNegative			= Num < 0;
-	const int32 TempBufferSize		= 16; // 16 is big enough
-	ElementType TempNum[TempBufferSize];
-	int32 TempAt					= TempBufferSize; // fill the temp string from the top down.
-
-	// Convert to string assuming base ten.
-	do 
-	{
-		TempNum[--TempAt] = DigitToChar[ZeroDigitIndex + (Num % 10)];
-		Num /= 10;
-	} while( Num );
-
-	if( bIsNumberNegative )
-	{
-		TempNum[--TempAt] = CHARTEXT(ElementType, '-');
-	}
-
-	const ElementType* CharPtr = TempNum + TempAt;
-	const int32 NumChars = TempBufferSize - TempAt;
-	Append(CharPtr, NumChars);
-}
-
-
-bool FString::ToBool() const
-{
-	return TCString<ElementType>::ToBool(**this);
-}
-
-FString FString::FromBlob(const uint8* SrcBuffer,const uint32 SrcSize)
-{
-	FString Result;
-	Result.Reserve( SrcSize * 3 );
-	// Convert and append each byte in the buffer
-	for (uint32 Count = 0; Count < SrcSize; Count++)
-	{
-		Result += FString::Printf(CHARTEXT(ElementType, "%03d"),(uint8)SrcBuffer[Count]);
-	}
-	return Result;
-}
-
-bool FString::ToBlob(const FString& Source,uint8* DestBuffer,const uint32 DestSize)
-{
-	// Make sure the buffer is at least half the size and that the string is an
-	// even number of characters long
-	if (DestSize >= (uint32)(Source.Len() / 3) &&
-		(Source.Len() % 3) == 0)
-	{
-		ElementType ConvBuffer[4];
-		ConvBuffer[3] = CHARTEXT(ElementType, '\0');
-		int32 WriteIndex = 0;
-		// Walk the string 3 chars at a time
-		for (int32 Index = 0; Index < Source.Len(); Index += 3, WriteIndex++)
-		{
-			ConvBuffer[0] = Source[Index];
-			ConvBuffer[1] = Source[Index + 1];
-			ConvBuffer[2] = Source[Index + 2];
-			DestBuffer[WriteIndex] = (uint8)TCString<ElementType>::Atoi(ConvBuffer);
-		}
-		return true;
-	}
-	return false;
-}
-
-FString FString::FromHexBlob( const uint8* SrcBuffer, const uint32 SrcSize )
-{
-	FString Result;
-	Result.Reserve( SrcSize * 2 );
-	// Convert and append each byte in the buffer
-	for (uint32 Count = 0; Count < SrcSize; Count++)
-	{
-		Result += FString::Printf( CHARTEXT(ElementType,  "%02X" ), (uint8)SrcBuffer[Count] );
-	}
-	return Result;
-}
-
-bool FString::ToHexBlob( const FString& Source, uint8* DestBuffer, const uint32 DestSize )
-{
-	// Make sure the buffer is at least half the size and that the string is an
-	// even number of characters long
-	if (DestSize >= (uint32)(Source.Len() / 2) &&
-		 (Source.Len() % 2) == 0)
-	{
-		ElementType ConvBuffer[3];
-		ConvBuffer[2] = CHARTEXT(ElementType, '\0' );
-		int32 WriteIndex = 0;
-		// Walk the string 2 chars at a time
-		ElementType* End = nullptr;
-		for (int32 Index = 0; Index < Source.Len(); Index += 2, WriteIndex++)
-		{
-			ConvBuffer[0] = Source[Index];
-			ConvBuffer[1] = Source[Index + 1];
-			DestBuffer[WriteIndex] = (uint8)TCString<ElementType>::Strtoi( ConvBuffer, &End, 16 );
-		}
-		return true;
-	}
-	return false;
-}
-
-UE_DISABLE_OPTIMIZATION_SHIP
-void StripNegativeZero(double& InFloat)
-{
-	// This works for translating a negative zero into a positive zero,
-	// but if optimizations are enabled when compiling with -ffast-math
-	// or /fp:fast, the compiler can strip it out.
-	InFloat += 0.0f;
-}
-UE_ENABLE_OPTIMIZATION_SHIP
-
-FString FString::SanitizeFloat( double InFloat, const int32 InMinFractionalDigits )
-{
-	// Avoids negative zero
-	StripNegativeZero(InFloat);
-
-	// First create the string
-	FString TempString = FString::Printf(CHARTEXT(ElementType, "%f"), InFloat);
-	if (!TempString.IsNumeric())
-	{
-		// String did not format as a valid decimal number so avoid messing with it
-		return TempString;
-	}
-
-	// Trim all trailing zeros (up-to and including the decimal separator) from the fractional part of the number
-	int32 TrimIndex = INDEX_NONE;
-	int32 DecimalSeparatorIndex = INDEX_NONE;
-	for (int32 CharIndex = TempString.Len() - 1; CharIndex >= 0; --CharIndex)
-	{
-		const ElementType Char = TempString[CharIndex];
-		if (Char == CHARTEXT(ElementType, '.'))
-		{
-			DecimalSeparatorIndex = CharIndex;
-			TrimIndex = FMath::Max(TrimIndex, DecimalSeparatorIndex);
-			break;
-		}
-		if (TrimIndex == INDEX_NONE && Char != CHARTEXT(ElementType, '0'))
-		{
-			TrimIndex = CharIndex + 1;
-		}
-	}
-	check(TrimIndex != INDEX_NONE && DecimalSeparatorIndex != INDEX_NONE);
-	TempString.RemoveAt(TrimIndex, TempString.Len() - TrimIndex, /*bAllowShrinking*/false);
-
-	// Pad the number back to the minimum number of fractional digits
-	if (InMinFractionalDigits > 0)
-	{
-		if (TrimIndex == DecimalSeparatorIndex)
-		{
-			// Re-add the decimal separator
-			TempString.AppendChar(CHARTEXT(ElementType, '.'));
-		}
-
-		const int32 NumFractionalDigits = (TempString.Len() - DecimalSeparatorIndex) - 1;
-		const int32 FractionalDigitsToPad = InMinFractionalDigits - NumFractionalDigits;
-		if (FractionalDigitsToPad > 0)
-		{
-			TempString.Reserve(TempString.Len() + FractionalDigitsToPad);
-			for (int32 Cx = 0; Cx < FractionalDigitsToPad; ++Cx)
-			{
-				TempString.AppendChar(CHARTEXT(ElementType, '0'));
-			}
-		}
-	}
-
-	return TempString;
-}
-
-FString FString::Chr(ElementType Ch)
-{
-	ElementType Temp[2]= { Ch, CHARTEXT(ElementType, '\0') };
-	return FString(Temp);
-}
-
-
-FString FString::ChrN( int32 NumCharacters, ElementType Char )
-{
-	check( NumCharacters >= 0 );
-
-	FString Temp;
-	Temp.Data.AddUninitialized(NumCharacters+1);
-	for( int32 Cx = 0; Cx < NumCharacters; ++Cx )
-	{
-		Temp[Cx] = Char;
-	}
-	Temp.Data[NumCharacters] = CHARTEXT(ElementType, '\0');
-	return Temp;
-}
-
-FString FString::LeftPad( int32 ChCount ) const
-{
-	int32 Pad = ChCount - Len();
-
-	if (Pad > 0)
-	{
-		return ChrN(Pad, CHARTEXT(ElementType, ' ')) + *this;
-	}
-	else
-	{
-		return *this;
-	}
-}
-FString FString::RightPad( int32 ChCount ) const
-{
-	int32 Pad = ChCount - Len();
-
-	if (Pad > 0)
-	{
-		return *this + ChrN(Pad, CHARTEXT(ElementType, ' '));
-	}
-	else
-	{
-		return *this;
-	}
-}
-
-bool FString::IsNumeric() const
-{
-	if (IsEmpty())
-	{
-		return 0;
-	}
-
-	return TCString<ElementType>::IsNumeric(Data.GetData());
-}
-
-int32 FString::ParseIntoArray( TArray<FString>& OutArray, const ElementType* pchDelim, const bool InCullEmpty ) const
-{
-	check(pchDelim);
-	OutArray.Reset();
-	// TODO Legacy behavior: if DelimLength is 0 we return an empty array. Can we change this to return { Text }?
-	if (pchDelim[0] != '\0')
-	{
-		UE::String::EParseTokensOptions ParseOptions = UE::String::EParseTokensOptions::IgnoreCase |
-			(InCullEmpty ? UE::String::EParseTokensOptions::SkipEmpty : UE::String::EParseTokensOptions::None);
-		UE::String::ParseTokens(FStringView(*this), FStringView(pchDelim),
-			[&OutArray](FStringView Token) { OutArray.Emplace(Token); },
-			ParseOptions);
-	}
-	return OutArray.Num();
-}
-
-bool FString::MatchesWildcard(const ElementType* InWildcard, int32 InWildcardLen, ESearchCase::Type SearchCase) const
-{
-	const ElementType* Target = **this;
-	int32        TargetLength = Len();
-
-	if (SearchCase == ESearchCase::CaseSensitive)
-	{
-		return UE::Core::Private::MatchesWildcardRecursive<UE::Core::Private::TCompareCharsCaseSensitive<ElementType>>(Target, TargetLength, InWildcard, InWildcardLen);
-	}
-	else
-	{
-		return UE::Core::Private::MatchesWildcardRecursive<UE::Core::Private::TCompareCharsCaseInsensitive<ElementType>>(Target, TargetLength, InWildcard, InWildcardLen);
-	}
-}
-
-
-/** Caution!! this routine is O(N^2) allocations...use it for parsing very short text or not at all */
-int32 FString::ParseIntoArrayWS( TArray<FString>& OutArray, const ElementType* pchExtraDelim, bool InCullEmpty ) const
-{
-	// default array of White Spaces, the last entry can be replaced with the optional pchExtraDelim string
-	// (if you want to split on white space and another character)
-	const ElementType* WhiteSpace[] =
-	{
-		CHARTEXT(ElementType, " "),
-		CHARTEXT(ElementType, "\t"),
-		CHARTEXT(ElementType, "\r"),
-		CHARTEXT(ElementType, "\n"),
-		CHARTEXT(ElementType, ""),
-	};
-
-	// start with just the standard whitespaces
-	int32 NumWhiteSpaces = UE_ARRAY_COUNT(WhiteSpace) - 1;
-	// if we got one passed in, use that in addition
-	if (pchExtraDelim && *pchExtraDelim)
-	{
-		WhiteSpace[NumWhiteSpaces++] = pchExtraDelim;
-	}
-
-	return ParseIntoArray(OutArray, WhiteSpace, NumWhiteSpaces, InCullEmpty);
-}
-
-int32 FString::ParseIntoArrayLines(TArray<FString>& OutArray, bool InCullEmpty) const
-{
-	// default array of LineEndings
-	static const ElementType* LineEndings[] =
-	{
-		CHARTEXT(ElementType, "\r\n"),
-		CHARTEXT(ElementType, "\r"),
-		CHARTEXT(ElementType, "\n"),
-	};
-
-	// start with just the standard line endings
-	int32 NumLineEndings = UE_ARRAY_COUNT(LineEndings);	
-	return ParseIntoArray(OutArray, LineEndings, NumLineEndings, InCullEmpty);
-}
-
-int32 FString::ParseIntoArray(TArray<FString>& OutArray, const ElementType* const * DelimArray, int32 NumDelims, bool InCullEmpty) const
-{
-	// Make sure the delimit string is not null or empty
-	check(DelimArray);
-	OutArray.Reset();
-	const ElementType* Start = Data.GetData();
-	const int32 Length = Len();
-	if (Start)
-	{
-		int32 SubstringBeginIndex = 0;
-
-		// Iterate through string.
-		for(int32 i = 0; i < Len();)
-		{
-			int32 SubstringEndIndex = INDEX_NONE;
-			int32 DelimiterLength = 0;
-
-			// Attempt each delimiter.
-			for(int32 DelimIndex = 0; DelimIndex < NumDelims; ++DelimIndex)
-			{
-				DelimiterLength = TCString<ElementType>::Strlen(DelimArray[DelimIndex]);
-
-				// If we found a delimiter...
-				if (TCString<ElementType>::Strncmp(Start + i, DelimArray[DelimIndex], DelimiterLength) == 0)
-				{
-					// Mark the end of the substring.
-					SubstringEndIndex = i;
-					break;
-				}
-			}
-
-			if (SubstringEndIndex != INDEX_NONE)
-			{
-				const int32 SubstringLength = SubstringEndIndex - SubstringBeginIndex;
-				// If we're not culling empty strings or if we are but the string isn't empty anyways...
-				if(!InCullEmpty || SubstringLength != 0)
-				{
-					// ... add new string from substring beginning up to the beginning of this delimiter.
-					OutArray.Emplace(SubstringEndIndex - SubstringBeginIndex, Start + SubstringBeginIndex);
-				}
-				// Next substring begins at the end of the discovered delimiter.
-				SubstringBeginIndex = SubstringEndIndex + DelimiterLength;
-				i = SubstringBeginIndex;
-			}
-			else
-			{
-				++i;
-			}
-		}
-
-		// Add any remaining characters after the last delimiter.
-		const int32 SubstringLength = Length - SubstringBeginIndex;
-		// If we're not culling empty strings or if we are but the string isn't empty anyways...
-		if(!InCullEmpty || SubstringLength != 0)
-		{
-			// ... add new string from substring beginning up to the beginning of this delimiter.
-			OutArray.Emplace(Start + SubstringBeginIndex);
-		}
-	}
-
-	return OutArray.Num();
-}
-
-FString FString::Replace(const ElementType* From, const ElementType* To, ESearchCase::Type SearchCase) const &
-{
-	// Previous code used to accidentally accept a nullptr replacement string - this is no longer accepted.
-	check(To);
-
-	if (IsEmpty() || !From || !*From)
-	{
-		return *this;
-	}
-
-	// get a pointer into the character data
-	const ElementType* Travel = Data.GetData();
-
-	// precalc the lengths of the replacement strings
-	int32 FromLength = TCString<ElementType>::Strlen(From);
-	int32 ToLength   = TCString<ElementType>::Strlen(To);
-
-	FString Result;
-	while (true)
-	{
-		// look for From in the remaining string
-		const ElementType* FromLocation = SearchCase == ESearchCase::IgnoreCase ? TCString<ElementType>::Stristr(Travel, From) : TCString<ElementType>::Strstr(Travel, From);
-		if (!FromLocation)
-		{
-			break;
-		}
-
-		// copy everything up to FromLocation
-		Result.AppendChars(Travel, UE_PTRDIFF_TO_INT32(FromLocation - Travel));
-
-		// copy over the To
-		Result.AppendChars(To, ToLength);
-
-		Travel = FromLocation + FromLength;
-	}
-
-	// copy anything left over
-	Result += Travel;
-
-	return Result;
-}
-
-FString FString::Replace(const ElementType* From, const ElementType* To, ESearchCase::Type SearchCase) &&
-{
-	ReplaceInline(From, To, SearchCase);
-	return MoveTemp(*this);
-}
-
-int32 FString::ReplaceInline(const ElementType* SearchText, const ElementType* ReplacementText, ESearchCase::Type SearchCase)
-{
-	int32 ReplacementCount = 0;
-
-	if (Len() > 0
-		&& SearchText != nullptr && *SearchText != 0
-		&& ReplacementText != nullptr && (SearchCase == ESearchCase::IgnoreCase || TCString<ElementType>::Strcmp(SearchText, ReplacementText) != 0))
-	{
-		const int32 NumCharsToReplace = TCString<ElementType>::Strlen(SearchText);
-		const int32 NumCharsToInsert = TCString<ElementType>::Strlen(ReplacementText);
-
-		if (NumCharsToInsert == NumCharsToReplace)
-		{
-			ElementType* Pos = SearchCase == ESearchCase::IgnoreCase ? TCString<ElementType>::Stristr(&(*this)[0], SearchText) : TCString<ElementType>::Strstr(&(*this)[0], SearchText);
-			while (Pos != nullptr)
-			{
-				ReplacementCount++;
-
-				// TCString<ElementType>::Strcpy now inserts a terminating zero so can't use that
-				for (int32 i = 0; i < NumCharsToInsert; i++)
-				{
-					Pos[i] = ReplacementText[i];
-				}
-
-				if (Pos + NumCharsToReplace - **this < Len())
-				{
-					Pos = SearchCase == ESearchCase::IgnoreCase ? TCString<ElementType>::Stristr(Pos + NumCharsToReplace, SearchText) : TCString<ElementType>::Strstr(Pos + NumCharsToReplace, SearchText);
-				}
-				else
-				{
-					break;
-				}
-			}
-		}
-		else if (Contains(SearchText, SearchCase))
-		{
-			FString Copy(MoveTemp(*this));
-
-			// get a pointer into the character data
-			ElementType* WritePosition = (ElementType*)Copy.Data.GetData();
-			// look for From in the remaining string
-			ElementType* SearchPosition = SearchCase == ESearchCase::IgnoreCase ? TCString<ElementType>::Stristr(WritePosition, SearchText) : TCString<ElementType>::Strstr(WritePosition, SearchText);
-			while (SearchPosition != nullptr)
-			{
-				ReplacementCount++;
-
-				// replace the first letter of the From with 0 so we can do a strcpy (FString +=)
-				*SearchPosition = CHARTEXT(ElementType, '\0');
-
-				// copy everything up to the SearchPosition
-				(*this) += WritePosition;
-
-				// copy over the ReplacementText
-				(*this) += ReplacementText;
-
-				// restore the letter, just so we don't have 0's in the string
-				*SearchPosition = *SearchText;
-
-				WritePosition = SearchPosition + NumCharsToReplace;
-				SearchPosition = SearchCase == ESearchCase::IgnoreCase ? TCString<ElementType>::Stristr(WritePosition, SearchText) : TCString<ElementType>::Strstr(WritePosition, SearchText);
-			}
-
-			// copy anything left over
-			(*this) += WritePosition;
-		}
-	}
-
-	return ReplacementCount;
-}
-
-
-/**
- * Returns a copy of this string with all quote marks escaped (unless the quote is already escaped)
- */
-FString FString::ReplaceQuotesWithEscapedQuotes() &&
-{
-	if (Contains(CHARTEXT(ElementType, "\""), ESearchCase::CaseSensitive))
-	{
-		FString Copy(MoveTemp(*this));
-
-		const ElementType* pChar = *Copy;
-
-		bool bEscaped = false;
-		while ( *pChar != 0 )
-		{
-			if ( bEscaped )
-			{
-				bEscaped = false;
-			}
-			else if ( *pChar == ElementType('\\') )
-			{
-				bEscaped = true;
-			}
-			else if ( *pChar == ElementType('"') )
-			{
-				*this += ElementType('\\');
-			}
-
-			*this += *pChar++;
-		}
-	}
-
-	return MoveTemp(*this);
-}
-
-template <typename CharType>
-static const CharType* CharToEscapeSeqMap[6][2] =
-{
-	// Always replace \\ first to avoid double-escaping characters
-	{ CHARTEXT(CharType, "\\"), CHARTEXT(CharType, "\\\\") },
-	{ CHARTEXT(CharType, "\n"), CHARTEXT(CharType, "\\n")  },
-	{ CHARTEXT(CharType, "\r"), CHARTEXT(CharType, "\\r")  },
-	{ CHARTEXT(CharType, "\t"), CHARTEXT(CharType, "\\t")  },
-	{ CHARTEXT(CharType, "\'"), CHARTEXT(CharType, "\\'")  },
-	{ CHARTEXT(CharType, "\""), CHARTEXT(CharType, "\\\"") }
-};
-
-template <typename CharType>
-static const uint32 MaxSupportedEscapeChars = UE_ARRAY_COUNT(CharToEscapeSeqMap<CharType>);
-
-void FString::ReplaceCharWithEscapedCharInline(const TArray<ElementType>* Chars/*=nullptr*/)
-{
-	if ( Len() > 0 && (Chars == nullptr || Chars->Num() > 0) )
-	{
-		for ( uint32 ChIdx = 0; ChIdx < MaxSupportedEscapeChars<ElementType>; ChIdx++ )
-		{
-			if ( Chars == nullptr || Chars->Contains(*(CharToEscapeSeqMap<ElementType>[ChIdx][0])) )
-			{
-				// use ReplaceInline as that won't create a copy of the string if the character isn't found
-				ReplaceInline(CharToEscapeSeqMap<ElementType>[ChIdx][0], CharToEscapeSeqMap<ElementType>[ChIdx][1]);
-			}
-		}
-	}
-}
-
-void FString::ReplaceEscapedCharWithCharInline(const TArray<ElementType>* Chars/*=nullptr*/)
-{
-	if ( Len() > 0 && (Chars == nullptr || Chars->Num() > 0) )
-	{
-		// Spin CharToEscapeSeqMap backwards to ensure we're doing the inverse of ReplaceCharWithEscapedChar
-		for ( int32 ChIdx = MaxSupportedEscapeChars<ElementType>; ChIdx > 0; )
-		{
-			--ChIdx;
-
-			if ( Chars == nullptr || Chars->Contains(*(CharToEscapeSeqMap<ElementType>[ChIdx][0])) )
-			{
-				// use ReplaceInline as that won't create a copy of the string if the character isn't found
-				ReplaceInline(CharToEscapeSeqMap<ElementType>[ChIdx][1], CharToEscapeSeqMap<ElementType>[ChIdx][0]);
-			}
-		}
-	}
-}
-
-/** 
- * Replaces all instances of '\t' with TabWidth number of spaces
- * @param InSpacesPerTab - Number of spaces that a tab represents
- */
-void FString::ConvertTabsToSpacesInline(const int32 InSpacesPerTab)
-{
-	//must call this with at least 1 space so the modulus operation works
-	check(InSpacesPerTab > 0);
-
-	int32 TabIndex = 0;
-	while ((TabIndex = Find(CHARTEXT(ElementType, "\t"), ESearchCase::CaseSensitive)) != INDEX_NONE )
-	{
-		FString RightSide = Mid(TabIndex+1);
-		LeftInline(TabIndex, false);
-
-		//for a tab size of 4, 
-		int32 LineBegin = Find(CHARTEXT(ElementType, "\n"), ESearchCase::CaseSensitive, ESearchDir::FromEnd, TabIndex);
-		if (LineBegin == INDEX_NONE)
-		{
-			LineBegin = 0;
-		}
-		const int32 CharactersOnLine = (Len()-LineBegin);
-
-		int32 NumSpacesForTab = InSpacesPerTab - (CharactersOnLine % InSpacesPerTab);
-		for (int32 i = 0; i < NumSpacesForTab; ++i)
-		{
-			AppendChar(CHARTEXT(ElementType, ' '));
-		}
-		Append(RightSide);
-	}
-}
-
-// This starting size catches 99.97% of printf calls - there are about 700k printf calls per level
-#define STARTING_BUFFER_SIZE		512
-
-FString FString::PrintfImpl(const ElementType* Fmt, ...)
-{
-	int32		BufferSize	= STARTING_BUFFER_SIZE;
-	ElementType	StartingBuffer[STARTING_BUFFER_SIZE];
-	ElementType*	Buffer		= StartingBuffer;
-	int32		Result		= -1;
-
-	// First try to print to a stack allocated location 
-	GET_VARARGS_RESULT( Buffer, BufferSize, BufferSize-1, Fmt, Fmt, Result );
-
-	// If that fails, start allocating regular memory
-	if( Result == -1 )
-	{
-		Buffer = nullptr;
-		while(Result == -1)
-		{
-			BufferSize *= 2;
-			Buffer = (ElementType*) FMemory::Realloc( Buffer, BufferSize * sizeof(ElementType) );
-			GET_VARARGS_RESULT( Buffer, BufferSize, BufferSize-1, Fmt, Fmt, Result );
-		};
-	}
-
-	Buffer[Result] = CHARTEXT(ElementType, '\0');
-
-	FString ResultString(Buffer);
-
-	if( BufferSize != STARTING_BUFFER_SIZE )
-	{
-		FMemory::Free( Buffer );
-	}
-
-	return ResultString;
-}
-
-void FString::AppendfImpl(FString& AppendToMe, const ElementType* Fmt, ...)
-{
-	int32		BufferSize = STARTING_BUFFER_SIZE;
-	ElementType	StartingBuffer[STARTING_BUFFER_SIZE];
-	ElementType*	Buffer = StartingBuffer;
-	int32		Result = -1;
-
-	// First try to print to a stack allocated location 
-	GET_VARARGS_RESULT(Buffer, BufferSize, BufferSize - 1, Fmt, Fmt, Result);
-
-	// If that fails, start allocating regular memory
-	if (Result == -1)
-	{
-		Buffer = nullptr;
-		while (Result == -1)
-		{
-			BufferSize *= 2;
-			Buffer = (ElementType*)FMemory::Realloc(Buffer, BufferSize * sizeof(ElementType));
-			GET_VARARGS_RESULT(Buffer, BufferSize, BufferSize - 1, Fmt, Fmt, Result);
-		};
-	}
-
-	Buffer[Result] = CHARTEXT(ElementType, '\0');
-
-	AppendToMe += Buffer;
-
-	if (BufferSize != STARTING_BUFFER_SIZE)
-	{
-		FMemory::Free(Buffer);
-	}
-}
-
-static_assert(PLATFORM_LITTLE_ENDIAN, "FString serialization needs updating to support big-endian platforms!");
-
-FArchive& operator<<( FArchive& Ar, FString& A )
-{
-	using ElementType = FString::ElementType;
-
-	if constexpr (std::is_same_v<ElementType, TCHAR>)
-	{
-		// > 0 for ANSICHAR, < 0 for UTF16CHAR serialization
-		static_assert(sizeof(UTF16CHAR) == sizeof(UCS2CHAR), "UTF16CHAR and UCS2CHAR are assumed to be the same size!");
-
-		if (Ar.IsLoading())
-		{
-			int32 SaveNum = 0;
-			Ar << SaveNum;
-
-			bool bLoadUnicodeChar = SaveNum < 0;
-			if (bLoadUnicodeChar)
-			{
-				// If SaveNum cannot be negated due to integer overflow, Ar is corrupted.
-				if (SaveNum == MIN_int32)
-				{
-					Ar.SetCriticalError();
-					UE_LOG(LogCore, Error, CHARTEXT(ElementType, "Archive is corrupted"));
-					return Ar;
-				}
-
-				SaveNum = -SaveNum;
-			}
-
-			int64 MaxSerializeSize = Ar.GetMaxSerializeSize();
-			// Protect against network packets allocating too much memory
-			if ((MaxSerializeSize > 0) && (SaveNum > MaxSerializeSize))
-			{
-				Ar.SetCriticalError();
-				UE_LOG(LogCore, Error, CHARTEXT(ElementType, "String is too large (Size: %i, Max: %i)"), SaveNum, MaxSerializeSize);
-				return Ar;
-			}
-
-			// Resize the array only if it passes the above tests to prevent rogue packets from crashing
-			A.Data.Empty(SaveNum);
-			A.Data.AddUninitialized(SaveNum);
-
-			if (SaveNum)
-			{
-				if (bLoadUnicodeChar)
-				{
-					// read in the unicode string
-					auto Passthru = StringMemoryPassthru<UCS2CHAR>(A.Data.GetData(), SaveNum, SaveNum);
-					Ar.Serialize(Passthru.Get(), SaveNum * sizeof(UCS2CHAR));
-					if (Ar.IsByteSwapping())
-					{
-						for (int32 CharIndex = 0; CharIndex < SaveNum; ++CharIndex)
-						{
-							Passthru.Get()[CharIndex] = ByteSwap(Passthru.Get()[CharIndex]);
-						}
-					}
-					// Ensure the string has a null terminator
-					Passthru.Get()[SaveNum - 1] = '\0';
-					Passthru.Apply();
-
-					// Inline combine any surrogate pairs in the data when loading into a UTF-32 string
-					StringConv::InlineCombineSurrogates(A);
-
-					// Since Microsoft's vsnwprintf implementation raises an invalid parameter warning
-					// with a character of 0xffff, scan for it and terminate the string there.
-					// 0xffff isn't an actual Unicode character anyway.
-					int Index = 0;
-					if (A.FindChar(0xffff, Index))
-					{
-						A[Index] = CHARTEXT(ElementType, '\0');
-						A.TrimToNullTerminator();
-					}
-				}
-				else
-				{
-					auto Passthru = StringMemoryPassthru<ANSICHAR>(A.Data.GetData(), SaveNum, SaveNum);
-					Ar.Serialize(Passthru.Get(), SaveNum * sizeof(ANSICHAR));
-					// Ensure the string has a null terminator
-					Passthru.Get()[SaveNum - 1] = '\0';
-					Passthru.Apply();
-				}
-
-				// Throw away empty string.
-				if (SaveNum == 1)
-				{
-					A.Data.Empty();
-				}
-			}
-		}
-		else
-		{
-			A.Data.CountBytes(Ar);
-
-			const bool bSaveUnicodeChar = Ar.IsForcingUnicode() || !TCString<ElementType>::IsPureAnsi(*A);
-			if (bSaveUnicodeChar)
-			{
-				// This preprocessor block should not be necessary when the StringCast above understands UTF16CHAR.
-				FTCHARToUTF16 UTF16String(*A, A.Len() + 1); // include the null terminator
-				int32 Num = UTF16String.Length() + 1; // include the null terminator
-
-				int32 SaveNum = -Num;
-				Ar << SaveNum;
-
-				if (Num)
-				{
-					if (!Ar.IsByteSwapping())
-					{
-						Ar.Serialize((void*)UTF16String.Get(), sizeof(UTF16CHAR) * Num);
-					}
-					else
-					{
-						TArray<UTF16CHAR> Swapped(UTF16String.Get(), Num);
-						for (int32 CharIndex = 0; CharIndex < Num; ++CharIndex)
-						{
-							Swapped[CharIndex] = ByteSwap(Swapped[CharIndex]);
-						}
-						Ar.Serialize((void*)Swapped.GetData(), sizeof(UTF16CHAR) * Num);
-					}
-				}
-			}
-			else
-			{
-				int32 Num = A.Data.Num();
-				Ar << Num;
-
-				if (Num)
-				{
-					Ar.Serialize((void*)StringCast<ANSICHAR>(A.Data.GetData(), Num).Get(), sizeof(ANSICHAR) * Num);
-				}
-			}
-		}
-	}
-	else
-	{
-		// Can just serialize as UTF-8 always
-		check(false); // TODO
-	}
-
-	return Ar;
-}
-
-int32 HexToBytes(const FString& HexString, uint8* OutBytes)
-{
-	return UE::String::HexToBytes(HexString, OutBytes);
-}
-
-int32 FindMatchingClosingParenthesis(const FString& TargetString, const int32 StartSearch)
-{
-	using ElementType = FString::ElementType;
-
-	check(StartSearch >= 0 && StartSearch <= TargetString.Len());// Check for usage, we do not accept INDEX_NONE like other string functions
-
-	const ElementType* const StartPosition = (*TargetString) + StartSearch;
-	const ElementType* CurrPosition = StartPosition;
-	int32 ParenthesisCount = 0;
-
-	// Move to first open parenthesis
-	while (*CurrPosition != 0 && *CurrPosition != CHARTEXT(ElementType, '('))
-	{
-		++CurrPosition;
-	}
-
-	// Did we find the open parenthesis
-	if (*CurrPosition == CHARTEXT(ElementType, '('))
-	{
-		++ParenthesisCount;
-		++CurrPosition;
-
-		while (*CurrPosition != 0 && ParenthesisCount > 0)
-		{
-			if (*CurrPosition == CHARTEXT(ElementType, '('))
-			{
-				++ParenthesisCount;
-			}
-			else if (*CurrPosition == CHARTEXT(ElementType, ')'))
-			{
-				--ParenthesisCount;
-			}
-			++CurrPosition;
-		}
-
-		// Did we find the matching close parenthesis
-		if (ParenthesisCount == 0 && *(CurrPosition - 1) == CHARTEXT(ElementType, ')'))
-		{
-			return StartSearch + UE_PTRDIFF_TO_INT32((CurrPosition - 1) - StartPosition);
-		}
-	}
-
-	return INDEX_NONE;
-}
-
-FString SlugStringForValidName(const FString& DisplayString, const FString::ElementType* ReplaceWith /*= CHARTEXT(ElementType, "")*/)
-{
-	using ElementType = FString::ElementType;
-
-	FString GeneratedName = DisplayString;
-
-	// Convert the display label, which may consist of just about any possible character, into a
-	// suitable name for a UObject (remove whitespace, certain symbols, etc.)
-	{
-		for ( int32 BadCharacterIndex = 0; BadCharacterIndex < UE_ARRAY_COUNT(INVALID_OBJECTNAME_CHARACTERS) - 1; ++BadCharacterIndex )
-		{
-			const ElementType TestChar[2] = { INVALID_OBJECTNAME_CHARACTERS[BadCharacterIndex], CHARTEXT(ElementType, '\0') };
-			const int32 NumReplacedChars = GeneratedName.ReplaceInline(TestChar, ReplaceWith);
-		}
-	}
-
-	return GeneratedName;
-}
-
-void FTextRange::CalculateLineRangesFromString(const FString& Input, TArray<FTextRange>& LineRanges)
-{
-	using ElementType = FString::ElementType;
-
-	int32 LineBeginIndex = 0;
-
-	// Loop through splitting at new-lines
-	const ElementType* const InputStart = *Input;
-	for (const ElementType* CurrentChar = InputStart; CurrentChar && *CurrentChar; ++CurrentChar)
-	{
-		// Handle a chain of \r\n slightly differently to stop the TChar<ElementType>::IsLinebreak adding two separate new-lines
-		const bool bIsWindowsNewLine = (*CurrentChar == '\r' && *(CurrentChar + 1) == '\n');
-		if (bIsWindowsNewLine || TChar<ElementType>::IsLinebreak(*CurrentChar))
-		{
-			const int32 LineEndIndex = UE_PTRDIFF_TO_INT32(CurrentChar - InputStart);
-			check(LineEndIndex >= LineBeginIndex);
-			LineRanges.Emplace(FTextRange(LineBeginIndex, LineEndIndex));
-
-			if (bIsWindowsNewLine)
-			{
-				++CurrentChar; // skip the \n of the \r\n chain
-			}
-			LineBeginIndex = UE_PTRDIFF_TO_INT32(CurrentChar - InputStart) + 1; // The next line begins after the end of the current line
-		}
-	}
-
-	// Process any remaining string after the last new-line
-	if (LineBeginIndex <= Input.Len())
-	{
-		LineRanges.Emplace(FTextRange(LineBeginIndex, Input.Len()));
-	}
-}
-
-void StringConv::InlineCombineSurrogates(FString& Str)
-{
-	InlineCombineSurrogates_Array(Str.GetCharArray());
-}
+#include "Containers/String.cpp.inl"
