Moved string code into .inl files in anticipation of the upcoming FUtf8String change.

diff --git a/Engine/Source/Runtime/Core/Public/Containers/UnrealString.h b/Engine/Source/Runtime/Core/Public/Containers/UnrealString.h
--- a/Engine/Source/Runtime/Core/Public/Containers/UnrealString.h
+++ b/Engine/Source/Runtime/Core/Public/Containers/UnrealString.h
@@ -5,2590 +5,7 @@
 
 #pragma once
 
-#include "CoreTypes.h"
-#include "Misc/VarArgs.h"
-#include "Misc/OutputDevice.h"
-#include "Misc/AssertionMacros.h"
-#include "HAL/UnrealMemory.h"
-#include "Templates/IsArithmetic.h"
-#include "Templates/IsArray.h"
-#include "Templates/UnrealTypeTraits.h"
-#include "Templates/UnrealTemplate.h"
-#include "Math/NumericLimits.h"
-#include "Containers/Array.h"
-#include "Misc/CString.h"
-#include "Misc/Crc.h"
-#include "Math/UnrealMathUtility.h"
-#include "Templates/Invoke.h"
-#include "Templates/IsValidVariadicFunctionArg.h"
-#include "Templates/AndOrNot.h"
-#include "Templates/IsArrayOrRefOfTypeByPredicate.h"
-#include "Templates/TypeHash.h"
-#include "Templates/IsFloatingPoint.h"
-#include "Templates/UnrealTypeTraits.h"
-#include "Traits/IsCharType.h"
-#include "Traits/IsCharEncodingCompatibleWith.h"
-#include "Traits/IsCharEncodingSimplyConvertibleTo.h"
+// Include UnrealString.h.inl's includes before defining (upcoming) macros, in case the macros 'poison' other headers or there are re-entrant includes.
+#include "Containers/UnrealStringIncludes.h.inl"
 
-#include <type_traits>
-
-struct FStringFormatArg;
-template<typename InKeyType,typename InValueType,typename SetAllocator ,typename KeyFuncs > class TMap;
-
-typedef TMap<FString, FStringFormatArg> FStringFormatNamedArguments;
-typedef TArray<FStringFormatArg> FStringFormatOrderedArguments;
-
-TCHAR*       GetData(FString&);
-const TCHAR* GetData(const FString&);
-int32        GetNum(const FString& String);
-/**
- * A dynamically sizeable string.
- * @see https://docs.unrealengine.com/latest/INT/Programming/UnrealArchitecture/StringHandling/FString/
- *
- * When dealing with UTF-8 literals, the following advice is recommended:
- *
- * - Do not use the u8"..." prefix (gives the wrong array type until C++20).
- * - Use UTF8TEXT("...") for array literals (type is const UTF8CHAR[n]).
- * - Use UTF8TEXTVIEW("...") for string view literals (type is FUtf8StringView).
- * - Use \uxxxx or \Uxxxxxxxx escape sequences rather than \x to specify Unicode code points.
- */
-class FString
-{
-public:
-	using AllocatorType = TSizedDefaultAllocator<32>;
-	using ElementType   = TCHAR;
-
-private:
-	/** Array holding the character data */
-	typedef TArray<ElementType, AllocatorType> DataType;
-	DataType Data;
-
-	template <typename RangeType>
-	using TRangeElementType = std::remove_cv_t<typename TRemovePointer<decltype(GetData(DeclVal<RangeType>()))>::Type>;
-
-	template <typename CharRangeType>
-	struct TIsRangeOfCharType : TIsCharType<TRangeElementType<CharRangeType>>
-	{
-	};
-
-	template <typename CharRangeType>
-	struct TIsRangeOfElementType
-	{
-		enum { Value = std::is_same_v<ElementType, TRangeElementType<CharRangeType>> };
-	};
-
-	/** Trait testing whether a type is a contiguous range of characters, and not CharType[]. */
-	template <typename CharRangeType>
-	using TIsCharRangeNotCArray = TAnd<
-		TIsContiguousContainer<CharRangeType>,
-		TNot<TIsArray<typename TRemoveReference<CharRangeType>::Type>>,
-		TIsRangeOfCharType<CharRangeType>>;
-
-	/** Trait testing whether a type is a contiguous range of characters, and not CharType[] and not a string class type. */
-	template <typename CharRangeType>
-	using TIsCharRangeNotCArrayNotStringClass = TAnd<
-		TIsCharRangeNotCArray<CharRangeType>,
-		TNot<TIsDerivedFrom<typename TDecay<CharRangeType>::Type, FString>>>;
-
-	/** Trait testing whether a type is a contiguous range of ElementType, and not CharType[]. */
-	template <typename CharRangeType>
-	using TIsTCharRangeNotCArray = TAnd<
-		TIsContiguousContainer<CharRangeType>,
-		TNot<TIsArray<typename TRemoveReference<CharRangeType>::Type>>,
-		TIsRangeOfElementType<CharRangeType>>;
-
-	/** Trait testing whether a type is a contiguous range of ElementType, and not CharType[] and not a string class type. */
-	template <typename CharRangeType>
-	using TIsTCharRangeNotCArrayNotStringClass = TAnd<
-		TIsTCharRangeNotCArray<CharRangeType>,
-		TNot<TIsDerivedFrom<typename TDecay<CharRangeType>::Type, FString>>>;
-
-public:
-	FString() = default;
-	FString(FString&&) = default;
-	FString(const FString&) = default;
-	FString& operator=(FString&&) = default;
-	FString& operator=(const FString&) = default;
-
-	/**
-	 * Create a copy of the Other string with extra space for characters at the end of the string
-	 *
-	 * @param Other the other string to create a new copy from
-	 * @param ExtraSlack number of extra characters to add to the end of the other string in this string
-	 */
-	FORCEINLINE FString(const FString& Other, int32 ExtraSlack)
-		: Data(Other.Data, ExtraSlack + ((Other.Data.Num() || !ExtraSlack) ? 0 : 1)) // Add 1 if the source string array is empty and we want some slack, because we'll need to include a null terminator which is currently missing
-	{
-	}
-
-	/**
-	 * Create a copy of the Other string with extra space for characters at the end of the string
-	 *
-	 * @param Other the other string to create a new copy from
-	 * @param ExtraSlack number of extra characters to add to the end of the other string in this string
-	 */
-	FORCEINLINE FString(FString&& Other, int32 ExtraSlack)
-		: Data(MoveTemp(Other.Data), ExtraSlack + ((Other.Data.Num() || !ExtraSlack) ? 0 : 1)) // Add 1 if the source string array is empty and we want some slack, because we'll need to include a null terminator which is currently missing
-	{
-	}
-
-	/** Construct from null-terminated C string or nullptr  */
-	CORE_API FString(const ANSICHAR* Str);
-	CORE_API FString(const WIDECHAR* Str);
-	CORE_API FString(const UTF8CHAR* Str);
-	CORE_API FString(const UCS2CHAR* Str);
-
-	/** Construct from null-terminated C substring or nullptr */
-	CORE_API FString(int32 Len, const ANSICHAR* Str);
-	CORE_API FString(int32 Len, const WIDECHAR* Str);
-	CORE_API FString(int32 Len, const UTF8CHAR* Str);
-	CORE_API FString(int32 Len, const UCS2CHAR* Str);
-
-	/** Construct from null-terminated C string or nullptr with extra slack on top of original string length */
-	CORE_API FString(const ANSICHAR* Str, int32 ExtraSlack);
-	CORE_API FString(const WIDECHAR* Str, int32 ExtraSlack);
-	CORE_API FString(const UTF8CHAR* Str, int32 ExtraSlack);
-	CORE_API FString(const UCS2CHAR* Str, int32 ExtraSlack);
-
-	/** Construct from contiguous range of characters such as a string view or string builder */
-	template <typename CharRangeType, typename TEnableIf<TIsCharRangeNotCArrayNotStringClass<CharRangeType>::Value>::Type* = nullptr>
-	FORCEINLINE explicit FString(CharRangeType&& Str) : FString(GetNum(Str), GetData(Forward<CharRangeType>(Str)))
-	{
-	}
-
-	/** Construct from contiguous range of characters with extra slack on top of original string length */
-	template <typename CharRangeType, typename TEnableIf<TIsCharRangeNotCArrayNotStringClass<CharRangeType>::Value>::Type* = nullptr>
-	explicit FString(CharRangeType&& Str, int32 ExtraSlack)
-	{
-		uint32 InLen = GetNum(Str);
-		Reserve(InLen + ExtraSlack);
-		AppendChars(GetData(Forward<CharRangeType>(Str)), InLen);
-	}
-
-public:
-#ifdef __OBJC__
-	/** Convert Objective-C NSString* to FString */
-	FORCEINLINE FString(const NSString* In)
-	{
-		if (In && [In length] > 0)
-		{
-			// Convert the NSString data into the native TCHAR format for UE4
-			// This returns a buffer of bytes, but they can be safely cast to a buffer of TCHARs
-#if PLATFORM_TCHAR_IS_4_BYTES
-			const CFStringEncoding Encoding = kCFStringEncodingUTF32LE;
-#else
-			const CFStringEncoding Encoding = kCFStringEncodingUTF16LE;
-#endif
-
-			CFRange Range = CFRangeMake(0, CFStringGetLength((__bridge CFStringRef)In));
-			CFIndex BytesNeeded;
-			if (CFStringGetBytes((__bridge CFStringRef)In, Range, Encoding, '?', false, NULL, 0, &BytesNeeded) > 0)
-			{
-				const size_t Length = BytesNeeded / sizeof(TCHAR);
-				Data.Reserve(Length + 1);
-				Data.AddUninitialized(Length + 1);
-				CFStringGetBytes((__bridge CFStringRef)In, Range, Encoding, '?', false, (uint8*)Data.GetData(), Length * sizeof(TCHAR) + 1, NULL);
-				Data[Length] = 0;
-			}
-		}
-	}
-#endif
-
-	CORE_API FString& operator=(const ElementType* Str);
-
-	template <typename CharRangeType, typename TEnableIf<TIsTCharRangeNotCArray<CharRangeType>::Value>::Type* = nullptr>
-	FORCEINLINE FString& operator=(CharRangeType&& Range)
-	{
-		AssignRange(GetData(Range), GetNum(Range));
-		return *this;
-	}
-
-private:
-	CORE_API void AssignRange(const ElementType* Str, int32 Len);
-
-public:
-	/**
-	 * Return specific character from this string
-	 *
-	 * @param Index into string
-	 * @return Character at Index
-	 */
-	FORCEINLINE ElementType& operator[]( int32 Index ) UE_LIFETIMEBOUND
-	{
-		checkf(IsValidIndex(Index), TEXT("String index out of bounds: Index %i from a string with a length of %i"), Index, Len());
-		return Data.GetData()[Index];
-	}
-
-	/**
-	 * Return specific const character from this string
-	 *
-	 * @param Index into string
-	 * @return const Character at Index
-	 */
-	FORCEINLINE const ElementType& operator[]( int32 Index ) const UE_LIFETIMEBOUND
-	{
-		checkf(IsValidIndex(Index), TEXT("String index out of bounds: Index %i from a string with a length of %i"), Index, Len());
-		return Data.GetData()[Index];
-	}
-
-	/**
-	 * Iterator typedefs
-	 */
-	typedef DataType::TIterator      TIterator;
-	typedef DataType::TConstIterator TConstIterator;
-
-	/** Creates an iterator for the characters in this string */
-	FORCEINLINE TIterator CreateIterator()
-	{
-		return Data.CreateIterator();
-	}
-
-	/** Creates a const iterator for the characters in this string */
-	FORCEINLINE TConstIterator CreateConstIterator() const
-	{
-		return Data.CreateConstIterator();
-	}
-
-	/**
-	 * DO NOT USE DIRECTLY
-	 * STL-like iterators to enable range-based for loop support.
-	 */
-	FORCEINLINE DataType::RangedForIteratorType             begin ()       { auto Result = Data.begin();                                return Result; }
-	FORCEINLINE DataType::RangedForConstIteratorType        begin () const { auto Result = Data.begin();                                return Result; }
-	FORCEINLINE DataType::RangedForIteratorType             end   ()       { auto Result = Data.end();    if (Data.Num()) { --Result; } return Result; }
-	FORCEINLINE DataType::RangedForConstIteratorType        end   () const { auto Result = Data.end();    if (Data.Num()) { --Result; } return Result; }
-	FORCEINLINE DataType::RangedForReverseIteratorType      rbegin()       { auto Result = Data.rbegin(); if (Data.Num()) { ++Result; } return Result; }
-	FORCEINLINE DataType::RangedForConstReverseIteratorType rbegin() const { auto Result = Data.rbegin(); if (Data.Num()) { ++Result; } return Result; }
-	FORCEINLINE DataType::RangedForReverseIteratorType      rend  ()       { auto Result = Data.rend();                                 return Result; }
-	FORCEINLINE DataType::RangedForConstReverseIteratorType rend  () const { auto Result = Data.rend();                                 return Result; }
-
-	FORCEINLINE SIZE_T GetAllocatedSize() const
-	{
-		return Data.GetAllocatedSize();
-	}
-
-	/**
-	 * Run slow checks on this string
-	 */
-	FORCEINLINE void CheckInvariants() const
-	{
-		int32 Num = Data.Num();
-		checkSlow(Num >= 0);
-		checkSlow(!Num || !Data.GetData()[Num - 1]);
-		checkSlow(Data.GetSlack() >= 0);
-	}
-
-	/**
-	 * Create empty string of given size with zero terminating character
-	 *
-	 * @param Slack length of empty string to create
-	 */
-	CORE_API void Empty();
-	CORE_API void Empty(int32 Slack);
-
-	/**
-	 * Test whether this string is empty
-	 *
-	 * @return true if this string is empty, otherwise return false.
-	 */
-	UE_NODISCARD FORCEINLINE bool IsEmpty() const
-	{
-		return Data.Num() <= 1;
-	}
-
-	/**
-	 * Empties the string, but doesn't change memory allocation, unless the new size is larger than the current string.
-	 *
-	 * @param NewReservedSize The expected usage size (in characters, not including the terminator) after calling this function.
-	 */
-	CORE_API void Reset(int32 NewReservedSize = 0);
-
-	/**
-	 * Remove unallocated empty character space from the end of this string
-	 */
-	CORE_API void Shrink();
-
-	/**
-	 * Tests if index is valid, i.e. greater than or equal to zero, and less than the number of characters in this string (excluding the null terminator).
-	 *
-	 * @param Index Index to test.
-	 *
-	 * @returns True if index is valid. False otherwise.
-	 */
-	UE_NODISCARD FORCEINLINE bool IsValidIndex(int32 Index) const
-	{
-		return Index >= 0 && Index < Len();
-	}
-
-	/**
-	 * Get pointer to the string
-	 *
-	 * @Return Pointer to Array of ElementType if Num, otherwise the empty string
-	 */
-	UE_NODISCARD FORCEINLINE const ElementType* operator*() const UE_LIFETIMEBOUND
-	{
-		return Data.Num() ? Data.GetData() : CHARTEXT(ElementType, "");
-	}
-
-	/** 
-	 *Get string as array of TCHARS 
-	 *
-	 * @warning: Operations on the TArray<*CHAR> can be unsafe, such as adding
-	 *		non-terminating 0's or removing the terminating zero.
-	 */
-	UE_NODISCARD FORCEINLINE DataType& GetCharArray() UE_LIFETIMEBOUND
-	{
-		return Data;
-	}
-
-	/** Get string as const array of TCHARS */
-	UE_NODISCARD FORCEINLINE const DataType& GetCharArray() const UE_LIFETIMEBOUND
-	{
-		return Data;
-	}
-
-#ifdef __OBJC__
-	/** Convert the string to Objective-C NSString */
-    CORE_API NSString* GetNSString() const;
-#endif
-
-	/** 
-	 * Appends a character range without null-terminators in it
-	 *
-	 * @param Str can be null if Count is 0. Can be unterminated, Str[Count] isn't read.
-	 */
-	CORE_API void AppendChars(const ANSICHAR* Str, int32 Count);
-	CORE_API void AppendChars(const WIDECHAR* Str, int32 Count);
-	CORE_API void AppendChars(const UCS2CHAR* Str, int32 Count);
-	CORE_API void AppendChars(const UTF8CHAR* Str, int32 Count);
-
-	/** Append a string and return a reference to this */
-	template<class CharType>
-	FORCEINLINE FString& Append(const CharType* Str, int32 Count)
-	{
-		AppendChars(Str, Count);
-		return *this;
-	}
-	
-	/**
-	 * Append a valid null-terminated string and return a reference to this
-	 *
-	 * CharType is not const to use this overload for mutable char arrays and call
-	 * Strlen() instead of getting the static length N from GetNum((&T)[N]).
-	 * Oddly MSVC ranks a const T* overload over T&& for T[N] while clang picks T&&.
-	 */
-	template<class CharType>
-	FORCEINLINE FString& Append(/* no const! */ CharType* Str)
-	{
-		checkSlow(Str);
-		AppendChars(Str, TCString<std::remove_const_t<CharType>>::Strlen(Str));
-		return *this;
-	}
-
-	/** Append a string and return a reference to this */
-	template <typename CharRangeType, typename TEnableIf<TIsCharRangeNotCArray<CharRangeType>::Value>::Type* = nullptr>
-	FORCEINLINE FString& Append(CharRangeType&& Str)
-	{
-		AppendChars(GetData(Str), GetNum(Str));
-		return *this;
-	}
-
-	/** Append a single character and return a reference to this */
-	CORE_API FString& AppendChar(ElementType InChar);
-
-	/** Append a string and return a reference to this */
-	template <typename StrType>
-	FORCEINLINE auto operator+=(StrType&& Str) -> decltype(Append(Forward<StrType>(Str)))
-	{
-		return Append(Forward<StrType>(Str));
-	}
-
-	/** Append a single character and return a reference to this */
-	template <
-		typename AppendedCharType,
-		std::enable_if_t<TIsCharType<AppendedCharType>::Value>* = nullptr
-	>
-	FORCEINLINE FString& operator+=(AppendedCharType Char)
-	{
-		if constexpr (TIsCharEncodingSimplyConvertibleTo_V<AppendedCharType, ElementType>)
-		{
-			return AppendChar((ElementType)Char);
-		}
-		else
-		{
-			AppendChars(&Char, 1);
-			return *this;
-		}
-	}
-
-	CORE_API void InsertAt(int32 Index, ElementType Character);
-	CORE_API void InsertAt(int32 Index, const FString& Characters);
-
-	/**
-	 * Removes characters within the string.
-	 *
-	 * @param Index           The index of the first character to remove.
-	 * @param Count           The number of characters to remove.
-	 * @param bAllowShrinking Whether or not to reallocate to shrink the storage after removal.
-	 */
-	CORE_API void RemoveAt(int32 Index, int32 Count = 1, bool bAllowShrinking = true);
-
-	/**
-	 * Removes the text from the start of the string if it exists.
-	 *
-	 * @param InPrefix the prefix to search for at the start of the string to remove.
-	 * @return true if the prefix was removed, otherwise false.
-	 */
-	template <typename TCharRangeType, std::enable_if_t<TIsCharRangeNotCArrayNotStringClass<TCharRangeType>::Value>* = nullptr>
-	bool RemoveFromStart(TCharRangeType&& InPrefix, ESearchCase::Type SearchCase = ESearchCase::IgnoreCase)
-	{
-		static_assert(std::is_same_v<typename TElementType<TCharRangeType>::Type, ElementType>, "Expected a range of ElementType");
-		return RemoveFromStart(GetData(InPrefix), GetNum(InPrefix), SearchCase);
-	}
-
-	/**
-	 * Removes the text from the start of the string if it exists.
-	 *
-	 * @param InPrefix the prefix to search for at the start of the string to remove.
-	 * @return true if the prefix was removed, otherwise false.
-	 */
-	bool RemoveFromStart(const ElementType* InPrefix, ESearchCase::Type SearchCase = ESearchCase::IgnoreCase)
-	{
-		return RemoveFromStart(InPrefix, InPrefix ? FCString::Strlen(InPrefix) : 0, SearchCase);
-	}
-
-	/**
-	 * Removes the text from the start of the string if it exists.
-	 *
-	 * @param InPrefix the prefix to search for at the start of the string to remove.
-	 * @return true if the prefix was removed, otherwise false.
-	 */
-	bool RemoveFromStart(const FString& InPrefix, ESearchCase::Type SearchCase = ESearchCase::IgnoreCase)
-	{
-		return RemoveFromStart(*InPrefix, InPrefix.Len(), SearchCase);
-	}
-
-	/**
-	 * Removes the text from the start of the string if it exists.
-	 *
-	 * @param InPrefix the prefix to search for at the start of the string to remove.
-	 * @param InPrefixLen length of InPrefix
-	 * @return true if the prefix was removed, otherwise false.
-	 */
-	CORE_API bool RemoveFromStart(const ElementType* InPrefix, int32 InPrefixLen, ESearchCase::Type SearchCase = ESearchCase::IgnoreCase);
-
-	/**
-	 * Removes the text from the end of the string if it exists.
-	 *
-	 * @param InSuffix the suffix to search for at the end of the string to remove.
-	 * @return true if the suffix was removed, otherwise false.
-	 */
-	template <typename TCharRangeType, std::enable_if_t<TIsCharRangeNotCArrayNotStringClass<TCharRangeType>::Value>* = nullptr>
-	bool RemoveFromEnd(TCharRangeType&& InSuffix, ESearchCase::Type SearchCase = ESearchCase::IgnoreCase)
-	{
-		static_assert(std::is_same_v<typename TElementType<TCharRangeType>::Type, ElementType>, "Expected a range of ElementType");
-		return RemoveFromEnd(GetData(InSuffix), GetNum(InSuffix), SearchCase);
-	}
-
-	/**
-	 * Removes the text from the end of the string if it exists.
-	 *
-	 * @param InSuffix the suffix to search for at the end of the string to remove.
-	 * @return true if the suffix was removed, otherwise false.
-	 */
-	bool RemoveFromEnd(const ElementType* InSuffix, ESearchCase::Type SearchCase = ESearchCase::IgnoreCase)
-	{
-		return RemoveFromEnd(InSuffix, InSuffix ? FCString::Strlen(InSuffix) : 0, SearchCase);
-	}
-
-	/**
-	 * Removes the text from the end of the string if it exists.
-	 *
-	 * @param InSuffix the suffix to search for at the end of the string to remove.
-	 * @return true if the suffix was removed, otherwise false.
-	 */
-	bool RemoveFromEnd(const FString& InSuffix, ESearchCase::Type SearchCase = ESearchCase::IgnoreCase)
-	{
-		return RemoveFromEnd(*InSuffix, InSuffix.Len(), SearchCase);
-	}
-
-	/**
-	 * Removes the text from the end of the string if it exists.
-	 *
-	 * @param InSuffix the suffix to search for at the end of the string to remove.
-	 * @param InSuffixLen length of InSuffix
-	 * @return true if the suffix was removed, otherwise false.
-	 */
-	CORE_API bool RemoveFromEnd(const ElementType* InSuffix, int32 InSuffixLen, ESearchCase::Type SearchCase = ESearchCase::IgnoreCase);
-
-	/**
-	 * Concatenate this path with given path ensuring the / character is used between them
-	 *
-	 * @param Str       Pointer to an array of TCHARs (not necessarily null-terminated) to be concatenated onto the end of this.
-	 * @param StrLength Exact number of characters from Str to append.
-	 */
-	CORE_API void PathAppend(const ElementType* Str, int32 StrLength);
-
-	/**
-	 * Concatenates a string with a character.
-	 * 
-	 * @param Lhs The string on the left-hand-side of the expression.
-	 * @param Rhs The char on the right-hand-side of the expression.
-	 *
-	 * @return The concatenated string.
-	 */
-	template <typename CharType>
-	UE_NODISCARD FORCEINLINE friend typename TEnableIf<TIsCharType<CharType>::Value, FString>::Type operator+(const FString& Lhs, CharType Rhs)
-	{
-		Lhs.CheckInvariants();
-
-		FString Result(Lhs, 1);
-		Result += Rhs;
-
-		return Result;
-	}
-
-	/**
-	 * Concatenates a string with a character.
-	 * 
-	 * @param Lhs The string on the left-hand-side of the expression.
-	 * @param Rhs The char on the right-hand-side of the expression.
-	 *
-	 * @return The concatenated string.
-	 */
-	template <typename CharType>
-	UE_NODISCARD FORCEINLINE friend typename TEnableIf<TIsCharType<CharType>::Value, FString>::Type operator+(FString&& Lhs, CharType Rhs)
-	{
-		Lhs.CheckInvariants();
-
-		FString Result(MoveTemp(Lhs), 1);
-		Result += Rhs;
-
-		return Result;
-	}
-
-private:
-	UE_NODISCARD static CORE_API FString ConcatFF(const FString& Lhs, const FString& Rhs);
-	UE_NODISCARD static CORE_API FString ConcatFF(FString&& Lhs, const FString& Rhs);
-	UE_NODISCARD static CORE_API FString ConcatFF(const FString& Lhs, FString&& Rhs);
-	UE_NODISCARD static CORE_API FString ConcatFF(FString&& Lhs, FString&& Rhs);
-	UE_NODISCARD static CORE_API FString ConcatFC(const FString& Lhs, const ElementType* Rhs);
-	UE_NODISCARD static CORE_API FString ConcatFC(FString&& Lhs,	const ElementType* Rhs);
-	UE_NODISCARD static CORE_API FString ConcatCF(const ElementType* Lhs, const FString& Rhs);
-	UE_NODISCARD static CORE_API FString ConcatCF(const ElementType* Lhs, FString&& Rhs);
-	UE_NODISCARD static CORE_API FString ConcatFR(const FString& Lhs, const ElementType* Rhs, int32 RhsLen);
-	UE_NODISCARD static CORE_API FString ConcatFR(FString&& Lhs,	const ElementType* Rhs, int32 RhsLen);
-	UE_NODISCARD static CORE_API FString ConcatRF(const ElementType* Lhs, int32 LhsLen, const FString& Rhs);
-	UE_NODISCARD static CORE_API FString ConcatRF(const ElementType* Lhs, int32 LhsLen, FString&& Rhs);
-
-public:
-	UE_NODISCARD FORCEINLINE friend FString operator+(const FString& Lhs, const FString& Rhs)			{ return ConcatFF(Lhs, Rhs); }
-	UE_NODISCARD FORCEINLINE friend FString operator+(FString&& Lhs, const FString& Rhs)				{ return ConcatFF(MoveTemp(Lhs), Rhs); }
-	UE_NODISCARD FORCEINLINE friend FString operator+(const FString& Lhs, FString&& Rhs)				{ return ConcatFF(Lhs,MoveTemp(Rhs)); }
-	UE_NODISCARD FORCEINLINE friend FString operator+(FString&& Lhs, FString&& Rhs)						{ return ConcatFF(MoveTemp(Lhs), MoveTemp(Rhs)); }
-	UE_NODISCARD FORCEINLINE friend FString operator+(const ElementType* Lhs, const FString& Rhs)		{ return ConcatCF(Lhs, Rhs); }
-	UE_NODISCARD FORCEINLINE friend FString operator+(const ElementType* Lhs, FString&& Rhs)			{ return ConcatCF(Lhs, MoveTemp(Rhs)); }
-	UE_NODISCARD FORCEINLINE friend FString operator+(const FString& Lhs, const ElementType* Rhs)		{ return ConcatFC(Lhs, Rhs); }
-	UE_NODISCARD FORCEINLINE friend FString operator+(FString&& Lhs, const ElementType* Rhs)			{ return ConcatFC(MoveTemp(Lhs), Rhs); }
-
-	template <typename T, typename TEnableIf<TIsTCharRangeNotCArrayNotStringClass<T>::Value>::Type* = nullptr>
-	UE_NODISCARD FORCEINLINE friend FString operator+(T&& Lhs, const FString& Rhs)						{ return ConcatRF(GetData(Lhs), GetNum(Lhs), Rhs); }
-	template <typename T, typename TEnableIf<TIsTCharRangeNotCArrayNotStringClass<T>::Value>::Type* = nullptr>
-	UE_NODISCARD FORCEINLINE friend FString operator+(T&& Lhs, FString&& Rhs)							{ return ConcatRF(GetData(Lhs), GetNum(Lhs), MoveTemp(Rhs)); }
-	template <typename T, typename TEnableIf<TIsTCharRangeNotCArrayNotStringClass<T>::Value>::Type* = nullptr>
-	UE_NODISCARD FORCEINLINE friend FString operator+(const FString& Lhs, T&& Rhs)						{ return ConcatFR(Lhs, GetData(Rhs), GetNum(Rhs)); }
-	template <typename T, typename TEnableIf<TIsTCharRangeNotCArrayNotStringClass<T>::Value>::Type* = nullptr>
-	UE_NODISCARD FORCEINLINE friend FString operator+(FString&& Lhs, T&& Rhs)							{ return ConcatFR(MoveTemp(Lhs), GetData(Rhs), GetNum(Rhs)); }
-
-	/**
-	 * Concatenate this path with given path ensuring the / character is used between them
-	 * 
-	 * @param Str path array of characters to be concatenated onto the end of this
-	 * @return reference to path
-	 */
-	FORCEINLINE FString& operator/=( const ElementType* Str )
-	{
-		checkSlow(Str);
-
-		PathAppend(Str, FCString::Strlen(Str));
-		return *this;
-	}
-
-	/**
-	* Concatenate this path with given path ensuring the / character is used between them
-	* 
-	* @param Str path CharRangeType (string class/string view/string builder etc) to be concatenated onto the end of this
-	* @return reference to path
-	*/
-	template <typename CharRangeType, typename TEnableIf<TIsTCharRangeNotCArray <CharRangeType>::Value>::Type* = nullptr>
-	FORCEINLINE FString& operator/=(CharRangeType&& Str)
-	{
-		PathAppend(GetData(Str), GetNum(Str));
-		return *this;
-	}
-
-	/**
-	* Concatenate this path with given path ensuring the / character is used between them
-	* 
-	* @param Str path array of CharType (that needs converting) to be concatenated onto the end of this
-	* @return reference to path
-	*/
-	template <typename CharType,typename TEnableIf<TIsCharType<CharType>::Value>::Type* = nullptr>
-	FORCEINLINE FString& operator/=(const CharType* Str)
-	{
-		FString Temp = Str;
-		PathAppend(GetData(Temp), GetNum(Temp));
-		return *this;
-	}
-
-	/**
-	 * Concatenate this path with given path ensuring the / character is used between them
-	 *
-	 * @param Lhs Path to concatenate onto.
-	 * @param Rhs Path to concatenate.
-	 * @return The new concatenated path
-	 */
-	UE_NODISCARD FORCEINLINE friend FString operator/(const FString& Lhs, const ElementType* Rhs)
-	{
-		checkSlow(Rhs);
-
-		int32 StrLength = FCString::Strlen(Rhs);
-
-		FString Result(Lhs, StrLength + 1);
-		Result.PathAppend(Rhs, StrLength);
-		return Result;
-	}
-
-	/**
-	 * Concatenate this path with given path ensuring the / character is used between them
-	 *
-	 * @param Lhs Path to concatenate onto.
-	 * @param Rhs Path to concatenate.
-	 * @return The new concatenated path
-	 */
-	UE_NODISCARD FORCEINLINE friend FString operator/(FString&& Lhs, const ElementType* Rhs)
-	{
-		checkSlow(Rhs);
-
-		int32 StrLength = FCString::Strlen(Rhs);
-
-		FString Result(MoveTemp(Lhs), StrLength + 1);
-		Result.PathAppend(Rhs, StrLength);
-		return Result;
-	}
-
-	/**
-	 * Concatenate this path with given path ensuring the / character is used between them
-	 *
-	 * @param Lhs Path to concatenate onto.
-	 * @param Rhs Path to concatenate.
-	 * @return The new concatenated path
-	 */
-	UE_NODISCARD FORCEINLINE friend FString operator/(const FString& Lhs, const FString& Rhs)
-	{
-		int32 StrLength = Rhs.Len();
-
-		FString Result(Lhs, StrLength + 1);
-		Result.PathAppend(Rhs.Data.GetData(), StrLength);
-		return Result;
-	}
-
-	/**
-	 * Concatenate this path with given path ensuring the / character is used between them
-	 *
-	 * @param Lhs Path to concatenate onto.
-	 * @param Rhs Path to concatenate.
-	 * @return The new concatenated path
-	 */
-	UE_NODISCARD FORCEINLINE friend FString operator/(FString&& Lhs, const FString& Rhs)
-	{
-		int32 StrLength = Rhs.Len();
-
-		FString Result(MoveTemp(Lhs), StrLength + 1);
-		Result.PathAppend(Rhs.Data.GetData(), StrLength);
-		return Result;
-	}
-
-	/**
-	 * Concatenate this path with given path ensuring the / character is used between them
-	 *
-	 * @param Lhs Path to concatenate onto.
-	 * @param Rhs Path to concatenate.
-	 * @return new string of the path
-	 */
-	UE_NODISCARD FORCEINLINE friend FString operator/(const ElementType* Lhs, const FString& Rhs)
-	{
-		int32 StrLength = Rhs.Len();
-
-		FString Result(FString(Lhs), StrLength + 1);
-		Result.PathAppend(Rhs.Data.GetData(), Rhs.Len());
-		return Result;
-	}
-
-	/**
-	 * Lexicographically test whether the left string is <= the right string
-	 *
-	 * @param Lhs String to compare against.
-	 * @param Rhs String to compare against.
-	 * @return true if the left string is lexicographically <= the right string, otherwise false
-	 * @note case insensitive
-	 */
-	UE_NODISCARD FORCEINLINE friend bool operator<=(const FString& Lhs, const FString& Rhs)
-	{
-		return FPlatformString::Stricmp(*Lhs, *Rhs) <= 0;
-	}
-
-	/**
-	 * Lexicographically test whether the left string is <= the right string
-	 *
-	 * @param Lhs String to compare against.
-	 * @param Rhs String to compare against.
-	 * @return true if the left string is lexicographically <= the right string, otherwise false
-	 * @note case insensitive
-	 */
-	template <typename CharType>
-	UE_NODISCARD FORCEINLINE friend bool operator<=(const FString& Lhs, const CharType* Rhs)
-	{
-		return FPlatformString::Stricmp(*Lhs, Rhs) <= 0;
-	}
-
-	/**
-	 * Lexicographically test whether the left string is <= the right string
-	 *
-	 * @param Lhs String to compare against.
-	 * @param Rhs String to compare against.
-	 * @return true if the left string is lexicographically <= the right string, otherwise false
-	 * @note case insensitive
-	 */
-	template <typename CharType>
-	UE_NODISCARD FORCEINLINE friend bool operator<=(const CharType* Lhs, const FString& Rhs)
-	{
-		return FPlatformString::Stricmp(Lhs, *Rhs) <= 0;
-	}
-
-	/**
-	 * Lexicographically test whether the left string is < the right string
-	 *
-	 * @param Lhs String to compare against.
-	 * @param Rhs String to compare against.
-	 * @return true if the left string is lexicographically < the right string, otherwise false
-	 * @note case insensitive
-	 */
-	UE_NODISCARD FORCEINLINE friend bool operator<(const FString& Lhs, const FString& Rhs)
-	{
-		return FPlatformString::Stricmp(*Lhs, *Rhs) < 0;
-	}
-
-	/**
-	 * Lexicographically test whether the left string is < the right string
-	 *
-	 * @param Lhs String to compare against.
-	 * @param Rhs String to compare against.
-	 * @return true if the left string is lexicographically < the right string, otherwise false
-	 * @note case insensitive
-	 */
-	template <typename CharType>
-	UE_NODISCARD FORCEINLINE friend bool operator<(const FString& Lhs, const CharType* Rhs)
-	{
-		return FPlatformString::Stricmp(*Lhs, Rhs) < 0;
-	}
-
-	/**
-	 * Lexicographically test whether the left string is < the right string
-	 *
-	 * @param Lhs String to compare against.
-	 * @param Rhs String to compare against.
-	 * @return true if the left string is lexicographically < the right string, otherwise false
-	 * @note case insensitive
-	 */
-	template <typename CharType>
-	UE_NODISCARD FORCEINLINE friend bool operator<(const CharType* Lhs, const FString& Rhs)
-	{
-		return FPlatformString::Stricmp(Lhs, *Rhs) < 0;
-	}
-
-	/**
-	 * Lexicographically test whether the left string is >= the right string
-	 *
-	 * @param Lhs String to compare against.
-	 * @param Rhs String to compare against.
-	 * @return true if the left string is lexicographically >= the right string, otherwise false
-	 * @note case insensitive
-	 */
-	UE_NODISCARD FORCEINLINE friend bool operator>=(const FString& Lhs, const FString& Rhs)
-	{
-		return FPlatformString::Stricmp(*Lhs, *Rhs) >= 0;
-	}
-
-	/**
-	 * Lexicographically test whether the left string is >= the right string
-	 *
-	 * @param Lhs String to compare against.
-	 * @param Rhs String to compare against.
-	 * @return true if the left string is lexicographically >= the right string, otherwise false
-	 * @note case insensitive
-	 */
-	template <typename CharType>
-	UE_NODISCARD FORCEINLINE friend bool operator>=(const FString& Lhs, const CharType* Rhs)
-	{
-		return FPlatformString::Stricmp(*Lhs, Rhs) >= 0;
-	}
-
-	/**
-	 * Lexicographically test whether the left string is >= the right string
-	 *
-	 * @param Lhs String to compare against.
-	 * @param Rhs String to compare against.
-	 * @return true if the left string is lexicographically >= the right string, otherwise false
-	 * @note case insensitive
-	 */
-	template <typename CharType>
-	UE_NODISCARD FORCEINLINE friend bool operator>=(const CharType* Lhs, const FString& Rhs)
-	{
-		return FPlatformString::Stricmp(Lhs, *Rhs) >= 0;
-	}
-
-	/**
-	 * Lexicographically test whether the left string is > the right string
-	 *
-	 * @param Lhs String to compare against.
-	 * @param Rhs String to compare against.
-	 * @return true if the left string is lexicographically > the right string, otherwise false
-	 * @note case insensitive
-	 */
-	UE_NODISCARD FORCEINLINE friend bool operator>(const FString& Lhs, const FString& Rhs)
-	{
-		return FPlatformString::Stricmp(*Lhs, *Rhs) > 0;
-	}
-
-	/**
-	 * Lexicographically test whether the left string is > the right string
-	 *
-	 * @param Lhs String to compare against.
-	 * @param Rhs String to compare against.
-	 * @return true if the left string is lexicographically > the right string, otherwise false
-	 * @note case insensitive
-	 */
-	template <typename CharType>
-	UE_NODISCARD FORCEINLINE friend bool operator>(const FString& Lhs, const CharType* Rhs)
-	{
-		return FPlatformString::Stricmp(*Lhs, Rhs) > 0;
-	}
-
-	/**
-	 * Lexicographically test whether the left string is > the right string
-	 *
-	 * @param Lhs String to compare against.
-	 * @param Rhs String to compare against.
-	 * @return true if the left string is lexicographically > the right string, otherwise false
-	 * @note case insensitive
-	 */
-	template <typename CharType>
-	UE_NODISCARD FORCEINLINE friend bool operator>(const CharType* Lhs, const FString& Rhs)
-	{
-		return FPlatformString::Stricmp(Lhs, *Rhs) > 0;
-	}
-
-	/**
-	 * Lexicographically test whether the left string is == the right string
-	 *
-	 * @param Lhs String to compare against.
-	 * @param Rhs String to compare against.
-	 * @return true if the left string is lexicographically == the right string, otherwise false
-	 * @note case insensitive
-	 */
-	UE_NODISCARD FORCEINLINE bool operator==(const FString& Rhs) const
-	{
-		return Equals(Rhs, ESearchCase::IgnoreCase);
-	}
-
-	/**
-	 * Lexicographically test whether the left string is == the right string
-	 *
-	 * @param Lhs String to compare against.
-	 * @param Rhs String to compare against.
-	 * @return true if the left string is lexicographically == the right string, otherwise false
-	 * @note case insensitive
-	 */
-	template <typename CharType>
-	UE_NODISCARD FORCEINLINE bool operator==(const CharType* Rhs) const
-	{
-		return FPlatformString::Stricmp(**this, Rhs) == 0;
-	}
-
-	/**
-	 * Lexicographically test whether the left string is == the right string
-	 *
-	 * @param Lhs String to compare against.
-	 * @param Rhs String to compare against.
-	 * @return true if the left string is lexicographically == the right string, otherwise false
-	 * @note case insensitive
-	 */
-	template <typename CharType>
-	UE_NODISCARD FORCEINLINE friend bool operator==(const CharType* Lhs, const FString& Rhs)
-	{
-		return FPlatformString::Stricmp(Lhs, *Rhs) == 0;
-	}
-
-	/**
-	 * Lexicographically test whether the left string is != the right string
-	 *
-	 * @param Lhs String to compare against.
-	 * @param Rhs String to compare against.
-	 * @return true if the left string is lexicographically != the right string, otherwise false
-	 * @note case insensitive
-	 */
-	UE_NODISCARD FORCEINLINE bool operator!=(const FString& Rhs) const
-	{
-		return !Equals(Rhs, ESearchCase::IgnoreCase);
-	}
-
-	/**
-	 * Lexicographically test whether the left string is != the right string
-	 *
-	 * @param Lhs String to compare against.
-	 * @param Rhs String to compare against.
-	 * @return true if the left string is lexicographically != the right string, otherwise false
-	 * @note case insensitive
-	 */
-	template <typename CharType>
-	UE_NODISCARD FORCEINLINE bool operator!=(const CharType* Rhs) const
-	{
-		return FPlatformString::Stricmp(**this, Rhs) != 0;
-	}
-
-	/**
-	 * Lexicographically test whether the left string is != the right string
-	 *
-	 * @param Lhs String to compare against.
-	 * @param Rhs String to compare against.
-	 * @return true if the left string is lexicographically != the right string, otherwise false
-	 * @note case insensitive
-	 */
-	template <typename CharType>
-	UE_NODISCARD FORCEINLINE friend bool operator!=(const CharType* Lhs, const FString& Rhs)
-	{
-		return FPlatformString::Stricmp(Lhs, *Rhs) != 0;
-	}
-
-	/** Get the length of the string, excluding terminating character */
-	UE_NODISCARD FORCEINLINE int32 Len() const
-	{
-		return Data.Num() ? Data.Num() - 1 : 0;
-	}
-
-	/** Returns the left most given number of characters */
-	UE_NODISCARD FORCEINLINE FString Left( int32 Count ) const &
-	{
-		return FString( FMath::Clamp(Count,0,Len()), **this );
-	}
-
-	UE_NODISCARD FORCEINLINE FString Left(int32 Count) &&
-	{
-		LeftInline(Count, false);
-		return MoveTemp(*this);
-	}
-
-	/** Modifies the string such that it is now the left most given number of characters */
-	FORCEINLINE void LeftInline(int32 Count, bool bAllowShrinking = true)
-	{
-		const int32 Length = Len();
-		Count = FMath::Clamp(Count, 0, Length);
-		RemoveAt(Count, Length-Count, bAllowShrinking);
-	}
-
-	/** Returns the left most characters from the string chopping the given number of characters from the end */
-	UE_NODISCARD FORCEINLINE FString LeftChop( int32 Count ) const &
-	{
-		const int32 Length = Len();
-		return FString( FMath::Clamp(Length-Count,0, Length), **this );
-	}
-
-	UE_NODISCARD FORCEINLINE FString LeftChop(int32 Count)&&
-	{
-		LeftChopInline(Count, false);
-		return MoveTemp(*this);
-	}
-
-	/** Modifies the string such that it is now the left most characters chopping the given number of characters from the end */
-	FORCEINLINE void LeftChopInline(int32 Count, bool bAllowShrinking = true)
-	{
-		const int32 Length = Len();
-		RemoveAt(FMath::Clamp(Length-Count, 0, Length), Count, bAllowShrinking);
-	}
-
-	/** Returns the string to the right of the specified location, counting back from the right (end of the word). */
-	UE_NODISCARD FORCEINLINE FString Right( int32 Count ) const &
-	{
-		const int32 Length = Len();
-		return FString( **this + Length-FMath::Clamp(Count,0,Length) );
-	}
-
-	UE_NODISCARD FORCEINLINE FString Right(int32 Count) &&
-	{
-		RightInline(Count, false);
-		return MoveTemp(*this);
-	}
-
-	/** Modifies the string such that it is now the right most given number of characters */
-	FORCEINLINE void RightInline(int32 Count, bool bAllowShrinking = true)
-	{
-		const int32 Length = Len();
-		RemoveAt(0, Length-FMath::Clamp(Count,0,Length), bAllowShrinking);
-	}
-
-	/** Returns the string to the right of the specified location, counting forward from the left (from the beginning of the word). */
-	UE_NODISCARD CORE_API FString RightChop( int32 Count ) const &;
-
-	UE_NODISCARD FORCEINLINE FString RightChop(int32 Count) &&
-	{
-		RightChopInline(Count, false);
-		return MoveTemp(*this);
-	}
-
-	/** Modifies the string such that it is now the string to the right of the specified location, counting forward from the left (from the beginning of the word). */
-	FORCEINLINE void RightChopInline(int32 Count, bool bAllowShrinking = true)
-	{
-		RemoveAt(0, Count, bAllowShrinking);
-	}
-
-	/** Returns the substring from Start position for Count characters. */
-	UE_NODISCARD CORE_API FString Mid(int32 Start, int32 Count) const &;
-	UE_NODISCARD CORE_API FString Mid(int32 Start, int32 Count) &&;
-
-	/** Returns the substring from Start position to the end */
-	UE_NODISCARD FORCEINLINE FString Mid(int32 Start) const & { return RightChop(Start); }
-	UE_NODISCARD FORCEINLINE FString Mid(int32 Start) && { return ((FString&&)*this).RightChop(Start); }
-
-	/** Modifies the string such that it is now the substring from Start position for Count characters. */
-	FORCEINLINE void MidInline(int32 Start, int32 Count = MAX_int32, bool bAllowShrinking = true)
-	{
-		if (Count != MAX_int32 && int64(Start) + Count < MAX_int32)
-		{
-			LeftInline(Count + Start, false);
-		}
-		RightChopInline(Start, bAllowShrinking);
-	}
-
-	/**
-	 * Searches the string for a substring, and returns index into this string of the first found instance. Can search
-	 * from beginning or end, and ignore case or not. If substring is empty, returns clamped StartPosition.
-	 *
-	 * @param SubStr			The string to search for
-	 * @param StartPosition		The start character position to search from
-	 * @param SearchCase		Indicates whether the search is case sensitive or not ( defaults to ESearchCase::IgnoreCase )
-	 * @param SearchDir			Indicates whether the search starts at the beginning or at the end ( defaults to ESearchDir::FromStart )
-	 *
-	 * @note  When searching backwards where a StartPosition is provided, searching will actually begin from
-	 *        StartPosition - SubStr.Len(), therefore:
-	 *
-	 *        FString("X").Find("X", ESearchCase::CaseSensitive, ESearchDir::FromEnd, 0) == INDEX_NONE
-	 *
-	 *        Consider using UE::String::FindLast() as an alternative.
-	 */
-	template <typename TCharRangeType, std::enable_if_t<TIsCharRangeNotCArrayNotStringClass<TCharRangeType>::Value>* = nullptr>
-	UE_NODISCARD int32 Find(TCharRangeType&& SubStr, ESearchCase::Type SearchCase = ESearchCase::IgnoreCase,
-		ESearchDir::Type SearchDir = ESearchDir::FromStart, int32 StartPosition = INDEX_NONE) const
-	{
-		static_assert(std::is_same_v<typename TElementType<TCharRangeType>::Type, ElementType>, "Expected a range of ElementType");
-		return Find(GetData(SubStr), GetNum(SubStr), SearchCase, SearchDir, StartPosition);
-	}
-
-	/**
-	 * Searches the string for a substring, and returns index into this string of the first found instance. Can search
-	 * from beginning or end, and ignore case or not. If substring is empty, returns clamped StartPosition.
-	 *
-	 * @param SubStr			The string array of characters to search for
-	 * @param StartPosition		The start character position to search from.  See note below.
-	 * @param SearchCase		Indicates whether the search is case sensitive or not
-	 * @param SearchDir			Indicates whether the search starts at the beginning or at the end.
-	 */
-	UE_NODISCARD int32 Find(const ElementType* SubStr, ESearchCase::Type SearchCase = ESearchCase::IgnoreCase,
-		ESearchDir::Type SearchDir = ESearchDir::FromStart, int32 StartPosition = INDEX_NONE) const
-	{
-		return SubStr ? Find(SubStr, FCString::Strlen(SubStr), SearchCase, SearchDir, StartPosition) : INDEX_NONE;
-	}
-
-	/**
-	 * Searches the string for a substring, and returns index into this string of the first found instance. Can search
-	 * from beginning or end, and ignore case or not. If substring is empty, returns clamped StartPosition.
-	 *
-	 * @param SubStr			The string array of characters to search for
-	 * @param StartPosition		The start character position to search from.  See note below.
-	 * @param SearchCase		Indicates whether the search is case sensitive or not
-	 * @param SearchDir			Indicates whether the search starts at the beginning or at the end.
-	 *
-	 * @note  When searching backwards where a StartPosition is provided, searching will actually begin from
-	 *        StartPosition - SubStr.Len(), therefore:
-	 *
-	 *        FString("X").Find("X", ESearchCase::CaseSensitive, ESearchDir::FromEnd, 0) == INDEX_NONE
-	 *
-	 *        Consider using UE::String::FindLast() as an alternative.
-	 */
-	UE_NODISCARD int32 Find(const FString& SubStr, ESearchCase::Type SearchCase = ESearchCase::IgnoreCase,
-		ESearchDir::Type SearchDir = ESearchDir::FromStart, int32 StartPosition = INDEX_NONE) const
-	{
-		return Find(*SubStr, SubStr.Len(), SearchCase, SearchDir, StartPosition);
-	}
-
-	/**
-	 * Searches the string for a substring, and returns index into this string of the first found instance. Can search
-	 * from beginning or end, and ignore case or not. If substring is empty, returns clamped StartPosition.
-	 *
-	 * @param SubStr			The string array of characters to search for
-	 * @param SubStrLen			The length of the SubStr array
-	 * @param StartPosition		The start character position to search from.  See note below.
-	 * @param SearchCase		Indicates whether the search is case sensitive or not
-	 * @param SearchDir			Indicates whether the search starts at the beginning or at the end.
-	 *
-	 * @note  When searching backwards where a StartPosition is provided, searching will actually begin from
-	 *        StartPosition - SubStr.Len(), therefore:
-	 *
-	 *        FString("X").Find("X", ESearchCase::CaseSensitive, ESearchDir::FromEnd, 0) == INDEX_NONE
-	 *
-	 *        Consider using UE::String::FindLast() as an alternative.
-	 */
-	UE_NODISCARD CORE_API int32 Find(const ElementType* SubStr, int32 InSubStrLen, ESearchCase::Type SearchCase = ESearchCase::IgnoreCase,
-		ESearchDir::Type SearchDir = ESearchDir::FromStart, int32 StartPosition = INDEX_NONE) const;
-
-
-	/**
-	 * Returns whether this string contains the specified substring.
-	 *
-	 * @param SubStr			Text to search for
-	 * @param SearchCase		Indicates whether the search is case sensitive or not ( defaults to ESearchCase::IgnoreCase )
-	 * @param SearchDir			Indicates whether the search starts at the beginning or at the end ( defaults to ESearchDir::FromStart )
-	 * @return					Returns whether the string contains the substring. If the substring is empty, returns true.
-	 **/
-	template <typename TCharRangeType, std::enable_if_t<TIsCharRangeNotCArrayNotStringClass<TCharRangeType>::Value>* = nullptr>
-	UE_NODISCARD FORCEINLINE bool Contains(TCharRangeType&& SubStr, ESearchCase::Type SearchCase = ESearchCase::IgnoreCase,
-		ESearchDir::Type SearchDir = ESearchDir::FromStart) const
-	{
-		static_assert(std::is_same_v<typename TElementType<TCharRangeType>::Type, ElementType>, "Expected a range of characters");
-		return Find(Forward<TCharRangeType>(SubStr), SearchCase, SearchDir) != INDEX_NONE;
-	}
-
-	/** 
-	 * Returns whether this string contains the specified substring.
-	 *
-	 * @param SubStr			Text to search for
-	 * @param SearchCase		Indicates whether the search is case sensitive or not ( defaults to ESearchCase::IgnoreCase )
-	 * @param SearchDir			Indicates whether the search starts at the beginning or at the end ( defaults to ESearchDir::FromStart )
-	 * @return					Returns whether the string contains the substring. If the substring is empty, returns true.
-	 **/
-	UE_NODISCARD FORCEINLINE bool Contains(const ElementType* SubStr, ESearchCase::Type SearchCase = ESearchCase::IgnoreCase,
-		ESearchDir::Type SearchDir = ESearchDir::FromStart) const
-	{
-		return Find(SubStr, SearchCase, SearchDir) != INDEX_NONE;
-	}
-
-	/** 
-	 * Returns whether this string contains the specified substring.
-	 *
-	 * @param SubStr			Text to search for
-	 * @param SearchCase		Indicates whether the search is case sensitive or not ( defaults to ESearchCase::IgnoreCase )
-	 * @param SearchDir			Indicates whether the search starts at the beginning or at the end ( defaults to ESearchDir::FromStart )
-	 * @return					Returns whether the string contains the substring. If the substring is empty, returns true.
-	 **/
-	UE_NODISCARD FORCEINLINE bool Contains(const FString& SubStr, ESearchCase::Type SearchCase = ESearchCase::IgnoreCase,
-							  ESearchDir::Type SearchDir = ESearchDir::FromStart ) const
-	{
-		return Find(*SubStr, SearchCase, SearchDir) != INDEX_NONE;
-	}
-
-	/**
-	 * Returns whether this string contains the specified substring.
-	 *
-	 * @param SubStr			Text to search for
-	 * @param SubStrLen			Length of the Text
-	 * @param SearchCase		Indicates whether the search is case sensitive or not ( defaults to ESearchCase::IgnoreCase )
-	 * @param SearchDir			Indicates whether the search starts at the beginning or at the end ( defaults to ESearchDir::FromStart )
-	 * @return					Returns whether the string contains the substring. If the substring is empty, returns true.
-	 **/
-	UE_NODISCARD FORCEINLINE bool Contains(const ElementType* SubStr, int32 SubStrLen,
-		ESearchCase::Type SearchCase = ESearchCase::IgnoreCase, ESearchDir::Type SearchDir = ESearchDir::FromStart) const
-	{
-		return Find(SubStr, SubStrLen, SearchCase, SearchDir) != INDEX_NONE;
-	}
-
-	/**
-	 * Searches the string for a character
-	 *
-	 * @param InChar the character to search for
-	 * @param Index out the position the character was found at, INDEX_NONE if return is false
-	 * @return true if character was found in this string, otherwise false
-	 */
-	FORCEINLINE bool FindChar(ElementType InChar, int32& Index ) const
-	{
-		return Data.Find(InChar, Index);
-	}
-
-	/**
-	 * Searches the string for the last occurrence of a character
-	 *
-	 * @param InChar the character to search for
-	 * @param Index out the position the character was found at, INDEX_NONE if return is false
-	 * @return true if character was found in this string, otherwise false
-	 */
-	FORCEINLINE bool FindLastChar( ElementType InChar, int32& Index ) const
-	{
-		return Data.FindLast(InChar, Index);
-	}
-
-	/**
-	 * Searches an initial substring for the last occurrence of a character which matches the specified predicate.
-	 *
-	 * @param Pred Predicate that takes a character and returns true if it matches search criteria, false otherwise.
-	 * @param Count The number of characters from the front of the string through which to search.
-	 *
-	 * @return Index of the found character, INDEX_NONE otherwise.
-	 */
-	template <typename Predicate>
-	FORCEINLINE int32 FindLastCharByPredicate(Predicate Pred, int32 Count) const
-	{
-		check(Count >= 0 && Count <= this->Len());
-		return Data.FindLastByPredicate(Pred, Count);
-	}
-
-	/**
-	 * Searches the string for the last occurrence of a character which matches the specified predicate.
-	 *
-	 * @param Pred Predicate that takes a character and returns true if it matches search criteria, false otherwise.
-	 * @param StartIndex Index of element from which to start searching. Defaults to the last character in string.
-	 *
-	 * @return Index of the found character, INDEX_NONE otherwise.
-	 */
-	template <typename Predicate>
-	FORCEINLINE int32 FindLastCharByPredicate(Predicate Pred) const
-	{
-		return Data.FindLastByPredicate(Pred, this->Len());
-	}
-
-	/**
-	 * Lexicographically tests whether this string is equivalent to the Other given string
-	 * 
-	 * @param Other 	The string test against
-	 * @param SearchCase 	Whether or not the comparison should ignore case
-	 * @return true if this string is lexicographically equivalent to the other, otherwise false
-	 */
-	UE_NODISCARD FORCEINLINE bool Equals(const FString& Other, ESearchCase::Type SearchCase = ESearchCase::CaseSensitive) const
-	{
-		int32 Num = Data.Num();
-		int32 OtherNum = Other.Data.Num();
-
-		if (Num != OtherNum)
-		{
-			// Handle special case where FString() == FString("")
-			return Num + OtherNum == 1;
-		}
-		else if (Num > 1)
-		{
-			if (SearchCase == ESearchCase::CaseSensitive)
-			{
-				return FCString::Strcmp(Data.GetData(), Other.Data.GetData()) == 0; 
-			}
-			else
-			{
-				return FCString::Stricmp(Data.GetData(), Other.Data.GetData()) == 0;
-			}
-		}
-
-		return true;
-	}
-
-	/**
-	 * Lexicographically tests how this string compares to the Other given string
-	 * 
-	 * @param Other 	The string test against
-	 * @param SearchCase 	Whether or not the comparison should ignore case
-	 * @return 0 if equal, negative if less than, positive if greater than
-	 */
-	UE_NODISCARD FORCEINLINE int32 Compare( const FString& Other, ESearchCase::Type SearchCase = ESearchCase::CaseSensitive ) const
-	{
-		if( SearchCase == ESearchCase::CaseSensitive )
-		{
-			return FCString::Strcmp( **this, *Other ); 
-		}
-		else
-		{
-			return FCString::Stricmp( **this, *Other );
-		}
-	}
-
-	/**
-	 * Splits this string at given string position case sensitive.
-	 *
-	 * @param InStr The string to search and split at
-	 * @param LeftS out the string to the left of InStr, not updated if return is false. LeftS must not point to the same location as RightS, but can point to this.
-	 * @param RightS out the string to the right of InStr, not updated if return is false. RightS must not point to the same location as LeftS, but can point to this.
-	 * @param SearchCase		Indicates whether the search is case sensitive or not ( defaults to ESearchCase::IgnoreCase )
-	 * @param SearchDir			Indicates whether the search starts at the beginning or at the end ( defaults to ESearchDir::FromStart )
-	 * @return true if string is split, otherwise false
-	 */
-	CORE_API bool Split(const FString& InS, FString* LeftS, FString* RightS, ESearchCase::Type SearchCase,
-		ESearchDir::Type SearchDir = ESearchDir::FromStart) const;
-
-	/** Split with ESearchCase::IgnoreCase and ESearchDir::FromStart. Allows compiler to avoid branches w/o inlining code. */
-	CORE_API bool Split(const FString& InS, FString* LeftS, FString* RightS) const;
-
-	/** Returns a new string with the characters of this converted to uppercase */
-	UE_NODISCARD CORE_API FString ToUpper() const &;
-
-	/**
-	 * Converts all characters in this rvalue string to uppercase and moves it into the returned string.
-	 * @return a new string with the characters of this converted to uppercase
-	 */
-	UE_NODISCARD CORE_API FString ToUpper() &&;
-
-	/** Converts all characters in this string to uppercase */
-	CORE_API void ToUpperInline();
-
-	/** Returns a new string with the characters of this converted to lowercase */
-	UE_NODISCARD CORE_API FString ToLower() const &;
-
-	/**
-	 * Converts all characters in this rvalue string to lowercase and moves it into the returned string.
-	 * @return a new string with the characters of this converted to lowercase
-	 */
-	UE_NODISCARD CORE_API FString ToLower() &&;
-
-	/** Converts all characters in this string to lowercase */
-	CORE_API void ToLowerInline();
-
-	/** Pad the left of this string for ChCount characters */
-	UE_NODISCARD CORE_API FString LeftPad( int32 ChCount ) const;
-
-	/** Pad the right of this string for ChCount characters */
-	UE_NODISCARD CORE_API FString RightPad( int32 ChCount ) const;
-	
-	/** Returns true if the string only contains numeric characters */
-	UE_NODISCARD CORE_API bool IsNumeric() const;
-	
-	/** Removes spaces from the string.  I.E. "Spaces Are Cool" --> "SpacesAreCool". */
-	CORE_API void RemoveSpacesInline();
-
-	/**
-	 * Constructs a string similarly to how classic sprintf works.
-	 *
-	 * @param Format	Format string that specifies how the string should be built optionally using additional args. Refer to standard printf format.
-	 * @param ...		Depending on format function may require additional arguments to build output object.
-	 *
-	 * @returns A string that was constructed using format and additional parameters.
-	 */
-	template <typename FmtType, typename... Types>
-	UE_NODISCARD static FString Printf(const FmtType& Fmt, Types... Args)
-	{
-		static_assert(TIsArrayOrRefOfTypeByPredicate<FmtType, TIsCharEncodingCompatibleWithTCHAR>::Value, "Formatting string must be a character array.");
-		static_assert(TAnd<TIsValidVariadicFunctionArg<Types>...>::Value, "Invalid argument(s) passed to Printf");
-
-		return PrintfImpl((const ElementType*)Fmt, Args...);
-	}
-
-	/**
-	 * Just like Printf, but appends the formatted text to the existing string instead.
-	 * @return a reference to the modified string, so that it can be chained
-	 */
-	template <typename FmtType, typename... Types>
-	FString& Appendf(const FmtType& Fmt, Types... Args)
-	{
-		static_assert(TIsArrayOrRefOfTypeByPredicate<FmtType, TIsCharEncodingCompatibleWithTCHAR>::Value, "Formatting string must be a character array.");
-		static_assert(TAnd<TIsValidVariadicFunctionArg<Types>...>::Value, "Invalid argument(s) passed to TString::Appendf");
-
-		AppendfImpl(*this, (const ElementType*)Fmt, Args...);
-		return *this;
-	}
-
-private:
-	static CORE_API FString VARARGS PrintfImpl(const ElementType* Fmt, ...);
-	static CORE_API void VARARGS AppendfImpl(FString& AppendToMe, const ElementType* Fmt, ...);
-public:
-
-	/**
-	 * Format the specified string using the specified arguments. Replaces instances of { Argument } with keys in the map matching 'Argument'
-	 * @param InFormatString		A string representing the format expression
-	 * @param InNamedArguments		A map of named arguments that match the tokens specified in InExpression
-	 * @return A string containing the formatted text
-	 */
-	UE_NODISCARD static CORE_API FString Format(const ElementType* InFormatString, const FStringFormatNamedArguments& InNamedArguments);
-
-	/**
-	 * Format the specified string using the specified arguments. Replaces instances of {0} with indices from the given array matching the index specified in the token
-	 * @param InFormatString		A string representing the format expression
-	 * @param InOrderedArguments	An array of ordered arguments that match the tokens specified in InExpression
-	 * @return A string containing the formatted text
-	 */
-	UE_NODISCARD static CORE_API FString Format(const ElementType* InFormatString, const FStringFormatOrderedArguments& InOrderedArguments);
-
-	/** Returns a string containing only the Ch character */
-	UE_NODISCARD static CORE_API FString Chr( ElementType Ch );
-
-	/**
-	 * Returns a string that is full of a variable number of characters
-	 *
-	 * @param NumCharacters Number of characters to put into the string
-	 * @param Char Character to put into the string
-	 * 
-	 * @return The string of NumCharacters characters.
-	 */
-	UE_NODISCARD static CORE_API FString ChrN( int32 NumCharacters, ElementType Char );
-
-	/**
-	 * Serializes the string.
-	 *
-	 * @param Ar Reference to the serialization archive.
-	 * @param S Reference to the string being serialized.
-	 *
-	 * @return Reference to the Archive after serialization.
-	 */
-	friend CORE_API FArchive& operator<<( FArchive& Ar, FString& S );
-
-
-	/**
-	 * Test whether this string starts with given prefix.
-	 *
-	 * @param SearchCase		Indicates whether the search is case sensitive or not ( defaults to ESearchCase::IgnoreCase )
-	 * @return true if this string begins with specified text, false otherwise
-	 */
-	template <typename TCharRangeType, std::enable_if_t<TIsCharRangeNotCArrayNotStringClass<TCharRangeType>::Value>* = nullptr>
-	UE_NODISCARD bool StartsWith(TCharRangeType&& InPrefix, ESearchCase::Type SearchCase = ESearchCase::IgnoreCase) const
-	{
-		static_assert(std::is_same_v<typename TElementType<TCharRangeType>::Type, ElementType>, "Expected a range of characters");
-		return StartsWith(GetData(InPrefix), GetNum(InPrefix), SearchCase);
-	}
-
-	/**
-	 * Test whether this string starts with given prefix.
-	 *
-	 * @param SearchCase		Indicates whether the search is case sensitive or not ( defaults to ESearchCase::IgnoreCase )
-	 * @return true if this string begins with specified text, false otherwise
-	 */
-	UE_NODISCARD bool StartsWith(const ElementType* InPrefix, ESearchCase::Type SearchCase = ESearchCase::IgnoreCase) const
-	{
-		return StartsWith(InPrefix, InPrefix ? FCString::Strlen(InPrefix) : 0, SearchCase);
-	}
-
-	/**
-	 * Test whether this string starts with given prefix.
-	 *
-	 * @param SearchCase		Indicates whether the search is case sensitive or not ( defaults to ESearchCase::IgnoreCase )
-	 * @return true if this string begins with specified text, false otherwise
-	 */
-	UE_NODISCARD bool StartsWith(const FString& InPrefix, ESearchCase::Type SearchCase = ESearchCase::IgnoreCase) const
-	{
-		return StartsWith(*InPrefix, InPrefix.Len(), SearchCase);
-	}
-
-	/**
-	 * Test whether this string starts with given prefix.
-	 *
-	 * @param SearchCase		Indicates whether the search is case sensitive or not ( defaults to ESearchCase::IgnoreCase )
-	 * @return true if this string begins with specified text, false otherwise
-	 */
-	UE_NODISCARD CORE_API bool StartsWith(const ElementType* InPrefix, int32 InPrefixLen, ESearchCase::Type SearchCase = ESearchCase::IgnoreCase) const;
-
-	/**
-	 * Test whether this string ends with given suffix.
-	 *
-	 * @param SearchCase		Indicates whether the search is case sensitive or not ( defaults to ESearchCase::IgnoreCase )
-	 * @return true if this string ends with specified text, false otherwise
-	 */
-	template <typename TCharRangeType, std::enable_if_t<TIsCharRangeNotCArrayNotStringClass<TCharRangeType>::Value>* = nullptr>
-	UE_NODISCARD bool EndsWith(TCharRangeType&& InSuffix, ESearchCase::Type SearchCase = ESearchCase::IgnoreCase) const
-	{
-		static_assert(std::is_same_v<typename TElementType<TCharRangeType>::Type, ElementType>, "Expected a range of characters");
-		return EndsWith(GetData(InSuffix), GetNum(InSuffix), SearchCase);
-	}
-
-	/**
-	 * Test whether this string ends with given suffix.
-	 *
-	 * @param SearchCase		Indicates whether the search is case sensitive or not ( defaults to ESearchCase::IgnoreCase )
-	 * @return true if this string ends with specified text, false otherwise
-	 */
-	UE_NODISCARD bool EndsWith(const ElementType* InSuffix, ESearchCase::Type SearchCase = ESearchCase::IgnoreCase) const
-	{
-		return EndsWith(InSuffix, InSuffix ? FCString::Strlen(InSuffix) : 0, SearchCase);
-	}
-
-	/**
-	 * Test whether this string ends with given suffix.
-	 *
-	 * @param SearchCase		Indicates whether the search is case sensitive or not ( defaults to ESearchCase::IgnoreCase )
-	 * @return true if this string ends with specified text, false otherwise
-	 */
-	UE_NODISCARD bool EndsWith(const FString& InSuffix, ESearchCase::Type SearchCase = ESearchCase::IgnoreCase) const
-	{
-		return EndsWith(*InSuffix, InSuffix.Len(), SearchCase);
-	}
-
-	/**
-	 * Test whether this string ends with given suffix.
-	 *
-	 * @param SearchCase		Indicates whether the search is case sensitive or not ( defaults to ESearchCase::IgnoreCase )
-	 * @return true if this string ends with specified text, false otherwise
-	 */
-	UE_NODISCARD CORE_API bool EndsWith(const ElementType* InSuffix, int32 InSuffixLen, ESearchCase::Type SearchCase = ESearchCase::IgnoreCase ) const;
-
-	/**
-	 * Searches this string for a given wild card
-	 *
-	 * @param Wildcard		*?-type wildcard
-	 * @param SearchCase	Indicates whether the search is case sensitive or not ( defaults to ESearchCase::IgnoreCase )
-	 * @return true if this string matches the *?-type wildcard given. 
-	 * @warning This is a simple, SLOW routine. Use with caution
-	 */
-	template <typename TCharRangeType, std::enable_if_t<TIsCharRangeNotCArrayNotStringClass<TCharRangeType>::Value>* = nullptr>
-	UE_NODISCARD bool MatchesWildcard(TCharRangeType&& Wildcard, ESearchCase::Type SearchCase = ESearchCase::IgnoreCase) const
-	{
-		static_assert(std::is_same_v<typename TElementType<TCharRangeType>::Type, ElementType>, "Expected a range of characters");
-		return MatchesWildcard(GetData(Wildcard), GetNum(Wildcard), SearchCase);
-	}
-
-	/**
-	 * Searches this string for a given wild card
-	 *
-	 * @param Wildcard		*?-type wildcard
-	 * @param SearchCase	Indicates whether the search is case sensitive or not ( defaults to ESearchCase::IgnoreCase )
-	 * @return true if this string matches the *?-type wildcard given.
-	 * @warning This is a simple, SLOW routine. Use with caution
-	 */
-	UE_NODISCARD bool MatchesWildcard(const ElementType* Wildcard, ESearchCase::Type SearchCase = ESearchCase::IgnoreCase) const
-	{
-		return MatchesWildcard(Wildcard, Wildcard ? FCString::Strlen(Wildcard) : 0, SearchCase);
-	}
-
-	/**
-	 * Searches this string for a given wild card
-	 *
-	 * @param Wildcard		*?-type wildcard
-	 * @param SearchCase	Indicates whether the search is case sensitive or not ( defaults to ESearchCase::IgnoreCase )
-	 * @return true if this string matches the *?-type wildcard given.
-	 * @warning This is a simple, SLOW routine. Use with caution
-	 */
-	UE_NODISCARD bool MatchesWildcard(const FString& Wildcard, ESearchCase::Type SearchCase = ESearchCase::IgnoreCase) const
-	{
-		return MatchesWildcard(*Wildcard, Wildcard.Len(), SearchCase);
-	}
-
-	/**
-	 * Searches this string for a given wild card
-	 *
-	 * @param Wildcard		*?-type wildcard
-	 * @param SearchCase	Indicates whether the search is case sensitive or not ( defaults to ESearchCase::IgnoreCase )
-	 * @return true if this string matches the *?-type wildcard given.
-	 * @warning This is a simple, SLOW routine. Use with caution
-	 */
-	UE_NODISCARD CORE_API bool MatchesWildcard(const ElementType* Wildcard, int32 WildcardLen, ESearchCase::Type SearchCase = ESearchCase::IgnoreCase) const;
-
-	/**
-	 * Removes whitespace characters from the start and end of this string. Modifies the string in-place.
-	 */
-	CORE_API void TrimStartAndEndInline();
-
-	/**
-	 * Removes whitespace characters from the start and end of this string.
-	 * @note Unlike Trim() this function returns a copy, and does not mutate the string.
-	 */
-	UE_NODISCARD CORE_API FString TrimStartAndEnd() const &;
-
-	/**
-	 * Removes whitespace characters from the start and end of this string.
-	 * @note Unlike Trim() this function returns a copy, and does not mutate the string.
-	 */
-	UE_NODISCARD CORE_API FString TrimStartAndEnd() &&;
-
-	/**
-	 * Removes whitespace characters from the start of this string. Modifies the string in-place.
-	 */
-	CORE_API void TrimStartInline();
-
-	/**
-	 * Removes whitespace characters from the start of this string.
-	 * @note Unlike Trim() this function returns a copy, and does not mutate the string.
-	 */
-	UE_NODISCARD CORE_API FString TrimStart() const &;
-
-	/**
-	 * Removes whitespace characters from the start of this string.
-	 * @note Unlike Trim() this function returns a copy, and does not mutate the string.
-	 */
-	UE_NODISCARD CORE_API FString TrimStart() &&;
-
-	/**
-	 * Removes whitespace characters from the end of this string. Modifies the string in-place.
-	 */
-	CORE_API void TrimEndInline();
-
-	/**
-	 * Removes whitespace characters from the end of this string.
-	 * @note Unlike TrimTrailing() this function returns a copy, and does not mutate the string.
-	 */
-	UE_NODISCARD CORE_API FString TrimEnd() const &;
-
-	/**
-	 * Removes whitespace characters from the end of this string.
-	 * @note Unlike TrimTrailing() this function returns a copy, and does not mutate the string.
-	 */
-	UE_NODISCARD CORE_API FString TrimEnd() &&;
-
-	/** 
-	 * Trims the inner array after the null terminator.
-	 */
-	CORE_API void TrimToNullTerminator();
-
-
-	/**
-	 * Trims wrapping quotation marks from this string.
-	 */
-	CORE_API void TrimQuotesInline(bool* bQuotesRemoved = nullptr);
-
-	/**
-	* Trims a single character from the start and end of the string (removes at max one instance in the beginning and end of the string).
-	* @see TrimChar for a variant that returns a modified copy of the string
-	*/
-	CORE_API void TrimCharInline(ElementType CharacterToTrim, bool* bCharRemoved);
-	
-	/**
-	 * Returns a copy of this string with wrapping quotation marks removed.
-	 */
-	UE_NODISCARD CORE_API FString TrimQuotes(bool* bQuotesRemoved = nullptr ) const &;
-
-	/**
-	 * Returns this string with wrapping quotation marks removed.
-	 */
-	UE_NODISCARD CORE_API FString TrimQuotes(bool* bQuotesRemoved = nullptr) &&;
-	
-	/**
-	* Returns a copy of this string with wrapping CharacterToTrim removed (removes at max one instance in the beginning and end of the string).
-	* @see TrimCharInline for an inline variant
-	*/
-	UE_NODISCARD CORE_API FString TrimChar(ElementType CharacterToTrim, bool* bCharRemoved = nullptr ) const &;
-
-	/**
-	* Returns a copy of this string with wrapping CharacterToTrim removed (removes at max one instance in the beginning and end of the string).
-	*/
-	UE_NODISCARD CORE_API FString TrimChar(ElementType CharacterToTrim, bool* bCharRemoved = nullptr) &&;
-	
-	/**
-	 * Breaks up a delimited string into elements of a string array.
-	 *
-	 * @param	InArray		The array to fill with the string pieces
-	 * @param	pchDelim	The string to delimit on
-	 * @param	InCullEmpty	If 1, empty strings are not added to the array
-	 *
-	 * @return	The number of elements in InArray
-	 */
-	CORE_API int32 ParseIntoArray( TArray<FString>& OutArray, const ElementType* pchDelim, bool InCullEmpty = true ) const;
-
-	/**
-	 * Breaks up a delimited string into elements of a string array, using any whitespace and an 
-	 * optional extra delimter, like a ","
-	 * @warning Caution!! this routine is O(N^2) allocations...use it for parsing very short text or not at all! 
-	 *
-	 * @param	InArray			The array to fill with the string pieces
-	 * @param	pchExtraDelim	The string to delimit on
-	 *
-	 * @return	The number of elements in InArray
-	 */
-	CORE_API int32 ParseIntoArrayWS( TArray<FString>& OutArray, const ElementType* pchExtraDelim = nullptr, bool InCullEmpty = true ) const;
-
-	/**
-	* Breaks up a delimited string into elements of a string array, using line ending characters
-	* @warning Caution!! this routine is O(N^2) allocations...use it for parsing very short text or not at all!
-	*
-	* @param	InArray			The array to fill with the string pieces	
-	*
-	* @return	The number of elements in InArray
-	*/
-	CORE_API int32 ParseIntoArrayLines(TArray<FString>& OutArray, bool InCullEmpty = true) const;
-
-	/**
-	* Breaks up a delimited string into elements of a string array, using the given delimiters
-	* @warning Caution!! this routine is O(N^2) allocations...use it for parsing very short text or not at all!
-	*
-	* @param	InArray			The array to fill with the string pieces
-	* @param	DelimArray		The strings to delimit on
-	* @param	NumDelims		The number of delimiters.
-	*
-	* @return	The number of elements in InArray
-	*/
-	CORE_API int32 ParseIntoArray(TArray<FString>& OutArray, const ElementType*const* DelimArray, int32 NumDelims, bool InCullEmpty = true) const;
-
-	/**
-	 * Takes an array of strings and removes any zero length entries.
-	 *
-	 * @param	InArray	The array to cull
-	 *
-	 * @return	The number of elements left in InArray
-	 */
-	static CORE_API int32 CullArray( TArray<FString>* InArray );
-
-	/**
-	 * Returns a copy of this string, with the characters in reverse order
-	 */
-	UE_NODISCARD CORE_API FString Reverse() const &;
-
-	/**
-	 * Returns this string, with the characters in reverse order
-	 */
-	UE_NODISCARD CORE_API FString Reverse() &&;
-
-	/**
-	 * Reverses the order of characters in this string
-	 */
-	CORE_API void ReverseString();
-
-	/**
-	 * Replace all occurrences of a substring in this string
-	 *
-	 * @param From substring to replace
-	 * @param To substring to replace From with
-	 * @param SearchCase	Indicates whether the search is case sensitive or not ( defaults to ESearchCase::IgnoreCase )
-	 * @return a copy of this string with the replacement made
-	 */
-	UE_NODISCARD CORE_API FString Replace(const ElementType* From, const ElementType* To, ESearchCase::Type SearchCase = ESearchCase::IgnoreCase) const &;
-
-	/**
-	 * Replace all occurrences of a substring in this string
-	 *
-	 * @param From substring to replace
-	 * @param To substring to replace From with
-	 * @param SearchCase	Indicates whether the search is case sensitive or not ( defaults to ESearchCase::IgnoreCase )
-	 * @return a copy of this string with the replacement made
-	 */
-	UE_NODISCARD CORE_API FString Replace(const ElementType* From, const ElementType* To, ESearchCase::Type SearchCase = ESearchCase::IgnoreCase) &&;
-
-	/**
-	 * Replace all occurrences of SearchText with ReplacementText in this string.
-	 *
-	 * @param	SearchText	the text that should be removed from this string
-	 * @param	ReplacementText		the text to insert in its place
-	 * @param SearchCase	Indicates whether the search is case sensitive or not ( defaults to ESearchCase::IgnoreCase )
-	 *
-	 * @return	the number of occurrences of SearchText that were replaced.
-	 */
-	CORE_API int32 ReplaceInline( const ElementType* SearchText, const ElementType* ReplacementText, ESearchCase::Type SearchCase = ESearchCase::IgnoreCase );
-
-	/**
-	 * Replace all occurrences of a character with another.
-	 *
-	 * @param SearchChar      Character to remove from this string
-	 * @param ReplacementChar Replacement character
-	 * @param SearchCase      Indicates whether the search is case sensitive or not ( defaults to ESearchCase::IgnoreCase )
-	 * @note no dynamic allocation
-	 */
-	void ReplaceCharInline(ElementType SearchChar, ElementType ReplacementChar, ESearchCase::Type SearchCase = ESearchCase::IgnoreCase)
-	{
-		if (SearchCase == ESearchCase::IgnoreCase && TChar<ElementType>::IsAlpha(SearchChar))
-		{
-			ReplaceCharInlineIgnoreCase(SearchChar, ReplacementChar);
-		}
-		else
-		{
-			ReplaceCharInlineCaseSensitive(SearchChar, ReplacementChar);
-		}
-	}
-
-private:
-	
-	CORE_API void ReplaceCharInlineCaseSensitive(ElementType SearchChar, ElementType ReplacementChar);
-	CORE_API void ReplaceCharInlineIgnoreCase(ElementType SearchChar, ElementType ReplacementChar);
-
-public:
-
-	/**
-	 * Returns a copy of this string with all quote marks escaped (unless the quote is already escaped)
-	 */
-	UE_NODISCARD FString ReplaceQuotesWithEscapedQuotes() const &
-	{
-		FString Result(*this);
-		return MoveTemp(Result).ReplaceQuotesWithEscapedQuotes();
-	}
-
-	/**
-	 * Returns a copy of this string with all quote marks escaped (unless the quote is already escaped)
-	 */
-	UE_NODISCARD CORE_API FString ReplaceQuotesWithEscapedQuotes() &&;
-
-	/**
-	 * Replaces certain characters with the "escaped" version of that character (i.e. replaces "\n" with "\\n").
-	 * The characters supported are: { \n, \r, \t, \', \", \\ }.
-	 *
-	 * @param	Chars	by default, replaces all supported characters; this parameter allows you to limit the replacement to a subset.
-	 */
-	CORE_API void ReplaceCharWithEscapedCharInline( const TArray<ElementType>* Chars = nullptr );
-
-	/**
-	 * Replaces certain characters with the "escaped" version of that character (i.e. replaces "\n" with "\\n").
-	 * The characters supported are: { \n, \r, \t, \', \", \\ }.
-	 *
-	 * @param	Chars	by default, replaces all supported characters; this parameter allows you to limit the replacement to a subset.
-	 *
-	 * @return	a string with all control characters replaced by the escaped version.
-	 */
-	UE_NODISCARD FString ReplaceCharWithEscapedChar( const TArray<ElementType>* Chars = nullptr ) const &
-	{
-		FString Result(*this);
-		Result.ReplaceCharWithEscapedCharInline(Chars);
-		return Result;
-	}
-
-	/**
-	 * Replaces certain characters with the "escaped" version of that character (i.e. replaces "\n" with "\\n").
-	 * The characters supported are: { \n, \r, \t, \', \", \\ }.
-	 *
-	 * @param	Chars	by default, replaces all supported characters; this parameter allows you to limit the replacement to a subset.
-	 *
-	 * @return	a string with all control characters replaced by the escaped version.
-	 */
-	UE_NODISCARD FString ReplaceCharWithEscapedChar( const TArray<ElementType>* Chars = nullptr ) &&
-	{
-		ReplaceCharWithEscapedCharInline(Chars);
-		return MoveTemp(*this);
-	}
-
-	/**
-	 * Removes the escape backslash for all supported characters, replacing the escape and character with the non-escaped version.  (i.e.
-	 * replaces "\\n" with "\n".  Counterpart to ReplaceCharWithEscapedCharInline().
-	 */
-	CORE_API void ReplaceEscapedCharWithCharInline( const TArray<ElementType>* Chars = nullptr );
-
-	/**
-	 * Removes the escape backslash for all supported characters, replacing the escape and character with the non-escaped version.  (i.e.
-	 * replaces "\\n" with "\n".  Counterpart to ReplaceCharWithEscapedChar().
-	 * @return copy of this string with replacement made
-	 */
-	UE_NODISCARD FString ReplaceEscapedCharWithChar( const TArray<ElementType>* Chars = nullptr ) const &
-	{
-		FString Result(*this);
-		Result.ReplaceEscapedCharWithCharInline(Chars);
-		return Result;
-	}
-
-	/**
-	 * Removes the escape backslash for all supported characters, replacing the escape and character with the non-escaped version.  (i.e.
-	 * replaces "\\n" with "\n".  Counterpart to ReplaceCharWithEscapedChar().
-	 * @return copy of this string with replacement made
-	 */
-	UE_NODISCARD FString ReplaceEscapedCharWithChar( const TArray<ElementType>* Chars = nullptr ) &&
-	{
-		ReplaceEscapedCharWithCharInline(Chars);
-		return MoveTemp(*this);
-	}
-
-	/**
-	 * Replaces all instances of '\t' with TabWidth number of spaces
-	 * @param InSpacesPerTab - Number of spaces that a tab represents
-	 */
-	CORE_API void ConvertTabsToSpacesInline(const int32 InSpacesPerTab);
-
-	/** 
-	 * Replaces all instances of '\t' with TabWidth number of spaces
-	 * @param InSpacesPerTab - Number of spaces that a tab represents
-	 * @return copy of this string with replacement made
-	 */
-	UE_NODISCARD FString ConvertTabsToSpaces(const int32 InSpacesPerTab) const &
-	{
-		FString FinalString(*this);
-		FinalString.ConvertTabsToSpacesInline(InSpacesPerTab);
-		return FinalString;
-	}
-
-	/**
-	 * Replaces all instances of '\t' with TabWidth number of spaces
-	 * @param InSpacesPerTab - Number of spaces that a tab represents
-	 * @return copy of this string with replacement made
-	 */
-	UE_NODISCARD FString ConvertTabsToSpaces(const int32 InSpacesPerTab) &&
-	{
-		ConvertTabsToSpacesInline(InSpacesPerTab);
-		return MoveTemp(*this);
-	}
-
-	// Takes the number passed in and formats the string in comma format ( 12345 becomes "12,345")
-	UE_NODISCARD static CORE_API FString FormatAsNumber( int32 InNumber );
-
-	// To allow more efficient memory handling, automatically adds one for the string termination.
-	CORE_API void Reserve(int32 CharacterCount);
-
-	/**
-	 * Serializes a string as ANSI char array.
-	 *
-	 * @param	String			String to serialize
-	 * @param	Ar				Archive to serialize with
-	 * @param	MinCharacters	Minimum number of characters to serialize.
-	 */
-	CORE_API void SerializeAsANSICharArray( FArchive& Ar, int32 MinCharacters=0 ) const;
-
-
-	/** Converts an integer to a string. */
-	UE_NODISCARD static FORCEINLINE FString FromInt( int32 Num )
-	{
-		FString Ret;
-		Ret.AppendInt(Num); 
-		return Ret;
-	}
-
-	/** appends the integer InNum to this string */
-	CORE_API void AppendInt( int32 InNum );
-
-	/**
-	 * Converts a string into a boolean value
-	 *   1, "True", "Yes", FCoreTexts::True, FCoreTexts::Yes, and non-zero integers become true
-	 *   0, "False", "No", FCoreTexts::False, FCoreTexts::No, and unparsable values become false
-	 *
-	 * @return The boolean value
-	 */
-	UE_NODISCARD CORE_API bool ToBool() const;
-
-	/**
-	 * Converts a buffer to a string
-	 *
-	 * @param SrcBuffer the buffer to stringify
-	 * @param SrcSize the number of bytes to convert
-	 *
-	 * @return the blob in string form
-	 */
-	UE_NODISCARD static CORE_API FString FromBlob(const uint8* SrcBuffer,const uint32 SrcSize);
-
-	/**
-	 * Converts a string into a buffer
-	 *
-	 * @param DestBuffer the buffer to fill with the string data
-	 * @param DestSize the size of the buffer in bytes (must be at least string len / 3)
-	 *
-	 * @return true if the conversion happened, false otherwise
-	 */
-	static CORE_API bool ToBlob(const FString& Source,uint8* DestBuffer,const uint32 DestSize);
-
-	/**
-	 * Converts a buffer to a string by hex-ifying the elements
-	 *
-	 * @param SrcBuffer the buffer to stringify
-	 * @param SrcSize the number of bytes to convert
-	 *
-	 * @return the blob in string form
-	 */
-	UE_NODISCARD static CORE_API FString FromHexBlob(const uint8* SrcBuffer,const uint32 SrcSize);
-
-	/**
-	 * Converts a string into a buffer
-	 *
-	 * @param DestBuffer the buffer to fill with the string data
-	 * @param DestSize the size of the buffer in bytes (must be at least string len / 2)
-	 *
-	 * @return true if the conversion happened, false otherwise
-	 */
-	static CORE_API bool ToHexBlob(const FString& Source,uint8* DestBuffer,const uint32 DestSize);
-
-	/**
-	 * Converts a float string with the trailing zeros stripped
-	 * For example - 1.234 will be "1.234" rather than "1.234000"
-	 * 
-	 * @param	InFloat					The float to sanitize
-	 * @param	InMinFractionalDigits	The minimum number of fractional digits the number should have (will be padded with zero)
-	 *
-	 * @return sanitized string version of float
-	 */
-	UE_NODISCARD static CORE_API FString SanitizeFloat( double InFloat, const int32 InMinFractionalDigits = 1 );
-
-	/**
-	 * Joins a range of 'something that can be concatentated to strings with +=' together into a single string with separators.
-	 *
-	 * @param	Range		The range of 'things' to concatenate.
-	 * @param	Separator	The string used to separate each element.
-	 *
-	 * @return	The final, joined, separated string.
-	 */
-	template <typename RangeType>
-	UE_NODISCARD static FString Join(const RangeType& Range, const ElementType* Separator)
-	{
-		FString Result;
-		bool    First = true;
-		for (const auto& Element : Range)
-		{
-			if (First)
-			{
-				First = false;
-			}
-			else
-			{
-				Result += Separator;
-			}
-
-			Result += Element;
-		}
-
-		return Result;
-	}
-
-	/**
-	 * Joins a range of elements together into a single string with separators using a projection function.
-	 *
-	 * @param	Range		The range of 'things' to concatenate.
-	 * @param	Separator	The string used to separate each element.
-	 * @param	Proj		The projection used to get a string for each element.
-	 *
-	 * @return	The final, joined, separated string.
-	 */
-	template <typename RangeType, typename ProjectionType>
-	UE_NODISCARD static FString JoinBy(const RangeType& Range, const ElementType* Separator, ProjectionType Proj)
-	{
-		FString Result;
-		bool    First = true;
-		for (const auto& Element : Range)
-		{
-			if (First)
-			{
-				First = false;
-			}
-			else
-			{
-				Result += Separator;
-			}
-
-			Result += Invoke(Proj, Element);
-		}
-
-		return Result;
-	}
-
-	FORCEINLINE void CountBytes(FArchive& Ar) const
-	{
-		Data.CountBytes(Ar);
-	}
-
-	/** Case insensitive string hash function. */
-	friend FORCEINLINE uint32 GetTypeHash(const FString& S)
-	{
-		// This must match the GetTypeHash behavior of FStringView
-		return FCrc::Strihash_DEPRECATED(S.Len(), *S);
-	}
-};
-
-template<> struct TIsZeroConstructType<FString> { enum { Value = true }; };
-Expose_TNameOf(FString)
-
-inline FString::ElementType* GetData(FString& String)
-{
-	return String.GetCharArray().GetData();
-}
-
-inline const FString::ElementType* GetData(const FString& String)
-{
-	return String.GetCharArray().GetData();
-}
-
-inline int32 GetNum(const FString& String)
-{
-	return String.Len();
-}
-
-/** 
- * Convert an array of bytes to a string
- * @param In byte array values to convert
- * @param Count number of bytes to convert
- * @return Valid string representing bytes.
- */
-UE_NODISCARD inline FString BytesToString(const uint8* In, int32 Count)
-{
-	FString Result;
-	Result.Empty(Count);
-
-	while (Count)
-	{
-		// Put the byte into an int16 and add 1 to it, this keeps anything from being put into the string as a null terminator
-		int16 Value = *In;
-		Value += 1;
-
-		Result += FString::ElementType(Value);
-
-		++In;
-		Count--;
-	}
-	return Result;
-}
-
-/** 
- * Convert FString of bytes into the byte array.
- * @param String		The FString of byte values
- * @param OutBytes		Ptr to memory must be preallocated large enough
- * @param MaxBufferSize	Max buffer size of the OutBytes array, to prevent overflow
- * @return	The number of bytes copied
- */
-inline int32 StringToBytes( const FString& String, uint8* OutBytes, int32 MaxBufferSize )
-{
-	int32 NumBytes = 0;
-	const FString::ElementType* CharPos = *String;
-
-	while( *CharPos && NumBytes < MaxBufferSize)
-	{
-		OutBytes[ NumBytes ] = (int8)(*CharPos - 1);
-		CharPos++;
-		++NumBytes;
-	}
-	return NumBytes;
-}
-
-/** Returns Char value of Nibble */
-UE_NODISCARD inline TCHAR NibbleToTChar(uint8 Num)
-{
-	if (Num > 9)
-	{
-		return (TCHAR)('A' + (Num - 10));
-	}
-	return (TCHAR)('0' + Num);
-}
-
-/** 
- * Convert a byte to hex
- * @param In byte value to convert
- * @param Result out hex value output
- */
-inline void ByteToHex(uint8 In, FString& Result)
-{
-	Result += NibbleToTChar(In >> 4);
-	Result += NibbleToTChar(In & 15);
-}
-
-/** Append bytes as uppercase hex string */
-CORE_API void BytesToHex(const uint8* Bytes, int32 NumBytes, FString& Out);
-/** Append bytes as lowercase hex string */
-CORE_API void BytesToHexLower(const uint8* Bytes, int32 NumBytes, FString& Out);
-
-/** Convert bytes to uppercase hex string */
-UE_NODISCARD inline FString BytesToHex(const uint8* Bytes, int32 NumBytes)
-{
-	FString Out;
-	BytesToHex(Bytes, NumBytes, Out);
-	return Out;
-}
-
-/** Convert bytes to lowercase hex string */
-UE_NODISCARD inline FString BytesToHexLower(const uint8* Bytes, int32 NumBytes)
-{
-	FString Out;
-	BytesToHexLower(Bytes, NumBytes, Out);
-	return Out;
-}
-
-/**
- * Checks if the TChar is a valid hex character
- * @param Char		The character
- * @return	True if in 0-9 and A-F ranges
- */
-UE_NODISCARD inline const bool CheckTCharIsHex( const TCHAR Char )
-{
-	return ( Char >= TEXT('0') && Char <= TEXT('9') ) || ( Char >= TEXT('A') && Char <= TEXT('F') ) || ( Char >= TEXT('a') && Char <= TEXT('f') );
-}
-
-/**
- * Convert a TChar to equivalent hex value as a uint8
- * @param Hex		The character
- * @return	The uint8 value of a hex character
- */
-UE_NODISCARD inline const uint8 TCharToNibble(const TCHAR Hex)
-{
-	if (Hex >= '0' && Hex <= '9')
-	{
-		return uint8(Hex - '0');
-	}
-	if (Hex >= 'A' && Hex <= 'F')
-	{
-		return uint8(Hex - 'A' + 10);
-	}
-	if (Hex >= 'a' && Hex <= 'f')
-	{
-		return uint8(Hex - 'a' + 10);
-	}
-	checkf(false, TEXT("'%c' (0x%02X) is not a valid hexadecimal digit"), Hex, Hex);
-	return 0;
-}
-
-/** 
- * Convert FString of Hex digits into the byte array.
- * @param HexString		The FString of Hex values
- * @param OutBytes		Ptr to memory must be preallocated large enough
- * @return	The number of bytes copied
- */
-CORE_API int32 HexToBytes(const FString& HexString, uint8* OutBytes);
-
-/**
- * Generalized API to convert something to a string. Function named after the (deprecated) Lex namespace, which
- * was deprecated because introducing customization points in a nested namespace didn't work in generic code because
- * it foiled 2-phase template instantiating compilers, which would bind to the qualified name (LexToString) in the first phase,
- * preventing future overloads defined in later headers to be considered for binding.
- */
- /**
- *	Expected functions in this namespace are as follows:
- *		bool								LexTryParseString(T& OutValue, const TCHAR* Buffer);
- *		void 								LexFromString(T& OutValue, const TCHAR* Buffer);
- *		<implicitly convertible to string>	LexToString(T);
- *		                    ^-- Generally this means it can return either FString or const TCHAR* 
- *		                        Generic code that uses ToString should assign to an FString or forward along to other functions
- *		                        that accept types that are also implicitly convertible to FString 
- *
- *	Implement custom functionality externally.
- */
-
- /** Covert a string buffer to intrinsic types */
-inline void LexFromString(int8& OutValue, 		const TCHAR* Buffer)	{	OutValue = (int8)FCString::Atoi(Buffer);		}
-inline void LexFromString(int16& OutValue,		const TCHAR* Buffer)	{	OutValue = (int16)FCString::Atoi(Buffer);		}
-inline void LexFromString(int32& OutValue,		const TCHAR* Buffer)	{	OutValue = (int32)FCString::Atoi(Buffer);		}
-inline void LexFromString(int64& OutValue,		const TCHAR* Buffer)	{	OutValue = FCString::Atoi64(Buffer);	}
-inline void LexFromString(uint8& OutValue,		const TCHAR* Buffer)	{	OutValue = (uint8)FCString::Atoi(Buffer);		}
-inline void LexFromString(uint16& OutValue, 	const TCHAR* Buffer)	{	OutValue = (uint16)FCString::Atoi(Buffer);		}
-inline void LexFromString(uint32& OutValue, 	const TCHAR* Buffer)	{	OutValue = (uint32)FCString::Atoi64(Buffer);	}	//64 because this unsigned and so Atoi might overflow
-inline void LexFromString(uint64& OutValue, 	const TCHAR* Buffer)	{	OutValue = FCString::Strtoui64(Buffer, nullptr, 0); }
-inline void LexFromString(float& OutValue,		const TCHAR* Buffer)	{	OutValue = FCString::Atof(Buffer);		}
-inline void LexFromString(double& OutValue, 	const TCHAR* Buffer)	{	OutValue = FCString::Atod(Buffer);		}
-inline void LexFromString(bool& OutValue, 		const TCHAR* Buffer)	{	OutValue = FCString::ToBool(Buffer);	}
-inline void LexFromString(FString& OutValue, 	const TCHAR* Buffer)	{	OutValue = Buffer;						}
-
- /** Convert numeric types to a string */
-template<typename T>
-UE_NODISCARD typename TEnableIf<TIsArithmetic<T>::Value, FString>::Type
-LexToString(const T& Value)
-{
-	// std::remove_cv_t to remove potential volatile decorations. Removing const is pointless, but harmless because it's specified in the param declaration.
-	return FString::Printf( TFormatSpecifier<std::remove_cv_t<T>>::GetFormatSpecifier(), Value );
-}
-
-template<typename CharType>
-UE_NODISCARD typename TEnableIf<TIsCharType<CharType>::Value, FString>::Type
-LexToString(const CharType* Ptr)
-{
-	return FString(Ptr);
-}
-
-UE_NODISCARD inline FString LexToString(bool Value)
-{
-	return Value ? TEXT("true") : TEXT("false");
-}
-
-UE_NODISCARD FORCEINLINE FString LexToString(FString&& Str)
-{
-	return MoveTemp(Str);
-}
-
-UE_NODISCARD FORCEINLINE FString LexToString(const FString& Str)
-{
-	return Str;
-}
-
-/** Helper template to convert to sanitized strings */
-template<typename T>
-UE_NODISCARD FString LexToSanitizedString(const T& Value)
-{
-	return LexToString(Value);
-}
-
-/** Overloaded for floats */
-UE_NODISCARD inline FString LexToSanitizedString(float Value)
-{
-	return FString::SanitizeFloat(Value);
-}
-
-/** Overloaded for doubles */
-UE_NODISCARD inline FString LexToSanitizedString(double Value)
-{
-	return FString::SanitizeFloat(Value);
-}
-
-/** Parse a string into this type, returning whether it was successful */
-/** Specialization for arithmetic types */
-template<typename T>
-static typename TEnableIf<TIsArithmetic<T>::Value, bool>::Type
-LexTryParseString(T& OutValue, const TCHAR* Buffer)
-{
-	if (Buffer[0] == '\0')
-	{
-		OutValue = 0;
-		return false;
-	}
-
-	LexFromString(OutValue, Buffer);
-	if (OutValue == 0 && FMath::IsFinite((float)OutValue)) //@TODO:FLOATPRECISION: ? huh ?
-	{
-		bool bSawZero = false;
-		TCHAR C = *Buffer;
-		while (C != '\0' && (C == '+' || C == '-' || FChar::IsWhitespace(C)))
-		{
-			C = *(++Buffer);
-		}
-
-		while (C != '\0' && !FChar::IsWhitespace(C) && (TIsFloatingPoint<T>::Value || C != '.'))
-		{
-			bSawZero = bSawZero || (C == '0');
-			if (!bSawZero && C != '.')
-			{
-				return false;
-			}
-
-			C = *(++Buffer);
-		}
-		return bSawZero;
-	}
-	
-	return true;
-}
-
-/** Try and parse a bool - always returns true */
-static bool LexTryParseString(bool& OutValue, const TCHAR* Buffer)
-{
-	LexFromString(OutValue, Buffer);
-	return true;
-}
-
-
-/** Shorthand legacy use for Lex functions */
-template<typename T>
-struct TTypeToString
-{
-	UE_NODISCARD static FString ToString(const T& Value)				{ return LexToString(Value); }
-	UE_NODISCARD static FString ToSanitizedString(const T& Value)	{ return LexToSanitizedString(Value); }
-};
-template<typename T>
-struct TTypeFromString
-{
-	static void FromString(T& Value, const TCHAR* Buffer) { return LexFromString(Value, Buffer); }
-};
-
-/**
- * Gets a non-owning TCHAR pointer from a string type.
- *
- * Can be used generically to get a const TCHAR*, when it is not known if the argument is a TCHAR* or an FString:
- *
- * template <typename T>
- * void LogValue(const T& Val)
- * {
- *     Logf(TEXT("Value: %s"), ToCStr(LexToString(Val)));
- * }
- */
-FORCEINLINE const TCHAR* ToCStr(const FString& Str)
-{
-	return *Str;
-}
-
-namespace StringConv
-{
-	/** Inline combine any UTF-16 surrogate pairs in the given string */
-	CORE_API void InlineCombineSurrogates(FString& Str);
-}
-
-/*----------------------------------------------------------------------------
-	Special archivers.
-----------------------------------------------------------------------------*/
-
-//
-// String output device.
-//
-class FStringOutputDevice : public FString, public FOutputDevice
-{
-public:
-	FStringOutputDevice( const TCHAR* OutputDeviceName=TEXT("") ):
-		FString( OutputDeviceName )
-	{
-		bAutoEmitLineTerminator = false;
-	}
-	virtual void Serialize( const TCHAR* InData, ELogVerbosity::Type Verbosity, const class FName& Category ) override
-	{
-		FString::operator+=((TCHAR*)InData);
-		if(bAutoEmitLineTerminator)
-		{
-			*this += LINE_TERMINATOR;
-		}
-	}
-
-	FStringOutputDevice(FStringOutputDevice&&) = default;
-	FStringOutputDevice(const FStringOutputDevice&) = default;
-	FStringOutputDevice& operator=(FStringOutputDevice&&) = default;
-	FStringOutputDevice& operator=(const FStringOutputDevice&) = default;
-
-	// Make += operator virtual.
-	virtual FString& operator+=(const FString& Other)
-	{
-		return FString::operator+=(Other);
-	}
-};
-
-template <>
-struct TIsContiguousContainer<FStringOutputDevice>
-{
-	enum { Value = true };
-};
-
-//
-// String output device.
-//
-class FStringOutputDeviceCountLines : public FStringOutputDevice
-{
-	typedef FStringOutputDevice Super;
-
-	int32 LineCount;
-public:
-	FStringOutputDeviceCountLines( const TCHAR* OutputDeviceName=TEXT("") )
-	:	Super( OutputDeviceName )
-	,	LineCount(0)
-	{}
-
-	virtual void Serialize(const TCHAR* InData, ELogVerbosity::Type Verbosity, const class FName& Category) override
-	{
-		Super::Serialize(InData, Verbosity, Category);
-		int32 TermLength = FCString::Strlen(LINE_TERMINATOR);
-		for (;;)
-		{
-			InData = FCString::Strstr(InData, LINE_TERMINATOR);
-			if (!InData)
-			{
-				break;
-			}
-			LineCount++;
-			InData += TermLength;
-		}
-
-		if (bAutoEmitLineTerminator)
-		{
-			LineCount++;
-		}
-	}
-
-	/**
-	 * Appends other FStringOutputDeviceCountLines object to this one.
-	 */
-	virtual FStringOutputDeviceCountLines& operator+=(const FStringOutputDeviceCountLines& Other)
-	{
-		FString::operator+=(static_cast<const FString&>(Other));
-
-		LineCount += Other.GetLineCount();
-
-		return *this;
-	}
-
-	/**
-	 * Appends other FString (as well as its specializations like FStringOutputDevice)
-	 * object to this.
-	 */
-	virtual FString& operator+=(const FString& Other) override
-	{
-		Log(Other);
-
-		return *this;
-	}
-
-	int32 GetLineCount() const
-	{
-		return LineCount;
-	}
-
-	FStringOutputDeviceCountLines(const FStringOutputDeviceCountLines&) = default;
-	FStringOutputDeviceCountLines& operator=(const FStringOutputDeviceCountLines&) = default;
-
-	FORCEINLINE FStringOutputDeviceCountLines(FStringOutputDeviceCountLines&& Other)
-		: Super    ((Super&&)Other)
-		, LineCount(Other.LineCount)
-	{
-		Other.LineCount = 0;
-	}
-
-	FORCEINLINE FStringOutputDeviceCountLines& operator=(FStringOutputDeviceCountLines&& Other)
-	{
-		if (this != &Other)
-		{
-			(Super&)*this = (Super&&)Other;
-			LineCount     = Other.LineCount;
-
-			Other.LineCount = 0;
-		}
-		return *this;
-	}
-};
-
-template <>
-struct TIsContiguousContainer<FStringOutputDeviceCountLines>
-{
-	enum { Value = true };
-};
-
-/**
- * A helper function to find closing parenthesis that matches the first open parenthesis found. The open parenthesis
- * referred to must be at or further up from the start index.
- *
- * @param TargetString      The string to search in
- * @param StartSearch       The index to start searching at
- * @return the index in the given string of the closing parenthesis
- */
-CORE_API int32 FindMatchingClosingParenthesis(const FString& TargetString, const int32 StartSearch = 0);
-
-/**
-* Given a display label string, generates an FString slug that only contains valid characters for an FName.
-* For example, "[MyObject]: Object Label" becomes "MyObjectObjectLabel" FName slug.
-*
-* @param DisplayLabel The label string to convert to an FName
-*
-* @return	The slugged string
-*/
-CORE_API FString SlugStringForValidName(const FString& DisplayString, const TCHAR* ReplaceWith = TEXT(""));
-
-struct FTextRange
-{
-	FTextRange()
-		: BeginIndex(INDEX_NONE)
-		, EndIndex(INDEX_NONE)
-	{
-
-	}
-
-	FTextRange(int32 InBeginIndex, int32 InEndIndex)
-		: BeginIndex(InBeginIndex)
-		, EndIndex(InEndIndex)
-	{
-
-	}
-
-	FORCEINLINE bool operator==(const FTextRange& Other) const
-	{
-		return BeginIndex == Other.BeginIndex
-			&& EndIndex == Other.EndIndex;
-	}
-
-	FORCEINLINE bool operator!=(const FTextRange& Other) const
-	{
-		return !(*this == Other);
-	}
-
-	friend inline uint32 GetTypeHash(const FTextRange& Key)
-	{
-		uint32 KeyHash = 0;
-		KeyHash = HashCombine(KeyHash, GetTypeHash(Key.BeginIndex));
-		KeyHash = HashCombine(KeyHash, GetTypeHash(Key.EndIndex));
-		return KeyHash;
-	}
-
-	int32 Len() const { return EndIndex - BeginIndex; }
-	bool IsEmpty() const { return (EndIndex - BeginIndex) <= 0; }
-	void Offset(int32 Amount) { BeginIndex += Amount; BeginIndex = FMath::Max(0, BeginIndex);  EndIndex += Amount; EndIndex = FMath::Max(0, EndIndex); }
-	bool Contains(int32 Index) const { return Index >= BeginIndex && Index < EndIndex; }
-	bool InclusiveContains(int32 Index) const { return Index >= BeginIndex && Index <= EndIndex; }
-
-	FTextRange Intersect(const FTextRange& Other) const
-	{
-		FTextRange Intersected(FMath::Max(BeginIndex, Other.BeginIndex), FMath::Min(EndIndex, Other.EndIndex));
-		if (Intersected.EndIndex <= Intersected.BeginIndex)
-		{
-			return FTextRange(0, 0);
-		}
-
-		return Intersected;
-	}
-
-	/**
-	 * Produce an array of line ranges from the given text, breaking at any new-line characters
-	 */
-	static CORE_API void CalculateLineRangesFromString(const FString& Input, TArray<FTextRange>& LineRanges);
-
-	int32 BeginIndex;
-	int32 EndIndex;
-};
-
-#include "Misc/StringFormatArg.h"
+#include "Containers/UnrealString.h.inl"
