Moved code out of generic string .inl files that we don't want to be generic (yet?).

diff --git a/Engine/Source/Runtime/Core/Private/Containers/String.cpp b/Engine/Source/Runtime/Core/Private/Containers/String.cpp
--- a/Engine/Source/Runtime/Core/Private/Containers/String.cpp
+++ b/Engine/Source/Runtime/Core/Private/Containers/String.cpp
@@ -6,3 +6,41 @@
 #include "Containers/StringIncludes.cpp.inl"
 
 #include "Containers/String.cpp.inl"
+
+void FTextRange::CalculateLineRangesFromString(const FString& Input, TArray<FTextRange>& LineRanges)
+{
+	using ElementType = FString::ElementType;
+
+	int32 LineBeginIndex = 0;
+
+	// Loop through splitting at new-lines
+	const ElementType* const InputStart = *Input;
+	for (const ElementType* CurrentChar = InputStart; CurrentChar && *CurrentChar; ++CurrentChar)
+	{
+		// Handle a chain of \r\n slightly differently to stop the TChar<ElementType>::IsLinebreak adding two separate new-lines
+		const bool bIsWindowsNewLine = (*CurrentChar == '\r' && *(CurrentChar + 1) == '\n');
+		if (bIsWindowsNewLine || TChar<ElementType>::IsLinebreak(*CurrentChar))
+		{
+			const int32 LineEndIndex = UE_PTRDIFF_TO_INT32(CurrentChar - InputStart);
+			check(LineEndIndex >= LineBeginIndex);
+			LineRanges.Emplace(FTextRange(LineBeginIndex, LineEndIndex));
+
+			if (bIsWindowsNewLine)
+			{
+				++CurrentChar; // skip the \n of the \r\n chain
+			}
+			LineBeginIndex = UE_PTRDIFF_TO_INT32(CurrentChar - InputStart) + 1; // The next line begins after the end of the current line
+		}
+	}
+
+	// Process any remaining string after the last new-line
+	if (LineBeginIndex <= Input.Len())
+	{
+		LineRanges.Emplace(FTextRange(LineBeginIndex, Input.Len()));
+	}
+}
+
+void StringConv::InlineCombineSurrogates(FString& Str)
+{
+	InlineCombineSurrogates_Array(Str.GetCharArray());
+}
diff --git a/Engine/Source/Runtime/Core/Private/Containers/String.cpp.inl b/Engine/Source/Runtime/Core/Private/Containers/String.cpp.inl
--- a/Engine/Source/Runtime/Core/Private/Containers/String.cpp.inl
+++ b/Engine/Source/Runtime/Core/Private/Containers/String.cpp.inl
@@ -1910,41 +1910,3 @@ FString SlugStringForValidName(const FString& DisplayString, const FString::Elem
 
 	return GeneratedName;
 }
-
-void FTextRange::CalculateLineRangesFromString(const FString& Input, TArray<FTextRange>& LineRanges)
-{
-	using ElementType = FString::ElementType;
-
-	int32 LineBeginIndex = 0;
-
-	// Loop through splitting at new-lines
-	const ElementType* const InputStart = *Input;
-	for (const ElementType* CurrentChar = InputStart; CurrentChar && *CurrentChar; ++CurrentChar)
-	{
-		// Handle a chain of \r\n slightly differently to stop the TChar<ElementType>::IsLinebreak adding two separate new-lines
-		const bool bIsWindowsNewLine = (*CurrentChar == '\r' && *(CurrentChar + 1) == '\n');
-		if (bIsWindowsNewLine || TChar<ElementType>::IsLinebreak(*CurrentChar))
-		{
-			const int32 LineEndIndex = UE_PTRDIFF_TO_INT32(CurrentChar - InputStart);
-			check(LineEndIndex >= LineBeginIndex);
-			LineRanges.Emplace(FTextRange(LineBeginIndex, LineEndIndex));
-
-			if (bIsWindowsNewLine)
-			{
-				++CurrentChar; // skip the \n of the \r\n chain
-			}
-			LineBeginIndex = UE_PTRDIFF_TO_INT32(CurrentChar - InputStart) + 1; // The next line begins after the end of the current line
-		}
-	}
-
-	// Process any remaining string after the last new-line
-	if (LineBeginIndex <= Input.Len())
-	{
-		LineRanges.Emplace(FTextRange(LineBeginIndex, Input.Len()));
-	}
-}
-
-void StringConv::InlineCombineSurrogates(FString& Str)
-{
-	InlineCombineSurrogates_Array(Str.GetCharArray());
-}
diff --git a/Engine/Source/Runtime/Core/Public/Containers/UnrealString.h b/Engine/Source/Runtime/Core/Public/Containers/UnrealString.h
--- a/Engine/Source/Runtime/Core/Public/Containers/UnrealString.h
+++ b/Engine/Source/Runtime/Core/Public/Containers/UnrealString.h
@@ -9,3 +9,439 @@
 #include "Containers/UnrealStringIncludes.h.inl"
 
 #include "Containers/UnrealString.h.inl"
+
+/**
+ * Convert an array of bytes to a string
+ * @param In byte array values to convert
+ * @param Count number of bytes to convert
+ * @return Valid string representing bytes.
+ */
+UE_NODISCARD inline FString BytesToString(const uint8* In, int32 Count)
+{
+	FString Result;
+	Result.Empty(Count);
+
+	while (Count)
+	{
+		// Put the byte into an int16 and add 1 to it, this keeps anything from being put into the string as a null terminator
+		int16 Value = *In;
+		Value += 1;
+
+		Result += FString::ElementType(Value);
+
+		++In;
+		Count--;
+	}
+	return Result;
+}
+
+/** 
+ * Convert FString of bytes into the byte array.
+ * @param String		The FString of byte values
+ * @param OutBytes		Ptr to memory must be preallocated large enough
+ * @param MaxBufferSize	Max buffer size of the OutBytes array, to prevent overflow
+ * @return	The number of bytes copied
+ */
+inline int32 StringToBytes( const FString& String, uint8* OutBytes, int32 MaxBufferSize )
+{
+	int32 NumBytes = 0;
+	const FString::ElementType* CharPos = *String;
+
+	while( *CharPos && NumBytes < MaxBufferSize)
+	{
+		OutBytes[ NumBytes ] = (int8)(*CharPos - 1);
+		CharPos++;
+		++NumBytes;
+	}
+	return NumBytes;
+}
+
+/** Returns Char value of Nibble */
+UE_NODISCARD inline TCHAR NibbleToTChar(uint8 Num)
+{
+	if (Num > 9)
+	{
+		return (TCHAR)('A' + (Num - 10));
+	}
+	return (TCHAR)('0' + Num);
+}
+
+/** 
+ * Convert a byte to hex
+ * @param In byte value to convert
+ * @param Result out hex value output
+ */
+inline void ByteToHex(uint8 In, FString& Result)
+{
+	Result += NibbleToTChar(In >> 4);
+	Result += NibbleToTChar(In & 15);
+}
+
+/** Convert bytes to uppercase hex string */
+UE_NODISCARD inline FString BytesToHex(const uint8* Bytes, int32 NumBytes)
+{
+	FString Out;
+	BytesToHex(Bytes, NumBytes, Out);
+	return Out;
+}
+
+/** Convert bytes to lowercase hex string */
+UE_NODISCARD inline FString BytesToHexLower(const uint8* Bytes, int32 NumBytes)
+{
+	FString Out;
+	BytesToHexLower(Bytes, NumBytes, Out);
+	return Out;
+}
+
+/**
+ * Checks if the TChar is a valid hex character
+ * @param Char		The character
+ * @return	True if in 0-9 and A-F ranges
+ */
+UE_NODISCARD inline const bool CheckTCharIsHex(const TCHAR Char)
+{
+	return (Char >= TEXT('0') && Char <= TEXT('9')) || (Char >= TEXT('A') && Char <= TEXT('F')) || (Char >= TEXT('a') && Char <= TEXT('f'));
+}
+
+/**
+ * Convert a TChar to equivalent hex value as a uint8
+ * @param Hex		The character
+ * @return	The uint8 value of a hex character
+ */
+UE_NODISCARD inline const uint8 TCharToNibble(const TCHAR Hex)
+{
+	if (Hex >= '0' && Hex <= '9')
+	{
+		return uint8(Hex - '0');
+	}
+	if (Hex >= 'A' && Hex <= 'F')
+	{
+		return uint8(Hex - 'A' + 10);
+	}
+	if (Hex >= 'a' && Hex <= 'f')
+	{
+		return uint8(Hex - 'a' + 10);
+	}
+	checkf(false, TEXT("'%c' (0x%02X) is not a valid hexadecimal digit"), Hex, Hex);
+	return 0;
+}
+
+/** Convert numeric types to a string */
+template<typename T>
+UE_NODISCARD typename TEnableIf<TIsArithmetic<T>::Value, FString>::Type
+LexToString(const T& Value)
+{
+	// std::remove_cv_t to remove potential volatile decorations. Removing const is pointless, but harmless because it's specified in the param declaration.
+	return FString::Printf(TFormatSpecifier<std::remove_cv_t<T>>::GetFormatSpecifier(), Value);
+}
+
+template<typename CharType>
+UE_NODISCARD typename TEnableIf<TIsCharType<CharType>::Value, FString>::Type
+LexToString(const CharType* Ptr)
+{
+	return FString(Ptr);
+}
+
+UE_NODISCARD inline FString LexToString(bool Value)
+{
+	return Value ? TEXT("true") : TEXT("false");
+}
+
+UE_NODISCARD FORCEINLINE FString LexToString(FString&& Str)
+{
+	return MoveTemp(Str);
+}
+
+UE_NODISCARD FORCEINLINE FString LexToString(const FString& Str)
+{
+	return Str;
+}
+
+/** Helper template to convert to sanitized strings */
+template<typename T>
+UE_NODISCARD FString LexToSanitizedString(const T& Value)
+{
+	return LexToString(Value);
+}
+
+/** Overloaded for floats */
+UE_NODISCARD inline FString LexToSanitizedString(float Value)
+{
+	return FString::SanitizeFloat(Value);
+}
+
+/** Overloaded for doubles */
+UE_NODISCARD inline FString LexToSanitizedString(double Value)
+{
+	return FString::SanitizeFloat(Value);
+}
+
+/** Shorthand legacy use for Lex functions */
+template<typename T>
+struct TTypeToString
+{
+	UE_NODISCARD static FString ToString(const T& Value)
+	{
+		return LexToString(Value);
+	}
+
+	UE_NODISCARD static FString ToSanitizedString(const T& Value)
+	{
+		return LexToSanitizedString(Value);
+	}
+};
+
+/** Parse a string into this type, returning whether it was successful */
+/** Specialization for arithmetic types */
+template<typename T>
+static typename TEnableIf<TIsArithmetic<T>::Value, bool>::Type
+LexTryParseString(T& OutValue, const TCHAR* Buffer)
+{
+	if (Buffer[0] == '\0')
+	{
+		OutValue = 0;
+		return false;
+	}
+
+	LexFromString(OutValue, Buffer);
+	if (OutValue == 0 && FMath::IsFinite((float)OutValue)) //@TODO:FLOATPRECISION: ? huh ?
+	{
+		bool bSawZero = false;
+		TCHAR C = *Buffer;
+		while (C != '\0' && (C == '+' || C == '-' || FChar::IsWhitespace(C)))
+		{
+			C = *(++Buffer);
+		}
+
+		while (C != '\0' && !FChar::IsWhitespace(C) && (TIsFloatingPoint<T>::Value || C != '.'))
+		{
+			bSawZero = bSawZero || (C == '0');
+			if (!bSawZero && C != '.')
+			{
+				return false;
+			}
+
+			C = *(++Buffer);
+		}
+		return bSawZero;
+	}
+
+	return true;
+}
+
+/** Try and parse a bool - always returns true */
+static bool LexTryParseString(bool& OutValue, const TCHAR* Buffer)
+{
+	LexFromString(OutValue, Buffer);
+	return true;
+}
+
+template<typename T>
+struct TTypeFromString
+{
+	static void FromString(T& Value, const TCHAR* Buffer)
+	{
+		return LexFromString(Value, Buffer);
+	}
+};
+
+namespace StringConv
+{
+	/** Inline combine any UTF-16 surrogate pairs in the given string */
+	CORE_API void InlineCombineSurrogates(FString& Str);
+}
+
+/*----------------------------------------------------------------------------
+	Special archivers.
+----------------------------------------------------------------------------*/
+
+//
+// String output device.
+//
+class FStringOutputDevice : public FString, public FOutputDevice
+{
+public:
+	FStringOutputDevice( const TCHAR* OutputDeviceName=TEXT("") ):
+		FString( OutputDeviceName )
+	{
+		bAutoEmitLineTerminator = false;
+	}
+	virtual void Serialize( const TCHAR* InData, ELogVerbosity::Type Verbosity, const class FName& Category ) override
+	{
+		FString::operator+=((TCHAR*)InData);
+		if(bAutoEmitLineTerminator)
+		{
+			*this += LINE_TERMINATOR;
+		}
+	}
+
+	FStringOutputDevice(FStringOutputDevice&&) = default;
+	FStringOutputDevice(const FStringOutputDevice&) = default;
+	FStringOutputDevice& operator=(FStringOutputDevice&&) = default;
+	FStringOutputDevice& operator=(const FStringOutputDevice&) = default;
+
+	// Make += operator virtual.
+	virtual FString& operator+=(const FString& Other)
+	{
+		return FString::operator+=(Other);
+	}
+};
+
+template <>
+struct TIsContiguousContainer<FStringOutputDevice>
+{
+	enum { Value = true };
+};
+
+//
+// String output device.
+//
+class FStringOutputDeviceCountLines : public FStringOutputDevice
+{
+	typedef FStringOutputDevice Super;
+
+	int32 LineCount;
+public:
+	FStringOutputDeviceCountLines( const TCHAR* OutputDeviceName=TEXT("") )
+	:	Super( OutputDeviceName )
+	,	LineCount(0)
+	{}
+
+	virtual void Serialize(const TCHAR* InData, ELogVerbosity::Type Verbosity, const class FName& Category) override
+	{
+		Super::Serialize(InData, Verbosity, Category);
+		int32 TermLength = FCString::Strlen(LINE_TERMINATOR);
+		for (;;)
+		{
+			InData = FCString::Strstr(InData, LINE_TERMINATOR);
+			if (!InData)
+			{
+				break;
+			}
+			LineCount++;
+			InData += TermLength;
+		}
+
+		if (bAutoEmitLineTerminator)
+		{
+			LineCount++;
+		}
+	}
+
+	/**
+	 * Appends other FStringOutputDeviceCountLines object to this one.
+	 */
+	virtual FStringOutputDeviceCountLines& operator+=(const FStringOutputDeviceCountLines& Other)
+	{
+		FString::operator+=(static_cast<const FString&>(Other));
+
+		LineCount += Other.GetLineCount();
+
+		return *this;
+	}
+
+	/**
+	 * Appends other FString (as well as its specializations like FStringOutputDevice)
+	 * object to this.
+	 */
+	virtual FString& operator+=(const FString& Other) override
+	{
+		Log(Other);
+
+		return *this;
+	}
+
+	int32 GetLineCount() const
+	{
+		return LineCount;
+	}
+
+	FStringOutputDeviceCountLines(const FStringOutputDeviceCountLines&) = default;
+	FStringOutputDeviceCountLines& operator=(const FStringOutputDeviceCountLines&) = default;
+
+	FORCEINLINE FStringOutputDeviceCountLines(FStringOutputDeviceCountLines&& Other)
+		: Super    ((Super&&)Other)
+		, LineCount(Other.LineCount)
+	{
+		Other.LineCount = 0;
+	}
+
+	FORCEINLINE FStringOutputDeviceCountLines& operator=(FStringOutputDeviceCountLines&& Other)
+	{
+		if (this != &Other)
+		{
+			(Super&)*this = (Super&&)Other;
+			LineCount     = Other.LineCount;
+
+			Other.LineCount = 0;
+		}
+		return *this;
+	}
+};
+
+template <>
+struct TIsContiguousContainer<FStringOutputDeviceCountLines>
+{
+	enum { Value = true };
+};
+
+struct FTextRange
+{
+	FTextRange()
+		: BeginIndex(INDEX_NONE)
+		, EndIndex(INDEX_NONE)
+	{
+
+	}
+
+	FTextRange(int32 InBeginIndex, int32 InEndIndex)
+		: BeginIndex(InBeginIndex)
+		, EndIndex(InEndIndex)
+	{
+
+	}
+
+	FORCEINLINE bool operator==(const FTextRange& Other) const
+	{
+		return BeginIndex == Other.BeginIndex
+			&& EndIndex == Other.EndIndex;
+	}
+
+	FORCEINLINE bool operator!=(const FTextRange& Other) const
+	{
+		return !(*this == Other);
+	}
+
+	friend inline uint32 GetTypeHash(const FTextRange& Key)
+	{
+		uint32 KeyHash = 0;
+		KeyHash = HashCombine(KeyHash, GetTypeHash(Key.BeginIndex));
+		KeyHash = HashCombine(KeyHash, GetTypeHash(Key.EndIndex));
+		return KeyHash;
+	}
+
+	int32 Len() const { return EndIndex - BeginIndex; }
+	bool IsEmpty() const { return (EndIndex - BeginIndex) <= 0; }
+	void Offset(int32 Amount) { BeginIndex += Amount; BeginIndex = FMath::Max(0, BeginIndex);  EndIndex += Amount; EndIndex = FMath::Max(0, EndIndex); }
+	bool Contains(int32 Index) const { return Index >= BeginIndex && Index < EndIndex; }
+	bool InclusiveContains(int32 Index) const { return Index >= BeginIndex && Index <= EndIndex; }
+
+	FTextRange Intersect(const FTextRange& Other) const
+	{
+		FTextRange Intersected(FMath::Max(BeginIndex, Other.BeginIndex), FMath::Min(EndIndex, Other.EndIndex));
+		if (Intersected.EndIndex <= Intersected.BeginIndex)
+		{
+			return FTextRange(0, 0);
+		}
+
+		return Intersected;
+	}
+
+	/**
+	 * Produce an array of line ranges from the given text, breaking at any new-line characters
+	 */
+	static CORE_API void CalculateLineRangesFromString(const FString& Input, TArray<FTextRange>& LineRanges);
+
+	int32 BeginIndex;
+	int32 EndIndex;
+};
diff --git a/Engine/Source/Runtime/Core/Public/Containers/UnrealString.h.inl b/Engine/Source/Runtime/Core/Public/Containers/UnrealString.h.inl
--- a/Engine/Source/Runtime/Core/Public/Containers/UnrealString.h.inl
+++ b/Engine/Source/Runtime/Core/Public/Containers/UnrealString.h.inl
@@ -2062,127 +2062,11 @@ inline int32 GetNum(const FString& String)
 	return String.Len();
 }
 
-/** 
- * Convert an array of bytes to a string
- * @param In byte array values to convert
- * @param Count number of bytes to convert
- * @return Valid string representing bytes.
- */
-UE_NODISCARD inline FString BytesToString(const uint8* In, int32 Count)
-{
-	FString Result;
-	Result.Empty(Count);
-
-	while (Count)
-	{
-		// Put the byte into an int16 and add 1 to it, this keeps anything from being put into the string as a null terminator
-		int16 Value = *In;
-		Value += 1;
-
-		Result += FString::ElementType(Value);
-
-		++In;
-		Count--;
-	}
-	return Result;
-}
-
-/** 
- * Convert FString of bytes into the byte array.
- * @param String		The FString of byte values
- * @param OutBytes		Ptr to memory must be preallocated large enough
- * @param MaxBufferSize	Max buffer size of the OutBytes array, to prevent overflow
- * @return	The number of bytes copied
- */
-inline int32 StringToBytes( const FString& String, uint8* OutBytes, int32 MaxBufferSize )
-{
-	int32 NumBytes = 0;
-	const FString::ElementType* CharPos = *String;
-
-	while( *CharPos && NumBytes < MaxBufferSize)
-	{
-		OutBytes[ NumBytes ] = (int8)(*CharPos - 1);
-		CharPos++;
-		++NumBytes;
-	}
-	return NumBytes;
-}
-
-/** Returns Char value of Nibble */
-UE_NODISCARD inline TCHAR NibbleToTChar(uint8 Num)
-{
-	if (Num > 9)
-	{
-		return (TCHAR)('A' + (Num - 10));
-	}
-	return (TCHAR)('0' + Num);
-}
-
-/** 
- * Convert a byte to hex
- * @param In byte value to convert
- * @param Result out hex value output
- */
-inline void ByteToHex(uint8 In, FString& Result)
-{
-	Result += NibbleToTChar(In >> 4);
-	Result += NibbleToTChar(In & 15);
-}
-
 /** Append bytes as uppercase hex string */
 CORE_API void BytesToHex(const uint8* Bytes, int32 NumBytes, FString& Out);
 /** Append bytes as lowercase hex string */
 CORE_API void BytesToHexLower(const uint8* Bytes, int32 NumBytes, FString& Out);
 
-/** Convert bytes to uppercase hex string */
-UE_NODISCARD inline FString BytesToHex(const uint8* Bytes, int32 NumBytes)
-{
-	FString Out;
-	BytesToHex(Bytes, NumBytes, Out);
-	return Out;
-}
-
-/** Convert bytes to lowercase hex string */
-UE_NODISCARD inline FString BytesToHexLower(const uint8* Bytes, int32 NumBytes)
-{
-	FString Out;
-	BytesToHexLower(Bytes, NumBytes, Out);
-	return Out;
-}
-
-/**
- * Checks if the TChar is a valid hex character
- * @param Char		The character
- * @return	True if in 0-9 and A-F ranges
- */
-UE_NODISCARD inline const bool CheckTCharIsHex( const TCHAR Char )
-{
-	return ( Char >= TEXT('0') && Char <= TEXT('9') ) || ( Char >= TEXT('A') && Char <= TEXT('F') ) || ( Char >= TEXT('a') && Char <= TEXT('f') );
-}
-
-/**
- * Convert a TChar to equivalent hex value as a uint8
- * @param Hex		The character
- * @return	The uint8 value of a hex character
- */
-UE_NODISCARD inline const uint8 TCharToNibble(const TCHAR Hex)
-{
-	if (Hex >= '0' && Hex <= '9')
-	{
-		return uint8(Hex - '0');
-	}
-	if (Hex >= 'A' && Hex <= 'F')
-	{
-		return uint8(Hex - 'A' + 10);
-	}
-	if (Hex >= 'a' && Hex <= 'f')
-	{
-		return uint8(Hex - 'a' + 10);
-	}
-	checkf(false, TEXT("'%c' (0x%02X) is not a valid hexadecimal digit"), Hex, Hex);
-	return 0;
-}
-
 /** 
  * Convert FString of Hex digits into the byte array.
  * @param HexString		The FString of Hex values
@@ -2223,115 +2107,6 @@ inline void LexFromString(double& OutValue, 	const TCHAR* Buffer)	{	OutValue = F
 inline void LexFromString(bool& OutValue, 		const TCHAR* Buffer)	{	OutValue = FCString::ToBool(Buffer);	}
 inline void LexFromString(FString& OutValue, 	const TCHAR* Buffer)	{	OutValue = Buffer;						}
 
- /** Convert numeric types to a string */
-template<typename T>
-UE_NODISCARD typename TEnableIf<TIsArithmetic<T>::Value, FString>::Type
-LexToString(const T& Value)
-{
-	// std::remove_cv_t to remove potential volatile decorations. Removing const is pointless, but harmless because it's specified in the param declaration.
-	return FString::Printf( TFormatSpecifier<std::remove_cv_t<T>>::GetFormatSpecifier(), Value );
-}
-
-template<typename CharType>
-UE_NODISCARD typename TEnableIf<TIsCharType<CharType>::Value, FString>::Type
-LexToString(const CharType* Ptr)
-{
-	return FString(Ptr);
-}
-
-UE_NODISCARD inline FString LexToString(bool Value)
-{
-	return Value ? TEXT("true") : TEXT("false");
-}
-
-UE_NODISCARD FORCEINLINE FString LexToString(FString&& Str)
-{
-	return MoveTemp(Str);
-}
-
-UE_NODISCARD FORCEINLINE FString LexToString(const FString& Str)
-{
-	return Str;
-}
-
-/** Helper template to convert to sanitized strings */
-template<typename T>
-UE_NODISCARD FString LexToSanitizedString(const T& Value)
-{
-	return LexToString(Value);
-}
-
-/** Overloaded for floats */
-UE_NODISCARD inline FString LexToSanitizedString(float Value)
-{
-	return FString::SanitizeFloat(Value);
-}
-
-/** Overloaded for doubles */
-UE_NODISCARD inline FString LexToSanitizedString(double Value)
-{
-	return FString::SanitizeFloat(Value);
-}
-
-/** Parse a string into this type, returning whether it was successful */
-/** Specialization for arithmetic types */
-template<typename T>
-static typename TEnableIf<TIsArithmetic<T>::Value, bool>::Type
-LexTryParseString(T& OutValue, const TCHAR* Buffer)
-{
-	if (Buffer[0] == '\0')
-	{
-		OutValue = 0;
-		return false;
-	}
-
-	LexFromString(OutValue, Buffer);
-	if (OutValue == 0 && FMath::IsFinite((float)OutValue)) //@TODO:FLOATPRECISION: ? huh ?
-	{
-		bool bSawZero = false;
-		TCHAR C = *Buffer;
-		while (C != '\0' && (C == '+' || C == '-' || FChar::IsWhitespace(C)))
-		{
-			C = *(++Buffer);
-		}
-
-		while (C != '\0' && !FChar::IsWhitespace(C) && (TIsFloatingPoint<T>::Value || C != '.'))
-		{
-			bSawZero = bSawZero || (C == '0');
-			if (!bSawZero && C != '.')
-			{
-				return false;
-			}
-
-			C = *(++Buffer);
-		}
-		return bSawZero;
-	}
-	
-	return true;
-}
-
-/** Try and parse a bool - always returns true */
-static bool LexTryParseString(bool& OutValue, const TCHAR* Buffer)
-{
-	LexFromString(OutValue, Buffer);
-	return true;
-}
-
-
-/** Shorthand legacy use for Lex functions */
-template<typename T>
-struct TTypeToString
-{
-	UE_NODISCARD static FString ToString(const T& Value)				{ return LexToString(Value); }
-	UE_NODISCARD static FString ToSanitizedString(const T& Value)	{ return LexToSanitizedString(Value); }
-};
-template<typename T>
-struct TTypeFromString
-{
-	static void FromString(T& Value, const TCHAR* Buffer) { return LexFromString(Value, Buffer); }
-};
-
 /**
  * Gets a non-owning TCHAR pointer from a string type.
  *
@@ -2348,146 +2123,6 @@ FORCEINLINE const TCHAR* ToCStr(const FString& Str)
 	return *Str;
 }
 
-namespace StringConv
-{
-	/** Inline combine any UTF-16 surrogate pairs in the given string */
-	CORE_API void InlineCombineSurrogates(FString& Str);
-}
-
-/*----------------------------------------------------------------------------
-	Special archivers.
-----------------------------------------------------------------------------*/
-
-//
-// String output device.
-//
-class FStringOutputDevice : public FString, public FOutputDevice
-{
-public:
-	FStringOutputDevice( const TCHAR* OutputDeviceName=TEXT("") ):
-		FString( OutputDeviceName )
-	{
-		bAutoEmitLineTerminator = false;
-	}
-	virtual void Serialize( const TCHAR* InData, ELogVerbosity::Type Verbosity, const class FName& Category ) override
-	{
-		FString::operator+=((TCHAR*)InData);
-		if(bAutoEmitLineTerminator)
-		{
-			*this += LINE_TERMINATOR;
-		}
-	}
-
-	FStringOutputDevice(FStringOutputDevice&&) = default;
-	FStringOutputDevice(const FStringOutputDevice&) = default;
-	FStringOutputDevice& operator=(FStringOutputDevice&&) = default;
-	FStringOutputDevice& operator=(const FStringOutputDevice&) = default;
-
-	// Make += operator virtual.
-	virtual FString& operator+=(const FString& Other)
-	{
-		return FString::operator+=(Other);
-	}
-};
-
-template <>
-struct TIsContiguousContainer<FStringOutputDevice>
-{
-	enum { Value = true };
-};
-
-//
-// String output device.
-//
-class FStringOutputDeviceCountLines : public FStringOutputDevice
-{
-	typedef FStringOutputDevice Super;
-
-	int32 LineCount;
-public:
-	FStringOutputDeviceCountLines( const TCHAR* OutputDeviceName=TEXT("") )
-	:	Super( OutputDeviceName )
-	,	LineCount(0)
-	{}
-
-	virtual void Serialize(const TCHAR* InData, ELogVerbosity::Type Verbosity, const class FName& Category) override
-	{
-		Super::Serialize(InData, Verbosity, Category);
-		int32 TermLength = FCString::Strlen(LINE_TERMINATOR);
-		for (;;)
-		{
-			InData = FCString::Strstr(InData, LINE_TERMINATOR);
-			if (!InData)
-			{
-				break;
-			}
-			LineCount++;
-			InData += TermLength;
-		}
-
-		if (bAutoEmitLineTerminator)
-		{
-			LineCount++;
-		}
-	}
-
-	/**
-	 * Appends other FStringOutputDeviceCountLines object to this one.
-	 */
-	virtual FStringOutputDeviceCountLines& operator+=(const FStringOutputDeviceCountLines& Other)
-	{
-		FString::operator+=(static_cast<const FString&>(Other));
-
-		LineCount += Other.GetLineCount();
-
-		return *this;
-	}
-
-	/**
-	 * Appends other FString (as well as its specializations like FStringOutputDevice)
-	 * object to this.
-	 */
-	virtual FString& operator+=(const FString& Other) override
-	{
-		Log(Other);
-
-		return *this;
-	}
-
-	int32 GetLineCount() const
-	{
-		return LineCount;
-	}
-
-	FStringOutputDeviceCountLines(const FStringOutputDeviceCountLines&) = default;
-	FStringOutputDeviceCountLines& operator=(const FStringOutputDeviceCountLines&) = default;
-
-	FORCEINLINE FStringOutputDeviceCountLines(FStringOutputDeviceCountLines&& Other)
-		: Super    ((Super&&)Other)
-		, LineCount(Other.LineCount)
-	{
-		Other.LineCount = 0;
-	}
-
-	FORCEINLINE FStringOutputDeviceCountLines& operator=(FStringOutputDeviceCountLines&& Other)
-	{
-		if (this != &Other)
-		{
-			(Super&)*this = (Super&&)Other;
-			LineCount     = Other.LineCount;
-
-			Other.LineCount = 0;
-		}
-		return *this;
-	}
-};
-
-template <>
-struct TIsContiguousContainer<FStringOutputDeviceCountLines>
-{
-	enum { Value = true };
-};
-
 /**
  * A helper function to find closing parenthesis that matches the first open parenthesis found. The open parenthesis
  * referred to must be at or further up from the start index.
@@ -2508,65 +2143,4 @@ CORE_API int32 FindMatchingClosingParenthesis(const FString& TargetString, const
 */
 CORE_API FString SlugStringForValidName(const FString& DisplayString, const TCHAR* ReplaceWith = TEXT(""));
 
-struct FTextRange
-{
-	FTextRange()
-		: BeginIndex(INDEX_NONE)
-		, EndIndex(INDEX_NONE)
-	{
-
-	}
-
-	FTextRange(int32 InBeginIndex, int32 InEndIndex)
-		: BeginIndex(InBeginIndex)
-		, EndIndex(InEndIndex)
-	{
-
-	}
-
-	FORCEINLINE bool operator==(const FTextRange& Other) const
-	{
-		return BeginIndex == Other.BeginIndex
-			&& EndIndex == Other.EndIndex;
-	}
-
-	FORCEINLINE bool operator!=(const FTextRange& Other) const
-	{
-		return !(*this == Other);
-	}
-
-	friend inline uint32 GetTypeHash(const FTextRange& Key)
-	{
-		uint32 KeyHash = 0;
-		KeyHash = HashCombine(KeyHash, GetTypeHash(Key.BeginIndex));
-		KeyHash = HashCombine(KeyHash, GetTypeHash(Key.EndIndex));
-		return KeyHash;
-	}
-
-	int32 Len() const { return EndIndex - BeginIndex; }
-	bool IsEmpty() const { return (EndIndex - BeginIndex) <= 0; }
-	void Offset(int32 Amount) { BeginIndex += Amount; BeginIndex = FMath::Max(0, BeginIndex);  EndIndex += Amount; EndIndex = FMath::Max(0, EndIndex); }
-	bool Contains(int32 Index) const { return Index >= BeginIndex && Index < EndIndex; }
-	bool InclusiveContains(int32 Index) const { return Index >= BeginIndex && Index <= EndIndex; }
-
-	FTextRange Intersect(const FTextRange& Other) const
-	{
-		FTextRange Intersected(FMath::Max(BeginIndex, Other.BeginIndex), FMath::Min(EndIndex, Other.EndIndex));
-		if (Intersected.EndIndex <= Intersected.BeginIndex)
-		{
-			return FTextRange(0, 0);
-		}
-
-		return Intersected;
-	}
-
-	/**
-	 * Produce an array of line ranges from the given text, breaking at any new-line characters
-	 */
-	static CORE_API void CalculateLineRangesFromString(const FString& Input, TArray<FTextRange>& LineRanges);
-
-	int32 BeginIndex;
-	int32 EndIndex;
-};
-
 #include "Misc/StringFormatArg.h"
