Added a CVar to control whether locales use their CLDR specified digit characters, or whether they should always use Arabic digit characters (eg, 1234)

This can be controlled via "Localization.UseLocaleSpecificDigitCharacters" and defaults to True to preserve the current behavior

diff --git a/Engine/Source/Runtime/Core/Private/Internationalization/ICUCulture.cpp b/Engine/Source/Runtime/Core/Private/Internationalization/ICUCulture.cpp
--- a/Engine/Source/Runtime/Core/Private/Internationalization/ICUCulture.cpp
+++ b/Engine/Source/Runtime/Core/Private/Internationalization/ICUCulture.cpp
@@ -10,6 +10,13 @@
 #include "Internationalization/ICUUtilities.h"
 #include "HAL/IConsoleManager.h"
 
+static TAutoConsoleVariable<bool> CVarUseLocaleSpecificDigitCharacters(
+	TEXT("Localization.UseLocaleSpecificDigitCharacters"),
+	true,
+	TEXT("False: Locales will always use Arabic digit characters (eg, 1234), True: Locales will use the digit characters specified in their CLDR data (default)."),
+	ECVF_Default
+	);
+
 static TAutoConsoleVariable<bool> CVarSpanishUsesRAENumberFormat(
 	TEXT("Localization.SpanishUsesRAENumberFormat"),
 	true,
@@ -686,16 +693,39 @@ FDecimalNumberFormattingRules ExtractNumberFormattingRulesFromICUDecimalFormatte
 	NewUEDecimalNumberFormattingRules.MinimumGroupingDigits			= static_cast<uint8>(FMath::Max(InICUDecimalFormat.getMinimumGroupingDigits(), 1));
 #endif
 
-	NewUEDecimalNumberFormattingRules.DigitCharacters[0]			= ExtractFormattingSymbolAsCharacter(icu::DecimalFormatSymbols::kZeroDigitSymbol,	TEXT('0'));
-	NewUEDecimalNumberFormattingRules.DigitCharacters[1]			= ExtractFormattingSymbolAsCharacter(icu::DecimalFormatSymbols::kOneDigitSymbol,	TEXT('1'));
-	NewUEDecimalNumberFormattingRules.DigitCharacters[2]			= ExtractFormattingSymbolAsCharacter(icu::DecimalFormatSymbols::kTwoDigitSymbol,	TEXT('2'));
-	NewUEDecimalNumberFormattingRules.DigitCharacters[3]			= ExtractFormattingSymbolAsCharacter(icu::DecimalFormatSymbols::kThreeDigitSymbol,	TEXT('3'));
-	NewUEDecimalNumberFormattingRules.DigitCharacters[4]			= ExtractFormattingSymbolAsCharacter(icu::DecimalFormatSymbols::kFourDigitSymbol,	TEXT('4'));
-	NewUEDecimalNumberFormattingRules.DigitCharacters[5]			= ExtractFormattingSymbolAsCharacter(icu::DecimalFormatSymbols::kFiveDigitSymbol,	TEXT('5'));
-	NewUEDecimalNumberFormattingRules.DigitCharacters[6]			= ExtractFormattingSymbolAsCharacter(icu::DecimalFormatSymbols::kSixDigitSymbol,	TEXT('6'));
-	NewUEDecimalNumberFormattingRules.DigitCharacters[7]			= ExtractFormattingSymbolAsCharacter(icu::DecimalFormatSymbols::kSevenDigitSymbol,	TEXT('7'));
-	NewUEDecimalNumberFormattingRules.DigitCharacters[8]			= ExtractFormattingSymbolAsCharacter(icu::DecimalFormatSymbols::kEightDigitSymbol,	TEXT('8'));
-	NewUEDecimalNumberFormattingRules.DigitCharacters[9]			= ExtractFormattingSymbolAsCharacter(icu::DecimalFormatSymbols::kNineDigitSymbol,	TEXT('9'));
+	if (CVarUseLocaleSpecificDigitCharacters.AsVariable()->GetBool())
+	{
+		NewUEDecimalNumberFormattingRules.DigitCharacters[0]		= ExtractFormattingSymbolAsCharacter(icu::DecimalFormatSymbols::kZeroDigitSymbol,	TEXT('0'));
+		NewUEDecimalNumberFormattingRules.DigitCharacters[1]		= ExtractFormattingSymbolAsCharacter(icu::DecimalFormatSymbols::kOneDigitSymbol,	TEXT('1'));
+		NewUEDecimalNumberFormattingRules.DigitCharacters[2]		= ExtractFormattingSymbolAsCharacter(icu::DecimalFormatSymbols::kTwoDigitSymbol,	TEXT('2'));
+		NewUEDecimalNumberFormattingRules.DigitCharacters[3]		= ExtractFormattingSymbolAsCharacter(icu::DecimalFormatSymbols::kThreeDigitSymbol,	TEXT('3'));
+		NewUEDecimalNumberFormattingRules.DigitCharacters[4]		= ExtractFormattingSymbolAsCharacter(icu::DecimalFormatSymbols::kFourDigitSymbol,	TEXT('4'));
+		NewUEDecimalNumberFormattingRules.DigitCharacters[5]		= ExtractFormattingSymbolAsCharacter(icu::DecimalFormatSymbols::kFiveDigitSymbol,	TEXT('5'));
+		NewUEDecimalNumberFormattingRules.DigitCharacters[6]		= ExtractFormattingSymbolAsCharacter(icu::DecimalFormatSymbols::kSixDigitSymbol,	TEXT('6'));
+		NewUEDecimalNumberFormattingRules.DigitCharacters[7]		= ExtractFormattingSymbolAsCharacter(icu::DecimalFormatSymbols::kSevenDigitSymbol,	TEXT('7'));
+		NewUEDecimalNumberFormattingRules.DigitCharacters[8]		= ExtractFormattingSymbolAsCharacter(icu::DecimalFormatSymbols::kEightDigitSymbol,	TEXT('8'));
+		NewUEDecimalNumberFormattingRules.DigitCharacters[9]		= ExtractFormattingSymbolAsCharacter(icu::DecimalFormatSymbols::kNineDigitSymbol,	TEXT('9'));
+	}
+	else
+	{
+		auto ReplaceLocaleSeparatorWithSuitableEquivalent = [](TCHAR& InOutSeparatorCharacter)
+		{
+			switch (InOutSeparatorCharacter)
+			{
+			case TEXT('\u066B'): // Arabic decimal separator
+				InOutSeparatorCharacter = TEXT('.');
+				break;
+			case TEXT('\u066C'): // Arabic group separator
+				InOutSeparatorCharacter = TEXT(',');
+				break;
+			default:
+				break;
+			}
+		};
+
+		ReplaceLocaleSeparatorWithSuitableEquivalent(NewUEDecimalNumberFormattingRules.GroupingSeparatorCharacter);
+		ReplaceLocaleSeparatorWithSuitableEquivalent(NewUEDecimalNumberFormattingRules.DecimalSeparatorCharacter);
+	}
 
 	if (FCStringAnsi::Strcmp(InICULocale.getLanguage(), "es") == 0)
 	{
