ConcurrentLinearAllocator: Allow not FastPath platforms to allocate blocks with supported alignment.
Currently Malloc will fail on such platforms, because not all platforms support 64kb alignment.

#tests preflights

diff --git a/Engine/Source/Runtime/Core/Public/Experimental/ConcurrentLinearAllocator.h b/Engine/Source/Runtime/Core/Public/Experimental/ConcurrentLinearAllocator.h
--- a/Engine/Source/Runtime/Core/Public/Experimental/ConcurrentLinearAllocator.h
+++ b/Engine/Source/Runtime/Core/Public/Experimental/ConcurrentLinearAllocator.h
@@ -285,9 +285,10 @@ class TConcurrentLinearAllocator
 		if constexpr (!BlockAllocationTag::InlineBlockAllocation)
 		{
 			static_assert(BlockAllocationTag::BlockSize >= sizeof(FBlockHeader) + sizeof(FAllocationHeader));
-			Header = new (BlockAllocationTag::Allocator::Malloc(BlockAllocationTag::BlockSize, BlockAllocationTag::BlockSize)) FBlockHeader;
+			uint32 BlockAlignment = SupportsFastPath ? BlockAllocationTag::BlockSize : alignof(FBlockHeader);
+			Header = new (BlockAllocationTag::Allocator::Malloc(BlockAllocationTag::BlockSize, BlockAlignment)) FBlockHeader;
 			MemoryTrace_MarkAllocAsHeap(uint64(Header), EMemoryTraceRootHeap::SystemMemory);
-			checkSlow(IsAligned(Header, BlockAllocationTag::BlockSize));
+			checkSlow(IsAligned(Header, BlockAlignment));
 			if constexpr (!SupportsFastPath)
 			{
 				ASAN_POISON_MEMORY_REGION( Header + 1, BlockAllocationTag::BlockSize - sizeof(FBlockHeader) );
@@ -329,9 +330,10 @@ public:
 			if constexpr (BlockAllocationTag::InlineBlockAllocation)
 			{
 				static_assert(BlockAllocationTag::BlockSize >= sizeof(FBlockHeader) + sizeof(FAllocationHeader));
-				Header = new (BlockAllocationTag::Allocator::Malloc(BlockAllocationTag::BlockSize, BlockAllocationTag::BlockSize)) FBlockHeader;
+				uint32 BlockAlignment = SupportsFastPath ? BlockAllocationTag::BlockSize : alignof(FBlockHeader);
+				Header = new (BlockAllocationTag::Allocator::Malloc(BlockAllocationTag::BlockSize, BlockAlignment)) FBlockHeader;
 				MemoryTrace_MarkAllocAsHeap(uint64(Header), EMemoryTraceRootHeap::SystemMemory);
-				checkSlow(IsAligned(Header, BlockAllocationTag::BlockSize));
+				checkSlow(IsAligned(Header, BlockAlignment));
 				if constexpr (!SupportsFastPath)
 				{
 					ASAN_POISON_MEMORY_REGION( Header + 1, BlockAllocationTag::BlockSize - sizeof(FBlockHeader) );
