Log platform, device id and machine id when startup.

diff --git a/Engine/Source/Runtime/Core/Private/Misc/App.cpp b/Engine/Source/Runtime/Core/Private/Misc/App.cpp
--- a/Engine/Source/Runtime/Core/Private/Misc/App.cpp
+++ b/Engine/Source/Runtime/Core/Private/Misc/App.cpp
@@ -375,6 +375,11 @@ void FApp::PrintStartupLogMessages()
 {
 	UE_LOG(LogInit, Log, TEXT("ExecutableName: %s"), FPlatformProcess::ExecutableName(false));
 	UE_LOG(LogInit, Log, TEXT("Build: %s"), FApp::GetBuildVersion());
+
+	UE_LOG(LogInit, Log, TEXT("Platform=%s"), ANSI_TO_TCHAR(FPlatformProperties::PlatformName()));
+	UE_LOG(LogInit, Log, TEXT("MachineId=%s"), *FPlatformMisc::GetLoginId());
+	UE_LOG(LogInit, Log, TEXT("Deviceid=%s"), *FPlatformMisc::GetDeviceId());
+
 	UE_LOG(LogInit, Log, TEXT("Engine Version: %s"), *FEngineVersion::Current().ToString());
 	UE_LOG(LogInit, Log, TEXT("Compatible Engine Version: %s"), *FEngineVersion::CompatibleWith().ToString());
 	UE_LOG(LogInit, Log, TEXT("Net CL: %u"), FNetworkVersion::GetNetworkCompatibleChangelist());
