Logging: Added support for structured logging of TOptional

diff --git a/Engine/Source/Runtime/Core/Public/Logging/StructuredLog.h b/Engine/Source/Runtime/Core/Public/Logging/StructuredLog.h
--- a/Engine/Source/Runtime/Core/Public/Logging/StructuredLog.h
+++ b/Engine/Source/Runtime/Core/Public/Logging/StructuredLog.h
@@ -7,6 +7,7 @@
 #include "Logging/LogCategory.h"
 #include "Logging/LogMacros.h"
 #include "Logging/LogVerbosity.h"
+#include "Misc/OptionalFwd.h"
 #include "Serialization/CompactBinary.h"
 #include "Serialization/CompactBinaryWriter.h"
 #include "Templates/IsArrayOrRefOfType.h"
@@ -236,6 +237,17 @@ inline void CallSerializeForLog(FCbWriter& Writer, ValueType&& Value)
 
 } // UE
 
+template <typename ValueType UE_REQUIRES(TModels_V<UE::CSerializableForLog, ValueType>)>
+inline void SerializeForLog(FCbWriter& Writer, const TOptional<ValueType>& Optional)
+{
+	Writer.BeginArray();
+	if (Optional)
+	{
+		UE::CallSerializeForLog(Writer, Optional.GetValue());
+	}
+	Writer.EndArray();
+}
+
 namespace UE::Logging::Private
 {
 
