Revert debug logging in UnixPlatformOutputDevices.cpp

diff --git a/Engine/Source/Runtime/Core/Private/Unix/UnixPlatformOutputDevices.cpp b/Engine/Source/Runtime/Core/Private/Unix/UnixPlatformOutputDevices.cpp
--- a/Engine/Source/Runtime/Core/Private/Unix/UnixPlatformOutputDevices.cpp
+++ b/Engine/Source/Runtime/Core/Private/Unix/UnixPlatformOutputDevices.cpp
@@ -8,7 +8,6 @@
 #include "Misc/Build.h"
 #include "Misc/CommandLine.h"
 #include "Misc/OutputDeviceConsole.h"
-#include "Misc/OutputDeviceDebug.h"
 #include "Misc/OutputDeviceRedirector.h"
 #include "Misc/Parse.h"
 #include "Misc/Paths.h"
@@ -35,15 +34,6 @@ void FUnixOutputDevices::SetupOutputDevices()
 		}
 	}
 
-#if USE_DEBUG_LOGGING
-	// If the platform has a separate debug output channel (e.g. OutputDebugString) then add an output device
-	// unless logging is turned off
-		if (FPlatformMisc::HasSeparateChannelForDebugOutput() && !FParse::Param(FCommandLine::Get(), TEXT("NODEBUGOUTPUT")))
-		{
-			GLog->AddOutputDevice(new FOutputDeviceDebug());
-		}
-#endif // USE_DEBUG_LOGGING
-
 	GLog->AddOutputDevice(FPlatformOutputDevices::GetEventLog());
 }
 
