Remove extra line
#rnx

diff --git a/Engine/Source/Runtime/Core/Private/Unix/UnixPlatformOutputDevices.cpp b/Engine/Source/Runtime/Core/Private/Unix/UnixPlatformOutputDevices.cpp
--- a/Engine/Source/Runtime/Core/Private/Unix/UnixPlatformOutputDevices.cpp
+++ b/Engine/Source/Runtime/Core/Private/Unix/UnixPlatformOutputDevices.cpp
@@ -33,8 +33,7 @@ void FUnixOutputDevices::SetupOutputDevices()
 			GLog->AddOutputDevice(GLogConsole);
 		}
 	}
-
-	GLog->AddOutputDevice(FPlatformOutputDevices::GetEventLog());
+	// debug and event logging is not really supported on Unix. 
 }
 
 FString FUnixOutputDevices::GetAbsoluteLogFilename()
