Add ability to presize name table in editor for large projects.
Example config in Editor.ini:

[Names]
PreallocateNames=10485760
PreallocateNameMemoryMB=900

diff --git a/Engine/Source/Runtime/Core/Private/UObject/UnrealNames.cpp b/Engine/Source/Runtime/Core/Private/UObject/UnrealNames.cpp
--- a/Engine/Source/Runtime/Core/Private/UObject/UnrealNames.cpp
+++ b/Engine/Source/Runtime/Core/Private/UObject/UnrealNames.cpp
@@ -1932,7 +1932,7 @@ void FNamePool::RetraceAll() const
 
 void FNamePool::LogStats(FOutputDevice& Ar) const
 {
-	Ar.Logf(TEXT("%i FNames using in %ikB + %ikB"), NumEntries(), sizeof(FNamePool), Entries.NumBlocks() * FNameEntryAllocator::BlockSizeBytes / 1024);
+	Ar.Logf(TEXT("%i FNames using in %ikB + %ikB"), NumEntries(), sizeof(FNamePool) / 1024, Entries.NumBlocks() * FNameEntryAllocator::BlockSizeBytes / 1024);
 	Ar.Logf(TEXT("%d ansi FNames"), NumAnsiEntries());
 	Ar.Logf(TEXT("%d wide FNames"), NumWideEntries());
 #if UE_FNAME_OUTLINE_NUMBER
@@ -2387,6 +2387,12 @@ FString FNameEntrySerialized::GetPlainNameString() const
 static TArray<FString> NameToDisplayStringExemptions;
 static FRWLock NameToDisplayStringExemptionLock;
 
+void FName::Reserve(uint32 NumBytes, uint32 NumNames)
+{
+	TRACE_CPUPROFILER_EVENT_SCOPE_STR("FName Reserve");
+	GetNamePoolPostInit().Reserve(NumBytes, NumNames);
+}
+
 int32 FName::GetNameEntryMemorySize()
 {
 	return GetNamePool().NumBlocks() * FNameEntryAllocator::BlockSizeBytes;
diff --git a/Engine/Source/Runtime/Core/Public/UObject/NameTypes.h b/Engine/Source/Runtime/Core/Public/UObject/NameTypes.h
--- a/Engine/Source/Runtime/Core/Public/UObject/NameTypes.h
+++ b/Engine/Source/Runtime/Core/Public/UObject/NameTypes.h
@@ -1008,6 +1008,8 @@ public:
 	CORE_API static void DisplayHash( class FOutputDevice& Ar );
 	CORE_API static FString SafeString(FNameEntryId InDisplayIndex, int32 InstanceNumber = NAME_NO_NUMBER_INTERNAL);
 
+	CORE_API static void Reserve(uint32 NumBytes, uint32 NumNames);
+
 	/**
 	 * @return Size of all name entries.
 	 */
diff --git a/Engine/Source/Runtime/Launch/Private/LaunchEngineLoop.cpp b/Engine/Source/Runtime/Launch/Private/LaunchEngineLoop.cpp
--- a/Engine/Source/Runtime/Launch/Private/LaunchEngineLoop.cpp
+++ b/Engine/Source/Runtime/Launch/Private/LaunchEngineLoop.cpp
@@ -6409,6 +6409,20 @@ bool FEngineLoop::AppInit( )
 		// init config system
 		FConfigCacheIni::InitializeConfigSystem();
 	}
+	
+#if WITH_EDITOR
+	if (GIsEditor)
+	{
+		int32 NumPreallocateNames = 0;
+		GConfig->GetInt(TEXT("NameTable"), TEXT("PreallocateNames"), NumPreallocateNames, GEditorIni);
+		int32 PreAllocateNameMB = 0;
+		GConfig->GetInt(TEXT("NameTable"), TEXT("PreallocateNameMemoryMB"), PreAllocateNameMB, GEditorIni);
+		if (NumPreallocateNames || PreAllocateNameMB)
+		{
+			FName::Reserve(PreAllocateNameMB * 1024 * 1024, NumPreallocateNames);
+		}
+	}
+#endif
 
 #if USE_IO_DISPATCHER
 	// Initialize on demand I/O dispatcher backend
