Add GPU breadcrumbs information into crash reports.

* Breadcrumbs are processed into a common data structure used to dump the contents to the log and also used by the crash reporting routines to format breadcrumbs into the crash report file.
* The format used by the crash report is versioned so that parsers know what to expect if the format ever changes. The format is specifically written for automated parsing - the human-readable dump stays in the log.
* A (per-queue) hash of the full breadcrumbs stack and a separate hash for the active subset are computed and stored in the crash report.
* The existing string dumped to the log is unchanged.

Example of the new crash report section:
<GPUBreadcrumbs>
	<FormatVersion>1.0</FormatVersion>
	<Queue>
		<Name>3D Queue 0</Name>
		<FullHash>D02D1A942643A9C549606769037ED96B52104742</FullHash>
		<ActiveHash>4DE9A4E81D71A1993E939D82A1C3ED063B6F6E7C</ActiveHash>
		<Breadcrumbs>{{Frame 2540},A,{{{BufferPoolCopyOps},F},{{TexturePoolCopyOps},F},{{WorldTick},F},{{SendAllEndOfFrameUpdates},F},{{GPUDebugCrash_DirectQueue_Hang},A},{{ClearGPUMessageBuffer},A},{{VirtualTextureClear},N},{{ShaderPrint::UploadParameters},N},{{UpdateDistanceFieldAtlas},N},{{Scene},N},{{EnqueueCopy(GPUMessageManager.MessageBuffer)},N},{{AccessModePass[Graphics] (Textures: 7, Buffers: 6)},N},{{SlateUI Title = Test Application },N}}}</Breadcrumbs>
	</Queue>
	<Queue>
		<Name>Compute Queue 0</Name>
		<FullHash>1DC8AED3D68A1B4DC2719B9D94F3EFBE16E8E371</FullHash>
		<ActiveHash>DA39A3EE5E6B4B0D3255BFEF95601890AFD80709</ActiveHash>
		<Breadcrumbs>{{AccessModePass[AsyncCompute] (Textures: 4, Buffers: 2)},F}{{Scene},F}{{AccessModePass[AsyncCompute] (Textures: 4, Buffers: 2)},F}{{Scene},F}{{AccessModePass[AsyncCompute] (Textures: 4, Buffers: 2)},F}{{Scene},F}{{AccessModePass[AsyncCompute] (Textures: 4, Buffers: 2)},F}{{Scene},N}{{AccessModePass[AsyncCompute] (Textures: 4, Buffers: 2)},N}</Breadcrumbs>
	</Queue>
</GPUBreadcrumbs>

diff --git a/Engine/Source/Runtime/Core/Private/GenericPlatform/GenericPlatformCrashContext.cpp b/Engine/Source/Runtime/Core/Private/GenericPlatform/GenericPlatformCrashContext.cpp
--- a/Engine/Source/Runtime/Core/Private/GenericPlatform/GenericPlatformCrashContext.cpp
+++ b/Engine/Source/Runtime/Core/Private/GenericPlatform/GenericPlatformCrashContext.cpp
@@ -86,6 +86,73 @@ static const TCHAR* AttendedStatusToString(EUnattendedStatus Status)
 	}
 }
 
+/* GPU breadcrumbs */
+class FGPUBreadcrumbCrashData
+{
+public:
+	FGPUBreadcrumbCrashData(const TArray<FBreadcrumbNode>& Breadcrumbs)
+	{
+		for (const FBreadcrumbNode& Node : Breadcrumbs)
+		{
+			ProcessBreadcrumbNode(Node);
+		}
+
+		FinalizedFullHash = FullHash.Finalize();
+		FinalizedActiveHash = ActiveHash.Finalize();
+	}
+
+	const FString& GetProcessedBreadcrumbString() const { return ProcessedBreadcrumbString; }
+	const FSHAHash GetFullHash() const { return FinalizedFullHash; }
+	const FSHAHash GetActiveHash() const { return FinalizedActiveHash; }
+	
+	/** 
+	 * This must be incremented whenever the format of the breadcrumb string
+	 * changes, in order to help parsers in dealing with strings from multiple
+	 * versions.
+	 */
+	static const TCHAR* const GetBreadcrumbStringFormatVersion() { return TEXT("1.0"); }
+
+private:
+	void ProcessBreadcrumbNode(const FBreadcrumbNode& Node)
+	{
+		FullHash.UpdateWithString(*Node.Name, Node.Name.Len());
+		if (Node.State == EBreadcrumbState::Active)
+		{
+			ActiveHash.UpdateWithString(*Node.Name, Node.Name.Len());
+		}
+
+		ProcessedBreadcrumbString.Append(FString::Printf(TEXT("{{%s},%c"), *SanitizeBreadcrumbEventName(Node.Name), Node.GetStateString()[0]));
+		if (!Node.Children.IsEmpty())
+		{
+			ProcessedBreadcrumbString.Append(TEXT(",{"));
+			for (int32 Child = 0; Child < Node.Children.Num(); Child++)
+			{
+				ProcessBreadcrumbNode(Node.Children[Child]);
+				if (Child != Node.Children.Num() - 1)
+				{
+					ProcessedBreadcrumbString.AppendChar(',');
+				}
+			}
+			ProcessedBreadcrumbString.AppendChar('}');
+		}
+		ProcessedBreadcrumbString.AppendChar('}');
+	}
+
+	// Sanitize the event name string to remove characters that are used
+	// as delimiters for parsing.
+	static FString SanitizeBreadcrumbEventName(const FString& EventName)
+	{
+		return EventName.Replace(TEXT("{"), TEXT("(")).Replace(TEXT("}"), TEXT(")"));
+	}
+
+	FString ProcessedBreadcrumbString;	
+
+	FSHA1 FullHash;
+	FSHAHash FinalizedFullHash;
+	FSHA1 ActiveHash;
+	FSHAHash FinalizedActiveHash;
+};
+
 /*-----------------------------------------------------------------------------
 	FGenericCrashContext
 -----------------------------------------------------------------------------*/
@@ -146,6 +213,7 @@ namespace NCached
 	static TArray<FString> EnabledPluginsList;
 	static TMap<FString, FString> EngineData;
 	static TMap<FString, FString> GameData;
+	static TMap<FString, FGPUBreadcrumbCrashData> GPUBreadcrumbsByQueue;
 
 	template <size_t CharCount, typename CharType>
 	void Set(CharType(&Dest)[CharCount], const CharType* pSrc)
@@ -802,6 +870,8 @@ void FGenericCrashContext::SerializeContentToBuffer() const
 	AddPortableCallStack();
 	AddPortableCallStackHash();
 
+	AddGPUBreadcrumbs();
+
 	{
 		FString AllThreadStacks;
 		if (GetPlatformAllThreadContextsString(AllThreadStacks))
@@ -992,7 +1062,33 @@ void FGenericCrashContext::AddPortableCallStack() const
 	EndSection(CommonBuffer, TEXT("PCallStack"));
 }
 
+void FGenericCrashContext::AddGPUBreadcrumbs() const
+{
+	if (NCached::GPUBreadcrumbsByQueue.IsEmpty())
+	{
+		return;
+	}
+	
+	BeginSection(CommonBuffer, TEXT("GPUBreadcrumbs"));
 
+	// We use a version indicator for the format used by the breadcrumbs
+	// string, so that parsers can know what to expect and don't break
+	// if changes are made in the format exported by the engine.
+	AddSection(CommonBuffer, TEXT("FormatVersion"), FGPUBreadcrumbCrashData::GetBreadcrumbStringFormatVersion());
+
+	for (auto& [Queue, Breadcrumbs] : NCached::GPUBreadcrumbsByQueue)
+	{
+		BeginSection(CommonBuffer, TEXT("Queue"));
+
+		AddSection(CommonBuffer, TEXT("Name"), Queue);
+		AddSection(CommonBuffer, TEXT("FullHash"), Breadcrumbs.GetFullHash().ToString());
+		AddSection(CommonBuffer, TEXT("ActiveHash"), Breadcrumbs.GetActiveHash().ToString());
+		AddSection(CommonBuffer, TEXT("Breadcrumbs"), Breadcrumbs.GetProcessedBreadcrumbString());
+
+		EndSection(CommonBuffer, TEXT("Queue"));
+	}
+	EndSection(CommonBuffer, TEXT("GPUBreadcrumbs"));
+}
 
 void FGenericCrashContext::AddHeader(FString& Buffer)
 {
@@ -1015,6 +1111,13 @@ void FGenericCrashContext::EndSection(FString& Buffer, const TCHAR* SectionName)
 	Buffer.Appendf(TEXT("</%s>" LINE_TERMINATOR_ANSI), SectionName);
 }
 
+void FGenericCrashContext::AddSection(FString& Buffer, const TCHAR* SectionName, const FString& SectionContent)
+{
+	BeginSection(Buffer, SectionName);
+	Buffer.Appendf(TEXT("%s"), *FXmlEscapedString(SectionContent));
+	EndSection(Buffer, SectionName);
+}
+
 template<typename DEST>
 static void AppendEscapedXMLString(DEST& OutBuffer, FStringView Text)
 {
@@ -1185,6 +1288,11 @@ void FGenericCrashContext::SetEngineData(const FString& Key, const FString& Valu
 	OnEngineDataSet.Broadcast(Key, Value);
 }
 
+void FGenericCrashContext::SetGPUBreadcrumbs(const FString& GPUQueueName, const TArray<FBreadcrumbNode>& Breadcrumbs)
+{
+	NCached::GPUBreadcrumbsByQueue.Emplace(GPUQueueName, FGPUBreadcrumbCrashData(Breadcrumbs));
+}
+
 const TMap<FString, FString>& FGenericCrashContext::GetEngineData()
 {
 	return NCached::EngineData;
diff --git a/Engine/Source/Runtime/Core/Public/GenericPlatform/GenericPlatformCrashContext.h b/Engine/Source/Runtime/Core/Public/GenericPlatform/GenericPlatformCrashContext.h
--- a/Engine/Source/Runtime/Core/Public/GenericPlatform/GenericPlatformCrashContext.h
+++ b/Engine/Source/Runtime/Core/Public/GenericPlatform/GenericPlatformCrashContext.h
@@ -290,6 +290,29 @@ struct FThreadCallStack
 	uint32 ThreadId;
 };
 
+/** GPU breadcrumbs. */
+enum EBreadcrumbState : uint8
+{
+	NotStarted = 0,
+	Active = 1,
+	Finished = 2,
+	Overflow = 3,
+	Invalid = 4,
+};
+const TCHAR* const EBreadcrumbStateStrings[] = { TEXT("Not started"), TEXT("Active"), TEXT("Finished"), TEXT("Overflow"), TEXT("Invalid") };
+
+struct FBreadcrumbNode
+{
+	EBreadcrumbState State = EBreadcrumbState::Invalid;
+	FString Name;
+	TArray<FBreadcrumbNode> Children;
+
+	const TCHAR* const GetStateString() const
+	{
+		return EBreadcrumbStateStrings[FMath::Min(State, EBreadcrumbState::Invalid)];
+	}
+};
+
 /**
  *	Contains a runtime crash's properties that are common for all platforms.
  *	This may change in the future.
@@ -474,6 +497,9 @@ public:
 	/** Updates (or adds if not already present) arbitrary engine data to the crash context (will remove the key if passed an empty string) */
 	CORE_API static void SetEngineData(const FString& Key, const FString& Value);
 
+	/** Updates (or adds if not already present) GPU breadcrumb data for a given GPU queue */
+	CORE_API static void SetGPUBreadcrumbs(const FString& GPUQueueName, const TArray<FBreadcrumbNode>& Breadcrumbs);
+
 	/** Accessor for engine data change callback delegate */
 	static FEngineDataSetDelegate& OnEngineDataSetDelegate() { return OnEngineDataSet; }
 
@@ -663,6 +689,9 @@ private:
 	/** Produces a hash based on the offsets of the portable callstack and adds it to the xml */
 	void AddPortableCallStackHash() const;
 
+	/** Add GPU breadcrumbs information to the crash report xml */
+	void AddGPUBreadcrumbs() const;
+
 	/** Add module/pdb information to the crash report xml */
 	void AddModules() const;
 
@@ -674,6 +703,7 @@ private:
 
 	static void BeginSection(FString& Buffer, const TCHAR* SectionName);
 	static void EndSection(FString& Buffer, const TCHAR* SectionName);
+	static void AddSection(FString& Buffer, const TCHAR* SectionName, const FString& SectionContent);
 
 	/** Called once when GConfig is initialized. Opportunity to cache values from config. */
 	static void InitializeFromConfig();
diff --git a/Engine/Source/Runtime/D3D12RHI/Private/D3D12Util.cpp b/Engine/Source/Runtime/D3D12RHI/Private/D3D12Util.cpp
--- a/Engine/Source/Runtime/D3D12RHI/Private/D3D12Util.cpp
+++ b/Engine/Source/Runtime/D3D12RHI/Private/D3D12Util.cpp
@@ -162,61 +162,70 @@ static FString GetD3D12TextureFlagString(uint32 TextureFlags)
 	return TextureFormatText;
 }
 
-void DumpBreadcrumbScopeRecursive(
+FBreadcrumbNode CollectBreadcrumbNode(
 	D3D12RHI::FD3DGPUProfiler& GPUProfiler,
 	const TSharedPtr<FBreadcrumbStack>& Stack,
-	FString& OutGpuProgress,
-	const FBreadcrumbStack::FScope& Scope,
-	uint32 Indent)
+	const FBreadcrumbStack::FScope& Scope)
 {
+	FBreadcrumbNode Node;
+
 	const volatile uint32* Markers = reinterpret_cast<uint32*>(Stack->CPUAddress);
 	const bool bHaveMarkerIndex = (Scope.MarkerIndex < Stack->MaxMarkers);
-	const uint32 Marker = bHaveMarkerIndex ? Markers[Scope.MarkerIndex] : 0;
+	Node.State = bHaveMarkerIndex ? (EBreadcrumbState)Markers[Scope.MarkerIndex] : EBreadcrumbState::Overflow;
 
-	const bool bDidOpen = Marker > 0;
-	const bool bDidClose = Marker > 1;
+	const FString* EventName = GPUProfiler.FindEventString(Scope.NameCRC);
+	Node.Name = EventName ? *EventName : TEXT("Unknown");
 
+	if (Node.State == EBreadcrumbState::Active)
+	{
+		for (uint32 Child = Scope.Child; Child != 0; Child = Stack->Scopes[Child].Sibling)
+		{
+			Node.Children.Emplace(CollectBreadcrumbNode(GPUProfiler, Stack, Stack->Scopes[Child]));
+		}
+	}
+
+	return Node;
+};
+
+void BuildBreadcrumbsDumpString(const FBreadcrumbNode& Breadcrumb, FString& OutString, uint32 Indent)
+{
 	const TCHAR* Prefix = TEXT("  ");
-	if (bDidOpen && bDidClose)
+	if (Breadcrumb.State == EBreadcrumbState::Finished)
 	{
 		Prefix = TEXT("| ");
 	}
-	else if (bDidOpen && !bDidClose)
+	else if (Breadcrumb.State == EBreadcrumbState::Active)
 	{
 		Prefix = TEXT("> ");
 	}
 
 	const TCHAR* Suffix = TEXT("");
-	if (!bHaveMarkerIndex)
+	if (Breadcrumb.State == EBreadcrumbState::Overflow)
 	{
 		Suffix = TEXT(" [overflow]");
 	}
 
-	const FString* EventName = GPUProfiler.FindEventString(Scope.NameCRC);
-
 	for (uint32 Idx = 0; Idx < Indent * 2; ++Idx)
 	{
-		OutGpuProgress.AppendChar(' ');
+		OutString.AppendChar(' ');
 	}
 
-	const TCHAR* MarkerState[] = { TEXT("Not started"), TEXT("Active"), TEXT("Finished"), TEXT("Invalid")};
-	OutGpuProgress.Append(FString::Printf(TEXT("Breadcrumbs: %s%s [%s]%s\n"), Prefix, **EventName, MarkerState[FMath::Min(3U, Marker)], Suffix));
+	OutString.Append(FString::Printf(TEXT("Breadcrumbs: %s%s [%s]%s\n"), Prefix, *Breadcrumb.Name, Breadcrumb.GetStateString(), Suffix));
 
-	if (bDidOpen && !bDidClose)
+	for (const FBreadcrumbNode& Child : Breadcrumb.Children)
 	{
-		for (uint32 Child = Scope.Child; Child != 0; Child = Stack->Scopes[Child].Sibling)
-		{
-			DumpBreadcrumbScopeRecursive(GPUProfiler, Stack, OutGpuProgress, Stack->Scopes[Child], Indent + 1);
-		}
+		BuildBreadcrumbsDumpString(Child, OutString, Indent + 1);
 	}
-};
+}
 
 /** Log the GPU progress of the given queue to the Error log if breadcrumb data is available */
 static bool LogBreadcrumbData(D3D12RHI::FD3DGPUProfiler& GPUProfiler, FD3D12Queue& Queue)
 {
-	FString GpuProgress = FString::Printf(TEXT("[GPUBreadCrumb]\t%s Queue %d\n"), GetD3DCommandQueueTypeName(Queue.QueueType),
+	FString GPUQueueName = FString::Printf(TEXT("%s Queue %d"), GetD3DCommandQueueTypeName(Queue.QueueType),
 		Queue.Device->GetGPUIndex());
 
+	FString GpuProgress = FString::Printf(TEXT("[GPUBreadCrumb]\t%s\n"), *GPUQueueName);
+
 	TArray<TSharedPtr<FBreadcrumbStack>, TInlineAllocator<8>> UniqueStacks;
 
 	while (const FD3D12Payload* Payload = Queue.PendingInterrupt.Peek())
@@ -228,6 +237,7 @@ static bool LogBreadcrumbData(D3D12RHI::FD3DGPUProfiler& GPUProfiler, FD3D12Queu
 		Queue.PendingInterrupt.Pop();
 	}
 
+	TArray<FBreadcrumbNode> Nodes;
 	for (const TSharedPtr<FBreadcrumbStack>& Stack : UniqueStacks)
 	{
 		if (!Stack->Scopes.IsEmpty())
@@ -237,15 +247,25 @@ static bool LogBreadcrumbData(D3D12RHI::FD3DGPUProfiler& GPUProfiler, FD3D12Queu
 				uint32 Scope = 0;
 				do
 				{
-					DumpBreadcrumbScopeRecursive(GPUProfiler, Stack, GpuProgress, Stack->Scopes[Scope], 0);
+					Nodes.Emplace(CollectBreadcrumbNode(GPUProfiler, Stack, Stack->Scopes[Scope]));
 					Scope = Stack->Scopes[Scope].Sibling;
 				} while (Scope != 0);
 			}
 		}
 	}
 
+	for (const FBreadcrumbNode& Node : Nodes)
+	{
+		BuildBreadcrumbsDumpString(Node, GpuProgress, 0);
+	}
+
 	UE_LOG(LogD3D12RHI, Error, TEXT("%s"), *GpuProgress);
 
+	if (!Nodes.IsEmpty())
+	{
+		FPlatformCrashContext::SetGPUBreadcrumbs(GPUQueueName, Nodes);
+	}
+
 	const FD3D12DiagnosticBufferData* DiagnosticData = Queue.GetDiagnosticBufferData();
 	if (DiagnosticData && DiagnosticData->Counter)
 	{
