Fix compile errors and missing include.

diff --git a/Engine/Source/Runtime/Core/Public/GenericPlatform/GenericPlatformCrashContext.h b/Engine/Source/Runtime/Core/Public/GenericPlatform/GenericPlatformCrashContext.h
--- a/Engine/Source/Runtime/Core/Public/GenericPlatform/GenericPlatformCrashContext.h
+++ b/Engine/Source/Runtime/Core/Public/GenericPlatform/GenericPlatformCrashContext.h
@@ -291,7 +291,7 @@ struct FThreadCallStack
 };
 
 /** GPU breadcrumbs. */
-enum EBreadcrumbState : uint8
+enum class EBreadcrumbState : uint8
 {
 	NotStarted = 0,
 	Active = 1,
@@ -309,7 +309,7 @@ struct FBreadcrumbNode
 
 	const TCHAR* const GetStateString() const
 	{
-		return EBreadcrumbStateStrings[FMath::Min(State, EBreadcrumbState::Invalid)];
+		return EBreadcrumbStateStrings[static_cast<uint32>(FMath::Min(State, EBreadcrumbState::Invalid))];
 	}
 };
 
diff --git a/Engine/Source/Runtime/D3D12RHI/Private/D3D12Util.cpp b/Engine/Source/Runtime/D3D12RHI/Private/D3D12Util.cpp
--- a/Engine/Source/Runtime/D3D12RHI/Private/D3D12Util.cpp
+++ b/Engine/Source/Runtime/D3D12RHI/Private/D3D12Util.cpp
@@ -13,6 +13,7 @@ D3D12Util.h: D3D RHI utility implementation.
 #if PLATFORM_WINDOWS
 #include "HAL/PlatformCrashContext.h"
 #endif
+#include "GenericPlatform/GenericPlatformCrashContext.h"
 
 #define D3DERR(x) case x: ErrorCodeText = TEXT(#x); break;
 #define LOCTEXT_NAMESPACE "Developer.MessageLog"
