Early out when a localization target with no loc files is added.

diff --git a/Engine/Source/Runtime/Core/Private/Internationalization/TextLocalizationManager.cpp b/Engine/Source/Runtime/Core/Private/Internationalization/TextLocalizationManager.cpp
--- a/Engine/Source/Runtime/Core/Private/Internationalization/TextLocalizationManager.cpp
+++ b/Engine/Source/Runtime/Core/Private/Internationalization/TextLocalizationManager.cpp
@@ -1019,7 +1019,7 @@ void FTextLocalizationManager::RefreshResources()
 
 void FTextLocalizationManager::HandleLocalizationTargetsMounted(TArrayView<const FString> LocalizationTargetPaths)
 {
-	if (!IsInitialized())
+	if (!IsInitialized() || LocalizationTargetPaths.IsEmpty())
 	{
 		// If we've not yet loaded localization data then there's nothing to do
 		return;
