Remove numeric characters when computing hash of GPU breadcrumb event names.

This is a shortcut but it's the easiest way to remove parameters from breadcrumb event names (e.g. "Frame 1234"), which should not be hashed.

diff --git a/Engine/Source/Runtime/Core/Private/GenericPlatform/GenericPlatformCrashContext.cpp b/Engine/Source/Runtime/Core/Private/GenericPlatform/GenericPlatformCrashContext.cpp
--- a/Engine/Source/Runtime/Core/Private/GenericPlatform/GenericPlatformCrashContext.cpp
+++ b/Engine/Source/Runtime/Core/Private/GenericPlatform/GenericPlatformCrashContext.cpp
@@ -3,6 +3,7 @@
 #include "GenericPlatform/GenericPlatformCrashContext.h"
 #include "HAL/PlatformTime.h"
 #include "HAL/PlatformStackWalk.h"
+#include "Misc/Char.h"
 #include "Misc/Parse.h"
 #include "Misc/FileHelper.h"
 #include "Misc/CommandLine.h"
@@ -115,11 +116,7 @@ public:
 private:
 	void ProcessBreadcrumbNode(const FBreadcrumbNode& Node)
 	{
-		FullHash.UpdateWithString(*Node.Name, Node.Name.Len());
-		if (Node.State == EBreadcrumbState::Active)
-		{
-			ActiveHash.UpdateWithString(*Node.Name, Node.Name.Len());
-		}
+		HashNode(Node);
 
 		ProcessedBreadcrumbString.Append(FString::Printf(TEXT("{{%s},%c"), *SanitizeBreadcrumbEventName(Node.Name), Node.GetStateString()[0]));
 		if (!Node.Children.IsEmpty())
@@ -138,6 +135,16 @@ private:
 		ProcessedBreadcrumbString.AppendChar('}');
 	}
 
+	void HashNode(const FBreadcrumbNode& Node)
+	{
+		FString NameForHash = SanitizeBreadcrumbEventNameForHash(Node.Name);
+		FullHash.UpdateWithString(*NameForHash, NameForHash.Len());
+		if (Node.State == EBreadcrumbState::Active)
+		{
+			ActiveHash.UpdateWithString(*NameForHash, NameForHash.Len());
+		}
+	}
+
 	// Sanitize the event name string to remove characters that are used
 	// as delimiters for parsing.
 	static FString SanitizeBreadcrumbEventName(const FString& EventName)
@@ -145,6 +152,22 @@ private:
 		return EventName.Replace(TEXT("{"), TEXT("(")).Replace(TEXT("}"), TEXT(")"));
 	}
 
+	// Event names include parameters, mostly numeric (e.g. "Frame 1234"), that should
+	// be ignored when computing the hash.
+	static FString SanitizeBreadcrumbEventNameForHash(const FString& EventName)
+	{
+		FString SanitizedName;
+		SanitizedName.Reserve(EventName.Len());
+		for (const TCHAR& Char : EventName)
+		{
+			if (!FChar::IsDigit(Char))
+			{
+				SanitizedName.AppendChar(Char);
+			}
+		}
+		return SanitizedName;
+	}
+
 	FString ProcessedBreadcrumbString;	
 
 	FSHA1 FullHash;
