[ConcurrentLinearAllocator] Fix third place where 64k was passed in as an alignment for non fastpath platforms.

#tests editor
#rnx

diff --git a/Engine/Source/Runtime/Core/Public/Experimental/ConcurrentLinearAllocator.h b/Engine/Source/Runtime/Core/Public/Experimental/ConcurrentLinearAllocator.h
--- a/Engine/Source/Runtime/Core/Public/Experimental/ConcurrentLinearAllocator.h
+++ b/Engine/Source/Runtime/Core/Public/Experimental/ConcurrentLinearAllocator.h
@@ -405,7 +405,7 @@ public:
 				//support for oversized Blocks
 				if (HeaderSize + Size + Alignment > BlockAllocationTag::BlockSize)
 				{
-					FBlockHeader* LargeHeader = new (BlockAllocationTag::Allocator::Malloc(HeaderSize + Size + Alignment, BlockAllocationTag::BlockSize)) FBlockHeader;
+					FBlockHeader* LargeHeader = new (BlockAllocationTag::Allocator::Malloc(HeaderSize + Size + Alignment, alignof(FBlockHeader))) FBlockHeader;
 					MemoryTrace_MarkAllocAsHeap(uint64(LargeHeader), EMemoryTraceRootHeap::SystemMemory);
 					checkSlow(IsAligned(LargeHeader, alignof(FBlockHeader)));
 
