[UnitConversion] Fix DataSize quantization factor to 1024 instead of 1000. The data size units are treated as KiB, MiB, GiB etc

diff --git a/Engine/Source/Runtime/Core/Private/Math/UnitConversion.cpp b/Engine/Source/Runtime/Core/Private/Math/UnitConversion.cpp
--- a/Engine/Source/Runtime/Core/Private/Math/UnitConversion.cpp
+++ b/Engine/Source/Runtime/Core/Private/Math/UnitConversion.cpp
@@ -838,10 +838,10 @@ namespace UnitConversion
 				Frequency.Emplace(EUnit::Megahertz,	1000.0f);
 				Frequency.Emplace(EUnit::Gigahertz,	0.0f);
 
-				DataSize.Emplace(EUnit::Bytes,		1000.0f);
-				DataSize.Emplace(EUnit::Kilobytes,	1000.0f);
-				DataSize.Emplace(EUnit::Megabytes,	1000.0f);
-				DataSize.Emplace(EUnit::Gigabytes,	1000.0f);
+				DataSize.Emplace(EUnit::Bytes,		1024.0f);
+				DataSize.Emplace(EUnit::Kilobytes,	1024.0f);
+				DataSize.Emplace(EUnit::Megabytes,	1024.0f);
+				DataSize.Emplace(EUnit::Gigabytes,	1024.0f);
 				DataSize.Emplace(EUnit::Terabytes,	0.0f);
 
 				Time.Emplace(EUnit::Milliseconds,		1000.0f);
