- Quick fix for the missing NDAd preview platform issue on Mac/Linux

diff --git a/Engine/Source/Runtime/Core/Private/Misc/DataDrivenPlatformInfoRegistry.cpp b/Engine/Source/Runtime/Core/Private/Misc/DataDrivenPlatformInfoRegistry.cpp
--- a/Engine/Source/Runtime/Core/Private/Misc/DataDrivenPlatformInfoRegistry.cpp
+++ b/Engine/Source/Runtime/Core/Private/Misc/DataDrivenPlatformInfoRegistry.cpp
@@ -276,6 +276,11 @@ static void ParsePreviewPlatforms(const FConfigFile& IniFile)
 
 				FName PlatformName = *GetSectionString(Section.Value, FName("PlatformName"));
 				checkf(PlatformName != NAME_None, TEXT("DataDrivenPlatformInfo section [%s] must specify a PlatformName"), *SectionName);
+				
+				if (!FDataDrivenPlatformInfoRegistry::GetPlatformInfo(PlatformName).bEnabledForUse)
+				{
+					continue;
+				}
 
 				FPreviewPlatformMenuItem Item;
 				Item.PlatformName = PlatformName;
