Trace log messages for ensure and checks

Ensure and check messages are not emitted to the log using the regular log macros which means we do not intercept them in log tracing. Insert explicit log tracing in the StaticFailDebugV function.

diff --git a/Engine/Source/Runtime/Core/Private/Misc/AssertionMacros.cpp b/Engine/Source/Runtime/Core/Private/Misc/AssertionMacros.cpp
--- a/Engine/Source/Runtime/Core/Private/Misc/AssertionMacros.cpp
+++ b/Engine/Source/Runtime/Core/Private/Misc/AssertionMacros.cpp
@@ -472,6 +472,18 @@ FORCENOINLINE void FDebug::EnsureFailed(const ANSICHAR* Expr, const ANSICHAR* Fi
 #endif
 			GLog->Flush();
 
+			// Trace the error
+#if !NO_LOGGING
+			if (GEnsuresAreErrors)
+			{
+				TRACE_LOG_MESSAGE(LogOutputDevice, Error, TEXT("%s"), ErrorMsg);
+			}
+			else
+			{
+				TRACE_LOG_MESSAGE(LogOutputDevice, Warning, TEXT("%s"), ErrorMsg);
+			}
+#endif
+			
 			// Submit the error report to the server! (and display a balloon in the system tray)
 			{
 				// How many unique previous errors we should keep track of
