Allow traced log messages to be longer than 1500 characters.

Adds a variant of the formatting function that allows for overflowing the stack allocated buffer when needed.

diff --git a/Engine/Source/Runtime/Core/Private/Logging/LogTrace.cpp b/Engine/Source/Runtime/Core/Private/Logging/LogTrace.cpp
--- a/Engine/Source/Runtime/Core/Private/Logging/LogTrace.cpp
+++ b/Engine/Source/Runtime/Core/Private/Logging/LogTrace.cpp
@@ -56,7 +56,7 @@ void FLogTrace::OutputLogMessageSpec(const void* LogPoint, const FLogCategoryBas
 		<< LogMessageSpec.FormatString(Format, FormatStringLen);
 }
 
-void FLogTrace::OutputLogMessageInternal(const void* LogPoint, uint16 EncodedFormatArgsSize, uint8* EncodedFormatArgs)
+void FLogTrace::OutputLogMessageInternal(const void* LogPoint, int32 EncodedFormatArgsSize, const uint8* EncodedFormatArgs)
 {
 	UE_TRACE_LOG(Logging, LogMessage, LogChannel)
 		<< LogMessage.LogPoint(LogPoint)
diff --git a/Engine/Source/Runtime/Core/Public/Logging/LogTrace.h b/Engine/Source/Runtime/Core/Public/Logging/LogTrace.h
--- a/Engine/Source/Runtime/Core/Public/Logging/LogTrace.h
+++ b/Engine/Source/Runtime/Core/Public/Logging/LogTrace.h
@@ -40,16 +40,18 @@ struct FLogTrace
 	template <typename... Types>
 	FORCENOINLINE static void OutputLogMessage(const void* LogPoint, Types... FormatArgs)
 	{
-		uint8 FormatArgsBuffer[3072];
-		uint16 FormatArgsSize = FFormatArgsTrace::EncodeArguments(FormatArgsBuffer, FormatArgs...);
-		if (FormatArgsSize)
+		// This inline size covers 99.96% of log messages in sampled traces
+		FFormatArgsTrace::TBufferWithOverflow<1600> Buffer;
+		const int32 FormatArgsSize = FFormatArgsTrace::EncodeArgumentsWithArray(Buffer, FormatArgs...);
+		const uint8* FormatArgsBuffer = Buffer.GetData();
+		if (FormatArgsSize && FormatArgsSize < INT32_MAX)
 		{
 			OutputLogMessageInternal(LogPoint, FormatArgsSize, FormatArgsBuffer);
 		}
 	}
 
 private:
-	CORE_API static void OutputLogMessageInternal(const void* LogPoint, uint16 EncodedFormatArgsSize, uint8* EncodedFormatArgs);
+	CORE_API static void OutputLogMessageInternal(const void* LogPoint, int32 EncodedFormatArgsSize, const uint8* EncodedFormatArgs);
 };
 
 #define TRACE_LOG_CATEGORY(Category, Name, DefaultVerbosity) \
diff --git a/Engine/Source/Runtime/Core/Public/ProfilingDebugging/FormatArgsTrace.h b/Engine/Source/Runtime/Core/Public/ProfilingDebugging/FormatArgsTrace.h
--- a/Engine/Source/Runtime/Core/Public/ProfilingDebugging/FormatArgsTrace.h
+++ b/Engine/Source/Runtime/Core/Public/ProfilingDebugging/FormatArgsTrace.h
@@ -2,11 +2,13 @@
 
 #pragma once
 
-#include "CoreTypes.h"
+#include <CoreTypes.h>
+#include "Containers/ContainerAllocationPolicies.h"
+#include "Math/UnrealMathUtility.h"
 #include "Misc/CString.h"
 #include "Templates/IsFloatingPoint.h"
-#include "Templates/UnrealTypeTraits.h"
 #include "Templates/UnrealTemplate.h"
+#include "Templates/UnrealTypeTraits.h"
 #include "Traits/IsCharType.h"
 
 struct FFormatArgsTrace
@@ -21,6 +23,28 @@ struct FFormatArgsTrace
 		FormatArgTypeCode_CategoryString = 3 << FormatArgTypeCode_CategoryBitShift,
 	};
 
+	// Convenient alias to use when allowing overflow for very large messages
+	template <int InlineSize>
+	using TBufferWithOverflow = TArray<uint8, TInlineAllocator<InlineSize>>;
+
+	// Encode arguments in any type of resizable array
+	template <typename BufferType, typename... Types>
+	static int32 EncodeArgumentsWithArray(BufferType& Array, Types... FormatArgs)
+	{
+		constexpr int32 FormatArgsCount = sizeof...(FormatArgs);
+		static_assert(FormatArgsCount < 256, "Maximum number of arguments is 256");
+		const int32 FormatArgsSize = 1 + FormatArgsCount + int32(GetArgumentsEncodedSize(FormatArgs...));
+		Array.SetNumUninitialized(FormatArgsSize);
+		uint8* Buffer = Array.GetData();
+		uint8* TypeCodesBufferPtr = Buffer;
+		*TypeCodesBufferPtr++ = uint8(FormatArgsCount);
+		uint8* PayloadBufferPtr = TypeCodesBufferPtr + FormatArgsCount;
+		EncodeArgumentsInternal(TypeCodesBufferPtr, PayloadBufferPtr, FormatArgs...);
+		check(PayloadBufferPtr - Buffer == FormatArgsSize);
+		return FormatArgsSize;
+	}
+
+	// Encode arguments in a fixed size buffer
 	template <int BufferSize, typename... Types>
 	static uint16 EncodeArguments(uint8(&Buffer)[BufferSize], Types... FormatArgs)
 	{
@@ -59,7 +83,7 @@ private:
 		{
 			if (Argument != nullptr)
 			{
-				return (TCString<CharType>::Strlen(Argument) + 1) * sizeof(CharType);
+				return FMath::Min<int32>((int32)UINT16_MAX, (TCString<CharType>::Strlen(Argument) + 1) * sizeof(CharType));
 			}
 			else
 			{
@@ -93,9 +117,16 @@ private:
 			*TypeCodesPtr++ = FormatArgTypeCode_CategoryString | sizeof(CharType);
 			if (Argument != nullptr)
 			{
-				uint16 Length = (uint16)((TCString<CharType>::Strlen(Argument) + 1) * sizeof(CharType));
-				memcpy(PayloadPtr, Argument, Length);
-				PayloadPtr += Length;
+				uint16 SizeBytes = (uint16)FMath::Min<int32>((int32)UINT16_MAX, (TCString<CharType>::Strlen(Argument) + 1) * sizeof(CharType));
+				memcpy(PayloadPtr, Argument, SizeBytes);
+				if (UNLIKELY(SizeBytes == UINT16_MAX))
+				{
+					const CharType Ellipsis[4] = {L'.', L'.', L'.', L'\0'};
+					// Patch ellipsis to the end of the string. Make sure write is aligned
+					void* PatchPoint = PayloadPtr + SizeBytes - sizeof(Ellipsis) - (SizeBytes % sizeof(CharType));
+					memcpy(PatchPoint, Ellipsis, sizeof(Ellipsis));
+				}
+				PayloadPtr += SizeBytes;
 			}
 			else
 			{
