Renamed FString to a macro which will be used to expand out to different names.

diff --git a/Engine/Source/Runtime/Core/Private/Containers/String.cpp b/Engine/Source/Runtime/Core/Private/Containers/String.cpp
--- a/Engine/Source/Runtime/Core/Private/Containers/String.cpp
+++ b/Engine/Source/Runtime/Core/Private/Containers/String.cpp
@@ -2,10 +2,12 @@
 
 #include "Containers/UnrealString.h"
 
-// Include String.cpp.inl's includes before defining (upcoming) macros, in case the macros 'poison' other headers or there are re-entrant includes.
+// Include String.cpp.inl's includes before defining the macros, in case the macros 'poison' other headers or there are re-entrant includes.
 #include "Containers/StringIncludes.cpp.inl"
 
-#include "Containers/String.cpp.inl"
+#define UE_STRING_CLASS FString
+	#include "Containers/String.cpp.inl"
+#undef UE_STRING_CLASS
 
 void FTextRange::CalculateLineRangesFromString(const FString& Input, TArray<FTextRange>& LineRanges)
 {
diff --git a/Engine/Source/Runtime/Core/Private/Containers/String.cpp.inl b/Engine/Source/Runtime/Core/Private/Containers/String.cpp.inl
--- a/Engine/Source/Runtime/Core/Private/Containers/String.cpp.inl
+++ b/Engine/Source/Runtime/Core/Private/Containers/String.cpp.inl
@@ -10,7 +10,11 @@
  * #includes needed to compile this file need to be specified in StringIncludes.cpp.inl file rather than here.       *
  *********************************************************************************************************************/
 
-/* FString implementation
+#ifndef UE_STRING_CLASS
+	#error "String.cpp.inl should only be included after defining UE_STRING_CLASS"
+#endif
+
+ /* String implementation
  *****************************************************************************/
 
 namespace UE::Core::Private
@@ -216,20 +220,20 @@ namespace UE::Core::Private
 	}
 } // namespace UE::Core::Private
 
-FString::FString(const ANSICHAR* Str)								{ UE::Core::Private::ConstructFromCString(Data, Str); }
-FString::FString(const WIDECHAR* Str)								{ UE::Core::Private::ConstructFromCString(Data, Str); }
-FString::FString(const UTF8CHAR* Str)								{ UE::Core::Private::ConstructFromCString(Data, Str); }
-FString::FString(const UCS2CHAR* Str)								{ UE::Core::Private::ConstructFromCString(Data, Str); }
-FString::FString(int32 Len, const ANSICHAR* Str)					{ UE::Core::Private::ConstructWithLength(Data, Len, Str); }
-FString::FString(int32 Len, const WIDECHAR* Str)					{ UE::Core::Private::ConstructWithLength(Data, Len, Str); }
-FString::FString(int32 Len, const UTF8CHAR* Str)					{ UE::Core::Private::ConstructWithLength(Data, Len, Str); }
-FString::FString(int32 Len, const UCS2CHAR* Str)					{ UE::Core::Private::ConstructWithLength(Data, Len, Str); }
-FString::FString(const ANSICHAR* Str, int32 ExtraSlack)				{ UE::Core::Private::ConstructWithSlack(Data, Str, ExtraSlack); }
-FString::FString(const WIDECHAR* Str, int32 ExtraSlack)				{ UE::Core::Private::ConstructWithSlack(Data, Str, ExtraSlack); }
-FString::FString(const UTF8CHAR* Str, int32 ExtraSlack)				{ UE::Core::Private::ConstructWithSlack(Data, Str, ExtraSlack); }
-FString::FString(const UCS2CHAR* Str, int32 ExtraSlack)				{ UE::Core::Private::ConstructWithSlack(Data, Str, ExtraSlack); }
-
-FString& FString::operator=( const ElementType* Other )
+UE_STRING_CLASS::UE_STRING_CLASS(const ANSICHAR* Str)								{ UE::Core::Private::ConstructFromCString(Data, Str); }
+UE_STRING_CLASS::UE_STRING_CLASS(const WIDECHAR* Str)								{ UE::Core::Private::ConstructFromCString(Data, Str); }
+UE_STRING_CLASS::UE_STRING_CLASS(const UTF8CHAR* Str)								{ UE::Core::Private::ConstructFromCString(Data, Str); }
+UE_STRING_CLASS::UE_STRING_CLASS(const UCS2CHAR* Str)								{ UE::Core::Private::ConstructFromCString(Data, Str); }
+UE_STRING_CLASS::UE_STRING_CLASS(int32 Len, const ANSICHAR* Str)					{ UE::Core::Private::ConstructWithLength(Data, Len, Str); }
+UE_STRING_CLASS::UE_STRING_CLASS(int32 Len, const WIDECHAR* Str)					{ UE::Core::Private::ConstructWithLength(Data, Len, Str); }
+UE_STRING_CLASS::UE_STRING_CLASS(int32 Len, const UTF8CHAR* Str)					{ UE::Core::Private::ConstructWithLength(Data, Len, Str); }
+UE_STRING_CLASS::UE_STRING_CLASS(int32 Len, const UCS2CHAR* Str)					{ UE::Core::Private::ConstructWithLength(Data, Len, Str); }
+UE_STRING_CLASS::UE_STRING_CLASS(const ANSICHAR* Str, int32 ExtraSlack)				{ UE::Core::Private::ConstructWithSlack(Data, Str, ExtraSlack); }
+UE_STRING_CLASS::UE_STRING_CLASS(const WIDECHAR* Str, int32 ExtraSlack)				{ UE::Core::Private::ConstructWithSlack(Data, Str, ExtraSlack); }
+UE_STRING_CLASS::UE_STRING_CLASS(const UTF8CHAR* Str, int32 ExtraSlack)				{ UE::Core::Private::ConstructWithSlack(Data, Str, ExtraSlack); }
+UE_STRING_CLASS::UE_STRING_CLASS(const UCS2CHAR* Str, int32 ExtraSlack)				{ UE::Core::Private::ConstructWithSlack(Data, Str, ExtraSlack); }
+
+UE_STRING_CLASS& UE_STRING_CLASS::operator=( const ElementType* Other )
 {
 	if (Data.GetData() != Other)
 	{
@@ -246,7 +250,7 @@ FString& FString::operator=( const ElementType* Other )
 }
 
 
-void FString::AssignRange(const ElementType* OtherData, int32 OtherLen)
+void UE_STRING_CLASS::AssignRange(const ElementType* OtherData, int32 OtherLen)
 {
 	if (OtherLen == 0)
 	{
@@ -274,7 +278,7 @@ void FString::AssignRange(const ElementType* OtherData, int32 OtherLen)
 	}
 }
 
-void FString::Reserve(int32 CharacterCount)
+void UE_STRING_CLASS::Reserve(int32 CharacterCount)
 {
 	checkSlow(CharacterCount >= 0 && CharacterCount < MAX_int32);
 	if (CharacterCount > 0)
@@ -283,17 +287,17 @@ void FString::Reserve(int32 CharacterCount)
 	}	
 }
 
-void FString::Empty(int32 Slack)
+void UE_STRING_CLASS::Empty(int32 Slack)
 {
 	Data.Empty(Slack ? Slack + 1 : 0);
 }
 
-void FString::Empty()
+void UE_STRING_CLASS::Empty()
 {
 	Data.Empty(0);
 }
 
-void FString::Reset(int32 NewReservedSize)
+void UE_STRING_CLASS::Reset(int32 NewReservedSize)
 {
 	const int32 NewSizeIncludingTerminator = (NewReservedSize > 0) ? (NewReservedSize + 1) : 0;
 	Data.Reset(NewSizeIncludingTerminator);
@@ -303,14 +307,14 @@ void FString::Reset(int32 NewReservedSize)
 	}
 }
 
-void FString::Shrink()
+void UE_STRING_CLASS::Shrink()
 {
 	Data.Shrink();
 }
 
 #ifdef __OBJC__
-/** Convert FString to Objective-C NSString */
-NSString* FString::GetNSString() const
+/** Convert string to Objective-C NSString */
+NSString* UE_STRING_CLASS::GetNSString() const
 {
 #if PLATFORM_TCHAR_IS_4_BYTES
     return [[[NSString alloc] initWithBytes:Data.GetData() length:Len() * sizeof(ElementType) encoding:NSUTF32LittleEndianStringEncoding] autorelease];
@@ -320,7 +324,7 @@ NSString* FString::GetNSString() const
 }
 #endif
 
-FString& FString::AppendChar(ElementType InChar)
+UE_STRING_CLASS& UE_STRING_CLASS::AppendChar(ElementType InChar)
 {
 	CheckInvariants();
 
@@ -341,31 +345,31 @@ FString& FString::AppendChar(ElementType InChar)
 	return *this;
 }
 
-void FString::AppendChars(const ANSICHAR* Str, int32 Count)
+void UE_STRING_CLASS::AppendChars(const ANSICHAR* Str, int32 Count)
 {
 	CheckInvariants();
 	UE::Core::Private::AppendCharacters(Data, Str, Count);
 }
 
-void FString::AppendChars(const WIDECHAR* Str, int32 Count)
+void UE_STRING_CLASS::AppendChars(const WIDECHAR* Str, int32 Count)
 {
 	CheckInvariants();
 	UE::Core::Private::AppendCharacters(Data, Str, Count);
 }
 
-void FString::AppendChars(const UCS2CHAR* Str, int32 Count)
+void UE_STRING_CLASS::AppendChars(const UCS2CHAR* Str, int32 Count)
 {
 	CheckInvariants();
 	UE::Core::Private::AppendCharacters(Data, Str, Count);
 }
 
-void FString::AppendChars(const UTF8CHAR* Str, int32 Count)
+void UE_STRING_CLASS::AppendChars(const UTF8CHAR* Str, int32 Count)
 {
 	CheckInvariants();
 	UE::Core::Private::AppendCharacters(Data, Str, Count);
 }
 
-void FString::TrimToNullTerminator()
+void UE_STRING_CLASS::TrimToNullTerminator()
 {
 	if( Data.Num() )
 	{
@@ -379,7 +383,7 @@ void FString::TrimToNullTerminator()
 }
 
 
-int32 FString::Find(const ElementType* SubStr, int32 SubStrLen, ESearchCase::Type SearchCase, ESearchDir::Type SearchDir, int32 StartPosition) const
+int32 UE_STRING_CLASS::Find(const ElementType* SubStr, int32 SubStrLen, ESearchCase::Type SearchCase, ESearchDir::Type SearchDir, int32 StartPosition) const
 {
 	checkf(SubStrLen >= 0, TEXT("Invalid SubStrLen: %d"), SubStrLen);
 
@@ -405,7 +409,7 @@ int32 FString::Find(const ElementType* SubStr, int32 SubStrLen, ESearchCase::Typ
 		// times in the loop below
 		if ( SearchCase == ESearchCase::IgnoreCase)
 		{
-			return ToUpper().Find(FString(SubStrLen, SubStr).ToUpper(), ESearchCase::CaseSensitive, SearchDir, StartPosition);
+			return ToUpper().Find(UE_STRING_CLASS(SubStrLen, SubStr).ToUpper(), ESearchCase::CaseSensitive, SearchDir, StartPosition);
 		}
 		else
 		{
@@ -437,7 +441,7 @@ int32 FString::Find(const ElementType* SubStr, int32 SubStrLen, ESearchCase::Typ
 	}
 }
 
-bool FString::Split(const FString& InS, FString* LeftS, FString* RightS, ESearchCase::Type SearchCase, ESearchDir::Type SearchDir) const
+bool UE_STRING_CLASS::Split(const UE_STRING_CLASS& InS, UE_STRING_CLASS* LeftS, UE_STRING_CLASS* RightS, ESearchCase::Type SearchCase, ESearchDir::Type SearchDir) const
 {
 	check(LeftS != RightS || LeftS == nullptr);
 
@@ -467,25 +471,25 @@ bool FString::Split(const FString& InS, FString* LeftS, FString* RightS, ESearch
 	return true;
 }
 
-bool FString::Split(const FString& InS, FString* LeftS, FString* RightS) const
+bool UE_STRING_CLASS::Split(const UE_STRING_CLASS& InS, UE_STRING_CLASS* LeftS, UE_STRING_CLASS* RightS) const
 {
 	return Split(InS, LeftS, RightS, ESearchCase::IgnoreCase);
 }
 
-FString FString::ToUpper() const &
+UE_STRING_CLASS UE_STRING_CLASS::ToUpper() const &
 {
-	FString New = *this;
+	UE_STRING_CLASS New = *this;
 	New.ToUpperInline();
 	return New;
 }
 
-FString FString::ToUpper() &&
+UE_STRING_CLASS UE_STRING_CLASS::ToUpper() &&
 {
 	this->ToUpperInline();
 	return MoveTemp(*this);
 }
 
-void FString::ToUpperInline()
+void UE_STRING_CLASS::ToUpperInline()
 {
 	const int32 StringLength = Len();
 	ElementType* RawData = Data.GetData();
@@ -496,20 +500,20 @@ void FString::ToUpperInline()
 }
 
 
-FString FString::ToLower() const &
+UE_STRING_CLASS UE_STRING_CLASS::ToLower() const &
 {
-	FString New = *this;
+	UE_STRING_CLASS New = *this;
 	New.ToLowerInline();
 	return New;
 }
 
-FString FString::ToLower() &&
+UE_STRING_CLASS UE_STRING_CLASS::ToLower() &&
 {
 	this->ToLowerInline();
 	return MoveTemp(*this);
 }
 
-void FString::ToLowerInline()
+void UE_STRING_CLASS::ToLowerInline()
 {
 	const int32 StringLength = Len();
 	ElementType* RawData = Data.GetData();
@@ -519,7 +523,7 @@ void FString::ToLowerInline()
 	}
 }
 
-void FString::RemoveSpacesInline()
+void UE_STRING_CLASS::RemoveSpacesInline()
 {
 	const int32 StringLength = Len();
 	if (StringLength == 0)
@@ -546,7 +550,7 @@ void FString::RemoveSpacesInline()
 	}
 }
 
-bool FString::StartsWith(const ElementType* InPrefix, int32 InPrefixLen, ESearchCase::Type SearchCase) const
+bool UE_STRING_CLASS::StartsWith(const ElementType* InPrefix, int32 InPrefixLen, ESearchCase::Type SearchCase) const
 {
 	if (SearchCase == ESearchCase::IgnoreCase)
 	{
@@ -558,7 +562,7 @@ bool FString::StartsWith(const ElementType* InPrefix, int32 InPrefixLen, ESearch
 	}
 }
 
-bool FString::EndsWith(const ElementType* InSuffix, int32 InSuffixLen, ESearchCase::Type SearchCase ) const
+bool UE_STRING_CLASS::EndsWith(const ElementType* InSuffix, int32 InSuffixLen, ESearchCase::Type SearchCase ) const
 {
 	if (SearchCase == ESearchCase::IgnoreCase)
 	{
@@ -574,7 +578,7 @@ bool FString::EndsWith(const ElementType* InSuffix, int32 InSuffixLen, ESearchCa
 	}
 }
 
-void FString::InsertAt(int32 Index, ElementType Character)
+void UE_STRING_CLASS::InsertAt(int32 Index, ElementType Character)
 {
 	if (Character != 0)
 	{
@@ -589,7 +593,7 @@ void FString::InsertAt(int32 Index, ElementType Character)
 	}
 }
 
-void FString::InsertAt(int32 Index, const FString& Characters)
+void UE_STRING_CLASS::InsertAt(int32 Index, const UE_STRING_CLASS& Characters)
 {
 	if (Characters.Len())
 	{
@@ -604,12 +608,12 @@ void FString::InsertAt(int32 Index, const FString& Characters)
 	}
 }
 
-void FString::RemoveAt(int32 Index, int32 Count, bool bAllowShrinking)
+void UE_STRING_CLASS::RemoveAt(int32 Index, int32 Count, bool bAllowShrinking)
 {
 	Data.RemoveAt(Index, FMath::Clamp(Count, 0, Len()-Index), bAllowShrinking);
 }
 
-bool FString::RemoveFromStart(const ElementType* InPrefix, int32 InPrefixLen, ESearchCase::Type SearchCase)
+bool UE_STRING_CLASS::RemoveFromStart(const ElementType* InPrefix, int32 InPrefixLen, ESearchCase::Type SearchCase)
 {
 	if (InPrefixLen == 0 )
 	{
@@ -625,7 +629,7 @@ bool FString::RemoveFromStart(const ElementType* InPrefix, int32 InPrefixLen, ES
 	return false;
 }
 
-bool FString::RemoveFromEnd(const ElementType* InSuffix, int32 InSuffixLen, ESearchCase::Type SearchCase)
+bool UE_STRING_CLASS::RemoveFromEnd(const ElementType* InSuffix, int32 InSuffixLen, ESearchCase::Type SearchCase)
 {
 	if (InSuffixLen == 0)
 	{
@@ -645,7 +649,7 @@ namespace UE::String::Private
 {
 
 template <typename LhsType, typename RhsType>
-UE_NODISCARD FORCEINLINE FString ConcatFStrings(LhsType&& Lhs, RhsType&& Rhs)
+UE_NODISCARD FORCEINLINE UE_STRING_CLASS ConcatFStrings(LhsType&& Lhs, RhsType&& Rhs)
 {
 	Lhs.CheckInvariants();
 	Rhs.CheckInvariants();
@@ -657,16 +661,16 @@ UE_NODISCARD FORCEINLINE FString ConcatFStrings(LhsType&& Lhs, RhsType&& Rhs)
 
 	int32 RhsLen = Rhs.Len();
 
-	FString Result(Forward<LhsType>(Lhs), /* extra slack */ RhsLen);
+	UE_STRING_CLASS Result(Forward<LhsType>(Lhs), /* extra slack */ RhsLen);
 	Result.AppendChars(Rhs.GetCharArray().GetData(), RhsLen);
 		
 	return Result;
 }
 
 template <typename LhsCharType, typename RhsType>
-UE_NODISCARD FORCEINLINE FString ConcatRangeFString(const LhsCharType* Lhs, int32 LhsLen, RhsType&& Rhs)
+UE_NODISCARD FORCEINLINE UE_STRING_CLASS ConcatRangeFString(const LhsCharType* Lhs, int32 LhsLen, RhsType&& Rhs)
 {
-	using ElementType = FString::ElementType;
+	using ElementType = UE_STRING_CLASS::ElementType;
 
 	checkSlow(LhsLen >= 0);
 	Rhs.CheckInvariants();
@@ -680,7 +684,7 @@ UE_NODISCARD FORCEINLINE FString ConcatRangeFString(const LhsCharType* Lhs, int3
 	// This is not entirely optimal, as if the Rhs is an rvalue and has enough slack space to hold Lhs, then
 	// the memory could be reused here without constructing a new object.  However, until there is proof otherwise,
 	// I believe this will be relatively rare and isn't worth making the code a lot more complex right now.
-	FString Result;
+	UE_STRING_CLASS Result;
 	Result.GetCharArray().Reserve(LhsLen + RhsLen + 1);
 	Result.GetCharArray().AddUninitialized(LhsLen + RhsLen + 1);
 
@@ -693,7 +697,7 @@ UE_NODISCARD FORCEINLINE FString ConcatRangeFString(const LhsCharType* Lhs, int3
 }
 
 template <typename LhsType, typename RhsCharType>
-UE_NODISCARD FORCEINLINE FString ConcatFStringRange(LhsType&& Lhs, const RhsCharType* Rhs, int32 RhsLen)
+UE_NODISCARD FORCEINLINE UE_STRING_CLASS ConcatFStringRange(LhsType&& Lhs, const RhsCharType* Rhs, int32 RhsLen)
 {
 	Lhs.CheckInvariants();
 	checkSlow(RhsLen >= 0);
@@ -702,14 +706,14 @@ UE_NODISCARD FORCEINLINE FString ConcatFStringRange(LhsType&& Lhs, const RhsChar
 		return Forward<LhsType>(Lhs);
 	}
 
-	FString Result(Forward<LhsType>(Lhs), /* extra slack */ RhsLen);
+	UE_STRING_CLASS Result(Forward<LhsType>(Lhs), /* extra slack */ RhsLen);
 	Result.AppendChars(Rhs, RhsLen);
 		
 	return Result;
 }
 
 template <typename LhsCharType, typename RhsType>
-UE_NODISCARD FORCEINLINE FString ConcatCStringFString(const LhsCharType* Lhs, RhsType&& Rhs)
+UE_NODISCARD FORCEINLINE UE_STRING_CLASS ConcatCStringFString(const LhsCharType* Lhs, RhsType&& Rhs)
 {
 	checkSlow(Lhs);
 	if (!Lhs)
@@ -721,7 +725,7 @@ UE_NODISCARD FORCEINLINE FString ConcatCStringFString(const LhsCharType* Lhs, Rh
 }
 
 template <typename LhsType, typename RhsCharType>
-UE_NODISCARD FORCEINLINE FString ConcatFStringCString(LhsType&& Lhs, const RhsCharType* Rhs)
+UE_NODISCARD FORCEINLINE UE_STRING_CLASS ConcatFStringCString(LhsType&& Lhs, const RhsCharType* Rhs)
 {
 	checkSlow(Rhs);
 	if (!Rhs)
@@ -733,18 +737,18 @@ UE_NODISCARD FORCEINLINE FString ConcatFStringCString(LhsType&& Lhs, const RhsCh
 
 } // namespace UE::String::Private
 
-FString FString::ConcatFF(const FString& Lhs, const FString& Rhs)						{ return UE::String::Private::ConcatFStrings(Lhs, Rhs); }
-FString FString::ConcatFF(FString&& Lhs, const FString& Rhs)							{ return UE::String::Private::ConcatFStrings(MoveTemp(Lhs), Rhs); }
-FString FString::ConcatFF(const FString& Lhs, FString&& Rhs)							{ return UE::String::Private::ConcatFStrings(Lhs, MoveTemp(Rhs)); }
-FString FString::ConcatFF(FString&& Lhs, FString&& Rhs)									{ return UE::String::Private::ConcatFStrings(MoveTemp(Lhs), MoveTemp(Rhs)); }
-FString FString::ConcatFC(const FString& Lhs, const ElementType* Rhs)					{ return UE::String::Private::ConcatFStringCString(Lhs, Rhs); }
-FString FString::ConcatFC(FString&& Lhs, const ElementType* Rhs)						{ return UE::String::Private::ConcatFStringCString(MoveTemp(Lhs), Rhs); }
-FString FString::ConcatCF(const ElementType* Lhs,	const FString& Rhs)					{ return UE::String::Private::ConcatCStringFString(Lhs, Rhs); }
-FString FString::ConcatCF(const ElementType* Lhs,	FString&& Rhs)						{ return UE::String::Private::ConcatCStringFString(Lhs, MoveTemp(Rhs)); }
-FString FString::ConcatFR(const FString& Lhs, const ElementType* Rhs, int32 RhsLen)		{ return UE::String::Private::ConcatFStringRange(Lhs, Rhs, RhsLen); }
-FString FString::ConcatFR(FString&& Lhs, const ElementType* Rhs, int32 RhsLen)			{ return UE::String::Private::ConcatFStringRange(MoveTemp(Lhs), Rhs, RhsLen); }
-FString FString::ConcatRF(const ElementType* Lhs, int32 LhsLen, const FString& Rhs)		{ return UE::String::Private::ConcatRangeFString(Lhs, LhsLen, Rhs); }
-FString FString::ConcatRF(const ElementType* Lhs, int32 LhsLen, FString&& Rhs)			{ return UE::String::Private::ConcatRangeFString(Lhs, LhsLen, MoveTemp(Rhs)); }
+UE_STRING_CLASS UE_STRING_CLASS::ConcatFF(const UE_STRING_CLASS& Lhs, const UE_STRING_CLASS& Rhs)				{ return UE::String::Private::ConcatFStrings(Lhs, Rhs); }
+UE_STRING_CLASS UE_STRING_CLASS::ConcatFF(UE_STRING_CLASS&& Lhs, const UE_STRING_CLASS& Rhs)					{ return UE::String::Private::ConcatFStrings(MoveTemp(Lhs), Rhs); }
+UE_STRING_CLASS UE_STRING_CLASS::ConcatFF(const UE_STRING_CLASS& Lhs, UE_STRING_CLASS&& Rhs)					{ return UE::String::Private::ConcatFStrings(Lhs, MoveTemp(Rhs)); }
+UE_STRING_CLASS UE_STRING_CLASS::ConcatFF(UE_STRING_CLASS&& Lhs, UE_STRING_CLASS&& Rhs)							{ return UE::String::Private::ConcatFStrings(MoveTemp(Lhs), MoveTemp(Rhs)); }
+UE_STRING_CLASS UE_STRING_CLASS::ConcatFC(const UE_STRING_CLASS& Lhs, const ElementType* Rhs)					{ return UE::String::Private::ConcatFStringCString(Lhs, Rhs); }
+UE_STRING_CLASS UE_STRING_CLASS::ConcatFC(UE_STRING_CLASS&& Lhs, const ElementType* Rhs)						{ return UE::String::Private::ConcatFStringCString(MoveTemp(Lhs), Rhs); }
+UE_STRING_CLASS UE_STRING_CLASS::ConcatCF(const ElementType* Lhs,	const UE_STRING_CLASS& Rhs)					{ return UE::String::Private::ConcatCStringFString(Lhs, Rhs); }
+UE_STRING_CLASS UE_STRING_CLASS::ConcatCF(const ElementType* Lhs, UE_STRING_CLASS&& Rhs)						{ return UE::String::Private::ConcatCStringFString(Lhs, MoveTemp(Rhs)); }
+UE_STRING_CLASS UE_STRING_CLASS::ConcatFR(const UE_STRING_CLASS& Lhs, const ElementType* Rhs, int32 RhsLen)		{ return UE::String::Private::ConcatFStringRange(Lhs, Rhs, RhsLen); }
+UE_STRING_CLASS UE_STRING_CLASS::ConcatFR(UE_STRING_CLASS&& Lhs, const ElementType* Rhs, int32 RhsLen)			{ return UE::String::Private::ConcatFStringRange(MoveTemp(Lhs), Rhs, RhsLen); }
+UE_STRING_CLASS UE_STRING_CLASS::ConcatRF(const ElementType* Lhs, int32 LhsLen, const UE_STRING_CLASS& Rhs)		{ return UE::String::Private::ConcatRangeFString(Lhs, LhsLen, Rhs); }
+UE_STRING_CLASS UE_STRING_CLASS::ConcatRF(const ElementType* Lhs, int32 LhsLen, UE_STRING_CLASS&& Rhs)			{ return UE::String::Private::ConcatRangeFString(Lhs, LhsLen, MoveTemp(Rhs)); }
 
 /**
  * Concatenate this path with given path ensuring the / character is used between them
@@ -752,7 +756,7 @@ FString FString::ConcatRF(const ElementType* Lhs, int32 LhsLen, FString&& Rhs)
  * @param Str       Pointer to an array of TCHARs (not necessarily null-terminated) to be concatenated onto the end of this.
  * @param StrLength Exact number of characters from Str to append.
  */
-void FString::PathAppend(const ElementType* Str, int32 StrLength)
+void UE_STRING_CLASS::PathAppend(const ElementType* Str, int32 StrLength)
 {
 	int32 DataNum = Data.Num();
 	if (StrLength == 0)
@@ -784,14 +788,14 @@ void FString::PathAppend(const ElementType* Str, int32 StrLength)
 	}
 }
 
-FString FString::RightChop(int32 Count) const &
+UE_STRING_CLASS UE_STRING_CLASS::RightChop(int32 Count) const &
 {
 	const int32 Length = Len();
 	const int32 Skip = FMath::Clamp(Count, 0, Length);
-	return FString(Length - Skip, **this + Skip);
+	return UE_STRING_CLASS(Length - Skip, **this + Skip);
 }
 
-FString FString::Mid(int32 Start, int32 Count) const &
+UE_STRING_CLASS UE_STRING_CLASS::Mid(int32 Start, int32 Count) const &
 {
 	if (Count >= 0)
 	{
@@ -799,19 +803,19 @@ FString FString::Mid(int32 Start, int32 Count) const &
 		const int32 RequestedStart = Start;
 		Start = FMath::Clamp(Start, 0, Length);
 		const int32 End = (int32)FMath::Clamp((int64)Count + RequestedStart, (int64)Start, (int64)Length);
-		return FString(End-Start, **this + Start);
+		return UE_STRING_CLASS(End-Start, **this + Start);
 	}
 
-	return FString();
+	return UE_STRING_CLASS();
 }
 
-FString FString::Mid(int32 Start, int32 Count) &&
+UE_STRING_CLASS UE_STRING_CLASS::Mid(int32 Start, int32 Count) &&
 {
 	MidInline(Start, Count, false);
 	return MoveTemp(*this);
 }
 
-void FString::ReplaceCharInlineCaseSensitive(const ElementType SearchChar, const ElementType ReplacementChar)
+void UE_STRING_CLASS::ReplaceCharInlineCaseSensitive(const ElementType SearchChar, const ElementType ReplacementChar)
 {
 	for (ElementType& Character : Data)
 	{
@@ -819,33 +823,33 @@ void FString::ReplaceCharInlineCaseSensitive(const ElementType SearchChar, const
 	}
 }
 
-void FString::ReplaceCharInlineIgnoreCase(const ElementType SearchChar, const ElementType ReplacementChar)
+void UE_STRING_CLASS::ReplaceCharInlineIgnoreCase(const ElementType SearchChar, const ElementType ReplacementChar)
 {
 	ElementType OtherCaseSearchChar = TChar<ElementType>::IsUpper(SearchChar) ? TChar<ElementType>::ToLower(SearchChar) : TChar<ElementType>::ToUpper(SearchChar);
 	ReplaceCharInlineCaseSensitive(OtherCaseSearchChar, ReplacementChar);
 	ReplaceCharInlineCaseSensitive(SearchChar, ReplacementChar);
 }
 
-void FString::TrimStartAndEndInline()
+void UE_STRING_CLASS::TrimStartAndEndInline()
 {
 	TrimEndInline();
 	TrimStartInline();
 }
 
-FString FString::TrimStartAndEnd() const &
+UE_STRING_CLASS UE_STRING_CLASS::TrimStartAndEnd() const &
 {
-	FString Result(*this);
+	UE_STRING_CLASS Result(*this);
 	Result.TrimStartAndEndInline();
 	return Result;
 }
 
-FString FString::TrimStartAndEnd() &&
+UE_STRING_CLASS UE_STRING_CLASS::TrimStartAndEnd() &&
 {
 	TrimStartAndEndInline();
 	return MoveTemp(*this);
 }
 
-void FString::TrimStartInline()
+void UE_STRING_CLASS::TrimStartInline()
 {
 	int32 Pos = 0;
 	while(Pos < Len() && TChar<ElementType>::IsWhitespace((*this)[Pos]))
@@ -855,20 +859,20 @@ void FString::TrimStartInline()
 	RemoveAt(0, Pos);
 }
 
-FString FString::TrimStart() const &
+UE_STRING_CLASS UE_STRING_CLASS::TrimStart() const &
 {
-	FString Result(*this);
+	UE_STRING_CLASS Result(*this);
 	Result.TrimStartInline();
 	return Result;
 }
 
-FString FString::TrimStart() &&
+UE_STRING_CLASS UE_STRING_CLASS::TrimStart() &&
 {
 	TrimStartInline();
 	return MoveTemp(*this);
 }
 
-void FString::TrimEndInline()
+void UE_STRING_CLASS::TrimEndInline()
 {
 	int32 End = Len();
 	while(End > 0 && TChar<ElementType>::IsWhitespace((*this)[End - 1]))
@@ -878,20 +882,20 @@ void FString::TrimEndInline()
 	RemoveAt(End, Len() - End);
 }
 
-FString FString::TrimEnd() const &
+UE_STRING_CLASS UE_STRING_CLASS::TrimEnd() const &
 {
-	FString Result(*this);
+	UE_STRING_CLASS Result(*this);
 	Result.TrimEndInline();
 	return Result;
 }
 
-FString FString::TrimEnd() &&
+UE_STRING_CLASS UE_STRING_CLASS::TrimEnd() &&
 {
 	TrimEndInline();
 	return MoveTemp(*this);
 }
 
-void FString::TrimCharInline(const ElementType CharacterToTrim, bool* bCharRemoved)
+void UE_STRING_CLASS::TrimCharInline(const ElementType CharacterToTrim, bool* bCharRemoved)
 {
 	bool bQuotesWereRemoved=false;
 	int32 Start = 0, Count = Len();
@@ -918,59 +922,59 @@ void FString::TrimCharInline(const ElementType CharacterToTrim, bool* bCharRemov
 	MidInline(Start, Count, false);
 }
 
-void FString::TrimQuotesInline(bool* bQuotesRemoved)
+void UE_STRING_CLASS::TrimQuotesInline(bool* bQuotesRemoved)
 {
 	TrimCharInline(ElementType('"'), bQuotesRemoved);
 }
 
-FString FString::TrimQuotes(bool* bQuotesRemoved) const &
+UE_STRING_CLASS UE_STRING_CLASS::TrimQuotes(bool* bQuotesRemoved) const &
 {
-	FString Result(*this);
+	UE_STRING_CLASS Result(*this);
 	Result.TrimQuotesInline(bQuotesRemoved);
 	return Result;
 }
 
-FString FString::TrimQuotes(bool* bQuotesRemoved) &&
+UE_STRING_CLASS UE_STRING_CLASS::TrimQuotes(bool* bQuotesRemoved) &&
 {
 	TrimQuotesInline(bQuotesRemoved);
 	return MoveTemp(*this);
 }
 
-FString FString::TrimChar(const ElementType CharacterToTrim, bool* bCharRemoved) const &
+UE_STRING_CLASS UE_STRING_CLASS::TrimChar(const ElementType CharacterToTrim, bool* bCharRemoved) const &
 {
-	FString Result(*this);
+	UE_STRING_CLASS Result(*this);
 	Result.TrimCharInline(CharacterToTrim, bCharRemoved);
 	return Result;
 }
 
-FString FString::TrimChar(const ElementType CharacterToTrim, bool* bCharRemoved) &&
+UE_STRING_CLASS UE_STRING_CLASS::TrimChar(const ElementType CharacterToTrim, bool* bCharRemoved) &&
 {
 	TrimCharInline(CharacterToTrim, bCharRemoved);
 	return MoveTemp(*this);
 }
 
-int32 FString::CullArray( TArray<FString>* InArray )
+int32 UE_STRING_CLASS::CullArray( TArray<UE_STRING_CLASS>* InArray )
 {
 	check(InArray);
-	FString Empty;
+	UE_STRING_CLASS Empty;
 	InArray->Remove(Empty);
 	return InArray->Num();
 }
 
-FString FString::Reverse() const &
+UE_STRING_CLASS UE_STRING_CLASS::Reverse() const &
 {
-	FString New(*this);
+	UE_STRING_CLASS New(*this);
 	New.ReverseString();
 	return New;
 }
 
-FString FString::Reverse() &&
+UE_STRING_CLASS UE_STRING_CLASS::Reverse() &&
 {
 	ReverseString();
 	return MoveTemp(*this);
 }
 
-void FString::ReverseString()
+void UE_STRING_CLASS::ReverseString()
 {
 	if ( Len() > 0 )
 	{
@@ -990,9 +994,9 @@ void FString::ReverseString()
 	}
 }
 
-FString FString::FormatAsNumber( int32 InNumber )
+UE_STRING_CLASS UE_STRING_CLASS::FormatAsNumber( int32 InNumber )
 {
-	FString Number = FString::FromInt( InNumber ), Result;
+	UE_STRING_CLASS Number = UE_STRING_CLASS::FromInt( InNumber ), Result;
 
 	int32 dec = 0;
 	for( int32 x = Number.Len()-1 ; x > -1 ; --x )
@@ -1017,7 +1021,7 @@ FString FString::FormatAsNumber( int32 InNumber )
  * @param	Ar				Archive to serialize with
  * @param	MinCharacters	Minimum number of characters to serialize.
  */
-void FString::SerializeAsANSICharArray( FArchive& Ar, int32 MinCharacters ) const
+void UE_STRING_CLASS::SerializeAsANSICharArray( FArchive& Ar, int32 MinCharacters ) const
 {
 	int32	Length = FMath::Max( Len(), MinCharacters );
 	Ar << Length;
@@ -1036,7 +1040,7 @@ void FString::SerializeAsANSICharArray( FArchive& Ar, int32 MinCharacters ) cons
 	}
 }
 
-void FString::AppendInt( int32 Num )
+void UE_STRING_CLASS::AppendInt( int32 Num )
 {
 	const ElementType* DigitToChar	= CHARTEXT(ElementType, "9876543210123456789");
 	constexpr int32 ZeroDigitIndex	= 9;
@@ -1063,24 +1067,24 @@ void FString::AppendInt( int32 Num )
 }
 
 
-bool FString::ToBool() const
+bool UE_STRING_CLASS::ToBool() const
 {
 	return TCString<ElementType>::ToBool(**this);
 }
 
-FString FString::FromBlob(const uint8* SrcBuffer,const uint32 SrcSize)
+UE_STRING_CLASS UE_STRING_CLASS::FromBlob(const uint8* SrcBuffer,const uint32 SrcSize)
 {
-	FString Result;
+	UE_STRING_CLASS Result;
 	Result.Reserve( SrcSize * 3 );
 	// Convert and append each byte in the buffer
 	for (uint32 Count = 0; Count < SrcSize; Count++)
 	{
-		Result += FString::Printf(CHARTEXT(ElementType, "%03d"),(uint8)SrcBuffer[Count]);
+		Result += UE_STRING_CLASS::Printf(CHARTEXT(ElementType, "%03d"),(uint8)SrcBuffer[Count]);
 	}
 	return Result;
 }
 
-bool FString::ToBlob(const FString& Source,uint8* DestBuffer,const uint32 DestSize)
+bool UE_STRING_CLASS::ToBlob(const UE_STRING_CLASS& Source,uint8* DestBuffer,const uint32 DestSize)
 {
 	// Make sure the buffer is at least half the size and that the string is an
 	// even number of characters long
@@ -1103,19 +1107,19 @@ bool FString::ToBlob(const FString& Source,uint8* DestBuffer,const uint32 DestSi
 	return false;
 }
 
-FString FString::FromHexBlob( const uint8* SrcBuffer, const uint32 SrcSize )
+UE_STRING_CLASS UE_STRING_CLASS::FromHexBlob( const uint8* SrcBuffer, const uint32 SrcSize )
 {
-	FString Result;
+	UE_STRING_CLASS Result;
 	Result.Reserve( SrcSize * 2 );
 	// Convert and append each byte in the buffer
 	for (uint32 Count = 0; Count < SrcSize; Count++)
 	{
-		Result += FString::Printf( CHARTEXT(ElementType,  "%02X" ), (uint8)SrcBuffer[Count] );
+		Result += UE_STRING_CLASS::Printf( CHARTEXT(ElementType,  "%02X" ), (uint8)SrcBuffer[Count] );
 	}
 	return Result;
 }
 
-bool FString::ToHexBlob( const FString& Source, uint8* DestBuffer, const uint32 DestSize )
+bool UE_STRING_CLASS::ToHexBlob( const UE_STRING_CLASS& Source, uint8* DestBuffer, const uint32 DestSize )
 {
 	// Make sure the buffer is at least half the size and that the string is an
 	// even number of characters long
@@ -1148,13 +1152,13 @@ void StripNegativeZero(double& InFloat)
 }
 UE_ENABLE_OPTIMIZATION_SHIP
 
-FString FString::SanitizeFloat( double InFloat, const int32 InMinFractionalDigits )
+UE_STRING_CLASS UE_STRING_CLASS::SanitizeFloat( double InFloat, const int32 InMinFractionalDigits )
 {
 	// Avoids negative zero
 	StripNegativeZero(InFloat);
 
 	// First create the string
-	FString TempString = FString::Printf(CHARTEXT(ElementType, "%f"), InFloat);
+	UE_STRING_CLASS TempString = UE_STRING_CLASS::Printf(CHARTEXT(ElementType, "%f"), InFloat);
 	if (!TempString.IsNumeric())
 	{
 		// String did not format as a valid decimal number so avoid messing with it
@@ -1205,18 +1209,18 @@ FString FString::SanitizeFloat( double InFloat, const int32 InMinFractionalDigit
 	return TempString;
 }
 
-FString FString::Chr(ElementType Ch)
+UE_STRING_CLASS UE_STRING_CLASS::Chr(ElementType Ch)
 {
 	ElementType Temp[2]= { Ch, CHARTEXT(ElementType, '\0') };
-	return FString(Temp);
+	return UE_STRING_CLASS(Temp);
 }
 
 
-FString FString::ChrN( int32 NumCharacters, ElementType Char )
+UE_STRING_CLASS UE_STRING_CLASS::ChrN( int32 NumCharacters, ElementType Char )
 {
 	check( NumCharacters >= 0 );
 
-	FString Temp;
+	UE_STRING_CLASS Temp;
 	Temp.Data.AddUninitialized(NumCharacters+1);
 	for( int32 Cx = 0; Cx < NumCharacters; ++Cx )
 	{
@@ -1226,7 +1230,7 @@ FString FString::ChrN( int32 NumCharacters, ElementType Char )
 	return Temp;
 }
 
-FString FString::LeftPad( int32 ChCount ) const
+UE_STRING_CLASS UE_STRING_CLASS::LeftPad( int32 ChCount ) const
 {
 	int32 Pad = ChCount - Len();
 
@@ -1239,7 +1243,7 @@ FString FString::LeftPad( int32 ChCount ) const
 		return *this;
 	}
 }
-FString FString::RightPad( int32 ChCount ) const
+UE_STRING_CLASS UE_STRING_CLASS::RightPad( int32 ChCount ) const
 {
 	int32 Pad = ChCount - Len();
 
@@ -1253,7 +1257,7 @@ FString FString::RightPad( int32 ChCount ) const
 	}
 }
 
-bool FString::IsNumeric() const
+bool UE_STRING_CLASS::IsNumeric() const
 {
 	if (IsEmpty())
 	{
@@ -1263,7 +1267,7 @@ bool FString::IsNumeric() const
 	return TCString<ElementType>::IsNumeric(Data.GetData());
 }
 
-int32 FString::ParseIntoArray( TArray<FString>& OutArray, const ElementType* pchDelim, const bool InCullEmpty ) const
+int32 UE_STRING_CLASS::ParseIntoArray( TArray<UE_STRING_CLASS>& OutArray, const ElementType* pchDelim, const bool InCullEmpty ) const
 {
 	check(pchDelim);
 	OutArray.Reset();
@@ -1279,7 +1283,7 @@ int32 FString::ParseIntoArray( TArray<FString>& OutArray, const ElementType* pch
 	return OutArray.Num();
 }
 
-bool FString::MatchesWildcard(const ElementType* InWildcard, int32 InWildcardLen, ESearchCase::Type SearchCase) const
+bool UE_STRING_CLASS::MatchesWildcard(const ElementType* InWildcard, int32 InWildcardLen, ESearchCase::Type SearchCase) const
 {
 	const ElementType* Target = **this;
 	int32        TargetLength = Len();
@@ -1296,7 +1300,7 @@ bool FString::MatchesWildcard(const ElementType* InWildcard, int32 InWildcardLen
 
 
 /** Caution!! this routine is O(N^2) allocations...use it for parsing very short text or not at all */
-int32 FString::ParseIntoArrayWS( TArray<FString>& OutArray, const ElementType* pchExtraDelim, bool InCullEmpty ) const
+int32 UE_STRING_CLASS::ParseIntoArrayWS( TArray<UE_STRING_CLASS>& OutArray, const ElementType* pchExtraDelim, bool InCullEmpty ) const
 {
 	// default array of White Spaces, the last entry can be replaced with the optional pchExtraDelim string
 	// (if you want to split on white space and another character)
@@ -1320,7 +1324,7 @@ int32 FString::ParseIntoArrayWS( TArray<FString>& OutArray, const ElementType* p
 	return ParseIntoArray(OutArray, WhiteSpace, NumWhiteSpaces, InCullEmpty);
 }
 
-int32 FString::ParseIntoArrayLines(TArray<FString>& OutArray, bool InCullEmpty) const
+int32 UE_STRING_CLASS::ParseIntoArrayLines(TArray<UE_STRING_CLASS>& OutArray, bool InCullEmpty) const
 {
 	// default array of LineEndings
 	static const ElementType* LineEndings[] =
@@ -1335,7 +1339,7 @@ int32 FString::ParseIntoArrayLines(TArray<FString>& OutArray, bool InCullEmpty)
 	return ParseIntoArray(OutArray, LineEndings, NumLineEndings, InCullEmpty);
 }
 
-int32 FString::ParseIntoArray(TArray<FString>& OutArray, const ElementType* const * DelimArray, int32 NumDelims, bool InCullEmpty) const
+int32 UE_STRING_CLASS::ParseIntoArray(TArray<UE_STRING_CLASS>& OutArray, const ElementType* const * DelimArray, int32 NumDelims, bool InCullEmpty) const
 {
 	// Make sure the delimit string is not null or empty
 	check(DelimArray);
@@ -1398,7 +1402,7 @@ int32 FString::ParseIntoArray(TArray<FString>& OutArray, const ElementType* cons
 	return OutArray.Num();
 }
 
-FString FString::Replace(const ElementType* From, const ElementType* To, ESearchCase::Type SearchCase) const &
+UE_STRING_CLASS UE_STRING_CLASS::Replace(const ElementType* From, const ElementType* To, ESearchCase::Type SearchCase) const &
 {
 	// Previous code used to accidentally accept a nullptr replacement string - this is no longer accepted.
 	check(To);
@@ -1415,7 +1419,7 @@ FString FString::Replace(const ElementType* From, const ElementType* To, ESearch
 	int32 FromLength = TCString<ElementType>::Strlen(From);
 	int32 ToLength   = TCString<ElementType>::Strlen(To);
 
-	FString Result;
+	UE_STRING_CLASS Result;
 	while (true)
 	{
 		// look for From in the remaining string
@@ -1440,13 +1444,13 @@ FString FString::Replace(const ElementType* From, const ElementType* To, ESearch
 	return Result;
 }
 
-FString FString::Replace(const ElementType* From, const ElementType* To, ESearchCase::Type SearchCase) &&
+UE_STRING_CLASS UE_STRING_CLASS::Replace(const ElementType* From, const ElementType* To, ESearchCase::Type SearchCase) &&
 {
 	ReplaceInline(From, To, SearchCase);
 	return MoveTemp(*this);
 }
 
-int32 FString::ReplaceInline(const ElementType* SearchText, const ElementType* ReplacementText, ESearchCase::Type SearchCase)
+int32 UE_STRING_CLASS::ReplaceInline(const ElementType* SearchText, const ElementType* ReplacementText, ESearchCase::Type SearchCase)
 {
 	int32 ReplacementCount = 0;
 
@@ -1482,7 +1486,7 @@ int32 FString::ReplaceInline(const ElementType* SearchText, const ElementType* R
 		}
 		else if (Contains(SearchText, SearchCase))
 		{
-			FString Copy(MoveTemp(*this));
+			UE_STRING_CLASS Copy(MoveTemp(*this));
 
 			// get a pointer into the character data
 			ElementType* WritePosition = (ElementType*)Copy.Data.GetData();
@@ -1492,7 +1496,7 @@ int32 FString::ReplaceInline(const ElementType* SearchText, const ElementType* R
 			{
 				ReplacementCount++;
 
-				// replace the first letter of the From with 0 so we can do a strcpy (FString +=)
+				// replace the first letter of the From with 0 so we can do a strcpy (via operator+=)
 				*SearchPosition = CHARTEXT(ElementType, '\0');
 
 				// copy everything up to the SearchPosition
@@ -1520,11 +1524,11 @@ int32 FString::ReplaceInline(const ElementType* SearchText, const ElementType* R
 /**
  * Returns a copy of this string with all quote marks escaped (unless the quote is already escaped)
  */
-FString FString::ReplaceQuotesWithEscapedQuotes() &&
+UE_STRING_CLASS UE_STRING_CLASS::ReplaceQuotesWithEscapedQuotes() &&
 {
 	if (Contains(CHARTEXT(ElementType, "\""), ESearchCase::CaseSensitive))
 	{
-		FString Copy(MoveTemp(*this));
+		UE_STRING_CLASS Copy(MoveTemp(*this));
 
 		const ElementType* pChar = *Copy;
 
@@ -1566,7 +1570,7 @@ static const CharType* CharToEscapeSeqMap[6][2] =
 template <typename CharType>
 static const uint32 MaxSupportedEscapeChars = UE_ARRAY_COUNT(CharToEscapeSeqMap<CharType>);
 
-void FString::ReplaceCharWithEscapedCharInline(const TArray<ElementType>* Chars/*=nullptr*/)
+void UE_STRING_CLASS::ReplaceCharWithEscapedCharInline(const TArray<ElementType>* Chars/*=nullptr*/)
 {
 	if ( Len() > 0 && (Chars == nullptr || Chars->Num() > 0) )
 	{
@@ -1581,7 +1585,7 @@ void FString::ReplaceCharWithEscapedCharInline(const TArray<ElementType>* Chars/
 	}
 }
 
-void FString::ReplaceEscapedCharWithCharInline(const TArray<ElementType>* Chars/*=nullptr*/)
+void UE_STRING_CLASS::ReplaceEscapedCharWithCharInline(const TArray<ElementType>* Chars/*=nullptr*/)
 {
 	if ( Len() > 0 && (Chars == nullptr || Chars->Num() > 0) )
 	{
@@ -1603,7 +1607,7 @@ void FString::ReplaceEscapedCharWithCharInline(const TArray<ElementType>* Chars/
  * Replaces all instances of '\t' with TabWidth number of spaces
  * @param InSpacesPerTab - Number of spaces that a tab represents
  */
-void FString::ConvertTabsToSpacesInline(const int32 InSpacesPerTab)
+void UE_STRING_CLASS::ConvertTabsToSpacesInline(const int32 InSpacesPerTab)
 {
 	//must call this with at least 1 space so the modulus operation works
 	check(InSpacesPerTab > 0);
@@ -1611,7 +1615,7 @@ void FString::ConvertTabsToSpacesInline(const int32 InSpacesPerTab)
 	int32 TabIndex = 0;
 	while ((TabIndex = Find(CHARTEXT(ElementType, "\t"), ESearchCase::CaseSensitive)) != INDEX_NONE )
 	{
-		FString RightSide = Mid(TabIndex+1);
+		UE_STRING_CLASS RightSide = Mid(TabIndex+1);
 		LeftInline(TabIndex, false);
 
 		//for a tab size of 4, 
@@ -1634,7 +1638,7 @@ void FString::ConvertTabsToSpacesInline(const int32 InSpacesPerTab)
 // This starting size catches 99.97% of printf calls - there are about 700k printf calls per level
 #define STARTING_BUFFER_SIZE		512
 
-FString FString::PrintfImpl(const ElementType* Fmt, ...)
+UE_STRING_CLASS UE_STRING_CLASS::PrintfImpl(const ElementType* Fmt, ...)
 {
 	int32		BufferSize	= STARTING_BUFFER_SIZE;
 	ElementType	StartingBuffer[STARTING_BUFFER_SIZE];
@@ -1658,7 +1662,7 @@ FString FString::PrintfImpl(const ElementType* Fmt, ...)
 
 	Buffer[Result] = CHARTEXT(ElementType, '\0');
 
-	FString ResultString(Buffer);
+	UE_STRING_CLASS ResultString(Buffer);
 
 	if( BufferSize != STARTING_BUFFER_SIZE )
 	{
@@ -1668,7 +1672,7 @@ FString FString::PrintfImpl(const ElementType* Fmt, ...)
 	return ResultString;
 }
 
-void FString::AppendfImpl(FString& AppendToMe, const ElementType* Fmt, ...)
+void UE_STRING_CLASS::AppendfImpl(UE_STRING_CLASS& AppendToMe, const ElementType* Fmt, ...)
 {
 	int32		BufferSize = STARTING_BUFFER_SIZE;
 	ElementType	StartingBuffer[STARTING_BUFFER_SIZE];
@@ -1700,11 +1704,11 @@ void FString::AppendfImpl(FString& AppendToMe, const ElementType* Fmt, ...)
 	}
 }
 
-static_assert(PLATFORM_LITTLE_ENDIAN, "FString serialization needs updating to support big-endian platforms!");
+static_assert(PLATFORM_LITTLE_ENDIAN, PREPROCESSOR_TO_STRING(UE_STRING_CLASS) " serialization needs updating to support big-endian platforms!");
 
-FArchive& operator<<( FArchive& Ar, FString& A )
+FArchive& operator<<( FArchive& Ar, UE_STRING_CLASS& A )
 {
-	using ElementType = FString::ElementType;
+	using ElementType = UE_STRING_CLASS::ElementType;
 
 	if constexpr (std::is_same_v<ElementType, TCHAR>)
 	{
@@ -1842,14 +1846,14 @@ FArchive& operator<<( FArchive& Ar, FString& A )
 	return Ar;
 }
 
-int32 HexToBytes(const FString& HexString, uint8* OutBytes)
+int32 HexToBytes(const UE_STRING_CLASS& HexString, uint8* OutBytes)
 {
 	return UE::String::HexToBytes(HexString, OutBytes);
 }
 
-int32 FindMatchingClosingParenthesis(const FString& TargetString, const int32 StartSearch)
+int32 FindMatchingClosingParenthesis(const UE_STRING_CLASS& TargetString, const int32 StartSearch)
 {
-	using ElementType = FString::ElementType;
+	using ElementType = UE_STRING_CLASS::ElementType;
 
 	check(StartSearch >= 0 && StartSearch <= TargetString.Len());// Check for usage, we do not accept INDEX_NONE like other string functions
 
@@ -1892,11 +1896,11 @@ int32 FindMatchingClosingParenthesis(const FString& TargetString, const int32 St
 	return INDEX_NONE;
 }
 
-FString SlugStringForValidName(const FString& DisplayString, const FString::ElementType* ReplaceWith /*= CHARTEXT(ElementType, "")*/)
+UE_STRING_CLASS SlugStringForValidName(const UE_STRING_CLASS& DisplayString, const UE_STRING_CLASS::ElementType* ReplaceWith /*= CHARTEXT(ElementType, "")*/)
 {
-	using ElementType = FString::ElementType;
+	using ElementType = UE_STRING_CLASS::ElementType;
 
-	FString GeneratedName = DisplayString;
+	UE_STRING_CLASS GeneratedName = DisplayString;
 
 	// Convert the display label, which may consist of just about any possible character, into a
 	// suitable name for a UObject (remove whitespace, certain symbols, etc.)
diff --git a/Engine/Source/Runtime/Core/Public/Containers/UnrealString.h b/Engine/Source/Runtime/Core/Public/Containers/UnrealString.h
--- a/Engine/Source/Runtime/Core/Public/Containers/UnrealString.h
+++ b/Engine/Source/Runtime/Core/Public/Containers/UnrealString.h
@@ -5,10 +5,12 @@
 
 #pragma once
 
-// Include UnrealString.h.inl's includes before defining (upcoming) macros, in case the macros 'poison' other headers or there are re-entrant includes.
+// Include UnrealString.h.inl's includes before defining the macros, in case the macros 'poison' other headers or there are re-entrant includes.
 #include "Containers/UnrealStringIncludes.h.inl"
 
-#include "Containers/UnrealString.h.inl"
+#define UE_STRING_CLASS FString
+	#include "Containers/UnrealString.h.inl"
+#undef UE_STRING_CLASS
 
 /**
  * Convert an array of bytes to a string
diff --git a/Engine/Source/Runtime/Core/Public/Containers/UnrealString.h.inl b/Engine/Source/Runtime/Core/Public/Containers/UnrealString.h.inl
--- a/Engine/Source/Runtime/Core/Public/Containers/UnrealString.h.inl
+++ b/Engine/Source/Runtime/Core/Public/Containers/UnrealString.h.inl
@@ -10,15 +10,15 @@
  * #includes needed to compile this file need to be specified in UnrealStringIncludes.h.inl file rather than here. *
  *******************************************************************************************************************/
 
-struct FStringFormatArg;
+struct PREPROCESSOR_JOIN(UE_STRING_CLASS, FormatArg);
 template<typename InKeyType,typename InValueType,typename SetAllocator ,typename KeyFuncs > class TMap;
 
-typedef TMap<FString, FStringFormatArg> FStringFormatNamedArguments;
-typedef TArray<FStringFormatArg> FStringFormatOrderedArguments;
+typedef TMap<UE_STRING_CLASS, PREPROCESSOR_JOIN(UE_STRING_CLASS, FormatArg)> PREPROCESSOR_JOIN(UE_STRING_CLASS, FormatNamedArguments);
+typedef TArray< PREPROCESSOR_JOIN(UE_STRING_CLASS, FormatArg)> PREPROCESSOR_JOIN(UE_STRING_CLASS, FormatOrderedArguments);
 
-TCHAR*       GetData(FString&);
-const TCHAR* GetData(const FString&);
-int32        GetNum(const FString& String);
+TCHAR*       GetData(UE_STRING_CLASS&);
+const TCHAR* GetData(const UE_STRING_CLASS&);
+int32        GetNum(const UE_STRING_CLASS& String);
 /**
  * A dynamically sizeable string.
  * @see https://docs.unrealengine.com/latest/INT/Programming/UnrealArchitecture/StringHandling/FString/
@@ -30,7 +30,7 @@ int32        GetNum(const FString& String);
  * - Use UTF8TEXTVIEW("...") for string view literals (type is FUtf8StringView).
  * - Use \uxxxx or \Uxxxxxxxx escape sequences rather than \x to specify Unicode code points.
  */
-class FString
+class UE_STRING_CLASS
 {
 public:
 	using AllocatorType = TSizedDefaultAllocator<32>;
@@ -66,7 +66,7 @@ private:
 	template <typename CharRangeType>
 	using TIsCharRangeNotCArrayNotStringClass = TAnd<
 		TIsCharRangeNotCArray<CharRangeType>,
-		TNot<TIsDerivedFrom<typename TDecay<CharRangeType>::Type, FString>>>;
+		TNot<TIsDerivedFrom<typename TDecay<CharRangeType>::Type, UE_STRING_CLASS>>>;
 
 	/** Trait testing whether a type is a contiguous range of ElementType, and not CharType[]. */
 	template <typename CharRangeType>
@@ -79,14 +79,14 @@ private:
 	template <typename CharRangeType>
 	using TIsTCharRangeNotCArrayNotStringClass = TAnd<
 		TIsTCharRangeNotCArray<CharRangeType>,
-		TNot<TIsDerivedFrom<typename TDecay<CharRangeType>::Type, FString>>>;
+		TNot<TIsDerivedFrom<typename TDecay<CharRangeType>::Type, UE_STRING_CLASS>>>;
 
 public:
-	FString() = default;
-	FString(FString&&) = default;
-	FString(const FString&) = default;
-	FString& operator=(FString&&) = default;
-	FString& operator=(const FString&) = default;
+	UE_STRING_CLASS() = default;
+	UE_STRING_CLASS(UE_STRING_CLASS&&) = default;
+	UE_STRING_CLASS(const UE_STRING_CLASS&) = default;
+	UE_STRING_CLASS& operator=(UE_STRING_CLASS&&) = default;
+	UE_STRING_CLASS& operator=(const UE_STRING_CLASS&) = default;
 
 	/**
 	 * Create a copy of the Other string with extra space for characters at the end of the string
@@ -94,7 +94,7 @@ public:
 	 * @param Other the other string to create a new copy from
 	 * @param ExtraSlack number of extra characters to add to the end of the other string in this string
 	 */
-	FORCEINLINE FString(const FString& Other, int32 ExtraSlack)
+	FORCEINLINE UE_STRING_CLASS(const UE_STRING_CLASS& Other, int32 ExtraSlack)
 		: Data(Other.Data, ExtraSlack + ((Other.Data.Num() || !ExtraSlack) ? 0 : 1)) // Add 1 if the source string array is empty and we want some slack, because we'll need to include a null terminator which is currently missing
 	{
 	}
@@ -105,38 +105,38 @@ public:
 	 * @param Other the other string to create a new copy from
 	 * @param ExtraSlack number of extra characters to add to the end of the other string in this string
 	 */
-	FORCEINLINE FString(FString&& Other, int32 ExtraSlack)
+	FORCEINLINE UE_STRING_CLASS(UE_STRING_CLASS&& Other, int32 ExtraSlack)
 		: Data(MoveTemp(Other.Data), ExtraSlack + ((Other.Data.Num() || !ExtraSlack) ? 0 : 1)) // Add 1 if the source string array is empty and we want some slack, because we'll need to include a null terminator which is currently missing
 	{
 	}
 
 	/** Construct from null-terminated C string or nullptr  */
-	CORE_API FString(const ANSICHAR* Str);
-	CORE_API FString(const WIDECHAR* Str);
-	CORE_API FString(const UTF8CHAR* Str);
-	CORE_API FString(const UCS2CHAR* Str);
+	CORE_API UE_STRING_CLASS(const ANSICHAR* Str);
+	CORE_API UE_STRING_CLASS(const WIDECHAR* Str);
+	CORE_API UE_STRING_CLASS(const UTF8CHAR* Str);
+	CORE_API UE_STRING_CLASS(const UCS2CHAR* Str);
 
 	/** Construct from null-terminated C substring or nullptr */
-	CORE_API FString(int32 Len, const ANSICHAR* Str);
-	CORE_API FString(int32 Len, const WIDECHAR* Str);
-	CORE_API FString(int32 Len, const UTF8CHAR* Str);
-	CORE_API FString(int32 Len, const UCS2CHAR* Str);
+	CORE_API UE_STRING_CLASS(int32 Len, const ANSICHAR* Str);
+	CORE_API UE_STRING_CLASS(int32 Len, const WIDECHAR* Str);
+	CORE_API UE_STRING_CLASS(int32 Len, const UTF8CHAR* Str);
+	CORE_API UE_STRING_CLASS(int32 Len, const UCS2CHAR* Str);
 
 	/** Construct from null-terminated C string or nullptr with extra slack on top of original string length */
-	CORE_API FString(const ANSICHAR* Str, int32 ExtraSlack);
-	CORE_API FString(const WIDECHAR* Str, int32 ExtraSlack);
-	CORE_API FString(const UTF8CHAR* Str, int32 ExtraSlack);
-	CORE_API FString(const UCS2CHAR* Str, int32 ExtraSlack);
+	CORE_API UE_STRING_CLASS(const ANSICHAR* Str, int32 ExtraSlack);
+	CORE_API UE_STRING_CLASS(const WIDECHAR* Str, int32 ExtraSlack);
+	CORE_API UE_STRING_CLASS(const UTF8CHAR* Str, int32 ExtraSlack);
+	CORE_API UE_STRING_CLASS(const UCS2CHAR* Str, int32 ExtraSlack);
 
 	/** Construct from contiguous range of characters such as a string view or string builder */
 	template <typename CharRangeType, typename TEnableIf<TIsCharRangeNotCArrayNotStringClass<CharRangeType>::Value>::Type* = nullptr>
-	FORCEINLINE explicit FString(CharRangeType&& Str) : FString(GetNum(Str), GetData(Forward<CharRangeType>(Str)))
+	FORCEINLINE explicit UE_STRING_CLASS(CharRangeType&& Str) : UE_STRING_CLASS(GetNum(Str), GetData(Forward<CharRangeType>(Str)))
 	{
 	}
 
 	/** Construct from contiguous range of characters with extra slack on top of original string length */
 	template <typename CharRangeType, typename TEnableIf<TIsCharRangeNotCArrayNotStringClass<CharRangeType>::Value>::Type* = nullptr>
-	explicit FString(CharRangeType&& Str, int32 ExtraSlack)
+	explicit UE_STRING_CLASS(CharRangeType&& Str, int32 ExtraSlack)
 	{
 		uint32 InLen = GetNum(Str);
 		Reserve(InLen + ExtraSlack);
@@ -145,8 +145,8 @@ public:
 
 public:
 #ifdef __OBJC__
-	/** Convert Objective-C NSString* to FString */
-	FORCEINLINE FString(const NSString* In)
+	/** Convert Objective-C NSString* to string class */
+	FORCEINLINE UE_STRING_CLASS(const NSString* In)
 	{
 		if (In && [In length] > 0)
 		{
@@ -172,10 +172,10 @@ public:
 	}
 #endif
 
-	CORE_API FString& operator=(const ElementType* Str);
+	CORE_API UE_STRING_CLASS& operator=(const ElementType* Str);
 
 	template <typename CharRangeType, typename TEnableIf<TIsTCharRangeNotCArray<CharRangeType>::Value>::Type* = nullptr>
-	FORCEINLINE FString& operator=(CharRangeType&& Range)
+	FORCEINLINE UE_STRING_CLASS& operator=(CharRangeType&& Range)
 	{
 		AssignRange(GetData(Range), GetNum(Range));
 		return *this;
@@ -342,7 +342,7 @@ public:
 
 	/** Append a string and return a reference to this */
 	template<class CharType>
-	FORCEINLINE FString& Append(const CharType* Str, int32 Count)
+	FORCEINLINE UE_STRING_CLASS& Append(const CharType* Str, int32 Count)
 	{
 		AppendChars(Str, Count);
 		return *this;
@@ -356,7 +356,7 @@ public:
 	 * Oddly MSVC ranks a const T* overload over T&& for T[N] while clang picks T&&.
 	 */
 	template<class CharType>
-	FORCEINLINE FString& Append(/* no const! */ CharType* Str)
+	FORCEINLINE UE_STRING_CLASS& Append(/* no const! */ CharType* Str)
 	{
 		checkSlow(Str);
 		AppendChars(Str, TCString<std::remove_const_t<CharType>>::Strlen(Str));
@@ -365,14 +365,14 @@ public:
 
 	/** Append a string and return a reference to this */
 	template <typename CharRangeType, typename TEnableIf<TIsCharRangeNotCArray<CharRangeType>::Value>::Type* = nullptr>
-	FORCEINLINE FString& Append(CharRangeType&& Str)
+	FORCEINLINE UE_STRING_CLASS& Append(CharRangeType&& Str)
 	{
 		AppendChars(GetData(Str), GetNum(Str));
 		return *this;
 	}
 
 	/** Append a single character and return a reference to this */
-	CORE_API FString& AppendChar(ElementType InChar);
+	CORE_API UE_STRING_CLASS& AppendChar(ElementType InChar);
 
 	/** Append a string and return a reference to this */
 	template <typename StrType>
@@ -386,7 +386,7 @@ public:
 		typename AppendedCharType,
 		std::enable_if_t<TIsCharType<AppendedCharType>::Value>* = nullptr
 	>
-	FORCEINLINE FString& operator+=(AppendedCharType Char)
+	FORCEINLINE UE_STRING_CLASS& operator+=(AppendedCharType Char)
 	{
 		if constexpr (TIsCharEncodingSimplyConvertibleTo_V<AppendedCharType, ElementType>)
 		{
@@ -400,7 +400,7 @@ public:
 	}
 
 	CORE_API void InsertAt(int32 Index, ElementType Character);
-	CORE_API void InsertAt(int32 Index, const FString& Characters);
+	CORE_API void InsertAt(int32 Index, const UE_STRING_CLASS& Characters);
 
 	/**
 	 * Removes characters within the string.
@@ -441,7 +441,7 @@ public:
 	 * @param InPrefix the prefix to search for at the start of the string to remove.
 	 * @return true if the prefix was removed, otherwise false.
 	 */
-	bool RemoveFromStart(const FString& InPrefix, ESearchCase::Type SearchCase = ESearchCase::IgnoreCase)
+	bool RemoveFromStart(const UE_STRING_CLASS& InPrefix, ESearchCase::Type SearchCase = ESearchCase::IgnoreCase)
 	{
 		return RemoveFromStart(*InPrefix, InPrefix.Len(), SearchCase);
 	}
@@ -485,7 +485,7 @@ public:
 	 * @param InSuffix the suffix to search for at the end of the string to remove.
 	 * @return true if the suffix was removed, otherwise false.
 	 */
-	bool RemoveFromEnd(const FString& InSuffix, ESearchCase::Type SearchCase = ESearchCase::IgnoreCase)
+	bool RemoveFromEnd(const UE_STRING_CLASS& InSuffix, ESearchCase::Type SearchCase = ESearchCase::IgnoreCase)
 	{
 		return RemoveFromEnd(*InSuffix, InSuffix.Len(), SearchCase);
 	}
@@ -516,11 +516,11 @@ public:
 	 * @return The concatenated string.
 	 */
 	template <typename CharType>
-	UE_NODISCARD FORCEINLINE friend typename TEnableIf<TIsCharType<CharType>::Value, FString>::Type operator+(const FString& Lhs, CharType Rhs)
+	UE_NODISCARD FORCEINLINE friend typename TEnableIf<TIsCharType<CharType>::Value, UE_STRING_CLASS>::Type operator+(const UE_STRING_CLASS& Lhs, CharType Rhs)
 	{
 		Lhs.CheckInvariants();
 
-		FString Result(Lhs, 1);
+		UE_STRING_CLASS Result(Lhs, 1);
 		Result += Rhs;
 
 		return Result;
@@ -535,48 +535,48 @@ public:
 	 * @return The concatenated string.
 	 */
 	template <typename CharType>
-	UE_NODISCARD FORCEINLINE friend typename TEnableIf<TIsCharType<CharType>::Value, FString>::Type operator+(FString&& Lhs, CharType Rhs)
+	UE_NODISCARD FORCEINLINE friend typename TEnableIf<TIsCharType<CharType>::Value, UE_STRING_CLASS>::Type operator+(UE_STRING_CLASS&& Lhs, CharType Rhs)
 	{
 		Lhs.CheckInvariants();
 
-		FString Result(MoveTemp(Lhs), 1);
+		UE_STRING_CLASS Result(MoveTemp(Lhs), 1);
 		Result += Rhs;
 
 		return Result;
 	}
 
 private:
-	UE_NODISCARD static CORE_API FString ConcatFF(const FString& Lhs, const FString& Rhs);
-	UE_NODISCARD static CORE_API FString ConcatFF(FString&& Lhs, const FString& Rhs);
-	UE_NODISCARD static CORE_API FString ConcatFF(const FString& Lhs, FString&& Rhs);
-	UE_NODISCARD static CORE_API FString ConcatFF(FString&& Lhs, FString&& Rhs);
-	UE_NODISCARD static CORE_API FString ConcatFC(const FString& Lhs, const ElementType* Rhs);
-	UE_NODISCARD static CORE_API FString ConcatFC(FString&& Lhs,	const ElementType* Rhs);
-	UE_NODISCARD static CORE_API FString ConcatCF(const ElementType* Lhs, const FString& Rhs);
-	UE_NODISCARD static CORE_API FString ConcatCF(const ElementType* Lhs, FString&& Rhs);
-	UE_NODISCARD static CORE_API FString ConcatFR(const FString& Lhs, const ElementType* Rhs, int32 RhsLen);
-	UE_NODISCARD static CORE_API FString ConcatFR(FString&& Lhs,	const ElementType* Rhs, int32 RhsLen);
-	UE_NODISCARD static CORE_API FString ConcatRF(const ElementType* Lhs, int32 LhsLen, const FString& Rhs);
-	UE_NODISCARD static CORE_API FString ConcatRF(const ElementType* Lhs, int32 LhsLen, FString&& Rhs);
+	UE_NODISCARD static CORE_API UE_STRING_CLASS ConcatFF(const UE_STRING_CLASS& Lhs, const UE_STRING_CLASS& Rhs);
+	UE_NODISCARD static CORE_API UE_STRING_CLASS ConcatFF(UE_STRING_CLASS&& Lhs, const UE_STRING_CLASS& Rhs);
+	UE_NODISCARD static CORE_API UE_STRING_CLASS ConcatFF(const UE_STRING_CLASS& Lhs, UE_STRING_CLASS&& Rhs);
+	UE_NODISCARD static CORE_API UE_STRING_CLASS ConcatFF(UE_STRING_CLASS&& Lhs, UE_STRING_CLASS&& Rhs);
+	UE_NODISCARD static CORE_API UE_STRING_CLASS ConcatFC(const UE_STRING_CLASS& Lhs, const ElementType* Rhs);
+	UE_NODISCARD static CORE_API UE_STRING_CLASS ConcatFC(UE_STRING_CLASS&& Lhs,	const ElementType* Rhs);
+	UE_NODISCARD static CORE_API UE_STRING_CLASS ConcatCF(const ElementType* Lhs, const UE_STRING_CLASS& Rhs);
+	UE_NODISCARD static CORE_API UE_STRING_CLASS ConcatCF(const ElementType* Lhs, UE_STRING_CLASS&& Rhs);
+	UE_NODISCARD static CORE_API UE_STRING_CLASS ConcatFR(const UE_STRING_CLASS& Lhs, const ElementType* Rhs, int32 RhsLen);
+	UE_NODISCARD static CORE_API UE_STRING_CLASS ConcatFR(UE_STRING_CLASS&& Lhs,	const ElementType* Rhs, int32 RhsLen);
+	UE_NODISCARD static CORE_API UE_STRING_CLASS ConcatRF(const ElementType* Lhs, int32 LhsLen, const UE_STRING_CLASS& Rhs);
+	UE_NODISCARD static CORE_API UE_STRING_CLASS ConcatRF(const ElementType* Lhs, int32 LhsLen, UE_STRING_CLASS&& Rhs);
 
 public:
-	UE_NODISCARD FORCEINLINE friend FString operator+(const FString& Lhs, const FString& Rhs)			{ return ConcatFF(Lhs, Rhs); }
-	UE_NODISCARD FORCEINLINE friend FString operator+(FString&& Lhs, const FString& Rhs)				{ return ConcatFF(MoveTemp(Lhs), Rhs); }
-	UE_NODISCARD FORCEINLINE friend FString operator+(const FString& Lhs, FString&& Rhs)				{ return ConcatFF(Lhs,MoveTemp(Rhs)); }
-	UE_NODISCARD FORCEINLINE friend FString operator+(FString&& Lhs, FString&& Rhs)						{ return ConcatFF(MoveTemp(Lhs), MoveTemp(Rhs)); }
-	UE_NODISCARD FORCEINLINE friend FString operator+(const ElementType* Lhs, const FString& Rhs)		{ return ConcatCF(Lhs, Rhs); }
-	UE_NODISCARD FORCEINLINE friend FString operator+(const ElementType* Lhs, FString&& Rhs)			{ return ConcatCF(Lhs, MoveTemp(Rhs)); }
-	UE_NODISCARD FORCEINLINE friend FString operator+(const FString& Lhs, const ElementType* Rhs)		{ return ConcatFC(Lhs, Rhs); }
-	UE_NODISCARD FORCEINLINE friend FString operator+(FString&& Lhs, const ElementType* Rhs)			{ return ConcatFC(MoveTemp(Lhs), Rhs); }
+	UE_NODISCARD FORCEINLINE friend UE_STRING_CLASS operator+(const UE_STRING_CLASS& Lhs, const UE_STRING_CLASS& Rhs)	{ return ConcatFF(Lhs, Rhs); }
+	UE_NODISCARD FORCEINLINE friend UE_STRING_CLASS operator+(UE_STRING_CLASS&& Lhs, const UE_STRING_CLASS& Rhs)		{ return ConcatFF(MoveTemp(Lhs), Rhs); }
+	UE_NODISCARD FORCEINLINE friend UE_STRING_CLASS operator+(const UE_STRING_CLASS& Lhs, UE_STRING_CLASS&& Rhs)		{ return ConcatFF(Lhs,MoveTemp(Rhs)); }
+	UE_NODISCARD FORCEINLINE friend UE_STRING_CLASS operator+(UE_STRING_CLASS&& Lhs, UE_STRING_CLASS&& Rhs)				{ return ConcatFF(MoveTemp(Lhs), MoveTemp(Rhs)); }
+	UE_NODISCARD FORCEINLINE friend UE_STRING_CLASS operator+(const ElementType* Lhs, const UE_STRING_CLASS& Rhs)		{ return ConcatCF(Lhs, Rhs); }
+	UE_NODISCARD FORCEINLINE friend UE_STRING_CLASS operator+(const ElementType* Lhs, UE_STRING_CLASS&& Rhs)			{ return ConcatCF(Lhs, MoveTemp(Rhs)); }
+	UE_NODISCARD FORCEINLINE friend UE_STRING_CLASS operator+(const UE_STRING_CLASS& Lhs, const ElementType* Rhs)		{ return ConcatFC(Lhs, Rhs); }
+	UE_NODISCARD FORCEINLINE friend UE_STRING_CLASS operator+(UE_STRING_CLASS&& Lhs, const ElementType* Rhs)			{ return ConcatFC(MoveTemp(Lhs), Rhs); }
 
 	template <typename T, typename TEnableIf<TIsTCharRangeNotCArrayNotStringClass<T>::Value>::Type* = nullptr>
-	UE_NODISCARD FORCEINLINE friend FString operator+(T&& Lhs, const FString& Rhs)						{ return ConcatRF(GetData(Lhs), GetNum(Lhs), Rhs); }
+	UE_NODISCARD FORCEINLINE friend UE_STRING_CLASS operator+(T&& Lhs, const UE_STRING_CLASS& Rhs)						{ return ConcatRF(GetData(Lhs), GetNum(Lhs), Rhs); }
 	template <typename T, typename TEnableIf<TIsTCharRangeNotCArrayNotStringClass<T>::Value>::Type* = nullptr>
-	UE_NODISCARD FORCEINLINE friend FString operator+(T&& Lhs, FString&& Rhs)							{ return ConcatRF(GetData(Lhs), GetNum(Lhs), MoveTemp(Rhs)); }
+	UE_NODISCARD FORCEINLINE friend UE_STRING_CLASS operator+(T&& Lhs, UE_STRING_CLASS&& Rhs)							{ return ConcatRF(GetData(Lhs), GetNum(Lhs), MoveTemp(Rhs)); }
 	template <typename T, typename TEnableIf<TIsTCharRangeNotCArrayNotStringClass<T>::Value>::Type* = nullptr>
-	UE_NODISCARD FORCEINLINE friend FString operator+(const FString& Lhs, T&& Rhs)						{ return ConcatFR(Lhs, GetData(Rhs), GetNum(Rhs)); }
+	UE_NODISCARD FORCEINLINE friend UE_STRING_CLASS operator+(const UE_STRING_CLASS& Lhs, T&& Rhs)						{ return ConcatFR(Lhs, GetData(Rhs), GetNum(Rhs)); }
 	template <typename T, typename TEnableIf<TIsTCharRangeNotCArrayNotStringClass<T>::Value>::Type* = nullptr>
-	UE_NODISCARD FORCEINLINE friend FString operator+(FString&& Lhs, T&& Rhs)							{ return ConcatFR(MoveTemp(Lhs), GetData(Rhs), GetNum(Rhs)); }
+	UE_NODISCARD FORCEINLINE friend UE_STRING_CLASS operator+(UE_STRING_CLASS&& Lhs, T&& Rhs)							{ return ConcatFR(MoveTemp(Lhs), GetData(Rhs), GetNum(Rhs)); }
 
 	/**
 	 * Concatenate this path with given path ensuring the / character is used between them
@@ -584,7 +584,7 @@ public:
 	 * @param Str path array of characters to be concatenated onto the end of this
 	 * @return reference to path
 	 */
-	FORCEINLINE FString& operator/=( const ElementType* Str )
+	FORCEINLINE UE_STRING_CLASS& operator/=( const ElementType* Str )
 	{
 		checkSlow(Str);
 
@@ -599,7 +599,7 @@ public:
 	* @return reference to path
 	*/
 	template <typename CharRangeType, typename TEnableIf<TIsTCharRangeNotCArray <CharRangeType>::Value>::Type* = nullptr>
-	FORCEINLINE FString& operator/=(CharRangeType&& Str)
+	FORCEINLINE UE_STRING_CLASS& operator/=(CharRangeType&& Str)
 	{
 		PathAppend(GetData(Str), GetNum(Str));
 		return *this;
@@ -612,9 +612,9 @@ public:
 	* @return reference to path
 	*/
 	template <typename CharType,typename TEnableIf<TIsCharType<CharType>::Value>::Type* = nullptr>
-	FORCEINLINE FString& operator/=(const CharType* Str)
+	FORCEINLINE UE_STRING_CLASS& operator/=(const CharType* Str)
 	{
-		FString Temp = Str;
+		UE_STRING_CLASS Temp = Str;
 		PathAppend(GetData(Temp), GetNum(Temp));
 		return *this;
 	}
@@ -626,13 +626,13 @@ public:
 	 * @param Rhs Path to concatenate.
 	 * @return The new concatenated path
 	 */
-	UE_NODISCARD FORCEINLINE friend FString operator/(const FString& Lhs, const ElementType* Rhs)
+	UE_NODISCARD FORCEINLINE friend UE_STRING_CLASS operator/(const UE_STRING_CLASS& Lhs, const ElementType* Rhs)
 	{
 		checkSlow(Rhs);
 
 		int32 StrLength = FCString::Strlen(Rhs);
 
-		FString Result(Lhs, StrLength + 1);
+		UE_STRING_CLASS Result(Lhs, StrLength + 1);
 		Result.PathAppend(Rhs, StrLength);
 		return Result;
 	}
@@ -644,13 +644,13 @@ public:
 	 * @param Rhs Path to concatenate.
 	 * @return The new concatenated path
 	 */
-	UE_NODISCARD FORCEINLINE friend FString operator/(FString&& Lhs, const ElementType* Rhs)
+	UE_NODISCARD FORCEINLINE friend UE_STRING_CLASS operator/(UE_STRING_CLASS&& Lhs, const ElementType* Rhs)
 	{
 		checkSlow(Rhs);
 
 		int32 StrLength = FCString::Strlen(Rhs);
 
-		FString Result(MoveTemp(Lhs), StrLength + 1);
+		UE_STRING_CLASS Result(MoveTemp(Lhs), StrLength + 1);
 		Result.PathAppend(Rhs, StrLength);
 		return Result;
 	}
@@ -662,11 +662,11 @@ public:
 	 * @param Rhs Path to concatenate.
 	 * @return The new concatenated path
 	 */
-	UE_NODISCARD FORCEINLINE friend FString operator/(const FString& Lhs, const FString& Rhs)
+	UE_NODISCARD FORCEINLINE friend UE_STRING_CLASS operator/(const UE_STRING_CLASS& Lhs, const UE_STRING_CLASS& Rhs)
 	{
 		int32 StrLength = Rhs.Len();
 
-		FString Result(Lhs, StrLength + 1);
+		UE_STRING_CLASS Result(Lhs, StrLength + 1);
 		Result.PathAppend(Rhs.Data.GetData(), StrLength);
 		return Result;
 	}
@@ -678,11 +678,11 @@ public:
 	 * @param Rhs Path to concatenate.
 	 * @return The new concatenated path
 	 */
-	UE_NODISCARD FORCEINLINE friend FString operator/(FString&& Lhs, const FString& Rhs)
+	UE_NODISCARD FORCEINLINE friend UE_STRING_CLASS operator/(UE_STRING_CLASS&& Lhs, const UE_STRING_CLASS& Rhs)
 	{
 		int32 StrLength = Rhs.Len();
 
-		FString Result(MoveTemp(Lhs), StrLength + 1);
+		UE_STRING_CLASS Result(MoveTemp(Lhs), StrLength + 1);
 		Result.PathAppend(Rhs.Data.GetData(), StrLength);
 		return Result;
 	}
@@ -694,11 +694,11 @@ public:
 	 * @param Rhs Path to concatenate.
 	 * @return new string of the path
 	 */
-	UE_NODISCARD FORCEINLINE friend FString operator/(const ElementType* Lhs, const FString& Rhs)
+	UE_NODISCARD FORCEINLINE friend UE_STRING_CLASS operator/(const ElementType* Lhs, const UE_STRING_CLASS& Rhs)
 	{
 		int32 StrLength = Rhs.Len();
 
-		FString Result(FString(Lhs), StrLength + 1);
+		UE_STRING_CLASS Result(UE_STRING_CLASS(Lhs), StrLength + 1);
 		Result.PathAppend(Rhs.Data.GetData(), Rhs.Len());
 		return Result;
 	}
@@ -711,7 +711,7 @@ public:
 	 * @return true if the left string is lexicographically <= the right string, otherwise false
 	 * @note case insensitive
 	 */
-	UE_NODISCARD FORCEINLINE friend bool operator<=(const FString& Lhs, const FString& Rhs)
+	UE_NODISCARD FORCEINLINE friend bool operator<=(const UE_STRING_CLASS& Lhs, const UE_STRING_CLASS& Rhs)
 	{
 		return FPlatformString::Stricmp(*Lhs, *Rhs) <= 0;
 	}
@@ -725,7 +725,7 @@ public:
 	 * @note case insensitive
 	 */
 	template <typename CharType>
-	UE_NODISCARD FORCEINLINE friend bool operator<=(const FString& Lhs, const CharType* Rhs)
+	UE_NODISCARD FORCEINLINE friend bool operator<=(const UE_STRING_CLASS& Lhs, const CharType* Rhs)
 	{
 		return FPlatformString::Stricmp(*Lhs, Rhs) <= 0;
 	}
@@ -739,7 +739,7 @@ public:
 	 * @note case insensitive
 	 */
 	template <typename CharType>
-	UE_NODISCARD FORCEINLINE friend bool operator<=(const CharType* Lhs, const FString& Rhs)
+	UE_NODISCARD FORCEINLINE friend bool operator<=(const CharType* Lhs, const UE_STRING_CLASS& Rhs)
 	{
 		return FPlatformString::Stricmp(Lhs, *Rhs) <= 0;
 	}
@@ -752,7 +752,7 @@ public:
 	 * @return true if the left string is lexicographically < the right string, otherwise false
 	 * @note case insensitive
 	 */
-	UE_NODISCARD FORCEINLINE friend bool operator<(const FString& Lhs, const FString& Rhs)
+	UE_NODISCARD FORCEINLINE friend bool operator<(const UE_STRING_CLASS& Lhs, const UE_STRING_CLASS& Rhs)
 	{
 		return FPlatformString::Stricmp(*Lhs, *Rhs) < 0;
 	}
@@ -766,7 +766,7 @@ public:
 	 * @note case insensitive
 	 */
 	template <typename CharType>
-	UE_NODISCARD FORCEINLINE friend bool operator<(const FString& Lhs, const CharType* Rhs)
+	UE_NODISCARD FORCEINLINE friend bool operator<(const UE_STRING_CLASS& Lhs, const CharType* Rhs)
 	{
 		return FPlatformString::Stricmp(*Lhs, Rhs) < 0;
 	}
@@ -780,7 +780,7 @@ public:
 	 * @note case insensitive
 	 */
 	template <typename CharType>
-	UE_NODISCARD FORCEINLINE friend bool operator<(const CharType* Lhs, const FString& Rhs)
+	UE_NODISCARD FORCEINLINE friend bool operator<(const CharType* Lhs, const UE_STRING_CLASS& Rhs)
 	{
 		return FPlatformString::Stricmp(Lhs, *Rhs) < 0;
 	}
@@ -793,7 +793,7 @@ public:
 	 * @return true if the left string is lexicographically >= the right string, otherwise false
 	 * @note case insensitive
 	 */
-	UE_NODISCARD FORCEINLINE friend bool operator>=(const FString& Lhs, const FString& Rhs)
+	UE_NODISCARD FORCEINLINE friend bool operator>=(const UE_STRING_CLASS& Lhs, const UE_STRING_CLASS& Rhs)
 	{
 		return FPlatformString::Stricmp(*Lhs, *Rhs) >= 0;
 	}
@@ -807,7 +807,7 @@ public:
 	 * @note case insensitive
 	 */
 	template <typename CharType>
-	UE_NODISCARD FORCEINLINE friend bool operator>=(const FString& Lhs, const CharType* Rhs)
+	UE_NODISCARD FORCEINLINE friend bool operator>=(const UE_STRING_CLASS& Lhs, const CharType* Rhs)
 	{
 		return FPlatformString::Stricmp(*Lhs, Rhs) >= 0;
 	}
@@ -821,7 +821,7 @@ public:
 	 * @note case insensitive
 	 */
 	template <typename CharType>
-	UE_NODISCARD FORCEINLINE friend bool operator>=(const CharType* Lhs, const FString& Rhs)
+	UE_NODISCARD FORCEINLINE friend bool operator>=(const CharType* Lhs, const UE_STRING_CLASS& Rhs)
 	{
 		return FPlatformString::Stricmp(Lhs, *Rhs) >= 0;
 	}
@@ -834,7 +834,7 @@ public:
 	 * @return true if the left string is lexicographically > the right string, otherwise false
 	 * @note case insensitive
 	 */
-	UE_NODISCARD FORCEINLINE friend bool operator>(const FString& Lhs, const FString& Rhs)
+	UE_NODISCARD FORCEINLINE friend bool operator>(const UE_STRING_CLASS& Lhs, const UE_STRING_CLASS& Rhs)
 	{
 		return FPlatformString::Stricmp(*Lhs, *Rhs) > 0;
 	}
@@ -848,7 +848,7 @@ public:
 	 * @note case insensitive
 	 */
 	template <typename CharType>
-	UE_NODISCARD FORCEINLINE friend bool operator>(const FString& Lhs, const CharType* Rhs)
+	UE_NODISCARD FORCEINLINE friend bool operator>(const UE_STRING_CLASS& Lhs, const CharType* Rhs)
 	{
 		return FPlatformString::Stricmp(*Lhs, Rhs) > 0;
 	}
@@ -862,7 +862,7 @@ public:
 	 * @note case insensitive
 	 */
 	template <typename CharType>
-	UE_NODISCARD FORCEINLINE friend bool operator>(const CharType* Lhs, const FString& Rhs)
+	UE_NODISCARD FORCEINLINE friend bool operator>(const CharType* Lhs, const UE_STRING_CLASS& Rhs)
 	{
 		return FPlatformString::Stricmp(Lhs, *Rhs) > 0;
 	}
@@ -875,7 +875,7 @@ public:
 	 * @return true if the left string is lexicographically == the right string, otherwise false
 	 * @note case insensitive
 	 */
-	UE_NODISCARD FORCEINLINE bool operator==(const FString& Rhs) const
+	UE_NODISCARD FORCEINLINE bool operator==(const UE_STRING_CLASS& Rhs) const
 	{
 		return Equals(Rhs, ESearchCase::IgnoreCase);
 	}
@@ -903,7 +903,7 @@ public:
 	 * @note case insensitive
 	 */
 	template <typename CharType>
-	UE_NODISCARD FORCEINLINE friend bool operator==(const CharType* Lhs, const FString& Rhs)
+	UE_NODISCARD FORCEINLINE friend bool operator==(const CharType* Lhs, const UE_STRING_CLASS& Rhs)
 	{
 		return FPlatformString::Stricmp(Lhs, *Rhs) == 0;
 	}
@@ -916,7 +916,7 @@ public:
 	 * @return true if the left string is lexicographically != the right string, otherwise false
 	 * @note case insensitive
 	 */
-	UE_NODISCARD FORCEINLINE bool operator!=(const FString& Rhs) const
+	UE_NODISCARD FORCEINLINE bool operator!=(const UE_STRING_CLASS& Rhs) const
 	{
 		return !Equals(Rhs, ESearchCase::IgnoreCase);
 	}
@@ -944,7 +944,7 @@ public:
 	 * @note case insensitive
 	 */
 	template <typename CharType>
-	UE_NODISCARD FORCEINLINE friend bool operator!=(const CharType* Lhs, const FString& Rhs)
+	UE_NODISCARD FORCEINLINE friend bool operator!=(const CharType* Lhs, const UE_STRING_CLASS& Rhs)
 	{
 		return FPlatformString::Stricmp(Lhs, *Rhs) != 0;
 	}
@@ -956,12 +956,12 @@ public:
 	}
 
 	/** Returns the left most given number of characters */
-	UE_NODISCARD FORCEINLINE FString Left( int32 Count ) const &
+	UE_NODISCARD FORCEINLINE UE_STRING_CLASS Left( int32 Count ) const &
 	{
-		return FString( FMath::Clamp(Count,0,Len()), **this );
+		return UE_STRING_CLASS( FMath::Clamp(Count,0,Len()), **this );
 	}
 
-	UE_NODISCARD FORCEINLINE FString Left(int32 Count) &&
+	UE_NODISCARD FORCEINLINE UE_STRING_CLASS Left(int32 Count) &&
 	{
 		LeftInline(Count, false);
 		return MoveTemp(*this);
@@ -976,13 +976,13 @@ public:
 	}
 
 	/** Returns the left most characters from the string chopping the given number of characters from the end */
-	UE_NODISCARD FORCEINLINE FString LeftChop( int32 Count ) const &
+	UE_NODISCARD FORCEINLINE UE_STRING_CLASS LeftChop( int32 Count ) const &
 	{
 		const int32 Length = Len();
-		return FString( FMath::Clamp(Length-Count,0, Length), **this );
+		return UE_STRING_CLASS( FMath::Clamp(Length-Count,0, Length), **this );
 	}
 
-	UE_NODISCARD FORCEINLINE FString LeftChop(int32 Count)&&
+	UE_NODISCARD FORCEINLINE UE_STRING_CLASS LeftChop(int32 Count)&&
 	{
 		LeftChopInline(Count, false);
 		return MoveTemp(*this);
@@ -996,13 +996,13 @@ public:
 	}
 
 	/** Returns the string to the right of the specified location, counting back from the right (end of the word). */
-	UE_NODISCARD FORCEINLINE FString Right( int32 Count ) const &
+	UE_NODISCARD FORCEINLINE UE_STRING_CLASS Right( int32 Count ) const &
 	{
 		const int32 Length = Len();
-		return FString( **this + Length-FMath::Clamp(Count,0,Length) );
+		return UE_STRING_CLASS( **this + Length-FMath::Clamp(Count,0,Length) );
 	}
 
-	UE_NODISCARD FORCEINLINE FString Right(int32 Count) &&
+	UE_NODISCARD FORCEINLINE UE_STRING_CLASS Right(int32 Count) &&
 	{
 		RightInline(Count, false);
 		return MoveTemp(*this);
@@ -1016,9 +1016,9 @@ public:
 	}
 
 	/** Returns the string to the right of the specified location, counting forward from the left (from the beginning of the word). */
-	UE_NODISCARD CORE_API FString RightChop( int32 Count ) const &;
+	UE_NODISCARD CORE_API UE_STRING_CLASS RightChop( int32 Count ) const &;
 
-	UE_NODISCARD FORCEINLINE FString RightChop(int32 Count) &&
+	UE_NODISCARD FORCEINLINE UE_STRING_CLASS RightChop(int32 Count) &&
 	{
 		RightChopInline(Count, false);
 		return MoveTemp(*this);
@@ -1031,12 +1031,12 @@ public:
 	}
 
 	/** Returns the substring from Start position for Count characters. */
-	UE_NODISCARD CORE_API FString Mid(int32 Start, int32 Count) const &;
-	UE_NODISCARD CORE_API FString Mid(int32 Start, int32 Count) &&;
+	UE_NODISCARD CORE_API UE_STRING_CLASS Mid(int32 Start, int32 Count) const &;
+	UE_NODISCARD CORE_API UE_STRING_CLASS Mid(int32 Start, int32 Count) &&;
 
 	/** Returns the substring from Start position to the end */
-	UE_NODISCARD FORCEINLINE FString Mid(int32 Start) const & { return RightChop(Start); }
-	UE_NODISCARD FORCEINLINE FString Mid(int32 Start) && { return ((FString&&)*this).RightChop(Start); }
+	UE_NODISCARD FORCEINLINE UE_STRING_CLASS Mid(int32 Start) const & { return RightChop(Start); }
+	UE_NODISCARD FORCEINLINE UE_STRING_CLASS Mid(int32 Start) && { return ((UE_STRING_CLASS&&)*this).RightChop(Start); }
 
 	/** Modifies the string such that it is now the substring from Start position for Count characters. */
 	FORCEINLINE void MidInline(int32 Start, int32 Count = MAX_int32, bool bAllowShrinking = true)
@@ -1103,7 +1103,7 @@ public:
 	 *
 	 *        Consider using UE::String::FindLast() as an alternative.
 	 */
-	UE_NODISCARD int32 Find(const FString& SubStr, ESearchCase::Type SearchCase = ESearchCase::IgnoreCase,
+	UE_NODISCARD int32 Find(const UE_STRING_CLASS& SubStr, ESearchCase::Type SearchCase = ESearchCase::IgnoreCase,
 		ESearchDir::Type SearchDir = ESearchDir::FromStart, int32 StartPosition = INDEX_NONE) const
 	{
 		return Find(*SubStr, SubStr.Len(), SearchCase, SearchDir, StartPosition);
@@ -1168,7 +1168,7 @@ public:
 	 * @param SearchDir			Indicates whether the search starts at the beginning or at the end ( defaults to ESearchDir::FromStart )
 	 * @return					Returns whether the string contains the substring. If the substring is empty, returns true.
 	 **/
-	UE_NODISCARD FORCEINLINE bool Contains(const FString& SubStr, ESearchCase::Type SearchCase = ESearchCase::IgnoreCase,
+	UE_NODISCARD FORCEINLINE bool Contains(const UE_STRING_CLASS& SubStr, ESearchCase::Type SearchCase = ESearchCase::IgnoreCase,
 							  ESearchDir::Type SearchDir = ESearchDir::FromStart ) const
 	{
 		return Find(*SubStr, SearchCase, SearchDir) != INDEX_NONE;
@@ -1249,7 +1249,7 @@ public:
 	 * @param SearchCase 	Whether or not the comparison should ignore case
 	 * @return true if this string is lexicographically equivalent to the other, otherwise false
 	 */
-	UE_NODISCARD FORCEINLINE bool Equals(const FString& Other, ESearchCase::Type SearchCase = ESearchCase::CaseSensitive) const
+	UE_NODISCARD FORCEINLINE bool Equals(const UE_STRING_CLASS& Other, ESearchCase::Type SearchCase = ESearchCase::CaseSensitive) const
 	{
 		int32 Num = Data.Num();
 		int32 OtherNum = Other.Data.Num();
@@ -1281,7 +1281,7 @@ public:
 	 * @param SearchCase 	Whether or not the comparison should ignore case
 	 * @return 0 if equal, negative if less than, positive if greater than
 	 */
-	UE_NODISCARD FORCEINLINE int32 Compare( const FString& Other, ESearchCase::Type SearchCase = ESearchCase::CaseSensitive ) const
+	UE_NODISCARD FORCEINLINE int32 Compare( const UE_STRING_CLASS& Other, ESearchCase::Type SearchCase = ESearchCase::CaseSensitive ) const
 	{
 		if( SearchCase == ESearchCase::CaseSensitive )
 		{
@@ -1303,41 +1303,41 @@ public:
 	 * @param SearchDir			Indicates whether the search starts at the beginning or at the end ( defaults to ESearchDir::FromStart )
 	 * @return true if string is split, otherwise false
 	 */
-	CORE_API bool Split(const FString& InS, FString* LeftS, FString* RightS, ESearchCase::Type SearchCase,
+	CORE_API bool Split(const UE_STRING_CLASS& InS, UE_STRING_CLASS* LeftS, UE_STRING_CLASS* RightS, ESearchCase::Type SearchCase,
 		ESearchDir::Type SearchDir = ESearchDir::FromStart) const;
 
 	/** Split with ESearchCase::IgnoreCase and ESearchDir::FromStart. Allows compiler to avoid branches w/o inlining code. */
-	CORE_API bool Split(const FString& InS, FString* LeftS, FString* RightS) const;
+	CORE_API bool Split(const UE_STRING_CLASS& InS, UE_STRING_CLASS* LeftS, UE_STRING_CLASS* RightS) const;
 
 	/** Returns a new string with the characters of this converted to uppercase */
-	UE_NODISCARD CORE_API FString ToUpper() const &;
+	UE_NODISCARD CORE_API UE_STRING_CLASS ToUpper() const &;
 
 	/**
 	 * Converts all characters in this rvalue string to uppercase and moves it into the returned string.
 	 * @return a new string with the characters of this converted to uppercase
 	 */
-	UE_NODISCARD CORE_API FString ToUpper() &&;
+	UE_NODISCARD CORE_API UE_STRING_CLASS ToUpper() &&;
 
 	/** Converts all characters in this string to uppercase */
 	CORE_API void ToUpperInline();
 
 	/** Returns a new string with the characters of this converted to lowercase */
-	UE_NODISCARD CORE_API FString ToLower() const &;
+	UE_NODISCARD CORE_API UE_STRING_CLASS ToLower() const &;
 
 	/**
 	 * Converts all characters in this rvalue string to lowercase and moves it into the returned string.
 	 * @return a new string with the characters of this converted to lowercase
 	 */
-	UE_NODISCARD CORE_API FString ToLower() &&;
+	UE_NODISCARD CORE_API UE_STRING_CLASS ToLower() &&;
 
 	/** Converts all characters in this string to lowercase */
 	CORE_API void ToLowerInline();
 
 	/** Pad the left of this string for ChCount characters */
-	UE_NODISCARD CORE_API FString LeftPad( int32 ChCount ) const;
+	UE_NODISCARD CORE_API UE_STRING_CLASS LeftPad( int32 ChCount ) const;
 
 	/** Pad the right of this string for ChCount characters */
-	UE_NODISCARD CORE_API FString RightPad( int32 ChCount ) const;
+	UE_NODISCARD CORE_API UE_STRING_CLASS RightPad( int32 ChCount ) const;
 	
 	/** Returns true if the string only contains numeric characters */
 	UE_NODISCARD CORE_API bool IsNumeric() const;
@@ -1354,7 +1354,7 @@ public:
 	 * @returns A string that was constructed using format and additional parameters.
 	 */
 	template <typename FmtType, typename... Types>
-	UE_NODISCARD static FString Printf(const FmtType& Fmt, Types... Args)
+	UE_NODISCARD static UE_STRING_CLASS Printf(const FmtType& Fmt, Types... Args)
 	{
 		static_assert(TIsArrayOrRefOfTypeByPredicate<FmtType, TIsCharEncodingCompatibleWithTCHAR>::Value, "Formatting string must be a character array.");
 		static_assert(TAnd<TIsValidVariadicFunctionArg<Types>...>::Value, "Invalid argument(s) passed to Printf");
@@ -1367,7 +1367,7 @@ public:
 	 * @return a reference to the modified string, so that it can be chained
 	 */
 	template <typename FmtType, typename... Types>
-	FString& Appendf(const FmtType& Fmt, Types... Args)
+	UE_STRING_CLASS& Appendf(const FmtType& Fmt, Types... Args)
 	{
 		static_assert(TIsArrayOrRefOfTypeByPredicate<FmtType, TIsCharEncodingCompatibleWithTCHAR>::Value, "Formatting string must be a character array.");
 		static_assert(TAnd<TIsValidVariadicFunctionArg<Types>...>::Value, "Invalid argument(s) passed to TString::Appendf");
@@ -1377,8 +1377,8 @@ public:
 	}
 
 private:
-	static CORE_API FString VARARGS PrintfImpl(const ElementType* Fmt, ...);
-	static CORE_API void VARARGS AppendfImpl(FString& AppendToMe, const ElementType* Fmt, ...);
+	static CORE_API UE_STRING_CLASS VARARGS PrintfImpl(const ElementType* Fmt, ...);
+	static CORE_API void VARARGS AppendfImpl(UE_STRING_CLASS& AppendToMe, const ElementType* Fmt, ...);
 public:
 
 	/**
@@ -1387,7 +1387,7 @@ public:
 	 * @param InNamedArguments		A map of named arguments that match the tokens specified in InExpression
 	 * @return A string containing the formatted text
 	 */
-	UE_NODISCARD static CORE_API FString Format(const ElementType* InFormatString, const FStringFormatNamedArguments& InNamedArguments);
+	UE_NODISCARD static CORE_API UE_STRING_CLASS Format(const ElementType* InFormatString, const PREPROCESSOR_JOIN(UE_STRING_CLASS, FormatNamedArguments)& InNamedArguments);
 
 	/**
 	 * Format the specified string using the specified arguments. Replaces instances of {0} with indices from the given array matching the index specified in the token
@@ -1395,10 +1395,10 @@ public:
 	 * @param InOrderedArguments	An array of ordered arguments that match the tokens specified in InExpression
 	 * @return A string containing the formatted text
 	 */
-	UE_NODISCARD static CORE_API FString Format(const ElementType* InFormatString, const FStringFormatOrderedArguments& InOrderedArguments);
+	UE_NODISCARD static CORE_API UE_STRING_CLASS Format(const ElementType* InFormatString, const PREPROCESSOR_JOIN(UE_STRING_CLASS, FormatOrderedArguments)& InOrderedArguments);
 
 	/** Returns a string containing only the Ch character */
-	UE_NODISCARD static CORE_API FString Chr( ElementType Ch );
+	UE_NODISCARD static CORE_API UE_STRING_CLASS Chr( ElementType Ch );
 
 	/**
 	 * Returns a string that is full of a variable number of characters
@@ -1408,7 +1408,7 @@ public:
 	 * 
 	 * @return The string of NumCharacters characters.
 	 */
-	UE_NODISCARD static CORE_API FString ChrN( int32 NumCharacters, ElementType Char );
+	UE_NODISCARD static CORE_API UE_STRING_CLASS ChrN( int32 NumCharacters, ElementType Char );
 
 	/**
 	 * Serializes the string.
@@ -1418,7 +1418,7 @@ public:
 	 *
 	 * @return Reference to the Archive after serialization.
 	 */
-	friend CORE_API FArchive& operator<<( FArchive& Ar, FString& S );
+	friend CORE_API FArchive& operator<<( FArchive& Ar, UE_STRING_CLASS& S );
 
 
 	/**
@@ -1451,7 +1451,7 @@ public:
 	 * @param SearchCase		Indicates whether the search is case sensitive or not ( defaults to ESearchCase::IgnoreCase )
 	 * @return true if this string begins with specified text, false otherwise
 	 */
-	UE_NODISCARD bool StartsWith(const FString& InPrefix, ESearchCase::Type SearchCase = ESearchCase::IgnoreCase) const
+	UE_NODISCARD bool StartsWith(const UE_STRING_CLASS& InPrefix, ESearchCase::Type SearchCase = ESearchCase::IgnoreCase) const
 	{
 		return StartsWith(*InPrefix, InPrefix.Len(), SearchCase);
 	}
@@ -1494,7 +1494,7 @@ public:
 	 * @param SearchCase		Indicates whether the search is case sensitive or not ( defaults to ESearchCase::IgnoreCase )
 	 * @return true if this string ends with specified text, false otherwise
 	 */
-	UE_NODISCARD bool EndsWith(const FString& InSuffix, ESearchCase::Type SearchCase = ESearchCase::IgnoreCase) const
+	UE_NODISCARD bool EndsWith(const UE_STRING_CLASS& InSuffix, ESearchCase::Type SearchCase = ESearchCase::IgnoreCase) const
 	{
 		return EndsWith(*InSuffix, InSuffix.Len(), SearchCase);
 	}
@@ -1543,7 +1543,7 @@ public:
 	 * @return true if this string matches the *?-type wildcard given.
 	 * @warning This is a simple, SLOW routine. Use with caution
 	 */
-	UE_NODISCARD bool MatchesWildcard(const FString& Wildcard, ESearchCase::Type SearchCase = ESearchCase::IgnoreCase) const
+	UE_NODISCARD bool MatchesWildcard(const UE_STRING_CLASS& Wildcard, ESearchCase::Type SearchCase = ESearchCase::IgnoreCase) const
 	{
 		return MatchesWildcard(*Wildcard, Wildcard.Len(), SearchCase);
 	}
@@ -1567,13 +1567,13 @@ public:
 	 * Removes whitespace characters from the start and end of this string.
 	 * @note Unlike Trim() this function returns a copy, and does not mutate the string.
 	 */
-	UE_NODISCARD CORE_API FString TrimStartAndEnd() const &;
+	UE_NODISCARD CORE_API UE_STRING_CLASS TrimStartAndEnd() const &;
 
 	/**
 	 * Removes whitespace characters from the start and end of this string.
 	 * @note Unlike Trim() this function returns a copy, and does not mutate the string.
 	 */
-	UE_NODISCARD CORE_API FString TrimStartAndEnd() &&;
+	UE_NODISCARD CORE_API UE_STRING_CLASS TrimStartAndEnd() &&;
 
 	/**
 	 * Removes whitespace characters from the start of this string. Modifies the string in-place.
@@ -1584,13 +1584,13 @@ public:
 	 * Removes whitespace characters from the start of this string.
 	 * @note Unlike Trim() this function returns a copy, and does not mutate the string.
 	 */
-	UE_NODISCARD CORE_API FString TrimStart() const &;
+	UE_NODISCARD CORE_API UE_STRING_CLASS TrimStart() const &;
 
 	/**
 	 * Removes whitespace characters from the start of this string.
 	 * @note Unlike Trim() this function returns a copy, and does not mutate the string.
 	 */
-	UE_NODISCARD CORE_API FString TrimStart() &&;
+	UE_NODISCARD CORE_API UE_STRING_CLASS TrimStart() &&;
 
 	/**
 	 * Removes whitespace characters from the end of this string. Modifies the string in-place.
@@ -1601,13 +1601,13 @@ public:
 	 * Removes whitespace characters from the end of this string.
 	 * @note Unlike TrimTrailing() this function returns a copy, and does not mutate the string.
 	 */
-	UE_NODISCARD CORE_API FString TrimEnd() const &;
+	UE_NODISCARD CORE_API UE_STRING_CLASS TrimEnd() const &;
 
 	/**
 	 * Removes whitespace characters from the end of this string.
 	 * @note Unlike TrimTrailing() this function returns a copy, and does not mutate the string.
 	 */
-	UE_NODISCARD CORE_API FString TrimEnd() &&;
+	UE_NODISCARD CORE_API UE_STRING_CLASS TrimEnd() &&;
 
 	/** 
 	 * Trims the inner array after the null terminator.
@@ -1629,23 +1629,23 @@ public:
 	/**
 	 * Returns a copy of this string with wrapping quotation marks removed.
 	 */
-	UE_NODISCARD CORE_API FString TrimQuotes(bool* bQuotesRemoved = nullptr ) const &;
+	UE_NODISCARD CORE_API UE_STRING_CLASS TrimQuotes(bool* bQuotesRemoved = nullptr ) const &;
 
 	/**
 	 * Returns this string with wrapping quotation marks removed.
 	 */
-	UE_NODISCARD CORE_API FString TrimQuotes(bool* bQuotesRemoved = nullptr) &&;
+	UE_NODISCARD CORE_API UE_STRING_CLASS TrimQuotes(bool* bQuotesRemoved = nullptr) &&;
 	
 	/**
 	* Returns a copy of this string with wrapping CharacterToTrim removed (removes at max one instance in the beginning and end of the string).
 	* @see TrimCharInline for an inline variant
 	*/
-	UE_NODISCARD CORE_API FString TrimChar(ElementType CharacterToTrim, bool* bCharRemoved = nullptr ) const &;
+	UE_NODISCARD CORE_API UE_STRING_CLASS TrimChar(ElementType CharacterToTrim, bool* bCharRemoved = nullptr ) const &;
 
 	/**
 	* Returns a copy of this string with wrapping CharacterToTrim removed (removes at max one instance in the beginning and end of the string).
 	*/
-	UE_NODISCARD CORE_API FString TrimChar(ElementType CharacterToTrim, bool* bCharRemoved = nullptr) &&;
+	UE_NODISCARD CORE_API UE_STRING_CLASS TrimChar(ElementType CharacterToTrim, bool* bCharRemoved = nullptr) &&;
 	
 	/**
 	 * Breaks up a delimited string into elements of a string array.
@@ -1656,7 +1656,7 @@ public:
 	 *
 	 * @return	The number of elements in InArray
 	 */
-	CORE_API int32 ParseIntoArray( TArray<FString>& OutArray, const ElementType* pchDelim, bool InCullEmpty = true ) const;
+	CORE_API int32 ParseIntoArray( TArray<UE_STRING_CLASS>& OutArray, const ElementType* pchDelim, bool InCullEmpty = true ) const;
 
 	/**
 	 * Breaks up a delimited string into elements of a string array, using any whitespace and an 
@@ -1668,7 +1668,7 @@ public:
 	 *
 	 * @return	The number of elements in InArray
 	 */
-	CORE_API int32 ParseIntoArrayWS( TArray<FString>& OutArray, const ElementType* pchExtraDelim = nullptr, bool InCullEmpty = true ) const;
+	CORE_API int32 ParseIntoArrayWS( TArray<UE_STRING_CLASS>& OutArray, const ElementType* pchExtraDelim = nullptr, bool InCullEmpty = true ) const;
 
 	/**
 	* Breaks up a delimited string into elements of a string array, using line ending characters
@@ -1678,7 +1678,7 @@ public:
 	*
 	* @return	The number of elements in InArray
 	*/
-	CORE_API int32 ParseIntoArrayLines(TArray<FString>& OutArray, bool InCullEmpty = true) const;
+	CORE_API int32 ParseIntoArrayLines(TArray<UE_STRING_CLASS>& OutArray, bool InCullEmpty = true) const;
 
 	/**
 	* Breaks up a delimited string into elements of a string array, using the given delimiters
@@ -1690,7 +1690,7 @@ public:
 	*
 	* @return	The number of elements in InArray
 	*/
-	CORE_API int32 ParseIntoArray(TArray<FString>& OutArray, const ElementType*const* DelimArray, int32 NumDelims, bool InCullEmpty = true) const;
+	CORE_API int32 ParseIntoArray(TArray<UE_STRING_CLASS>& OutArray, const ElementType*const* DelimArray, int32 NumDelims, bool InCullEmpty = true) const;
 
 	/**
 	 * Takes an array of strings and removes any zero length entries.
@@ -1699,17 +1699,17 @@ public:
 	 *
 	 * @return	The number of elements left in InArray
 	 */
-	static CORE_API int32 CullArray( TArray<FString>* InArray );
+	static CORE_API int32 CullArray( TArray<UE_STRING_CLASS>* InArray );
 
 	/**
 	 * Returns a copy of this string, with the characters in reverse order
 	 */
-	UE_NODISCARD CORE_API FString Reverse() const &;
+	UE_NODISCARD CORE_API UE_STRING_CLASS Reverse() const &;
 
 	/**
 	 * Returns this string, with the characters in reverse order
 	 */
-	UE_NODISCARD CORE_API FString Reverse() &&;
+	UE_NODISCARD CORE_API UE_STRING_CLASS Reverse() &&;
 
 	/**
 	 * Reverses the order of characters in this string
@@ -1724,7 +1724,7 @@ public:
 	 * @param SearchCase	Indicates whether the search is case sensitive or not ( defaults to ESearchCase::IgnoreCase )
 	 * @return a copy of this string with the replacement made
 	 */
-	UE_NODISCARD CORE_API FString Replace(const ElementType* From, const ElementType* To, ESearchCase::Type SearchCase = ESearchCase::IgnoreCase) const &;
+	UE_NODISCARD CORE_API UE_STRING_CLASS Replace(const ElementType* From, const ElementType* To, ESearchCase::Type SearchCase = ESearchCase::IgnoreCase) const &;
 
 	/**
 	 * Replace all occurrences of a substring in this string
@@ -1734,7 +1734,7 @@ public:
 	 * @param SearchCase	Indicates whether the search is case sensitive or not ( defaults to ESearchCase::IgnoreCase )
 	 * @return a copy of this string with the replacement made
 	 */
-	UE_NODISCARD CORE_API FString Replace(const ElementType* From, const ElementType* To, ESearchCase::Type SearchCase = ESearchCase::IgnoreCase) &&;
+	UE_NODISCARD CORE_API UE_STRING_CLASS Replace(const ElementType* From, const ElementType* To, ESearchCase::Type SearchCase = ESearchCase::IgnoreCase) &&;
 
 	/**
 	 * Replace all occurrences of SearchText with ReplacementText in this string.
@@ -1777,16 +1777,16 @@ public:
 	/**
 	 * Returns a copy of this string with all quote marks escaped (unless the quote is already escaped)
 	 */
-	UE_NODISCARD FString ReplaceQuotesWithEscapedQuotes() const &
+	UE_NODISCARD UE_STRING_CLASS ReplaceQuotesWithEscapedQuotes() const &
 	{
-		FString Result(*this);
+		UE_STRING_CLASS Result(*this);
 		return MoveTemp(Result).ReplaceQuotesWithEscapedQuotes();
 	}
 
 	/**
 	 * Returns a copy of this string with all quote marks escaped (unless the quote is already escaped)
 	 */
-	UE_NODISCARD CORE_API FString ReplaceQuotesWithEscapedQuotes() &&;
+	UE_NODISCARD CORE_API UE_STRING_CLASS ReplaceQuotesWithEscapedQuotes() &&;
 
 	/**
 	 * Replaces certain characters with the "escaped" version of that character (i.e. replaces "\n" with "\\n").
@@ -1804,9 +1804,9 @@ public:
 	 *
 	 * @return	a string with all control characters replaced by the escaped version.
 	 */
-	UE_NODISCARD FString ReplaceCharWithEscapedChar( const TArray<ElementType>* Chars = nullptr ) const &
+	UE_NODISCARD UE_STRING_CLASS ReplaceCharWithEscapedChar( const TArray<ElementType>* Chars = nullptr ) const &
 	{
-		FString Result(*this);
+		UE_STRING_CLASS Result(*this);
 		Result.ReplaceCharWithEscapedCharInline(Chars);
 		return Result;
 	}
@@ -1819,7 +1819,7 @@ public:
 	 *
 	 * @return	a string with all control characters replaced by the escaped version.
 	 */
-	UE_NODISCARD FString ReplaceCharWithEscapedChar( const TArray<ElementType>* Chars = nullptr ) &&
+	UE_NODISCARD UE_STRING_CLASS ReplaceCharWithEscapedChar( const TArray<ElementType>* Chars = nullptr ) &&
 	{
 		ReplaceCharWithEscapedCharInline(Chars);
 		return MoveTemp(*this);
@@ -1836,9 +1836,9 @@ public:
 	 * replaces "\\n" with "\n".  Counterpart to ReplaceCharWithEscapedChar().
 	 * @return copy of this string with replacement made
 	 */
-	UE_NODISCARD FString ReplaceEscapedCharWithChar( const TArray<ElementType>* Chars = nullptr ) const &
+	UE_NODISCARD UE_STRING_CLASS ReplaceEscapedCharWithChar( const TArray<ElementType>* Chars = nullptr ) const &
 	{
-		FString Result(*this);
+		UE_STRING_CLASS Result(*this);
 		Result.ReplaceEscapedCharWithCharInline(Chars);
 		return Result;
 	}
@@ -1848,7 +1848,7 @@ public:
 	 * replaces "\\n" with "\n".  Counterpart to ReplaceCharWithEscapedChar().
 	 * @return copy of this string with replacement made
 	 */
-	UE_NODISCARD FString ReplaceEscapedCharWithChar( const TArray<ElementType>* Chars = nullptr ) &&
+	UE_NODISCARD UE_STRING_CLASS ReplaceEscapedCharWithChar( const TArray<ElementType>* Chars = nullptr ) &&
 	{
 		ReplaceEscapedCharWithCharInline(Chars);
 		return MoveTemp(*this);
@@ -1865,9 +1865,9 @@ public:
 	 * @param InSpacesPerTab - Number of spaces that a tab represents
 	 * @return copy of this string with replacement made
 	 */
-	UE_NODISCARD FString ConvertTabsToSpaces(const int32 InSpacesPerTab) const &
+	UE_NODISCARD UE_STRING_CLASS ConvertTabsToSpaces(const int32 InSpacesPerTab) const &
 	{
-		FString FinalString(*this);
+		UE_STRING_CLASS FinalString(*this);
 		FinalString.ConvertTabsToSpacesInline(InSpacesPerTab);
 		return FinalString;
 	}
@@ -1877,14 +1877,14 @@ public:
 	 * @param InSpacesPerTab - Number of spaces that a tab represents
 	 * @return copy of this string with replacement made
 	 */
-	UE_NODISCARD FString ConvertTabsToSpaces(const int32 InSpacesPerTab) &&
+	UE_NODISCARD UE_STRING_CLASS ConvertTabsToSpaces(const int32 InSpacesPerTab) &&
 	{
 		ConvertTabsToSpacesInline(InSpacesPerTab);
 		return MoveTemp(*this);
 	}
 
 	// Takes the number passed in and formats the string in comma format ( 12345 becomes "12,345")
-	UE_NODISCARD static CORE_API FString FormatAsNumber( int32 InNumber );
+	UE_NODISCARD static CORE_API UE_STRING_CLASS FormatAsNumber( int32 InNumber );
 
 	// To allow more efficient memory handling, automatically adds one for the string termination.
 	CORE_API void Reserve(int32 CharacterCount);
@@ -1900,9 +1900,9 @@ public:
 
 
 	/** Converts an integer to a string. */
-	UE_NODISCARD static FORCEINLINE FString FromInt( int32 Num )
+	UE_NODISCARD static FORCEINLINE UE_STRING_CLASS FromInt( int32 Num )
 	{
-		FString Ret;
+		UE_STRING_CLASS Ret;
 		Ret.AppendInt(Num); 
 		return Ret;
 	}
@@ -1927,7 +1927,7 @@ public:
 	 *
 	 * @return the blob in string form
 	 */
-	UE_NODISCARD static CORE_API FString FromBlob(const uint8* SrcBuffer,const uint32 SrcSize);
+	UE_NODISCARD static CORE_API UE_STRING_CLASS FromBlob(const uint8* SrcBuffer,const uint32 SrcSize);
 
 	/**
 	 * Converts a string into a buffer
@@ -1937,7 +1937,7 @@ public:
 	 *
 	 * @return true if the conversion happened, false otherwise
 	 */
-	static CORE_API bool ToBlob(const FString& Source,uint8* DestBuffer,const uint32 DestSize);
+	static CORE_API bool ToBlob(const UE_STRING_CLASS& Source,uint8* DestBuffer,const uint32 DestSize);
 
 	/**
 	 * Converts a buffer to a string by hex-ifying the elements
@@ -1947,7 +1947,7 @@ public:
 	 *
 	 * @return the blob in string form
 	 */
-	UE_NODISCARD static CORE_API FString FromHexBlob(const uint8* SrcBuffer,const uint32 SrcSize);
+	UE_NODISCARD static CORE_API UE_STRING_CLASS FromHexBlob(const uint8* SrcBuffer,const uint32 SrcSize);
 
 	/**
 	 * Converts a string into a buffer
@@ -1957,7 +1957,7 @@ public:
 	 *
 	 * @return true if the conversion happened, false otherwise
 	 */
-	static CORE_API bool ToHexBlob(const FString& Source,uint8* DestBuffer,const uint32 DestSize);
+	static CORE_API bool ToHexBlob(const UE_STRING_CLASS& Source,uint8* DestBuffer,const uint32 DestSize);
 
 	/**
 	 * Converts a float string with the trailing zeros stripped
@@ -1968,7 +1968,7 @@ public:
 	 *
 	 * @return sanitized string version of float
 	 */
-	UE_NODISCARD static CORE_API FString SanitizeFloat( double InFloat, const int32 InMinFractionalDigits = 1 );
+	UE_NODISCARD static CORE_API UE_STRING_CLASS SanitizeFloat( double InFloat, const int32 InMinFractionalDigits = 1 );
 
 	/**
 	 * Joins a range of 'something that can be concatentated to strings with +=' together into a single string with separators.
@@ -1979,10 +1979,10 @@ public:
 	 * @return	The final, joined, separated string.
 	 */
 	template <typename RangeType>
-	UE_NODISCARD static FString Join(const RangeType& Range, const ElementType* Separator)
+	UE_NODISCARD static UE_STRING_CLASS Join(const RangeType& Range, const ElementType* Separator)
 	{
-		FString Result;
-		bool    First = true;
+		UE_STRING_CLASS Result;
+		bool            First = true;
 		for (const auto& Element : Range)
 		{
 			if (First)
@@ -2010,10 +2010,10 @@ public:
 	 * @return	The final, joined, separated string.
 	 */
 	template <typename RangeType, typename ProjectionType>
-	UE_NODISCARD static FString JoinBy(const RangeType& Range, const ElementType* Separator, ProjectionType Proj)
+	UE_NODISCARD static UE_STRING_CLASS JoinBy(const RangeType& Range, const ElementType* Separator, ProjectionType Proj)
 	{
-		FString Result;
-		bool    First = true;
+		UE_STRING_CLASS Result;
+		bool            First = true;
 		for (const auto& Element : Range)
 		{
 			if (First)
@@ -2037,43 +2037,43 @@ public:
 	}
 
 	/** Case insensitive string hash function. */
-	friend FORCEINLINE uint32 GetTypeHash(const FString& S)
+	friend FORCEINLINE uint32 GetTypeHash(const UE_STRING_CLASS& S)
 	{
 		// This must match the GetTypeHash behavior of FStringView
 		return FCrc::Strihash_DEPRECATED(S.Len(), *S);
 	}
 };
 
-template<> struct TIsZeroConstructType<FString> { enum { Value = true }; };
-Expose_TNameOf(FString)
+template<> struct TIsZeroConstructType<UE_STRING_CLASS> { enum { Value = true }; };
+Expose_TNameOf(UE_STRING_CLASS)
 
-inline FString::ElementType* GetData(FString& String)
+inline UE_STRING_CLASS::ElementType* GetData(UE_STRING_CLASS& String)
 {
 	return String.GetCharArray().GetData();
 }
 
-inline const FString::ElementType* GetData(const FString& String)
+inline const UE_STRING_CLASS::ElementType* GetData(const UE_STRING_CLASS& String)
 {
 	return String.GetCharArray().GetData();
 }
 
-inline int32 GetNum(const FString& String)
+inline int32 GetNum(const UE_STRING_CLASS& String)
 {
 	return String.Len();
 }
 
 /** Append bytes as uppercase hex string */
-CORE_API void BytesToHex(const uint8* Bytes, int32 NumBytes, FString& Out);
+CORE_API void BytesToHex(const uint8* Bytes, int32 NumBytes, UE_STRING_CLASS& Out);
 /** Append bytes as lowercase hex string */
-CORE_API void BytesToHexLower(const uint8* Bytes, int32 NumBytes, FString& Out);
+CORE_API void BytesToHexLower(const uint8* Bytes, int32 NumBytes, UE_STRING_CLASS& Out);
 
 /** 
- * Convert FString of Hex digits into the byte array.
- * @param HexString		The FString of Hex values
+ * Convert a string of Hex digits into the byte array.
+ * @param HexString		The string of Hex values
  * @param OutBytes		Ptr to memory must be preallocated large enough
  * @return	The number of bytes copied
  */
-CORE_API int32 HexToBytes(const FString& HexString, uint8* OutBytes);
+CORE_API int32 HexToBytes(const UE_STRING_CLASS& HexString, uint8* OutBytes);
 
 /**
  * Generalized API to convert something to a string. Function named after the (deprecated) Lex namespace, which
@@ -2086,26 +2086,26 @@ CORE_API int32 HexToBytes(const FString& HexString, uint8* OutBytes);
  *		bool								LexTryParseString(T& OutValue, const TCHAR* Buffer);
  *		void 								LexFromString(T& OutValue, const TCHAR* Buffer);
  *		<implicitly convertible to string>	LexToString(T);
- *		                    ^-- Generally this means it can return either FString or const TCHAR* 
- *		                        Generic code that uses ToString should assign to an FString or forward along to other functions
- *		                        that accept types that are also implicitly convertible to FString 
+ *		                    ^-- Generally this means it can return either a string or const TCHAR* 
+ *		                        Generic code that uses ToString should assign to a string or forward along to other functions
+ *		                        that accept types that are also implicitly convertible to a string 
  *
  *	Implement custom functionality externally.
  */
 
  /** Covert a string buffer to intrinsic types */
-inline void LexFromString(int8& OutValue, 		const TCHAR* Buffer)	{	OutValue = (int8)FCString::Atoi(Buffer);		}
-inline void LexFromString(int16& OutValue,		const TCHAR* Buffer)	{	OutValue = (int16)FCString::Atoi(Buffer);		}
-inline void LexFromString(int32& OutValue,		const TCHAR* Buffer)	{	OutValue = (int32)FCString::Atoi(Buffer);		}
-inline void LexFromString(int64& OutValue,		const TCHAR* Buffer)	{	OutValue = FCString::Atoi64(Buffer);	}
-inline void LexFromString(uint8& OutValue,		const TCHAR* Buffer)	{	OutValue = (uint8)FCString::Atoi(Buffer);		}
-inline void LexFromString(uint16& OutValue, 	const TCHAR* Buffer)	{	OutValue = (uint16)FCString::Atoi(Buffer);		}
-inline void LexFromString(uint32& OutValue, 	const TCHAR* Buffer)	{	OutValue = (uint32)FCString::Atoi64(Buffer);	}	//64 because this unsigned and so Atoi might overflow
-inline void LexFromString(uint64& OutValue, 	const TCHAR* Buffer)	{	OutValue = FCString::Strtoui64(Buffer, nullptr, 0); }
-inline void LexFromString(float& OutValue,		const TCHAR* Buffer)	{	OutValue = FCString::Atof(Buffer);		}
-inline void LexFromString(double& OutValue, 	const TCHAR* Buffer)	{	OutValue = FCString::Atod(Buffer);		}
-inline void LexFromString(bool& OutValue, 		const TCHAR* Buffer)	{	OutValue = FCString::ToBool(Buffer);	}
-inline void LexFromString(FString& OutValue, 	const TCHAR* Buffer)	{	OutValue = Buffer;						}
+inline void LexFromString(int8& OutValue, 				const TCHAR* Buffer)	{	OutValue = (int8)FCString::Atoi(Buffer);		}
+inline void LexFromString(int16& OutValue,				const TCHAR* Buffer)	{	OutValue = (int16)FCString::Atoi(Buffer);		}
+inline void LexFromString(int32& OutValue,				const TCHAR* Buffer)	{	OutValue = (int32)FCString::Atoi(Buffer);		}
+inline void LexFromString(int64& OutValue,				const TCHAR* Buffer)	{	OutValue = FCString::Atoi64(Buffer);	}
+inline void LexFromString(uint8& OutValue,				const TCHAR* Buffer)	{	OutValue = (uint8)FCString::Atoi(Buffer);		}
+inline void LexFromString(uint16& OutValue, 			const TCHAR* Buffer)	{	OutValue = (uint16)FCString::Atoi(Buffer);		}
+inline void LexFromString(uint32& OutValue, 			const TCHAR* Buffer)	{	OutValue = (uint32)FCString::Atoi64(Buffer);	}	//64 because this unsigned and so Atoi might overflow
+inline void LexFromString(uint64& OutValue, 			const TCHAR* Buffer)	{	OutValue = FCString::Strtoui64(Buffer, nullptr, 0); }
+inline void LexFromString(float& OutValue,				const TCHAR* Buffer)	{	OutValue = FCString::Atof(Buffer);		}
+inline void LexFromString(double& OutValue, 			const TCHAR* Buffer)	{	OutValue = FCString::Atod(Buffer);		}
+inline void LexFromString(bool& OutValue, 				const TCHAR* Buffer)	{	OutValue = FCString::ToBool(Buffer);	}
+inline void LexFromString(UE_STRING_CLASS& OutValue, 	const TCHAR* Buffer)	{	OutValue = Buffer;						}
 
 /**
  * Gets a non-owning TCHAR pointer from a string type.
@@ -2118,7 +2118,7 @@ inline void LexFromString(FString& OutValue, 	const TCHAR* Buffer)	{	OutValue =
  *     Logf(TEXT("Value: %s"), ToCStr(LexToString(Val)));
  * }
  */
-FORCEINLINE const TCHAR* ToCStr(const FString& Str)
+FORCEINLINE const TCHAR* ToCStr(const UE_STRING_CLASS& Str)
 {
 	return *Str;
 }
@@ -2131,16 +2131,16 @@ FORCEINLINE const TCHAR* ToCStr(const FString& Str)
  * @param StartSearch       The index to start searching at
  * @return the index in the given string of the closing parenthesis
  */
-CORE_API int32 FindMatchingClosingParenthesis(const FString& TargetString, const int32 StartSearch = 0);
+CORE_API int32 FindMatchingClosingParenthesis(const UE_STRING_CLASS& TargetString, const int32 StartSearch = 0);
 
 /**
-* Given a display label string, generates an FString slug that only contains valid characters for an FName.
+* Given a display label string, generates a slug string that only contains valid characters for an FName.
 * For example, "[MyObject]: Object Label" becomes "MyObjectObjectLabel" FName slug.
 *
 * @param DisplayLabel The label string to convert to an FName
 *
 * @return	The slugged string
 */
-CORE_API FString SlugStringForValidName(const FString& DisplayString, const TCHAR* ReplaceWith = TEXT(""));
+CORE_API UE_STRING_CLASS SlugStringForValidName(const UE_STRING_CLASS& DisplayString, const TCHAR* ReplaceWith = TEXT(""));
 
 #include "Misc/StringFormatArg.h"
