Fixed TNameOf<FString>.

diff --git a/Engine/Source/Runtime/Core/Public/Containers/UnrealString.h.inl b/Engine/Source/Runtime/Core/Public/Containers/UnrealString.h.inl
--- a/Engine/Source/Runtime/Core/Public/Containers/UnrealString.h.inl
+++ b/Engine/Source/Runtime/Core/Public/Containers/UnrealString.h.inl
@@ -2045,7 +2045,15 @@ public:
 };
 
 template<> struct TIsZeroConstructType<UE_STRING_CLASS> { enum { Value = true }; };
-Expose_TNameOf(UE_STRING_CLASS)
+
+template<>
+struct TNameOf<UE_STRING_CLASS>
+{
+	FORCEINLINE static TCHAR const* GetName()
+	{
+		return TEXT(PREPROCESSOR_TO_STRING(UE_STRING_CLASS));
+	}
+};
 
 inline UE_STRING_CLASS::ElementType* GetData(UE_STRING_CLASS& String)
 {
