Fix dead lock in FAutomationTestFramework::FAutomationTestMessageFilter::Serialize

#rnx

diff --git a/Engine/Source/Runtime/Core/Private/Misc/AutomationTest.cpp b/Engine/Source/Runtime/Core/Private/Misc/AutomationTest.cpp
--- a/Engine/Source/Runtime/Core/Private/Misc/AutomationTest.cpp
+++ b/Engine/Source/Runtime/Core/Private/Misc/AutomationTest.cpp
@@ -10,7 +10,6 @@
 #include "HAL/ThreadHeartBeat.h"
 #include "Internationalization/Internationalization.h"
 #include "Internationalization/Regex.h"
-#include "Logging/StructuredLog.h"
 #include "Logging/TokenizedMessage.h"
 #include "Misc/App.h"
 #include "Misc/CommandLine.h"
@@ -18,6 +17,8 @@
 #include "Misc/OutputDeviceRedirector.h"
 #include "Misc/Paths.h"
 #include "Misc/ScopedSlowTask.h"
+#include "Misc/ScopeTryLock.h"
+#include "Misc/ScopeRWLock.h"
 #include "Modules/ModuleManager.h"
 
 DEFINE_LOG_CATEGORY(LogLatentCommands)
@@ -178,7 +179,6 @@ void FAutomationTestFramework::FAutomationTestOutputDevice::Serialize( const TCH
 	FAutomationTestBase* const LocalCurTest = CurTest.load(std::memory_order_relaxed);
 	if (LocalCurTest)
 	{
-		FScopeLock Lock(&ActionCS);
 		bool CaptureLog = !LocalCurTest->SuppressLogs()
 			&& (Verbosity == ELogVerbosity::Error || Verbosity == ELogVerbosity::Warning || Verbosity == ELogVerbosity::Display)
 			&& LocalCurTest->ShouldCaptureLogCategory(Category);
@@ -193,6 +193,7 @@ void FAutomationTestFramework::FAutomationTestOutputDevice::Serialize( const TCH
 			// Errors
 			if (EffectiveVerbosity == ELogVerbosity::Error)
 			{
+				FScopeLock Lock(&ActionCS);
 				if (!LoggedFailureCause.Contains(LocalCurTest))
 				{
 					LocalCurTest->AddError(FString::Printf(TEXT("%s will be marked as failing due to errors being logged"), *LocalCurTest->GetTestFullName()), STACK_OFFSET);
@@ -248,12 +249,31 @@ void FAutomationTestFramework::FAutomationTestMessageFilter::Serialize(const TCH
 	FAutomationTestBase* const LocalCurTest = CurTest.load(std::memory_order_relaxed);
 	if (LocalDestinationContext)
 	{
-		FScopeLock Lock(&ActionCS);
 		if (LocalCurTest && LocalCurTest->IsExpectedMessage(FString(V), Verbosity))
 		{
 			Verbosity = ELogVerbosity::Verbose;
 		}
-		LocalDestinationContext->Serialize(V, Verbosity, Category, Time);
+		FScopeTryLock Lock(&ActionCS);
+		if (Lock.IsLocked() || LocalDestinationContext->CanBeUsedOnMultipleThreads())
+		{
+			LocalDestinationContext->Serialize(V, Verbosity, Category, Time);
+			if (!Backlog.IsEmpty())
+			{
+				UE::FLogRecord Item;
+				while (Backlog.Dequeue(Item))
+				{
+					LocalDestinationContext->SerializeRecord(Item);
+				}
+			}
+		}
+		else
+		{
+			UE::FLogRecord Item;
+			Item.SetFormat(V);
+			Item.SetVerbosity(Verbosity);
+			Item.SetCategory(Category);
+			Backlog.Enqueue(Item);
+		}
 	}
 }
 
@@ -266,17 +286,32 @@ void FAutomationTestFramework::FAutomationTestMessageFilter::SerializeRecord(con
 	{
 		UE::FLogRecord LocalRecord = Record;
 		const ELogVerbosity::Type Verbosity = LocalRecord.GetVerbosity();
-		FScopeLock Lock(&ActionCS);
 		if ((Verbosity == ELogVerbosity::Warning) || (Verbosity == ELogVerbosity::Error))
 		{
 			TStringBuilder<512> Line;
 			Record.FormatMessageTo(Line);
-			if (LocalCurTest->IsExpectedMessage(FString(Line), ELogVerbosity::Warning))
+			if (LocalCurTest && LocalCurTest->IsExpectedMessage(FString(Line), ELogVerbosity::Warning))
 			{
 				LocalRecord.SetVerbosity(ELogVerbosity::Verbose);
 			}
 		}
-		LocalDestinationContext->SerializeRecord(LocalRecord);
+		FScopeTryLock Lock(&ActionCS);
+		if (Lock.IsLocked() || LocalDestinationContext->CanBeUsedOnMultipleThreads())
+		{
+			LocalDestinationContext->SerializeRecord(LocalRecord);
+			if (!Backlog.IsEmpty())
+			{
+				UE::FLogRecord Item;
+				while (Backlog.Dequeue(Item))
+				{
+					LocalDestinationContext->SerializeRecord(Item);
+				}
+			}
+		}
+		else
+		{
+			Backlog.Enqueue(LocalRecord);
+		}
 	}
 }
 
@@ -917,7 +952,7 @@ bool FAutomationTestFramework::InternalStopTest(FAutomationTestExecutionInfo& Ou
 	bTestSuccessful = bTestSuccessful && !CurrentTest->HasAnyErrors() && CurrentTest->HasMetExpectedMessages();
 
 	{
-		FScopeLock Lock(&CurrentTest->ActionCS);
+		FWriteScopeLock Lock(CurrentTest->ActionCS);
 		CurrentTest->ExpectedMessages.Empty();
 	}
 
@@ -1240,7 +1275,7 @@ void FAutomationTestBase::AddError(const FString& InError, int32 StackOffset)
 {
 	if( !IsExpectedMessage(InError, ELogVerbosity::Warning))
 	{
-		FScopeLock Lock(&ActionCS);
+		FWriteScopeLock Lock(ActionCS);
 		ExecutionInfo.AddEvent(FAutomationEvent(EAutomationEventType::Error, InError), StackOffset + 1);
 	}
 }
@@ -1258,7 +1293,7 @@ void FAutomationTestBase::AddErrorS(const FString& InError, const FString& InFil
 {
 	if ( !IsExpectedMessage(InError, ELogVerbosity::Warning))
 	{
-		FScopeLock Lock(&ActionCS);
+		FWriteScopeLock Lock(ActionCS);
 		//ExecutionInfo.AddEvent(FAutomationEvent(EAutomationEventType::Error, InError, ExecutionInfo.GetContext(), InFilename, InLineNumber));
 	}
 }
@@ -1267,7 +1302,7 @@ void FAutomationTestBase::AddWarningS(const FString& InWarning, const FString& I
 {
 	if ( !IsExpectedMessage(InWarning, ELogVerbosity::Warning))
 	{
-		FScopeLock Lock(&ActionCS);
+		FWriteScopeLock Lock(ActionCS);
 		//ExecutionInfo.AddEvent(FAutomationEvent(EAutomationEventType::Warning, InWarning, ExecutionInfo.GetContext(), InFilename, InLineNumber));
 	}
 }
@@ -1276,7 +1311,7 @@ void FAutomationTestBase::AddWarning( const FString& InWarning, int32 StackOffse
 {
 	if ( !IsExpectedMessage(InWarning, ELogVerbosity::Warning))
 	{
-		FScopeLock Lock(&ActionCS);
+		FWriteScopeLock Lock(ActionCS);
 		ExecutionInfo.AddEvent(FAutomationEvent(EAutomationEventType::Warning, InWarning), StackOffset + 1);
 	}
 }
@@ -1285,26 +1320,26 @@ void FAutomationTestBase::AddInfo( const FString& InLogItem, int32 StackOffset,
 {
 	if ( !IsExpectedMessage(InLogItem, ELogVerbosity::Display))
 	{
-		FScopeLock Lock(&ActionCS);
+		FWriteScopeLock Lock(ActionCS);
 		ExecutionInfo.AddEvent(FAutomationEvent(EAutomationEventType::Info, InLogItem), StackOffset + 1, bCaptureStack);
 	}
 }
 
 void FAutomationTestBase::AddAnalyticsItem(const FString& InAnalyticsItem)
 {
-	FScopeLock Lock(&ActionCS);
+	FWriteScopeLock Lock(ActionCS);
 	ExecutionInfo.AnalyticsItems.Add(InAnalyticsItem);
 }
 
 void FAutomationTestBase::AddTelemetryData(const FString& DataPoint, double Measurement, const FString& Context)
 {
-	FScopeLock Lock(&ActionCS);
+	FWriteScopeLock Lock(ActionCS);
 	ExecutionInfo.TelemetryItems.Add(FAutomationTelemetryData(DataPoint, Measurement, Context));
 }
 
 void FAutomationTestBase::AddTelemetryData(const TMap <FString, double>& ValuePairs, const FString& Context)
 {
-	FScopeLock Lock(&ActionCS);
+	FWriteScopeLock Lock(ActionCS);
 	for (const TPair<FString, double>& Item : ValuePairs)
 	{
 		ExecutionInfo.TelemetryItems.Add(FAutomationTelemetryData(Item.Key, Item.Value, Context));
@@ -1318,7 +1353,7 @@ void FAutomationTestBase::SetTelemetryStorage(const FString& StorageName)
 
 void FAutomationTestBase::AddEvent(const FAutomationEvent& InEvent, int32 StackOffset, bool bCaptureStack)
 {
-	FScopeLock Lock(&ActionCS);
+	FWriteScopeLock Lock(ActionCS);
 	ExecutionInfo.AddEvent(InEvent, StackOffset + 1, bCaptureStack);
 }
 
@@ -1329,10 +1364,13 @@ bool FAutomationTestBase::HasAnyErrors() const
 
 bool FAutomationTestBase::HasMetExpectedMessages(ELogVerbosity::Type VerbosityType)
 {
-	FScopeLock Lock(&ActionCS);
 	bool bHasMetAllExpectedMessages = true;
-
-	for (FAutomationExpectedMessage& ExpectedMessage : ExpectedMessages)
+	TArray<FAutomationExpectedMessage> ExpectedMessagesArray;
+	{
+		FReadScopeLock RLock(ActionCS);
+		ExpectedMessagesArray = ExpectedMessages.Array();
+	}
+	for (FAutomationExpectedMessage& ExpectedMessage : ExpectedMessagesArray)
 	{
 		if (!LogCategoryMatchesSeverityInclusive(ExpectedMessage.Verbosity, VerbosityType))
 		{
@@ -1345,6 +1383,7 @@ bool FAutomationTestBase::HasMetExpectedMessages(ELogVerbosity::Type VerbosityTy
 		const bool bExpectsOneOrMore = ExpectedMessage.ExpectedNumberOfOccurrences == 0;
 		if (!bExpectsOneOrMore && (ExpectedMessage.ExpectedNumberOfOccurrences != ExpectedMessage.ActualNumberOfOccurrences))
 		{
+			FWriteScopeLock WLock(ActionCS);
 			bHasMetAllExpectedMessages = false;
 
 			ExecutionInfo.AddEvent(FAutomationEvent(EAutomationEventType::Error,
@@ -1358,6 +1397,7 @@ bool FAutomationTestBase::HasMetExpectedMessages(ELogVerbosity::Type VerbosityTy
 		}
 		else if (bExpectsOneOrMore)
 		{
+			FWriteScopeLock WLock(ActionCS);
 			if (ExpectedMessage.ActualNumberOfOccurrences == 0)
 			{
 				bHasMetAllExpectedMessages = false;
@@ -1412,7 +1452,7 @@ void FAutomationTestBase::AddExpectedMessage(
 {
 	if (Occurrences >= 0)
 	{
-		FScopeLock Lock(&ActionCS);
+		FWriteScopeLock Lock(ActionCS);
 		ExpectedMessages.Add(FAutomationExpectedMessage(ExpectedPatternString, ExpectedVerbosity, CompareType, Occurrences, IsRegex));
 	}
 	else
@@ -1739,7 +1779,7 @@ bool FAutomationTestBase::IsExpectedMessage(
 	const FString& Message,
 	const ELogVerbosity::Type& Verbosity)
 {
-	FScopeLock Lock(&ActionCS);
+	FReadScopeLock Lock(ActionCS);
 	for (FAutomationExpectedMessage& ExpectedMessage : ExpectedMessages)
 	{
 		// Maintains previous behavior: Adjust so that error and fatal messages are tested against when the input verbosity is "Warning"
diff --git a/Engine/Source/Runtime/Core/Public/Misc/AutomationTest.h b/Engine/Source/Runtime/Core/Public/Misc/AutomationTest.h
--- a/Engine/Source/Runtime/Core/Public/Misc/AutomationTest.h
+++ b/Engine/Source/Runtime/Core/Public/Misc/AutomationTest.h
@@ -24,6 +24,7 @@
 #include "HAL/ThreadSafeBool.h"
 #include "Internationalization/Regex.h"
 #include "Logging/LogVerbosity.h"
+#include "Logging/StructuredLog.h"
 #include "Math/Color.h"
 #include "Math/MathFwd.h"
 #include "Math/Rotator.h"
@@ -1234,7 +1235,8 @@ private:
 	 private:
 		 std::atomic<FAutomationTestBase*> CurTest;
 		 std::atomic<FFeedbackContext*> DestinationContext;
-		 FCriticalSection ActionCS;;
+		 FCriticalSection ActionCS;
+		 TQueue<UE::FLogRecord> Backlog;
 	 };
 
 	friend class FAutomationTestOutputDevice;
@@ -2151,7 +2153,7 @@ private:
 	TSet<FAutomationExpectedMessage> ExpectedMessages;
 
 	/** Critical section lock */
-	FCriticalSection ActionCS;
+	FRWLock ActionCS;
 };
 
 class FBDDAutomationTestBase : public FAutomationTestBase
