Back out changelist 28034836 - fix for automation test log serialization deadlock
#rnx

diff --git a/Engine/Source/Runtime/Core/Private/Misc/AutomationTest.cpp b/Engine/Source/Runtime/Core/Private/Misc/AutomationTest.cpp
--- a/Engine/Source/Runtime/Core/Private/Misc/AutomationTest.cpp
+++ b/Engine/Source/Runtime/Core/Private/Misc/AutomationTest.cpp
@@ -10,6 +10,7 @@
 #include "HAL/ThreadHeartBeat.h"
 #include "Internationalization/Internationalization.h"
 #include "Internationalization/Regex.h"
+#include "Logging/StructuredLog.h"
 #include "Logging/TokenizedMessage.h"
 #include "Misc/App.h"
 #include "Misc/CommandLine.h"
@@ -17,8 +18,6 @@
 #include "Misc/OutputDeviceRedirector.h"
 #include "Misc/Paths.h"
 #include "Misc/ScopedSlowTask.h"
-#include "Misc/ScopeTryLock.h"
-#include "Misc/ScopeRWLock.h"
 #include "Modules/ModuleManager.h"
 
 DEFINE_LOG_CATEGORY(LogLatentCommands)
@@ -179,6 +178,7 @@ void FAutomationTestFramework::FAutomationTestOutputDevice::Serialize( const TCH
 	FAutomationTestBase* const LocalCurTest = CurTest.load(std::memory_order_relaxed);
 	if (LocalCurTest)
 	{
+		FScopeLock Lock(&ActionCS);
 		bool CaptureLog = !LocalCurTest->SuppressLogs()
 			&& (Verbosity == ELogVerbosity::Error || Verbosity == ELogVerbosity::Warning || Verbosity == ELogVerbosity::Display)
 			&& LocalCurTest->ShouldCaptureLogCategory(Category);
@@ -193,7 +193,6 @@ void FAutomationTestFramework::FAutomationTestOutputDevice::Serialize( const TCH
 			// Errors
 			if (EffectiveVerbosity == ELogVerbosity::Error)
 			{
-				FScopeLock Lock(&ActionCS);
 				if (!LoggedFailureCause.Contains(LocalCurTest))
 				{
 					LocalCurTest->AddError(FString::Printf(TEXT("%s will be marked as failing due to errors being logged"), *LocalCurTest->GetTestFullName()), STACK_OFFSET);
@@ -249,31 +248,12 @@ void FAutomationTestFramework::FAutomationTestMessageFilter::Serialize(const TCH
 	FAutomationTestBase* const LocalCurTest = CurTest.load(std::memory_order_relaxed);
 	if (LocalDestinationContext)
 	{
+		FScopeLock Lock(&ActionCS);
 		if (LocalCurTest && LocalCurTest->IsExpectedMessage(FString(V), Verbosity))
 		{
 			Verbosity = ELogVerbosity::Verbose;
 		}
-		FScopeTryLock Lock(&ActionCS);
-		if (Lock.IsLocked() || LocalDestinationContext->CanBeUsedOnMultipleThreads())
-		{
-			LocalDestinationContext->Serialize(V, Verbosity, Category, Time);
-			if (!Backlog.IsEmpty())
-			{
-				UE::FLogRecord Item;
-				while (Backlog.Dequeue(Item))
-				{
-					LocalDestinationContext->SerializeRecord(Item);
-				}
-			}
-		}
-		else
-		{
-			UE::FLogRecord Item;
-			Item.SetFormat(V);
-			Item.SetVerbosity(Verbosity);
-			Item.SetCategory(Category);
-			Backlog.Enqueue(Item);
-		}
+		LocalDestinationContext->Serialize(V, Verbosity, Category, Time);
 	}
 }
 
@@ -286,32 +266,17 @@ void FAutomationTestFramework::FAutomationTestMessageFilter::SerializeRecord(con
 	{
 		UE::FLogRecord LocalRecord = Record;
 		const ELogVerbosity::Type Verbosity = LocalRecord.GetVerbosity();
+		FScopeLock Lock(&ActionCS);
 		if ((Verbosity == ELogVerbosity::Warning) || (Verbosity == ELogVerbosity::Error))
 		{
 			TStringBuilder<512> Line;
 			Record.FormatMessageTo(Line);
-			if (LocalCurTest && LocalCurTest->IsExpectedMessage(FString(Line), ELogVerbosity::Warning))
+			if (LocalCurTest->IsExpectedMessage(FString(Line), ELogVerbosity::Warning))
 			{
 				LocalRecord.SetVerbosity(ELogVerbosity::Verbose);
 			}
 		}
-		FScopeTryLock Lock(&ActionCS);
-		if (Lock.IsLocked() || LocalDestinationContext->CanBeUsedOnMultipleThreads())
-		{
-			LocalDestinationContext->SerializeRecord(LocalRecord);
-			if (!Backlog.IsEmpty())
-			{
-				UE::FLogRecord Item;
-				while (Backlog.Dequeue(Item))
-				{
-					LocalDestinationContext->SerializeRecord(Item);
-				}
-			}
-		}
-		else
-		{
-			Backlog.Enqueue(LocalRecord);
-		}
+		LocalDestinationContext->SerializeRecord(LocalRecord);
 	}
 }
 
@@ -952,7 +917,7 @@ bool FAutomationTestFramework::InternalStopTest(FAutomationTestExecutionInfo& Ou
 	bTestSuccessful = bTestSuccessful && !CurrentTest->HasAnyErrors() && CurrentTest->HasMetExpectedMessages();
 
 	{
-		FWriteScopeLock Lock(CurrentTest->ActionCS);
+		FScopeLock Lock(&CurrentTest->ActionCS);
 		CurrentTest->ExpectedMessages.Empty();
 	}
 
@@ -1275,7 +1240,7 @@ void FAutomationTestBase::AddError(const FString& InError, int32 StackOffset)
 {
 	if( !IsExpectedMessage(InError, ELogVerbosity::Warning))
 	{
-		FWriteScopeLock Lock(ActionCS);
+		FScopeLock Lock(&ActionCS);
 		ExecutionInfo.AddEvent(FAutomationEvent(EAutomationEventType::Error, InError), StackOffset + 1);
 	}
 }
@@ -1293,7 +1258,7 @@ void FAutomationTestBase::AddErrorS(const FString& InError, const FString& InFil
 {
 	if ( !IsExpectedMessage(InError, ELogVerbosity::Warning))
 	{
-		FWriteScopeLock Lock(ActionCS);
+		FScopeLock Lock(&ActionCS);
 		//ExecutionInfo.AddEvent(FAutomationEvent(EAutomationEventType::Error, InError, ExecutionInfo.GetContext(), InFilename, InLineNumber));
 	}
 }
@@ -1302,7 +1267,7 @@ void FAutomationTestBase::AddWarningS(const FString& InWarning, const FString& I
 {
 	if ( !IsExpectedMessage(InWarning, ELogVerbosity::Warning))
 	{
-		FWriteScopeLock Lock(ActionCS);
+		FScopeLock Lock(&ActionCS);
 		//ExecutionInfo.AddEvent(FAutomationEvent(EAutomationEventType::Warning, InWarning, ExecutionInfo.GetContext(), InFilename, InLineNumber));
 	}
 }
@@ -1311,7 +1276,7 @@ void FAutomationTestBase::AddWarning( const FString& InWarning, int32 StackOffse
 {
 	if ( !IsExpectedMessage(InWarning, ELogVerbosity::Warning))
 	{
-		FWriteScopeLock Lock(ActionCS);
+		FScopeLock Lock(&ActionCS);
 		ExecutionInfo.AddEvent(FAutomationEvent(EAutomationEventType::Warning, InWarning), StackOffset + 1);
 	}
 }
@@ -1320,26 +1285,26 @@ void FAutomationTestBase::AddInfo( const FString& InLogItem, int32 StackOffset,
 {
 	if ( !IsExpectedMessage(InLogItem, ELogVerbosity::Display))
 	{
-		FWriteScopeLock Lock(ActionCS);
+		FScopeLock Lock(&ActionCS);
 		ExecutionInfo.AddEvent(FAutomationEvent(EAutomationEventType::Info, InLogItem), StackOffset + 1, bCaptureStack);
 	}
 }
 
 void FAutomationTestBase::AddAnalyticsItem(const FString& InAnalyticsItem)
 {
-	FWriteScopeLock Lock(ActionCS);
+	FScopeLock Lock(&ActionCS);
 	ExecutionInfo.AnalyticsItems.Add(InAnalyticsItem);
 }
 
 void FAutomationTestBase::AddTelemetryData(const FString& DataPoint, double Measurement, const FString& Context)
 {
-	FWriteScopeLock Lock(ActionCS);
+	FScopeLock Lock(&ActionCS);
 	ExecutionInfo.TelemetryItems.Add(FAutomationTelemetryData(DataPoint, Measurement, Context));
 }
 
 void FAutomationTestBase::AddTelemetryData(const TMap <FString, double>& ValuePairs, const FString& Context)
 {
-	FWriteScopeLock Lock(ActionCS);
+	FScopeLock Lock(&ActionCS);
 	for (const TPair<FString, double>& Item : ValuePairs)
 	{
 		ExecutionInfo.TelemetryItems.Add(FAutomationTelemetryData(Item.Key, Item.Value, Context));
@@ -1353,7 +1318,7 @@ void FAutomationTestBase::SetTelemetryStorage(const FString& StorageName)
 
 void FAutomationTestBase::AddEvent(const FAutomationEvent& InEvent, int32 StackOffset, bool bCaptureStack)
 {
-	FWriteScopeLock Lock(ActionCS);
+	FScopeLock Lock(&ActionCS);
 	ExecutionInfo.AddEvent(InEvent, StackOffset + 1, bCaptureStack);
 }
 
@@ -1364,13 +1329,10 @@ bool FAutomationTestBase::HasAnyErrors() const
 
 bool FAutomationTestBase::HasMetExpectedMessages(ELogVerbosity::Type VerbosityType)
 {
+	FScopeLock Lock(&ActionCS);
 	bool bHasMetAllExpectedMessages = true;
-	TArray<FAutomationExpectedMessage> ExpectedMessagesArray;
-	{
-		FReadScopeLock RLock(ActionCS);
-		ExpectedMessagesArray = ExpectedMessages.Array();
-	}
-	for (FAutomationExpectedMessage& ExpectedMessage : ExpectedMessagesArray)
+
+	for (FAutomationExpectedMessage& ExpectedMessage : ExpectedMessages)
 	{
 		if (!LogCategoryMatchesSeverityInclusive(ExpectedMessage.Verbosity, VerbosityType))
 		{
@@ -1383,7 +1345,6 @@ bool FAutomationTestBase::HasMetExpectedMessages(ELogVerbosity::Type VerbosityTy
 		const bool bExpectsOneOrMore = ExpectedMessage.ExpectedNumberOfOccurrences == 0;
 		if (!bExpectsOneOrMore && (ExpectedMessage.ExpectedNumberOfOccurrences != ExpectedMessage.ActualNumberOfOccurrences))
 		{
-			FWriteScopeLock WLock(ActionCS);
 			bHasMetAllExpectedMessages = false;
 
 			ExecutionInfo.AddEvent(FAutomationEvent(EAutomationEventType::Error,
@@ -1397,7 +1358,6 @@ bool FAutomationTestBase::HasMetExpectedMessages(ELogVerbosity::Type VerbosityTy
 		}
 		else if (bExpectsOneOrMore)
 		{
-			FWriteScopeLock WLock(ActionCS);
 			if (ExpectedMessage.ActualNumberOfOccurrences == 0)
 			{
 				bHasMetAllExpectedMessages = false;
@@ -1452,7 +1412,7 @@ void FAutomationTestBase::AddExpectedMessage(
 {
 	if (Occurrences >= 0)
 	{
-		FWriteScopeLock Lock(ActionCS);
+		FScopeLock Lock(&ActionCS);
 		ExpectedMessages.Add(FAutomationExpectedMessage(ExpectedPatternString, ExpectedVerbosity, CompareType, Occurrences, IsRegex));
 	}
 	else
@@ -1779,7 +1739,7 @@ bool FAutomationTestBase::IsExpectedMessage(
 	const FString& Message,
 	const ELogVerbosity::Type& Verbosity)
 {
-	FReadScopeLock Lock(ActionCS);
+	FScopeLock Lock(&ActionCS);
 	for (FAutomationExpectedMessage& ExpectedMessage : ExpectedMessages)
 	{
 		// Maintains previous behavior: Adjust so that error and fatal messages are tested against when the input verbosity is "Warning"
diff --git a/Engine/Source/Runtime/Core/Public/Misc/AutomationTest.h b/Engine/Source/Runtime/Core/Public/Misc/AutomationTest.h
--- a/Engine/Source/Runtime/Core/Public/Misc/AutomationTest.h
+++ b/Engine/Source/Runtime/Core/Public/Misc/AutomationTest.h
@@ -24,7 +24,6 @@
 #include "HAL/ThreadSafeBool.h"
 #include "Internationalization/Regex.h"
 #include "Logging/LogVerbosity.h"
-#include "Logging/StructuredLog.h"
 #include "Math/Color.h"
 #include "Math/MathFwd.h"
 #include "Math/Rotator.h"
@@ -1235,8 +1234,7 @@ private:
 	 private:
 		 std::atomic<FAutomationTestBase*> CurTest;
 		 std::atomic<FFeedbackContext*> DestinationContext;
-		 FCriticalSection ActionCS;
-		 TQueue<UE::FLogRecord> Backlog;
+		 FCriticalSection ActionCS;;
 	 };
 
 	friend class FAutomationTestOutputDevice;
@@ -2153,7 +2151,7 @@ private:
 	TSet<FAutomationExpectedMessage> ExpectedMessages;
 
 	/** Critical section lock */
-	FRWLock ActionCS;
+	FCriticalSection ActionCS;
 };
 
 class FBDDAutomationTestBase : public FAutomationTestBase
