Fixed build when ENABLE_UNINITIALIZED_BOX_DIAGNOSTIC is set to 1

diff --git a/Engine/Source/Runtime/Core/Public/Math/Box.h b/Engine/Source/Runtime/Core/Public/Math/Box.h
--- a/Engine/Source/Runtime/Core/Public/Math/Box.h
+++ b/Engine/Source/Runtime/Core/Public/Math/Box.h
@@ -88,7 +88,7 @@ public:
 		{
 #if ENABLE_UNINITIALIZED_BOX_DIAGNOSTIC
 			// In case ENABLE_UNINITIALIZED_BOX_DIAGNOSTIC was explicitly enabled then report each encountered accesses to uninitialized members.
-			ensureAlwaysMsgf(TEXT("FBox contains NaN: %s. Initializing with zero extent and marking as invalid for safety."), *ToString());			
+			ensureAlwaysMsgf(false, TEXT("FBox contains NaN: %s. Initializing with zero extent and marking as invalid for safety."), *ToString());
 #else
 			// In case diagnostic check is enable from NaN diagnostic we report the error using the NaN error reporting behavior.
 			logOrEnsureNanError(TEXT("FBox contains NaN: %s. Initializing with zero extent and marking as invalid for safety."), *ToString());
