Add support for generating OnDemandToc data from UToc files on disk rather than downloading the data each time the process is launched. This feature is off by default.

#rnx

### Configuration

- New cvar 's.IasGenerateOnDemandUtoc' which when set to true will tell the IasBackend to generate the OnDemand toc file from all of the utoc files marked as supporting OnDemand currently on disk. Defaults to false.
-- When we ship the IAS feature this will be changed to true.
- Note that both cvars are used before '-dpcvars' is applied so cannot be overriden from the cmdline.
- When UE_VALIDATE_GENERATED_TOC is enabled we will download the ondemand toc from the CDN anyway and check that the generated version from disk has the same values (where important) and trigger an assert if there are differences. This is intended to help debugging future compatibility issues.

## IoDispatcher Changes

- Add new method FIoStoreReader::GetChunkCompressedInfo allowing the caller to get info about a compressed chunk without actually loading it from disk.
- Changed FIoStoreReader::Initialize to take a string view rather than raw string pointer.
- FIoStoreTocResource now has an array of meta data for chunks and another for blocks, specifically for IAS (FIoStoreTocResource::OnDemandChunkMeta/OnDemandCompressedBlockMeta)
-- This data will only be generated for utoc files with OnDemand support to avoid bloating the non-ondemand utocs.
-- This data is considered meta data and will be read when EIoStoreTocReadOptions::ReadTocMeta is enabled.
- Reordered some includes to match the IWYU tool.
- Minor optimization to FIoStoreToc::Initialize

### IAS Changes

- The IoStoreOnDemand module now includes the 'Pak'File' module so that we have access to the pak directories.
- Add FEncryptionKeyManager::GetAllKeys which would be needed if IAS data ends up being encrypted.
-- The encryption path has not yet been tested and probably will need some form of defered mounting (for cases where the encryption keys are applied at a later point)

diff --git a/Engine/Source/Runtime/Core/Internal/IO/IoStore.h b/Engine/Source/Runtime/Core/Internal/IO/IoStore.h
--- a/Engine/Source/Runtime/Core/Internal/IO/IoStore.h
+++ b/Engine/Source/Runtime/Core/Internal/IO/IoStore.h
@@ -27,6 +27,7 @@ enum class EIoStoreTocVersion : uint8
 	PartitionSize,
 	PerfectHash,
 	PerfectHashWithOverflow,
+	OnDemandMetaData,
 	LatestPlusOne,
 	Latest = LatestPlusOne - 1
 };
@@ -92,6 +93,18 @@ struct FIoStoreTocEntryMeta
 	FIoStoreTocEntryMetaFlags Flags;
 };
 
+struct FIoStoreTocOnDemandChunkMeta
+{
+	/** Hash of the chunk on disk after both compression and encryption */
+	FIoHash DiskHash;
+};
+
+struct FIoStoreTocOnDemandCompressedBlockMeta
+{
+	/** Hash of the block on disk */
+	FIoHash DiskHash;
+};
+
 /**
  * Compression block entry.
  */
@@ -193,9 +206,12 @@ struct FIoStoreTocResource
 	
 	TArray<FSHAHash> ChunkBlockSignatures;
 
+	TArray<uint8> DirectoryIndexBuffer; 
+	
 	TArray<FIoStoreTocEntryMeta> ChunkMetas;
 
-	TArray<uint8> DirectoryIndexBuffer;
+	TArray<FIoStoreTocOnDemandChunkMeta> OnDemandChunkMeta;
+	TArray<FIoStoreTocOnDemandCompressedBlockMeta> OnDemandCompressedBlockMeta;
 
 	UE_NODISCARD CORE_API static FIoStatus Read(const TCHAR* TocFilePath, EIoStoreTocReadOptions ReadOptions, FIoStoreTocResource& OutTocResource);
 
diff --git a/Engine/Source/Runtime/Core/Private/IO/IoStore.cpp b/Engine/Source/Runtime/Core/Private/IO/IoStore.cpp
--- a/Engine/Source/Runtime/Core/Private/IO/IoStore.cpp
+++ b/Engine/Source/Runtime/Core/Private/IO/IoStore.cpp
@@ -1,31 +1,32 @@
 // Copyright Epic Games, Inc. All Rights Reserved.
 
 #include "IO/IoStore.h"
-#include "Containers/Map.h"
-#include "HAL/FileManager.h"
-#include "Templates/UniquePtr.h"
-#include "Misc/Paths.h"
-#include "Misc/Compression.h"
-#include "Serialization/BufferWriter.h"
+
 #include "Async/Async.h"
+#include "Async/AsyncFileHandle.h"
 #include "Async/Future.h"
 #include "Async/ParallelFor.h"
-#include "HAL/PlatformFileManager.h"
-#include "Misc/StringBuilder.h"
+#include "Containers/Map.h"
+#include "DerivedDataCacheInterface.h"
+#include "DerivedDataPluginInterface.h"
 #include "Features/IModularFeatures.h"
-#include "Modules/ModuleManager.h"
-#include "Misc/CoreDelegates.h"
-#include "Serialization/MemoryWriter.h"
-#include "Async/AsyncFileHandle.h"
+#include "HAL/FileManager.h"
+#include "HAL/PlatformFileManager.h"
+#include "IO/IoDirectoryIndex.h"
 #include "Interfaces/ITargetPlatform.h"
+#include "Misc/Compression.h"
 #include "Misc/ConfigCacheIni.h"
-#include "DerivedDataPluginInterface.h"
-#include "DerivedDataCacheInterface.h"
+#include "Misc/CoreDelegates.h"
+#include "Misc/Paths.h"
+#include "Misc/StringBuilder.h"
+#include "Modules/ModuleManager.h"
+#include "ProfilingDebugging/CountersTrace.h"
+#include "Serialization/BufferWriter.h"
 #include "Serialization/LargeMemoryReader.h"
 #include "Serialization/MemoryWriter.h"
-#include "ProfilingDebugging/CountersTrace.h"
+#include "Serialization/MemoryWriter.h"
 #include "Tasks/Task.h"
-#include "IO/IoDirectoryIndex.h"
+#include "Templates/UniquePtr.h"
 
 DEFINE_LOG_CATEGORY(LogIoStore);
 
@@ -170,6 +171,9 @@ struct FChunkBlock
 	uint64 UncompressedSize = 0;
 	FName CompressionMethod = NAME_None;
 	FSHAHash Signature;
+
+	/** Hash of the block data as it would be found on disk */
+	FIoHash DiskHash;
 };
 
 struct FIoStoreWriteQueueEntry
@@ -179,6 +183,9 @@ struct FIoStoreWriteQueueEntry
 	IIoStoreWriteRequest* Request = nullptr;
 	FIoChunkId ChunkId;
 	FIoChunkHash ChunkHash;
+	/** Hash of the block data as it would be found on disk after compression and encryption */
+	FIoHash ChunkDiskHash;
+
 	uint64 Sequence = 0;
 	
 	// We make this optional because at the latest it might not be valid until FinishCompressionBarrior
@@ -194,6 +201,7 @@ struct FIoStoreWriteQueueEntry
 	FGraphEventRef BeginCompressionBarrier;
 	FGraphEventRef FinishCompressionBarrier;
 	FGraphEventRef FinishEncryptionAndSigningBarrier;
+	FGraphEventRef BeginOnDemandDataBarrier;
 	FGraphEventRef BeginWriteBarrier;
 	FGraphEventRef WriteFinishedEvent;
 	TAtomic<int32> CompressedBlocksCount{ 0 };
@@ -302,10 +310,12 @@ public:
 		BeginCompressionQueue.CompleteAdding();
 		BeginEncryptionAndSigningQueue.CompleteAdding();
 		FinishEncryptionAndSigningQueue.CompleteAdding();
+		OnDemandDataQueue.CompleteAdding();
 		WriterQueue.CompleteAdding();
 		BeginCompressionThread.Wait();
 		BeginEncryptionAndSigningThread.Wait();
 		FinishEncryptionAndSigningThread.Wait();
+		BeginOnDemandDataThread.Wait();
 		WriterThread.Wait();
 		if (CompressionBufferAvailableEvent)
 		{
@@ -347,6 +357,7 @@ public:
 		BeginCompressionThread = Async(EAsyncExecution::Thread, [this]() { BeginCompressionThreadFunc(); });
 		BeginEncryptionAndSigningThread = Async(EAsyncExecution::Thread, [this]() { BeginEncryptionAndSigningThreadFunc(); });
 		FinishEncryptionAndSigningThread = Async(EAsyncExecution::Thread, [this]() { FinishEncryptionAndSigningThreadFunc(); });
+		BeginOnDemandDataThread = Async(EAsyncExecution::Thread, [this]() { BeginOnDemandDataThreadFunc(); });
 		WriterThread = Async(EAsyncExecution::Thread, [this]() { WriterThreadFunc(); });
 
 		return FIoStatus::Ok;
@@ -462,6 +473,7 @@ private:
 	void BeginCompressionThreadFunc();
 	void BeginEncryptionAndSigningThreadFunc();
 	void FinishEncryptionAndSigningThreadFunc();
+	void BeginOnDemandDataThreadFunc();
 	void WriterThreadFunc();
 
 	FIoStoreWriterSettings WriterSettings;
@@ -470,10 +482,12 @@ private:
 	TFuture<void> BeginCompressionThread;
 	TFuture<void> BeginEncryptionAndSigningThread;
 	TFuture<void> FinishEncryptionAndSigningThread;
+	TFuture<void> BeginOnDemandDataThread;
 	TFuture<void> WriterThread;
 	FIoStoreWriteQueue BeginCompressionQueue;
 	FIoStoreWriteQueue BeginEncryptionAndSigningQueue;
 	FIoStoreWriteQueue FinishEncryptionAndSigningQueue;
+	FIoStoreWriteQueue OnDemandDataQueue;
 	FIoStoreWriteQueue WriterQueue;
 	TAtomic<uint64> TotalChunksCount{ 0 };
 	TAtomic<uint64> HashedChunksCount{ 0 };
@@ -540,7 +554,7 @@ public:
 
 	void Initialize()
 	{
-		ChunkIdToIndex.Empty(false);
+		ChunkIdToIndex.Empty(Toc.ChunkIds.Num());
 
 		for (int32 ChunkIndex = 0; ChunkIndex < Toc.ChunkIds.Num(); ++ChunkIndex)
 		{
@@ -548,7 +562,7 @@ public:
 		}
 	}
 
-	int32 AddChunkEntry(const FIoChunkId& ChunkId, const FIoOffsetAndLength& OffsetLength, const FIoStoreTocEntryMeta& Meta)
+	int32 AddChunkEntry(const FIoChunkId& ChunkId, const FIoOffsetAndLength& OffsetLength, const FIoStoreTocEntryMeta& Meta, const FIoStoreTocOnDemandChunkMeta& OnDemandData)
 	{
 		int32& Index = ChunkIdToIndex.FindOrAdd(ChunkId);
 
@@ -557,6 +571,7 @@ public:
 			Index = Toc.ChunkIds.Add(ChunkId);
 			Toc.ChunkOffsetLengths.Add(OffsetLength);
 			Toc.ChunkMetas.Add(Meta);
+			Toc.OnDemandChunkMeta.Add(OnDemandData);
 
 			return Index;
 		}
@@ -574,6 +589,11 @@ public:
 		return Toc.ChunkBlockSignatures.AddDefaulted_GetRef();
 	}
 
+	FIoStoreTocOnDemandCompressedBlockMeta& AddCompressionBlockMetaEntry()
+	{
+		return Toc.OnDemandCompressedBlockMeta.AddDefaulted_GetRef();
+	}
+
 	uint8 AddCompressionMethodEntry(FName CompressionMethod)
 	{
 		if (CompressionMethod == NAME_None)
@@ -825,6 +845,7 @@ public:
 		Entry->BeginCompressionBarrier = FGraphEvent::CreateGraphEvent();
 		Entry->FinishCompressionBarrier = FGraphEvent::CreateGraphEvent();
 		Entry->FinishEncryptionAndSigningBarrier = FGraphEvent::CreateGraphEvent();
+		Entry->BeginOnDemandDataBarrier = FGraphEvent::CreateGraphEvent();
 		Entry->BeginWriteBarrier = FGraphEvent::CreateGraphEvent();
 		Entry->WriteFinishedEvent = FGraphEvent::CreateGraphEvent();
 		
@@ -948,6 +969,8 @@ public:
 		OutTocOffsetAndLengths.SetNum(ChunkCount);
 		TArray<FIoStoreTocEntryMeta> OutTocChunkMetas;
 		OutTocChunkMetas.SetNum(ChunkCount);
+		TArray<FIoStoreTocOnDemandChunkMeta> OutTocOnDemandChunkMeta;
+		OutTocOnDemandChunkMeta.SetNum(ChunkCount);
 		TArray<int32> OutTocChunkHashSeeds;
 		OutTocChunkHashSeeds.SetNumZeroed(SeedCount);
 		TArray<int32> OutTocChunkIndicesWithoutPerfectHash;
@@ -1107,6 +1130,7 @@ public:
 					OutTocChunkIds[Slot] = ChunkId;
 					OutTocOffsetAndLengths[Slot] = TocResource.ChunkOffsetLengths[ChunkIndex];
 					OutTocChunkMetas[Slot] = TocResource.ChunkMetas[ChunkIndex];
+					OutTocOnDemandChunkMeta[Slot] = TocResource.OnDemandChunkMeta[ChunkIndex];
 				}
 			}
 		}
@@ -1128,6 +1152,7 @@ public:
 				OutTocChunkIds[Slot] = ChunkId;
 				OutTocOffsetAndLengths[Slot] = TocResource.ChunkOffsetLengths[ChunkIndex];
 				OutTocChunkMetas[Slot] = TocResource.ChunkMetas[ChunkIndex];
+				OutTocOnDemandChunkMeta[Slot] = TocResource.OnDemandChunkMeta[ChunkIndex];
 			}
 		}
 
@@ -1143,6 +1168,7 @@ public:
 				OutTocChunkIds[Slot] = ChunkId;
 				OutTocOffsetAndLengths[Slot] = TocResource.ChunkOffsetLengths[OverflowEntryIndex];
 				OutTocChunkMetas[Slot] = TocResource.ChunkMetas[OverflowEntryIndex];
+				OutTocOnDemandChunkMeta[Slot] = TocResource.OnDemandChunkMeta[OverflowEntryIndex];
 				OverflowEntryIndex = Slot;
 			}
 		}
@@ -1156,6 +1182,7 @@ public:
 		TocResource.ChunkIds = MoveTemp(OutTocChunkIds);
 		TocResource.ChunkOffsetLengths = MoveTemp(OutTocOffsetAndLengths);
 		TocResource.ChunkMetas = MoveTemp(OutTocChunkMetas);
+		TocResource.OnDemandChunkMeta = MoveTemp(OutTocOnDemandChunkMeta);
 		TocResource.ChunkPerfectHashSeeds = MoveTemp(OutTocChunkHashSeeds);
 		TocResource.ChunkIndicesWithoutPerfectHash = MoveTemp(OutTocChunkIndicesWithoutPerfectHash);
 
@@ -1210,7 +1237,7 @@ public:
 			TRACE_CPUPROFILER_EVENT_SCOPE(BuildIndex);
 			TArray<FStringView> FilesToIndex;
 			Toc.GetFileNamesToIndex(FilesToIndex);
-			
+
 			FString MountPoint = IoDirectoryIndexUtils::GetCommonRootPath(FilesToIndex);
 			FIoDirectoryIndexWriter DirectoryIndexWriter;
 			DirectoryIndexWriter.SetMountPoint(MountPoint);
@@ -1746,6 +1773,33 @@ private:
 		}
 	}
 
+	void BeginOnDemandData(FIoStoreWriteQueueEntry* Entry)
+	{
+		if (ContainerSettings.IsOnDemand())
+		{
+			// TODO: Worth threading this branch? (needs profiling on larger data sets)
+			FIoHashBuilder HashBuilder;
+			for (FChunkBlock& Block : Entry->ChunkBlocks)
+			{
+				check(Align(Block.CompressedSize, FAES::AESBlockSize) == Block.Size);
+
+				// Note that the IoBuffer size is not the size of the data!
+				FMemoryView BlockData = MakeMemoryView(Block.IoBuffer->Data(), Block.Size);
+		
+				HashBuilder.Update(BlockData);
+				Block.DiskHash = FIoHash::HashBuffer(BlockData);
+			}
+
+			Entry->ChunkDiskHash = HashBuilder.Finalize();
+
+			Entry->BeginWriteBarrier->DispatchSubsequents();
+		}
+		else
+		{
+			Entry->BeginWriteBarrier->DispatchSubsequents();
+		}
+	}
+
 	void WriteEntry(FIoStoreWriteQueueEntry* Entry)
 	{
 		TRACE_CPUPROFILER_EVENT_SCOPE(WriteEntry);
@@ -1824,7 +1878,7 @@ private:
 			CreatePartitionContainerFile(*TargetPartition);
 		}
 		Entry->Offset = TargetPartition->Offset;
-		
+
 		FIoOffsetAndLength OffsetLength;
 		OffsetLength.SetOffset(UncompressedFileOffset);
 		OffsetLength.SetLength(Entry->UncompressedSize.GetValue());
@@ -1855,9 +1909,17 @@ private:
 				FSHAHash& Signature = Toc.AddBlockSignatureEntry();
 				Signature = ChunkBlock.Signature;
 			}
+
+			if (ContainerSettings.IsOnDemand())
+			{
+				FIoStoreTocOnDemandCompressedBlockMeta& Meta = Toc.AddCompressionBlockMetaEntry();
+				Meta.DiskHash = ChunkBlock.DiskHash;
+			}
 		}
 
-		const int32 TocEntryIndex = Toc.AddChunkEntry(Entry->ChunkId, OffsetLength, ChunkMeta);
+		FIoStoreTocOnDemandChunkMeta OnDemandChunkMeta{ Entry->ChunkDiskHash };
+
+		const int32 TocEntryIndex = Toc.AddChunkEntry(Entry->ChunkId, OffsetLength, ChunkMeta, OnDemandChunkMeta);
 		check(TocEntryIndex != INDEX_NONE);
 
 		if (ContainerSettings.IsIndexed() && Entry->Options.FileName.Len() > 0)
@@ -2095,13 +2157,37 @@ void FIoStoreWriterContextImpl::FinishEncryptionAndSigningThreadFunc()
 		{
 			FIoStoreWriteQueueEntry* Next = Entry->Next;
 			Entry->FinishEncryptionAndSigningBarrier->Wait();
-			WriterQueue.Enqueue(Entry);
+			OnDemandDataQueue.Enqueue(Entry);
 			Entry->CompressedSize = 0;
 			for (const FChunkBlock& ChunkBlock : Entry->ChunkBlocks)
 			{
 				Entry->CompressedSize += ChunkBlock.Size;
 			}
-			Entry->BeginWriteBarrier->DispatchSubsequents();
+			Entry->BeginOnDemandDataBarrier->DispatchSubsequents();
+
+			Entry = Next;
+		}
+	}
+}
+
+void FIoStoreWriterContextImpl::BeginOnDemandDataThreadFunc()
+{
+	for (;;)
+	{
+		FIoStoreWriteQueueEntry* Entry = OnDemandDataQueue.DequeueOrWait();
+		if (!Entry)
+		{
+			return;
+		}
+
+		while (Entry)
+		{
+			FIoStoreWriteQueueEntry* Next = Entry->Next;
+			Entry->BeginOnDemandDataBarrier->Wait();
+
+			WriterQueue.Enqueue(Entry);
+
+			Entry->Writer->BeginOnDemandData(Entry);
 
 			Entry = Next;
 		}
@@ -2216,7 +2302,7 @@ public:
 		});
 	}
 
-	UE_NODISCARD FIoStatus Initialize(const TCHAR* InContainerPath, const TMap<FGuid, FAES::FAESKey>& InDecryptionKeys)
+	UE_NODISCARD FIoStatus Initialize(FStringView InContainerPath, const TMap<FGuid, FAES::FAESKey>& InDecryptionKeys)
 	{
 		TRACE_CPUPROFILER_EVENT_SCOPE(FIoStoreReader::Initialize);
 		ContainerPath = InContainerPath;
@@ -2345,6 +2431,73 @@ public:
 		}
 	}
 
+	TIoStatusOr<FIoStoreCompressedChunkInfo> GetChunkCompressedInfo(const FIoChunkId& ChunkId) const
+	{
+		const int32* TocEntryIndex = Toc.GetTocEntryIndex(ChunkId);
+		if (TocEntryIndex != nullptr)
+		{
+			// Find where in the virtual file the chunk exists.
+			const FIoOffsetAndLength* OffsetAndLength = Toc.GetOffsetAndLength(ChunkId);
+			if (!OffsetAndLength)
+			{
+				return FIoStatus(EIoErrorCode::NotFound, TEXT("Unknown chunk ID"));
+			}
+
+			const uint64 ResolvedOffset = OffsetAndLength->GetOffset();
+			const uint64 ResolvedSize = OffsetAndLength->GetLength();
+
+			// Find what compressed blocks this read straddles.
+			const FIoStoreTocResource& TocResource = Toc.GetTocResource();
+			const uint64 CompressionBlockSize = TocResource.Header.CompressionBlockSize;
+
+			const int32 FirstBlockIndex = int32(ResolvedOffset / CompressionBlockSize);
+			const int32 LastBlockIndex = int32((Align(ResolvedOffset + ResolvedSize, CompressionBlockSize) - 1) / CompressionBlockSize);
+
+			// Determine size of the result and set up output buffers
+			uint64 TotalCompressedSize = 0;
+			for (int32 BlockIndex = FirstBlockIndex; BlockIndex <= LastBlockIndex; ++BlockIndex)
+			{
+				const FIoStoreTocCompressedBlockEntry& CompressionBlock = TocResource.CompressionBlocks[BlockIndex];
+				TotalCompressedSize += CompressionBlock.GetCompressedSize();
+			}
+
+			FIoStoreCompressedChunkInfo Info;
+			if (!TocResource.OnDemandChunkMeta.IsEmpty())
+			{
+				Info.DiskHash = TocResource.OnDemandChunkMeta[*TocEntryIndex].DiskHash;
+			}
+
+			Info.Blocks.Reserve(LastBlockIndex + 1 - FirstBlockIndex);
+			Info.UncompressedOffset = ResolvedOffset % CompressionBlockSize;
+			Info.UncompressedSize = ResolvedSize;
+			Info.TotalCompressedSize = TotalCompressedSize;
+
+			// Set up the result blocks.
+			uint64 CurrentOffset = 0;
+			for (int32 BlockIndex = FirstBlockIndex; BlockIndex <= LastBlockIndex; ++BlockIndex)
+			{
+				const FIoStoreTocCompressedBlockEntry& CompressionBlock = TocResource.CompressionBlocks[BlockIndex];
+				FIoStoreCompressedBlockInfo& BlockInfo = Info.Blocks.AddDefaulted_GetRef();
+				if (!TocResource.OnDemandCompressedBlockMeta.IsEmpty())
+				{
+					BlockInfo.DiskHash = TocResource.OnDemandCompressedBlockMeta[BlockIndex].DiskHash;
+				}
+				BlockInfo.CompressionMethod = TocResource.CompressionMethods[CompressionBlock.GetCompressionMethodIndex()];
+				BlockInfo.CompressedSize = CompressionBlock.GetCompressedSize();
+				BlockInfo.UncompressedSize = CompressionBlock.GetUncompressedSize();
+				BlockInfo.OffsetInBuffer = CurrentOffset;
+				BlockInfo.AlignedSize = Align(CompressionBlock.GetCompressedSize(), FAES::AESBlockSize);
+				CurrentOffset += BlockInfo.AlignedSize;
+			}
+
+			return Info;
+		}
+		else
+		{
+			return FIoStatus(EIoErrorCode::NotFound, TEXT("Not found"));
+		}
+	}
+
 	UE::Tasks::TTask<TIoStatusOr<FIoBuffer>> ReadAsync(const FIoChunkId& ChunkId, const FIoReadOptions& Options) const
 	{
 		TRACE_CPUPROFILER_EVENT_SCOPE(ReadChunkAsync);
@@ -2688,6 +2841,10 @@ public:
 		{
 			const FIoStoreTocCompressedBlockEntry& CompressionBlock = TocResource.CompressionBlocks[BlockIndex];
 			FIoStoreCompressedBlockInfo& BlockInfo = Result.Blocks.AddDefaulted_GetRef();
+			if (!TocResource.OnDemandCompressedBlockMeta.IsEmpty())
+			{
+				BlockInfo.DiskHash = TocResource.OnDemandCompressedBlockMeta[BlockIndex].DiskHash;
+			}
 			BlockInfo.CompressionMethod = TocResource.CompressionMethods[CompressionBlock.GetCompressionMethodIndex()];
 			BlockInfo.CompressedSize = CompressionBlock.GetCompressedSize();
 			BlockInfo.UncompressedSize = CompressionBlock.GetUncompressedSize();
@@ -2851,7 +3008,7 @@ FIoStoreReader::~FIoStoreReader()
 	delete Impl;
 }
 
-FIoStatus FIoStoreReader::Initialize(const TCHAR* InContainerPath, const TMap<FGuid, FAES::FAESKey>& InDecryptionKeys)
+FIoStatus FIoStoreReader::Initialize(FStringView InContainerPath, const TMap<FGuid, FAES::FAESKey>& InDecryptionKeys)
 {
 	return Impl->Initialize(InContainerPath, InDecryptionKeys);
 }
@@ -2891,6 +3048,11 @@ TIoStatusOr<FIoStoreTocChunkInfo> FIoStoreReader::GetChunkInfo(const uint32 TocE
 	return Impl->GetChunkInfo(TocEntryIndex);
 }
 
+TIoStatusOr<FIoStoreCompressedChunkInfo> FIoStoreReader::GetChunkCompressedInfo(const FIoChunkId& Chunk) const
+{
+	return Impl->GetChunkCompressedInfo(Chunk);
+}
+
 TIoStatusOr<FIoBuffer> FIoStoreReader::Read(const FIoChunkId& Chunk, const FIoReadOptions& Options) const
 {
 	return Impl->Read(Chunk, Options);
@@ -2980,9 +3142,14 @@ FIoStatus FIoStoreTocResource::Read(const TCHAR* TocFilePath, EIoStoreTocReadOpt
 		return FIoStatusBuilder(EIoErrorCode::CorruptToc) << TEXT("Too new TOC header version while reading '") << TocFilePath << TEXT("'");
 	}
 
+	const bool bHasOnDemandMetaData = Header.Version >= static_cast<uint8>(EIoStoreTocVersion::OnDemandMetaData) && EnumHasAnyFlags(Header.ContainerFlags, EIoContainerFlags::OnDemand);
+
 	const uint64 TotalTocSize = TocFileHandle->Size() - sizeof(FIoStoreTocHeader);
 	const uint64 TocMetaSize = Header.TocEntryCount * sizeof(FIoStoreTocEntryMeta);
-	const uint64 DefaultTocSize = TotalTocSize - Header.DirectoryIndexSize - TocMetaSize;
+	const uint64 TocOnDemandMetaSize = !bHasOnDemandMetaData ? 0 :	(Header.TocEntryCount * sizeof(FIoStoreTocOnDemandChunkMeta)) +
+																	(Header.TocCompressedBlockEntryCount * sizeof(FIoStoreTocOnDemandCompressedBlockMeta));
+
+	const uint64 DefaultTocSize = TotalTocSize - (Header.DirectoryIndexSize + TocMetaSize + TocOnDemandMetaSize);
 	uint64 TocSize = DefaultTocSize;
 
 	if (EnumHasAnyFlags(ReadOptions, EIoStoreTocReadOptions::ReadTocMeta))
@@ -3069,10 +3236,10 @@ FIoStatus FIoStoreTocResource::Read(const TCHAR* TocFilePath, EIoStoreTocReadOpt
 		const int32* HashSize = reinterpret_cast<const int32*>(SignatureBuffer);
 		TArrayView<const uint8> TocSignature = MakeArrayView<const uint8>(reinterpret_cast<const uint8*>(HashSize + 1), *HashSize);
 		TArrayView<const uint8> BlockSignature = MakeArrayView<const uint8>(TocSignature.GetData() + *HashSize, *HashSize);
-		
+
 		TArrayView<const uint8> BothSignatures = MakeArrayView<const uint8>(TocSignature.GetData(), *HashSize * 2);
 		FSHA1::HashBuffer(BothSignatures.GetData(), BothSignatures.Num(), OutTocResource.SignatureHash.Hash);
-		
+
 		TArrayView<const FSHAHash> ChunkBlockSignatures = MakeArrayView<const FSHAHash>(reinterpret_cast<const FSHAHash*>(BlockSignature.GetData() + *HashSize), Header.TocCompressedBlockEntryCount);
 
 		// Adjust address to meta data
@@ -3099,11 +3266,26 @@ FIoStatus FIoStoreTocResource::Read(const TCHAR* TocFilePath, EIoStoreTocReadOpt
 	}
 
 	// Meta
-	const uint8* TocMeta = DirectoryIndexBuffer + Header.DirectoryIndexSize;
 	if (EnumHasAnyFlags(ReadOptions, EIoStoreTocReadOptions::ReadTocMeta))
 	{
+		const uint8* TocMeta = (uint8*)DirectoryIndexBuffer + Header.DirectoryIndexSize;
+
 		const FIoStoreTocEntryMeta* ChunkMetas = reinterpret_cast<const FIoStoreTocEntryMeta*>(TocMeta);
 		OutTocResource.ChunkMetas = MakeArrayView<FIoStoreTocEntryMeta const>(ChunkMetas, Header.TocEntryCount);
+
+		// OnDemand 
+		if (bHasOnDemandMetaData)
+		{
+			uint8 const* OnDemandPtr = TocMeta + TocMetaSize;
+
+			const FIoStoreTocOnDemandChunkMeta* OnDemandChunkMeta = reinterpret_cast<const FIoStoreTocOnDemandChunkMeta*>(OnDemandPtr);
+			OutTocResource.OnDemandChunkMeta = MakeArrayView<FIoStoreTocOnDemandChunkMeta const>(OnDemandChunkMeta, Header.TocEntryCount);
+
+			OnDemandPtr += sizeof(FIoStoreTocOnDemandChunkMeta) * Header.TocEntryCount;
+
+			const FIoStoreTocOnDemandCompressedBlockMeta* OnDemandCompressedBlockMeta = reinterpret_cast<const FIoStoreTocOnDemandCompressedBlockMeta*>(OnDemandPtr);
+			OutTocResource.OnDemandCompressedBlockMeta = MakeArrayView<FIoStoreTocOnDemandCompressedBlockMeta const>(OnDemandCompressedBlockMeta, Header.TocCompressedBlockEntryCount);
+		}
 	}
 
 	if (Header.Version < static_cast<uint8>(EIoStoreTocVersion::PartitionSize))
@@ -3278,16 +3460,36 @@ TIoStatusOr<uint64> FIoStoreTocResource::Write(
 		}
 	}
 
-	// Directory index
+	// Directory index (EIoStoreTocReadOptions::ReadDirectoryIndex)
 	if (EnumHasAnyFlags(TocHeader.ContainerFlags, EIoContainerFlags::Indexed))
 	{
-		TocFileHandle->Write(TocResource.DirectoryIndexBuffer.GetData(), TocResource.DirectoryIndexBuffer.Num());
+		if (!TocFileHandle->Write(TocResource.DirectoryIndexBuffer.GetData(), TocResource.DirectoryIndexBuffer.Num()))
+		{
+			return FIoStatus(EIoErrorCode::WriteError, TEXT("Failed to write directory index buffer"));
+		}
 	}
 
-	// Meta
-	if (!WriteArray(TocFileHandle.Get(), TocResource.ChunkMetas))
+	// Meta data (EIoStoreTocReadOptions::ReadTocMeta)
 	{
-		return FIoStatus(EIoErrorCode::WriteError, TEXT("Failed to write chunk meta data"));
+		// First write out the standard meta info
+		if (!WriteArray(TocFileHandle.Get(), TocResource.ChunkMetas))
+		{
+			return FIoStatus(EIoErrorCode::WriteError, TEXT("Failed to write chunk meta data"));
+		}
+
+		// Write out the metadata for the 'OnDemand' feature if the container supports this.
+		if (EnumHasAnyFlags(TocHeader.ContainerFlags, EIoContainerFlags::OnDemand))
+		{
+			if (!WriteArray(TocFileHandle.Get(), TocResource.OnDemandChunkMeta))
+			{
+				return FIoStatus(EIoErrorCode::WriteError, TEXT("Failed to write on demand chunk meta data"));
+			}
+
+			if (!WriteArray(TocFileHandle.Get(), TocResource.OnDemandCompressedBlockMeta))
+			{
+				return FIoStatus(EIoErrorCode::WriteError, TEXT("Failed to write on demand compressed block meta data"));
+			}
+		}
 	}
 
 	TocFileHandle->Flush(true);
diff --git a/Engine/Source/Runtime/Core/Public/IO/IoDispatcher.h b/Engine/Source/Runtime/Core/Public/IO/IoDispatcher.h
--- a/Engine/Source/Runtime/Core/Public/IO/IoDispatcher.h
+++ b/Engine/Source/Runtime/Core/Public/IO/IoDispatcher.h
@@ -617,20 +617,52 @@ struct FIoStoreTocCompressedBlockInfo
 
 struct FIoStoreCompressedBlockInfo
 {
+	/**
+	* Hash of the block on disk. Note that this can be all zero if the hash info was not computed when
+	* the utoc was created.
+	*/
+	FIoHash DiskHash;
+
+	/** Name of the method used to compress the block. */
 	FName CompressionMethod;
-	// The size of relevant data in the block (i.e. what you pass to decompress)
+	/** The size of relevant data in the block (i.e. what you pass to decompress). */
 	uint32 CompressedSize;
-	// The size of the _block_ after decompression. This is not adjusted for any FIoReadOptions used.
+	/** The size of the _block_ after decompression. This is not adjusted for any FIoReadOptions used. */
 	uint32 UncompressedSize;
-	// The size of the data this block takes in IoBuffer (i.e. after padding for decryption)
+	/** The size of the data this block takes in IoBuffer (i.e. after padding for decryption). */
 	uint32 AlignedSize;
-	// Where in IoBuffer this block starts.
+	/** Where in IoBuffer this block starts. */
 	uint64 OffsetInBuffer;
 };
 
+struct FIoStoreCompressedChunkInfo
+{
+	/** Info about the blocks that the chunk is split up into. */
+	TArray<FIoStoreCompressedBlockInfo> Blocks;
+
+	/**
+	* Hash of the compressed chunk on disk. Note that this can be all zero if the hash info was
+	* not computed when the utoc was created.
+	*/
+	FIoHash DiskHash;
+
+	/** There is where the data starts in IoBuffer(for when you pass in a data range via FIoReadOptions). */
+	uint64 UncompressedOffset = 0;
+	/**
+	 * This is the total size requested via FIoReadOptions. Notably, if you requested a narrow range, you could
+	 * add up all the block uncompressed sizes and it would be larger than this.
+	 */
+	uint64 UncompressedSize = 0;
+	/** This is the total size of compressed data, which is less than IoBuffer size due to padding for decryption. */
+	uint64 TotalCompressedSize = 0;
+};
+
 struct FIoStoreCompressedReadResult
 {
+	/** The buffer containing the chunk. */
 	FIoBuffer IoBuffer;
+
+	/** Info about the blocks that the chunk is split up into. */
 	TArray<FIoStoreCompressedBlockInfo> Blocks;
 	// There is where the data starts in IoBuffer (for when you pass in a data range via FIoReadOptions)
 	uint64 UncompressedOffset = 0;
@@ -641,7 +673,6 @@ struct FIoStoreCompressedReadResult
 	uint64 TotalCompressedSize = 0;
 };
 
-
 class IIoStoreWriterReferenceChunkDatabase
 {
 public:
@@ -718,14 +749,16 @@ public:
 	CORE_API FIoStoreReader();
 	CORE_API ~FIoStoreReader();
 
-	UE_NODISCARD CORE_API FIoStatus Initialize(const TCHAR* ContainerPath, const TMap<FGuid, FAES::FAESKey>& InDecryptionKeys);
+	UE_NODISCARD CORE_API FIoStatus Initialize(FStringView ContainerPath, const TMap<FGuid, FAES::FAESKey>& InDecryptionKeys);
 	CORE_API FIoContainerId GetContainerId() const;
 	CORE_API uint32 GetVersion() const;
 	CORE_API EIoContainerFlags GetContainerFlags() const;
 	CORE_API FGuid GetEncryptionKeyGuid() const;
+
 	CORE_API void EnumerateChunks(TFunction<bool(FIoStoreTocChunkInfo&&)>&& Callback) const;
 	CORE_API TIoStatusOr<FIoStoreTocChunkInfo> GetChunkInfo(const FIoChunkId& Chunk) const;
 	CORE_API TIoStatusOr<FIoStoreTocChunkInfo> GetChunkInfo(const uint32 TocEntryIndex) const;
+	CORE_API TIoStatusOr<FIoStoreCompressedChunkInfo> GetChunkCompressedInfo(const FIoChunkId& Chunk) const;
 
 	// Reads the chunk off the disk, decrypting/decompressing as necessary.
 	CORE_API TIoStatusOr<FIoBuffer> Read(const FIoChunkId& Chunk, const FIoReadOptions& Options) const;
diff --git a/Engine/Source/Runtime/Experimental/IoStoreOnDemand/IoStoreOnDemand.Build.cs b/Engine/Source/Runtime/Experimental/IoStoreOnDemand/IoStoreOnDemand.Build.cs
--- a/Engine/Source/Runtime/Experimental/IoStoreOnDemand/IoStoreOnDemand.Build.cs
+++ b/Engine/Source/Runtime/Experimental/IoStoreOnDemand/IoStoreOnDemand.Build.cs
@@ -12,7 +12,7 @@ public class IoStoreOnDemand : ModuleRules
 			new string[] {
 				"HTTP",
 				"Json",
-				"Analytics",
+				"Analytics"
 			}
 		);
 
@@ -24,5 +24,17 @@ public class IoStoreOnDemand : ModuleRules
 		{
 			PrivateDependencyModuleNames.AddRange(new string[] { "S3Client", "RSA" });
 		}
+
+		bool bFindUcasViaPakFileModule = false;
+
+		if (bFindUcasViaPakFileModule)
+		{
+			PrivateDependencyModuleNames.Add("PakFile");
+			PublicDefinitions.Add("UE_IAS_LINKPAKFILE=1");
+		}
+		else
+		{
+			PublicDefinitions.Add("UE_IAS_LINKPAKFILE=0");
+		}
 	}
 }
diff --git a/Engine/Source/Runtime/Experimental/IoStoreOnDemand/Private/EncryptionKeyManager.cpp b/Engine/Source/Runtime/Experimental/IoStoreOnDemand/Private/EncryptionKeyManager.cpp
--- a/Engine/Source/Runtime/Experimental/IoStoreOnDemand/Private/EncryptionKeyManager.cpp
+++ b/Engine/Source/Runtime/Experimental/IoStoreOnDemand/Private/EncryptionKeyManager.cpp
@@ -68,6 +68,12 @@ bool FEncryptionKeyManager::TryGetKey(const FGuid& Id, FAES::FAESKey& OutKey)
 	return false;
 }
 
+TMap<FGuid, FAES::FAESKey> FEncryptionKeyManager::GetAllKeys()
+{
+	FScopeLock _(&CriticalSection);
+	return Keys;
+}
+
 FEncryptionKeyManager& FEncryptionKeyManager::Get()
 {
 	static FEncryptionKeyManager Mgr;
diff --git a/Engine/Source/Runtime/Experimental/IoStoreOnDemand/Private/EncryptionKeyManager.h b/Engine/Source/Runtime/Experimental/IoStoreOnDemand/Private/EncryptionKeyManager.h
--- a/Engine/Source/Runtime/Experimental/IoStoreOnDemand/Private/EncryptionKeyManager.h
+++ b/Engine/Source/Runtime/Experimental/IoStoreOnDemand/Private/EncryptionKeyManager.h
@@ -26,6 +26,8 @@ public:
 	FAES::FAESKey* GetKey(const FGuid& Id);
 	/** Try retrieve the encryption key for the specified key ID. */
 	bool TryGetKey(const FGuid& Id, FAES::FAESKey& OutKey);
+	/** Returns a map of all available keys */
+	TMap<FGuid, FAES::FAESKey> GetAllKeys();
 
 	DECLARE_MULTICAST_DELEGATE_TwoParams(FEncryptionKeyAddedDelegate, const FGuid&, const FAES::FAESKey&);
 	/** Event triggered when a new key as been added. */
diff --git a/Engine/Source/Runtime/Experimental/IoStoreOnDemand/Private/OnDemandIoDispatcherBackend.cpp b/Engine/Source/Runtime/Experimental/IoStoreOnDemand/Private/OnDemandIoDispatcherBackend.cpp
--- a/Engine/Source/Runtime/Experimental/IoStoreOnDemand/Private/OnDemandIoDispatcherBackend.cpp
+++ b/Engine/Source/Runtime/Experimental/IoStoreOnDemand/Private/OnDemandIoDispatcherBackend.cpp
@@ -17,6 +17,7 @@
 #include "HAL/PreprocessorHelpers.h"
 #include "HAL/Runnable.h"
 #include "HAL/RunnableThread.h"
+#include "HAL/PlatformFileManager.h"
 #include "HttpManager.h"
 #include "IO/IoAllocators.h"
 #include "IO/IoChunkEncoding.h"
@@ -27,7 +28,10 @@
 #include "IO/IoStoreOnDemand.h"
 #include "IasCache.h"
 #include "Math/NumericLimits.h"
+#include "Misc/CommandLine.h"
 #include "Misc/EnumClassFlags.h"
+#include "Misc/Paths.h"
+#include "Misc/PathViews.h"
 #include "Misc/ScopeLock.h"
 #include "Misc/ScopeRWLock.h"
 #include "OnDemandHttpClient.h"
@@ -39,10 +43,20 @@
 
 #include <atomic>
 
+#if UE_IAS_LINKPAKFILE
+#include "IPlatformFilePak.h"
+#endif //UE_IAS_LINKPAKFILE
+
 #if !UE_BUILD_SHIPPING
 #include "Modules/ModuleManager.h"
 #endif 
 
+/** 
+ * When enabledOnDemandToc files generated from disk will be compared to the downloaded version
+ * and checked for compatibility issues. These checks will assert if a problem is found.
+ */
+#define UE_VALIDATE_GENERATED_TOC (0 && !UE_BUILD_SHIPPING)
+
 namespace UE::IO::IAS
 {
 
@@ -108,6 +122,13 @@ static TAutoConsoleVariable<bool> CVar_IoReportAnalytics(
 	true,
 	TEXT("Enables reporting statics to the analytics system"));
 
+static FAutoConsoleVariable CVar_IasGenerateOnDemandUtoc(
+	TEXT("s.IasGenerateOnDemandUtoc"),
+	false,
+	TEXT("Enables generating the FOnDemandToc from utoc files on disk rather than downloading them"),
+	ECVF_ReadOnly
+);
+
 #if !UE_BUILD_SHIPPING
 static FAutoConsoleCommand CVar_IasAbandonCache(
 	TEXT("Ias.AbandonCache"),
@@ -731,6 +752,7 @@ struct FChunkRequest
 };
 
 ///////////////////////////////////////////////////////////////////////////////
+
 static void LogIoResult(
 	const FIoChunkId& ChunkId,
 	const FIoHash& UrlHash,
@@ -760,6 +782,108 @@ static void LogIoResult(
 		CompressedOffset,
 		Priority);
 };
+
+/** Utility for finding all available on demand utoc files currently in valid pak directories */
+static TArray<FString> FindOnDemandUtocFilesOnDisk()
+{
+#if UE_IAS_LINKPAKFILE
+	// TODO: This line is all over the engine, should make it look a bit nicer
+	FPakPlatformFile* PakPlatformFile = static_cast<FPakPlatformFile*>(FPlatformFileManager::Get().FindPlatformFile(TEXT("PakFile")));
+	if (!PakPlatformFile)
+	{
+		return TArray<FString>();
+	}
+
+	TArray<FString> PakFolders;
+	PakPlatformFile->GetPakFolders(FCommandLine::Get(), PakFolders);
+
+	IPlatformFile* SearchFile = PakPlatformFile->GetLowerLevel();
+#else
+	TStringBuilder<260> PakPath;
+	FPathViews::Append(PakPath, FPaths::ProjectContentDir(), TEXT("Paks"));
+
+	TArray<FString> PakFolders;
+	PakFolders.Add(PakPath.ToString());
+
+	IPlatformFile* SearchFile = &FPlatformFileManager::Get().GetPlatformFile();
+#endif // UE_IAS_LINKPAKFILE
+
+	TArray<FString> FoundFiles;
+	for (const FString& Directory : PakFolders)
+	{
+		//PakPlatformFile
+		SearchFile->IterateDirectoryRecursively(*Directory, [&FoundFiles/*, PakPlatformFile*/](const TCHAR* Path, bool bIsDirectory) -> bool
+			{
+				if (!bIsDirectory)
+				{
+					FString Filename(Path);
+					if (Filename.EndsWith(TEXT(".utoc")) && Filename.Contains(TEXT("ondemand")))
+					{
+						// TODO: Cannot call IsPakFileInstalled at this point
+						//if (PakPlatformFile->IsPakFileInstalled(Filename))
+						{
+							FoundFiles.Emplace(MoveTemp(Filename));
+						}
+					}
+				}
+
+				return true;
+			});
+	}
+
+	return FoundFiles;
+}
+
+/** Validation code used during development to ensure that the results are correct */
+#if UE_VALIDATE_GENERATED_TOC
+
+static TArray<const FOnDemandTocContainerEntry*> SortContainers(const TArray<FOnDemandTocContainerEntry>& Containers)
+{
+	TArray<const FOnDemandTocContainerEntry*> SortedContainers;
+	for (const FOnDemandTocContainerEntry& Container : Containers)
+	{
+		SortedContainers.Add(&Container);
+	}
+
+	Algo::Sort(SortedContainers, [](const FOnDemandTocContainerEntry* LHS, const FOnDemandTocContainerEntry* RHS)->bool
+		{
+			return LHS->ContainerName < RHS->ContainerName;
+		});
+
+	return SortedContainers;
+}
+
+static bool operator == (const FOnDemandTocEntry& LHS, const FOnDemandTocEntry& RHS)
+{
+	return FMemory::Memcmp(&LHS, &RHS, sizeof(FOnDemandTocEntry)) == 0;
+}
+
+static void ValidateToc(const FOnDemandToc& RefToc, const FOnDemandToc& NewToc)
+{
+	check(RefToc.Header.BlockSize == NewToc.Header.BlockSize);
+	check(RefToc.Header.CompressionFormat == NewToc.Header.CompressionFormat);
+	check(RefToc.Header.ChunksDirectory == NewToc.Header.ChunksDirectory);
+
+	TArray<const FOnDemandTocContainerEntry*> RefContainers = SortContainers(RefToc.Containers);
+	TArray<const FOnDemandTocContainerEntry*> NewContainers = SortContainers(NewToc.Containers);
+
+	check(RefContainers.Num() == NewContainers.Num());
+
+	for (int32 Index = 0; Index < RefContainers.Num(); ++Index)
+	{
+		const FOnDemandTocContainerEntry* RefContainer = RefContainers[Index];
+		const FOnDemandTocContainerEntry* NewContainer = NewContainers[Index];
+
+		check(RefContainer->ContainerName == NewContainer->ContainerName);
+
+		check(RefContainer->Entries == NewContainer->Entries);
+		check(RefContainer->BlockSizes == NewContainer->BlockSizes);
+		check(RefContainer->BlockHashes == NewContainer->BlockHashes);
+	}
+}
+
+#endif //UE_VALIDATE_GENERATED_TOC
+
 ///////////////////////////////////////////////////////////////////////////////
 struct FBackendStatus
 {
@@ -1100,9 +1224,12 @@ public:
 	virtual uint32 Run() override;
 
 private:
+	TIoStatusOr<FOnDemandToc> GenerateOnDemandTocFromDisk() const;
+
 	FString GetEndpointTestPath() const;
 	void ConditionallyStartBackendThread();
 	void CompleteRequest(FChunkRequest* ChunkRequest);
+	FIoStatus LoadOnDemandToc(const FString& CdnUrl, const FString& TocPath);
 	void AddDeferredTocs();
 	void ProcessHttpRequests(FOnDemandHttpClient* HttpClient, FBitWindow& HttpErrors, int32 MaxConcurrentRequests);
 
@@ -1197,6 +1324,98 @@ void FOnDemandIoBackend::Shutdown()
 	BackendContext.Reset();
 }
 
+TIoStatusOr<FOnDemandToc> FOnDemandIoBackend::GenerateOnDemandTocFromDisk() const
+{
+	TRACE_CPUPROFILER_EVENT_SCOPE(FOnDemandIoBackend::GenerateOnDemandTocFromDisk);
+
+	FOnDemandToc OutToc;
+
+	TArray<FString> UtocFilePaths = FindOnDemandUtocFilesOnDisk();
+	const TMap<FGuid, FAES::FAESKey> EncryptionKeys = FEncryptionKeyManager::Get().GetAllKeys();
+
+	for (const FString& UtocFilePath : UtocFilePaths)
+	{
+		FIoStoreReader Reader;
+		FIoStatus Status = Reader.Initialize(FPathViews::GetBaseFilenameWithPath(UtocFilePath), EncryptionKeys);
+		if (Status.IsOk() && EnumHasAnyFlags(Reader.GetContainerFlags(), EIoContainerFlags::OnDemand))
+		{
+			FOnDemandTocContainerEntry Container;
+
+			Container.ContainerName = FPathViews::GetBaseFilename(UtocFilePath);
+
+			const uint32 BlockSize = Reader.GetCompressionBlockSize();
+			if (OutToc.Header.BlockSize == 0)
+			{
+				OutToc.Header.BlockSize = Reader.GetCompressionBlockSize();
+			}
+			check(OutToc.Header.BlockSize == Reader.GetCompressionBlockSize());
+
+			TArray<FIoStoreTocChunkInfo> ChunkInfos;
+			Reader.EnumerateChunks([&ChunkInfos](FIoStoreTocChunkInfo&& Info)
+				{
+					ChunkInfos.Emplace(MoveTemp(Info));
+					return true;
+				});
+
+			// We can't actually hit this until we solve the FASEKey Initialize issue above
+			if (EnumHasAnyFlags(Reader.GetContainerFlags(), EIoContainerFlags::Encrypted))
+			{
+				Container.EncryptionKeyGuid = LexToString(Reader.GetEncryptionKeyGuid());
+			}
+
+			for (const FIoStoreTocChunkInfo& ChunkInfo : ChunkInfos)
+			{
+				TIoStatusOr<FIoStoreCompressedChunkInfo> InfoStatus = Reader.GetChunkCompressedInfo(ChunkInfo.Id);
+				if (!InfoStatus.IsOk())
+				{
+					return InfoStatus.Status();
+				}
+				FIoStoreCompressedChunkInfo CompressedChunkInfo = InfoStatus.ConsumeValueOrDie();
+
+				const uint32 BlockOffset = Container.BlockSizes.Num();
+				const uint32 BlockCount = CompressedChunkInfo.Blocks.Num();
+
+				uint64 RawChunkSize = 0;
+				uint64 EncodedChunkSize = 0;
+				for (const FIoStoreCompressedBlockInfo& BlockInfo : CompressedChunkInfo.Blocks)
+				{
+					const uint64 EncodedBlockSize = Align(BlockInfo.CompressedSize, FAES::AESBlockSize);
+					Container.BlockSizes.Add(uint32(EncodedBlockSize));
+
+					FIoBlockHash BlockHash;
+					FMemory::Memcpy(&BlockHash, &BlockInfo.DiskHash, sizeof(FIoBlockHash));
+					Container.BlockHashes.Add(BlockHash);
+
+					EncodedChunkSize += EncodedBlockSize;
+					RawChunkSize += BlockInfo.UncompressedSize;
+
+					if (OutToc.Header.CompressionFormat.IsEmpty() && BlockInfo.CompressionMethod != NAME_None)
+					{
+						OutToc.Header.CompressionFormat = BlockInfo.CompressionMethod.ToString();
+					}
+				}
+
+				FOnDemandTocEntry& TocEntry = Container.Entries.AddDefaulted_GetRef();
+				TocEntry.ChunkId = ChunkInfo.Id;
+				TocEntry.Hash = CompressedChunkInfo.DiskHash;
+				TocEntry.RawSize = RawChunkSize;
+				TocEntry.EncodedSize = EncodedChunkSize;
+				TocEntry.BlockOffset = BlockOffset;
+				TocEntry.BlockCount = BlockCount;
+			}
+
+			OutToc.Containers.Emplace(MoveTemp(Container));
+		}
+	}
+
+	if (!OutToc.Containers.IsEmpty())
+	{
+		OutToc.Header.ChunksDirectory = FString::Printf(TEXT("IoChunksV%u"), EOnDemandChunkVersion::Latest).ToLower();
+	}
+
+	return OutToc;
+}
+
 FString FOnDemandIoBackend::GetEndpointTestPath() const
 {
 	FString TestPath = IoStore.IsValid() ? IoStore->GetFirstTocPath() : FString();
@@ -1485,6 +1704,54 @@ TIoStatusOr<FIoMappedRegion> FOnDemandIoBackend::OpenMapped(const FIoChunkId& Ch
 	return FIoStatus::Unknown;
 }
 
+FIoStatus FOnDemandIoBackend::LoadOnDemandToc(const FString& CdnUrl, const FString& TocPath)
+{
+	TRACE_CPUPROFILER_EVENT_SCOPE(FOnDemandIoBackend::LoadOnDemandToc);
+
+	if (CVar_IasGenerateOnDemandUtoc->GetBool())
+	{
+		UE_LOG(LogIas, Log, TEXT("Generating OnDemandToc from disk"));
+
+		// TODO: We don't really handle the case where the utoc on disk are encrypted and we don't have the key yet/
+		// See the 'DeferredTocs' code below.
+		TIoStatusOr<FOnDemandToc> GeneratedTocResult = GenerateOnDemandTocFromDisk();
+
+		if (GeneratedTocResult.IsOk())
+		{
+#if UE_VALIDATE_GENERATED_TOC == 0
+			IoStore->AddToc(TocPath, GeneratedTocResult.ConsumeValueOrDie());
+#else
+			FOnDemandToc UrlToc = LoadTocFromUrl(CdnUrl, TocPath, 1).ConsumeValueOrDie();
+			FOnDemandToc GeneratedToc = GeneratedTocResult.ConsumeValueOrDie();
+
+			UE::IO::IAS::ValidateToc(UrlToc, GeneratedToc);
+
+			IoStore->AddToc(TocPath, MoveTemp(GeneratedToc));
+#endif // UE_VALIDATE_GENERATED_TOC
+		}
+		else
+		{
+			return GeneratedTocResult.Status();
+		}
+	}
+	else
+	{
+		UE_LOG(LogIas, Log, TEXT("Downloading OnDemandToc from CDN"));
+
+		TIoStatusOr<FOnDemandToc> TocResult = LoadTocFromUrl(CdnUrl, TocPath, 1);
+		if (TocResult.IsOk())
+		{
+			IoStore->AddToc(TocPath, TocResult.ConsumeValueOrDie());
+		}
+		else
+		{
+			return TocResult.Status();
+		}
+	}
+
+	return EIoErrorCode::Ok;
+}
+
 void FOnDemandIoBackend::AddDeferredTocs()
 {
 	TArray<FString> TocPaths;
@@ -1499,14 +1766,10 @@ void FOnDemandIoBackend::AddDeferredTocs()
 
 	for (const FString& TocPath : TocPaths)
 	{
-		TIoStatusOr<FOnDemandToc> Toc = LoadTocFromUrl(AvailableEps.GetCurrent(), TocPath, 1);
-		if (Toc.IsOk())
+		FIoStatus Result = LoadOnDemandToc(AvailableEps.GetCurrent(), TocPath);
+		if (!Result.IsOk())
 		{
-			IoStore->AddToc(TocPath, Toc.ConsumeValueOrDie());
-		}
-		else
-		{
-			UE_LOG(LogIas, Error, TEXT("Failed to add TOC '%s/%s', reason '%s'"), *AvailableEps.GetCurrent(), *TocPath, *Toc.Status().ToString());
+			UE_LOG(LogIas, Error, TEXT("Failed to add TOC '%s/%s', reason '%s'"), *AvailableEps.GetCurrent(), *TocPath, *Result.ToString());
 		}
 	}
 }
@@ -1546,14 +1809,11 @@ void FOnDemandIoBackend::Mount(const FOnDemandEndpoint& Endpoint)
 		}
 
 		check(AvailableEps.HasCurrent());
-		TIoStatusOr<FOnDemandToc> Toc = LoadTocFromUrl(AvailableEps.GetCurrent(), Endpoint.TocPath, 1);
-		if (Toc.IsOk())
-		{
-			IoStore->AddToc(Endpoint.TocPath, Toc.ConsumeValueOrDie());
-		}
-		else
+
+		FIoStatus Result = LoadOnDemandToc(AvailableEps.GetCurrent(), Endpoint.TocPath);
+		if (!Result.IsOk())
 		{
-			UE_LOG(LogIas, Error, TEXT("Deferring TOC '%s/%s' due to '%s'"), *AvailableEps.GetCurrent(), *Endpoint.TocPath, *Toc.Status().ToString());
+			UE_LOG(LogIas, Error, TEXT("Deferring TOC '%s/%s' due to '%s'"), *AvailableEps.GetCurrent(), *Endpoint.TocPath, *Result.ToString());
 			BackendStatus.SetHttpError(true);
 			FWriteScopeLock _(Lock);
 			DeferredTocs.Add(Endpoint.TocPath);
@@ -1785,3 +2045,5 @@ TSharedPtr<IOnDemandIoDispatcherBackend> MakeOnDemandIoDispatcherBackend(TUnique
 }
 
 } // namespace UE::IO::IAS
+
+#undef UE_VALIDATE_GENERATED_TOC
