FVector2 - Added FVector2<T>::Clamp()
* Using FMath::Clamp() with vectors would result in warnings and incorrect behavior

diff --git a/Engine/Source/Runtime/Core/Public/Math/Vector2D.h b/Engine/Source/Runtime/Core/Public/Math/Vector2D.h
--- a/Engine/Source/Runtime/Core/Public/Math/Vector2D.h
+++ b/Engine/Source/Runtime/Core/Public/Math/Vector2D.h
@@ -236,7 +236,7 @@ public:
 	* @param Other The vector to compare against.
 	* @return true if this is the smaller vector, otherwise false.
 	*/
-	UE_DEPRECATED(5.1, "TVector2 comparison operators are deprecated. Use ComponentwiseAllLessThan. For componentwise min/max/abs, use TVector2::{Min,Max,GetAbs}, FMath::{Min,Max,Abs} compute something different.")
+	UE_DEPRECATED(5.1, "TVector2 comparison operators are deprecated. Use ComponentwiseAllLessThan. For componentwise min/max/abs/clamp, use TVector2::{Min,Max,GetAbs,Clamp}, FMath::{Min,Max,Abs,Clamp} compute something different.")
 	bool operator<(const TVector2<T>& Other) const
 	{
 		return ComponentwiseAllLessThan(Other);
@@ -248,7 +248,7 @@ public:
 	* @param Other The vector to compare against.
 	* @return true if this is the larger vector, otherwise false.
 	*/
-	UE_DEPRECATED(5.1, "TVector2 comparison operators are deprecated. Use ComponentwiseAllGreaterThan. For componentwise min/max/abs, use TVector2::{Min,Max,GetAbs}, FMath::{Min,Max,Abs} compute something different.")
+	UE_DEPRECATED(5.1, "TVector2 comparison operators are deprecated. Use ComponentwiseAllGreaterThan. For componentwise min/max/abs/clamp, use TVector2::{Min,Max,GetAbs,Clamp}, FMath::{Min,Max,Abs,Clamp} compute something different.")
 	bool operator>(const TVector2<T>& Other) const
 	{
 		return ComponentwiseAllGreaterThan(Other);
@@ -260,7 +260,7 @@ public:
 	* @param Other The vector to compare against.
 	* @return true if this vector is less than or equal to the other vector, otherwise false.
 	*/
-	UE_DEPRECATED(5.1, "TVector2 comparison operators are deprecated. Use ComponentwiseAllLessOrEqual. For componentwise min/max/abs, use TVector2::{Min,Max,GetAbs}, FMath::{Min,Max,Abs} compute something different.")
+	UE_DEPRECATED(5.1, "TVector2 comparison operators are deprecated. Use ComponentwiseAllLessOrEqual. For componentwise min/max/abs/clamp, use TVector2::{Min,Max,GetAbs,Clamp}, FMath::{Min,Max,Abs,Clamp} compute something different.")
 	bool operator<=(const TVector2<T>& Other) const
 	{
 		return ComponentwiseAllLessOrEqual(Other);
@@ -272,7 +272,7 @@ public:
 	* @param Other The vector to compare against.
 	* @return true if this vector is greater than or equal to the other vector, otherwise false.
 	*/
-	UE_DEPRECATED(5.1, "TVector2 comparison operators are deprecated. Use ComponentwiseAllGreaterOrEqual. For componentwise min/max/abs, use TVector2::{Min,Max,GetAbs}, FMath::{Min,Max,Abs} compute something different.")
+	UE_DEPRECATED(5.1, "TVector2 comparison operators are deprecated. Use ComponentwiseAllGreaterOrEqual. For componentwise min/max/abs/clamp, use TVector2::{Min,Max,GetAbs,Clamp}, FMath::{Min,Max,Abs,Clamp} compute something different.")
 	bool operator>=(const TVector2<T>& Other) const
 	{
 		return ComponentwiseAllGreaterOrEqual(Other);
@@ -454,6 +454,16 @@ public:
 	*/
 	FORCEINLINE static TVector2<T> Min(const TVector2<T>& A, const TVector2<T>& B);
 
+	/**
+	* Returns a vector with each component clamped between a minimum and a maximum.
+	*
+	* @param V The vector to clamp.
+	* @param MinValue The minimum vector.
+	* @param MaxValue The maximum vector.
+	* @return The clamped vector.
+	*/
+	FORCEINLINE static TVector2<T> Clamp(const TVector2<T>& V, const TVector2<T>& MinValue, const TVector2<T>& MaxValue);
+
 	/**
 	* Checks for equality with error-tolerant comparison.
 	*
@@ -928,6 +938,12 @@ FORCEINLINE TVector2<T> TVector2<T>::Min(const TVector2<T>& A, const TVector2<T>
 	return TVector2<T>(FMath::Min(A.X, B.X), FMath::Min(A.Y, B.Y));
 }
 
+template<typename T>
+FORCEINLINE TVector2<T> TVector2<T>::Clamp(const TVector2<T>& V, const TVector2<T>& MinValue, const TVector2<T>& MaxValue)
+{
+	return TVector2<T>(FMath::Clamp(V.X, MinValue.X, MaxValue.X), FMath::Clamp(V.Y, MinValue.Y, MaxValue.Y));
+}
+
 template<typename T>
 FORCEINLINE bool TVector2<T>::operator==(const TVector2<T>& V) const
 {
