New Container Classes:

- TWeakCastable<T, CastTo>
- TStrongCastable<T, CastTo>

- TWeakContainerTo<CastTo>
- TStrongContainerTo<CastTo>

- TPointerContainer<CastTo>

These are Pointer Containers for a 'T' (UObject xOr TSharedFromThis) that can be statically casted to 'CastTo'
Useful if you want your system to support both memory management systems: (1) UObject Garbage Collection and (2) Smart Pointer Reference counting

diff --git a/Engine/Plugins/Experimental/ControlFlows/Source/ControlFlows/Public/ControlFlowContainer.h b/Engine/Plugins/Experimental/ControlFlows/Source/ControlFlows/Public/ControlFlowContainer.h
--- a/Engine/Plugins/Experimental/ControlFlows/Source/ControlFlows/Public/ControlFlowContainer.h
+++ b/Engine/Plugins/Experimental/ControlFlows/Source/ControlFlows/Public/ControlFlowContainer.h
@@ -5,6 +5,7 @@
 // HEADER_UNIT_SKIP - Not included directly
 
 #include "Templates/SharedPointer.h"
+#include "Containers/CastedTo.h"
 
 class FControlFlowContainerBase : public TSharedFromThis<FControlFlowContainerBase>
 {
@@ -36,47 +37,19 @@ private:
 	FString FlowName;
 };
 
-template<typename OwningObjectT, bool /*bDerivedFromUObject*/, bool /*bDerivedFromTSharedFromThis*/>
-class FControlFlowContainerInternal;
-
-template<typename OwningObjectT>
-class FControlFlowContainerInternal<OwningObjectT, true, false> : public FControlFlowContainerBase
+template<typename T>
+class TControlFlowContainer : public FControlFlowContainerBase, public TWeakContainer<T>
 {
 public:
-	FControlFlowContainerInternal() = delete;
-	FControlFlowContainerInternal(OwningObjectT* InOwner, TSharedRef<FControlFlow> InFlow, const FString& FlowId)
+	TControlFlowContainer() = delete;
+	TControlFlowContainer(T* InOwner, TSharedRef<FControlFlow> InFlow, const FString& FlowId)
 		: FControlFlowContainerBase(InFlow, FlowId)
-		, OwningObject(InOwner)
+		, TWeakContainer<T>(InOwner)
 	{}
-protected:
-	virtual const void* const GetOwningObject() const override { return OwningObject.IsValid() ? OwningObject.Get() : nullptr; }
-private:
-	TWeakObjectPtr<const OwningObjectT> OwningObject;
-};
 
-template<typename OwningObjectT>
-class FControlFlowContainerInternal<OwningObjectT, false, true> : public FControlFlowContainerBase
-{
-public:
-	FControlFlowContainerInternal() = delete;
-	FControlFlowContainerInternal(OwningObjectT* InOwner, TSharedRef<FControlFlow> InFlow, const FString& FlowId)
-		: FControlFlowContainerBase(InFlow, FlowId)
-		, OwningObject(StaticCastSharedRef<OwningObjectT>(InOwner->AsShared()))
-	{}
-protected:
-	virtual const void* const GetOwningObject() const override { return OwningObject.IsValid() ? OwningObject.Pin().Get() : nullptr; }
 private:
-	TWeakPtr<const OwningObjectT> OwningObject;
-};
-
-#define FControlFlowContainerInternal_Decl FControlFlowContainerInternal<ExternalObjectT, TIsDerivedFrom<ExternalObjectT, UObject>::IsDerived, IsDerivedFromSharedFromThis<ExternalObjectT>()>
-
-template<typename ExternalObjectT>
-class TControlFlowContainer : public FControlFlowContainerInternal_Decl
-{
-public:
-	TControlFlowContainer() = delete;
-	TControlFlowContainer(ExternalObjectT* InOwner, TSharedRef<FControlFlow> InFlow, const FString& FlowId) : FControlFlowContainerInternal_Decl(InOwner, InFlow, FlowId) {}
-};
-
-#undef FControlFlowContainerInternal_Decl
\ No newline at end of file
+	virtual const void* const GetOwningObject() const override final
+	{
+		return static_cast<const TWeakContainerTo<T>*>(this)->Cast();
+	}
+};
\ No newline at end of file
diff --git a/Engine/Source/Runtime/Core/Public/Containers/CastedTo.h b/Engine/Source/Runtime/Core/Public/Containers/CastedTo.h
new file mode 100644
--- /dev/null
+++ b/Engine/Source/Runtime/Core/Public/Containers/CastedTo.h
@@ -0,0 +1,152 @@
+// Copyright Epic Games, Inc. All Rights Reserved.
+
+#pragma once
+
+#include "UObject/StrongObjectPtr.h"
+#include "UObject/WeakObjectPtr.h"
+#include "Templates/SharedPointer.h"
+
+/* A Pointer Container for a 'T' (UObject xOr TSharedFromThis) that can be statically casted to 'CastTo'
+ * These classes are useful if you want your system to support both memory management systems: (1) UObject Garbage Collection and (2) Smart Pointer Reference counting
+ */
+template<typename CastTo>
+class TPointerContainer : public TSharedFromThis<TPointerContainer<CastTo>>
+{
+public:
+	virtual ~TPointerContainer() {}
+
+	virtual CastTo* Cast() const = 0;
+	virtual bool IsValid() const = 0;
+	virtual void Reset() = 0;
+
+public:
+	CastTo* Get() const { return Cast(); }
+	void IsPtrValid() const { return IsValid(); }
+	void ResetPtr() { Reset(); }
+};
+
+/* A "Weak" Pointer Reference Container for a 'T' (UObject xOr TSharedFromThis) that can be statically casted to 'CastTo' */
+template<typename CastTo>
+class TWeakContainerTo : public TPointerContainer<CastTo>
+{
+public:
+	virtual ~TWeakContainerTo() {}
+};
+
+template<typename T, typename CastTo, bool /*bDerivedFromUObject*/, bool /*bDerivedFromTSharedFromThis*/>
+class TWeakCastableImpl;
+
+template<typename T, typename CastTo>
+class TWeakCastableImpl<T, CastTo, true, false> : public TWeakContainerTo<CastTo>
+{
+public:
+	TWeakCastableImpl() {}
+	TWeakCastableImpl(T* InObject) : Object(InObject) {}
+	virtual ~TWeakCastableImpl() {}
+
+public:
+	virtual CastTo* Cast() const override final { return static_cast<CastTo*>(Object.Get()); }
+	virtual bool IsValid() const override final { return Object.IsValid(); }
+	virtual void Reset() override final { Object.Reset(); }
+
+private:
+	TWeakObjectPtr<T> Object;
+};
+
+template<typename T, typename CastTo>
+class TWeakCastableImpl<T, CastTo, false, true> : public TWeakContainerTo<CastTo>
+{
+public:
+	TWeakCastableImpl() {}
+	TWeakCastableImpl(T* InObject) : Object(StaticCastSharedRef<T>(InObject->AsShared())) {}
+	virtual ~TWeakCastableImpl() {}
+
+public:
+	virtual CastTo* Cast() const override final { return static_cast<CastTo*>(IsValid() ? Object.Pin().Get() : nullptr); }
+	virtual bool IsValid() const override final { return Object.IsValid(); }
+	virtual void Reset() override final { Object.Reset(); }
+
+private:
+	TWeakPtr<T> Object;
+};
+
+#define TWeakCastableImpl_Decl TWeakCastableImpl<CastFrom, CastTo, TIsDerivedFrom<CastFrom, UObject>::IsDerived, IsDerivedFromSharedFromThis<CastFrom>()>
+template<typename CastFrom, typename CastTo>
+class TWeakCastable : public TWeakCastableImpl_Decl
+{
+public:
+	TWeakCastable() {}
+	TWeakCastable(CastFrom* InObject) : TWeakCastableImpl_Decl(InObject) {}
+};
+#undef TWeakCastableImpl_Decl
+
+template<typename T>
+class TWeakContainer : public TWeakCastable<T, T>
+{
+public:
+	TWeakContainer() : TWeakCastable<T, T>() {}
+	TWeakContainer(T* InObject) : TWeakCastable<T, T>(InObject) {}
+};
+
+/* A "Strong" Pointer Reference Container for a 'T' (UObject xOr TSharedFromThis) that can be statically casted to 'CastTo' */
+template<typename CastTo>
+class TStrongContainerTo : public TPointerContainer<CastTo>
+{
+public:
+	virtual ~TStrongContainerTo() {}
+};
+
+template<typename T, typename CastTo, bool /*bDerivedFromUObject*/, bool /*bDerivedFromTSharedFromThis*/>
+class TStrongContainerToImpl;
+
+template<typename T, typename CastTo>
+class TStrongContainerToImpl<T, CastTo, true, false> : public TStrongContainerTo<CastTo>
+{
+public:
+	TStrongContainerToImpl() {}
+	TStrongContainerToImpl(T* InObject) : Object(InObject) {}
+	virtual ~TStrongContainerToImpl() {}
+
+public:
+	virtual CastTo* Cast() const override final { return static_cast<CastTo*>(Object.Get()); }
+	virtual bool IsValid() const override final { return Object.IsValid(); }
+	virtual void Reset() override final { Object.Reset(); }
+
+private:
+	TStrongObjectPtr<T> Object;
+};
+
+template<typename T, typename CastTo>
+class TStrongContainerToImpl<T, CastTo, false, true> : public TStrongContainerTo<CastTo>
+{
+public:
+	TStrongContainerToImpl() {}
+	TStrongContainerToImpl(T* InObject) : Object(StaticCastSharedRef<T>(InObject->AsShared())) {}
+	virtual ~TStrongContainerToImpl() {}
+
+public:
+	virtual CastTo* Cast() const override final { return static_cast<CastTo*>(Object.IsValid() ? Object.Get() : nullptr); }
+	virtual bool IsValid() const override final { return Object.IsValid(); }
+	virtual void Reset() override final { Object.Reset(); }
+
+private:
+	TSharedPtr<T> Object;
+};
+
+#define TStrongContainerToImpl_Decl TStrongContainerToImpl<CastFrom, CastTo, TIsDerivedFrom<CastFrom, UObject>::IsDerived, IsDerivedFromSharedFromThis<CastFrom>()>
+template<typename CastFrom, typename CastTo>
+class TStrongCastable : public TStrongContainerToImpl_Decl
+{
+public:
+	TStrongCastable() : TStrongContainerToImpl_Decl() {}
+	TStrongCastable(CastFrom* InObject) : TStrongContainerToImpl_Decl(InObject) {}
+};
+#undef TStrongContainerToImpl_Decl
+
+template<typename T>
+class TStrongContainer : public TStrongCastable<T, T>
+{
+public:
+	TStrongContainer() : TStrongCastable<T, T>() {}
+	TStrongContainer(T* InObject) : TStrongCastable<T, T>(InObject) {}
+};
\ No newline at end of file
