Fix log spam about degenerate triangles when computing barycentric coordinates on a triangle + minor cleanup:

1. Add a version of ComputeBaryCentric2D that (1) returns false on failure, rather than logging a warning, and (2) is less confusingly-named.  Also minorly clean up the shared code for this new function and ComputeBaryCentric2D.

2. Update SegmentTriangleIntersection to use the new barycentric function and to exit earlier on degenerate triangles.

diff --git a/Engine/Source/Runtime/Core/Private/Math/UnrealMath.cpp b/Engine/Source/Runtime/Core/Private/Math/UnrealMath.cpp
--- a/Engine/Source/Runtime/Core/Private/Math/UnrealMath.cpp
+++ b/Engine/Source/Runtime/Core/Private/Math/UnrealMath.cpp
@@ -1920,7 +1920,10 @@ bool FMath::SegmentTriangleIntersection(const FVector& StartPoint, const FVector
 	FVector Edge2(C - A);
 	Edge2.Normalize();
 	FVector TriNormal = Edge2 ^ Edge1;
-	TriNormal.Normalize();
+	if (!TriNormal.Normalize())
+	{
+		return false;
+	}
 
 	bool bCollide = FMath::SegmentPlaneIntersection(StartPoint, EndPoint, FPlane(A, TriNormal), OutIntersectPoint);
 	if (!bCollide)
@@ -1928,10 +1931,9 @@ bool FMath::SegmentTriangleIntersection(const FVector& StartPoint, const FVector
 		return false;
 	}
 
-	FVector BaryCentric = FMath::ComputeBaryCentric2D(OutIntersectPoint, A, B, C);
-
-	// ComputeBaryCenteric2D returns ZeroVector when the triangle is too small.
-	if (BaryCentric == FVector::ZeroVector)
+	// ComputeBarycentricTri returns false when the triangle is too small.
+	FVector BaryCentric;
+	if (!FMath::ComputeBarycentricTri(OutIntersectPoint, A, B, C, BaryCentric))
 	{
 		return false;
 	}
@@ -2221,33 +2223,48 @@ FVector FMath::GetBaryCentric2D(const FVector2D& Point, const FVector2D& A, cons
 	return FVector(a, b, 1.0f - a - b);
 }
 
-FVector FMath::ComputeBaryCentric2D(const FVector& Point, const FVector& A, const FVector& B, const FVector& C)
+bool FMath::ComputeBarycentricTri(const FVector& Point, const FVector& A, const FVector& B, const FVector& C, FVector& OutBarycentric, double Tolerance)
 {
-	// Compute the normal of the triangle
+	// Tolerance cannot be negative
+	checkSlow(Tolerance >= 0);
+
+	// Compute the normal direction of the triangle (scaled by 2x area)
 	const FVector TriNorm = (B-A) ^ (C-A);
 
-	// Check the size of the triangle is reasonable (TriNorm.Size() will be twice the triangle area)
-	if(TriNorm.SizeSquared() <= UE_SMALL_NUMBER)
+	// Check if the triangle is too small, according to the tolerance
+	const FVector::FReal TriNormSizeSquared = TriNorm.SizeSquared();
+	if(TriNormSizeSquared <= Tolerance)
 	{
-		UE_LOG(LogUnrealMath, Warning, TEXT("Small triangle detected in FMath::ComputeBaryCentric2D(), can't compute valid barycentric coordinate."));
-		return FVector(0.0f, 0.0f, 0.0f);
+		return false;
 	}
 
-	const FVector N = TriNorm.GetSafeNormal();
+	// Compute 1 / twice area of triangle ABC
+	const FVector::FReal AreaABCInv = FMath::InvSqrt(TriNormSizeSquared);
 
-	// Compute twice area of triangle ABC
-	const FVector::FReal AreaABCInv = 1.0f / (N | TriNorm);
+	// Compute unit-length Tri normal
+	const FVector N = TriNorm * AreaABCInv;
 
-	// Compute a contribution
+	// Compute a contribution (in X)
 	const FVector::FReal AreaPBC = N | ((B-Point) ^ (C-Point));
-	const FVector::FReal a = AreaPBC * AreaABCInv;
+	OutBarycentric.X = AreaPBC * AreaABCInv;
 
-	// Compute b contribution
+	// Compute b contribution (in Y)
 	const FVector::FReal AreaPCA = N | ((C-Point) ^ (A-Point));
-	const FVector::FReal b = AreaPCA * AreaABCInv;
+	OutBarycentric.Y = AreaPCA * AreaABCInv;
 
-	// Compute c contribution
-	return FVector(a, b, 1.0f - a - b);
+	// Compute c contribution (in Z)
+	OutBarycentric.Z = 1.0 - OutBarycentric.X - OutBarycentric.Y;
+	return true;
+}
+
+FVector FMath::ComputeBaryCentric2D(const FVector& Point, const FVector& A, const FVector& B, const FVector& C)
+{
+	FVector ToRet(0, 0, 0);
+	if (!ComputeBarycentricTri(Point, A, B, C, ToRet, UE_DOUBLE_SMALL_NUMBER))
+	{
+		UE_LOG(LogUnrealMath, Warning, TEXT("Small triangle detected in FMath::ComputeBaryCentric2D(); can't compute valid barycentric coordinate."));
+	}
+	return ToRet;
 }
 
 FVector4 FMath::ComputeBaryCentric3D(const FVector& Point, const FVector& A, const FVector& B, const FVector& C, const FVector& D)
diff --git a/Engine/Source/Runtime/Core/Public/Math/UnrealMathUtility.h b/Engine/Source/Runtime/Core/Public/Math/UnrealMathUtility.h
--- a/Engine/Source/Runtime/Core/Public/Math/UnrealMathUtility.h
+++ b/Engine/Source/Runtime/Core/Public/Math/UnrealMathUtility.h
@@ -2295,7 +2295,8 @@ public:
 	UE_NODISCARD static CORE_API FVector GetBaryCentric2D(const FVector2D& Point, const FVector2D& A, const FVector2D& B, const FVector2D& C);
 
 	/**
-	 * Computes the barycentric coordinates for a given point in a triangle
+	 * Computes the barycentric coordinates for a given point in a 3D triangle.
+	 * Note: Prefer the more accurately-named ComputeBarycentricTri instead.
 	 *
 	 * @param	Point			point to convert to barycentric coordinates (in plane of ABC)
 	 * @param	A,B,C			three non-collinear points defining a triangle in CCW
@@ -2305,6 +2306,19 @@ public:
 	 */
 	UE_NODISCARD static CORE_API FVector ComputeBaryCentric2D(const FVector& Point, const FVector& A, const FVector& B, const FVector& C);
 
+	/**
+	 * Computes the barycentric coordinates for a given point in a triangle
+	 *
+	 * @param	Point			point to convert to barycentric coordinates (in plane of ABC)
+	 * @param	A,B,C			three non-collinear points defining a triangle in CCW
+	 * @param	OutBarycentric	Vector containing the three weights a,b,c such that Point = a*A + b*B + c*C
+	 *							                               or Point = A + b*(B-A) + c*(C-A) = (1-b-c)*A + b*B + c*C
+	 * @param	Tolerance		Tolerance for ignoring too-small triangles
+	 * 
+	 * @return false if the result could not be computed (occurs when the triangle area is too small)
+	 */
+	UE_NODISCARD static CORE_API bool ComputeBarycentricTri(const FVector& Point, const FVector& A, const FVector& B, const FVector& C, FVector& OutBarycentric, double Tolerance = UE_DOUBLE_SMALL_NUMBER);
+
 	/**
 	 * Computes the barycentric coordinates for a given point on a tetrahedron (3D)
 	 *
