replaced accidental void with bool

diff --git a/Engine/Source/Runtime/Core/Public/Containers/CastedTo.h b/Engine/Source/Runtime/Core/Public/Containers/CastedTo.h
--- a/Engine/Source/Runtime/Core/Public/Containers/CastedTo.h
+++ b/Engine/Source/Runtime/Core/Public/Containers/CastedTo.h
@@ -21,7 +21,7 @@ public:
 
 public:
 	CastTo* Get() const { return Cast(); }
-	void IsPtrValid() const { return IsValid(); }
+	bool IsPtrValid() const { return IsValid(); }
 	void ResetPtr() { Reset(); }
 };
 
