Improve comments on IsRunningDedicatedServer and similar functions to clearly indicate they are meant to be used for data loading and not gameplay purposes

diff --git a/Engine/Source/Runtime/Core/Public/Misc/CoreMisc.h b/Engine/Source/Runtime/Core/Public/Misc/CoreMisc.h
--- a/Engine/Source/Runtime/Core/Public/Misc/CoreMisc.h
+++ b/Engine/Source/Runtime/Core/Public/Misc/CoreMisc.h
@@ -139,8 +139,9 @@ CORE_API bool WillNeedAudioVisualData();
 -----------------------------------------------------------------------------*/
 
 /**
- * Check to see if this executable is running as dedicated server
- * Editor can run as dedicated with -server
+ * Check to see if this executable was launched as a dedicated server process and should not load client only data.
+ * An editor build can be launched with -server to set this to true, but it will be false during single process PlayInEditor mode.
+ * This function should not be used for gameplay or networking purposes, check for NM_DedicatedServer instead.
  */
 FORCEINLINE bool IsRunningDedicatedServer()
 {
@@ -163,9 +164,9 @@ FORCEINLINE bool IsRunningDedicatedServer()
 }
 
 /**
- * Check to see if this executable is running as "the game"
- * - contains all net code (WITH_SERVER_CODE=1)
- * Editor can run as a game with -game
+ * Check to see if this executable was launched as a game process that can be both a client and server (WITH_SERVER_CODE=1).
+ * An editor build can be launched with -game to set this to true, but it will be false during single process PlayInEditor mode.
+ * This function should not be used for gameplay or networking purposes, check the NetMode instead.
  */
 FORCEINLINE bool IsRunningGame()
 {
@@ -188,9 +189,9 @@ FORCEINLINE bool IsRunningGame()
 }
 
 /**
- * Check to see if this executable is running as "the client"
- * - removes all net code (WITH_SERVER_CODE=0)
- * Editor can run as a game with -clientonly
+ * Check to see if this executable was launched as a a client only process that should not load server data (WITH_SERVER_CODE=0).
+ * An editor build can be launched with -clientonly to set this to true, but it will be false during single process PlayInEditor mode.
+ * This function should not be used for gameplay or networking purposes, check for NM_Client instead.
  */
 FORCEINLINE bool IsRunningClientOnly()
 {
