Added [[nodiscard]] to the constructors of TGuardValue, TOptionalGuardValue, TGuardValue_Bitfield_Cleanup and TScopeCounter.

diff --git a/Engine/Source/Editor/Experimental/EditorInteractiveToolsFramework/Private/EditorGizmos/TransformGizmo.cpp b/Engine/Source/Editor/Experimental/EditorInteractiveToolsFramework/Private/EditorGizmos/TransformGizmo.cpp
--- a/Engine/Source/Editor/Experimental/EditorInteractiveToolsFramework/Private/EditorGizmos/TransformGizmo.cpp
+++ b/Engine/Source/Editor/Experimental/EditorInteractiveToolsFramework/Private/EditorGizmos/TransformGizmo.cpp
@@ -759,7 +759,7 @@ void UTransformGizmo::SetNewChildScale(const FVector& NewChildScale)
 	FTransform NewTransform = ActiveTarget->GetTransform();
 	NewTransform.SetScale3D(NewChildScale);
 
-	TGuardValue<bool>(ActiveTarget->bSetPivotMode, true);
+	TGuardValue GuardValue(ActiveTarget->bSetPivotMode, true);
 	ActiveTarget->SetTransform(NewTransform);
 }
 
diff --git a/Engine/Source/Runtime/Core/Public/Templates/UnrealTemplate.h b/Engine/Source/Runtime/Core/Public/Templates/UnrealTemplate.h
--- a/Engine/Source/Runtime/Core/Public/Templates/UnrealTemplate.h
+++ b/Engine/Source/Runtime/Core/Public/Templates/UnrealTemplate.h
@@ -331,7 +331,7 @@ private:
 template <typename RefType, typename AssignedType = RefType>
 struct TGuardValue : private FNoncopyable
 {
-	TGuardValue(RefType& ReferenceValue, const AssignedType& NewValue)
+	[[nodiscard]] TGuardValue(RefType& ReferenceValue, const AssignedType& NewValue)
 	: RefValue(ReferenceValue), OldValue(ReferenceValue)
 	{
 		RefValue = NewValue;
@@ -368,7 +368,7 @@ private:
 template <typename RefType, typename AssignedType = RefType>
 struct TOptionalGuardValue : private FNoncopyable
 {
-	TOptionalGuardValue(RefType& ReferenceValue, const AssignedType& NewValue)
+	[[nodiscard]] TOptionalGuardValue(RefType& ReferenceValue, const AssignedType& NewValue)
 		: RefValue(ReferenceValue), OldValue(ReferenceValue)
 	{
 		if (RefValue != NewValue)
@@ -403,7 +403,7 @@ private:
 template <typename FuncType>
 struct TGuardValue_Bitfield_Cleanup : public FNoncopyable
 {
-	explicit TGuardValue_Bitfield_Cleanup(FuncType&& InFunc)
+	[[nodiscard]] explicit TGuardValue_Bitfield_Cleanup(FuncType&& InFunc)
 		: Func(MoveTemp(InFunc))
 	{
 	}
@@ -433,7 +433,7 @@ private:
 template <typename Type>
 struct TScopeCounter : private FNoncopyable
 {
-	TScopeCounter(Type& ReferenceValue)
+	[[nodiscard]] explicit TScopeCounter(Type& ReferenceValue)
 		: RefValue(ReferenceValue)
 	{
 		++RefValue;
