Optimize GetNetworkConnection

#android

diff --git a/Engine/Source/Runtime/Core/Private/Android/AndroidPlatformMisc.cpp b/Engine/Source/Runtime/Core/Private/Android/AndroidPlatformMisc.cpp
--- a/Engine/Source/Runtime/Core/Private/Android/AndroidPlatformMisc.cpp
+++ b/Engine/Source/Runtime/Core/Private/Android/AndroidPlatformMisc.cpp
@@ -2800,7 +2800,10 @@ bool FAndroidMisc::AreHeadPhonesPluggedIn()
 #define ANDROIDTHUNK_CONNECTION_TYPE_WIMAX 5
 #define ANDROIDTHUNK_CONNECTION_TYPE_BLUETOOTH 6
 
-ENetworkConnectionType FAndroidMisc::GetNetworkConnectionType()
+static bool bLastConnectionTypeValid = false;
+static ENetworkConnectionType LastNetworkConnectionType = ENetworkConnectionType::None;
+
+static ENetworkConnectionType PrivateGetNetworkConnectionType()
 {
 #if USE_ANDROID_JNI
 	extern int32 AndroidThunkCpp_GetNetworkConnectionType();
@@ -2819,6 +2822,16 @@ ENetworkConnectionType FAndroidMisc::GetNetworkConnectionType()
 	return ENetworkConnectionType::Unknown;
 }
 
+ENetworkConnectionType FAndroidMisc::GetNetworkConnectionType()
+{
+	if (!bLastConnectionTypeValid)
+	{
+		LastNetworkConnectionType = PrivateGetNetworkConnectionType();
+		bLastConnectionTypeValid = true;
+	}
+	return LastNetworkConnectionType;
+}
+
 #if USE_ANDROID_JNI
 bool FAndroidMisc::HasActiveWiFiConnection()
 {
@@ -2830,6 +2843,9 @@ bool FAndroidMisc::HasActiveWiFiConnection()
 
 JNI_METHOD void Java_com_epicgames_unreal_GameActivity_nativeNetworkChanged(JNIEnv* jenv, jobject thiz)
 {
+	LastNetworkConnectionType = PrivateGetNetworkConnectionType();
+	bLastConnectionTypeValid = true;
+	
 	if (FTaskGraphInterface::IsRunning())
 	{
 		FFunctionGraphTask::CreateAndDispatchWhenReady([]()
@@ -3190,6 +3206,7 @@ void FAndroidMisc::ShowConsoleWindow()
 
 FDelegateHandle FAndroidMisc::AddNetworkListener(FCoreDelegates::FOnNetworkConnectionChanged::FDelegate&& InNewDelegate)
 {
+	// not really necessary since PlatformInit calls AddNetworkListener but doesn't hurt anything
 	if (!FCoreDelegates::OnNetworkConnectionChanged.IsBound())
 	{
 #if USE_ANDROID_JNI
@@ -3205,13 +3222,14 @@ bool FAndroidMisc::RemoveNetworkListener(FDelegateHandle Handle)
 {
 	bool bSuccess = FCoreDelegates::OnNetworkConnectionChanged.Remove(Handle);
 
-	if (!FCoreDelegates::OnNetworkConnectionChanged.IsBound())
-	{
-#if USE_ANDROID_JNI
-		extern void AndroidThunkJava_RemoveNetworkListener();
-		AndroidThunkJava_RemoveNetworkListener();
-#endif
-	}
+	// we don't really want to remove listener since we're using it for GetNetworkConnectionType
+//	if (!FCoreDelegates::OnNetworkConnectionChanged.IsBound())
+//	{
+//#if USE_ANDROID_JNI
+//		extern void AndroidThunkJava_RemoveNetworkListener();
+//		AndroidThunkJava_RemoveNetworkListener();
+//#endif
+//	}
 
 	return bSuccess;
 }
