Simplified the Pointer Container Classes into TSmartPtr

diff --git a/Engine/Source/Runtime/Core/Public/Containers/CastedTo.h b/Engine/Source/Runtime/Core/Public/Containers/CastedTo.h
--- a/Engine/Source/Runtime/Core/Public/Containers/CastedTo.h
+++ b/Engine/Source/Runtime/Core/Public/Containers/CastedTo.h
@@ -10,10 +10,10 @@
  * These classes are useful if you want your system to support both memory management systems: (1) UObject Garbage Collection and (2) Smart Pointer Reference counting
  */
 template<typename CastTo>
-class TPointerContainer : public TSharedFromThis<TPointerContainer<CastTo>>
+class TPointerContainerBase : public TSharedFromThis<TPointerContainerBase<CastTo>>
 {
 public:
-	virtual ~TPointerContainer() {}
+	virtual ~TPointerContainerBase() {}
 
 	virtual CastTo* Cast() const = 0;
 	virtual bool IsValid() const = 0;
@@ -25,6 +25,17 @@ public:
 	void ResetPtr() { Reset(); }
 };
 
+template<typename CastTo>
+class TPointerContainer : public TPointerContainerBase<CastTo>
+{
+public:
+	virtual ~TPointerContainer() {}
+
+	virtual CastTo* Cast() const override { return nullptr; }
+	virtual bool IsValid() const override { return false; }
+	virtual void Reset() override {}
+};
+
 /* A "Weak" Pointer Reference Container for a 'T' (UObject xOr TSharedFromThis) that can be statically casted to 'CastTo' */
 template<typename CastTo>
 class TWeakContainerTo : public TPointerContainer<CastTo>
diff --git a/Engine/Source/Runtime/Core/Public/Templates/SmartPtr.h b/Engine/Source/Runtime/Core/Public/Templates/SmartPtr.h
new file mode 100644
--- /dev/null
+++ b/Engine/Source/Runtime/Core/Public/Templates/SmartPtr.h
@@ -0,0 +1,55 @@
+// Copyright Epic Games, Inc. All Rights Reserved.
+
+#pragma once
+
+#include "Containers/CastedTo.h"
+
+enum class ESmartPointer
+{
+	Strong,
+	Shared = Strong, // Shared and Strong pointers do the same thing, which is prevent automatic memory deletion using different systems (Garbage Collection vs. Reference Counting)
+	Weak,
+	// Unique(?) - not sure if needed
+};
+
+/* Smart Pointer Class where the pointer can be statically casted to 'T' */
+template<typename T>
+class TSmartPtr
+{
+public:
+	TSmartPtr() : Container(MakeShared<TPointerContainer<T>>()) {}
+
+	template<typename From>
+	TSmartPtr(From* InObject)
+	{
+		Set(InObject, ESmartPointer::Weak);
+	}
+
+	template<typename From>
+	TSmartPtr(From* InObject, const ESmartPointer InType)
+	{
+		Set(InObject, InType);
+	}
+
+	template<typename From>
+	TSmartPtr<T>& Set(From* InObject, const ESmartPointer InType)
+	{
+		if (InType == ESmartPointer::Strong || InType == ESmartPointer::Shared)
+		{
+			Container = MakeShared<TStrongCastable<From, T>>(InObject);
+		}
+		else if (InType == ESmartPointer::Weak)
+		{
+			Container = MakeShared<TWeakCastable<From, T>>(InObject);
+		}
+
+		return *this;
+	}
+
+	T* Get() const { return Container->Get(); }
+	bool IsValid() const { return Container->IsValid(); }
+	void Reset() const { Container->Reset(); }
+
+private:
+	TSharedRef<TPointerContainer<T>> Container = MakeShared<TPointerContainer<T>>();
+};
\ No newline at end of file
