Fixed compiler errors when building with header units

diff --git a/Engine/Source/Runtime/Core/Public/AutoRTFM/AutoRTFMConstants.h b/Engine/Source/Runtime/Core/Public/AutoRTFM/AutoRTFMConstants.h
--- a/Engine/Source/Runtime/Core/Public/AutoRTFM/AutoRTFMConstants.h
+++ b/Engine/Source/Runtime/Core/Public/AutoRTFM/AutoRTFMConstants.h
@@ -15,22 +15,22 @@ namespace AutoRTFM
 namespace Constants
 {
 
-constexpr size_t LogLineBytes = 3;
-constexpr size_t LineBytes = static_cast<size_t>(1) << LogLineBytes;
-constexpr size_t LineTableSize = 33554432;
-
-constexpr size_t Offset_Context_CurrentTransaction = 0;
-constexpr size_t Offset_Context_LineTable = 192;
-constexpr size_t Offset_Context_Status = 152;
-
-constexpr size_t LogSize_LineEntry = 5;
-constexpr size_t Size_LineEntry = static_cast<size_t>(1) << LogSize_LineEntry;
-constexpr size_t Offset_LineEntry_LogicalLine = 0;
-constexpr size_t Offset_LineEntry_ActiveLine = 16;
-constexpr size_t Offset_LineEntry_LoggingTransaction = 24;
-constexpr size_t Offset_LineEntry_AccessMask = 14;
-
-constexpr uint32_t Context_Status_OnTrack = 1;
+inline constexpr size_t LogLineBytes = 3;
+inline constexpr size_t LineBytes = static_cast<size_t>(1) << LogLineBytes;
+inline constexpr size_t LineTableSize = 33554432;
+
+inline constexpr size_t Offset_Context_CurrentTransaction = 0;
+inline constexpr size_t Offset_Context_LineTable = 192;
+inline constexpr size_t Offset_Context_Status = 152;
+
+inline constexpr size_t LogSize_LineEntry = 5;
+inline constexpr size_t Size_LineEntry = static_cast<size_t>(1) << LogSize_LineEntry;
+inline constexpr size_t Offset_LineEntry_LogicalLine = 0;
+inline constexpr size_t Offset_LineEntry_ActiveLine = 16;
+inline constexpr size_t Offset_LineEntry_LoggingTransaction = 24;
+inline constexpr size_t Offset_LineEntry_AccessMask = 14;
+
+inline constexpr uint32_t Context_Status_OnTrack = 1;
 
 } // namespace Constants
 } // namespace AutoRTFM
diff --git a/Engine/Source/Runtime/Core/Public/HAL/PThreadSemaphore.h b/Engine/Source/Runtime/Core/Public/HAL/PThreadSemaphore.h
--- a/Engine/Source/Runtime/Core/Public/HAL/PThreadSemaphore.h
+++ b/Engine/Source/Runtime/Core/Public/HAL/PThreadSemaphore.h
@@ -2,6 +2,8 @@
 
 #pragma once
 
+// HEADER_UNIT_UNSUPPORTED - Unsupported platform
+
 #include "CoreTypes.h"
 #include "Misc/Timespan.h"
 #include "Misc/AssertionMacros.h"
diff --git a/Engine/Source/Runtime/Core/Public/Math/Vector.h b/Engine/Source/Runtime/Core/Public/Math/Vector.h
--- a/Engine/Source/Runtime/Core/Public/Math/Vector.h
+++ b/Engine/Source/Runtime/Core/Public/Math/Vector.h
@@ -2689,7 +2689,7 @@ inline TVector<T> TVector2<T>::SphericalToUnitCartesian() const
 	
 namespace LWC
 {
-constexpr FVector::FReal DefaultFloatPrecision = 1./16.;
+inline constexpr FVector::FReal DefaultFloatPrecision = 1./16.;
 
 // Validated narrowing cast for world positions. FVector -> FVector3f
 FORCEINLINE FVector3f NarrowWorldPositionChecked(const FVector& WorldPosition)
diff --git a/Engine/Source/Runtime/CoreUObject/Public/UObject/ObjectHandle.h b/Engine/Source/Runtime/CoreUObject/Public/UObject/ObjectHandle.h
--- a/Engine/Source/Runtime/CoreUObject/Public/UObject/ObjectHandle.h
+++ b/Engine/Source/Runtime/CoreUObject/Public/UObject/ObjectHandle.h
@@ -338,9 +338,9 @@ namespace UE::CoreUObject::Private
 		};
 	};
 
-	constexpr uint32 ObjectIdShift = 1;
-	constexpr uint32 PackageIdShift = 33;
-	constexpr uint32 PackageIdMask = 0x7FFF'FFFF;
+	inline constexpr uint32 ObjectIdShift = 1;
+	inline constexpr uint32 PackageIdShift = 33;
+	inline constexpr uint32 PackageIdMask = 0x7FFF'FFFF;
 
 #if UE_WITH_OBJECT_HANDLE_LATE_RESOLVE
 	//forward declarations
