- Deprecated FPaths::EnginePlatformExtensionsDir() and FPaths::ProjectPlatformExtensionsDir(), as they do not work with Programs
- Introduced the following to FPaths:
   - EnginePlatformExtensionDir(Platform) - gets the Engine PlatExt dir for the given platform
   - ProjectPlatformExtensionDir(Platform) - gets the Project PlatExt dir for the given platform, where Project can be a Program
   - ConvertPath(Path, Method, ExtraData) - converts a Path for platform extensions and Restricted locations. This is going to be "ground truth" for how to get an extension directory (the two functions above call this).
- Added FConfigContext::VisualizeHierarchy() (static and member versions with differing params) that can log ini hierarchies

diff --git a/Engine/Plugins/Experimental/GameFeatures/Source/GameFeatures/Private/GameFeatureData.cpp b/Engine/Plugins/Experimental/GameFeatures/Source/GameFeatures/Private/GameFeatureData.cpp
--- a/Engine/Plugins/Experimental/GameFeatures/Source/GameFeatures/Private/GameFeatureData.cpp
+++ b/Engine/Plugins/Experimental/GameFeatures/Source/GameFeatures/Private/GameFeatureData.cpp
@@ -138,7 +138,7 @@ void UGameFeatureData::InitializeHierarchicalPluginIniFiles(const FString& Plugi
 	}
 
 	const FString PluginName = FPaths::GetBaseFilename(PluginInstalledStandardFilename);
-	const FString PlatformExtensionDir = FPaths::ProjectPlatformExtensionsDir() / (PlatformName + "/");
+	const FString PlatformExtensionDir = FPaths::ProjectPlatformExtensionDir(*PlatformName);
 	const FString EngineConfigDir = FPaths::EngineConfigDir();
 	const FString PluginConfigDir = FPaths::GetPath(PluginInstalledStandardFilename) / TEXT("Config/");
 	const FString PluginPlatformConfigDir = FPaths::Combine(PluginConfigDir, PlatformName);
diff --git a/Engine/Source/Runtime/Core/Private/Misc/ConfigContext.cpp b/Engine/Source/Runtime/Core/Private/Misc/ConfigContext.cpp
--- a/Engine/Source/Runtime/Core/Private/Misc/ConfigContext.cpp
+++ b/Engine/Source/Runtime/Core/Private/Misc/ConfigContext.cpp
@@ -112,7 +112,8 @@ const FConfigContext::FPerPlatformDirs& FConfigContext::GetPerPlatformDirs(const
 	FConfigContext::FPerPlatformDirs* Dirs = FConfigContext::PerPlatformDirs.Find(PlatformName);
 	if (Dirs == nullptr)
 	{
-		FString PluginExtDir;
+		// default to <skip> so we don't look in non-existant platform extension directories
+		FString PluginExtDir = TEXT("<skip>");
 		if (bIsForPlugin)
 		{
 			// look if there's a plugin extension for this platform, it will have the platform name in the path
@@ -129,9 +130,9 @@ const FConfigContext::FPerPlatformDirs& FConfigContext::GetPerPlatformDirs(const
 		Dirs = &PerPlatformDirs.Emplace(PlatformName, FConfigContext::FPerPlatformDirs
 			{
 				// PlatformExtensionEngineDir
-				FPaths::Combine(*FPaths::EnginePlatformExtensionsDir(), *PlatformName).Replace(*FPaths::EngineDir(), *(EngineRootDir + "/")),
+				FPaths::ConvertPath(EngineRootDir, FPaths::EPathConversion::Engine_PlatformExtension, *PlatformName),
 				// PlatformExtensionProjectDir
-				FPaths::Combine(*FPaths::ProjectPlatformExtensionsDir(), *PlatformName).Replace(*FPaths::ProjectDir(), *(ProjectRootDir + "/")),
+				FPaths::ConvertPath(ProjectRootDir, FPaths::EPathConversion::Project_PlatformExtension, *PlatformName, *ProjectRootDir),
 				// PluginExtensionDir
 				PluginExtDir,
 			});
@@ -484,7 +485,7 @@ FString FConfigContext::PerformFinalExpansions(const FString& InString, const FS
 		OutString = OutString.Replace(TEXT("{OPT_SUBDIR}"), TEXT(""));
 	}
 	
-	if (Platform.Len() > 0)
+	if (InPlatform.Len() > 0)
 	{
 		OutString = OutString.Replace(TEXT("{EXTENGINE}"), *GetPerPlatformDirs(InPlatform).PlatformExtensionEngineDir);
 		OutString = OutString.Replace(TEXT("{EXTPROJECT}"), *GetPerPlatformDirs(InPlatform).PlatformExtensionProjectDir);
@@ -502,7 +503,7 @@ FString FConfigContext::PerformFinalExpansions(const FString& InString, const FS
 
 
 
-void FConfigContext::AddStaticLayersToHierarchy()
+void FConfigContext::AddStaticLayersToHierarchy(TArray<FString>* GatheredLayerFilenames, bool bIsLogging)
 {
 	// remember where this file was loaded from
 	ConfigFile->SourceEngineConfigDir = EngineConfigDir;
@@ -612,9 +613,28 @@ void FConfigContext::AddStaticLayersToHierarchy()
 					{
 						return;
 					}
+					
+					if (PlatformPath.StartsWith(TEXT("<skip>")))
+					{
+						continue;
+					}
 
 					// add this to the list!
-					ConfigFile->SourceIniHierarchy.AddStaticLayer(PlatformPath, LayerIndex, ExpansionIndex, PlatformIndex);
+					if (GatheredLayerFilenames != nullptr)
+					{
+						if (bIsLogging)
+						{
+							GatheredLayerFilenames->Add(FString::Printf(TEXT("%s[Exp-%d]: %s"), Layer.EditorName, ExpansionIndex, *PlatformPath));
+						}
+						else
+						{
+							GatheredLayerFilenames->Add(PlatformPath);
+						}
+					}
+					else
+					{
+						ConfigFile->SourceIniHierarchy.AddStaticLayer(PlatformPath, LayerIndex, ExpansionIndex, PlatformIndex);
+					}
 				}
 			}
 		}
@@ -633,7 +653,21 @@ void FConfigContext::AddStaticLayersToHierarchy()
 			}
 
 			// add with no expansion
-			ConfigFile->SourceIniHierarchy.AddStaticLayer(FinalPath, LayerIndex);
+			if (GatheredLayerFilenames != nullptr)
+			{
+				if (bIsLogging)
+				{
+					GatheredLayerFilenames->Add(FString::Printf(TEXT("%s: %s"), Layer.EditorName, *FinalPath));
+				}
+				else
+				{
+					GatheredLayerFilenames->Add(FinalPath);
+				}
+			}
+			else
+			{
+				ConfigFile->SourceIniHierarchy.AddStaticLayer(FinalPath, LayerIndex);
+			}
 		}
 	}
 }
@@ -945,3 +979,67 @@ void FConfigContext::EnsureRequiredGlobalPathsHaveBeenInitialized()
 {
 	PerformBasicReplacements(TEXT(""), TEXT("")); // requests user directories and FConfigCacheIni::GetCustomConfigString
 }
+
+
+void FConfigContext::VisualizeHierarchy(FOutputDevice& Ar, const TCHAR* IniName, const TCHAR* OverridePlatform, const TCHAR* OverrideProjectOrProgramDataDir, const TCHAR* OverridePluginDir, const TArray<FString>* ChildPluginBaseDirs)
+{
+	FConfigFile Test;
+	FConfigContext Context(nullptr, true, OverridePlatform ? FString(OverridePlatform) : FString(), &Test);
+	if (OverridePluginDir != nullptr)
+	{
+		Context.bIsForPlugin = true;
+		Context.PluginRootDir = OverridePluginDir;
+		if (ChildPluginBaseDirs != nullptr)
+		{
+			Context.ChildPluginBaseDirs = *ChildPluginBaseDirs;
+		}
+	}
+	
+	if (OverrideProjectOrProgramDataDir != nullptr)
+	{
+		Context.ProjectConfigDir = FPaths::Combine(OverrideProjectOrProgramDataDir, "Config/");
+	}
+
+	
+	Context.VisualizeHierarchy(Ar, IniName);
+}
+
+void FConfigContext::VisualizeHierarchy(FOutputDevice& Ar, const TCHAR* IniName)
+{
+	Ar.Logf(TEXT("======================================================="));
+
+	ResetBaseIni(IniName);
+	CachePaths();
+	bool _;
+	PrepareForLoad(_);
+
+	Ar.Logf(TEXT("Config hierarchy:"));
+	if (ProjectRootDir.Contains(TEXT("/Programs/")))
+	{
+		Ar.Logf(TEXT("  Program Data Dir: %s"), *ProjectRootDir);
+	}
+	else
+	{
+		Ar.Logf(TEXT("  Project Dir: %s"), *ProjectRootDir);
+	}
+	Ar.Logf(TEXT("  Platform: %s"), *Platform);
+	if (bIsForPlugin)
+	{
+		Ar.Logf(TEXT("  Plugin Root Dir: %s"), *PluginRootDir);
+		for (FString& Child : ChildPluginBaseDirs)
+		{
+			Ar.Logf(TEXT("  Plugin Children Dir: %s"), *Child);
+		}
+	}
+	
+	TArray<FString> FileList;
+	AddStaticLayersToHierarchy(&FileList, true);
+	
+	Ar.Logf(TEXT("  Files:"));
+	for (const FString& File : FileList)
+	{
+		Ar.Logf(TEXT("    %s"), *File);
+	}
+
+	Ar.Logf(TEXT("======================================================="));
+}
diff --git a/Engine/Source/Runtime/Core/Private/Misc/DataDrivenPlatformInfoRegistry.cpp b/Engine/Source/Runtime/Core/Private/Misc/DataDrivenPlatformInfoRegistry.cpp
--- a/Engine/Source/Runtime/Core/Private/Misc/DataDrivenPlatformInfoRegistry.cpp
+++ b/Engine/Source/Runtime/Core/Private/Misc/DataDrivenPlatformInfoRegistry.cpp
@@ -38,11 +38,11 @@ static const TArray<FString>& GetDataDrivenIniFilenames()
 
 		// manually look through the platform directories - we can't use GetExtensionDirs(), since that function uses the results of this function 
 		TArray<FString> PlatformDirs;
-		IFileManager::Get().FindFiles(PlatformDirs, *FPaths::Combine(FPaths::EnginePlatformExtensionsDir(), TEXT("*")), false, true);
+		IFileManager::Get().FindFiles(PlatformDirs, *FPaths::Combine(FPaths::EngineDir(), TEXT("Platforms"), TEXT("*")), false, true);
 
 		for (const FString& PlatformDir : PlatformDirs)
 		{
-			FString IniPath = FPaths::Combine(FPaths::EnginePlatformExtensionsDir(), PlatformDir, TEXT("Config/DataDrivenPlatformInfo.ini"));
+			FString IniPath = FPaths::Combine(FPaths::EnginePlatformExtensionDir(*PlatformDir), TEXT("Config/DataDrivenPlatformInfo.ini"));
 			if (IFileManager::Get().FileExists(*IniPath))
 			{
 				DataDrivenIniFilenames.Add(IniPath);
@@ -50,11 +50,12 @@ static const TArray<FString>& GetDataDrivenIniFilenames()
 		}
 
 		// look for the special files in any project config subdirectories
+		// @note: this will not work for Programs, the Platforms dir is outside of the ProjectDir in their case
 		TArray<FString> ProjectPlatformDirs;
-		IFileManager::Get().FindFiles(ProjectPlatformDirs, *FPaths::Combine(FPaths::ProjectConfigDir(), TEXT("*")), false, true);
+		IFileManager::Get().FindFiles(ProjectPlatformDirs, *FPaths::Combine(FPaths::ProjectDir(), TEXT("Platforms"), TEXT("*")), false, true);
 		for (const FString& PlatformDir : ProjectPlatformDirs)
 		{
-			FString IniPath = FPaths::Combine(FPaths::ProjectPlatformExtensionsDir(), PlatformDir, TEXT("Config/DataDrivenPlatformInfo.ini"));
+			FString IniPath = FPaths::Combine(FPaths::ProjectPlatformExtensionDir(*PlatformDir), TEXT("Config/DataDrivenPlatformInfo.ini"));
 			if (IFileManager::Get().FileExists(*IniPath))
 			{
 				DataDrivenIniFilenames.Add(IniPath);
@@ -86,11 +87,11 @@ bool FDataDrivenPlatformInfoRegistry::LoadDataDrivenIniFile(int32 Index, FConfig
 		IniFile.ProcessInputFileContents(IniContents, IniFilenames[Index]);
 
 		// platform extension paths are different (engine/platforms/platform/config, not engine/config/platform)
-		if (IniFilenames[Index].StartsWith(FPaths::EnginePlatformExtensionsDir()))
+		if (IniFilenames[Index].StartsWith(FPaths::EnginePlatformExtensionDir(TEXT("")).TrimChar('/')))
 		{
 			PlatformName = FPaths::GetCleanFilename(FPaths::GetPath(FPaths::GetPath(IniFilenames[Index])));
 		}
-		else if (IniFilenames[Index].StartsWith(FPaths::ProjectPlatformExtensionsDir()))
+		else if (IniFilenames[Index].StartsWith(FPaths::ProjectPlatformExtensionDir(TEXT("")).TrimChar('/')))
 		{
 			PlatformName = FPaths::GetCleanFilename(FPaths::GetPath(FPaths::GetPath(IniFilenames[Index])));
 		}
diff --git a/Engine/Source/Runtime/Core/Private/Misc/Paths.cpp b/Engine/Source/Runtime/Core/Private/Misc/Paths.cpp
--- a/Engine/Source/Runtime/Core/Private/Misc/Paths.cpp
+++ b/Engine/Source/Runtime/Core/Private/Misc/Paths.cpp
@@ -271,6 +271,82 @@ FString FPaths::ProjectPlatformExtensionsDir()
 }
 
 
+FString FPaths::ConvertPath(const FString& Path, EPathConversion Method, const TCHAR* ExtraData, const TCHAR* OverrideProjectDir)
+{
+	// Basic idea here is: 
+	//   We have a path coming in that we want to get the correct corresponding extension/restricted location 
+	//   So, depending on parameters, and if it's a program path or not (programs are special locations), we split
+	//   up the path into a Prefix and Suffix, and insert the correct folders into the middle that matches the Conversion method
+	// Note that ExtraData is currently only for PlatformExtensions, and contains the Platform name
+
+	// NOTE: We can't test that a Project_ type starts with ProjectDir, because of this case of a Program:
+	//   Path0 = ProjectDir() + "Config";
+	//   Path1 = ConvertPath(Path0, Project_NoRedist);
+	//   Path2 = ConvertPath(Path1, Project_PlatformExtension, "Mac");
+	// This is valid, and will generate the correct output, but note that Path1 does _not_ start with ProjectDir (which is ../../../Engine/Programs/Foo):
+	//   Path0 = ../../../Engine/Programs/Foo/Config
+	//   Path1 = ../../../Engine/Restricted/NoRedist/Programs/Foo/Config
+	//   Path2 = ../../../Engine/Restricted/NoRedist/Platforms/Mac/Programs/Foo/Config
+	// Aren't Programs fun??
+		
+	const bool bIsProgram = OverrideProjectDir == nullptr ? IS_PROGRAM : FCString::Strstr(OverrideProjectDir, TEXT("/Programs/")) != nullptr;
+	const bool bIsProjectConversion = (int)Method >= (int)EPathConversion::Project_First;
+	const bool bIsEngineConversion = !bIsProjectConversion;
+	
+	FString Prefix, Suffix;
+	// programs need special love for project types, where we need to split on whatever is before the Programs, NOT ProjectDir
+	if (bIsProgram && bIsProjectConversion)
+	{
+		int ProgramsLoc = Path.Replace(TEXT("\\"), TEXT("/")).Find(TEXT("/Programs/"));
+		Prefix = Path.Mid(0, ProgramsLoc + 1);
+		Suffix = Path.Mid(ProgramsLoc + 1);
+	}
+	// handle the override case, simply
+	else if (OverrideProjectDir != nullptr && (int)Method >= (int)EPathConversion::Project_First)
+	{
+		// skip over engine or project directory
+		Prefix = OverrideProjectDir;
+		// grab what's left
+		Suffix = Path.Mid(Prefix.Len());
+	}
+	else
+	{
+		checkfSlow(!(bIsProjectConversion && !Path.StartsWith(FPaths::ProjectDir())), TEXT("Called ConvertPath with a Platform_ method, but Path didn't start with FPaths::ProjectDir(). This is not supported"));
+		checkfSlow(!(bIsEngineConversion && !Path.StartsWith(FPaths::EngineDir())), TEXT("Called ConvertPath with a Engine_ method, but Path didn't start with FPaths::EngineDir(). This is not supported"));
+
+		// skip over engine or project directory
+		Prefix = (int)Method < (int)EPathConversion::Project_First ? FPaths::EngineDir() : FPaths::ProjectDir();
+		// grab what's left
+		Suffix = Path.Mid(Prefix.Len());
+	}
+	
+	// we check we can skip appending Suffix if it's empty, otherwise, we will end up with an extra trailing /
+	// but we do need to make sure we end with a / if the incoming path did
+	if (!Suffix.EndsWith(TEXT("/")) && Path.EndsWith(TEXT("/")))
+	{
+		Suffix += TEXT("/");
+	}
+	bool bAppendSuffix = Suffix.Len() > 0;
+	
+	switch (Method)
+	{
+		case EPathConversion::Engine_PlatformExtension:
+		case EPathConversion::Project_PlatformExtension:
+			return bAppendSuffix ? FPaths::Combine(Prefix, TEXT("Platforms"), ExtraData, Suffix) : FPaths::Combine(Prefix, TEXT("Platforms"), ExtraData);
+			
+		case EPathConversion::Engine_NotForLicensees:
+		case EPathConversion::Project_NotForLicensees:
+			return bAppendSuffix ? FPaths::Combine(Prefix, TEXT("Restricted/NotForLicensees"), Suffix) : FPaths::Combine(Prefix, TEXT("Restricted/NotForLicensees"));
+			
+		case EPathConversion::Engine_NoRedist:
+		case EPathConversion::Project_NoRedist:
+			return bAppendSuffix ? FPaths::Combine(Prefix, TEXT("Restricted/NoRedist"), Suffix) : FPaths::Combine(Prefix, TEXT("Restricted/NoRedist"));
+	}
+	
+	return TEXT("");
+}
+
+
 static void AddIfDirectoryExists(TArray<FString>& ExtensionDirs, FString&& Dir)
 {
 	if (FPaths::DirectoryExists(Dir))
diff --git a/Engine/Source/Runtime/Core/Public/Misc/ConfigContext.h b/Engine/Source/Runtime/Core/Public/Misc/ConfigContext.h
--- a/Engine/Source/Runtime/Core/Public/Misc/ConfigContext.h
+++ b/Engine/Source/Runtime/Core/Public/Misc/ConfigContext.h
@@ -123,6 +123,12 @@ public:
 	 * Call to make before attempting parallel config init
 	 */
 	static CORE_API void EnsureRequiredGlobalPathsHaveBeenInitialized();
+	
+	/** Log out the config full hierarchy of a file, with various overrides to see other projects/platforms/etc hierarchies */
+	static CORE_API void VisualizeHierarchy(FOutputDevice& Ar, const TCHAR* IniName, const TCHAR* OverridePlatform, const TCHAR* OverrideProjectOrProgramDataDir, const TCHAR* OverridePluginDir=nullptr, const TArray<FString>* ChildPluginBaseDirs=nullptr);
+	
+	/** Visualize an existing hierarchy */
+	void CORE_API VisualizeHierarchy(FOutputDevice& Ar, const TCHAR* IniName);
 
 	/**
 	 * Use the context to perform the actual load operation. Note that this is where you specify the Ini name (for instance "Engine"), meaning
@@ -198,7 +204,7 @@ protected:
 	CORE_API bool PrepareForLoad(bool& bPerformLoad);
 	CORE_API bool PerformLoad();
 
-	CORE_API void AddStaticLayersToHierarchy();
+	void AddStaticLayersToHierarchy(TArray<FString>* GatheredLayerFilenames=nullptr, bool bIsForLogging=false);
 	CORE_API bool GenerateDestIniFile();
 	CORE_API FString PerformFinalExpansions(const FString& InString, const FString& Platform);
 };
diff --git a/Engine/Source/Runtime/Core/Public/Misc/Paths.h b/Engine/Source/Runtime/Core/Public/Misc/Paths.h
--- a/Engine/Source/Runtime/Core/Public/Misc/Paths.h
+++ b/Engine/Source/Runtime/Core/Public/Misc/Paths.h
@@ -156,15 +156,58 @@ public:
 	 *
 	 * @return engine platform extensions directory
 	 */
+	UE_DEPRECATED(5.4, "Use EnginePlatformExtensionDir(Platform) instead - ProjectPlatformExtensionsDir did not handle programs properly, so for consistency this is being removed as well")
 	static CORE_API FString EnginePlatformExtensionsDir();
 
+	/**
+	 * Returns the directory where the engine's platform extensions resides for the given platform
+	 *
+	 * @param Platform the platform to get the extension directory for
+	 * @return engine's platform extension directory
+	 */
+	static CORE_API FString EnginePlatformExtensionDir(const TCHAR* Platform)
+	{
+		return ConvertPath(FPaths::EngineDir(), EPathConversion::Engine_PlatformExtension, Platform);
+	}
+
 	/**
 	 * Returns the directory where the project's platform extensions reside
 	 *
 	 * @return project platform extensions directory
 	 */
+	UE_DEPRECATED(5.4, "Use ProjectPlatformExtensionDir(Platform) instead - this function does not handle Programs properly")
 	static CORE_API FString ProjectPlatformExtensionsDir();
 
+	/**
+	 * Returns the directory where the project's platform extensions resides for the given platform
+	 *
+	 * @param Platform the platform to get the extension directory for
+	 * @return project's platform extension directory
+	 */
+	static CORE_API FString ProjectPlatformExtensionDir(const TCHAR* Platform)
+	{
+		return ConvertPath(FPaths::ProjectDir(), EPathConversion::Project_PlatformExtension, Platform);
+	}
+	
+	enum class EPathConversion : uint8
+	{
+		// ExtraData is name of platform
+		Engine_PlatformExtension,
+		Engine_NotForLicensees,
+		Engine_NoRedist,
+
+		Project_First,
+		Project_PlatformExtension = Project_First,
+		Project_NotForLicensees,
+		Project_NoRedist,
+	};
+	
+	/**
+	  * Converts a path with the given method. Initially used for converting from normal paths to platform extensions (programs make it confusing, so haveing
+	  * one source of truth is a good idea)
+	 */
+	static CORE_API FString ConvertPath(const FString& Path, EPathConversion Method, const TCHAR* ExtraData=nullptr, const TCHAR* OverrideProjectDir=nullptr);
+	
 	/**
 	 * Returns platform and restricted extensions that are present and valid (for platforms, it uses FDataDrivePlatformInfo to determine valid platforms, it doesn't just use what's present)
 	 *
diff --git a/Engine/Source/Runtime/CoreUObject/Private/UObject/Obj.cpp b/Engine/Source/Runtime/CoreUObject/Private/UObject/Obj.cpp
--- a/Engine/Source/Runtime/CoreUObject/Private/UObject/Obj.cpp
+++ b/Engine/Source/Runtime/CoreUObject/Private/UObject/Obj.cpp
@@ -3122,7 +3122,7 @@ FString UObject::GetDefaultConfigFilename() const
 	FString OverridePlatform = GetFinalOverridePlatform(this);
 	if (OverridePlatform.Len())
 	{
-		bool bIsPlatformExtension = FPaths::DirectoryExists(FPaths::Combine(FPaths::EnginePlatformExtensionsDir(), OverridePlatform));
+		bool bIsPlatformExtension = FPaths::DirectoryExists(FPaths::EnginePlatformExtensionDir(*OverridePlatform));
 		FString RegularPath = FString::Printf(TEXT("%s%s"), *FPaths::SourceConfigDir(), *OverridePlatform, *OverridePlatform, *GetClass()->ClassConfigName.ToString());
 		FString SelectedPath = RegularPath;
 
@@ -3140,7 +3140,7 @@ FString UObject::GetDefaultConfigFilename() const
 		// if the project already uses platform configs in the regular directory, just use that, otherwise check if this is a platform extensions
 		if (bIsPlatformExtension && !bPlatformConfigExistsInRegular)
 		{
-			SelectedPath = FString::Printf(TEXT("%s%s/Config"), *FPaths::ProjectPlatformExtensionsDir(), *OverridePlatform);
+			SelectedPath = FPaths::Combine(FPaths::ProjectPlatformExtensionDir(*OverridePlatform), TEXT("Config"));
 		}
 
 		return FConfigCacheIni::NormalizeConfigIniPath(FString::Printf(TEXT("%s/%s%s.ini"), *SelectedPath, *OverridePlatform, *GetClass()->ClassConfigName.ToString()));
diff --git a/Engine/Source/Runtime/NetworkFile/Private/NetworkPlatformFile.cpp b/Engine/Source/Runtime/NetworkFile/Private/NetworkPlatformFile.cpp
--- a/Engine/Source/Runtime/NetworkFile/Private/NetworkPlatformFile.cpp
+++ b/Engine/Source/Runtime/NetworkFile/Private/NetworkPlatformFile.cpp
@@ -794,17 +794,19 @@ void FNetworkPlatformFile::FillGetFileList(FNetworkFileArchive& Payload)
 		GameName = FPaths::GetProjectFilePath();
 	}
 
+	FString IniPlatformName = FPlatformProperties::IniPlatformName();
 	FString EngineRelPath = FPaths::EngineDir();
 	FString EngineRelPluginPath = FPaths::EnginePluginsDir();
 	FString GameRelPath = FPaths::ProjectDir();
 	FString GameRelPluginPath = FPaths::ProjectPluginsDir();
-	FString EnginePlatformExtensionsDir = FPaths::EnginePlatformExtensionsDir();
-	FString ProjectPlatformExtensionsDir = FPaths::ProjectPlatformExtensionsDir();
+	FString EnginePlatformExtensionsDir = FPaths::EnginePlatformExtensionDir(*IniPlatformName);
+	FString ProjectPlatformExtensionsDir = FPaths::ProjectPlatformExtensionDir(*IniPlatformName);
 
 	TArray<FString> Directories;
 	Directories.Add(EngineRelPath);
 	Directories.Add(GameRelPath);
 
+	Payload << IniPlatformName;
 	Payload << TargetPlatformNames;
 	Payload << GameName;
 	Payload << EngineRelPath;
@@ -841,14 +843,15 @@ void FNetworkPlatformFile::ProcessServerInitialResponse(FArrayReader& InResponse
 	InResponse << ServerEnginePlatformExtensionsDir;
 	InResponse << ServerProjectPlatformExtensionsDir;
 
+	FString IniPlatformName = FPlatformProperties::IniPlatformName();
 	UE_LOG(LogNetworkPlatformFile, Display, TEXT("    Server EngineDir      = %s"), *ServerEngineDir);
 	UE_LOG(LogNetworkPlatformFile, Display, TEXT("     Local EngineDir      = %s"), *FPaths::EngineDir());
 	UE_LOG(LogNetworkPlatformFile, Display, TEXT("    Server ProjectDir     = %s"), *ServerProjectDir);
 	UE_LOG(LogNetworkPlatformFile, Display, TEXT("     Local ProjectDir     = %s"), *FPaths::ProjectDir());
 	UE_LOG(LogNetworkPlatformFile, Display, TEXT("    Server EnginePlatformExtDir = %s"), *ServerEnginePlatformExtensionsDir);
-	UE_LOG(LogNetworkPlatformFile, Display, TEXT("     Local EnginePlatformExtDir = %s"), *FPaths::EnginePlatformExtensionsDir());
+	UE_LOG(LogNetworkPlatformFile, Display, TEXT("     Local EnginePlatformExtDir = %s"), *FPaths::EnginePlatformExtensionDir(*IniPlatformName));
 	UE_LOG(LogNetworkPlatformFile, Display, TEXT("    Server ProjectPlatformExtDir = %s"), *ServerProjectPlatformExtensionsDir);
-	UE_LOG(LogNetworkPlatformFile, Display, TEXT("     Local ProjectPlatformExtDir = %s"), *FPaths::ProjectPlatformExtensionsDir());
+	UE_LOG(LogNetworkPlatformFile, Display, TEXT("     Local ProjectPlatformExtDir = %s"), *FPaths::ProjectPlatformExtensionDir(*IniPlatformName));
 
 	// Receive a list of files and their timestamps.
 	TMap<FString, FDateTime> ServerFileMap;
@@ -1479,11 +1482,13 @@ void FNetworkPlatformFile::ConvertServerFilenameToClientFilename(FString& Filena
 	}
 	else if (FilenameToConvert.StartsWith(InServerEnginePlatformExtensionsDir))
 	{
-		FilenameToConvert = FilenameToConvert.Replace(*InServerEnginePlatformExtensionsDir, *(FPaths::EnginePlatformExtensionsDir()));
+		FString IniPlatformName(FPlatformProperties::IniPlatformName());
+		FilenameToConvert = FilenameToConvert.Replace(*InServerEnginePlatformExtensionsDir, *(FPaths::EnginePlatformExtensionDir(*IniPlatformName)));
 	}
 	else if (FilenameToConvert.StartsWith(InServerProjectPlatformExtensionsDir))
 	{
-		FilenameToConvert = FilenameToConvert.Replace(*InServerProjectPlatformExtensionsDir, *(FPaths::ProjectPlatformExtensionsDir()));
+		FString IniPlatformName(FPlatformProperties::IniPlatformName());
+		FilenameToConvert = FilenameToConvert.Replace(*InServerProjectPlatformExtensionsDir, *(FPaths::ProjectPlatformExtensionDir(*IniPlatformName)));
 	}
 }
 
diff --git a/Engine/Source/Runtime/NetworkFileSystem/Private/NetworkFileServerConnection.cpp b/Engine/Source/Runtime/NetworkFileSystem/Private/NetworkFileServerConnection.cpp
--- a/Engine/Source/Runtime/NetworkFileSystem/Private/NetworkFileServerConnection.cpp
+++ b/Engine/Source/Runtime/NetworkFileSystem/Private/NetworkFileServerConnection.cpp
@@ -91,8 +91,6 @@ FNetworkFileServerClientConnection::FNetworkFileServerClientConnection(const FNe
 
 	LocalEngineDir = FPaths::EngineDir();
 	LocalProjectDir = FPaths::ProjectDir();
-	LocalEnginePlatformExtensionsDir = FPaths::EnginePlatformExtensionsDir();
-	LocalProjectPlatformExtensionsDir = FPaths::ProjectPlatformExtensionsDir();
 
 	if (FPaths::IsProjectFilePathSet())
 	{
@@ -102,8 +100,6 @@ FNetworkFileServerClientConnection::FNetworkFileServerClientConnection(const FNe
 
 	LocalEngineDirAbs = MakeAbsoluteNormalizedDir(LocalEngineDir);
 	LocalProjectDirAbs = MakeAbsoluteNormalizedDir(LocalProjectDir);
-	LocalEnginePlatformExtensionsDirAbs = MakeAbsoluteNormalizedDir(LocalEnginePlatformExtensionsDir);
-	LocalProjectPlatformExtensionsDirAbs = MakeAbsoluteNormalizedDir(LocalEnginePlatformExtensionsDir);
 }
 
 
@@ -157,11 +153,11 @@ void FNetworkFileServerClientConnection::ConvertClientFilenameToServerFilename(F
 		// We do *not* want to replace the directory in that case.
 		return;
 	}
-	if (TrySubstituteDirectory(FilenameToConvert, FPaths::EnginePlatformExtensionsDir(), ConnectedEnginePlatformExtensionsDir))
+	if (TrySubstituteDirectory(FilenameToConvert, FPaths::EnginePlatformExtensionDir(*ConnectedIniPlatformName), ConnectedEnginePlatformExtensionsDir))
 	{
 		return;
 	}
-	if (TrySubstituteDirectory(FilenameToConvert, FPaths::ProjectPlatformExtensionsDir(), ConnectedProjectPlatformExtensionsDir))
+	if (TrySubstituteDirectory(FilenameToConvert, FPaths::ProjectPlatformExtensionDir(*ConnectedIniPlatformName), ConnectedProjectPlatformExtensionsDir))
 	{
 		return;
 	}
@@ -875,6 +871,7 @@ bool FNetworkFileServerClientConnection::ProcessGetFileList( FArchive& In, FArch
 	FString ClientVersionInfo;
 	FString TargetAddress;
 
+	In << ConnectedIniPlatformName;
 	In << TargetPlatformNames;
 	In << GameName;
 	In << EngineRelativePath;
@@ -961,6 +958,12 @@ bool FNetworkFileServerClientConnection::ProcessGetFileList( FArchive& In, FArch
 	ConnectedProjectDir = GameRelativePath;
 	ConnectedEnginePlatformExtensionsDir = EnginePlatformExtensionsRelativePath;
 	ConnectedProjectPlatformExtensionsDir = ProjectPlatformExtensionsRelativePath;
+	
+	// now that we have a connected platform for this connection, we can query the platform extension directory locations
+	LocalEnginePlatformExtensionsDir = FPaths::EnginePlatformExtensionDir(*ConnectedIniPlatformName);
+	LocalProjectPlatformExtensionsDir = FPaths::ProjectPlatformExtensionDir(*ConnectedIniPlatformName);
+	LocalEnginePlatformExtensionsDirAbs = MakeAbsoluteNormalizedDir(LocalEnginePlatformExtensionsDir);
+	LocalProjectPlatformExtensionsDirAbs = MakeAbsoluteNormalizedDir(LocalEnginePlatformExtensionsDir);
 
 	UE_LOG(LogFileServer, Display, TEXT("    Connected EngineDir      = %s"), *ConnectedEngineDir);
 	UE_LOG(LogFileServer, Display, TEXT("        Local EngineDir      = %s"), *LocalEngineDir);
diff --git a/Engine/Source/Runtime/NetworkFileSystem/Private/NetworkFileServerConnection.h b/Engine/Source/Runtime/NetworkFileSystem/Private/NetworkFileServerConnection.h
--- a/Engine/Source/Runtime/NetworkFileSystem/Private/NetworkFileServerConnection.h
+++ b/Engine/Source/Runtime/NetworkFileSystem/Private/NetworkFileServerConnection.h
@@ -188,7 +188,10 @@ protected:
 	
 private:
 
-	// Hold the name of the currently connected platform.
+	// Hold the ini name of the currently connected platform.
+	FString ConnectedIniPlatformName;
+
+	// Hold the name of the currently connected target platform.
 	FString ConnectedPlatformName;
 
 	// Hold the ip address of the currently connected platform.
diff --git a/Engine/Source/Runtime/Projects/Private/PluginManager.cpp b/Engine/Source/Runtime/Projects/Private/PluginManager.cpp
--- a/Engine/Source/Runtime/Projects/Private/PluginManager.cpp
+++ b/Engine/Source/Runtime/Projects/Private/PluginManager.cpp
@@ -832,7 +832,6 @@ void FPluginManager::ReadAllPlugins(FDiscoveredPluginMap& Plugins, const TSet<FS
 
 			// Get all the standard plugin directories
 			const FString EngineDir = FPaths::EngineDir();
-			const FString PlatformExtensionEngineDir = FPaths::EnginePlatformExtensionsDir();
 			const FString EnterpriseDir = FPaths::EnterpriseDir();
 			const FString ProjectModsDir = FPaths::ProjectModsDir();
 
@@ -843,7 +842,7 @@ void FPluginManager::ReadAllPlugins(FDiscoveredPluginMap& Plugins, const TSet<FS
 				SlowTask_ReadManifestContents.EnterProgressFrame();
 
 				EPluginType Type;
-				if (Entry.File.StartsWith(EngineDir) || Entry.File.StartsWith(PlatformExtensionEngineDir))
+				if (Entry.File.StartsWith(EngineDir))
 				{
 					Type = EPluginType::Engine;
 				}
diff --git a/Engine/Source/Runtime/SlateCore/Private/Styling/StarshipCoreStyle.cpp b/Engine/Source/Runtime/SlateCore/Private/Styling/StarshipCoreStyle.cpp
--- a/Engine/Source/Runtime/SlateCore/Private/Styling/StarshipCoreStyle.cpp
+++ b/Engine/Source/Runtime/SlateCore/Private/Styling/StarshipCoreStyle.cpp
@@ -976,10 +976,10 @@ TSharedRef<ISlateStyle> FStarshipCoreStyle::Create()
 			{
 				if (NormalIconPath.StartsWith(TEXT("/Platforms/")))
 				{
-#define PLATFORM_IMAGE_BRUSH( PlatformPath, ... ) FSlateImageBrush( PlatformPath.Replace(TEXT("/Platforms/"), *FPaths::EnginePlatformExtensionsDir()) + TEXT(".png") , __VA_ARGS__ )
-					Style->Set(PlatformInfo.GetIconStyleName(EPlatformIconSize::Normal), new PLATFORM_IMAGE_BRUSH(NormalIconPath, Icon24x24));
-					Style->Set(PlatformInfo.GetIconStyleName(EPlatformIconSize::Large), new PLATFORM_IMAGE_BRUSH(PlatformInfo.GetIconPath(EPlatformIconSize::Large), Icon64x64));
-					Style->Set(PlatformInfo.GetIconStyleName(EPlatformIconSize::XLarge), new PLATFORM_IMAGE_BRUSH(PlatformInfo.GetIconPath(EPlatformIconSize::XLarge), Icon128x128));
+#define PLATFORM_IMAGE_BRUSH( Platform, PlatformPath, ... ) FSlateImageBrush( PlatformPath.Replace(*(FString(TEXT("/Platforms/")) + Platform.ToString() + TEXT("/")), *FPaths::EnginePlatformExtensionDir(*Platform.ToString())) + TEXT(".png") , __VA_ARGS__ )
+					Style->Set(PlatformInfo.GetIconStyleName(EPlatformIconSize::Normal), new PLATFORM_IMAGE_BRUSH(PlatformInfo.IniPlatformName, NormalIconPath, Icon24x24));
+					Style->Set(PlatformInfo.GetIconStyleName(EPlatformIconSize::Large), new PLATFORM_IMAGE_BRUSH(PlatformInfo.IniPlatformName, PlatformInfo.GetIconPath(EPlatformIconSize::Large), Icon64x64));
+					Style->Set(PlatformInfo.GetIconStyleName(EPlatformIconSize::XLarge), new PLATFORM_IMAGE_BRUSH(PlatformInfo.IniPlatformName, PlatformInfo.GetIconPath(EPlatformIconSize::XLarge), Icon128x128));
 				}
 				else
 				{
@@ -997,11 +997,11 @@ TSharedRef<ISlateStyle> FStarshipCoreStyle::Create()
 		{
 			if (!Item.ActiveIconPath.IsEmpty())
 			{
-				Style->Set(Item.ActiveIconName, new PLATFORM_IMAGE_BRUSH(Item.ActiveIconPath, Icon40x40));
+				Style->Set(Item.ActiveIconName, new PLATFORM_IMAGE_BRUSH(Item.PlatformName, Item.ActiveIconPath, Icon40x40));
 			}
 			if (!Item.InactiveIconPath.IsEmpty())
 			{
-				Style->Set(Item.InactiveIconName, new PLATFORM_IMAGE_BRUSH(Item.InactiveIconPath, Icon40x40));
+				Style->Set(Item.InactiveIconName, new PLATFORM_IMAGE_BRUSH(Item.PlatformName, Item.InactiveIconPath, Icon40x40));
 			}
 		}
 #endif
