Disable UE_DEBUG_SECTION when building with header units because it causes linker errors

diff --git a/Engine/Source/Programs/UnrealBuildTool/Configuration/UEBuildTarget.cs b/Engine/Source/Programs/UnrealBuildTool/Configuration/UEBuildTarget.cs
--- a/Engine/Source/Programs/UnrealBuildTool/Configuration/UEBuildTarget.cs
+++ b/Engine/Source/Programs/UnrealBuildTool/Configuration/UEBuildTarget.cs
@@ -4865,6 +4865,11 @@ namespace UnrealBuildTool
 				GlobalCompileEnvironment.Definitions.Add("ENABLE_PGO_PROFILE=0");
 			}
 
+			if (Rules.bUseHeaderUnitsForPch)
+			{
+				GlobalCompileEnvironment.Definitions.Add("UE_DEBUG_SECTION=");
+			}
+
 			// Toggle to enable vorbis for audio streaming where available
 			GlobalCompileEnvironment.Definitions.Add("USE_VORBIS_FOR_STREAMING=1");
 
diff --git a/Engine/Source/Runtime/Core/Public/Misc/AssertionMacros.h b/Engine/Source/Runtime/Core/Public/Misc/AssertionMacros.h
--- a/Engine/Source/Runtime/Core/Public/Misc/AssertionMacros.h
+++ b/Engine/Source/Runtime/Core/Public/Misc/AssertionMacros.h
@@ -14,6 +14,7 @@
 #include "Misc/VarArgs.h"
 #include "String/FormatStringSan.h"
 
+#ifndef UE_DEBUG_SECTION
 #if (DO_CHECK || DO_GUARD_SLOW || DO_ENSURE) && !PLATFORM_CPU_ARM_FAMILY
 	// We'll put all assert implementation code into a separate section in the linked
 	// executable. This code should never execute so using a separate section keeps
@@ -27,6 +28,7 @@
 	// is present it will generate code that it cannot link.
 	#define UE_DEBUG_SECTION
 #endif // DO_CHECK || DO_GUARD_SLOW
+#endif
 
 namespace ELogVerbosity
 {
