Fix for niagara attribute trimmer edge case where we would miss dependencies for the first read of a variable within a stage (in very specific circumstances).

diff --git a/Engine/Plugins/FX/Niagara/Source/NiagaraEditor/Private/NiagaraAttributeTrimmer.cpp b/Engine/Plugins/FX/Niagara/Source/NiagaraEditor/Private/NiagaraAttributeTrimmer.cpp
--- a/Engine/Plugins/FX/Niagara/Source/NiagaraEditor/Private/NiagaraAttributeTrimmer.cpp
+++ b/Engine/Plugins/FX/Niagara/Source/NiagaraEditor/Private/NiagaraAttributeTrimmer.cpp
@@ -537,7 +537,13 @@ void FNiagaraAttributeTrimmerHelper<GraphBridge>::ResolveDependencyChain(const F
 
 										if (!bValidDefaultWriteFound || bBackCompat_AlwaysIncludeFirstReadPin)
 										{
-											ResolvedDependencies.Pins.Add(ParamMap.PerVariableReadHistory[DefaultBoundVariable][0].ReadPin);
+											const typename FParamMapHistory::FReadHistory& InitialReadHistory = ParamMap.PerVariableReadHistory[DefaultBoundVariable][0];
+
+											if (InitialReadHistory.PreviousWritePin.Pin)
+											{
+												PinsToResolve.Add(InitialReadHistory.PreviousWritePin);
+											}
+											ResolvedDependencies.Pins.Add(InitialReadHistory.ReadPin);
 										}
 									}
 								}
diff --git a/Engine/Source/Runtime/Core/Private/UObject/FortniteMainBranchObjectVersion.cpp b/Engine/Source/Runtime/Core/Private/UObject/FortniteMainBranchObjectVersion.cpp
--- a/Engine/Source/Runtime/Core/Private/UObject/FortniteMainBranchObjectVersion.cpp
+++ b/Engine/Source/Runtime/Core/Private/UObject/FortniteMainBranchObjectVersion.cpp
@@ -11,7 +11,7 @@ TMap<FGuid, FGuid> FFortniteMainBranchObjectVersion::GetSystemGuids()
 	SystemGuids.Add(DevGuids.MATERIALSHADERMAP_DERIVEDDATA_VER, FGuid("CE7776E63E9E43FBB8C103EE1AB40B5D"));
 	SystemGuids.Add(DevGuids.NANITE_DERIVEDDATA_VER, FGuid("D28C003F825F4C178F698938FB0E47F7"));
 	SystemGuids.Add(DevGuids.NIAGARASHADERMAP_DERIVEDDATA_VER, FGuid("D0DF49D408F34E38A98FA097E0280EC7"));
-	SystemGuids.Add(DevGuids.Niagara_LatestScriptCompileVersion, FGuid("3C56691355BF1941B5D7A4DB0AFE5F95"));
+	SystemGuids.Add(DevGuids.Niagara_LatestScriptCompileVersion, FGuid("2017726CDB77425CB8F3D6278D491E65"));
 	SystemGuids.Add(DevGuids.SkeletalMeshDerivedDataVersion, FGuid("518D7C51D31D9BE4A9F4B81542AC8119"));
 	SystemGuids.Add(DevGuids.STATICMESH_DERIVEDDATA_VER, FGuid("543843B103794E0E9BA4BE80FB602F79"));
 	return SystemGuids;
