Rework the hitset to use a fibonacci hash (simplifies the implementation a fair bit).

diff --git a/Engine/Source/Runtime/Core/Private/AutoRTFM/HitSet.h b/Engine/Source/Runtime/Core/Private/AutoRTFM/HitSet.h
--- a/Engine/Source/Runtime/Core/Private/AutoRTFM/HitSet.h
+++ b/Engine/Source/Runtime/Core/Private/AutoRTFM/HitSet.h
@@ -10,29 +10,28 @@ namespace AutoRTFM
 	class FHitSet final
 	{
 		// TODO: Revisit a good probe depth for the hashset.
-		static constexpr uint32_t LinearProbeDepth = 21;
+		static constexpr uint32_t LinearProbeDepth = 16;
 
 		// TODO: Revisit a good initial capacity for the hitset.
-		static constexpr uint32_t InitialCapacity = 128;
+		static constexpr uint32_t LogInitialCapacity = 4;
 
 	public:
 		using Key = TTaggedPtr<void>;
 
 		explicit FHitSet()
 		{
+			constexpr uint32_t InitialCapacity = 1u << LogInitialCapacity;
+
 			// Do not want to deal with null payloads.
-			ASSERT(0 != InitialCapacity);
+			static_assert(0 != InitialCapacity);
 
-			// The capacity increases by a power of two so that the modulus is optimal.
-			ASSERT(0 == (InitialCapacity & (InitialCapacity - 1)));
+			// The capacity is always a power of two so that the range reduction is optimal.
+			static_assert(0 == (InitialCapacity & (InitialCapacity - 1)));
 
-			// To make the linear probe check simpler, our payload has a padded
-			// linear probe depth so the last element doesn't have to wrap
-			// around and insert at the front.
-			Payload = new uintptr_t[InitialCapacity + LinearProbeDepth]();
+			Payload = new uintptr_t[InitialCapacity]();
 			ASSERT(nullptr != Payload);
 
-			Capacity = InitialCapacity;
+			SixtyFourMinusLogCapacity = 64 - LogInitialCapacity;
 			Size = 0;
 		}
 
@@ -55,14 +54,14 @@ namespace AutoRTFM
 					return true;
 				case EInsertResult::Exists:
 					return false;
-				case EInsertResult::NeedResize:
+				case EInsertResult::NotInserted:
 					Resize();
 					continue;
 				}
 			}
 		}
 
-		bool IsEmpty() const
+		UE_AUTORTFM_FORCEINLINE bool IsEmpty() const
 		{
 			return 0 == Size;
 		}
@@ -70,22 +69,46 @@ namespace AutoRTFM
 		// Clear out the data stored in the set, but does not reduce the capacity.
 		void Reset()
 		{
-			memset(Payload, 0, sizeof(uintptr_t) * (Capacity + LinearProbeDepth));
+			memset(Payload, 0, sizeof(uintptr_t) * Capacity());
 			Size = 0;
 		}
 
+		uint32_t GetCapacity() const
+		{
+			return Capacity();
+		}
+
+		uint32_t GetSize() const
+		{
+			return Size;
+		}
+
 	private:
 		uintptr_t* Payload;
 
-		// TODO: Make this (Capacity - 1) to save a minus in the modulus calculation.
-		uint32_t Capacity;
+		uint32_t SixtyFourMinusLogCapacity;
 
 		uint32_t Size;
 
+		UE_AUTORTFM_FORCEINLINE uint32_t Capacity() const
+		{
+			return 1u << (64 - SixtyFourMinusLogCapacity);
+		}
+
+		UE_AUTORTFM_FORCEINLINE void IncreaseCapacity()
+		{
+			// It seems odd - but subtracting 1 here is totally intentional
+			// because of how we store our capacity as (64 - log2(Capacity)).
+			SixtyFourMinusLogCapacity -= 1;
+
+			// Check that we haven't overflowed our capacity!
+			ASSERT(0 != SixtyFourMinusLogCapacity);
+		}
+
 		void Resize()
 		{
 			const uintptr_t* const OldPayload = Payload;
-			const uint32_t OldCapacity = Capacity;
+			const uint32_t OldCapacity = Capacity();
 			const uint32_t OldSize = Size;
 
 			while (true)
@@ -94,17 +117,14 @@ namespace AutoRTFM
 
 				Size = 0;
 
-				Capacity *= 2;
-
-				// Check that we haven't overflowed our capacity!
-				ASSERT(0 != Capacity);
+				IncreaseCapacity();
 
-				Payload = new uintptr_t[Capacity + LinearProbeDepth]();
+				Payload = new uintptr_t[Capacity()]();
 				ASSERT(nullptr != Payload);
 
 				// Now we need to rehash and reinsert all the items. We need to
 				// remember about the extra linear probe section here too.
-				for (size_t I = 0; I < OldCapacity + LinearProbeDepth; I++)
+				for (size_t I = 0; I < OldCapacity; I++)
 				{
 					// Skip empty locations.
 					if (0 == OldPayload[I])
@@ -114,7 +134,7 @@ namespace AutoRTFM
 
 					const EInsertResult Result = InsertNoResize(OldPayload[I]);
 
-					if (EInsertResult::NeedResize == Result)
+					if (EInsertResult::NotInserted == Result)
 					{
 						NeedAnotherResize = true;
 						break;
@@ -145,50 +165,83 @@ namespace AutoRTFM
 		{
 			Exists,
 			Inserted,
-			NeedResize
+			NotInserted
 		};
 
-		// Insert something in the HitSet, returning true if the insert
-		// succeeded (EG. the key was not already in the set).
-		EInsertResult InsertNoResize(const uintptr_t Raw)
+		UE_AUTORTFM_FORCEINLINE uintptr_t FirstHashInRange(const uintptr_t Hashee) const
+		{
+			constexpr uintptr_t Fibonacci = 0x9E3779B97F4A7C15;
+			return (Hashee * Fibonacci) >> SixtyFourMinusLogCapacity;
+		}
+
+		UE_AUTORTFM_FORCEINLINE uintptr_t SecondHash(const uintptr_t Hashee) const
 		{
 			// TODO: Find a great hashing function for the set.
 			constexpr uintptr_t A = 0x65d200ce55b19ad8;
 			constexpr uintptr_t B = 0x4f2162926e40c299;
 			constexpr uintptr_t C = 0x162dd799029970f8;
 
-			const uintptr_t Low = Raw & 0x00000000FFFFFFFF;
-			const uintptr_t High = Raw >> 32;
+			const uintptr_t Low = Hashee & 0x00000000FFFFFFFF;
+			const uintptr_t High = Hashee >> 32;
 
-			const uint32_t Hash = static_cast<uint32_t>((A * Low + B * High + C) >> 32);
+			const uintptr_t Hash = A * Low + B * High + C;
 
-			if (Size < Capacity)
+			return Hash;
+		}
+
+		UE_AUTORTFM_FORCEINLINE EInsertResult TryInsertAtIndex(const uintptr_t Raw, const uintptr_t I)
+		{
+			// We have a free location in the set.
+			if (0 == Payload[I])
 			{
-				const uint32_t HashModCapacity = Hash & (Capacity - 1);
+				Payload[I] = Raw;
+				Size++;
+				return EInsertResult::Inserted;
+			}
+
+			// We're already in the set.
+			if (Raw == Payload[I])
+			{
+				return EInsertResult::Exists;
+			}
+
+			return EInsertResult::NotInserted;
+		}
+
+		// Insert something in the HitSet, returning true if the insert
+		// succeeded (EG. the key was not already in the set).
+		UE_AUTORTFM_FORCEINLINE EInsertResult InsertNoResize(const uintptr_t Raw)
+		{
+			// First try to use our simple first hash to get a location.
+			{
+				const uintptr_t I = FirstHashInRange(Raw);
 
-				for (uint32_t D = 0; D < LinearProbeDepth; D++)
+				const EInsertResult R = TryInsertAtIndex(Raw, I);
+
+				if (EInsertResult::NotInserted != R)
 				{
-					const uintptr_t I = HashModCapacity + D;
+					return R;
+				}
+			}
 
-					// We have a free location in the set.
-					if (0 == Payload[I])
-					{
-						Payload[I] = Raw;
-						Size++;
-						return EInsertResult::Inserted;
-					}
+			// Then, if that fails, we use a second more complicated hash with a linear probe.
+			const uintptr_t Hash = SecondHash(Raw);
 
-					// We're already in the set.
-					if (Raw == Payload[I])
-					{
-						return EInsertResult::Exists;
-					}
+			for (uint32_t D = 0; D < LinearProbeDepth; D++)
+			{
+				const uintptr_t I = FirstHashInRange(Hash + D);
+
+				const EInsertResult R = TryInsertAtIndex(Raw, I);
 
-					// Otherwise we need to keep going with our linear probe...
+				if (EInsertResult::NotInserted != R)
+				{
+					return R;
 				}
+
+				// Otherwise we need to keep going with our linear probe...
 			}
 
-			return EInsertResult::NeedResize;
+			return EInsertResult::NotInserted;
 		}
 	};
 }
diff --git a/Engine/Source/Runtime/Core/Private/AutoRTFM/Stats.cpp b/Engine/Source/Runtime/Core/Private/AutoRTFM/Stats.cpp
--- a/Engine/Source/Runtime/Core/Private/AutoRTFM/Stats.cpp
+++ b/Engine/Source/Runtime/Core/Private/AutoRTFM/Stats.cpp
@@ -98,6 +98,20 @@ template<> void AutoRTFM::FStats::Report<AutoRTFM::EStatsKind::NewMemoryTrackerM
     UE_LOG(LogAutoRTFM, Display, TEXT("  New memory misses:         %11u"), Data);
 }
 
+template<> void AutoRTFM::FStats::Report<AutoRTFM::EStatsKind::AverageHitSetSize>(const uint64_t Data) const
+{
+    const uint64_t TotalTransactions = Datas[static_cast<size_t>(AutoRTFM::EStatsKind::Transaction)];
+    UE_LOG(LogAutoRTFM, Display, TEXT("  Average hit set size:      %15.3f"), (static_cast<double>(Data) / static_cast<double>(TotalTransactions)));
+}
+
+template<> void AutoRTFM::FStats::Report<AutoRTFM::EStatsKind::AverageHitSetCapacity>(const uint64_t Data) const
+{
+    const uint64_t TotalTransactions = Datas[static_cast<size_t>(AutoRTFM::EStatsKind::Transaction)];
+    UE_LOG(LogAutoRTFM, Display, TEXT("  Average hit set capacity:  %15.3f"), (static_cast<double>(Data) / static_cast<double>(TotalTransactions)));
+
+    const uint64_t TotalHitSetSize = Datas[static_cast<size_t>(AutoRTFM::EStatsKind::AverageHitSetSize)];
+    UE_LOG(LogAutoRTFM, Display, TEXT("  Average hit set occupancy: %15.3f"), 100 * (static_cast<double>(TotalHitSetSize) / static_cast<double>(Data)));
+}
 
 void AutoRTFM::FStats::Report() const
 {
@@ -129,6 +143,8 @@ void AutoRTFM::FStats::Report() const
             REPORT_CASE(AutoRTFM::EStatsKind::MaximumAbortTasks);
             REPORT_CASE(AutoRTFM::EStatsKind::NewMemoryTrackerHit);
             REPORT_CASE(AutoRTFM::EStatsKind::NewMemoryTrackerMiss);
+			REPORT_CASE(AutoRTFM::EStatsKind::AverageHitSetSize);
+			REPORT_CASE(AutoRTFM::EStatsKind::AverageHitSetCapacity);
 #undef REPORT_CASE
 			}
 		}
diff --git a/Engine/Source/Runtime/Core/Private/AutoRTFM/Stats.h b/Engine/Source/Runtime/Core/Private/AutoRTFM/Stats.h
--- a/Engine/Source/Runtime/Core/Private/AutoRTFM/Stats.h
+++ b/Engine/Source/Runtime/Core/Private/AutoRTFM/Stats.h
@@ -28,6 +28,8 @@ enum class EStatsKind : uint8_t
 	MaximumAbortTasks,
 	NewMemoryTrackerHit,
 	NewMemoryTrackerMiss,
+	AverageHitSetSize,
+	AverageHitSetCapacity,
 	Total
 };
 
@@ -65,6 +67,8 @@ template<EStatsKind Kind> constexpr bool IsStatsKindOneArg()
 	case EStatsKind::MaximumCommitTasks:
 	case EStatsKind::AverageAbortTasks:
 	case EStatsKind::MaximumAbortTasks:
+	case EStatsKind::AverageHitSetSize:
+	case EStatsKind::AverageHitSetCapacity:
 		return true;
 	}
 }
@@ -80,6 +84,8 @@ template<EStatsKind Kind> constexpr bool IsStatsKindAverage()
 	case EStatsKind::AverageWriteLogBytes:
 	case EStatsKind::AverageCommitTasks:
 	case EStatsKind::AverageAbortTasks:
+	case EStatsKind::AverageHitSetSize:
+	case EStatsKind::AverageHitSetCapacity:
 		return true;
 	}
 }
diff --git a/Engine/Source/Runtime/Core/Private/AutoRTFM/TransactionInlines.h b/Engine/Source/Runtime/Core/Private/AutoRTFM/TransactionInlines.h
--- a/Engine/Source/Runtime/Core/Private/AutoRTFM/TransactionInlines.h
+++ b/Engine/Source/Runtime/Core/Private/AutoRTFM/TransactionInlines.h
@@ -175,6 +175,9 @@ UE_AUTORTFM_FORCEINLINE void FTransaction::CollectStats() const
 
     Stats.Collect<EStatsKind::AverageAbortTasks>(AbortTasks.Num());
     Stats.Collect<EStatsKind::MaximumAbortTasks>(AbortTasks.Num());
+
+    Stats.Collect<EStatsKind::AverageHitSetSize>(HitSet.GetSize());
+    Stats.Collect<EStatsKind::AverageHitSetCapacity>(HitSet.GetCapacity());
 }
 
 } // namespace AutoRTFM
