Increase default IO block size in editor for Unix platforms from 1KB to 256KB

#rnx

diff --git a/Engine/Source/Runtime/Core/Public/Mac/MacPlatform.h b/Engine/Source/Runtime/Core/Public/Mac/MacPlatform.h
--- a/Engine/Source/Runtime/Core/Public/Mac/MacPlatform.h
+++ b/Engine/Source/Runtime/Core/Public/Mac/MacPlatform.h
@@ -80,6 +80,10 @@ typedef FMacPlatformTypes FPlatformTypes;
 
 #define PLATFORM_GLOBAL_LOG_CATEGORY					LogMac
 
+#if WITH_EDITOR
+#define PLATFORM_FILE_READER_BUFFER_SIZE				(256*1024)
+#endif
+
 #if PLATFORM_MAC_X86
 	#define PLATFORM_BREAK()							__asm__("int $3")
 #else
diff --git a/Engine/Source/Runtime/Core/Public/Unix/UnixPlatform.h b/Engine/Source/Runtime/Core/Public/Unix/UnixPlatform.h
--- a/Engine/Source/Runtime/Core/Public/Unix/UnixPlatform.h
+++ b/Engine/Source/Runtime/Core/Public/Unix/UnixPlatform.h
@@ -71,6 +71,10 @@ typedef FUnixPlatformTypes FPlatformTypes;
 
 #define PLATFORM_RETURN_ADDRESS_FOR_CALLSTACKTRACING	PLATFORM_RETURN_ADDRESS
 
+#if WITH_EDITOR
+#define PLATFORM_FILE_READER_BUFFER_SIZE				(256*1024)
+#endif
+
 #if PLATFORM_CPU_X86_FAMILY
 	#define PLATFORM_BREAK()							__asm__ volatile("int $0x03")
 #else
