Fixed usage of CpuProfiler trace in TaskTrace.
FCpuProfilerTrace::OutputEventType() and FCpuProfilerTrace::OutputBeginDynamicEvent() needs to be called only if the CpuChannel is enabled, otherwise the "spec" trace event is not emitted, so it will result in "<unknown>" cpu timers. The issue can be easily reproduced with -trace=tasks (i.e. without enabling cpu channel).

diff --git a/Engine/Source/Runtime/Core/Private/Async/TaskTrace.cpp b/Engine/Source/Runtime/Core/Private/Async/TaskTrace.cpp
--- a/Engine/Source/Runtime/Core/Private/Async/TaskTrace.cpp
+++ b/Engine/Source/Runtime/Core/Private/Async/TaskTrace.cpp
@@ -81,7 +81,7 @@ namespace TaskTrace
 	}
 
 	static bool bGTaskTraceInitialized = false;
-	uint32 ExecuteTaskSpecId = 0;
+	std::atomic<uint32> ExecuteTaskSpecId = 0;
 
 	void Init()
 	{
@@ -90,7 +90,10 @@ namespace TaskTrace
 
 		bGTaskTraceInitialized = true;
 #if CPUPROFILERTRACE_ENABLED
-		ExecuteTaskSpecId = FCpuProfilerTrace::OutputEventType("ExecuteTask");
+		if (UE_TRACE_CHANNELEXPR_IS_ENABLED(CpuChannel))
+		{
+			ExecuteTaskSpecId = FCpuProfilerTrace::OutputEventType("ExecuteTask", __FILE__, __LINE__);
+		}
 #endif
 	}
 
@@ -262,10 +265,14 @@ namespace TaskTrace
 		// The RenderingThread outputs a BeginEvent in the BeginFrameRenderThread function and an EndEvent in EndFrameRenderThread.
 		// Outputing the ExecuteTask event would cause the Frame event to be closed incorrectly.
 #if CPUPROFILERTRACE_ENABLED
-		if (UE_TRACE_CHANNELEXPR_IS_ENABLED(TaskChannel) && !IsInRenderingThread())
+		if (UE_TRACE_CHANNELEXPR_IS_ENABLED(TaskChannel|CpuChannel) && !IsInRenderingThread())
 		{
-			bIsActive = true;
+			if (ExecuteTaskSpecId == 0)
+			{
+				ExecuteTaskSpecId = FCpuProfilerTrace::OutputEventType("ExecuteTask", __FILE__, __LINE__);
+			}
 			FCpuProfilerTrace::OutputBeginEvent(ExecuteTaskSpecId);
+			bIsActive = true;
 		}
 #endif
 	}
