- Fixed an checkfSlow to handle being passed in EngineDir but missing a trailing slash
#rnx

diff --git a/Engine/Source/Runtime/Core/Private/Misc/Paths.cpp b/Engine/Source/Runtime/Core/Private/Misc/Paths.cpp
--- a/Engine/Source/Runtime/Core/Private/Misc/Paths.cpp
+++ b/Engine/Source/Runtime/Core/Private/Misc/Paths.cpp
@@ -311,8 +311,8 @@ FString FPaths::ConvertPath(const FString& Path, EPathConversion Method, const T
 	}
 	else
 	{
-		checkfSlow(!(bIsProjectConversion && !Path.StartsWith(FPaths::ProjectDir())), TEXT("Called ConvertPath with a Platform_ method, but Path didn't start with FPaths::ProjectDir(). This is not supported"));
-		checkfSlow(!(bIsEngineConversion && !Path.StartsWith(FPaths::EngineDir())), TEXT("Called ConvertPath with a Engine_ method, but Path didn't start with FPaths::EngineDir(). This is not supported"));
+		checkfSlow(!(bIsProjectConversion && !Path.StartsWith(FPaths::ProjectDir().TrimChar('/'))), TEXT("Called ConvertPath with a Platform_ method, but Path [%s] didn't start with FPaths::ProjectDir() [%s]. This is not supported"), *Path, *FPaths::ProjectDir());
+		checkfSlow(!(bIsEngineConversion && !Path.StartsWith(FPaths::EngineDir().TrimChar('/'))), TEXT("Called ConvertPath with a Engine_ method, but Path [%s] didn't start with FPaths::EngineDir() [%s]. This is not supported"), *Path, *FPaths::EngineDir());
 
 		// skip over engine or project directory
 		Prefix = (int)Method < (int)EPathConversion::Project_First ? FPaths::EngineDir() : FPaths::ProjectDir();
