- Fixed platform extensions DDPIs being added twice when running a commandlet
#rnx

diff --git a/Engine/Source/Runtime/Core/Private/Misc/DataDrivenPlatformInfoRegistry.cpp b/Engine/Source/Runtime/Core/Private/Misc/DataDrivenPlatformInfoRegistry.cpp
--- a/Engine/Source/Runtime/Core/Private/Misc/DataDrivenPlatformInfoRegistry.cpp
+++ b/Engine/Source/Runtime/Core/Private/Misc/DataDrivenPlatformInfoRegistry.cpp
@@ -50,9 +50,8 @@ static const TArray<FString>& GetDataDrivenIniFilenames()
 		}
 
 		// look for the special files in any project config subdirectories
-		// @note: this will not work for Programs, the Platforms dir is outside of the ProjectDir in their case
 		TArray<FString> ProjectPlatformDirs;
-		IFileManager::Get().FindFiles(ProjectPlatformDirs, *FPaths::Combine(FPaths::ProjectDir(), TEXT("Platforms"), TEXT("*")), false, true);
+		IFileManager::Get().FindFiles(ProjectPlatformDirs, *FPaths::Combine(FPaths::ProjectConfigDir(), TEXT("*")), false, true);
 		for (const FString& PlatformDir : ProjectPlatformDirs)
 		{
 			FString IniPath = FPaths::Combine(FPaths::ProjectPlatformExtensionDir(*PlatformDir), TEXT("Config/DataDrivenPlatformInfo.ini"));
