- Removed some checkSlows in the ConvertPath after an edge case was found of using absolute path to Engine, will rework the function to handle absolute paths if possible, and return a failure if it can't handle it
#rnx

diff --git a/Engine/Source/Runtime/Core/Private/Misc/Paths.cpp b/Engine/Source/Runtime/Core/Private/Misc/Paths.cpp
--- a/Engine/Source/Runtime/Core/Private/Misc/Paths.cpp
+++ b/Engine/Source/Runtime/Core/Private/Misc/Paths.cpp
@@ -311,8 +311,10 @@ FString FPaths::ConvertPath(const FString& Path, EPathConversion Method, const T
 	}
 	else
 	{
-		checkfSlow(!(bIsProjectConversion && !Path.StartsWith(FPaths::ProjectDir().TrimChar('/'))), TEXT("Called ConvertPath with a Platform_ method, but Path [%s] didn't start with FPaths::ProjectDir() [%s]. This is not supported"), *Path, *FPaths::ProjectDir());
-		checkfSlow(!(bIsEngineConversion && !Path.StartsWith(FPaths::EngineDir().TrimChar('/'))), TEXT("Called ConvertPath with a Engine_ method, but Path [%s] didn't start with FPaths::EngineDir() [%s]. This is not supported"), *Path, *FPaths::EngineDir());
+		// removing these temporarily until we handle an absolute path to Engine/Project - the GetUnrealBuildToolExecutableFilename() function is passing
+		// in an absolute path to Engine which causes Config path conversion to fail here
+//		checkfSlow(!(bIsProjectConversion && !Path.StartsWith(FPaths::ProjectDir().TrimChar('/'))), TEXT("Called ConvertPath with a Platform_ method, but Path [%s] didn't start with FPaths::ProjectDir() [%s]. This is not supported"), *Path, *FPaths::ProjectDir());
+//		checkfSlow(!(bIsEngineConversion && !Path.StartsWith(FPaths::EngineDir().TrimChar('/'))), TEXT("Called ConvertPath with a Engine_ method, but Path [%s] didn't start with FPaths::EngineDir() [%s]. This is not supported"), *Path, *FPaths::EngineDir());
 
 		// skip over engine or project directory
 		Prefix = (int)Method < (int)EPathConversion::Project_First ? FPaths::EngineDir() : FPaths::ProjectDir();
