AutoRTFM-ifi some missing modules
Add missing std lib calls to auto open

diff --git a/Engine/Plugins/Experimental/UIFramework/Source/UIFramework.Build.cs b/Engine/Plugins/Experimental/UIFramework/Source/UIFramework.Build.cs
--- a/Engine/Plugins/Experimental/UIFramework/Source/UIFramework.Build.cs
+++ b/Engine/Plugins/Experimental/UIFramework/Source/UIFramework.Build.cs
@@ -35,6 +35,8 @@ public class UIFramework : ModuleRules
 			}
 		);
 
+		bAllowAutoRTFMInstrumentation = true;
+
 		SetupIrisSupport(Target);
 	}
 }
diff --git a/Engine/Plugins/Runtime/ReplicationGraph/Source/ReplicationGraph.Build.cs b/Engine/Plugins/Runtime/ReplicationGraph/Source/ReplicationGraph.Build.cs
--- a/Engine/Plugins/Runtime/ReplicationGraph/Source/ReplicationGraph.Build.cs
+++ b/Engine/Plugins/Runtime/ReplicationGraph/Source/ReplicationGraph.Build.cs
@@ -19,6 +19,8 @@ public class ReplicationGraph : ModuleRules
 
 		UnsafeTypeCastWarningLevel = WarningLevel.Error;
 
+		bAllowAutoRTFMInstrumentation = true;
+
 		SetupGameplayDebuggerSupport(Target);
 	}
 }
diff --git a/Engine/Source/Runtime/Core/Private/AutoRTFM/StdLibSupport.cpp b/Engine/Source/Runtime/Core/Private/AutoRTFM/StdLibSupport.cpp
--- a/Engine/Source/Runtime/Core/Private/AutoRTFM/StdLibSupport.cpp
+++ b/Engine/Source/Runtime/Core/Private/AutoRTFM/StdLibSupport.cpp
@@ -153,6 +153,7 @@ UE_AUTORTFM_REGISTER_SELF_FUNCTION(exp);
 UE_AUTORTFM_REGISTER_SELF_FUNCTION(log);
 UE_AUTORTFM_REGISTER_SELF_FUNCTION(pow);
 UE_AUTORTFM_REGISTER_SELF_FUNCTION(llrint);
+UE_AUTORTFM_REGISTER_SELF_FUNCTION(fmod);
 // Linux (likely Mac) have ambiguous overrides to these math functions
 #else
 UE_AUTORTFM_REGISTER_SELF_FUNCTION(static_cast<float(*)(float)>(&sqrt));
@@ -191,6 +192,8 @@ UE_AUTORTFM_REGISTER_SELF_FUNCTION(static_cast<float(*)(float, float)>(&pow));
 UE_AUTORTFM_REGISTER_SELF_FUNCTION(static_cast<long double(*)(long double, long double)>(&pow));
 UE_AUTORTFM_REGISTER_SELF_FUNCTION(static_cast<long long(*)(float)>(&llrint));
 UE_AUTORTFM_REGISTER_SELF_FUNCTION(static_cast<long long(*)(long double)>(&llrint));
+UE_AUTORTFM_REGISTER_SELF_FUNCTION(static_cast<float(*)(float, float)>(&fmod));
+UE_AUTORTFM_REGISTER_SELF_FUNCTION(static_cast<long double(*)(long double, long double)>(&fmod));
 #endif
 
 // Self register Math functions
@@ -207,6 +210,8 @@ UE_AUTORTFM_REGISTER_SELF_FUNCTION(tanhf);
 UE_AUTORTFM_REGISTER_SELF_FUNCTION(expf);
 UE_AUTORTFM_REGISTER_SELF_FUNCTION(logf);
 UE_AUTORTFM_REGISTER_SELF_FUNCTION(powf);
+UE_AUTORTFM_REGISTER_SELF_FUNCTION(fmodf);
+UE_AUTORTFM_REGISTER_SELF_FUNCTION(fmodl);
 
 // FIXME: This is only correct when:
 // - Str is newly allocated
diff --git a/Engine/Source/Runtime/Experimental/Iris/Core/IrisCore.Build.cs b/Engine/Source/Runtime/Experimental/Iris/Core/IrisCore.Build.cs
--- a/Engine/Source/Runtime/Experimental/Iris/Core/IrisCore.Build.cs
+++ b/Engine/Source/Runtime/Experimental/Iris/Core/IrisCore.Build.cs
@@ -30,6 +30,7 @@ namespace UnrealBuildTool.Rules
 				}
 			}
 
+			bAllowAutoRTFMInstrumentation = true;
 		}
 	}
 }
diff --git a/Engine/Source/Runtime/GameplayTags/GameplayTags.Build.cs b/Engine/Source/Runtime/GameplayTags/GameplayTags.Build.cs
--- a/Engine/Source/Runtime/GameplayTags/GameplayTags.Build.cs
+++ b/Engine/Source/Runtime/GameplayTags/GameplayTags.Build.cs
@@ -36,6 +36,8 @@ namespace UnrealBuildTool.Rules
                 );
             }
 
+			bAllowAutoRTFMInstrumentation = true;
+
 			SetupIrisSupport(Target);
 		}
 	}
diff --git a/Engine/Source/Runtime/PhysicsCore/PhysicsCore.Build.cs b/Engine/Source/Runtime/PhysicsCore/PhysicsCore.Build.cs
--- a/Engine/Source/Runtime/PhysicsCore/PhysicsCore.Build.cs
+++ b/Engine/Source/Runtime/PhysicsCore/PhysicsCore.Build.cs
@@ -27,5 +27,7 @@ public class PhysicsCore: ModuleRules
 		PublicDependencyModuleNames.Remove("PhysicsCore");
 
 		PrivateDefinitions.Add("CHAOS_INCLUDE_LEVEL_1=1");
+
+		bAllowAutoRTFMInstrumentation = true;
 	}
 }
