IoDispatcher: Protect the Backends array with a FRWLock since it may be modified by Mount() on GT while being used in the IoDispatcher thread.

#rnx

diff --git a/Engine/Source/Runtime/Core/Private/IO/IoDispatcher.cpp b/Engine/Source/Runtime/Core/Private/IO/IoDispatcher.cpp
--- a/Engine/Source/Runtime/Core/Private/IO/IoDispatcher.cpp
+++ b/Engine/Source/Runtime/Core/Private/IO/IoDispatcher.cpp
@@ -425,6 +425,7 @@ public:
 	{
 		if (ChunkId.IsValid())
 		{
+			FReadScopeLock _(BackendsLock);
 			for (const FBackendAndPriority& Backend : Backends)
 			{
 				TIoStatusOr<FIoMappedRegion> Result = Backend.Value->OpenMapped(ChunkId, Options);
@@ -445,20 +446,24 @@ public:
 	{
 		check(IsInGameThread());
 
-		int32 Index = Algo::LowerBoundBy(Backends, Priority, &FBackendAndPriority::Key, TGreater<>());
-		Backends.Insert(MakeTuple(Priority, Backend), Index);
 		if (bIsInitialized)
 		{
 			Backend->Initialize(BackendContext);
-			if (!Thread)
-			{
-				StartThread();
-			}
+		}
+		{
+			FWriteScopeLock _(BackendsLock);
+			int32 Index = Algo::LowerBoundBy(Backends, Priority, &FBackendAndPriority::Key, TGreater<>());
+			Backends.Insert(MakeTuple(Priority, Backend), Index);
+		}
+		if (bIsInitialized && !Thread)
+		{
+			StartThread();
 		}
 	}
 
 	bool DoesChunkExist(const FIoChunkId& ChunkId) const
 	{
+		FReadScopeLock _(BackendsLock);
 		for (const FBackendAndPriority& Backend : Backends)
 		{
 			if (Backend.Value->DoesChunkExist(ChunkId))
@@ -471,6 +476,7 @@ public:
 
 	bool DoesChunkExist(const FIoChunkId& ChunkId, const FIoOffsetAndLength& ChunkRange) const
 	{
+		FReadScopeLock _(BackendsLock);
 		for (const FBackendAndPriority& Backend : Backends)
 		{
 			if (Backend.Value->DoesChunkExist(ChunkId, ChunkRange))
@@ -486,6 +492,7 @@ public:
 		// Only attempt to find the size if the FIoChunkId is valid
 		if (ChunkId.IsValid())
 		{
+			FReadScopeLock _(BackendsLock);
 			for (const FBackendAndPriority& Backend : Backends)
 			{
 				TIoStatusOr<uint64> Result = Backend.Value->GetSizeForChunk(ChunkId);
@@ -507,6 +514,7 @@ public:
 		// Only attempt to find the size if the FIoChunkId is valid
 		if (ChunkId.IsValid())
 		{
+			FReadScopeLock _(BackendsLock);
 			for (const FBackendAndPriority& Backend : Backends)
 			{
 				TIoStatusOr<uint64> Result = Backend.Value->GetSizeForChunk(ChunkId, ChunkRange, OutAvailable);
@@ -540,7 +548,7 @@ public:
 		}
 		check(Batch.TailRequest);
 
-		if (Backends.IsEmpty())
+		if (!HasMountedBackend())
 		{
 			FIoRequestImpl* Request = Batch.HeadRequest;
 			while (Request)
@@ -634,6 +642,7 @@ public:
 
 	bool HasMountedBackend() const
 	{
+		FReadScopeLock _(BackendsLock);
 		return Backends.Num() > 0;
 	}
 
@@ -651,6 +660,7 @@ private:
 	{
 		//TRACE_CPUPROFILER_EVENT_SCOPE(ProcessCompletedRequests);
 
+		FReadScopeLock _(BackendsLock);
 		for (const FBackendAndPriority& Backend : Backends)
 		{
 			FIoRequestImpl* CompletedRequestsHead = Backend.Value->GetCompletedRequests();
@@ -811,6 +821,7 @@ private:
 			{
 				TRACE_CPUPROFILER_EVENT_SCOPE(ResolveRequest);
 				bool bResolved = false;
+				FReadScopeLock _(BackendsLock);
 				for (const FBackendAndPriority& Backend : Backends)
 				{
 					if (Backend.Value->Resolve(Request))
@@ -882,6 +893,7 @@ private:
 
 	TSharedRef<FIoDispatcherBackendContext> BackendContext;
 	FDelegateHandle MemoryTrimDelegateHandle;
+	mutable FRWLock BackendsLock;
 	TArray<FBackendAndPriority> Backends;
 	FIoRequestAllocator* RequestAllocator = nullptr;
 	FBatchAllocator BatchAllocator;
