"any task" support

diff --git a/Engine/Source/Runtime/Core/Private/Async/TaskGraph.cpp b/Engine/Source/Runtime/Core/Private/Async/TaskGraph.cpp
--- a/Engine/Source/Runtime/Core/Private/Async/TaskGraph.cpp
+++ b/Engine/Source/Runtime/Core/Private/Async/TaskGraph.cpp
@@ -2800,3 +2800,81 @@ static FAutoConsoleCommand TaskThreadPriorityCmd(
 	TEXT("Sets the priority of the task threads. Argument is one of belownormal, normal or abovenormal."),
 	FConsoleCommandWithArgsDelegate::CreateStatic(&SetTaskThreadPriority)
 	);
+
+/////////////////////////////////////////////////////////////
+// "any task" support. these functions allocate excessively (per input task plus more)
+// can be reduced to a single alloc if this is a perf issue
+
+int32 WaitForAnyTaskCompleted(const FGraphEventArray& GraphEvents, FTimespan Timeout /*= FTimespan::MaxValue()*/)
+{
+	if (UNLIKELY(GraphEvents.Num() == 0))
+	{
+		return INDEX_NONE;
+	}
+
+	FSharedEventRef SystemEvent;
+	std::atomic<int32> CompletedTaskIndex;
+
+	for (int32 Index = 0; Index != GraphEvents.Num(); ++Index)
+	{
+		FFunctionGraphTask::CreateAndDispatchWhenReady(
+			[SystemEvent, Index, &CompletedTaskIndex] 
+			{ 
+				CompletedTaskIndex.store(Index, std::memory_order_relaxed);
+				SystemEvent->Trigger(); 
+			}, 
+			TStatId{}, 
+			GraphEvents[Index]
+		);
+	}
+
+	return SystemEvent->Wait(Timeout) ? CompletedTaskIndex.load(std::memory_order_relaxed) : INDEX_NONE;
+}
+
+FGraphEventRef AnyTaskCompleted(const FGraphEventArray& GraphEvents)
+{
+	if (UNLIKELY(GraphEvents.Num() == 0))
+	{
+		FGraphEventRef Result = FGraphEvent::CreateGraphEvent();
+		Result->DispatchSubsequents();
+		return Result;
+	}
+
+	struct FSharedData
+	{
+		explicit FSharedData(uint32 InitRefCount)
+			: RefCount(InitRefCount)
+		{
+		}
+
+		FGraphEventRef Event = FGraphEvent::CreateGraphEvent();
+		std::atomic<uint32> RefCount;
+	};
+
+	FSharedData* SharedData = new FSharedData(GraphEvents.Num());
+	// `SharedData` can be destroyed before leaving the scope, cache the result locally
+	FGraphEventRef Result = SharedData->Event;
+
+	for (const FGraphEventRef& GraphEvent : GraphEvents)
+	{
+		FFunctionGraphTask::CreateAndDispatchWhenReady(
+			[SharedData, Num = GraphEvents.Num()]
+			{
+				// cache the local copy as `SharedData` can be concurrently deleted right after decrementing the ref counter
+				FGraphEventRef Event = SharedData->Event;
+				uint32 PrevRefCount = SharedData->RefCount.fetch_sub(1, std::memory_order_acq_rel); // acq_rel to sync between tasks
+				
+				if (UNLIKELY(PrevRefCount == Num))
+				{	// the first completed task
+					Event->DispatchSubsequents();
+				}
+				else if (UNLIKELY(PrevRefCount == 1))
+				{	// the last competed task
+					delete SharedData;
+				}
+			}
+		);
+	}
+
+	return Result;
+}
\ No newline at end of file
diff --git a/Engine/Source/Runtime/Core/Public/Async/TaskGraphInterfaces.h b/Engine/Source/Runtime/Core/Public/Async/TaskGraphInterfaces.h
--- a/Engine/Source/Runtime/Core/Public/Async/TaskGraphInterfaces.h
+++ b/Engine/Source/Runtime/Core/Public/Async/TaskGraphInterfaces.h
@@ -1401,6 +1401,14 @@ private:
 
 #endif // TASKGRAPH_NEW_FRONTEND
 
+// Blocks the current thread until any of the given tasks is completed.
+// Is slightly more efficient than `AnyTaskCompleted()->Wait()` and supports timeout while `FGraphEvent::Wait()` doesn't.
+// Returns the index of the first completed task, or `INDEX_NONE` on timeout.
+int32 WaitForAnyTaskCompleted(const FGraphEventArray& GraphEvents, FTimespan Timeout = FTimespan::MaxValue());
+
+// Returns a graph event that gets completed as soon as any of the given tasks gets completed
+FGraphEventRef AnyTaskCompleted(const FGraphEventArray& GraphEvents);
+
 /** 
  *	FReturnGraphTask is a task used to return flow control from a named thread back to the original caller of ProcessThreadUntilRequestReturn
  **/
diff --git a/Engine/Source/Runtime/Core/Public/Tasks/Task.h b/Engine/Source/Runtime/Core/Public/Tasks/Task.h
--- a/Engine/Source/Runtime/Core/Public/Tasks/Task.h
+++ b/Engine/Source/Runtime/Core/Public/Tasks/Task.h
@@ -400,6 +400,92 @@ namespace UE::Tasks
 		return Tasks;
 	}
 
+	/////////////////////////////////////////////////////////////
+	// "any task" support. these functions allocate excessively (per input task plus more).
+	// can be reduced to a single alloc if this is a perf issue
+	
+	// Blocks the current thread until any of the given tasks is completed.
+	// Is slightly more efficient than `Any().Wait()`.
+	// Returns the index of the first completed task, or `INDEX_NONE` on timeout
+	template<typename TaskCollectionType>
+	int32 WaitAny(const TaskCollectionType& Tasks, FTimespan Timeout = FTimespan::MaxValue())
+	{
+		if (UNLIKELY(Tasks.Num() == 0))
+		{
+			return INDEX_NONE;
+		}
+
+		FSharedEventRef Event;
+		std::atomic<int32> CompletedTaskIndex;
+
+		for (int32 Index = 0; Index != Tasks.Num(); ++Index)
+		{
+			Launch(UE_SOURCE_LOCATION, 
+				[Event, Index, &CompletedTaskIndex] 
+				{ 
+					CompletedTaskIndex.store(Index, std::memory_order_relaxed);
+					Event->Trigger(); 
+				}, 
+				Prerequisites(Tasks[Index]),
+				ETaskPriority::Default, 
+				EExtendedTaskPriority::Inline
+			);
+		}
+
+		return Event->Wait(Timeout) ? CompletedTaskIndex.load(std::memory_order_relaxed) : INDEX_NONE;
+	}
+
+	// Returns a task that gets completed as soon as any of the given tasks gets completed
+	template<typename TaskCollectionType>
+	FTask Any(const TaskCollectionType& Tasks)
+	{
+		if (UNLIKELY(Tasks.Num() == 0))
+		{
+			return FTask{};
+		}
+
+		struct FSharedData
+		{
+			explicit FSharedData(uint32 InitRefCount)
+				: RefCount(InitRefCount)
+			{
+			}
+
+			FTaskEvent Event{ UE_SOURCE_LOCATION };
+			std::atomic<uint32> RefCount;
+		};
+
+		FSharedData* SharedData = new FSharedData(Tasks.Num());
+		// `SharedData` can be destroyed before leaving the scope, cache the result locally
+		FTaskEvent Result = SharedData->Event;
+
+		for (const FTask& Task : Tasks)
+		{
+			Launch(UE_SOURCE_LOCATION,
+				[SharedData, Num = Tasks.Num()]
+				{
+					// cache the local copy as `SharedData` can be concurrently deleted right after decrementing the ref counter
+					FTaskEvent Event = SharedData->Event;
+					uint32 PrevRefCount = SharedData->RefCount.fetch_sub(1, std::memory_order_acq_rel); // acq_rel to sync between tasks
+
+					if (UNLIKELY(PrevRefCount == Num))
+					{	// the first completed task
+						Event.Trigger();
+					}
+					else if (UNLIKELY(PrevRefCount == 1))
+					{	// the last competed task
+						delete SharedData;
+					}
+				},
+				Prerequisites(Task),
+				ETaskPriority::Default,
+				EExtendedTaskPriority::Inline
+			);
+		}
+
+		return Result;
+	}
+
 	// Adds the nested task to the task that is being currently executed by the current thread. A parent task is not flagged completed
 	// until all nested tasks are completed. It's similar to explicitly waiting for a sub-task at the end of its parent task, except explicit waiting
 	// blocks the worker executing the parent task until the sub-task is completed. With nested tasks, the worker won't be blocked.
diff --git a/Engine/Source/Runtime/Core/Tests/Async/TaskGraphTest.cpp b/Engine/Source/Runtime/Core/Tests/Async/TaskGraphTest.cpp
--- a/Engine/Source/Runtime/Core/Tests/Async/TaskGraphTest.cpp
+++ b/Engine/Source/Runtime/Core/Tests/Async/TaskGraphTest.cpp
@@ -1385,6 +1385,73 @@ PRAGMA_ENABLE_DEPRECATION_WARNINGS
 
 		check(bDestroyed);
 	}
+
+	TEST_CASE_NAMED(FTaskGraphWaitForAnyTask, "System::Core::Async::TaskGraph::WaitForAnyTask", "[.][ApplicationContextMask][EngineFilter]")
+	{
+		{	// blocks if none of tasks is completed
+			FGraphEventRef Blocker = FGraphEvent::CreateGraphEvent(); // blocks all tasks
+
+			FGraphEventArray Tasks
+			{
+				FFunctionGraphTask::CreateAndDispatchWhenReady([] {}, TStatId{}, Blocker),
+				FFunctionGraphTask::CreateAndDispatchWhenReady([] {}, TStatId{}, Blocker)
+			};
+
+			verify(WaitForAnyTaskCompleted(Tasks, FTimespan::FromMilliseconds(1.0)) == INDEX_NONE);
+
+			Blocker->DispatchSubsequents();
+
+			verify(WaitForAnyTaskCompleted(Tasks) != INDEX_NONE);
+		}
+
+		{	// doesn't wait for all tasks
+			FGraphEventRef Blocker = FGraphEvent::CreateGraphEvent();
+
+			FGraphEventArray Tasks
+			{
+				FFunctionGraphTask::CreateAndDispatchWhenReady([] {}),
+				FFunctionGraphTask::CreateAndDispatchWhenReady([] {}, TStatId{}, Blocker) // is blocked
+			};
+
+			verify(WaitForAnyTaskCompleted(Tasks) == 0);
+
+			Blocker->DispatchSubsequents();
+		}
+	}
+
+	TEST_CASE_NAMED(FTaskGraphAnyTask, "System::Core::Async::TaskGraph::AnyTask", "[.][ApplicationContextMask][EngineFilter]")
+	{
+		{	// blocks if none of tasks is completed
+			FGraphEventRef Blocker = FGraphEvent::CreateGraphEvent(); // blocks all tasks
+
+			FGraphEventArray Tasks
+			{
+				FFunctionGraphTask::CreateAndDispatchWhenReady([] {}, TStatId{}, Blocker),
+				FFunctionGraphTask::CreateAndDispatchWhenReady([] {}, TStatId{}, Blocker)
+			};
+
+			FPlatformProcess::Sleep(0.1f);
+			verify(!AnyTaskCompleted(Tasks)->IsComplete());
+
+			Blocker->DispatchSubsequents();
+
+			AnyTaskCompleted(Tasks)->Wait();
+		}
+
+		{	// doesn't wait for all tasks
+			FGraphEventRef Blocker = FGraphEvent::CreateGraphEvent();
+
+			FGraphEventArray Tasks
+			{
+				FFunctionGraphTask::CreateAndDispatchWhenReady([] {}),
+				FFunctionGraphTask::CreateAndDispatchWhenReady([] {}, TStatId{}, Blocker) // is blocked
+			};
+
+			AnyTaskCompleted(Tasks)->Wait();
+
+			Blocker->DispatchSubsequents();
+		}
+	}
 }
 
 #endif //WITH_TESTS
diff --git a/Engine/Source/Runtime/Core/Tests/Tasks/TasksTest.cpp b/Engine/Source/Runtime/Core/Tests/Tasks/TasksTest.cpp
--- a/Engine/Source/Runtime/Core/Tests/Tasks/TasksTest.cpp
+++ b/Engine/Source/Runtime/Core/Tests/Tasks/TasksTest.cpp
@@ -1453,6 +1453,72 @@ namespace UE { namespace TasksTests
 			Wait(TArray{ Task1, Task2 });
 		}
 	}
+
+	TEST_CASE_NAMED(FTasksWaitAny, "System::Core::Async::Tasks::WaitAny", "[.][ApplicationContextMask][EngineFilter]")
+	{
+		{	// blocks if none of tasks is completed
+			FTaskEvent Blocker{ UE_SOURCE_LOCATION }; // blocks all tasks
+
+			TArray<FTask> Tasks
+			{
+				Launch(UE_SOURCE_LOCATION, [] {}, Prerequisites(Blocker)),
+				Launch(UE_SOURCE_LOCATION, [] {}, Prerequisites(Blocker))
+			};
+
+			verify(WaitAny(Tasks, FTimespan::FromMilliseconds(1.0)) == INDEX_NONE);
+
+			Blocker.Trigger();
+
+			verify(WaitAny(Tasks) != INDEX_NONE);
+		}
+
+		{	// doesn't wait for all tasks
+			FTaskEvent Blocker{ UE_SOURCE_LOCATION };
+
+			TArray<FTask> Tasks
+			{
+				Launch(UE_SOURCE_LOCATION, [] {}),
+				Launch(UE_SOURCE_LOCATION, [] {}, Prerequisites(Blocker)) // is blocked
+			};
+
+			verify(WaitAny(Tasks) == 0);
+
+			Blocker.Trigger();
+		}
+	}
+
+	TEST_CASE_NAMED(FTasksAny, "System::Core::Async::Tasks::Any", "[.][ApplicationContextMask][EngineFilter]")
+	{
+		{	// blocks if none of tasks is completed
+			FTaskEvent Blocker{ UE_SOURCE_LOCATION }; // blocks all tasks
+
+			TArray<FTask> Tasks
+			{
+				Launch(UE_SOURCE_LOCATION, [] {}, Prerequisites(Blocker)),
+				Launch(UE_SOURCE_LOCATION, [] {}, Prerequisites(Blocker))
+			};
+
+			verify(!Any(Tasks).Wait(FTimespan::FromMilliseconds(0.1)));
+
+			Blocker.Trigger();
+
+			Any(Tasks).Wait();
+		}
+
+		{	// doesn't wait for all tasks
+			FTaskEvent Blocker{ UE_SOURCE_LOCATION };
+
+			TArray<FTask> Tasks
+			{
+				Launch(UE_SOURCE_LOCATION, [] {}),
+				Launch(UE_SOURCE_LOCATION, [] {}, Prerequisites(Blocker)) // is blocked
+			};
+
+			Any(Tasks).Wait();
+
+			Blocker.Trigger();
+		}
+	}
 }}
 
 #endif // WITH_TESTS
