Add 31 bit precision rand function.
#rnx

diff --git a/Engine/Source/Runtime/Core/Public/GenericPlatform/GenericPlatformMath.h b/Engine/Source/Runtime/Core/Public/GenericPlatform/GenericPlatformMath.h
--- a/Engine/Source/Runtime/Core/Public/GenericPlatform/GenericPlatformMath.h
+++ b/Engine/Source/Runtime/Core/Public/GenericPlatform/GenericPlatformMath.h
@@ -615,6 +615,9 @@ struct FGenericPlatformMath
 	/** Returns a random integer between 0 and RAND_MAX, inclusive */
 	static FORCEINLINE int32 Rand() { return rand(); }
 
+	/** Returns a random integer between 0 and MAX_int32, inclusive. RAND_MAX may only be 15 bits, so compose from multiple calls. */
+	static FORCEINLINE int32 Rand32() { return ((rand() & 0x7fff) << 16) | ((rand() & 0x7fff) << 1) | (rand() & 0x1); }
+
 	/** Seeds global random number functions Rand() and FRand() */
 	static FORCEINLINE void RandInit(int32 Seed) { srand( Seed ); }
 
