Add chunk install manifest interface declaration

diff --git a/Engine/Source/Runtime/Core/Public/GenericPlatform/GenericPlatformChunkInstall.h b/Engine/Source/Runtime/Core/Public/GenericPlatform/GenericPlatformChunkInstall.h
--- a/Engine/Source/Runtime/Core/Public/GenericPlatform/GenericPlatformChunkInstall.h
+++ b/Engine/Source/Runtime/Core/Public/GenericPlatform/GenericPlatformChunkInstall.h
@@ -87,6 +87,18 @@ public:
 	virtual IPlatformChunkInstall* GetPlatformChunkInstall() = 0;
 };
 
+/**
+ * Platform Chunk Install Manifest Interface
+ */
+class IPlatformChunkInstallManifest
+{
+public:
+	virtual bool HasManifest() const = 0;
+	virtual int32 GetChunkIDFromPakchunkIndex(int32 PakchunkIndex) const = 0;
+	virtual TArray<FString> GetPakFilesInChunk(int32 ChunkID) const = 0;
+};
+
+
 /** Delegate called when a chunk either successfully installs or fails to install, bool is success */
 DECLARE_DELEGATE_TwoParams(FPlatformChunkInstallDelegate, uint32, bool);
 DECLARE_MULTICAST_DELEGATE_TwoParams(FPlatformChunkInstallMultiDelegate, uint32, bool);
