Fix the reentrant exceptions test under AutoRTFM.

diff --git a/Engine/Source/Runtime/Core/Public/Templates/SharedPointerInternals.h b/Engine/Source/Runtime/Core/Public/Templates/SharedPointerInternals.h
--- a/Engine/Source/Runtime/Core/Public/Templates/SharedPointerInternals.h
+++ b/Engine/Source/Runtime/Core/Public/Templates/SharedPointerInternals.h
@@ -149,37 +149,55 @@ namespace SharedPointerInternals
 		{
 			if constexpr (Mode == ESPMode::ThreadSafe)
 			{
-				// See AddSharedReference for the same reasons that std::memory_order_relaxed is used in this function.
-
-				// Peek at the current shared reference count.  Remember, this value may be updated by
-				// multiple threads.
-				int32 OriginalCount = SharedReferenceCount.load(std::memory_order_relaxed);
+				bool bSucceeded = false;
 
-				for ( ; ; )
+				UE_AUTORTFM_OPEN(
 				{
-					if( OriginalCount == 0 )
+					// See AddSharedReference for the same reasons that std::memory_order_relaxed is used in this function.
+
+					// Peek at the current shared reference count.  Remember, this value may be updated by
+					// multiple threads.
+					int32 OriginalCount = SharedReferenceCount.load(std::memory_order_relaxed);
+
+					for (; ; )
 					{
-						// Never add a shared reference if the pointer has already expired
-						return false;
+						if (OriginalCount == 0)
+						{
+							// Never add a shared reference if the pointer has already expired
+							bSucceeded = false;
+							break;
+						}
+
+						// Attempt to increment the reference count.
+						//
+						// We need to make sure that we never revive a counter that has already expired, so if the
+						// actual value what we expected (because it was touched by another thread), then we'll try
+						// again.  Note that only in very unusual cases will this actually have to loop.
+						//
+						// We do a weak read here because we require a loop and this is the recommendation:
+						//
+						// https://en.cppreference.com/w/cpp/atomic/atomic/compare_exchange
+						//
+						// > When a compare-and-exchange is in a loop, the weak version will yield better performance on some platforms.
+						// > When a weak compare-and-exchange would require a loop and a strong one would not, the strong one is preferable
+						if (SharedReferenceCount.compare_exchange_weak(OriginalCount, OriginalCount + 1, std::memory_order_relaxed))
+						{
+							bSucceeded = true;
+							break;
+						}
 					}
+				});
 
-					// Attempt to increment the reference count.
-					//
-					// We need to make sure that we never revive a counter that has already expired, so if the
-					// actual value what we expected (because it was touched by another thread), then we'll try
-					// again.  Note that only in very unusual cases will this actually have to loop.
-					//
-					// We do a weak read here because we require a loop and this is the recommendation:
-					//
-					// https://en.cppreference.com/w/cpp/atomic/atomic/compare_exchange
-					//
-					// > When a compare-and-exchange is in a loop, the weak version will yield better performance on some platforms.
-					// > When a weak compare-and-exchange would require a loop and a strong one would not, the strong one is preferable
-					if (SharedReferenceCount.compare_exchange_weak(OriginalCount, OriginalCount + 1, std::memory_order_relaxed))
+				// If we succeedd in taking a shared reference count, we need to undo that on an abort.
+				if (bSucceeded)
+				{
+					UE_AUTORTFM_OPENABORT(
 					{
-						return true;
-					}
+						ReleaseSharedReference();
+					});
 				}
+
+				return bSucceeded;
 			}
 			else
 			{
@@ -525,18 +543,18 @@ namespace SharedPointerInternals
 		}
 
 		/** Creates a shared referencer object from a weak referencer object.  This will only result
-		    in a valid object reference if the object already has at least one other shared referencer. */
-		FSharedReferencer( FWeakReferencer< Mode > const& InWeakReference )
-			: ReferenceController( InWeakReference.ReferenceController )
+			in a valid object reference if the object already has at least one other shared referencer. */
+		FSharedReferencer(FWeakReferencer< Mode > const& InWeakReference)
+			: ReferenceController(InWeakReference.ReferenceController)
 		{
 			// If the incoming reference had an object associated with it, then go ahead and increment the
 			// shared reference count
-			if( ReferenceController != nullptr )
+			if (ReferenceController != nullptr)
 			{
 				// Attempt to elevate a weak reference to a shared one.  For this to work, the object this
 				// weak counter is associated with must already have at least one shared reference.  We'll
 				// never revive a pointer that has already expired!
-				if( !ReferenceController->ConditionallyAddSharedReference() )
+				if (!ReferenceController->ConditionallyAddSharedReference())
 				{
 					ReferenceController = nullptr;
 				}
