Make failing to find the transactional variant of a function cause a fatal UE_LOG in AutoRTFM.

diff --git a/Engine/Source/Runtime/Core/Private/AutoRTFM/CompilerInterface.cpp b/Engine/Source/Runtime/Core/Private/AutoRTFM/CompilerInterface.cpp
--- a/Engine/Source/Runtime/Core/Private/AutoRTFM/CompilerInterface.cpp
+++ b/Engine/Source/Runtime/Core/Private/AutoRTFM/CompilerInterface.cpp
@@ -82,11 +82,11 @@ extern "C" void autortfm_llvm_fail(const char* Message)
 {
     if (Message)
     {
-		UE_LOG(LogAutoRTFM, Warning, TEXT("Transaction failing because of language issue '%s'."), ANSI_TO_TCHAR(Message));
+		UE_LOG(LogAutoRTFM, Fatal, TEXT("Transaction failing because of language issue '%s'."), ANSI_TO_TCHAR(Message));
     }
     else
     {
-		UE_LOG(LogAutoRTFM, Warning, TEXT("Transaction failing because of language issue."));
+		UE_LOG(LogAutoRTFM, Fatal, TEXT("Transaction failing because of language issue."));
 	}
 
 	FContext* Context = FContext::Get();
diff --git a/Engine/Source/Runtime/Core/Private/AutoRTFM/FunctionMapInlines.h b/Engine/Source/Runtime/Core/Private/AutoRTFM/FunctionMapInlines.h
--- a/Engine/Source/Runtime/Core/Private/AutoRTFM/FunctionMapInlines.h
+++ b/Engine/Source/Runtime/Core/Private/AutoRTFM/FunctionMapInlines.h
@@ -17,12 +17,13 @@ inline void* FunctionMapLookup(void* OldFunction, const char* Where)
     {
 		if (Where)
 		{
-			UE_LOG(LogAutoRTFM, Warning, TEXT("Could not find function %p '%s' where '%s'."), OldFunction, *GetFunctionDescription(OldFunction), ANSI_TO_TCHAR(Where));
+			UE_LOG(LogAutoRTFM, Fatal, TEXT("Could not find function %p '%s' where '%s'."), OldFunction, *GetFunctionDescription(OldFunction), ANSI_TO_TCHAR(Where));
 		}
 		else
 		{
-			UE_LOG(LogAutoRTFM, Warning, TEXT("Could not find function %p '%s'."), OldFunction, *GetFunctionDescription(OldFunction));
+			UE_LOG(LogAutoRTFM, Fatal, TEXT("Could not find function %p '%s'."), OldFunction, *GetFunctionDescription(OldFunction));
 		}
+		
 		FContext* Context = FContext::Get();
         Context->AbortByLanguageAndThrow();
     }
