Remove buffering of the log in FAutomationMessageFilter

#rnx

diff --git a/Engine/Source/Runtime/Core/Private/Misc/AutomationTest.cpp b/Engine/Source/Runtime/Core/Private/Misc/AutomationTest.cpp
--- a/Engine/Source/Runtime/Core/Private/Misc/AutomationTest.cpp
+++ b/Engine/Source/Runtime/Core/Private/Misc/AutomationTest.cpp
@@ -18,7 +18,6 @@
 #include "Misc/OutputDeviceRedirector.h"
 #include "Misc/Paths.h"
 #include "Misc/ScopedSlowTask.h"
-#include "Misc/ScopeTryLock.h"
 #include "Misc/ScopeRWLock.h"
 #include "Modules/ModuleManager.h"
 
@@ -254,26 +253,9 @@ void FAutomationTestFramework::FAutomationTestMessageFilter::Serialize(const TCH
 		{
 			Verbosity = ELogVerbosity::Verbose;
 		}
-		FScopeTryLock Lock(&ActionCS);
-		if (Lock.IsLocked() || LocalDestinationContext->CanBeUsedOnMultipleThreads())
 		{
+			FScopeLock CriticalSection(&ActionCS);
 			LocalDestinationContext->Serialize(V, Verbosity, Category, Time);
-			if (!Backlog.IsEmpty())
-			{
-				TSharedPtr<UE::FLogRecord> Item;
-				while (Backlog.Dequeue(Item))
-				{
-					LocalDestinationContext->SerializeRecord(*Item);
-				}
-			}
-		}
-		else
-		{
-			TSharedPtr<UE::FLogRecord> Item = MakeShareable(new UE::FLogRecord());
-			Item->SetFormat(V);
-			Item->SetVerbosity(Verbosity);
-			Item->SetCategory(Category);
-			Backlog.Enqueue(Item);
 		}
 	}
 }
@@ -296,22 +278,9 @@ void FAutomationTestFramework::FAutomationTestMessageFilter::SerializeRecord(con
 				LocalRecord.SetVerbosity(ELogVerbosity::Verbose);
 			}
 		}
-		FScopeTryLock Lock(&ActionCS);
-		if (Lock.IsLocked() || LocalDestinationContext->CanBeUsedOnMultipleThreads())
 		{
+			FScopeLock CriticalSection(&ActionCS);
 			LocalDestinationContext->SerializeRecord(LocalRecord);
-			if (!Backlog.IsEmpty())
-			{
-				TSharedPtr<UE::FLogRecord> Item;
-				while (Backlog.Dequeue(Item))
-				{
-					LocalDestinationContext->SerializeRecord(*Item);
-				}
-			}
-		}
-		else
-		{
-			Backlog.Enqueue(MakeShareable(new UE::FLogRecord(LocalRecord)));
 		}
 	}
 }
diff --git a/Engine/Source/Runtime/Core/Public/Misc/AutomationTest.h b/Engine/Source/Runtime/Core/Public/Misc/AutomationTest.h
--- a/Engine/Source/Runtime/Core/Public/Misc/AutomationTest.h
+++ b/Engine/Source/Runtime/Core/Public/Misc/AutomationTest.h
@@ -1235,7 +1235,6 @@ private:
 		 std::atomic<FAutomationTestBase*> CurTest;
 		 std::atomic<FFeedbackContext*> DestinationContext;
 		 FCriticalSection ActionCS;
-		 TQueue<TSharedPtr<UE::FLogRecord>> Backlog;
 	 };
 
 	friend class FAutomationTestOutputDevice;
