FIoDirectoryIndexReader::IterateDirectoryIndex: Use FStringViews and TStringBuilders instead of FStrings to avoid potentially unnecessary heap allocations
FPakPlatformFile::ForeachPackageInIostoreWhile: Update usage of IterateDirectoryIndex, and use overload of TryConvertFilenameToLongPackageName using string views and builders to avoid allocations

diff --git a/Engine/Source/Developer/IoStoreUtilities/Private/IoStoreUtilities.cpp b/Engine/Source/Developer/IoStoreUtilities/Private/IoStoreUtilities.cpp
--- a/Engine/Source/Developer/IoStoreUtilities/Private/IoStoreUtilities.cpp
+++ b/Engine/Source/Developer/IoStoreUtilities/Private/IoStoreUtilities.cpp
@@ -5745,12 +5745,12 @@ int32 CreateContentPatch(const FIoStoreArguments& Arguments, const FIoStoreWrite
 
 		TMap<FIoChunkId, FString> ChunkFileNamesMap;
 		TargetReader->GetDirectoryIndexReader().IterateDirectoryIndex(FIoDirectoryIndexHandle::RootDirectory(), TEXT(""),
-		[&ChunkFileNamesMap, &TargetReader](FString Filename, uint32 TocEntryIndex) -> bool
+		[&ChunkFileNamesMap, &TargetReader](FStringView Filename, uint32 TocEntryIndex) -> bool
 		{
 			TIoStatusOr<FIoStoreTocChunkInfo> ChunkInfo = TargetReader->GetChunkInfo(TocEntryIndex);
 			if (ChunkInfo.IsOk())
 			{
-				ChunkFileNamesMap.Add(ChunkInfo.ValueOrDie().Id, Filename);
+				ChunkFileNamesMap.Add(ChunkInfo.ValueOrDie().Id, FString(Filename));
 			}
 			return true;
 		});
diff --git a/Engine/Source/Runtime/Core/Private/IO/IoDirectoryIndex.cpp b/Engine/Source/Runtime/Core/Private/IO/IoDirectoryIndex.cpp
--- a/Engine/Source/Runtime/Core/Private/IO/IoDirectoryIndex.cpp
+++ b/Engine/Source/Runtime/Core/Private/IO/IoDirectoryIndex.cpp
@@ -9,6 +9,7 @@
 #include "Misc/AssertionMacros.h"
 #include "Misc/CString.h"
 #include "Misc/Paths.h"
+#include "Misc/PathViews.h"
 #include "Serialization/Archive.h"
 #include "Serialization/MemoryReader.h"
 #include "Serialization/MemoryWriter.h"
@@ -398,16 +399,18 @@ public:
 			: ~uint32(0);
 	}
 
-	bool IterateDirectoryIndex(FIoDirectoryIndexHandle DirectoryIndexHandle, const FString& Path, FDirectoryIndexVisitorFunction Visit)
+	bool IterateDirectoryIndex(FIoDirectoryIndexHandle DirectoryIndexHandle, FStringView Path, FDirectoryIndexVisitorFunction Visit)
 	{
 		FIoDirectoryIndexHandle File = GetFile(DirectoryIndexHandle);
 		while (File.IsValid())
 		{
 			const uint32 TocEntryIndex = GetFileData(File);
 			FStringView FileName = GetFileName(File);
-			FString FilePath = GetMountPoint() / Path / FString(FileName);
+			TStringBuilder<256> FilePathBuilder;
+			FPathViews::Append(FilePathBuilder, GetMountPoint(), Path, FileName);
+			const FStringView FilePath = FilePathBuilder.ToView();
 
-			if (!Visit(MoveTemp(FilePath), TocEntryIndex))
+			if (!Visit(FilePath, TocEntryIndex))
 			{
 				return false;
 			}
@@ -419,7 +422,9 @@ public:
 		while (ChildDirectory.IsValid())
 		{
 			FStringView DirectoryName = GetDirectoryName(ChildDirectory);
-			FString ChildDirectoryPath = Path / FString(DirectoryName);
+			TStringBuilder<256> ChildDirectoryPathBuilder;
+			FPathViews::Append(ChildDirectoryPathBuilder, Path, DirectoryName);
+			const FStringView ChildDirectoryPath = ChildDirectoryPathBuilder.ToView();
 
 			if (!IterateDirectoryIndex(ChildDirectory, ChildDirectoryPath, Visit))
 			{
@@ -506,7 +511,7 @@ uint32 FIoDirectoryIndexReader::GetFileData(FIoDirectoryIndexHandle File) const
 	return Impl->GetFileData(File);
 }
 
-bool FIoDirectoryIndexReader::IterateDirectoryIndex(FIoDirectoryIndexHandle Directory, const FString& Path, FDirectoryIndexVisitorFunction Visit) const
+bool FIoDirectoryIndexReader::IterateDirectoryIndex(FIoDirectoryIndexHandle Directory, FStringView Path, FDirectoryIndexVisitorFunction Visit) const
 {
 	return Impl->IterateDirectoryIndex(Directory, Path, Visit);
 }
\ No newline at end of file
diff --git a/Engine/Source/Runtime/Core/Private/IO/IoStore.cpp b/Engine/Source/Runtime/Core/Private/IO/IoStore.cpp
--- a/Engine/Source/Runtime/Core/Private/IO/IoStore.cpp
+++ b/Engine/Source/Runtime/Core/Private/IO/IoStore.cpp
@@ -2360,9 +2360,9 @@ public:
 			DirectoryIndexReader.IterateDirectoryIndex(
 				FIoDirectoryIndexHandle::RootDirectory(),
 				TEXT(""),
-				[this](FString Filename, uint32 TocEntryIndex) -> bool
+				[this](FStringView Filename, uint32 TocEntryIndex) -> bool
 				{
-					ChunkFileNamesMap.Add(TocEntryIndex, MoveTemp(Filename));
+					ChunkFileNamesMap.Add(TocEntryIndex, FString(Filename));
 					return true;
 				});
 		}
@@ -3516,9 +3516,9 @@ void FIoStoreReader::GetFilenames(TArray<FString>& OutFileList) const
 	DirectoryIndex.IterateDirectoryIndex(
 		FIoDirectoryIndexHandle::RootDirectory(),
 		TEXT(""),
-		[&OutFileList](FString Filename, uint32 TocEntryIndex) -> bool
+		[&OutFileList](FStringView Filename, uint32 TocEntryIndex) -> bool
 		{
-			OutFileList.AddUnique(Filename);
+			OutFileList.AddUnique(FString(Filename));
 			return true;
 		});
 }
@@ -3528,13 +3528,13 @@ void FIoStoreReader::GetFilenamesByBlockIndex(const TArray<int32>& InBlockIndexL
 	const FIoDirectoryIndexReader& DirectoryIndex = GetDirectoryIndexReader();
 
 	DirectoryIndex.IterateDirectoryIndex(FIoDirectoryIndexHandle::RootDirectory(), TEXT(""),
-		[this, &InBlockIndexList, &OutFileList](FString Filename, uint32 TocEntryIndex) -> bool
+		[this, &InBlockIndexList, &OutFileList](FStringView Filename, uint32 TocEntryIndex) -> bool
 		{
 			for (int32 BlockIndex : InBlockIndexList)
 			{
 				if (Impl->TocChunkContainsBlockIndex(TocEntryIndex, BlockIndex))
 				{
-					OutFileList.AddUnique(Filename);
+					OutFileList.AddUnique(FString(Filename));
 					break;
 				}
 			}
diff --git a/Engine/Source/Runtime/Core/Public/IO/IoDispatcher.h b/Engine/Source/Runtime/Core/Public/IO/IoDispatcher.h
--- a/Engine/Source/Runtime/Core/Public/IO/IoDispatcher.h
+++ b/Engine/Source/Runtime/Core/Public/IO/IoDispatcher.h
@@ -419,7 +419,7 @@ private:
 	uint32 Handle = InvalidHandle;
 };
 
-using FDirectoryIndexVisitorFunction = TFunctionRef<bool(FString, const uint32)>;
+using FDirectoryIndexVisitorFunction = TFunctionRef<bool(FStringView, const uint32)>;
 
 class FIoDirectoryIndexReader
 {
@@ -437,7 +437,7 @@ public:
 	CORE_API FStringView GetFileName(FIoDirectoryIndexHandle File) const;
 	CORE_API uint32 GetFileData(FIoDirectoryIndexHandle File) const;
 
-	CORE_API bool IterateDirectoryIndex(FIoDirectoryIndexHandle Directory, const FString& Path, FDirectoryIndexVisitorFunction Visit) const;
+	CORE_API bool IterateDirectoryIndex(FIoDirectoryIndexHandle Directory, FStringView Path, FDirectoryIndexVisitorFunction Visit) const;
 
 private:
 	UE_NONCOPYABLE(FIoDirectoryIndexReader);
diff --git a/Engine/Source/Runtime/PakFile/Private/IPlatformFilePak.cpp b/Engine/Source/Runtime/PakFile/Private/IPlatformFilePak.cpp
--- a/Engine/Source/Runtime/PakFile/Private/IPlatformFilePak.cpp
+++ b/Engine/Source/Runtime/PakFile/Private/IPlatformFilePak.cpp
@@ -853,18 +853,18 @@ void FPakPlatformFile::ForeachPackageInIostoreWhile(TFunctionRef<bool(FName)> Pr
 			const bool Result = DirectoryIndex.IterateDirectoryIndex(
 				FIoDirectoryIndexHandle::RootDirectory(),
 				TEXT(""),
-				[Predicate](FString Filename, uint32) -> bool
+				[Predicate](FStringView Filename, uint32) -> bool
 				{
-					const FString Ext = FPaths::GetExtension(Filename);
-					if (Ext != TEXT("umap") && Ext != TEXT("uasset"))
+					const FStringView Ext = FPathViews::GetExtension(Filename);
+					if (Ext != TEXTVIEW("umap") && Ext != TEXTVIEW("uasset"))
 					{
 						return true; // ignore non package files
 					}
 
-					FString PackageName;
-					if (FPackageName::TryConvertFilenameToLongPackageName(Filename, PackageName))
+					TStringBuilder<256> PackageNameBuilder;
+					if (FPackageName::TryConvertFilenameToLongPackageName(Filename, PackageNameBuilder))
 					{
-						return Invoke(Predicate, FName(*PackageName));
+						return Invoke(Predicate, FName(PackageNameBuilder.ToView()));
 					}
 
 					return true; // ignore not mapped packages
