Write the boot hotfix file in background.

diff --git a/Engine/Source/Runtime/Core/Private/Misc/ConfigUtilities.cpp b/Engine/Source/Runtime/Core/Private/Misc/ConfigUtilities.cpp
--- a/Engine/Source/Runtime/Core/Private/Misc/ConfigUtilities.cpp
+++ b/Engine/Source/Runtime/Core/Private/Misc/ConfigUtilities.cpp
@@ -6,12 +6,15 @@
 #include "HAL/IConsoleManager.h"
 #include "Misc/CoreDelegates.h"
 #include "Misc/FileHelper.h"
+#include "Tasks/Pipe.h"
 
 #define UE_HOTFIX_FOR_NEXT_BOOT_FILENAME TEXT("HotfixForNextBoot.txt")
 
 namespace UE::ConfigUtilities
 {
 
+UE::Tasks::FPipe AsyncTaskPipe( TEXT("SaveHotfixForNextBootPipe") );
+
 const TCHAR* ConvertValueFromHumanFriendlyValue( const TCHAR* Value )
 {
 	static const TCHAR* OnValue = TEXT("1");
@@ -94,23 +97,28 @@ void SaveCVarForNextBoot(const TCHAR* Key, const TCHAR* Value)
 		return;
 	}
 
-	TMap<FString, FString> CVarsToSave;
+	FString StrKey(Key);
+	FString StrValue(Value);
 
-	// Read from file, in case there are more than one cvar hotfix event in same run
-	LoadCVarsFromFileForNextBoot(CVarsToSave);
+	AsyncTaskPipe.Launch(UE_SOURCE_LOCATION, [StrKey, StrValue] {
+		TMap<FString, FString> CVarsToSave;
 
-	CVarsToSave.FindOrAdd(Key) = Value;
+		// Read from file, in case there are more than one cvar hotfix event in same run
+		LoadCVarsFromFileForNextBoot(CVarsToSave);
 
-	FString ContentToSave;
-	for (const TPair<FString, FString>& CVarPair : CVarsToSave)
-	{
-		ContentToSave.Append(FString::Format(TEXT("{0}={1}\r\n"), { CVarPair.Key, CVarPair.Value }));
-	}
+		CVarsToSave.FindOrAdd(StrKey) = StrValue;
 
-	const FString FullPath = FPaths::ProjectPersistentDownloadDir() / UE_HOTFIX_FOR_NEXT_BOOT_FILENAME;
-	FFileHelper::SaveStringToFile(ContentToSave, *FullPath);
+		FString ContentToSave;
+		for (const TPair<FString, FString>& CVarPair : CVarsToSave)
+		{
+			ContentToSave.Append(FString::Format(TEXT("{0}={1}\r\n"), { CVarPair.Key, CVarPair.Value }));
+		}
+
+		const FString FullPath = FPaths::ProjectPersistentDownloadDir() / UE_HOTFIX_FOR_NEXT_BOOT_FILENAME;
+		FFileHelper::SaveStringToFile(ContentToSave, *FullPath);
 
-	UE_LOG(LogConfig, Log, TEXT("Local boot hotfix file [%s] saved with hotfixed CVar: %s=%s"), *FullPath, Key, Value);
+		UE_LOG(LogConfig, Log, TEXT("Local boot hotfix file [%s] saved with hotfixed CVar: %s=%s"), *FullPath, *StrKey, *StrValue);
+	}, LowLevelTasks::ETaskPriority::BackgroundLow);
 #endif // !UE_SERVER
 }
 
