added comment to discourage TQueue and TCircularQueue usage as we plan to deprecate them

diff --git a/Engine/Source/Runtime/Core/Public/Containers/CircularQueue.h b/Engine/Source/Runtime/Core/Public/Containers/CircularQueue.h
--- a/Engine/Source/Runtime/Core/Public/Containers/CircularQueue.h
+++ b/Engine/Source/Runtime/Core/Public/Containers/CircularQueue.h
@@ -6,6 +6,8 @@
 #include "Containers/CircularBuffer.h"
 #include "Templates/Atomic.h"
 
+// WARNING: This queue is planned for deprecation in favor of TSpscQueue
+
 /**
  * Implements a lock-free first-in first-out queue using a circular array.
  *
diff --git a/Engine/Source/Runtime/Core/Public/Containers/Queue.h b/Engine/Source/Runtime/Core/Public/Containers/Queue.h
--- a/Engine/Source/Runtime/Core/Public/Containers/Queue.h
+++ b/Engine/Source/Runtime/Core/Public/Containers/Queue.h
@@ -7,6 +7,8 @@
 #include "HAL/PlatformAtomics.h"
 #include "HAL/PlatformMisc.h"
 
+// WARNING: This queue is planned for deprecation in favor of TSpscQueue or TMpscQueue
+
 /**
  * Enumerates concurrent queue modes.
  */
