Add flush feature to CpuProfiler that can be used before entering a long single operation like a wait that could be missed in UnrealInsights in case of deadlock

#rnx

diff --git a/Engine/Source/Runtime/Core/Private/ProfilingDebugging/CpuProfilerTrace.cpp b/Engine/Source/Runtime/Core/Private/ProfilingDebugging/CpuProfilerTrace.cpp
--- a/Engine/Source/Runtime/Core/Private/ProfilingDebugging/CpuProfilerTrace.cpp
+++ b/Engine/Source/Runtime/Core/Private/ProfilingDebugging/CpuProfilerTrace.cpp
@@ -360,4 +360,13 @@ uint32 FCpuProfilerTrace::OutputEventType(const ANSICHAR* Name, const ANSICHAR*
 	return SpecId;
 }
 
+void FCpuProfilerTrace::FlushThreadBuffer()
+{
+	FCpuProfilerTraceInternal::FThreadBuffer* ThreadBuffer = FCpuProfilerTraceInternal::ThreadBuffer;
+	if (ThreadBuffer && ThreadBuffer->BufferSize > 0)
+	{
+		FCpuProfilerTraceInternal::FlushThreadBuffer(ThreadBuffer);
+	}
+}
+
 #endif
diff --git a/Engine/Source/Runtime/Core/Public/ProfilingDebugging/CpuProfilerTrace.h b/Engine/Source/Runtime/Core/Public/ProfilingDebugging/CpuProfilerTrace.h
--- a/Engine/Source/Runtime/Core/Public/ProfilingDebugging/CpuProfilerTrace.h
+++ b/Engine/Source/Runtime/Core/Public/ProfilingDebugging/CpuProfilerTrace.h
@@ -109,6 +109,11 @@ struct FCpuProfilerTrace
 	*/
 	CORE_API static void OutputSuspendEvent();
 
+	/*
+	* Make sure all thread data has reached the destination. Can be useful to call this before entering a wait condition that might take a while.
+	*/
+	CORE_API static void FlushThreadBuffer();
+
 	class FEventScope
 	{
 	public:
@@ -320,6 +325,11 @@ struct FCpuProfilerTrace
 #define TRACE_CPUPROFILER_EVENT_SCOPE_TEXT_CONDITIONAL(Name, Condition) \
 	FCpuProfilerTrace::FDynamicEventScope PREPROCESSOR_JOIN(__CpuProfilerEventScope, __LINE__)(Name, (Condition), __FILE__, __LINE__);
 
+// Make sure all thread data has reached the destination.
+// Note: Can be useful to call this before entering a wait condition that might take a while.
+#define TRACE_CPUPROFILER_EVENT_FLUSH() \
+	FCpuProfilerTrace::FlushThreadBuffer(); 
+
 #else
 
 #define TRACE_CPUPROFILER_EVENT_DECLARE(DeclName)
@@ -334,5 +344,6 @@ struct FCpuProfilerTrace
 #define TRACE_CPUPROFILER_EVENT_SCOPE_TEXT_ON_CHANNEL(Name, Channel)
 #define TRACE_CPUPROFILER_EVENT_SCOPE_TEXT(Name)
 #define TRACE_CPUPROFILER_EVENT_SCOPE_TEXT_CONDITIONAL(Name, Condition)
+#define TRACE_CPUPROFILER_EVENT_FLUSH()
 
 #endif
