Fixed up some use of TCHAR instead of ElementType.

diff --git a/Engine/Source/Runtime/Core/Private/Containers/String.cpp.inl b/Engine/Source/Runtime/Core/Private/Containers/String.cpp.inl
--- a/Engine/Source/Runtime/Core/Private/Containers/String.cpp.inl
+++ b/Engine/Source/Runtime/Core/Private/Containers/String.cpp.inl
@@ -167,7 +167,7 @@ namespace UE::Core::Private
 		}
 
 		// (Re-)establish the null terminator
-		*NewEnd = '\0';
+		*NewEnd = CHARTEXT(DestCharType, '\0');
 	}
 
 	template <typename DestCharType, typename SrcCharType>
@@ -1276,8 +1276,8 @@ int32 UE_STRING_CLASS::ParseIntoArray( TArray<UE_STRING_CLASS>& OutArray, const
 	{
 		UE::String::EParseTokensOptions ParseOptions = UE::String::EParseTokensOptions::IgnoreCase |
 			(InCullEmpty ? UE::String::EParseTokensOptions::SkipEmpty : UE::String::EParseTokensOptions::None);
-		UE::String::ParseTokens(FStringView(*this), FStringView(pchDelim),
-			[&OutArray](FStringView Token) { OutArray.Emplace(Token); },
+		UE::String::ParseTokens(TStringView<ElementType>(*this), TStringView<ElementType>(pchDelim),
+			[&OutArray](TStringView<ElementType> Token) { OutArray.Emplace(Token); },
 			ParseOptions);
 	}
 	return OutArray.Num();
@@ -1907,7 +1907,7 @@ UE_STRING_CLASS SlugStringForValidName(const UE_STRING_CLASS& DisplayString, con
 	{
 		for ( int32 BadCharacterIndex = 0; BadCharacterIndex < UE_ARRAY_COUNT(INVALID_OBJECTNAME_CHARACTERS) - 1; ++BadCharacterIndex )
 		{
-			const ElementType TestChar[2] = { INVALID_OBJECTNAME_CHARACTERS[BadCharacterIndex], CHARTEXT(ElementType, '\0') };
+			const ElementType TestChar[2] = { (ElementType)INVALID_OBJECTNAME_CHARACTERS[BadCharacterIndex], CHARTEXT(ElementType, '\0') };
 			const int32 NumReplacedChars = GeneratedName.ReplaceInline(TestChar, ReplaceWith);
 		}
 	}
diff --git a/Engine/Source/Runtime/Core/Public/Containers/UnrealString.h.inl b/Engine/Source/Runtime/Core/Public/Containers/UnrealString.h.inl
--- a/Engine/Source/Runtime/Core/Public/Containers/UnrealString.h.inl
+++ b/Engine/Source/Runtime/Core/Public/Containers/UnrealString.h.inl
@@ -432,7 +432,7 @@ public:
 	 */
 	bool RemoveFromStart(const ElementType* InPrefix, ESearchCase::Type SearchCase = ESearchCase::IgnoreCase)
 	{
-		return RemoveFromStart(InPrefix, InPrefix ? FCString::Strlen(InPrefix) : 0, SearchCase);
+		return RemoveFromStart(InPrefix, InPrefix ? TCString<ElementType>::Strlen(InPrefix) : 0, SearchCase);
 	}
 
 	/**
@@ -476,7 +476,7 @@ public:
 	 */
 	bool RemoveFromEnd(const ElementType* InSuffix, ESearchCase::Type SearchCase = ESearchCase::IgnoreCase)
 	{
-		return RemoveFromEnd(InSuffix, InSuffix ? FCString::Strlen(InSuffix) : 0, SearchCase);
+		return RemoveFromEnd(InSuffix, InSuffix ? TCString<ElementType>::Strlen(InSuffix) : 0, SearchCase);
 	}
 
 	/**
@@ -588,7 +588,7 @@ public:
 	{
 		checkSlow(Str);
 
-		PathAppend(Str, FCString::Strlen(Str));
+		PathAppend(Str, TCString<ElementType>::Strlen(Str));
 		return *this;
 	}
 
@@ -630,7 +630,7 @@ public:
 	{
 		checkSlow(Rhs);
 
-		int32 StrLength = FCString::Strlen(Rhs);
+		int32 StrLength = TCString<ElementType>::Strlen(Rhs);
 
 		UE_STRING_CLASS Result(Lhs, StrLength + 1);
 		Result.PathAppend(Rhs, StrLength);
@@ -648,7 +648,7 @@ public:
 	{
 		checkSlow(Rhs);
 
-		int32 StrLength = FCString::Strlen(Rhs);
+		int32 StrLength = TCString<ElementType>::Strlen(Rhs);
 
 		UE_STRING_CLASS Result(MoveTemp(Lhs), StrLength + 1);
 		Result.PathAppend(Rhs, StrLength);
@@ -1084,7 +1084,7 @@ public:
 	UE_NODISCARD int32 Find(const ElementType* SubStr, ESearchCase::Type SearchCase = ESearchCase::IgnoreCase,
 		ESearchDir::Type SearchDir = ESearchDir::FromStart, int32 StartPosition = INDEX_NONE) const
 	{
-		return SubStr ? Find(SubStr, FCString::Strlen(SubStr), SearchCase, SearchDir, StartPosition) : INDEX_NONE;
+		return SubStr ? Find(SubStr, TCString<ElementType>::Strlen(SubStr), SearchCase, SearchDir, StartPosition) : INDEX_NONE;
 	}
 
 	/**
@@ -1263,11 +1263,11 @@ public:
 		{
 			if (SearchCase == ESearchCase::CaseSensitive)
 			{
-				return FCString::Strcmp(Data.GetData(), Other.Data.GetData()) == 0; 
+				return TCString<ElementType>::Strcmp(Data.GetData(), Other.Data.GetData()) == 0; 
 			}
 			else
 			{
-				return FCString::Stricmp(Data.GetData(), Other.Data.GetData()) == 0;
+				return TCString<ElementType>::Stricmp(Data.GetData(), Other.Data.GetData()) == 0;
 			}
 		}
 
@@ -1285,11 +1285,11 @@ public:
 	{
 		if( SearchCase == ESearchCase::CaseSensitive )
 		{
-			return FCString::Strcmp( **this, *Other ); 
+			return TCString<ElementType>::Strcmp( **this, *Other ); 
 		}
 		else
 		{
-			return FCString::Stricmp( **this, *Other );
+			return TCString<ElementType>::Stricmp( **this, *Other );
 		}
 	}
 
@@ -1442,7 +1442,7 @@ public:
 	 */
 	UE_NODISCARD bool StartsWith(const ElementType* InPrefix, ESearchCase::Type SearchCase = ESearchCase::IgnoreCase) const
 	{
-		return StartsWith(InPrefix, InPrefix ? FCString::Strlen(InPrefix) : 0, SearchCase);
+		return StartsWith(InPrefix, InPrefix ? TCString<ElementType>::Strlen(InPrefix) : 0, SearchCase);
 	}
 
 	/**
@@ -1485,7 +1485,7 @@ public:
 	 */
 	UE_NODISCARD bool EndsWith(const ElementType* InSuffix, ESearchCase::Type SearchCase = ESearchCase::IgnoreCase) const
 	{
-		return EndsWith(InSuffix, InSuffix ? FCString::Strlen(InSuffix) : 0, SearchCase);
+		return EndsWith(InSuffix, InSuffix ? TCString<ElementType>::Strlen(InSuffix) : 0, SearchCase);
 	}
 
 	/**
@@ -1532,7 +1532,7 @@ public:
 	 */
 	UE_NODISCARD bool MatchesWildcard(const ElementType* Wildcard, ESearchCase::Type SearchCase = ESearchCase::IgnoreCase) const
 	{
-		return MatchesWildcard(Wildcard, Wildcard ? FCString::Strlen(Wildcard) : 0, SearchCase);
+		return MatchesWildcard(Wildcard, Wildcard ? TCString<ElementType>::Strlen(Wildcard) : 0, SearchCase);
 	}
 
 	/**
