Fix for linuxarm64 cook platform.

diff --git a/Engine/Source/Runtime/Core/Public/Misc/DataDrivenPlatformInfoRegistry.h b/Engine/Source/Runtime/Core/Public/Misc/DataDrivenPlatformInfoRegistry.h
--- a/Engine/Source/Runtime/Core/Public/Misc/DataDrivenPlatformInfoRegistry.h
+++ b/Engine/Source/Runtime/Core/Public/Misc/DataDrivenPlatformInfoRegistry.h
@@ -166,7 +166,7 @@ struct FDataDrivenPlatformInfo
 
 	const FName& GetCookPlatformName() const
 	{
-		if (!OverrideCookPlatformName.IsValid())
+		if (OverrideCookPlatformName.IsValid())
 		{
 			return OverrideCookPlatformName;
 		}
