Replaced GET_VARARGS* macros with GET_TYPED_VARARGS*.

diff --git a/Engine/Plugins/Runtime/NetworkPrediction/Source/NetworkPrediction/Private/NetworkPredictionTrace.cpp b/Engine/Plugins/Runtime/NetworkPrediction/Source/NetworkPrediction/Private/NetworkPredictionTrace.cpp
--- a/Engine/Plugins/Runtime/NetworkPrediction/Source/NetworkPrediction/Private/NetworkPredictionTrace.cpp
+++ b/Engine/Plugins/Runtime/NetworkPrediction/Source/NetworkPrediction/Private/NetworkPredictionTrace.cpp
@@ -324,7 +324,7 @@ void FNetworkPredictionTrace::TraceOOBStateMod(int32 TraceID, int32 Frame, const
 \
 	/* first, try using the stack buffer */ \
 	Buffer = StackBuffer; \
-	GET_VARARGS_RESULT( Buffer, UE_ARRAY_COUNT(StackBuffer), UE_ARRAY_COUNT(StackBuffer) - 1, Fmt, Fmt, Result ); \
+	GET_TYPED_VARARGS_RESULT( TCHAR, Buffer, UE_ARRAY_COUNT(StackBuffer), UE_ARRAY_COUNT(StackBuffer) - 1, Fmt, Fmt, Result ); \
 \
 	/* if that fails, then use heap allocation to make enough space */ \
 	while(Result == -1) \
@@ -336,7 +336,7 @@ void FNetworkPredictionTrace::TraceOOBStateMod(int32 TraceID, int32 Frame, const
 		{ \
 			return; \
 		} \
-		GET_VARARGS_RESULT( Buffer, BufferSize, BufferSize-1, Fmt, Fmt, Result ); \
+		GET_TYPED_VARARGS_RESULT( TCHAR, Buffer, BufferSize, BufferSize-1, Fmt, Fmt, Result ); \
 		BufferSize *= 2; \
 	}; \
 	Buffer[Result] = 0; \
diff --git a/Engine/Plugins/Runtime/StateTree/Source/StateTreeModule/Private/Debugger/StateTreeTrace.cpp b/Engine/Plugins/Runtime/StateTree/Source/StateTreeModule/Private/Debugger/StateTreeTrace.cpp
--- a/Engine/Plugins/Runtime/StateTree/Source/StateTreeModule/Private/Debugger/StateTreeTrace.cpp
+++ b/Engine/Plugins/Runtime/StateTree/Source/StateTreeModule/Private/Debugger/StateTreeTrace.cpp
@@ -465,7 +465,7 @@ void OutputLogEventTrace(
 	)
 {
 	static TCHAR TraceStaticBuffer[8192];
-	GET_VARARGS(TraceStaticBuffer, UE_ARRAY_COUNT(TraceStaticBuffer), UE_ARRAY_COUNT(TraceStaticBuffer) - 1, Fmt, Fmt);
+	GET_TYPED_VARARGS(TCHAR, TraceStaticBuffer, UE_ARRAY_COUNT(TraceStaticBuffer), UE_ARRAY_COUNT(TraceStaticBuffer) - 1, Fmt, Fmt);
 
 	TraceBufferedEvents(InstanceId);
 
diff --git a/Engine/Source/Developer/DerivedDataCache/Private/S3CacheStore.cpp b/Engine/Source/Developer/DerivedDataCache/Private/S3CacheStore.cpp
--- a/Engine/Source/Developer/DerivedDataCache/Private/S3CacheStore.cpp
+++ b/Engine/Source/Developer/DerivedDataCache/Private/S3CacheStore.cpp
@@ -106,7 +106,7 @@ public:
 	static FStringAnsi Printf(const ANSICHAR* Format, ...)
 	{
 		ANSICHAR Buffer[1024];
-		GET_VARARGS_ANSI(Buffer, UE_ARRAY_COUNT(Buffer), UE_ARRAY_COUNT(Buffer) - 1, Format, Format);
+		GET_TYPED_VARARGS(ANSICHAR, Buffer, UE_ARRAY_COUNT(Buffer), UE_ARRAY_COUNT(Buffer) - 1, Format, Format);
 		return Buffer;
 	}
 
diff --git a/Engine/Source/Programs/UnrealLightmass/Private/ImportExport/LightmassSwarm.cpp b/Engine/Source/Programs/UnrealLightmass/Private/ImportExport/LightmassSwarm.cpp
--- a/Engine/Source/Programs/UnrealLightmass/Private/ImportExport/LightmassSwarm.cpp
+++ b/Engine/Source/Programs/UnrealLightmass/Private/ImportExport/LightmassSwarm.cpp
@@ -288,7 +288,7 @@ void FLightmassSwarm::SendTextMessageImpl(const TCHAR* Fmt, ...)
 	{
 		FMemory::Free(Buffer);
 		Buffer = (TCHAR*) FMemory::Malloc( BufferSize * sizeof(TCHAR) );
-		GET_VARARGS_RESULT( Buffer, BufferSize, BufferSize-1, Fmt, Fmt, Result );
+		GET_TYPED_VARARGS_RESULT( ElementType, Buffer, BufferSize, BufferSize-1, Fmt, Fmt, Result );
 		BufferSize *= 2;
 	};
 	Buffer[Result] = 0;
diff --git a/Engine/Source/Runtime/Core/Private/Containers/String.cpp.inl b/Engine/Source/Runtime/Core/Private/Containers/String.cpp.inl
--- a/Engine/Source/Runtime/Core/Private/Containers/String.cpp.inl
+++ b/Engine/Source/Runtime/Core/Private/Containers/String.cpp.inl
@@ -1646,7 +1646,7 @@ UE_STRING_CLASS UE_STRING_CLASS::PrintfImpl(const ElementType* Fmt, ...)
 	int32		Result		= -1;
 
 	// First try to print to a stack allocated location 
-	GET_VARARGS_RESULT( Buffer, BufferSize, BufferSize-1, Fmt, Fmt, Result );
+	GET_TYPED_VARARGS_RESULT( ElementType, Buffer, BufferSize, BufferSize-1, Fmt, Fmt, Result );
 
 	// If that fails, start allocating regular memory
 	if( Result == -1 )
@@ -1656,7 +1656,7 @@ UE_STRING_CLASS UE_STRING_CLASS::PrintfImpl(const ElementType* Fmt, ...)
 		{
 			BufferSize *= 2;
 			Buffer = (ElementType*) FMemory::Realloc( Buffer, BufferSize * sizeof(ElementType) );
-			GET_VARARGS_RESULT( Buffer, BufferSize, BufferSize-1, Fmt, Fmt, Result );
+			GET_TYPED_VARARGS_RESULT( ElementType, Buffer, BufferSize, BufferSize-1, Fmt, Fmt, Result );
 		};
 	}
 
@@ -1680,7 +1680,7 @@ void UE_STRING_CLASS::AppendfImpl(UE_STRING_CLASS& AppendToMe, const ElementType
 	int32		Result = -1;
 
 	// First try to print to a stack allocated location 
-	GET_VARARGS_RESULT(Buffer, BufferSize, BufferSize - 1, Fmt, Fmt, Result);
+	GET_TYPED_VARARGS_RESULT(ElementType, Buffer, BufferSize, BufferSize - 1, Fmt, Fmt, Result);
 
 	// If that fails, start allocating regular memory
 	if (Result == -1)
@@ -1690,7 +1690,7 @@ void UE_STRING_CLASS::AppendfImpl(UE_STRING_CLASS& AppendToMe, const ElementType
 		{
 			BufferSize *= 2;
 			Buffer = (ElementType*)FMemory::Realloc(Buffer, BufferSize * sizeof(ElementType));
-			GET_VARARGS_RESULT(Buffer, BufferSize, BufferSize - 1, Fmt, Fmt, Result);
+			GET_TYPED_VARARGS_RESULT(ElementType, Buffer, BufferSize, BufferSize - 1, Fmt, Fmt, Result);
 		};
 	}
 
diff --git a/Engine/Source/Runtime/Core/Private/HAL/LowLevelMemTracker.cpp b/Engine/Source/Runtime/Core/Private/HAL/LowLevelMemTracker.cpp
--- a/Engine/Source/Runtime/Core/Private/HAL/LowLevelMemTracker.cpp
+++ b/Engine/Source/Runtime/Core/Private/HAL/LowLevelMemTracker.cpp
@@ -2872,7 +2872,7 @@ bool HandleAssert(bool bLog, const TCHAR* Format, ...)
 	if (bLog)
 	{
 		TCHAR DescriptionString[4096];
-		GET_VARARGS(DescriptionString, UE_ARRAY_COUNT(DescriptionString), UE_ARRAY_COUNT(DescriptionString) - 1,
+		GET_TYPED_VARARGS(TCHAR, DescriptionString, UE_ARRAY_COUNT(DescriptionString), UE_ARRAY_COUNT(DescriptionString) - 1,
 			Format, Format);
 
 		FPlatformMisc::LowLevelOutputDebugString(DescriptionString);
diff --git a/Engine/Source/Runtime/Core/Private/Logging/LogTrace.cpp b/Engine/Source/Runtime/Core/Private/Logging/LogTrace.cpp
--- a/Engine/Source/Runtime/Core/Private/Logging/LogTrace.cpp
+++ b/Engine/Source/Runtime/Core/Private/Logging/LogTrace.cpp
@@ -76,7 +76,7 @@ static TCHAR LogTraceStaticBuffer[8192];
 void FLogTrace::OutputLogMessageSimple(const void* LogPoint, const TCHAR* Fmt, ...)
 {
 	FScopeLock MsgLock(GetLogTraceStaticBufferGuard());
-	GET_VARARGS(LogTraceStaticBuffer, UE_ARRAY_COUNT(LogTraceStaticBuffer), UE_ARRAY_COUNT(LogTraceStaticBuffer) - 1, Fmt, Fmt);
+	GET_TYPED_VARARGS(TCHAR, LogTraceStaticBuffer, UE_ARRAY_COUNT(LogTraceStaticBuffer), UE_ARRAY_COUNT(LogTraceStaticBuffer) - 1, Fmt, Fmt);
 	FLogTrace::OutputLogMessage(LogPoint, LogTraceStaticBuffer);
 }
 #endif // LOGTRACE_RUNTIME_FORMATTING_ENABLED
diff --git a/Engine/Source/Runtime/Core/Private/Misc/CoreMisc.cpp b/Engine/Source/Runtime/Core/Private/Misc/CoreMisc.cpp
--- a/Engine/Source/Runtime/Core/Private/Misc/CoreMisc.cpp
+++ b/Engine/Source/Runtime/Core/Private/Misc/CoreMisc.cpp
@@ -514,7 +514,7 @@ namespace UEAsserts_Private
 	{		
 		const int32 TempStrSize = 4096;
 		TCHAR TempStr[TempStrSize];
-		GET_VARARGS(TempStr, TempStrSize, TempStrSize - 1, FormattedMsg, FormattedMsg);
+		GET_TYPED_VARARGS(TCHAR, TempStr, TempStrSize, TempStrSize - 1, FormattedMsg, FormattedMsg);
 		UE_LOG(LogCore, Error, TEXT("%s"), TempStr);
 	}
 }
diff --git a/Engine/Source/Runtime/Core/Private/Misc/VarargsHelper.h b/Engine/Source/Runtime/Core/Private/Misc/VarargsHelper.h
--- a/Engine/Source/Runtime/Core/Private/Misc/VarargsHelper.h
+++ b/Engine/Source/Runtime/Core/Private/Misc/VarargsHelper.h
@@ -8,7 +8,7 @@
 #include "Templates/UnrealTemplate.h"
 
 // Pulled the two FOutputDevice::Logf functions into shared code. Needs to be a #define
-// since it uses GET_VARARGS_RESULT which uses the va_list stuff which operates on the
+// since it uses GET_TYPED_VARARGS_RESULT which uses the va_list stuff which operates on the
 // current function, so we can't easily call a function
 #define GROWABLE_LOGF(SerializeFunc) \
 	int32	BufferSize	= 1024; \
@@ -20,7 +20,7 @@
 \
 	/* first, try using the stack buffer */ \
 	Buffer = StackBuffer; \
-	GET_VARARGS_RESULT( Buffer, UE_ARRAY_COUNT(StackBuffer), UE_ARRAY_COUNT(StackBuffer) - 1, Fmt, Fmt, Result ); \
+	GET_TYPED_VARARGS_RESULT( TCHAR, Buffer, UE_ARRAY_COUNT(StackBuffer), UE_ARRAY_COUNT(StackBuffer) - 1, Fmt, Fmt, Result ); \
 \
 	/* if that fails, then use heap allocation to make enough space */ \
 	while(Result == -1) \
@@ -32,7 +32,7 @@
 		{ \
 			return; \
 		} \
-		GET_VARARGS_RESULT( Buffer, BufferSize, BufferSize-1, Fmt, Fmt, Result ); \
+		GET_TYPED_VARARGS_RESULT( TCHAR, Buffer, BufferSize, BufferSize-1, Fmt, Fmt, Result ); \
 		BufferSize *= 2; \
 	}; \
 	Buffer[Result] = TEXT('\0'); \
diff --git a/Engine/Source/Runtime/Core/Private/ProfilingDebugging/CsvProfiler.cpp b/Engine/Source/Runtime/Core/Private/ProfilingDebugging/CsvProfiler.cpp
--- a/Engine/Source/Runtime/Core/Private/ProfilingDebugging/CsvProfiler.cpp
+++ b/Engine/Source/Runtime/Core/Private/ProfilingDebugging/CsvProfiler.cpp
@@ -3529,7 +3529,7 @@ void FCsvProfiler::RecordEventfInternal(int32 CategoryIndex, const TCHAR* Fmt, .
 	{
 		LLM_SCOPE(ELLMTag::CsvProfiler);
 		TCHAR Buffer[256];
-		GET_VARARGS(Buffer, UE_ARRAY_COUNT(Buffer), UE_ARRAY_COUNT(Buffer) - 1, Fmt, Fmt);
+		GET_TYPED_VARARGS(TCHAR, Buffer, UE_ARRAY_COUNT(Buffer), UE_ARRAY_COUNT(Buffer) - 1, Fmt, Fmt);
 		Buffer[255] = '\0';
 		FString Str = Buffer;
 
diff --git a/Engine/Source/Runtime/Core/Private/Serialization/Archive.cpp b/Engine/Source/Runtime/Core/Private/Serialization/Archive.cpp
--- a/Engine/Source/Runtime/Core/Private/Serialization/Archive.cpp
+++ b/Engine/Source/Runtime/Core/Private/Serialization/Archive.cpp
@@ -1379,7 +1379,7 @@ void FArchive::LogfImpl(const TCHAR* Fmt, ...)
 	{
 		FMemory::SystemFree(Buffer);
 		Buffer = (TCHAR*) FMemory::SystemMalloc( BufferSize * sizeof(TCHAR) );
-		GET_VARARGS_RESULT( Buffer, BufferSize, BufferSize-1, Fmt, Fmt, Result );
+		GET_TYPED_VARARGS_RESULT( TCHAR, Buffer, BufferSize, BufferSize-1, Fmt, Fmt, Result );
 		BufferSize *= 2;
 	};
 	Buffer[Result] = TEXT('\0');
diff --git a/Engine/Source/Runtime/Core/Private/Windows/WindowsPlatformAtomics.cpp b/Engine/Source/Runtime/Core/Private/Windows/WindowsPlatformAtomics.cpp
--- a/Engine/Source/Runtime/Core/Private/Windows/WindowsPlatformAtomics.cpp
+++ b/Engine/Source/Runtime/Core/Private/Windows/WindowsPlatformAtomics.cpp
@@ -10,6 +10,6 @@
 void FWindowsPlatformAtomics::HandleAtomicsFailure(const TCHAR* InFormat, ...)
 {
 	TCHAR TempStr[1024];
-	GET_VARARGS(TempStr, UE_ARRAY_COUNT(TempStr), UE_ARRAY_COUNT(TempStr) - 1, InFormat, InFormat);
+	GET_TYPED_VARARGS(TCHAR, TempStr, UE_ARRAY_COUNT(TempStr), UE_ARRAY_COUNT(TempStr) - 1, InFormat, InFormat);
 	UE_LOG(LogWindows, Fatal, TEXT("%s"), TempStr);
 }
diff --git a/Engine/Source/Runtime/Core/Public/Misc/CString.h b/Engine/Source/Runtime/Core/Public/Misc/CString.h
--- a/Engine/Source/Runtime/Core/Public/Misc/CString.h
+++ b/Engine/Source/Runtime/Core/Public/Misc/CString.h
@@ -1085,53 +1085,19 @@ int32 TCString<T>::GetVarArgs( CharType* Dest, SIZE_T DestSize, const CharType*&
 template <typename T>
 int32 TCString<T>::SprintfImpl(CharType* Dest, const CharType* Fmt, ...)
 {
-	if constexpr (std::is_same_v<CharType, ANSICHAR>)
-	{
-		int32	Result = -1;
-		GET_VARARGS_RESULT_ANSI(Dest, MAX_SPRINTF, MAX_SPRINTF - 1, Fmt, Fmt, Result);
-		return Result;
-	}
-	else if constexpr (std::is_same_v<CharType, WIDECHAR>)
-	{
-		int32	Result = -1;
-		GET_VARARGS_RESULT_WIDE(Dest, MAX_SPRINTF, MAX_SPRINTF - 1, Fmt, Fmt, Result);
-		return Result;
-	}
-	else if constexpr (std::is_same_v<CharType, UTF8CHAR>)
-	{
-		int32	Result = -1;
-		GET_VARARGS_RESULT_UTF8(Dest, MAX_SPRINTF, MAX_SPRINTF - 1, Fmt, Fmt, Result);
-		return Result;
-	}
-	else
-	{
-		static_assert(sizeof(CharType) == 0, "Not supported");
-	}
+	static_assert(std::is_same_v<CharType, ANSICHAR> || std::is_same_v<CharType, WIDECHAR> || std::is_same_v<CharType, UTF8CHAR>, "Not supported");
+
+	int32	Result = -1;
+	GET_TYPED_VARARGS_RESULT(CharType, Dest, MAX_SPRINTF, MAX_SPRINTF - 1, Fmt, Fmt, Result);
+	return Result;
 }
 
 template <typename T>
 int32 TCString<T>::SnprintfImpl(CharType* Dest, int32 DestSize, const CharType* Fmt, ...)
 {
-	if constexpr (std::is_same_v<CharType, ANSICHAR>)
-	{
-		int32	Result = -1;
-		GET_VARARGS_RESULT_ANSI(Dest, DestSize, DestSize - 1, Fmt, Fmt, Result);
-		return Result;
-	}
-	else if constexpr (std::is_same_v<CharType, WIDECHAR>)
-	{
-		int32	Result = -1;
-		GET_VARARGS_RESULT_WIDE(Dest, DestSize, DestSize - 1, Fmt, Fmt, Result);
-		return Result;
-	}
-	else if constexpr (std::is_same_v<CharType, UTF8CHAR>)
-	{
-		int32	Result = -1;
-		GET_VARARGS_RESULT_UTF8(Dest, DestSize, DestSize - 1, Fmt, Fmt, Result);
-		return Result;
-	}
-	else
-	{
-		static_assert(sizeof(CharType) == 0, "Not supported");
-	}
+	static_assert(std::is_same_v<CharType, ANSICHAR> || std::is_same_v<CharType, WIDECHAR> || std::is_same_v<CharType, UTF8CHAR>, "Not supported");
+
+	int32	Result = -1;
+	GET_TYPED_VARARGS_RESULT(CharType, Dest, DestSize, DestSize - 1, Fmt, Fmt, Result);
+	return Result;
 }
diff --git a/Engine/Source/Runtime/Core/Public/Misc/VarArgs.h b/Engine/Source/Runtime/Core/Public/Misc/VarArgs.h
--- a/Engine/Source/Runtime/Core/Public/Misc/VarArgs.h
+++ b/Engine/Source/Runtime/Core/Public/Misc/VarArgs.h
@@ -14,66 +14,23 @@ VarArgs helper macros.
 	#define VARARG_BODY( FuncRet, FuncName, FmtType, ExtraDecl ) FuncRet FuncName( ExtraDecl FmtType Fmt, ... )
 #endif // VISUAL_ASSIST_HACK
 
-
-#define GET_VARARGS(msg, msgsize, len, lastarg, fmt) \
-	{ \
-		va_list ap; \
-		va_start(ap, lastarg); \
-		FCString::GetVarArgs(msg, msgsize, fmt, ap); \
-		va_end(ap); \
-	}
-#define GET_VARARGS_WIDE(msg, msgsize, len, lastarg, fmt) \
-	{ \
-		va_list ap; \
-		va_start(ap, lastarg); \
-		FCStringWide::GetVarArgs(msg, msgsize, fmt, ap); \
-		va_end(ap); \
-	}
-#define GET_VARARGS_ANSI(msg, msgsize, len, lastarg, fmt) \
+#define GET_TYPED_VARARGS(CharType, msg, msgsize, len, lastarg, fmt) \
 	{ \
 		va_list ap; \
 		va_start(ap, lastarg); \
-		FCStringAnsi::GetVarArgs(msg, msgsize, fmt, ap); \
+		TCString<CharType>::GetVarArgs(msg, msgsize, fmt, ap); \
 		va_end(ap); \
 	}
-#define GET_VARARGS_RESULT(msg, msgsize, len, lastarg, fmt, result) \
-	{ \
-		va_list ap; \
-		va_start(ap, lastarg); \
-		result = FCString::GetVarArgs(msg, msgsize, fmt, ap); \
-		if (result >= msgsize) \
-		{ \
-			result = -1; \
-		} \
-		va_end(ap); \
-	}
-#define GET_VARARGS_RESULT_WIDE(msg, msgsize, len, lastarg, fmt, result) \
-	{ \
-		va_list ap; \
-		va_start(ap, lastarg); \
-		result = FCStringWide::GetVarArgs(msg, msgsize, fmt, ap); \
-		if (result >= msgsize) \
-		{ \
-			result = -1; \
-		} \
-		va_end(ap); \
-	}
-#define GET_VARARGS_RESULT_ANSI(msg, msgsize, len, lastarg, fmt, result) \
-	{ \
-		va_list ap; \
-		va_start(ap, lastarg); \
-		result = FCStringAnsi::GetVarArgs(msg, msgsize, fmt, ap); \
-		if (result >= msgsize) \
-		{ \
-			result = -1; \
-		} \
-		va_end(ap); \
-	}
-#define GET_VARARGS_RESULT_UTF8(msg, msgsize, len, lastarg, fmt, result) \
+
+#define GET_VARARGS(     msg, msgsize, len, lastarg, fmt) UE_DEPRECATED_MACRO(5.4, "GET_VARARGS(...) has been deprecated - please use GET_TYPED_VARARGS(TCHAR, ...) instead")         GET_TYPED_VARARGS(TCHAR,    msg, msgsize, len, lastarg, fmt)
+#define GET_VARARGS_WIDE(msg, msgsize, len, lastarg, fmt) UE_DEPRECATED_MACRO(5.4, "GET_VARARGS_WIDE(...) has been deprecated - please use GET_TYPED_VARARGS(WIDECHAR, ...) instead") GET_TYPED_VARARGS(WIDECHAR, msg, msgsize, len, lastarg, fmt)
+#define GET_VARARGS_ANSI(msg, msgsize, len, lastarg, fmt) UE_DEPRECATED_MACRO(5.4, "GET_VARARGS_ANSI(...) has been deprecated - please use GET_TYPED_VARARGS(ANSICHAR, ...) instead") GET_TYPED_VARARGS(ANSICHAR, msg, msgsize, len, lastarg, fmt)
+
+#define GET_TYPED_VARARGS_RESULT(CharType, msg, msgsize, len, lastarg, fmt, result) \
 	{ \
 		va_list ap; \
 		va_start(ap, lastarg); \
-		result = FCStringUtf8::GetVarArgs(msg, msgsize, fmt, ap); \
+		result = TCString<CharType>::GetVarArgs(msg, msgsize, fmt, ap); \
 		if (result >= msgsize) \
 		{ \
 			result = -1; \
@@ -81,6 +38,11 @@ VarArgs helper macros.
 		va_end(ap); \
 	}
 
+#define GET_VARARGS_RESULT(     msg, msgsize, len, lastarg, fmt, result) UE_DEPRECATED_MACRO(5.4, "GET_VARARGS(...) has been deprecated - please use GET_TYPED_VARARGS(TCHAR, ...) instead") GET_TYPED_VARARGS_RESULT(TCHAR,    msg, msgsize, len, lastarg, fmt, result)
+#define GET_VARARGS_RESULT_WIDE(msg, msgsize, len, lastarg, fmt, result) UE_DEPRECATED_MACRO(5.4, "GET_VARARGS_RESULT_WIDE(...) has been deprecated - please use GET_TYPED_VARARGS(WIDECHAR, ...) instead") GET_TYPED_VARARGS_RESULT(WIDECHAR, msg, msgsize, len, lastarg, fmt, result)
+#define GET_VARARGS_RESULT_ANSI(msg, msgsize, len, lastarg, fmt, result) UE_DEPRECATED_MACRO(5.4, "GET_VARARGS_RESULT_ANSI(...) has been deprecated - please use GET_TYPED_VARARGS(ANSICHAR, ...) instead") GET_TYPED_VARARGS_RESULT(ANSICHAR, msg, msgsize, len, lastarg, fmt, result)
+#define GET_VARARGS_RESULT_UTF8(msg, msgsize, len, lastarg, fmt, result) UE_DEPRECATED_MACRO(5.4, "GET_VARARGS_RESULT_UTF8(...) has been deprecated - please use GET_TYPED_VARARGS(UTF8CHAR, ...) instead") GET_TYPED_VARARGS_RESULT(UTF8CHAR, msg, msgsize, len, lastarg, fmt, result)
+
 /*-----------------------------------------------------------------------------
 Ugly VarArgs type checking (debug builds only).
 -----------------------------------------------------------------------------*/
diff --git a/Engine/Source/Runtime/Core/Public/ProfilingDebugging/DiagnosticTable.h b/Engine/Source/Runtime/Core/Public/ProfilingDebugging/DiagnosticTable.h
--- a/Engine/Source/Runtime/Core/Public/ProfilingDebugging/DiagnosticTable.h
+++ b/Engine/Source/Runtime/Core/Public/ProfilingDebugging/DiagnosticTable.h
@@ -47,7 +47,7 @@ public:
 
 		// Format the input string.
 		TCHAR ColumnText[4096];
-		GET_VARARGS(ColumnText,UE_ARRAY_COUNT(ColumnText),UE_ARRAY_COUNT(ColumnText)-1,Format,Format);
+		GET_TYPED_VARARGS(TCHAR,ColumnText,UE_ARRAY_COUNT(ColumnText),UE_ARRAY_COUNT(ColumnText)-1,Format,Format);
 
 		// If this isn't the first column in the row, add a comma to separate the columns.
 		if(CurrentRow.Len())
diff --git a/Engine/Source/Runtime/Datasmith/CADKernel/Private/CADKernel/UI/Message.cpp b/Engine/Source/Runtime/Datasmith/CADKernel/Private/CADKernel/UI/Message.cpp
--- a/Engine/Source/Runtime/Datasmith/CADKernel/Private/CADKernel/UI/Message.cpp
+++ b/Engine/Source/Runtime/Datasmith/CADKernel/Private/CADKernel/UI/Message.cpp
@@ -49,7 +49,7 @@ void FMessage::VPrintf(EVerboseLevel Level, const TCHAR* Text, ...)
 		{
 			FMemory::Free(Buffer);
 			Buffer = (TCHAR*)FMemory::Malloc(BufferSize * sizeof(TCHAR));
-			GET_VARARGS_RESULT(Buffer, BufferSize, BufferSize - 1, Text, Text, Result);
+			GET_TYPED_VARARGS_RESULT(TCHAR, Buffer, BufferSize, BufferSize - 1, Text, Text, Result);
 			BufferSize *= 2;
 		};
 		Buffer[Result] = 0;
@@ -112,7 +112,7 @@ void FMessage::VReportPrintF(FString Header, const TCHAR* Text, ...)
 		{
 			FMemory::Free(Buffer);
 			Buffer = (TCHAR*)FMemory::Malloc(BufferSize * sizeof(TCHAR));
-			GET_VARARGS_RESULT(Buffer, BufferSize, BufferSize - 1, Text, Text, Result);
+			GET_TYPED_VARARGS_RESULT(TCHAR, Buffer, BufferSize, BufferSize - 1, Text, Text, Result);
 			BufferSize *= 2;
 		};
 		Buffer[Result] = 0;
diff --git a/Engine/Source/Runtime/Engine/Private/Materials/HLSLMaterialTranslator.cpp b/Engine/Source/Runtime/Engine/Private/Materials/HLSLMaterialTranslator.cpp
--- a/Engine/Source/Runtime/Engine/Private/Materials/HLSLMaterialTranslator.cpp
+++ b/Engine/Source/Runtime/Engine/Private/Materials/HLSLMaterialTranslator.cpp
@@ -3564,7 +3564,7 @@ int32 FHLSLMaterialTranslator::AddCodeChunkInner(EMaterialValueType Type, EDeriv
 	while(Result == -1)
 	{
 		FormattedCode = (TCHAR*) FMemory::Realloc( FormattedCode, BufferSize * sizeof(TCHAR) );
-		GET_VARARGS_RESULT( FormattedCode, BufferSize, BufferSize-1, Format, Format, Result );
+		GET_TYPED_VARARGS_RESULT( TCHAR, FormattedCode, BufferSize, BufferSize-1, Format, Format, Result );
 		BufferSize *= 2;
 	};
 	FormattedCode[Result] = 0;
@@ -3585,7 +3585,7 @@ int32 FHLSLMaterialTranslator::AddCodeChunkWithHash(uint64 BaseHash, EMaterialVa
 	while (Result == -1)
 	{
 		FormattedCode = (TCHAR*)FMemory::Realloc(FormattedCode, BufferSize * sizeof(TCHAR));
-		GET_VARARGS_RESULT(FormattedCode, BufferSize, BufferSize - 1, Format, Format, Result);
+		GET_TYPED_VARARGS_RESULT(TCHAR, FormattedCode, BufferSize, BufferSize - 1, Format, Format, Result);
 		BufferSize *= 2;
 	};
 	FormattedCode[Result] = 0;
@@ -3722,7 +3722,7 @@ int32 FHLSLMaterialTranslator::AddUniformExpression(FMaterialUniformExpression*
 	while (Result == -1)
 	{
 		FormattedCode = (TCHAR*)FMemory::Realloc(FormattedCode, BufferSize * sizeof(TCHAR));
-		GET_VARARGS_RESULT(FormattedCode, BufferSize, BufferSize - 1, Format, Format, Result);
+		GET_TYPED_VARARGS_RESULT(TCHAR, FormattedCode, BufferSize, BufferSize - 1, Format, Format, Result);
 		BufferSize *= 2;
 	};
 	FormattedCode[Result] = 0;
diff --git a/Engine/Source/Runtime/Engine/Private/Materials/MaterialShared.cpp b/Engine/Source/Runtime/Engine/Private/Materials/MaterialShared.cpp
--- a/Engine/Source/Runtime/Engine/Private/Materials/MaterialShared.cpp
+++ b/Engine/Source/Runtime/Engine/Private/Materials/MaterialShared.cpp
@@ -301,7 +301,7 @@ bool ReloadMaterialResource(
 int32 FMaterialCompiler::Errorf(const TCHAR* Format,...)
 {
 	TCHAR	ErrorText[2048];
-	GET_VARARGS( ErrorText, UE_ARRAY_COUNT(ErrorText), UE_ARRAY_COUNT(ErrorText)-1, Format, Format );
+	GET_TYPED_VARARGS( TCHAR, ErrorText, UE_ARRAY_COUNT(ErrorText), UE_ARRAY_COUNT(ErrorText)-1, Format, Format );
 	return Error(ErrorText);
 }
 
diff --git a/Engine/Source/Runtime/Engine/Private/UserInterface/Canvas.cpp b/Engine/Source/Runtime/Engine/Private/UserInterface/Canvas.cpp
--- a/Engine/Source/Runtime/Engine/Private/UserInterface/Canvas.cpp
+++ b/Engine/Source/Runtime/Engine/Private/UserInterface/Canvas.cpp
@@ -1607,7 +1607,7 @@ void UCanvas::ClippedStrLen( const UFont* Font, float ScaleX, float ScaleY, int3
 void VARARGS UCanvas::WrappedStrLenf( const UFont* Font, float ScaleX, float ScaleY, int32& XL, int32& YL, const TCHAR* Fmt, ... ) 
 {
 	TCHAR Text[4096];
-	GET_VARARGS( Text, UE_ARRAY_COUNT(Text), UE_ARRAY_COUNT(Text)-1, Fmt, Fmt );
+	GET_TYPED_VARARGS( TCHAR, Text, UE_ARRAY_COUNT(Text), UE_ARRAY_COUNT(Text)-1, Fmt, Fmt );
 
 	FFontRenderInfo Info;
 	WrappedPrint( false, 0.0f, 0.0f, XL, YL, Font, ScaleX, ScaleY, false, false, Text, Info ); 
diff --git a/Engine/Source/Runtime/Engine/Private/VisualLogger/VisualLogger.cpp b/Engine/Source/Runtime/Engine/Private/VisualLogger/VisualLogger.cpp
--- a/Engine/Source/Runtime/Engine/Private/VisualLogger/VisualLogger.cpp
+++ b/Engine/Source/Runtime/Engine/Private/VisualLogger/VisualLogger.cpp
@@ -25,7 +25,7 @@ DEFINE_LOG_CATEGORY(LogVisual);
 
 DEFINE_STAT(STAT_VisualLog);
 
-// Unfortunately needs to be a #define since it uses GET_VARARGS_RESULT which uses the va_list stuff which operates on the
+// Unfortunately needs to be a #define since it uses GET_TYPED_VARARGS_RESULT which uses the va_list stuff which operates on the
 // current function, so we can't easily call a function
 #define COLLAPSED_LOGF(SerializeFunc) \
 	SCOPE_CYCLE_COUNTER(STAT_VisualLog); \
@@ -44,7 +44,7 @@ DEFINE_STAT(STAT_VisualLog);
 	\
 	/* first, try using the stack buffer */ \
 	Buffer = StackBuffer; \
-	GET_VARARGS_RESULT( Buffer, UE_ARRAY_COUNT(StackBuffer), UE_ARRAY_COUNT(StackBuffer) - 1, Fmt, Fmt, Result ); \
+	GET_TYPED_VARARGS_RESULT( TCHAR, Buffer, UE_ARRAY_COUNT(StackBuffer), UE_ARRAY_COUNT(StackBuffer) - 1, Fmt, Fmt, Result ); \
 	\
 	/* if that fails, then use heap allocation to make enough space */ \
 			while(Result == -1) \
@@ -52,7 +52,7 @@ DEFINE_STAT(STAT_VisualLog);
 		FMemory::SystemFree(AllocatedBuffer); \
 		/* We need to use malloc here directly as GMalloc might not be safe. */ \
 		Buffer = AllocatedBuffer = (TCHAR*) FMemory::SystemMalloc( BufferSize * sizeof(TCHAR) ); \
-		GET_VARARGS_RESULT( Buffer, BufferSize, BufferSize-1, Fmt, Fmt, Result ); \
+		GET_TYPED_VARARGS_RESULT( TCHAR, Buffer, BufferSize, BufferSize-1, Fmt, Fmt, Result ); \
 		BufferSize *= 2; \
 						}; \
 	Buffer[Result] = 0; \
diff --git a/Engine/Source/Runtime/GameplayDebugger/Private/GameplayDebuggerTypes.cpp b/Engine/Source/Runtime/GameplayDebugger/Private/GameplayDebuggerTypes.cpp
--- a/Engine/Source/Runtime/GameplayDebugger/Private/GameplayDebuggerTypes.cpp
+++ b/Engine/Source/Runtime/GameplayDebugger/Private/GameplayDebuggerTypes.cpp
@@ -665,7 +665,7 @@ void FGameplayDebuggerCanvasContext::PrintAt(float PosX, float PosY, const FColo
 \
 	/* first, try using the stack buffer */ \
 	Buffer = StackBuffer; \
-	GET_VARARGS_RESULT( Buffer, UE_ARRAY_COUNT(StackBuffer), UE_ARRAY_COUNT(StackBuffer) - 1, Fmt, Fmt, Result ); \
+	GET_TYPED_VARARGS_RESULT( TCHAR, Buffer, UE_ARRAY_COUNT(StackBuffer), UE_ARRAY_COUNT(StackBuffer) - 1, Fmt, Fmt, Result ); \
 \
 	/* if that fails, then use heap allocation to make enough space */ \
 	while(Result == -1) \
@@ -673,7 +673,7 @@ void FGameplayDebuggerCanvasContext::PrintAt(float PosX, float PosY, const FColo
 		FMemory::SystemFree(AllocatedBuffer); \
 		/* We need to use malloc here directly as GMalloc might not be safe. */ \
 		Buffer = AllocatedBuffer = (TCHAR*) FMemory::SystemMalloc( BufferSize * sizeof(TCHAR) ); \
-		GET_VARARGS_RESULT( Buffer, BufferSize, BufferSize-1, Fmt, Fmt, Result ); \
+		GET_TYPED_VARARGS_RESULT( TCHAR, Buffer, BufferSize, BufferSize-1, Fmt, Fmt, Result ); \
 		BufferSize *= 2; \
 	}; \
 	Buffer[Result] = 0; \
diff --git a/Engine/Source/Runtime/RHI/Private/RHIBreadcrumbs.cpp b/Engine/Source/Runtime/RHI/Private/RHIBreadcrumbs.cpp
--- a/Engine/Source/Runtime/RHI/Private/RHIBreadcrumbs.cpp
+++ b/Engine/Source/Runtime/RHI/Private/RHIBreadcrumbs.cpp
@@ -39,7 +39,7 @@ FRHIBreadcrumb* FRHIBreadcrumbStack::PushBreadcrumbPrintf(FMemStackBase& Allocat
 	TCHAR BreadcrumbString[1024];
 	int32 WrittenLength = 0;
 
-	GET_VARARGS_RESULT(BreadcrumbString, UE_ARRAY_COUNT(BreadcrumbString), UE_ARRAY_COUNT(BreadcrumbString) - 1, InFormat, InFormat, WrittenLength);
+	GET_TYPED_VARARGS_RESULT(TCHAR, BreadcrumbString, UE_ARRAY_COUNT(BreadcrumbString), UE_ARRAY_COUNT(BreadcrumbString) - 1, InFormat, InFormat, WrittenLength);
 
 	return PushBreadcrumb(Allocator, BreadcrumbString, WrittenLength);
 }
@@ -111,7 +111,8 @@ void FRHIBreadcrumbStack::WriteRenderBreadcrumbs(FCrashContextExtendedWriter& Wr
 		if (BreadcrumbStackStringSize < MAX_BREADCRUMB_STRING)
 		{
 			int32 WrittenLength = 0;
-			GET_VARARGS_RESULT(
+			GET_TYPED_VARARGS_RESULT(
+				TCHAR,
 				&StaticBreadcrumbStackString[BreadcrumbStackStringSize],
 				MAX_BREADCRUMB_STRING - BreadcrumbStackStringSize,
 				MAX_BREADCRUMB_STRING - BreadcrumbStackStringSize - 1,
