Make Flush function public.

diff --git a/Engine/Source/Runtime/Core/Public/Logging/MessageLog.h b/Engine/Source/Runtime/Core/Public/Logging/MessageLog.h
--- a/Engine/Source/Runtime/Core/Public/Logging/MessageLog.h
+++ b/Engine/Source/Runtime/Core/Public/Logging/MessageLog.h
@@ -31,6 +31,9 @@ public:
 	 */
 	CORE_API ~FMessageLog();
 
+	/** Send the currently buffered messages to the log & clear the buffer */
+	CORE_API void Flush();
+
 	/**
 	 * Add a message to the log.
 	 * @param	InMessage	The message to add
@@ -123,10 +126,6 @@ public:
 	/** Helper function to retrieve a log color for a specified message severity */
 	CORE_API static const TCHAR* const GetLogColor( EMessageSeverity::Type InSeverity );
 
-private:
-	/** Send the currently buffered messages to the log & clear the buffer */
-	void Flush();
-
 private:
 	/** Buffer for messages */
 	TArray< TSharedRef<FTokenizedMessage> > Messages;
