PR #10928: FHashBuilder::AppendRaw overload only takes two arguments

diff --git a/Engine/Source/Runtime/Core/Public/Misc/HashBuilder.h b/Engine/Source/Runtime/Core/Public/Misc/HashBuilder.h
--- a/Engine/Source/Runtime/Core/Public/Misc/HashBuilder.h
+++ b/Engine/Source/Runtime/Core/Public/Misc/HashBuilder.h
@@ -29,7 +29,7 @@ public:
 	template <typename T>
 	typename TEnableIf<TIsPODType<T>::Value, FHashBuilder&>::Type AppendRaw(const T& InData)
 	{
-		AppendRaw(&InData, sizeof(T), Hash);
+		AppendRaw(&InData, sizeof(T));
 		return *this;
 	}
 
