Update VMA pool scale to match actual usage in games and servers.

diff --git a/Engine/Source/Runtime/Core/Private/HAL/Allocators/PooledVirtualMemoryAllocator.cpp b/Engine/Source/Runtime/Core/Private/HAL/Allocators/PooledVirtualMemoryAllocator.cpp
--- a/Engine/Source/Runtime/Core/Private/HAL/Allocators/PooledVirtualMemoryAllocator.cpp
+++ b/Engine/Source/Runtime/Core/Private/HAL/Allocators/PooledVirtualMemoryAllocator.cpp
@@ -15,7 +15,12 @@
 using T64KBAlignedPool = TMemoryPool<65536>;
 
 #ifndef UE_VMA_POOL_SCALE
-#define UE_VMA_POOL_SCALE 1.4f
+	#if !UE_EDITOR
+		// all non-editor targets have been long running with -vmapoolscale=1.0
+		#define UE_VMA_POOL_SCALE 1.0f
+	#else
+		#define UE_VMA_POOL_SCALE 1.4f
+	#endif
 #endif
 
 /** Scale parameter used when growing the pools on allocation (and scaling them back), configurable from the commandline */
