Removed some more TCHAR from UnrealString.h.inl.

diff --git a/Engine/Source/Runtime/Core/Public/Containers/UnrealString.h.inl b/Engine/Source/Runtime/Core/Public/Containers/UnrealString.h.inl
--- a/Engine/Source/Runtime/Core/Public/Containers/UnrealString.h.inl
+++ b/Engine/Source/Runtime/Core/Public/Containers/UnrealString.h.inl
@@ -81,6 +81,10 @@ private:
 		TIsTCharRangeNotCArray<CharRangeType>,
 		TNot<TIsDerivedFrom<typename TDecay<CharRangeType>::Type, UE_STRING_CLASS>>>;
 
+	/** Like the TIsCharEncodingCompatibleWithTCHAR trait, but for the element type of the string */
+	template <typename SrcEncoding>
+	using TIsCharEncodingCompatibleWithElementType = TIsCharEncodingCompatibleWith<SrcEncoding, ElementType>;
+
 public:
 	UE_STRING_CLASS() = default;
 	UE_STRING_CLASS(UE_STRING_CLASS&&) = default;
@@ -1356,7 +1360,7 @@ public:
 	template <typename FmtType, typename... Types>
 	UE_NODISCARD static UE_STRING_CLASS Printf(const FmtType& Fmt, Types... Args)
 	{
-		static_assert(TIsArrayOrRefOfTypeByPredicate<FmtType, TIsCharEncodingCompatibleWithTCHAR>::Value, "Formatting string must be a character array.");
+		static_assert(TIsArrayOrRefOfTypeByPredicate<FmtType, TIsCharEncodingCompatibleWithElementType>::Value, "Formatting string must be a character array.");
 		static_assert(TAnd<TIsValidVariadicFunctionArg<Types>...>::Value, "Invalid argument(s) passed to Printf");
 
 		return PrintfImpl((const ElementType*)Fmt, Args...);
@@ -1369,7 +1373,7 @@ public:
 	template <typename FmtType, typename... Types>
 	UE_STRING_CLASS& Appendf(const FmtType& Fmt, Types... Args)
 	{
-		static_assert(TIsArrayOrRefOfTypeByPredicate<FmtType, TIsCharEncodingCompatibleWithTCHAR>::Value, "Formatting string must be a character array.");
+		static_assert(TIsArrayOrRefOfTypeByPredicate<FmtType, TIsCharEncodingCompatibleWithElementType>::Value, "Formatting string must be a character array.");
 		static_assert(TAnd<TIsValidVariadicFunctionArg<Types>...>::Value, "Invalid argument(s) passed to TString::Appendf");
 
 		AppendfImpl(*this, (const ElementType*)Fmt, Args...);
