Moved LexToString overloads of string classes to to the generic header.

diff --git a/Engine/Source/Runtime/Core/Public/Containers/UnrealString.h b/Engine/Source/Runtime/Core/Public/Containers/UnrealString.h
--- a/Engine/Source/Runtime/Core/Public/Containers/UnrealString.h
+++ b/Engine/Source/Runtime/Core/Public/Containers/UnrealString.h
@@ -149,16 +149,6 @@ UE_NODISCARD inline FString LexToString(bool Value)
 	return Value ? TEXT("true") : TEXT("false");
 }
 
-UE_NODISCARD FORCEINLINE FString LexToString(FString&& Str)
-{
-	return MoveTemp(Str);
-}
-
-UE_NODISCARD FORCEINLINE FString LexToString(const FString& Str)
-{
-	return Str;
-}
-
 /** Helper template to convert to sanitized strings */
 template<typename T>
 UE_NODISCARD FString LexToSanitizedString(const T& Value)
diff --git a/Engine/Source/Runtime/Core/Public/Containers/UnrealString.h.inl b/Engine/Source/Runtime/Core/Public/Containers/UnrealString.h.inl
--- a/Engine/Source/Runtime/Core/Public/Containers/UnrealString.h.inl
+++ b/Engine/Source/Runtime/Core/Public/Containers/UnrealString.h.inl
@@ -2119,6 +2119,18 @@ inline void LexFromString(double& OutValue, 			const TCHAR* Buffer)	{	OutValue =
 inline void LexFromString(bool& OutValue, 				const TCHAR* Buffer)	{	OutValue = FCString::ToBool(Buffer);	}
 inline void LexFromString(UE_STRING_CLASS& OutValue, 	const TCHAR* Buffer)	{	OutValue = Buffer;						}
 
+template <typename StringType = FString>
+UE_NODISCARD FORCEINLINE StringType LexToString(UE_STRING_CLASS&& Str)
+{
+	return MoveTemp(Str);
+}
+
+template <typename StringType = FString>
+UE_NODISCARD FORCEINLINE StringType LexToString(const UE_STRING_CLASS& Str)
+{
+	return Str;
+}
+
 /**
  * Gets a non-owning TCHAR pointer from a string type.
  *
