Made LexToString more generic by providing a string type as an optional template parameter.

diff --git a/Engine/Source/Runtime/Core/Public/Containers/UnrealString.h b/Engine/Source/Runtime/Core/Public/Containers/UnrealString.h
--- a/Engine/Source/Runtime/Core/Public/Containers/UnrealString.h
+++ b/Engine/Source/Runtime/Core/Public/Containers/UnrealString.h
@@ -129,65 +129,71 @@ UE_NODISCARD inline const uint8 TCharToNibble(const TCHAR Hex)
 }
 
 /** Convert numeric types to a string */
-template<typename T>
-UE_NODISCARD typename TEnableIf<TIsArithmetic<T>::Value, FString>::Type
+template <typename StringType = FString, typename T>
+UE_NODISCARD typename TEnableIf<TIsArithmetic<T>::Value, StringType>::Type
 LexToString(const T& Value)
 {
 	// std::remove_cv_t to remove potential volatile decorations. Removing const is pointless, but harmless because it's specified in the param declaration.
-	return FString::Printf(TFormatSpecifier<std::remove_cv_t<T>>::GetFormatSpecifier(), Value);
+	return StringType::Printf(TFormatSpecifier<std::remove_cv_t<T>>::GetFormatSpecifier(), Value);
 }
 
-template<typename CharType>
-UE_NODISCARD typename TEnableIf<TIsCharType<CharType>::Value, FString>::Type
+template <typename StringType = FString, typename CharType>
+UE_NODISCARD typename TEnableIf<TIsCharType<CharType>::Value, StringType>::Type
 LexToString(const CharType* Ptr)
 {
-	return FString(Ptr);
+	return StringType(Ptr);
 }
 
-UE_NODISCARD inline FString LexToString(bool Value)
+template <typename StringType = FString>
+UE_NODISCARD inline StringType LexToString(bool Value)
 {
-	return Value ? TEXT("true") : TEXT("false");
+	using ElementType = typename StringType::ElementType;
+	return Value ? CHARTEXT(ElementType, "true") : CHARTEXT(ElementType, "false");
 }
 
 /** Helper template to convert to sanitized strings */
-template<typename T>
-UE_NODISCARD FString LexToSanitizedString(const T& Value)
+template <typename StringType = FString, typename T>
+UE_NODISCARD StringType LexToSanitizedString(const T& Value)
 {
-	return LexToString(Value);
+	return LexToString<StringType>(Value);
 }
 
 /** Overloaded for floats */
-UE_NODISCARD inline FString LexToSanitizedString(float Value)
+template <typename StringType = FString>
+UE_NODISCARD inline StringType LexToSanitizedString(float Value)
 {
-	return FString::SanitizeFloat(Value);
+	return StringType::SanitizeFloat(Value);
 }
 
 /** Overloaded for doubles */
-UE_NODISCARD inline FString LexToSanitizedString(double Value)
+template <typename StringType = FString>
+UE_NODISCARD inline StringType LexToSanitizedString(double Value)
 {
-	return FString::SanitizeFloat(Value);
+	return StringType::SanitizeFloat(Value);
 }
 
 /** Shorthand legacy use for Lex functions */
 template<typename T>
 struct TTypeToString
 {
-	UE_NODISCARD static FString ToString(const T& Value)
+	template <typename StringType = FString>
+	UE_NODISCARD static StringType ToString(const T& Value)
 	{
-		return LexToString(Value);
+		return LexToString<StringType>(Value);
 	}
 
-	UE_NODISCARD static FString ToSanitizedString(const T& Value)
+	template <typename StringType = FString>
+	UE_NODISCARD static StringType ToSanitizedString(const T& Value)
 	{
-		return LexToSanitizedString(Value);
+		return LexToSanitizedString<StringType>(Value);
 	}
 };
 
 /** Parse a string into this type, returning whether it was successful */
 /** Specialization for arithmetic types */
-template<typename T>
+template <typename T, typename CharType>
 typename TEnableIf<TIsArithmetic<T>::Value, bool>::Type
-LexTryParseString(T& OutValue, const TCHAR* Buffer)
+LexTryParseString(T& OutValue, const CharType* Buffer)
 {
 	if (Buffer[0] == '\0')
 	{
@@ -199,16 +205,16 @@ LexTryParseString(T& OutValue, const TCHAR* Buffer)
 	if (OutValue == 0 && FMath::IsFinite((float)OutValue)) //@TODO:FLOATPRECISION: ? huh ?
 	{
 		bool bSawZero = false;
-		TCHAR C = *Buffer;
-		while (C != '\0' && (C == '+' || C == '-' || FChar::IsWhitespace(C)))
+		CharType C = *Buffer;
+		while (C != CHARTEXT(CharType, '\0') && (C == CHARTEXT(CharType, '+') || C == CHARTEXT(CharType, '-') || TChar<CharType>::IsWhitespace(C)))
 		{
 			C = *(++Buffer);
 		}
 
-		while (C != '\0' && !FChar::IsWhitespace(C) && (TIsFloatingPoint<T>::Value || C != '.'))
+		while (C != CHARTEXT(CharType, '\0') && !TChar<CharType>::IsWhitespace(C) && (TIsFloatingPoint<T>::Value || C != CHARTEXT(CharType, '.')))
 		{
-			bSawZero = bSawZero || (C == '0');
-			if (!bSawZero && C != '.')
+			bSawZero = bSawZero || (C == CHARTEXT(CharType, '0'));
+			if (!bSawZero && C != CHARTEXT(CharType, '.'))
 			{
 				return false;
 			}
@@ -222,7 +228,8 @@ LexTryParseString(T& OutValue, const TCHAR* Buffer)
 }
 
 /** Try and parse a bool - always returns true */
-inline bool LexTryParseString(bool& OutValue, const TCHAR* Buffer)
+template <typename CharType>
+inline bool LexTryParseString(bool& OutValue, const CharType* Buffer)
 {
 	LexFromString(OutValue, Buffer);
 	return true;
@@ -231,7 +238,8 @@ inline bool LexTryParseString(bool& OutValue, const TCHAR* Buffer)
 template<typename T>
 struct TTypeFromString
 {
-	static void FromString(T& Value, const TCHAR* Buffer)
+	template <typename CharType>
+	static void FromString(T& Value, const CharType* Buffer)
 	{
 		return LexFromString(Value, Buffer);
 	}
