Fixed return type from ToCStr and some comments.

diff --git a/Engine/Source/Runtime/Core/Public/Containers/UnrealString.h.inl b/Engine/Source/Runtime/Core/Public/Containers/UnrealString.h.inl
--- a/Engine/Source/Runtime/Core/Public/Containers/UnrealString.h.inl
+++ b/Engine/Source/Runtime/Core/Public/Containers/UnrealString.h.inl
@@ -2105,7 +2105,7 @@ CORE_API int32 HexToBytes(const UE_STRING_CLASS& HexString, uint8* OutBytes);
  *	Implement custom functionality externally.
  */
 
- /** Covert a string buffer to intrinsic types */
+ /** Convert a string buffer to intrinsic types */
 inline void LexFromString(int8& OutValue, 				const TCHAR* Buffer)	{	OutValue = (int8)FCString::Atoi(Buffer);		}
 inline void LexFromString(int16& OutValue,				const TCHAR* Buffer)	{	OutValue = (int16)FCString::Atoi(Buffer);		}
 inline void LexFromString(int32& OutValue,				const TCHAR* Buffer)	{	OutValue = (int32)FCString::Atoi(Buffer);		}
@@ -2132,9 +2132,9 @@ UE_NODISCARD FORCEINLINE StringType LexToString(const UE_STRING_CLASS& Str)
 }
 
 /**
- * Gets a non-owning TCHAR pointer from a string type.
+ * Gets a non-owning character pointer from a string type.
  *
- * Can be used generically to get a const TCHAR*, when it is not known if the argument is a TCHAR* or an FString:
+ * Can be used generically to get a const char pointer, when it is not known if the argument is a char pointer or a string:
  *
  * template <typename T>
  * void LogValue(const T& Val)
@@ -2142,7 +2142,7 @@ UE_NODISCARD FORCEINLINE StringType LexToString(const UE_STRING_CLASS& Str)
  *     Logf(TEXT("Value: %s"), ToCStr(LexToString(Val)));
  * }
  */
-FORCEINLINE const TCHAR* ToCStr(const UE_STRING_CLASS& Str)
+FORCEINLINE const UE_STRING_CLASS::ElementType* ToCStr(const UE_STRING_CLASS& Str)
 {
 	return *Str;
 }
