Open around getting custom versions which call into a cache with a lock

diff --git a/Engine/Source/Runtime/Core/Private/Serialization/CustomVersion.cpp b/Engine/Source/Runtime/Core/Private/Serialization/CustomVersion.cpp
--- a/Engine/Source/Runtime/Core/Private/Serialization/CustomVersion.cpp
+++ b/Engine/Source/Runtime/Core/Private/Serialization/CustomVersion.cpp
@@ -415,7 +415,10 @@ void FCustomVersionContainer::SetVersionUsingRegistry(FGuid CustomKey, ESetCusto
 			return;
 		}
 
-		TOptional<FCustomVersion> RegisteredVersion = FCurrentCustomVersions::Get(CustomKey);
+		TOptional<FCustomVersion> RegisteredVersion;
+		UE_AUTORTFM_OPEN({
+			RegisteredVersion = FCurrentCustomVersions::Get(CustomKey);
+		});
 		checkf(RegisteredVersion, TEXT("Attempted to set a version that is not registered"));
 
 		Found->Version      = RegisteredVersion->Version;
@@ -423,7 +426,11 @@ void FCustomVersionContainer::SetVersionUsingRegistry(FGuid CustomKey, ESetCusto
 	}
 	else
 	{
-		TOptional<FCustomVersion> RegisteredVersion = FCurrentCustomVersions::Get(CustomKey);
+		TOptional<FCustomVersion> RegisteredVersion;
+		UE_AUTORTFM_OPEN({
+			RegisteredVersion = FCurrentCustomVersions::Get(CustomKey);
+		});
+
 		checkf(RegisteredVersion, TEXT("Attempted to set a version that is not registered"));
 
 		Versions.Emplace(FCustomVersion(CustomKey, RegisteredVersion->Version, RegisteredVersion->FriendlyName));
