Added UE_STRING_CHARTYPE for configuring the string class element type.

diff --git a/Engine/Source/Runtime/Core/Private/Containers/String.cpp b/Engine/Source/Runtime/Core/Private/Containers/String.cpp
--- a/Engine/Source/Runtime/Core/Private/Containers/String.cpp
+++ b/Engine/Source/Runtime/Core/Private/Containers/String.cpp
@@ -5,8 +5,10 @@
 // Include String.cpp.inl's includes before defining the macros, in case the macros 'poison' other headers or there are re-entrant includes.
 #include "Containers/StringIncludes.cpp.inl"
 
-#define UE_STRING_CLASS FString
+#define UE_STRING_CLASS    FString
+#define UE_STRING_CHARTYPE TCHAR
 	#include "Containers/String.cpp.inl"
+#undef UE_STRING_CHARTYPE
 #undef UE_STRING_CLASS
 
 void FTextRange::CalculateLineRangesFromString(const FString& Input, TArray<FTextRange>& LineRanges)
diff --git a/Engine/Source/Runtime/Core/Private/Containers/String.cpp.inl b/Engine/Source/Runtime/Core/Private/Containers/String.cpp.inl
--- a/Engine/Source/Runtime/Core/Private/Containers/String.cpp.inl
+++ b/Engine/Source/Runtime/Core/Private/Containers/String.cpp.inl
@@ -13,6 +13,9 @@
 #ifndef UE_STRING_CLASS
 	#error "String.cpp.inl should only be included after defining UE_STRING_CLASS"
 #endif
+#ifndef UE_STRING_CHARTYPE
+	#error "String.cpp.inl should only be included after defining UE_STRING_CHARTYPE"
+#endif
 
  /* String implementation
  *****************************************************************************/
diff --git a/Engine/Source/Runtime/Core/Private/Containers/StringIncludes.cpp.inl b/Engine/Source/Runtime/Core/Private/Containers/StringIncludes.cpp.inl
--- a/Engine/Source/Runtime/Core/Private/Containers/StringIncludes.cpp.inl
+++ b/Engine/Source/Runtime/Core/Private/Containers/StringIncludes.cpp.inl
@@ -7,6 +7,13 @@
  * String.cpp.inl.                                                                                     *
  *******************************************************************************************************/
 
+#ifdef UE_STRING_CLASS
+	#error "StringIncludes.cpp.inl should not be included after defining UE_STRING_CLASS"
+#endif
+#ifdef UE_STRING_CHARTYPE
+	#error "StringIncludes.cpp.inl should not be included after defining UE_STRING_CHARTYPE"
+#endif
+
 #include "Containers/Array.h"
 #include "Containers/StringConv.h"
 #include "CoreGlobals.h"
diff --git a/Engine/Source/Runtime/Core/Public/Containers/UnrealString.h b/Engine/Source/Runtime/Core/Public/Containers/UnrealString.h
--- a/Engine/Source/Runtime/Core/Public/Containers/UnrealString.h
+++ b/Engine/Source/Runtime/Core/Public/Containers/UnrealString.h
@@ -8,8 +8,10 @@
 // Include UnrealString.h.inl's includes before defining the macros, in case the macros 'poison' other headers or there are re-entrant includes.
 #include "Containers/UnrealStringIncludes.h.inl"
 
-#define UE_STRING_CLASS FString
+#define UE_STRING_CLASS    FString
+#define UE_STRING_CHARTYPE TCHAR
 	#include "Containers/UnrealString.h.inl"
+#undef UE_STRING_CHARTYPE
 #undef UE_STRING_CLASS
 
 /**
diff --git a/Engine/Source/Runtime/Core/Public/Containers/UnrealString.h.inl b/Engine/Source/Runtime/Core/Public/Containers/UnrealString.h.inl
--- a/Engine/Source/Runtime/Core/Public/Containers/UnrealString.h.inl
+++ b/Engine/Source/Runtime/Core/Public/Containers/UnrealString.h.inl
@@ -10,15 +10,22 @@
  * #includes needed to compile this file need to be specified in UnrealStringIncludes.h.inl file rather than here. *
  *******************************************************************************************************************/
 
+#ifndef UE_STRING_CLASS
+	#error "UnrealString.h.inl should only be included after defining UE_STRING_CLASS"
+#endif
+#ifndef UE_STRING_CHARTYPE
+	#error "UnrealString.h.inl should only be included after defining UE_STRING_CHARTYPE"
+#endif
+
 struct PREPROCESSOR_JOIN(UE_STRING_CLASS, FormatArg);
 template<typename InKeyType,typename InValueType,typename SetAllocator ,typename KeyFuncs > class TMap;
 
 typedef TMap<UE_STRING_CLASS, PREPROCESSOR_JOIN(UE_STRING_CLASS, FormatArg)> PREPROCESSOR_JOIN(UE_STRING_CLASS, FormatNamedArguments);
 typedef TArray< PREPROCESSOR_JOIN(UE_STRING_CLASS, FormatArg)> PREPROCESSOR_JOIN(UE_STRING_CLASS, FormatOrderedArguments);
 
-TCHAR*       GetData(UE_STRING_CLASS&);
-const TCHAR* GetData(const UE_STRING_CLASS&);
-int32        GetNum(const UE_STRING_CLASS& String);
+UE_STRING_CHARTYPE*       GetData(UE_STRING_CLASS&);
+const UE_STRING_CHARTYPE* GetData(const UE_STRING_CLASS&);
+int32                     GetNum(const UE_STRING_CLASS& String);
 /**
  * A dynamically sizeable string.
  * @see https://docs.unrealengine.com/latest/INT/Programming/UnrealArchitecture/StringHandling/FString/
@@ -34,7 +41,7 @@ class UE_STRING_CLASS
 {
 public:
 	using AllocatorType = TSizedDefaultAllocator<32>;
-	using ElementType   = TCHAR;
+	using ElementType   = UE_STRING_CHARTYPE;
 
 private:
 	/** Array holding the character data */
@@ -2106,18 +2113,18 @@ CORE_API int32 HexToBytes(const UE_STRING_CLASS& HexString, uint8* OutBytes);
  */
 
  /** Convert a string buffer to intrinsic types */
-inline void LexFromString(int8& OutValue, 				const TCHAR* Buffer)	{	OutValue = (int8)FCString::Atoi(Buffer);		}
-inline void LexFromString(int16& OutValue,				const TCHAR* Buffer)	{	OutValue = (int16)FCString::Atoi(Buffer);		}
-inline void LexFromString(int32& OutValue,				const TCHAR* Buffer)	{	OutValue = (int32)FCString::Atoi(Buffer);		}
-inline void LexFromString(int64& OutValue,				const TCHAR* Buffer)	{	OutValue = FCString::Atoi64(Buffer);	}
-inline void LexFromString(uint8& OutValue,				const TCHAR* Buffer)	{	OutValue = (uint8)FCString::Atoi(Buffer);		}
-inline void LexFromString(uint16& OutValue, 			const TCHAR* Buffer)	{	OutValue = (uint16)FCString::Atoi(Buffer);		}
-inline void LexFromString(uint32& OutValue, 			const TCHAR* Buffer)	{	OutValue = (uint32)FCString::Atoi64(Buffer);	}	//64 because this unsigned and so Atoi might overflow
-inline void LexFromString(uint64& OutValue, 			const TCHAR* Buffer)	{	OutValue = FCString::Strtoui64(Buffer, nullptr, 0); }
-inline void LexFromString(float& OutValue,				const TCHAR* Buffer)	{	OutValue = FCString::Atof(Buffer);		}
-inline void LexFromString(double& OutValue, 			const TCHAR* Buffer)	{	OutValue = FCString::Atod(Buffer);		}
-inline void LexFromString(bool& OutValue, 				const TCHAR* Buffer)	{	OutValue = FCString::ToBool(Buffer);	}
-inline void LexFromString(UE_STRING_CLASS& OutValue, 	const TCHAR* Buffer)	{	OutValue = Buffer;						}
+inline void LexFromString(int8& OutValue, 				const UE_STRING_CHARTYPE* Buffer)	{	OutValue = (int8)TCString<UE_STRING_CHARTYPE>::Atoi(Buffer);		}
+inline void LexFromString(int16& OutValue,				const UE_STRING_CHARTYPE* Buffer)	{	OutValue = (int16)TCString<UE_STRING_CHARTYPE>::Atoi(Buffer);		}
+inline void LexFromString(int32& OutValue,				const UE_STRING_CHARTYPE* Buffer)	{	OutValue = (int32)TCString<UE_STRING_CHARTYPE>::Atoi(Buffer);		}
+inline void LexFromString(int64& OutValue,				const UE_STRING_CHARTYPE* Buffer)	{	OutValue = TCString<UE_STRING_CHARTYPE>::Atoi64(Buffer);	}
+inline void LexFromString(uint8& OutValue,				const UE_STRING_CHARTYPE* Buffer)	{	OutValue = (uint8)TCString<UE_STRING_CHARTYPE>::Atoi(Buffer);		}
+inline void LexFromString(uint16& OutValue, 			const UE_STRING_CHARTYPE* Buffer)	{	OutValue = (uint16)TCString<UE_STRING_CHARTYPE>::Atoi(Buffer);		}
+inline void LexFromString(uint32& OutValue, 			const UE_STRING_CHARTYPE* Buffer)	{	OutValue = (uint32)TCString<UE_STRING_CHARTYPE>::Atoi64(Buffer);	}	//64 because this unsigned and so Atoi might overflow
+inline void LexFromString(uint64& OutValue, 			const UE_STRING_CHARTYPE* Buffer)	{	OutValue = TCString<UE_STRING_CHARTYPE>::Strtoui64(Buffer, nullptr, 0); }
+inline void LexFromString(float& OutValue,				const UE_STRING_CHARTYPE* Buffer)	{	OutValue = TCString<UE_STRING_CHARTYPE>::Atof(Buffer);		}
+inline void LexFromString(double& OutValue, 			const UE_STRING_CHARTYPE* Buffer)	{	OutValue = TCString<UE_STRING_CHARTYPE>::Atod(Buffer);		}
+inline void LexFromString(bool& OutValue, 				const UE_STRING_CHARTYPE* Buffer)	{	OutValue = TCString<UE_STRING_CHARTYPE>::ToBool(Buffer);	}
+inline void LexFromString(UE_STRING_CLASS& OutValue,	const UE_STRING_CHARTYPE* Buffer)	{	OutValue = Buffer;						}
 
 template <typename StringType = FString>
 UE_NODISCARD FORCEINLINE StringType LexToString(UE_STRING_CLASS&& Str)
diff --git a/Engine/Source/Runtime/Core/Public/Containers/UnrealStringIncludes.h.inl b/Engine/Source/Runtime/Core/Public/Containers/UnrealStringIncludes.h.inl
--- a/Engine/Source/Runtime/Core/Public/Containers/UnrealStringIncludes.h.inl
+++ b/Engine/Source/Runtime/Core/Public/Containers/UnrealStringIncludes.h.inl
@@ -7,6 +7,13 @@
  * UnrealString.h.inl.                                                                                 *
  *******************************************************************************************************/
 
+#ifdef UE_STRING_CLASS
+	#error "UnrealStringIncludes.h.inl should not be included after defining UE_STRING_CLASS"
+#endif
+#ifdef UE_STRING_CHARTYPE
+	#error "UnrealStringIncludes.h.inl should not be included after defining UE_STRING_CHARTYPE"
+#endif
+
 #include "CoreTypes.h"
 #include "Misc/VarArgs.h"
 #include "Misc/OutputDevice.h"
