Add an Open around getting the text revision number as there is not big deal not rolling back caching the text or simply getting a cached value

diff --git a/Engine/Source/Runtime/Core/Private/Internationalization/TextHistory.cpp b/Engine/Source/Runtime/Core/Private/Internationalization/TextHistory.cpp
--- a/Engine/Source/Runtime/Core/Private/Internationalization/TextHistory.cpp
+++ b/Engine/Source/Runtime/Core/Private/Internationalization/TextHistory.cpp
@@ -687,7 +687,9 @@ void FTextHistory::UpdateDisplayStringIfOutOfDate()
 	{
 		uint16 CurrentGlobalRevision = 0;
 		uint16 CurrentLocalRevision = 0;
-		FTextLocalizationManager::Get().GetTextRevisions(GetTextId(), CurrentGlobalRevision, CurrentLocalRevision);
+		UE_AUTORTFM_OPEN({
+			FTextLocalizationManager::Get().GetTextRevisions(GetTextId(), CurrentGlobalRevision, CurrentLocalRevision);
+		});
 
 		if (GlobalRevision != CurrentGlobalRevision || LocalRevision != CurrentLocalRevision)
 		{
