Removed FAnsiStringView overload for HexToBytes, because it's already handled by the UTF-8 overload and causes ambiguous overload resolution.

diff --git a/Engine/Source/Runtime/Core/Private/String/HexToBytes.cpp b/Engine/Source/Runtime/Core/Private/String/HexToBytes.cpp
--- a/Engine/Source/Runtime/Core/Private/String/HexToBytes.cpp
+++ b/Engine/Source/Runtime/Core/Private/String/HexToBytes.cpp
@@ -27,11 +27,6 @@ static inline int32 HexToBytesImpl(TStringView<CharType> Hex, uint8* const OutBy
 	return static_cast<int32>(OutPos - OutBytes);
 }
 
-int32 HexToBytes(FAnsiStringView Hex, uint8* OutBytes)
-{
-	return HexToBytesImpl<ANSICHAR>(Hex, OutBytes);
-}
-
 int32 HexToBytes(FWideStringView Hex, uint8* OutBytes)
 {
 	return HexToBytesImpl<WIDECHAR>(Hex, OutBytes);
diff --git a/Engine/Source/Runtime/Core/Public/String/HexToBytes.h b/Engine/Source/Runtime/Core/Public/String/HexToBytes.h
--- a/Engine/Source/Runtime/Core/Public/String/HexToBytes.h
+++ b/Engine/Source/Runtime/Core/Public/String/HexToBytes.h
@@ -16,7 +16,6 @@ namespace UE::String
  *
  * @return Number of bytes written to the output.
  */
-CORE_API int32 HexToBytes(FAnsiStringView Hex, uint8* OutBytes);
 CORE_API int32 HexToBytes(FWideStringView Hex, uint8* OutBytes);
 CORE_API int32 HexToBytes(FUtf8StringView Hex, uint8* OutBytes);
 
