Added UE_STRING_CHARTYPE_IS_TCHAR for switching on that because there's no easy way to compare the UE_STRING_CHARTYPE macro.
Replaced the if constexpr usage in FString's operator<< with the new macro.
Fixed some incorrect usage of CHARTEXT when logging.

diff --git a/Engine/Source/Runtime/Core/Private/Containers/String.cpp b/Engine/Source/Runtime/Core/Private/Containers/String.cpp
--- a/Engine/Source/Runtime/Core/Private/Containers/String.cpp
+++ b/Engine/Source/Runtime/Core/Private/Containers/String.cpp
@@ -5,9 +5,11 @@
 // Include String.cpp.inl's includes before defining the macros, in case the macros 'poison' other headers or there are re-entrant includes.
 #include "Containers/StringIncludes.cpp.inl"
 
-#define UE_STRING_CLASS    FString
-#define UE_STRING_CHARTYPE TCHAR
+#define UE_STRING_CLASS             FString
+#define UE_STRING_CHARTYPE          TCHAR
+#define UE_STRING_CHARTYPE_IS_TCHAR 1
 	#include "Containers/String.cpp.inl"
+#undef UE_STRING_CHARTYPE_IS_TCHAR
 #undef UE_STRING_CHARTYPE
 #undef UE_STRING_CLASS
 
diff --git a/Engine/Source/Runtime/Core/Private/Containers/String.cpp.inl b/Engine/Source/Runtime/Core/Private/Containers/String.cpp.inl
--- a/Engine/Source/Runtime/Core/Private/Containers/String.cpp.inl
+++ b/Engine/Source/Runtime/Core/Private/Containers/String.cpp.inl
@@ -16,6 +16,9 @@
 #ifndef UE_STRING_CHARTYPE
 	#error "String.cpp.inl should only be included after defining UE_STRING_CHARTYPE"
 #endif
+#ifndef UE_STRING_CHARTYPE_IS_TCHAR
+	#error "String.cpp.inl should only be included after defining UE_STRING_CHARTYPE_IS_TCHAR"
+#endif
 
  /* String implementation
  *****************************************************************************/
@@ -1713,138 +1716,135 @@ FArchive& operator<<( FArchive& Ar, UE_STRING_CLASS& A )
 {
 	using ElementType = UE_STRING_CLASS::ElementType;
 
-	if constexpr (std::is_same_v<ElementType, TCHAR>)
+#if UE_STRING_CHARTYPE_IS_TCHAR
+	// > 0 for ANSICHAR, < 0 for UTF16CHAR serialization
+	static_assert(sizeof(UTF16CHAR) == sizeof(UCS2CHAR), "UTF16CHAR and UCS2CHAR are assumed to be the same size!");
+
+	if (Ar.IsLoading())
 	{
-		// > 0 for ANSICHAR, < 0 for UTF16CHAR serialization
-		static_assert(sizeof(UTF16CHAR) == sizeof(UCS2CHAR), "UTF16CHAR and UCS2CHAR are assumed to be the same size!");
+		int32 SaveNum = 0;
+		Ar << SaveNum;
 
-		if (Ar.IsLoading())
+		bool bLoadUnicodeChar = SaveNum < 0;
+		if (bLoadUnicodeChar)
 		{
-			int32 SaveNum = 0;
-			Ar << SaveNum;
-
-			bool bLoadUnicodeChar = SaveNum < 0;
-			if (bLoadUnicodeChar)
-			{
-				// If SaveNum cannot be negated due to integer overflow, Ar is corrupted.
-				if (SaveNum == MIN_int32)
-				{
-					Ar.SetCriticalError();
-					UE_LOG(LogCore, Error, CHARTEXT(ElementType, "Archive is corrupted"));
-					return Ar;
-				}
-
-				SaveNum = -SaveNum;
-			}
-
-			int64 MaxSerializeSize = Ar.GetMaxSerializeSize();
-			// Protect against network packets allocating too much memory
-			if ((MaxSerializeSize > 0) && (SaveNum > MaxSerializeSize))
+			// If SaveNum cannot be negated due to integer overflow, Ar is corrupted.
+			if (SaveNum == MIN_int32)
 			{
 				Ar.SetCriticalError();
-				UE_LOG(LogCore, Error, CHARTEXT(ElementType, "String is too large (Size: %i, Max: %i)"), SaveNum, MaxSerializeSize);
+				UE_LOG(LogCore, Error, TEXT("Archive is corrupted"));
 				return Ar;
 			}
 
-			// Resize the array only if it passes the above tests to prevent rogue packets from crashing
-			A.Data.Empty(SaveNum);
-			A.Data.AddUninitialized(SaveNum);
+			SaveNum = -SaveNum;
+		}
+
+		int64 MaxSerializeSize = Ar.GetMaxSerializeSize();
+		// Protect against network packets allocating too much memory
+		if ((MaxSerializeSize > 0) && (SaveNum > MaxSerializeSize))
+		{
+			Ar.SetCriticalError();
+			UE_LOG(LogCore, Error, TEXT("String is too large (Size: %i, Max: %i)"), SaveNum, MaxSerializeSize);
+			return Ar;
+		}
+
+		// Resize the array only if it passes the above tests to prevent rogue packets from crashing
+		A.Data.Empty(SaveNum);
+		A.Data.AddUninitialized(SaveNum);
 
-			if (SaveNum)
+		if (SaveNum)
+		{
+			if (bLoadUnicodeChar)
 			{
-				if (bLoadUnicodeChar)
+				// read in the unicode string
+				auto Passthru = StringMemoryPassthru<UCS2CHAR>(A.Data.GetData(), SaveNum, SaveNum);
+				Ar.Serialize(Passthru.Get(), SaveNum * sizeof(UCS2CHAR));
+				if (Ar.IsByteSwapping())
 				{
-					// read in the unicode string
-					auto Passthru = StringMemoryPassthru<UCS2CHAR>(A.Data.GetData(), SaveNum, SaveNum);
-					Ar.Serialize(Passthru.Get(), SaveNum * sizeof(UCS2CHAR));
-					if (Ar.IsByteSwapping())
+					for (int32 CharIndex = 0; CharIndex < SaveNum; ++CharIndex)
 					{
-						for (int32 CharIndex = 0; CharIndex < SaveNum; ++CharIndex)
-						{
-							Passthru.Get()[CharIndex] = ByteSwap(Passthru.Get()[CharIndex]);
-						}
-					}
-					// Ensure the string has a null terminator
-					Passthru.Get()[SaveNum - 1] = '\0';
-					Passthru.Apply();
-
-					// Inline combine any surrogate pairs in the data when loading into a UTF-32 string
-					StringConv::InlineCombineSurrogates(A);
-
-					// Since Microsoft's vsnwprintf implementation raises an invalid parameter warning
-					// with a character of 0xffff, scan for it and terminate the string there.
-					// 0xffff isn't an actual Unicode character anyway.
-					int Index = 0;
-					if (A.FindChar(0xffff, Index))
-					{
-						A[Index] = CHARTEXT(ElementType, '\0');
-						A.TrimToNullTerminator();
+						Passthru.Get()[CharIndex] = ByteSwap(Passthru.Get()[CharIndex]);
 					}
 				}
-				else
+				// Ensure the string has a null terminator
+				Passthru.Get()[SaveNum - 1] = '\0';
+				Passthru.Apply();
+
+				// Inline combine any surrogate pairs in the data when loading into a UTF-32 string
+				StringConv::InlineCombineSurrogates(A);
+
+				// Since Microsoft's vsnwprintf implementation raises an invalid parameter warning
+				// with a character of 0xffff, scan for it and terminate the string there.
+				// 0xffff isn't an actual Unicode character anyway.
+				int Index = 0;
+				if (A.FindChar(0xffff, Index))
 				{
-					auto Passthru = StringMemoryPassthru<ANSICHAR>(A.Data.GetData(), SaveNum, SaveNum);
-					Ar.Serialize(Passthru.Get(), SaveNum * sizeof(ANSICHAR));
-					// Ensure the string has a null terminator
-					Passthru.Get()[SaveNum - 1] = '\0';
-					Passthru.Apply();
+					A[Index] = CHARTEXT(ElementType, '\0');
+					A.TrimToNullTerminator();
 				}
+			}
+			else
+			{
+				auto Passthru = StringMemoryPassthru<ANSICHAR>(A.Data.GetData(), SaveNum, SaveNum);
+				Ar.Serialize(Passthru.Get(), SaveNum * sizeof(ANSICHAR));
+				// Ensure the string has a null terminator
+				Passthru.Get()[SaveNum - 1] = '\0';
+				Passthru.Apply();
+			}
 
-				// Throw away empty string.
-				if (SaveNum == 1)
-				{
-					A.Data.Empty();
-				}
+			// Throw away empty string.
+			if (SaveNum == 1)
+			{
+				A.Data.Empty();
 			}
 		}
-		else
-		{
-			A.Data.CountBytes(Ar);
+	}
+	else
+	{
+		A.Data.CountBytes(Ar);
 
-			const bool bSaveUnicodeChar = Ar.IsForcingUnicode() || !TCString<ElementType>::IsPureAnsi(*A);
-			if (bSaveUnicodeChar)
-			{
-				// This preprocessor block should not be necessary when the StringCast above understands UTF16CHAR.
-				FTCHARToUTF16 UTF16String(*A, A.Len() + 1); // include the null terminator
-				int32 Num = UTF16String.Length() + 1; // include the null terminator
+		const bool bSaveUnicodeChar = Ar.IsForcingUnicode() || !TCString<ElementType>::IsPureAnsi(*A);
+		if (bSaveUnicodeChar)
+		{
+			// This preprocessor block should not be necessary when the StringCast above understands UTF16CHAR.
+			FTCHARToUTF16 UTF16String(*A, A.Len() + 1); // include the null terminator
+			int32 Num = UTF16String.Length() + 1; // include the null terminator
 
-				int32 SaveNum = -Num;
-				Ar << SaveNum;
+			int32 SaveNum = -Num;
+			Ar << SaveNum;
 
-				if (Num)
+			if (Num)
+			{
+				if (!Ar.IsByteSwapping())
 				{
-					if (!Ar.IsByteSwapping())
-					{
-						Ar.Serialize((void*)UTF16String.Get(), sizeof(UTF16CHAR) * Num);
-					}
-					else
+					Ar.Serialize((void*)UTF16String.Get(), sizeof(UTF16CHAR) * Num);
+				}
+				else
+				{
+					TArray<UTF16CHAR> Swapped(UTF16String.Get(), Num);
+					for (int32 CharIndex = 0; CharIndex < Num; ++CharIndex)
 					{
-						TArray<UTF16CHAR> Swapped(UTF16String.Get(), Num);
-						for (int32 CharIndex = 0; CharIndex < Num; ++CharIndex)
-						{
-							Swapped[CharIndex] = ByteSwap(Swapped[CharIndex]);
-						}
-						Ar.Serialize((void*)Swapped.GetData(), sizeof(UTF16CHAR) * Num);
+						Swapped[CharIndex] = ByteSwap(Swapped[CharIndex]);
 					}
+					Ar.Serialize((void*)Swapped.GetData(), sizeof(UTF16CHAR) * Num);
 				}
 			}
-			else
-			{
-				int32 Num = A.Data.Num();
-				Ar << Num;
+		}
+		else
+		{
+			int32 Num = A.Data.Num();
+			Ar << Num;
 
-				if (Num)
-				{
-					Ar.Serialize((void*)StringCast<ANSICHAR>(A.Data.GetData(), Num).Get(), sizeof(ANSICHAR) * Num);
-				}
+			if (Num)
+			{
+				Ar.Serialize((void*)StringCast<ANSICHAR>(A.Data.GetData(), Num).Get(), sizeof(ANSICHAR) * Num);
 			}
 		}
 	}
-	else
-	{
-		// Can just serialize as UTF-8 always
-		check(false); // TODO
-	}
+#else
+	// Can just serialize as UTF-8 always
+	check(false); // TODO
+#endif
 
 	return Ar;
 }
diff --git a/Engine/Source/Runtime/Core/Private/Containers/StringIncludes.cpp.inl b/Engine/Source/Runtime/Core/Private/Containers/StringIncludes.cpp.inl
--- a/Engine/Source/Runtime/Core/Private/Containers/StringIncludes.cpp.inl
+++ b/Engine/Source/Runtime/Core/Private/Containers/StringIncludes.cpp.inl
@@ -13,6 +13,9 @@
 #ifdef UE_STRING_CHARTYPE
 	#error "StringIncludes.cpp.inl should not be included after defining UE_STRING_CHARTYPE"
 #endif
+#ifdef UE_STRING_CHARTYPE_IS_TCHAR
+	#error "StringIncludes.cpp.inl should not be included after defining UE_STRING_CHARTYPE_IS_TCHAR"
+#endif
 
 #include "Containers/Array.h"
 #include "Containers/StringConv.h"
diff --git a/Engine/Source/Runtime/Core/Public/Containers/UnrealString.h b/Engine/Source/Runtime/Core/Public/Containers/UnrealString.h
--- a/Engine/Source/Runtime/Core/Public/Containers/UnrealString.h
+++ b/Engine/Source/Runtime/Core/Public/Containers/UnrealString.h
@@ -8,9 +8,11 @@
 // Include UnrealString.h.inl's includes before defining the macros, in case the macros 'poison' other headers or there are re-entrant includes.
 #include "Containers/UnrealStringIncludes.h.inl"
 
-#define UE_STRING_CLASS    FString
-#define UE_STRING_CHARTYPE TCHAR
+#define UE_STRING_CLASS             FString
+#define UE_STRING_CHARTYPE          TCHAR
+#define UE_STRING_CHARTYPE_IS_TCHAR 1
 	#include "Containers/UnrealString.h.inl"
+#undef UE_STRING_CHARTYPE_IS_TCHAR
 #undef UE_STRING_CHARTYPE
 #undef UE_STRING_CLASS
 
diff --git a/Engine/Source/Runtime/Core/Public/Containers/UnrealString.h.inl b/Engine/Source/Runtime/Core/Public/Containers/UnrealString.h.inl
--- a/Engine/Source/Runtime/Core/Public/Containers/UnrealString.h.inl
+++ b/Engine/Source/Runtime/Core/Public/Containers/UnrealString.h.inl
@@ -16,6 +16,9 @@
 #ifndef UE_STRING_CHARTYPE
 	#error "UnrealString.h.inl should only be included after defining UE_STRING_CHARTYPE"
 #endif
+#ifndef UE_STRING_CHARTYPE_IS_TCHAR
+	#error "UnrealString.h.inl should only be included after defining UE_STRING_CHARTYPE_IS_TCHAR"
+#endif
 
 struct PREPROCESSOR_JOIN(UE_STRING_CLASS, FormatArg);
 template<typename InKeyType,typename InValueType,typename SetAllocator ,typename KeyFuncs > class TMap;
diff --git a/Engine/Source/Runtime/Core/Public/Containers/UnrealStringIncludes.h.inl b/Engine/Source/Runtime/Core/Public/Containers/UnrealStringIncludes.h.inl
--- a/Engine/Source/Runtime/Core/Public/Containers/UnrealStringIncludes.h.inl
+++ b/Engine/Source/Runtime/Core/Public/Containers/UnrealStringIncludes.h.inl
@@ -13,6 +13,9 @@
 #ifdef UE_STRING_CHARTYPE
 	#error "UnrealStringIncludes.h.inl should not be included after defining UE_STRING_CHARTYPE"
 #endif
+#ifdef UE_STRING_CHARTYPE_IS_TCHAR
+	#error "UnrealStringIncludes.h.inl should not be included after defining UE_STRING_CHARTYPE_IS_TCHAR"
+#endif
 
 #include "CoreTypes.h"
 #include "Misc/VarArgs.h"
