Moved reusable StripNegativeZero out of string template code.

diff --git a/Engine/Source/Runtime/Core/Private/Containers/String.cpp b/Engine/Source/Runtime/Core/Private/Containers/String.cpp
--- a/Engine/Source/Runtime/Core/Private/Containers/String.cpp
+++ b/Engine/Source/Runtime/Core/Private/Containers/String.cpp
@@ -50,3 +50,16 @@ void StringConv::InlineCombineSurrogates(FString& Str)
 {
 	InlineCombineSurrogates_Array(Str.GetCharArray());
 }
+
+namespace UE::Core::Private
+{
+	UE_DISABLE_OPTIMIZATION_SHIP
+	void StripNegativeZero(double& InFloat)
+	{
+		// This works for translating a negative zero into a positive zero,
+		// but if optimizations are enabled when compiling with -ffast-math
+		// or /fp:fast, the compiler can strip it out.
+		InFloat += 0.0f;
+	}
+	UE_ENABLE_OPTIMIZATION_SHIP
+}
diff --git a/Engine/Source/Runtime/Core/Private/Containers/String.cpp.inl b/Engine/Source/Runtime/Core/Private/Containers/String.cpp.inl
--- a/Engine/Source/Runtime/Core/Private/Containers/String.cpp.inl
+++ b/Engine/Source/Runtime/Core/Private/Containers/String.cpp.inl
@@ -1148,20 +1148,10 @@ bool UE_STRING_CLASS::ToHexBlob( const UE_STRING_CLASS& Source, uint8* DestBuffe
 	return false;
 }
 
-UE_DISABLE_OPTIMIZATION_SHIP
-void StripNegativeZero(double& InFloat)
-{
-	// This works for translating a negative zero into a positive zero,
-	// but if optimizations are enabled when compiling with -ffast-math
-	// or /fp:fast, the compiler can strip it out.
-	InFloat += 0.0f;
-}
-UE_ENABLE_OPTIMIZATION_SHIP
-
 UE_STRING_CLASS UE_STRING_CLASS::SanitizeFloat( double InFloat, const int32 InMinFractionalDigits )
 {
 	// Avoids negative zero
-	StripNegativeZero(InFloat);
+	UE::Core::Private::StripNegativeZero(InFloat);
 
 	// First create the string
 	UE_STRING_CLASS TempString = UE_STRING_CLASS::Printf(CHARTEXT(ElementType, "%f"), InFloat);
@@ -1717,133 +1707,136 @@ FArchive& operator<<( FArchive& Ar, UE_STRING_CLASS& A )
 	using ElementType = UE_STRING_CLASS::ElementType;
 
 #if UE_STRING_CHARTYPE_IS_TCHAR
-	// > 0 for ANSICHAR, < 0 for UTF16CHAR serialization
-	static_assert(sizeof(UTF16CHAR) == sizeof(UCS2CHAR), "UTF16CHAR and UCS2CHAR are assumed to be the same size!");
-
-	if (Ar.IsLoading())
-	{
-		int32 SaveNum = 0;
-		Ar << SaveNum;
+		// > 0 for ANSICHAR, < 0 for UTF16CHAR serialization
+		static_assert(sizeof(UTF16CHAR) == sizeof(UCS2CHAR), "UTF16CHAR and UCS2CHAR are assumed to be the same size!");
 
-		bool bLoadUnicodeChar = SaveNum < 0;
-		if (bLoadUnicodeChar)
+		if (Ar.IsLoading())
 		{
-			// If SaveNum cannot be negated due to integer overflow, Ar is corrupted.
-			if (SaveNum == MIN_int32)
+			int32 SaveNum = 0;
+			Ar << SaveNum;
+
+			bool bLoadUnicodeChar = SaveNum < 0;
+			if (bLoadUnicodeChar)
 			{
-				Ar.SetCriticalError();
+				// If SaveNum cannot be negated due to integer overflow, Ar is corrupted.
+				if (SaveNum == MIN_int32)
+				{
+					Ar.SetCriticalError();
 				UE_LOG(LogCore, Error, TEXT("Archive is corrupted"));
-				return Ar;
-			}
+					return Ar;
+				}
 
-			SaveNum = -SaveNum;
-		}
+				SaveNum = -SaveNum;
+			}
 
-		int64 MaxSerializeSize = Ar.GetMaxSerializeSize();
-		// Protect against network packets allocating too much memory
-		if ((MaxSerializeSize > 0) && (SaveNum > MaxSerializeSize))
-		{
-			Ar.SetCriticalError();
+			int64 MaxSerializeSize = Ar.GetMaxSerializeSize();
+			// Protect against network packets allocating too much memory
+			if ((MaxSerializeSize > 0) && (SaveNum > MaxSerializeSize))
+			{
+				Ar.SetCriticalError();
 			UE_LOG(LogCore, Error, TEXT("String is too large (Size: %i, Max: %i)"), SaveNum, MaxSerializeSize);
-			return Ar;
-		}
+				return Ar;
+			}
 
-		// Resize the array only if it passes the above tests to prevent rogue packets from crashing
-		A.Data.Empty(SaveNum);
-		A.Data.AddUninitialized(SaveNum);
+			// Resize the array only if it passes the above tests to prevent rogue packets from crashing
+			A.Data.Empty(SaveNum);
+			A.Data.AddUninitialized(SaveNum);
 
-		if (SaveNum)
-		{
-			if (bLoadUnicodeChar)
+			if (SaveNum)
 			{
-				// read in the unicode string
-				auto Passthru = StringMemoryPassthru<UCS2CHAR>(A.Data.GetData(), SaveNum, SaveNum);
-				Ar.Serialize(Passthru.Get(), SaveNum * sizeof(UCS2CHAR));
-				if (Ar.IsByteSwapping())
+				if (bLoadUnicodeChar)
 				{
-					for (int32 CharIndex = 0; CharIndex < SaveNum; ++CharIndex)
+					// read in the unicode string
+					auto Passthru = StringMemoryPassthru<UCS2CHAR>(A.Data.GetData(), SaveNum, SaveNum);
+					Ar.Serialize(Passthru.Get(), SaveNum * sizeof(UCS2CHAR));
+					if (Ar.IsByteSwapping())
 					{
-						Passthru.Get()[CharIndex] = ByteSwap(Passthru.Get()[CharIndex]);
+						for (int32 CharIndex = 0; CharIndex < SaveNum; ++CharIndex)
+						{
+							Passthru.Get()[CharIndex] = ByteSwap(Passthru.Get()[CharIndex]);
+						}
 					}
+					// Ensure the string has a null terminator
+					Passthru.Get()[SaveNum - 1] = '\0';
+					Passthru.Apply();
+
+				if constexpr (std::is_same_v<UE_STRING_CLASS::ElementType, TCHAR>)
+				{
+					// Inline combine any surrogate pairs in the data when loading into a UTF-32 string
+					StringConv::InlineCombineSurrogates(A);
 				}
-				// Ensure the string has a null terminator
-				Passthru.Get()[SaveNum - 1] = '\0';
-				Passthru.Apply();
-
-				// Inline combine any surrogate pairs in the data when loading into a UTF-32 string
-				StringConv::InlineCombineSurrogates(A);
-
-				// Since Microsoft's vsnwprintf implementation raises an invalid parameter warning
-				// with a character of 0xffff, scan for it and terminate the string there.
-				// 0xffff isn't an actual Unicode character anyway.
-				int Index = 0;
-				if (A.FindChar(0xffff, Index))
+
+					// Since Microsoft's vsnwprintf implementation raises an invalid parameter warning
+					// with a character of 0xffff, scan for it and terminate the string there.
+					// 0xffff isn't an actual Unicode character anyway.
+					int Index = 0;
+					if (A.FindChar(0xffff, Index))
+					{
+						A[Index] = CHARTEXT(ElementType, '\0');
+						A.TrimToNullTerminator();
+					}
+				}
+				else
 				{
-					A[Index] = CHARTEXT(ElementType, '\0');
-					A.TrimToNullTerminator();
+					auto Passthru = StringMemoryPassthru<ANSICHAR>(A.Data.GetData(), SaveNum, SaveNum);
+					Ar.Serialize(Passthru.Get(), SaveNum * sizeof(ANSICHAR));
+					// Ensure the string has a null terminator
+					Passthru.Get()[SaveNum - 1] = '\0';
+					Passthru.Apply();
 				}
-			}
-			else
-			{
-				auto Passthru = StringMemoryPassthru<ANSICHAR>(A.Data.GetData(), SaveNum, SaveNum);
-				Ar.Serialize(Passthru.Get(), SaveNum * sizeof(ANSICHAR));
-				// Ensure the string has a null terminator
-				Passthru.Get()[SaveNum - 1] = '\0';
-				Passthru.Apply();
-			}
 
-			// Throw away empty string.
-			if (SaveNum == 1)
-			{
-				A.Data.Empty();
+				// Throw away empty string.
+				if (SaveNum == 1)
+				{
+					A.Data.Empty();
+				}
 			}
 		}
-	}
-	else
-	{
-		A.Data.CountBytes(Ar);
-
-		const bool bSaveUnicodeChar = Ar.IsForcingUnicode() || !TCString<ElementType>::IsPureAnsi(*A);
-		if (bSaveUnicodeChar)
+		else
 		{
-			// This preprocessor block should not be necessary when the StringCast above understands UTF16CHAR.
-			FTCHARToUTF16 UTF16String(*A, A.Len() + 1); // include the null terminator
-			int32 Num = UTF16String.Length() + 1; // include the null terminator
+			A.Data.CountBytes(Ar);
 
-			int32 SaveNum = -Num;
-			Ar << SaveNum;
-
-			if (Num)
+			const bool bSaveUnicodeChar = Ar.IsForcingUnicode() || !TCString<ElementType>::IsPureAnsi(*A);
+			if (bSaveUnicodeChar)
 			{
-				if (!Ar.IsByteSwapping())
-				{
-					Ar.Serialize((void*)UTF16String.Get(), sizeof(UTF16CHAR) * Num);
-				}
-				else
+				// This preprocessor block should not be necessary when the StringCast above understands UTF16CHAR.
+				FTCHARToUTF16 UTF16String(*A, A.Len() + 1); // include the null terminator
+				int32 Num = UTF16String.Length() + 1; // include the null terminator
+
+				int32 SaveNum = -Num;
+				Ar << SaveNum;
+
+				if (Num)
 				{
-					TArray<UTF16CHAR> Swapped(UTF16String.Get(), Num);
-					for (int32 CharIndex = 0; CharIndex < Num; ++CharIndex)
+					if (!Ar.IsByteSwapping())
 					{
-						Swapped[CharIndex] = ByteSwap(Swapped[CharIndex]);
+						Ar.Serialize((void*)UTF16String.Get(), sizeof(UTF16CHAR) * Num);
+					}
+					else
+					{
+						TArray<UTF16CHAR> Swapped(UTF16String.Get(), Num);
+						for (int32 CharIndex = 0; CharIndex < Num; ++CharIndex)
+						{
+							Swapped[CharIndex] = ByteSwap(Swapped[CharIndex]);
+						}
+						Ar.Serialize((void*)Swapped.GetData(), sizeof(UTF16CHAR) * Num);
 					}
-					Ar.Serialize((void*)Swapped.GetData(), sizeof(UTF16CHAR) * Num);
 				}
 			}
-		}
-		else
-		{
-			int32 Num = A.Data.Num();
-			Ar << Num;
-
-			if (Num)
+			else
 			{
-				Ar.Serialize((void*)StringCast<ANSICHAR>(A.Data.GetData(), Num).Get(), sizeof(ANSICHAR) * Num);
+				int32 Num = A.Data.Num();
+				Ar << Num;
+
+				if (Num)
+				{
+					Ar.Serialize((void*)StringCast<ANSICHAR>(A.Data.GetData(), Num).Get(), sizeof(ANSICHAR) * Num);
+				}
 			}
 		}
-	}
 #else
-	// Can just serialize as UTF-8 always
-	check(false); // TODO
+		// Can just serialize as UTF-8 always
+		check(false); // TODO
 #endif
 
 	return Ar;
diff --git a/Engine/Source/Runtime/Core/Public/Containers/UnrealString.h b/Engine/Source/Runtime/Core/Public/Containers/UnrealString.h
--- a/Engine/Source/Runtime/Core/Public/Containers/UnrealString.h
+++ b/Engine/Source/Runtime/Core/Public/Containers/UnrealString.h
@@ -449,3 +449,8 @@ struct FTextRange
 	int32 BeginIndex;
 	int32 EndIndex;
 };
+
+namespace UE::Core::Private
+{
+	CORE_API void StripNegativeZero(double& InFloat);
+}
