Add failsafe for memory tracing on Windows

Adds an extra check when looking for functions on callstack tracing on Windows. This issue has been observed on Clang compiled targets. Thanks to [REDACTED_PERSON] for PR.

diff --git a/Engine/Source/Runtime/Core/Private/ProfilingDebugging/Microsoft/WindowsCallstackTrace.cpp b/Engine/Source/Runtime/Core/Private/ProfilingDebugging/Microsoft/WindowsCallstackTrace.cpp
--- a/Engine/Source/Runtime/Core/Private/ProfilingDebugging/Microsoft/WindowsCallstackTrace.cpp
+++ b/Engine/Source/Runtime/Core/Private/ProfilingDebugging/Microsoft/WindowsCallstackTrace.cpp
@@ -552,6 +552,10 @@ const FBacktracer::FFunction* FBacktracer::LookupFunction(UPTRINT Address, FLook
 	uint32 FuncId = uint32(Address - IdToAddress(Module->Id)); 
 	TArrayView<FFunction> FuncsView(Module->Functions, Module->NumFunctions);
 	uint32 Index = Algo::UpperBound(FuncsView, FuncId, IdPredicate);
+	if (Index == 0)
+	{
+		return nullptr;
+	}
 
 	const FFunction* Function = Module->Functions + (Index - 1);
 #if BACKTRACE_DBGLVL >= 2
