Don't try to auto-connect when already connected

Trace auto-connecting did not check current state before trying to connect, causing an error to be emitted.

diff --git a/Engine/Source/Runtime/Core/Private/ProfilingDebugging/TraceAuxiliary.cpp b/Engine/Source/Runtime/Core/Private/ProfilingDebugging/TraceAuxiliary.cpp
--- a/Engine/Source/Runtime/Core/Private/ProfilingDebugging/TraceAuxiliary.cpp
+++ b/Engine/Source/Runtime/Core/Private/ProfilingDebugging/TraceAuxiliary.cpp
@@ -1801,7 +1801,7 @@ void FTraceAuxiliary::TryAutoConnect()
 {
 #if UE_TRACE_ENABLED
 #if PLATFORM_WINDOWS
-	if (GTraceAutoStart)
+	if (GTraceAutoStart && !IsConnected())
 	{
 		// If we can detect a named event it means UnrealInsights (Browser Mode) is running.
 		// In this case, we try to auto-connect with the Trace Server.
