Added UTF-8 serialization for non-TCHAR string classes.

diff --git a/Engine/Source/Runtime/Core/Private/Containers/String.cpp.inl b/Engine/Source/Runtime/Core/Private/Containers/String.cpp.inl
--- a/Engine/Source/Runtime/Core/Private/Containers/String.cpp.inl
+++ b/Engine/Source/Runtime/Core/Private/Containers/String.cpp.inl
@@ -1706,15 +1706,18 @@ FArchive& operator<<( FArchive& Ar, UE_STRING_CLASS& A )
 {
 	using ElementType = UE_STRING_CLASS::ElementType;
 
-#if UE_STRING_CHARTYPE_IS_TCHAR
+	#if UE_STRING_CHARTYPE_IS_TCHAR
 		// > 0 for ANSICHAR, < 0 for UTF16CHAR serialization
 		static_assert(sizeof(UTF16CHAR) == sizeof(UCS2CHAR), "UTF16CHAR and UCS2CHAR are assumed to be the same size!");
+	#endif
 
-		if (Ar.IsLoading())
-		{
-			int32 SaveNum = 0;
-			Ar << SaveNum;
+	if (Ar.IsLoading())
+	{
+		int32 SaveNum = 0;
+		Ar << SaveNum;
 
+		// Validate loaded num to ensure the archive is in a good state
+		#if UE_STRING_CHARTYPE_IS_TCHAR
 			bool bLoadUnicodeChar = SaveNum < 0;
 			if (bLoadUnicodeChar)
 			{
@@ -1722,22 +1725,31 @@ FArchive& operator<<( FArchive& Ar, UE_STRING_CLASS& A )
 				if (SaveNum == MIN_int32)
 				{
 					Ar.SetCriticalError();
-				UE_LOG(LogCore, Error, TEXT("Archive is corrupted"));
+					UE_LOG(LogCore, Error, TEXT("Archive is corrupted"));
 					return Ar;
 				}
 
 				SaveNum = -SaveNum;
 			}
-
-			int64 MaxSerializeSize = Ar.GetMaxSerializeSize();
-			// Protect against network packets allocating too much memory
-			if ((MaxSerializeSize > 0) && (SaveNum > MaxSerializeSize))
+		#else
+			if (SaveNum < 0)
 			{
 				Ar.SetCriticalError();
-			UE_LOG(LogCore, Error, TEXT("String is too large (Size: %i, Max: %i)"), SaveNum, MaxSerializeSize);
+				UE_LOG(LogCore, Error, TEXT("Archive is corrupted"));
 				return Ar;
 			}
+		#endif
+
+		int64 MaxSerializeSize = Ar.GetMaxSerializeSize();
+		// Protect against network packets allocating too much memory
+		if ((MaxSerializeSize > 0) && (SaveNum > MaxSerializeSize))
+		{
+			Ar.SetCriticalError();
+			UE_LOG(LogCore, Error, TEXT("String is too large (Size: %i, Max: %i)"), SaveNum, MaxSerializeSize);
+			return Ar;
+		}
 
+		#if UE_STRING_CHARTYPE_IS_TCHAR
 			// Resize the array only if it passes the above tests to prevent rogue packets from crashing
 			A.Data.Empty(SaveNum);
 			A.Data.AddUninitialized(SaveNum);
@@ -1760,11 +1772,8 @@ FArchive& operator<<( FArchive& Ar, UE_STRING_CLASS& A )
 					Passthru.Get()[SaveNum - 1] = '\0';
 					Passthru.Apply();
 
-				if constexpr (std::is_same_v<UE_STRING_CLASS::ElementType, TCHAR>)
-				{
 					// Inline combine any surrogate pairs in the data when loading into a UTF-32 string
 					StringConv::InlineCombineSurrogates(A);
-				}
 
 					// Since Microsoft's vsnwprintf implementation raises an invalid parameter warning
 					// with a character of 0xffff, scan for it and terminate the string there.
@@ -1791,11 +1800,31 @@ FArchive& operator<<( FArchive& Ar, UE_STRING_CLASS& A )
 					A.Data.Empty();
 				}
 			}
-		}
-		else
-		{
-			A.Data.CountBytes(Ar);
+		#else
+			if (SaveNum)
+			{
+				A.Data.Empty(SaveNum + 1);
+				A.Data.AddUninitialized(SaveNum + 1);
 
+				auto Passthru = StringMemoryPassthru<ElementType>(A.Data.GetData(), SaveNum, SaveNum);
+				Ar.Serialize(Passthru.Get(), SaveNum * sizeof(ElementType));
+				// Ensure the string has a null terminator
+				Passthru.Get()[SaveNum] = UTF8TEXT('\0');
+				Passthru.Apply();
+
+				// We don't need to throw away empty strings here, unlike above, because we never saved a null terminator
+			}
+			else
+			{
+				A.Empty();
+			}
+		#endif
+	}
+	else
+	{
+		A.Data.CountBytes(Ar);
+
+		#if UE_STRING_CHARTYPE_IS_TCHAR
 			const bool bSaveUnicodeChar = Ar.IsForcingUnicode() || !TCString<ElementType>::IsPureAnsi(*A);
 			if (bSaveUnicodeChar)
 			{
@@ -1833,11 +1862,18 @@ FArchive& operator<<( FArchive& Ar, UE_STRING_CLASS& A )
 					Ar.Serialize((void*)StringCast<ANSICHAR>(A.Data.GetData(), Num).Get(), sizeof(ANSICHAR) * Num);
 				}
 			}
-		}
-#else
-		// Can just serialize as UTF-8 always
-		check(false); // TODO
-#endif
+		#else
+			// Unlike the TCHAR case, we don't bother to save the null terminator
+			int32 SaveNum = A.Len();
+			Ar << SaveNum;
+
+			if (SaveNum)
+			{
+				auto CompactString = StringCast<UTF8CHAR>(A.Data.GetData(), SaveNum);
+				Ar.Serialize((void*)CompactString.Get(), sizeof(UTF8CHAR) * CompactString.Length());
+			}
+		#endif
+	}
 
 	return Ar;
 }
