Add two OPENS around FName caches which create, store and own the FName memory while returning ID/numbers which we can leak not cleaning these up on a failed transaction

diff --git a/Engine/Source/Runtime/Core/Private/UObject/UnrealNames.cpp b/Engine/Source/Runtime/Core/Private/UObject/UnrealNames.cpp
--- a/Engine/Source/Runtime/Core/Private/UObject/UnrealNames.cpp
+++ b/Engine/Source/Runtime/Core/Private/UObject/UnrealNames.cpp
@@ -2853,7 +2853,10 @@ struct FNameHelper
 		FNamePool& Pool = GetNamePool();
 		if (FindType == FNAME_Add)
 		{
-			FNameEntryId DisplayId = Pool.StoreWithNumber(BaseIds, InternalNumber);
+			FNameEntryId DisplayId;
+			UE_AUTORTFM_OPEN({
+				DisplayId = Pool.StoreWithNumber(BaseIds, InternalNumber);
+			});
 			return FinalConstruct(FNameEntryIds{ ResolveComparisonId(DisplayId), DisplayId });
 		}
 		else
@@ -2958,7 +2961,10 @@ private:
 
 	static FName MakeInternal(FNameStringView View, EFindName FindType, int32 InternalNumber)
 	{
-		FNameEntryIds Ids = FindOrStoreString(View, FindType);
+		FNameEntryIds Ids;
+		UE_AUTORTFM_OPEN({
+			Ids = FindOrStoreString(View, FindType);
+		});
 #if UE_FNAME_OUTLINE_NUMBER
 		if (FindType == FNAME_Find && !Ids.DisplayId)
 		{
@@ -4089,7 +4095,7 @@ FName FName::CreateNumberedName(FNameEntryId ComparisonId, FNameEntryId DisplayI
 #endif
 	checkName(ResolveEntry(ComparisonId)->IsNumbered() == false);
 
-	return FNameHelper::MakeWithNumber(FNameEntryIds{ComparisonId, DisplayId}, FNAME_Add, Number);
+	return FNameHelper::MakeWithNumber(FNameEntryIds{ ComparisonId, DisplayId }, FNAME_Add, Number);
 }
 #endif // UE_FNAME_OUTLINE_NUMBER
 
