MpscQueue - Fix for dequeuing move only types.

diff --git a/Engine/Source/Runtime/Core/Public/Containers/MpscQueue.h b/Engine/Source/Runtime/Core/Public/Containers/MpscQueue.h
--- a/Engine/Source/Runtime/Core/Public/Containers/MpscQueue.h
+++ b/Engine/Source/Runtime/Core/Public/Containers/MpscQueue.h
@@ -79,7 +79,7 @@ public:
 		TOptional<ElementType> LocalElement = Dequeue();
 		if (LocalElement.IsSet())
 		{
-			OutElem = LocalElement.GetValue();
+			OutElem = MoveTempIfPossible(LocalElement.GetValue());
 			return true;
 		}
 
