Removed templating from FString implementation functions, to allow them to be defined per-char type and not clash in unity builds.

diff --git a/Engine/Source/Runtime/Core/Private/Containers/String.cpp.inl b/Engine/Source/Runtime/Core/Private/Containers/String.cpp.inl
--- a/Engine/Source/Runtime/Core/Private/Containers/String.cpp.inl
+++ b/Engine/Source/Runtime/Core/Private/Containers/String.cpp.inl
@@ -25,26 +25,8 @@
 
 namespace UE::Core::Private
 {
-	template <typename CharType>
-	struct TCompareCharsCaseSensitive
-	{
-		static FORCEINLINE bool Compare(CharType Lhs, CharType Rhs)
-		{
-			return Lhs == Rhs;
-		}
-	};
-
-	template <typename CharType>
-	struct TCompareCharsCaseInsensitive
-	{
-		static FORCEINLINE bool Compare(CharType Lhs, CharType Rhs)
-		{
-			return TChar<CharType>::ToLower(Lhs) == TChar<CharType>::ToLower(Rhs);
-		}
-	};
-
-	template <typename CompareType, typename CharType>
-	bool MatchesWildcardRecursive(const CharType* Target, int32 TargetLength, const CharType* Wildcard, int32 WildcardLength)
+	template <typename CompareType>
+	bool MatchesWildcardRecursive(const UE_STRING_CHARTYPE* Target, int32 TargetLength, const UE_STRING_CHARTYPE* Wildcard, int32 WildcardLength, CompareType Compare)
 	{
 		// Skip over common initial non-wildcard-char sequence of Target and Wildcard
 		for (;;)
@@ -54,13 +36,13 @@ namespace UE::Core::Private
 				return TargetLength == 0;
 			}
 
-			CharType WCh = *Wildcard;
-			if (WCh == CharType('*') || WCh == CharType('?'))
+			UE_STRING_CHARTYPE WCh = *Wildcard;
+			if (WCh == CHARTEXT(UE_STRING_CHARTYPE, '*') || WCh == CHARTEXT(UE_STRING_CHARTYPE, '?'))
 			{
 				break;
 			}
 
-			if (!CompareType::Compare(*Target, WCh))
+			if (!Compare(*Target, WCh))
 			{
 				return false;
 			}
@@ -72,20 +54,20 @@ namespace UE::Core::Private
 		}
 
 		// Test for common suffix
-		const CharType* TPtr = Target   + TargetLength;
-		const CharType* WPtr = Wildcard + WildcardLength;
+		const UE_STRING_CHARTYPE* TPtr = Target   + TargetLength;
+		const UE_STRING_CHARTYPE* WPtr = Wildcard + WildcardLength;
 		for (;;)
 		{
 			--TPtr;
 			--WPtr;
 
-			CharType WCh = *WPtr;
-			if (WCh == CharType('*') || WCh == CharType('?'))
+			UE_STRING_CHARTYPE WCh = *WPtr;
+			if (WCh == CHARTEXT(UE_STRING_CHARTYPE, '*') || WCh == CHARTEXT(UE_STRING_CHARTYPE, '?'))
 			{
 				break;
 			}
 
-			if (!CompareType::Compare(*TPtr, WCh))
+			if (!Compare(*TPtr, WCh))
 			{
 				return false;
 			}
@@ -100,8 +82,8 @@ namespace UE::Core::Private
 		}
 
 		// Match * against anything and ? against single (and zero?) chars
-		CharType FirstWild = *Wildcard;
-		if (WildcardLength == 1 && (FirstWild == CharType('*') || TargetLength < 2))
+		UE_STRING_CHARTYPE FirstWild = *Wildcard;
+		if (WildcardLength == 1 && (FirstWild == CHARTEXT(UE_STRING_CHARTYPE, '*') || TargetLength < 2))
 		{
 			return true;
 		}
@@ -110,14 +92,14 @@ namespace UE::Core::Private
 
 		// This routine is very slow, though it does ok with one wildcard
 		int32 MaxNum = TargetLength;
-		if (FirstWild == CharType('?') && MaxNum > 1)
+		if (FirstWild == CHARTEXT(UE_STRING_CHARTYPE, '?') && MaxNum > 1)
 		{
 			MaxNum = 1;
 		}
 
 		for (int32 Index = 0; Index <= MaxNum; ++Index)
 		{
-			if (MatchesWildcardRecursive<CompareType>(Target + Index, TargetLength - Index, Wildcard, WildcardLength))
+			if (MatchesWildcardRecursive(Target + Index, TargetLength - Index, Wildcard, WildcardLength, Compare))
 			{
 				return true;
 			}
@@ -125,8 +107,8 @@ namespace UE::Core::Private
 		return false;
 	}
 
-	template <typename DestCharType, typename SrcCharType>
-	void AppendCharacters(TArray<DestCharType>& Out, const SrcCharType* Str, int32 Count)
+	template <typename SrcCharType>
+	void AppendCharacters(TArray<UE_STRING_CHARTYPE>& Out, const SrcCharType* Str, int32 Count)
 	{
 		check(Count >= 0);
 
@@ -144,15 +126,15 @@ namespace UE::Core::Private
 		Out.AddUninitialized(Count + (OldEnd ? 0 : 1));
 		OldEnd -= OldEnd ? 1 : 0;
 
-		DestCharType* Dest = Out.GetData() + OldEnd;
+		UE_STRING_CHARTYPE* Dest = Out.GetData() + OldEnd;
 
 		// Try copying characters to end of string, overwriting null terminator if we already have one
-		DestCharType* NewEnd = FPlatformString::Convert(Dest, Count, Str, Count);
+		UE_STRING_CHARTYPE* NewEnd = FPlatformString::Convert(Dest, Count, Str, Count);
 		if (!NewEnd)
 		{
 			// If that failed, it will have meant that conversion likely contained multi-code unit characters
 			// and so the buffer wasn't long enough, so calculate it properly.
-			int32 Length = FPlatformString::ConvertedLength<DestCharType>(Str, Count);
+			int32 Length = FPlatformString::ConvertedLength<UE_STRING_CHARTYPE>(Str, Count);
 
 			// Add the extra bytes that we need
 			Out.AddUninitialized(Length - Count);
@@ -173,16 +155,16 @@ namespace UE::Core::Private
 		}
 
 		// (Re-)establish the null terminator
-		*NewEnd = CHARTEXT(DestCharType, '\0');
+		*NewEnd = CHARTEXT(UE_STRING_CHARTYPE, '\0');
 	}
 
-	template <typename DestCharType, typename SrcCharType>
-	FORCEINLINE void ConstructFromCString(/* Out */ TArray<DestCharType>& Data, const SrcCharType* Src)
+	template <typename SrcCharType>
+	FORCEINLINE void ConstructFromCString(/* Out */ TArray<UE_STRING_CHARTYPE>& Data, const SrcCharType* Src)
 	{
 		if (Src && *Src)
 		{
 			int32 SrcLen  = TCString<SrcCharType>::Strlen(Src) + 1;
-			int32 DestLen = FPlatformString::ConvertedLength<DestCharType>(Src, SrcLen);
+			int32 DestLen = FPlatformString::ConvertedLength<UE_STRING_CHARTYPE>(Src, SrcLen);
 			Data.Reserve(DestLen);
 			Data.AddUninitialized(DestLen);
 
@@ -190,30 +172,30 @@ namespace UE::Core::Private
 		}
 	}
 
-	template <typename DestCharType, typename SrcCharType>
-	FORCEINLINE void ConstructWithLength(/* Out */ TArray<DestCharType>& Data, int32 InCount, const SrcCharType* InSrc)
+	template <typename SrcCharType>
+	FORCEINLINE void ConstructWithLength(/* Out */ TArray<UE_STRING_CHARTYPE>& Data, int32 InCount, const SrcCharType* InSrc)
 	{
 		if (InSrc)
 		{
-			int32 DestLen = FPlatformString::ConvertedLength<DestCharType>(InSrc, InCount);
+			int32 DestLen = FPlatformString::ConvertedLength<UE_STRING_CHARTYPE>(InSrc, InCount);
 			if (DestLen > 0 && *InSrc)
 			{
 				Data.Reserve(DestLen + 1);
 				Data.AddUninitialized(DestLen + 1);
 
 				FPlatformString::Convert(Data.GetData(), DestLen, InSrc, InCount);
-				*(Data.GetData() + Data.Num() - 1) = DestCharType('\0');
+				*(Data.GetData() + Data.Num() - 1) = CHARTEXT(UE_STRING_CHARTYPE, '\0');
 			}
 		}
 	}
 
-	template <typename DestCharType, typename SrcCharType>
-	FORCEINLINE void ConstructWithSlack(/* Out */ TArray<DestCharType>& Data, const SrcCharType* Src, int32 ExtraSlack)
+	template <typename SrcCharType>
+	FORCEINLINE void ConstructWithSlack(/* Out */ TArray<UE_STRING_CHARTYPE>& Data, const SrcCharType* Src, int32 ExtraSlack)
 	{
 		if (Src && *Src)
 		{
 			int32 SrcLen = TCString<SrcCharType>::Strlen(Src) + 1;
-			int32 DestLen = FPlatformString::ConvertedLength<DestCharType>(Src, SrcLen);
+			int32 DestLen = FPlatformString::ConvertedLength<UE_STRING_CHARTYPE>(Src, SrcLen);
 			Data.Reserve(DestLen + ExtraSlack);
 			Data.AddUninitialized(DestLen);
 
@@ -1286,11 +1268,29 @@ bool UE_STRING_CLASS::MatchesWildcard(const ElementType* InWildcard, int32 InWil
 
 	if (SearchCase == ESearchCase::CaseSensitive)
 	{
-		return UE::Core::Private::MatchesWildcardRecursive<UE::Core::Private::TCompareCharsCaseSensitive<ElementType>>(Target, TargetLength, InWildcard, InWildcardLen);
+		return UE::Core::Private::MatchesWildcardRecursive(
+			Target,
+			TargetLength,
+			InWildcard,
+			InWildcardLen,
+			[](UE_STRING_CHARTYPE Lhs, UE_STRING_CHARTYPE Rhs)
+			{
+				return Lhs == Rhs;
+			}
+		);
 	}
 	else
 	{
-		return UE::Core::Private::MatchesWildcardRecursive<UE::Core::Private::TCompareCharsCaseInsensitive<ElementType>>(Target, TargetLength, InWildcard, InWildcardLen);
+		return UE::Core::Private::MatchesWildcardRecursive(
+			Target,
+			TargetLength,
+			InWildcard,
+			InWildcardLen,
+			[](UE_STRING_CHARTYPE Lhs, UE_STRING_CHARTYPE Rhs)
+			{
+				return TChar<UE_STRING_CHARTYPE>::ToLower(Lhs) == TChar<UE_STRING_CHARTYPE>::ToLower(Rhs);
+			}
+		);
 	}
 }
 
