Removed CharToEscapeSeqMap and MaxSupportedEscapeChars templates.

diff --git a/Engine/Source/Runtime/Core/Private/Containers/String.cpp.inl b/Engine/Source/Runtime/Core/Private/Containers/String.cpp.inl
--- a/Engine/Source/Runtime/Core/Private/Containers/String.cpp.inl
+++ b/Engine/Source/Runtime/Core/Private/Containers/String.cpp.inl
@@ -1551,31 +1551,29 @@ UE_STRING_CLASS UE_STRING_CLASS::ReplaceQuotesWithEscapedQuotes() &&
 	return MoveTemp(*this);
 }
 
-template <typename CharType>
-static const CharType* CharToEscapeSeqMap[6][2] =
+static const UE_STRING_CHARTYPE* PREPROCESSOR_JOIN(CharToEscapeSeqMap_, UE_STRING_CLASS)[6][2] =
 {
 	// Always replace \\ first to avoid double-escaping characters
-	{ CHARTEXT(CharType, "\\"), CHARTEXT(CharType, "\\\\") },
-	{ CHARTEXT(CharType, "\n"), CHARTEXT(CharType, "\\n")  },
-	{ CHARTEXT(CharType, "\r"), CHARTEXT(CharType, "\\r")  },
-	{ CHARTEXT(CharType, "\t"), CHARTEXT(CharType, "\\t")  },
-	{ CHARTEXT(CharType, "\'"), CHARTEXT(CharType, "\\'")  },
-	{ CHARTEXT(CharType, "\""), CHARTEXT(CharType, "\\\"") }
+	{ CHARTEXT(UE_STRING_CHARTYPE, "\\"), CHARTEXT(UE_STRING_CHARTYPE, "\\\\") },
+	{ CHARTEXT(UE_STRING_CHARTYPE, "\n"), CHARTEXT(UE_STRING_CHARTYPE, "\\n")  },
+	{ CHARTEXT(UE_STRING_CHARTYPE, "\r"), CHARTEXT(UE_STRING_CHARTYPE, "\\r")  },
+	{ CHARTEXT(UE_STRING_CHARTYPE, "\t"), CHARTEXT(UE_STRING_CHARTYPE, "\\t")  },
+	{ CHARTEXT(UE_STRING_CHARTYPE, "\'"), CHARTEXT(UE_STRING_CHARTYPE, "\\'")  },
+	{ CHARTEXT(UE_STRING_CHARTYPE, "\""), CHARTEXT(UE_STRING_CHARTYPE, "\\\"") }
 };
 
-template <typename CharType>
-static const uint32 MaxSupportedEscapeChars = UE_ARRAY_COUNT(CharToEscapeSeqMap<CharType>);
-
 void UE_STRING_CLASS::ReplaceCharWithEscapedCharInline(const TArray<ElementType>* Chars/*=nullptr*/)
 {
 	if ( Len() > 0 && (Chars == nullptr || Chars->Num() > 0) )
 	{
-		for ( uint32 ChIdx = 0; ChIdx < MaxSupportedEscapeChars<ElementType>; ChIdx++ )
+		const auto& CharToEscapeSeqMap = PREPROCESSOR_JOIN(CharToEscapeSeqMap_, UE_STRING_CLASS);
+
+		for ( uint32 ChIdx = 0; ChIdx < UE_ARRAY_COUNT(CharToEscapeSeqMap); ChIdx++ )
 		{
-			if ( Chars == nullptr || Chars->Contains(*(CharToEscapeSeqMap<ElementType>[ChIdx][0])) )
+			if ( Chars == nullptr || Chars->Contains(*(CharToEscapeSeqMap[ChIdx][0])) )
 			{
 				// use ReplaceInline as that won't create a copy of the string if the character isn't found
-				ReplaceInline(CharToEscapeSeqMap<ElementType>[ChIdx][0], CharToEscapeSeqMap<ElementType>[ChIdx][1]);
+				ReplaceInline(CharToEscapeSeqMap[ChIdx][0], CharToEscapeSeqMap[ChIdx][1]);
 			}
 		}
 	}
@@ -1585,15 +1583,17 @@ void UE_STRING_CLASS::ReplaceEscapedCharWithCharInline(const TArray<ElementType>
 {
 	if ( Len() > 0 && (Chars == nullptr || Chars->Num() > 0) )
 	{
+		const auto& CharToEscapeSeqMap = PREPROCESSOR_JOIN(CharToEscapeSeqMap_, UE_STRING_CLASS);
+
 		// Spin CharToEscapeSeqMap backwards to ensure we're doing the inverse of ReplaceCharWithEscapedChar
-		for ( int32 ChIdx = MaxSupportedEscapeChars<ElementType>; ChIdx > 0; )
+		for ( int32 ChIdx = UE_ARRAY_COUNT(CharToEscapeSeqMap); ChIdx > 0; )
 		{
 			--ChIdx;
 
-			if ( Chars == nullptr || Chars->Contains(*(CharToEscapeSeqMap<ElementType>[ChIdx][0])) )
+			if ( Chars == nullptr || Chars->Contains(*(CharToEscapeSeqMap[ChIdx][0])) )
 			{
 				// use ReplaceInline as that won't create a copy of the string if the character isn't found
-				ReplaceInline(CharToEscapeSeqMap<ElementType>[ChIdx][1], CharToEscapeSeqMap<ElementType>[ChIdx][0]);
+				ReplaceInline(CharToEscapeSeqMap[ChIdx][1], CharToEscapeSeqMap[ChIdx][0]);
 			}
 		}
 	}
