Added FAnsiString and FUtf8String.

diff --git a/Engine/Source/Runtime/Core/Private/Containers/AnsiString.cpp b/Engine/Source/Runtime/Core/Private/Containers/AnsiString.cpp
new file mode 100644
--- /dev/null
+++ b/Engine/Source/Runtime/Core/Private/Containers/AnsiString.cpp
@@ -0,0 +1,12 @@
+// Copyright Epic Games, Inc. All Rights Reserved.
+
+#include "Containers/AnsiString.h"
+#include "Containers/StringIncludes.cpp.inl"
+
+#define UE_STRING_CLASS             FAnsiString
+#define UE_STRING_CHARTYPE          ANSICHAR
+#define UE_STRING_CHARTYPE_IS_TCHAR 0
+	#include "Containers/String.cpp.inl"
+#undef UE_STRING_CHARTYPE_IS_TCHAR
+#undef UE_STRING_CHARTYPE
+#undef UE_STRING_CLASS
diff --git a/Engine/Source/Runtime/Core/Private/Containers/Utf8String.cpp b/Engine/Source/Runtime/Core/Private/Containers/Utf8String.cpp
new file mode 100644
--- /dev/null
+++ b/Engine/Source/Runtime/Core/Private/Containers/Utf8String.cpp
@@ -0,0 +1,12 @@
+// Copyright Epic Games, Inc. All Rights Reserved.
+
+#include "Containers/Utf8String.h"
+#include "Containers/StringIncludes.cpp.inl"
+
+#define UE_STRING_CLASS             FUtf8String
+#define UE_STRING_CHARTYPE          UTF8CHAR
+#define UE_STRING_CHARTYPE_IS_TCHAR 0
+	#include "Containers/String.cpp.inl"
+#undef UE_STRING_CHARTYPE_IS_TCHAR
+#undef UE_STRING_CHARTYPE
+#undef UE_STRING_CLASS
diff --git a/Engine/Source/Runtime/Core/Public/Containers/AnsiString.h b/Engine/Source/Runtime/Core/Public/Containers/AnsiString.h
new file mode 100644
--- /dev/null
+++ b/Engine/Source/Runtime/Core/Public/Containers/AnsiString.h
@@ -0,0 +1,14 @@
+// Copyright Epic Games, Inc. All Rights Reserved.
+
+#pragma once
+
+// Include UnrealString.h.inl's includes before defining the macros, in case the macros 'poison' other headers or there are re-entrant includes.
+#include "Containers/UnrealStringIncludes.h.inl"
+
+#define UE_STRING_CLASS             FAnsiString
+#define UE_STRING_CHARTYPE          ANSICHAR
+#define UE_STRING_CHARTYPE_IS_TCHAR 0
+	#include "Containers/UnrealString.h.inl"
+#undef UE_STRING_CHARTYPE_IS_TCHAR
+#undef UE_STRING_CHARTYPE
+#undef UE_STRING_CLASS
diff --git a/Engine/Source/Runtime/Core/Public/Containers/ContainersFwd.h b/Engine/Source/Runtime/Core/Public/Containers/ContainersFwd.h
--- a/Engine/Source/Runtime/Core/Public/Containers/ContainersFwd.h
+++ b/Engine/Source/Runtime/Core/Public/Containers/ContainersFwd.h
@@ -13,8 +13,12 @@ using FDefaultAllocator64 = TSizedDefaultAllocator<64>;
 class FDefaultSetAllocator;
 
 class FString;
+class FAnsiString;
+class FUtf8String;
 
-template<> struct TIsContiguousContainer<FString> { static constexpr bool Value = true; };
+template<> struct TIsContiguousContainer<FString>     { static constexpr bool Value = true; };
+template<> struct TIsContiguousContainer<FAnsiString> { static constexpr bool Value = true; };
+template<> struct TIsContiguousContainer<FUtf8String> { static constexpr bool Value = true; };
 
 template<typename T, typename Allocator = FDefaultAllocator> class TArray;
 template<typename T> using TArray64 = TArray<T, FDefaultAllocator64>;
diff --git a/Engine/Source/Runtime/Core/Public/Containers/UnrealString.h b/Engine/Source/Runtime/Core/Public/Containers/UnrealString.h
--- a/Engine/Source/Runtime/Core/Public/Containers/UnrealString.h
+++ b/Engine/Source/Runtime/Core/Public/Containers/UnrealString.h
@@ -454,3 +454,5 @@ namespace UE::Core::Private
 {
 	CORE_API void StripNegativeZero(double& InFloat);
 }
+
+#include "Misc/StringFormatArg.h"
diff --git a/Engine/Source/Runtime/Core/Public/Containers/UnrealString.h.inl b/Engine/Source/Runtime/Core/Public/Containers/UnrealString.h.inl
--- a/Engine/Source/Runtime/Core/Public/Containers/UnrealString.h.inl
+++ b/Engine/Source/Runtime/Core/Public/Containers/UnrealString.h.inl
@@ -2176,5 +2176,3 @@ CORE_API int32 FindMatchingClosingParenthesis(const UE_STRING_CLASS& TargetStrin
 * @return	The slugged string
 */
 CORE_API UE_STRING_CLASS SlugStringForValidName(const UE_STRING_CLASS& DisplayString, const TCHAR* ReplaceWith = TEXT(""));
-
-#include "Misc/StringFormatArg.h"
diff --git a/Engine/Source/Runtime/Core/Public/Containers/Utf8String.h b/Engine/Source/Runtime/Core/Public/Containers/Utf8String.h
new file mode 100644
--- /dev/null
+++ b/Engine/Source/Runtime/Core/Public/Containers/Utf8String.h
@@ -0,0 +1,14 @@
+// Copyright Epic Games, Inc. All Rights Reserved.
+
+#pragma once
+
+// Include UnrealString.h.inl's includes before defining the macros, in case the macros 'poison' other headers or there are re-entrant includes.
+#include "Containers/UnrealStringIncludes.h.inl"
+
+#define UE_STRING_CLASS             FUtf8String
+#define UE_STRING_CHARTYPE          UTF8CHAR
+#define UE_STRING_CHARTYPE_IS_TCHAR 0
+	#include "Containers/UnrealString.h.inl"
+#undef UE_STRING_CHARTYPE_IS_TCHAR
+#undef UE_STRING_CHARTYPE
+#undef UE_STRING_CLASS
