CIS fix.

diff --git a/Engine/Source/Runtime/Core/Public/Containers/UnrealString.h.inl b/Engine/Source/Runtime/Core/Public/Containers/UnrealString.h.inl
--- a/Engine/Source/Runtime/Core/Public/Containers/UnrealString.h.inl
+++ b/Engine/Source/Runtime/Core/Public/Containers/UnrealString.h.inl
@@ -158,7 +158,7 @@ public:
 	}
 
 public:
-#ifdef __OBJC__
+#if defined(__OBJC__) && UE_STRING_CHARTYPE_IS_TCHAR
 	/** Convert Objective-C NSString* to string class */
 	FORCEINLINE UE_STRING_CLASS(const NSString* In)
 	{
