removed a member that is not used anywhere, a remnant of old versions

diff --git a/Engine/Source/Runtime/Core/Private/Async/TaskGraph.cpp b/Engine/Source/Runtime/Core/Private/Async/TaskGraph.cpp
--- a/Engine/Source/Runtime/Core/Private/Async/TaskGraph.cpp
+++ b/Engine/Source/Runtime/Core/Private/Async/TaskGraph.cpp
@@ -1824,7 +1824,6 @@ class FTaskGraphCompatibilityImplementation final : public FTaskGraphInterface
 	TArray<FWorkerThread> NamedThreads;
 
 	FThreadSafeCounter	ReentrancyCheck;
-	FAAArrayQueue<FBaseGraphTask>	QueuedBackgroundTasks;
 
 	std::atomic<bool> bReserveWorkersEnabled{ false };
 
