Silenced IncludeTool from complaining about the string *.inl files.

diff --git a/Engine/Source/Runtime/Core/Private/Containers/String.cpp.inl b/Engine/Source/Runtime/Core/Private/Containers/String.cpp.inl
--- a/Engine/Source/Runtime/Core/Private/Containers/String.cpp.inl
+++ b/Engine/Source/Runtime/Core/Private/Containers/String.cpp.inl
@@ -9,6 +9,7 @@
  *                                                                                                                   *
  * #includes needed to compile this file need to be specified in StringIncludes.cpp.inl file rather than here.       *
  *********************************************************************************************************************/
+#define UE_INCLUDETOOL_IGNORE_INCONSISTENT_STATE
 
 #ifndef UE_STRING_CLASS
 	#error "String.cpp.inl should only be included after defining UE_STRING_CLASS"
@@ -1946,3 +1947,5 @@ UE_STRING_CLASS SlugStringForValidName(const UE_STRING_CLASS& DisplayString, con
 
 	return GeneratedName;
 }
+
+#undef UE_INCLUDETOOL_IGNORE_INCONSISTENT_STATE
diff --git a/Engine/Source/Runtime/Core/Public/Containers/UnrealString.h.inl b/Engine/Source/Runtime/Core/Public/Containers/UnrealString.h.inl
--- a/Engine/Source/Runtime/Core/Public/Containers/UnrealString.h.inl
+++ b/Engine/Source/Runtime/Core/Public/Containers/UnrealString.h.inl
@@ -9,6 +9,7 @@
  *                                                                                                                 *
  * #includes needed to compile this file need to be specified in UnrealStringIncludes.h.inl file rather than here. *
  *******************************************************************************************************************/
+#define UE_INCLUDETOOL_IGNORE_INCONSISTENT_STATE
 
 #ifndef UE_STRING_CLASS
 	#error "UnrealString.h.inl should only be included after defining UE_STRING_CLASS"
@@ -2176,3 +2177,5 @@ CORE_API int32 FindMatchingClosingParenthesis(const UE_STRING_CLASS& TargetStrin
 * @return	The slugged string
 */
 CORE_API UE_STRING_CLASS SlugStringForValidName(const UE_STRING_CLASS& DisplayString, const TCHAR* ReplaceWith = TEXT(""));
+
+#undef UE_INCLUDETOOL_IGNORE_INCONSISTENT_STATE
