Add support for labels in the file open order, specified with the FILE_OPEN_ORDER_LABEL() macro. This allows us to delineate different sections of loading to support offline processing (e.g splitting out map loading screen etc). These use existing comment syntax (#)
Also add logging on initialization indicating where the log file is located

diff --git a/Engine/Source/Runtime/Core/Private/HAL/IPlatformFileOpenLogWrapper.cpp b/Engine/Source/Runtime/Core/Private/HAL/IPlatformFileOpenLogWrapper.cpp
--- a/Engine/Source/Runtime/Core/Private/HAL/IPlatformFileOpenLogWrapper.cpp
+++ b/Engine/Source/Runtime/Core/Private/HAL/IPlatformFileOpenLogWrapper.cpp
@@ -3,6 +3,11 @@
 #include "HAL/IPlatformFileOpenLogWrapper.h"
 
 #if !UE_BUILD_SHIPPING
+
+DEFINE_LOG_CATEGORY_STATIC(LogFileOpenOrder, Log, All);
+
+static FPlatformFileOpenLog* GPlatformFileOpenLog = nullptr;
+
 IAsyncReadRequest* FLoggingAsyncReadFileHandle::ReadRequest(int64 Offset, int64 BytesToRead, EAsyncIOPriorityAndFlags PriorityAndFlags, FAsyncFileCallBack* CompleteCallback, uint8* UserSuppliedMemory)
 {
 	if ((PriorityAndFlags & AIOP_FLAG_PRECACHE) == 0)
@@ -11,4 +16,101 @@ IAsyncReadRequest* FLoggingAsyncReadFileHandle::ReadRequest(int64 Offset, int64
 	}
 	return ActualRequest->ReadRequest(Offset, BytesToRead, PriorityAndFlags, CompleteCallback, UserSuppliedMemory);
 }
+
+bool FPlatformFileOpenLog::Initialize(IPlatformFile* Inner, const TCHAR* CommandLineParam)
+{
+	LowerLevel = Inner;
+	FString LogFileDirectory;
+	FString LogFilePath;
+	FString PlatformStr;
+	FString	OutputDirectoryBase;
+
+	if (!FParse::Value(CommandLineParam, TEXT("FOOBASEDIR="), OutputDirectoryBase))
+	{
+		if (FParse::Param(CommandLineParam, TEXT("FOOUSESAVEDDIR")))
+		{
+			OutputDirectoryBase = FPaths::ProjectSavedDir();
+		}
+		else
+		{
+			OutputDirectoryBase = FPlatformMisc::ProjectDir();
+		}
+	}
+
+	if (FParse::Value(CommandLineParam, TEXT("TARGETPLATFORM="), PlatformStr))
+	{
+		TArray<FString> PlatformNames;
+		if (!(PlatformStr == TEXT("None") || PlatformStr == TEXT("All")))
+		{
+			PlatformStr.ParseIntoArray(PlatformNames, TEXT("+"), true);
+		}
+
+		for (int32 Platform = 0; Platform < PlatformNames.Num(); ++Platform)
+		{
+			if (FDataDrivenPlatformInfoRegistry::GetPlatformInfo(PlatformNames[Platform]).bIsConfidential)
+			{
+				LogFileDirectory = FPaths::Combine(OutputDirectoryBase, TEXT("Platforms"), *PlatformNames[Platform], TEXT("Build"), TEXT("FileOpenOrder"));
+			}
+			else
+			{
+				LogFileDirectory = FPaths::Combine(OutputDirectoryBase, TEXT("Build"), *PlatformNames[Platform], TEXT("FileOpenOrder"));
+			}
+#if WITH_EDITOR
+			LogFilePath = FPaths::Combine(*LogFileDirectory, TEXT("EditorOpenOrder.log"));
+#else 
+			FString OrderSuffix = TEXT("");
+			FParse::Value(CommandLineParam, TEXT("FooSuffix="), OrderSuffix);
+
+			FString GameOpenOrderFileName = FString::Printf(TEXT("GameOpenOrder%s.log"), *OrderSuffix);
+			LogFilePath = FPaths::Combine(*LogFileDirectory, *GameOpenOrderFileName);
+#endif
+			Inner->CreateDirectoryTree(*LogFileDirectory);
+			auto* FileHandle = Inner->OpenWrite(*LogFilePath, false, false);
+			if (FileHandle)
+			{
+				LogOutput.Add(FileHandle);
+			}
+			UE_LOG(LogFileOpenOrder, Log, TEXT("Initialized file open order log : %s"), *LogFilePath);
+		}
+	}
+	else
+	{
+		if (FDataDrivenPlatformInfoRegistry::GetPlatformInfo(FPlatformProperties::IniPlatformName()).bIsConfidential)
+		{
+			LogFileDirectory = FPaths::Combine(OutputDirectoryBase, TEXT("Platforms"), StringCast<TCHAR>(FPlatformProperties::PlatformName()).Get(), TEXT("Build"), TEXT("FileOpenOrder"));
+		}
+		else
+		{
+			LogFileDirectory = FPaths::Combine(OutputDirectoryBase, TEXT("Build"), StringCast<TCHAR>(FPlatformProperties::PlatformName()).Get(), TEXT("FileOpenOrder"));
+		}
+#if WITH_EDITOR
+		LogFilePath = FPaths::Combine(*LogFileDirectory, TEXT("EditorOpenOrder.log"));
+#else 
+		LogFilePath = FPaths::Combine(*LogFileDirectory, TEXT("GameOpenOrder.log"));
+#endif
+		Inner->CreateDirectoryTree(*LogFileDirectory);
+		auto* FileHandle = Inner->OpenWrite(*LogFilePath, false, false);
+		if (FileHandle)
+		{
+			LogOutput.Add(FileHandle);
+		}
+		UE_LOG(LogFileOpenOrder, Log, TEXT("Initialized file open order log : %s"), *LogFilePath);
+	}
+	GPlatformFileOpenLog = this;
+	return true;
+}
+
+void FPlatformFileOpenLog::AddLabelInternal(const TCHAR* Fmt, ...)
+{
+	if (GPlatformFileOpenLog)
+	{
+		TCHAR Buffer[256];
+		GET_VARARGS(Buffer, UE_ARRAY_COUNT(Buffer), UE_ARRAY_COUNT(Buffer) - 1, Fmt, Fmt);
+		Buffer[255] = '\0';
+		GPlatformFileOpenLog->AddLabelToOpenLog(Buffer);
+
+		UE_LOG(LogFileOpenOrder, Log, TEXT("Label: %s"), Buffer);
+	}
+}
+
 #endif // !UE_BUILD_SHIPPING
diff --git a/Engine/Source/Runtime/Core/Public/HAL/IPlatformFileOpenLogWrapper.h b/Engine/Source/Runtime/Core/Public/HAL/IPlatformFileOpenLogWrapper.h
--- a/Engine/Source/Runtime/Core/Public/HAL/IPlatformFileOpenLogWrapper.h
+++ b/Engine/Source/Runtime/Core/Public/HAL/IPlatformFileOpenLogWrapper.h
@@ -24,6 +24,9 @@ class IMappedFileHandle;
 
 #if !UE_BUILD_SHIPPING
 
+#define FILE_OPEN_ORDER_LABEL(Format, ...) \
+	FPlatformFileOpenLog::AddLabel( Format, ##__VA_ARGS__ );
+
 class FLoggingAsyncReadFileHandle final : public IAsyncReadFileHandle
 {
 	FPlatformFileOpenLog* Owner;
@@ -82,85 +85,8 @@ public:
 		return bResult;
 	}
 
-	virtual bool Initialize(IPlatformFile* Inner, const TCHAR* CommandLineParam) override
-	{
-		LowerLevel = Inner;
-		FString LogFileDirectory;
-		FString LogFilePath;
-		FString PlatformStr;
-		FString	OutputDirectoryBase;
-
-		if (!FParse::Value(CommandLineParam, TEXT("FOOBASEDIR="), OutputDirectoryBase))
-		{
-			if (FParse::Param(CommandLineParam, TEXT("FOOUSESAVEDDIR")))
-			{
-				OutputDirectoryBase = FPaths::ProjectSavedDir();
-			}
-			else
-			{
-				OutputDirectoryBase = FPlatformMisc::ProjectDir();
-			}
-		}
+	virtual bool Initialize(IPlatformFile* Inner, const TCHAR* CommandLineParam) override;
 
-		if (FParse::Value(CommandLineParam, TEXT("TARGETPLATFORM="), PlatformStr))
-		{
-			TArray<FString> PlatformNames;
-			if (!(PlatformStr == TEXT("None") || PlatformStr == TEXT("All")))
-			{
-				PlatformStr.ParseIntoArray(PlatformNames, TEXT("+"), true);
-			}
-
-			for (int32 Platform = 0;Platform < PlatformNames.Num(); ++Platform)
-			{
-				if (FDataDrivenPlatformInfoRegistry::GetPlatformInfo(PlatformNames[Platform]).bIsConfidential)
-				{
-					LogFileDirectory = FPaths::Combine(OutputDirectoryBase, TEXT("Platforms"), *PlatformNames[Platform], TEXT("Build"), TEXT("FileOpenOrder"));
-				}
-				else
-				{
-					LogFileDirectory = FPaths::Combine(OutputDirectoryBase, TEXT( "Build" ), *PlatformNames[Platform], TEXT("FileOpenOrder"));
-				}
-#if WITH_EDITOR
-				LogFilePath = FPaths::Combine( *LogFileDirectory, TEXT("EditorOpenOrder.log"));
-#else 
-				FString OrderSuffix = TEXT("");
-				FParse::Value(CommandLineParam, TEXT("FooSuffix="), OrderSuffix);
-
-				FString GameOpenOrderFileName = FString::Printf(TEXT("GameOpenOrder%s.log"), *OrderSuffix);
-				LogFilePath = FPaths::Combine( *LogFileDirectory, *GameOpenOrderFileName);
-#endif
-				Inner->CreateDirectoryTree(*LogFileDirectory);
-				auto* FileHandle = Inner->OpenWrite(*LogFilePath, false, false);
-				if (FileHandle) 
-				{
-					LogOutput.Add(FileHandle);
-				}
-			}
-		}
-		else
-		{
-			if (FDataDrivenPlatformInfoRegistry::GetPlatformInfo(FPlatformProperties::IniPlatformName()).bIsConfidential)
-			{
-				LogFileDirectory = FPaths::Combine(OutputDirectoryBase, TEXT("Platforms"), StringCast<TCHAR>(FPlatformProperties::PlatformName()).Get(), TEXT("Build"), TEXT("FileOpenOrder"));
-			}
-			else
-			{
-				LogFileDirectory = FPaths::Combine(OutputDirectoryBase, TEXT("Build"), StringCast<TCHAR>(FPlatformProperties::PlatformName()).Get(), TEXT("FileOpenOrder"));
-			}
-#if WITH_EDITOR
-			LogFilePath = FPaths::Combine( *LogFileDirectory, TEXT("EditorOpenOrder.log"));
-#else 
-			LogFilePath = FPaths::Combine( *LogFileDirectory, TEXT("GameOpenOrder.log"));
-#endif
-			Inner->CreateDirectoryTree(*LogFileDirectory);
-			auto* FileHandle = Inner->OpenWrite(*LogFilePath, false, false);
-			if (FileHandle)
-			{
-				LogOutput.Add(FileHandle);
-			}
-		}
-		return true;
-	}
 	virtual IPlatformFile* GetLowerLevel() override
 	{
 		return LowerLevel;
@@ -315,6 +241,18 @@ public:
 		CriticalSection.Unlock();
 	}
 
+	void AddLabelToOpenLog(const TCHAR* LabelStr)
+	{
+		CriticalSection.Lock();
+		for (auto File = LogOutput.CreateIterator(); File; ++File)
+		{
+			FString Text = FString::Printf(TEXT("# %s\n"), LabelStr);
+			(*File)->Write((uint8*)StringCast<ANSICHAR>(*Text).Get(), Text.Len());
+		}
+		CriticalSection.Unlock();
+	}
+
+
 	void AddPackageToOpenLog(const TCHAR* Filename)
 	{
 		CriticalSection.Lock();
@@ -329,5 +267,22 @@ public:
 		}
 		CriticalSection.Unlock();
 	}
+
+	template <typename FmtType, typename... Types>
+	static void AddLabel(const FmtType& Fmt, Types... Args)
+	{
+		static_assert(TIsArrayOrRefOfTypeByPredicate<FmtType, TIsCharEncodingCompatibleWithTCHAR>::Value, "Formatting string must be a TCHAR array.");
+		static_assert(TAnd<TIsValidVariadicFunctionArg<Types>...>::Value, "Invalid argument(s) passed to FCsvProfiler::AddLabel");
+		AddLabelInternal((const TCHAR*)Fmt, Args...);
+	}
+
+private:
+	static CORE_API void VARARGS AddLabelInternal(const TCHAR* Fmt, ...);
+
 };
-#endif // !UE_BUILD_SHIPPING
+#else // !UE_BUILD_SHIPPING
+
+#define FILE_OPEN_ORDER_LABEL(Format, ...)
+
+#endif
+
