Fix warning due to GET_VARARGS deprecation

diff --git a/Engine/Source/Runtime/Core/Private/HAL/IPlatformFileOpenLogWrapper.cpp b/Engine/Source/Runtime/Core/Private/HAL/IPlatformFileOpenLogWrapper.cpp
--- a/Engine/Source/Runtime/Core/Private/HAL/IPlatformFileOpenLogWrapper.cpp
+++ b/Engine/Source/Runtime/Core/Private/HAL/IPlatformFileOpenLogWrapper.cpp
@@ -105,7 +105,7 @@ void FPlatformFileOpenLog::AddLabelInternal(const TCHAR* Fmt, ...)
 	if (GPlatformFileOpenLog)
 	{
 		TCHAR Buffer[256];
-		GET_VARARGS(Buffer, UE_ARRAY_COUNT(Buffer), UE_ARRAY_COUNT(Buffer) - 1, Fmt, Fmt);
+		GET_TYPED_VARARGS(TCHAR, Buffer, UE_ARRAY_COUNT(Buffer), UE_ARRAY_COUNT(Buffer) - 1, Fmt, Fmt);
 		Buffer[255] = '\0';
 		GPlatformFileOpenLog->AddLabelToOpenLog(Buffer);
 
