- Enabling BSD_SOCKET_FEATURE_IOCTL and BSD_SOCKET_FEATURE_POLL for Mac platform
(Copied over from source CL 28461631)

diff --git a/Engine/Source/Runtime/Core/Public/Mac/MacPlatform.h b/Engine/Source/Runtime/Core/Public/Mac/MacPlatform.h
--- a/Engine/Source/Runtime/Core/Public/Mac/MacPlatform.h
+++ b/Engine/Source/Runtime/Core/Public/Mac/MacPlatform.h
@@ -67,6 +67,8 @@ typedef FMacPlatformTypes FPlatformTypes;
 #endif
 #define PLATFORM_HAS_BSD_TIME							1
 #define PLATFORM_HAS_BSD_IPV6_SOCKETS					1
+#define PLATFORM_HAS_BSD_SOCKET_FEATURE_IOCTL			1
+#define PLATFORM_HAS_BSD_SOCKET_FEATURE_POLL			1
 //#define PLATFORM_USE_PTHREADS							1
 #define PLATFORM_MAX_FILEPATH_LENGTH_DEPRECATED			MAC_MAX_PATH
 #define PLATFORM_SUPPORTS_TBB							1
