Add helpers to quickly add text to tokenized messages

diff --git a/Engine/Source/Runtime/Core/Private/Logging/TokenizedMessage.cpp b/Engine/Source/Runtime/Core/Private/Logging/TokenizedMessage.cpp
--- a/Engine/Source/Runtime/Core/Private/Logging/TokenizedMessage.cpp
+++ b/Engine/Source/Runtime/Core/Private/Logging/TokenizedMessage.cpp
@@ -117,6 +117,12 @@ TSharedRef<FTokenizedMessage> FTokenizedMessage::AddToken( const TSharedRef<IMes
 	return AsShared();
 }
 
+TSharedRef<FTokenizedMessage> FTokenizedMessage::AddText(const FText& InText)
+{
+	MessageTokens.Add(FTextToken::Create(InText));
+	return AsShared();
+}
+
 
 void FTokenizedMessage::SetMessageLink(const TSharedRef<IMessageToken>& InToken)
 {
diff --git a/Engine/Source/Runtime/Core/Public/Logging/TokenizedMessage.h b/Engine/Source/Runtime/Core/Public/Logging/TokenizedMessage.h
--- a/Engine/Source/Runtime/Core/Public/Logging/TokenizedMessage.h
+++ b/Engine/Source/Runtime/Core/Public/Logging/TokenizedMessage.h
@@ -138,6 +138,26 @@ public:
 	 */
 	CORE_API TSharedRef<FTokenizedMessage> AddToken( const TSharedRef<IMessageToken>& InToken );
 
+	/** 
+	 * Adds a text token to a message.
+	 * @param	InMessage	The message to insert a token into
+	 * @param	InText		The text to insert as a token
+	 * @returns this message, for chaining calls.
+	 */
+	CORE_API TSharedRef<FTokenizedMessage> AddText(const FText& InText);
+
+	/** 
+	 * Adds a text token to a message as by calling FText::FormatOrdered
+	 * @param	InMessage	The message to insert a token into
+	 * @param	InText		The text to insert as a token
+	 * @returns this message, for chaining calls.
+	 */
+	template<typename... TArguments>
+	TSharedRef<FTokenizedMessage> AddText(FTextFormat InTextFormat, TArguments&&... InArgs)
+	{
+		return AddText(FText::FormatOrdered(InTextFormat, Forward<TArguments...>(InArgs...)));
+	}
+
 	/** 
 	 * Sets the severity of this message
 	 * 
