Configurable timeouts for crashes.

Licensees has been requesting a way to configure the timeout used when processing a crash.

diff --git a/Engine/Source/Runtime/Core/Private/Windows/WindowsPlatformCrashContext.cpp b/Engine/Source/Runtime/Core/Private/Windows/WindowsPlatformCrashContext.cpp
--- a/Engine/Source/Runtime/Core/Private/Windows/WindowsPlatformCrashContext.cpp
+++ b/Engine/Source/Runtime/Core/Private/Windows/WindowsPlatformCrashContext.cpp
@@ -90,7 +90,7 @@ enum EConstants
 const uint32 EnsureExceptionCode = ECrashExitCodes::UnhandledEnsure; // Use a rather unique exception code in case SEH doesn't handle it as expected.
 const uint32 AssertExceptionCode = 0x4000;
 const uint32 GPUCrashExceptionCode = 0x8000;
-constexpr double CrashHandlingTimeoutSecs = 60.0;
+constexpr double DefaultCrashHandlingTimeoutSecs = 60.0;
 
 namespace {
 	/**
@@ -147,6 +147,20 @@ namespace {
 
 		return Result == TRUE;
 	}
+	
+	/** Returns the crash timeout in seconds. */
+	FORCEINLINE double GetCrashTimeoutSeconds()
+	{
+		// By default, wait 60s for crash handling. This should normally be enough.
+		double TimeoutSeconds = DefaultCrashHandlingTimeoutSecs;
+		if (GConfig)
+		{
+			// If available override with configurable value. Negative values are interpreted as infinite wait (not generally recommended)
+			GConfig->GetDouble(TEXT("CrashReportClient"), TEXT("CrashHandlingTimeoutSecs"), TimeoutSeconds, GEngineIni);
+		}
+		return TimeoutSeconds;
+	}
+
 }
 
 /**
@@ -726,6 +740,7 @@ int32 ReportCrashForMonitor(
 		// Wait for a response, saying it's ok to continue
 		bool bCanContinueExecution = false;
 		int32 ExitCode = 0;
+		const double TimeoutSecs = GetCrashTimeoutSeconds();
 		// Would like to use TInlineAllocator here to avoid heap allocation on crashes, but it doesn't work since ReadPipeToArray 
 		// cannot take array with non-default allocator
 		TArray<uint8> ResponseBuffer;
@@ -744,7 +759,8 @@ int32 ReportCrashForMonitor(
 			// In general, the crash monitor app (CRC) is expected to respond within ~5 seconds, but it might be busy sending an
 			// ensure/stall that occurred just before. In some degenerated cases, CRC may hang several minutes and cause the crash
 			// reporting thread to timeout and resume the crashing thread, likely resulting in this process to exit or to request it exit.
-			if (IsEngineExitRequested() && FPlatformTime::Seconds() - WaitResponseStartTimeSecs >= CrashHandlingTimeoutSecs)
+			const double WaitResponseTimeSecs = FPlatformTime::Seconds() - WaitResponseStartTimeSecs;
+			if (IsEngineExitRequested() && (TimeoutSecs >= 0 && WaitResponseTimeSecs >= TimeoutSecs))
 			{
 				break;
 			}
@@ -1307,8 +1323,9 @@ public:
 	/** The thread that crashed calls this function to wait for the report to be generated */
 	FORCEINLINE bool WaitUntilCrashIsHandled()
 	{
-		// Wait 60s, it's more than enough to generate crash report. We don't want to stall forever otherwise.
-		return WaitForSingleObject(CrashHandledEvent, static_cast<DWORD>(CrashHandlingTimeoutSecs * 1000)) == WAIT_OBJECT_0;
+		const double TimeoutSeconds = GetCrashTimeoutSeconds();
+		const DWORD TimeoutMs = TimeoutSeconds < 0.0f ? INFINITE : (static_cast<DWORD>(TimeoutSeconds * 1000));
+		return WaitForSingleObject(CrashHandledEvent, TimeoutMs) == WAIT_OBJECT_0;
 	}
 
 	/** Crashes during static init should be reported directly to crash monitor. */
