Named chunk callback is always triggered with failure for unknown chunks.

#rnx

diff --git a/Engine/Source/Runtime/Core/Private/GenericPlatform/GenericPlatformChunkInstall.cpp b/Engine/Source/Runtime/Core/Private/GenericPlatform/GenericPlatformChunkInstall.cpp
--- a/Engine/Source/Runtime/Core/Private/GenericPlatform/GenericPlatformChunkInstall.cpp
+++ b/Engine/Source/Runtime/Core/Private/GenericPlatform/GenericPlatformChunkInstall.cpp
@@ -1,30 +1,48 @@
 // Copyright Epic Games, Inc. All Rights Reserved.
 
 #include "GenericPlatform/GenericPlatformChunkInstall.h"
+#include "Containers/Ticker.h"
 
 DEFINE_LOG_CATEGORY(LogChunkInstaller)
 
-void FGenericPlatformChunkInstall::DoNamedChunkCompleteCallbacks( const FName NamedChunk, EChunkLocation::Type Location, bool bHasSucceeded )
+void FGenericPlatformChunkInstall::DoNamedChunkCompleteCallbacks( const FName NamedChunk, EChunkLocation::Type Location, bool bHasSucceeded ) const
 {
-	check(IsInGameThread());
+	DoNamedChunkCompleteCallbacks( MakeArrayView(&NamedChunk,1), Location, bHasSucceeded );
+}
 
-	bool bIsInstalled = (Location == EChunkLocation::LocalFast) || (Location == EChunkLocation::LocalSlow);
 
-	if (NamedChunkCompleteDelegate.IsBound())
+void FGenericPlatformChunkInstall::DoNamedChunkCompleteCallbacks( const TArrayView<const FName>& NamedChunks, EChunkLocation::Type Location, bool bHasSucceeded ) const
+{
+	if (NamedChunks.Num() == 0)
 	{
-		FNamedChunkCompleteCallbackParam Param;
-		Param.NamedChunk = NamedChunk;
-		Param.Location = Location;
-		Param.bIsInstalled = bIsInstalled;
-		Param.bHasSucceeded = bHasSucceeded;
-
-		NamedChunkCompleteDelegate.Broadcast(Param);
+		return;
 	}
 
-PRAGMA_DISABLE_DEPRECATION_WARNINGS
-	if (bIsInstalled && NamedChunkInstallDelegate.IsBound())
+	// always defer the callback until the next gamethread tick even if we're on the gamethread already to ensure consistency
+	TArray<FName> CachedNamedChunks(NamedChunks);
+	ExecuteOnGameThread(UE_SOURCE_LOCATION, [this, NamedChunks = MoveTemp(CachedNamedChunks), Location, bHasSucceeded]()
 	{
-		NamedChunkInstallDelegate.Broadcast(NamedChunk, bHasSucceeded);
-	}
-PRAGMA_ENABLE_DEPRECATION_WARNINGS
+		bool bIsInstalled = (Location == EChunkLocation::LocalFast) || (Location == EChunkLocation::LocalSlow);
+
+		for (const FName& NamedChunk : NamedChunks)
+		{
+			if (NamedChunkCompleteDelegate.IsBound())
+			{
+				FNamedChunkCompleteCallbackParam Param;
+				Param.NamedChunk = NamedChunk;
+				Param.Location = Location;
+				Param.bIsInstalled = bIsInstalled;
+				Param.bHasSucceeded = bHasSucceeded;
+
+				NamedChunkCompleteDelegate.Broadcast(Param);
+			}
+
+			PRAGMA_DISABLE_DEPRECATION_WARNINGS
+			if (bIsInstalled && NamedChunkInstallDelegate.IsBound())
+			{
+				NamedChunkInstallDelegate.Broadcast(NamedChunk, bHasSucceeded);
+			}
+			PRAGMA_ENABLE_DEPRECATION_WARNINGS
+		}
+	});
 }
diff --git a/Engine/Source/Runtime/Core/Public/GenericPlatform/GenericPlatformChunkInstall.h b/Engine/Source/Runtime/Core/Public/GenericPlatform/GenericPlatformChunkInstall.h
--- a/Engine/Source/Runtime/Core/Public/GenericPlatform/GenericPlatformChunkInstall.h
+++ b/Engine/Source/Runtime/Core/Public/GenericPlatform/GenericPlatformChunkInstall.h
@@ -588,7 +588,8 @@ public:
 
 protected:
 
-	void DoNamedChunkCompleteCallbacks( const FName NamedChunk, EChunkLocation::Type Location, bool bHasSucceeded );
+	void DoNamedChunkCompleteCallbacks( const FName NamedChunk, EChunkLocation::Type Location, bool bHasSucceeded ) const;
+	void DoNamedChunkCompleteCallbacks( const TArrayView<const FName>& NamedChunks, EChunkLocation::Type Location, bool bHasSucceeded ) const;
 
 	/** Delegates called when installation succeeds or fails */
 	FPlatformChunkInstallMultiDelegate InstallDelegate;
