Added missing defines when UE_MEMORY_TAGS_TRACE_ENABLED is not defined
#rnx

diff --git a/Engine/Source/Runtime/Core/Public/ProfilingDebugging/TagTrace.h b/Engine/Source/Runtime/Core/Public/ProfilingDebugging/TagTrace.h
--- a/Engine/Source/Runtime/Core/Public/ProfilingDebugging/TagTrace.h
+++ b/Engine/Source/Runtime/Core/Public/ProfilingDebugging/TagTrace.h
@@ -100,6 +100,8 @@ inline constexpr int32 TRACE_TAG = 257;
 #define UE_MEMSCOPE(...)
 #define UE_MEMSCOPE_PTR(...)
 #define UE_MEMSCOPE_DEFAULT(...)
+#define UE_MEMSCOPE_UNINITIALIZED(...)
+#define UE_MEMSCOPE_ACTIVATE(...)
 
 #endif // UE_MEMORY_TAGS_TRACE_ENABLED
 
