Allow FTextFormat to optionally disable its argument modifier evaluation

diff --git a/Engine/Source/Runtime/Core/Private/Internationalization/TextFormatter.cpp b/Engine/Source/Runtime/Core/Private/Internationalization/TextFormatter.cpp
--- a/Engine/Source/Runtime/Core/Private/Internationalization/TextFormatter.cpp
+++ b/Engine/Source/Runtime/Core/Private/Internationalization/TextFormatter.cpp
@@ -367,13 +367,13 @@ public:
 	 * Construct an instance from an FText.
 	 * The text will be immediately compiled. 
 	 */
-	FTextFormatData(FText&& InText, FTextFormatPatternDefinitionConstRef InPatternDef);
+	FTextFormatData(FText&& InText, ETextFormatFlags InFormatFlags, FTextFormatPatternDefinitionConstRef InPatternDef);
 
 	/**
 	 * Construct an instance from an FString.
 	 * The string will be immediately compiled. 
 	 */
-	FTextFormatData(FString&& InString, FTextFormatPatternDefinitionConstRef InPatternDef);
+	FTextFormatData(FString&& InString, ETextFormatFlags InFormatFlags, FTextFormatPatternDefinitionConstRef InPatternDef);
 
 	/**
 	 * Test to see whether this instance contains valid compiled data.
@@ -444,6 +444,14 @@ public:
 		return CompiledExpressionType;
 	}
 
+	/**
+	 * Get the format flags being used.
+	 */
+	FORCEINLINE ETextFormatFlags GetFormatFlags() const
+	{
+		return FormatFlags;
+	}
+
 	/**
 	 * Get the format pattern definition being used.
 	 */
@@ -500,6 +508,11 @@ private:
 	 */
 	ESourceType SourceType;
 
+	/**
+	 * Flags controlling the behavior of the format.
+	 */
+	ETextFormatFlags FormatFlags;
+
 	/**
 	 * Definition of the pattern used during a text format.
 	 */
@@ -567,39 +580,39 @@ FTextFormat::FTextFormat()
 {
 }
 
-FTextFormat::FTextFormat(const FText& InText)
-	: TextFormatData(MakeShared<FTextFormatData, ESPMode::ThreadSafe>(CopyTemp(InText), FTextFormatPatternDefinition::GetDefault()))
+FTextFormat::FTextFormat(const FText& InText, ETextFormatFlags InFormatFlags)
+	: TextFormatData(MakeShared<FTextFormatData, ESPMode::ThreadSafe>(CopyTemp(InText), InFormatFlags, FTextFormatPatternDefinition::GetDefault()))
 {
 }
 
-FTextFormat::FTextFormat(const FText& InText, FTextFormatPatternDefinitionConstRef InCustomPatternDef)
-	: TextFormatData(MakeShared<FTextFormatData, ESPMode::ThreadSafe>(CopyTemp(InText), MoveTemp(InCustomPatternDef)))
+FTextFormat::FTextFormat(const FText& InText, FTextFormatPatternDefinitionConstRef InCustomPatternDef, ETextFormatFlags InFormatFlags)
+	: TextFormatData(MakeShared<FTextFormatData, ESPMode::ThreadSafe>(CopyTemp(InText), InFormatFlags, MoveTemp(InCustomPatternDef)))
 {
 }
 
-FTextFormat::FTextFormat(FString&& InString, FTextFormatPatternDefinitionConstRef InCustomPatternDef)
-	: TextFormatData(MakeShared<FTextFormatData, ESPMode::ThreadSafe>(MoveTemp(InString), MoveTemp(InCustomPatternDef)))
+FTextFormat::FTextFormat(FString&& InString, FTextFormatPatternDefinitionConstRef InCustomPatternDef, ETextFormatFlags InFormatFlags)
+	: TextFormatData(MakeShared<FTextFormatData, ESPMode::ThreadSafe>(MoveTemp(InString), InFormatFlags, MoveTemp(InCustomPatternDef)))
 {
 }
 
-FTextFormat FTextFormat::FromString(const FString& InString)
+FTextFormat FTextFormat::FromString(const FString& InString, ETextFormatFlags InFormatFlags)
 {
-	return FTextFormat(CopyTemp(InString), FTextFormatPatternDefinition::GetDefault());
+	return FTextFormat(CopyTemp(InString), FTextFormatPatternDefinition::GetDefault(), InFormatFlags);
 }
 
-FTextFormat FTextFormat::FromString(FString&& InString)
+FTextFormat FTextFormat::FromString(FString&& InString, ETextFormatFlags InFormatFlags)
 {
-	return FTextFormat(MoveTemp(InString), FTextFormatPatternDefinition::GetDefault());
+	return FTextFormat(MoveTemp(InString), FTextFormatPatternDefinition::GetDefault(), InFormatFlags);
 }
 
-FTextFormat FTextFormat::FromString(const FString& InString, FTextFormatPatternDefinitionConstRef InCustomPatternDef)
+FTextFormat FTextFormat::FromString(const FString& InString, FTextFormatPatternDefinitionConstRef InCustomPatternDef, ETextFormatFlags InFormatFlags)
 {
-	return FTextFormat(CopyTemp(InString), MoveTemp(InCustomPatternDef));
+	return FTextFormat(CopyTemp(InString), MoveTemp(InCustomPatternDef), InFormatFlags);
 }
 
-FTextFormat FTextFormat::FromString(FString&& InString, FTextFormatPatternDefinitionConstRef InCustomPatternDef)
+FTextFormat FTextFormat::FromString(FString&& InString, FTextFormatPatternDefinitionConstRef InCustomPatternDef, ETextFormatFlags InFormatFlags)
 {
-	return FTextFormat(MoveTemp(InString), MoveTemp(InCustomPatternDef));
+	return FTextFormat(MoveTemp(InString), MoveTemp(InCustomPatternDef), InFormatFlags);
 }
 
 bool FTextFormat::IsValid() const
@@ -627,6 +640,11 @@ FTextFormat::EExpressionType FTextFormat::GetExpressionType() const
 	return TextFormatData->GetExpressionType();
 }
 
+ETextFormatFlags FTextFormat::GetFormatFlags() const
+{
+	return TextFormatData->GetFormatFlags();
+}
+
 FTextFormatPatternDefinitionConstRef FTextFormat::GetPatternDefinition() const
 {
 	return TextFormatData->GetPatternDefinition();
@@ -645,20 +663,22 @@ void FTextFormat::GetFormatArgumentNames(TArray<FString>& OutArgumentNames) cons
 
 TSharedRef<FTextFormatData, ESPMode::ThreadSafe> FTextFormatData::GetSharedEmptyInstance()
 {
-	static const TSharedRef<FTextFormatData> EmptyInstance = MakeShared<FTextFormatData, ESPMode::ThreadSafe>(FText(), FTextFormatPatternDefinition::GetDefault());
+	static const TSharedRef<FTextFormatData> EmptyInstance = MakeShared<FTextFormatData, ESPMode::ThreadSafe>(FText(), ETextFormatFlags::Default, FTextFormatPatternDefinition::GetDefault());
 	return EmptyInstance;
 }
 
-FTextFormatData::FTextFormatData(FText&& InText, FTextFormatPatternDefinitionConstRef InPatternDef)
+FTextFormatData::FTextFormatData(FText&& InText, ETextFormatFlags InFormatFlags, FTextFormatPatternDefinitionConstRef InPatternDef)
 	: SourceType(ESourceType::Text)
+	, FormatFlags(InFormatFlags)
 	, PatternDef(MoveTemp(InPatternDef))
 	, SourceText(MoveTemp(InText))
 {
 	Compile_NoLock();
 }
 
-FTextFormatData::FTextFormatData(FString&& InString, FTextFormatPatternDefinitionConstRef InPatternDef)
+FTextFormatData::FTextFormatData(FString&& InString, ETextFormatFlags InFormatFlags, FTextFormatPatternDefinitionConstRef InPatternDef)
 	: SourceType(ESourceType::String)
+	, FormatFlags(InFormatFlags)
 	, PatternDef(MoveTemp(InPatternDef))
 	, SourceExpression(MoveTemp(InString))
 {
@@ -739,12 +759,19 @@ void FTextFormatData::Compile_NoLock()
 					// Peek to see if the next token is an argument modifier
 					if (const auto* ArgumentModifierToken = NextToken.Node.Cast<TextFormatTokens::FArgumentModifierTokenSpecifier>())
 					{
-						int32 ArgModLength = 0;
-						bool ArgModUsesFormatArgs = false;
-						ArgumentModifierToken->TextFormatArgumentModifier->EstimateLength(ArgModLength, ArgModUsesFormatArgs);
-
-						BaseFormatStringLength += ArgModLength;
-						FormatArgumentEstimateMultiplier += (ArgModUsesFormatArgs) ? 1 : 0;
+						if (EnumHasAnyFlags(FormatFlags, ETextFormatFlags::EvaluateArgumentModifiers))
+						{
+							int32 ArgModLength = 0;
+							bool ArgModUsesFormatArgs = false;
+							ArgumentModifierToken->TextFormatArgumentModifier->EstimateLength(ArgModLength, ArgModUsesFormatArgs);
+
+							BaseFormatStringLength += ArgModLength;
+							FormatArgumentEstimateMultiplier += (ArgModUsesFormatArgs) ? 1 : 0;
+						}
+						else
+						{
+							BaseFormatStringLength += ArgumentModifierToken->ModifierPatternLen + 1;
+						}
 
 						++TokenIndex; // walk over the argument token so that the next iteration will skip over the argument modifier
 						continue;
@@ -816,11 +843,14 @@ bool FTextFormatData::ValidatePattern_NoLock(const FCulturePtr& InCulture, TArra
 	const FCultureRef ResolvedCulture = InCulture ? InCulture.ToSharedRef() : FInternationalization::Get().GetCurrentLanguage();
 
 	bool bIsValidPattern = true;
-	for (const FExpressionToken& Token : LexedExpression)
+	if (EnumHasAnyFlags(FormatFlags, ETextFormatFlags::EvaluateArgumentModifiers))
 	{
-		if (const auto* ArgumentModifierToken = Token.Node.Cast<TextFormatTokens::FArgumentModifierTokenSpecifier>())
+		for (const FExpressionToken& Token : LexedExpression)
 		{
-			bIsValidPattern &= ArgumentModifierToken->TextFormatArgumentModifier->Validate(ResolvedCulture, OutValidationErrors);
+			if (const auto* ArgumentModifierToken = Token.Node.Cast<TextFormatTokens::FArgumentModifierTokenSpecifier>())
+			{
+				bIsValidPattern &= ArgumentModifierToken->TextFormatArgumentModifier->Validate(ResolvedCulture, OutValidationErrors);
+			}
 		}
 	}
 	return bIsValidPattern;
@@ -868,7 +898,16 @@ FString FTextFormatData::Format_NoLock(const FPrivateTextFormatArguments& InForm
 					// Peek to see if the next token is an argument modifier
 					if (const auto* ArgumentModifierToken = NextToken.Node.Cast<TextFormatTokens::FArgumentModifierTokenSpecifier>())
 					{
-						ArgumentModifierToken->TextFormatArgumentModifier->Evaluate(*PossibleArgumentValue, InFormatArgs, ResultString);
+						if (EnumHasAnyFlags(FormatFlags, ETextFormatFlags::EvaluateArgumentModifiers))
+						{
+							ArgumentModifierToken->TextFormatArgumentModifier->Evaluate(*PossibleArgumentValue, InFormatArgs, ResultString);
+						}
+						else
+						{
+							// If evaluation is disabled we just write the literal value of the argument modifier back into the final string...
+							ResultString.AppendChar(PatternDef->ArgModChar);
+							ResultString.AppendChars(ArgumentModifierToken->ModifierPatternStartPos, ArgumentModifierToken->ModifierPatternLen);
+						}
 						++TokenIndex; // walk over the argument token so that the next iteration will skip over the argument modifier
 						continue;
 					}
@@ -922,7 +961,10 @@ void FTextFormatData::GetFormatArgumentNames_NoLock(TArray<FString>& OutArgument
 		}
 		else if (const auto* ArgumentModifierToken = Token.Node.Cast<TextFormatTokens::FArgumentModifierTokenSpecifier>())
 		{
-			ArgumentModifierToken->TextFormatArgumentModifier->GetFormatArgumentNames(OutArgumentNames);
+			if (EnumHasAnyFlags(FormatFlags, ETextFormatFlags::EvaluateArgumentModifiers))
+			{
+				ArgumentModifierToken->TextFormatArgumentModifier->GetFormatArgumentNames(OutArgumentNames);
+			}
 		}
 	}
 }
@@ -1127,7 +1169,7 @@ FString FTextFormatter::Format(const FTextFormat& InFmt, const FPrivateTextForma
 			FmtText.Rebuild();
 		}
 
-		FmtPattern = FTextFormat(FmtText.BuildSourceString(), FmtPattern.GetPatternDefinition());
+		FmtPattern = FTextFormat(FmtText.BuildSourceString(), FmtPattern.GetPatternDefinition(), FmtPattern.GetFormatFlags());
 	}
 
 	return FmtPattern.TextFormatData->Format(InFormatArgs);
diff --git a/Engine/Source/Runtime/Core/Public/Internationalization/Text.h b/Engine/Source/Runtime/Core/Public/Internationalization/Text.h
--- a/Engine/Source/Runtime/Core/Public/Internationalization/Text.h
+++ b/Engine/Source/Runtime/Core/Public/Internationalization/Text.h
@@ -73,6 +73,22 @@ enum class ETextIdenticalModeFlags : uint8
 };
 ENUM_CLASS_FLAGS(ETextIdenticalModeFlags);
 
+enum class ETextFormatFlags : uint8
+{
+	/** No special behavior */
+	None = 0,
+
+	/**
+	 * Set to evaluate argument modifiers when formatting text
+	 * Unset to print the literal argument modifier syntax into the result
+	 */
+	EvaluateArgumentModifiers = 1<<0,
+
+	/** Default formatting flags */
+	Default = EvaluateArgumentModifiers,
+};
+ENUM_CLASS_FLAGS(ETextFormatFlags);
+
 enum class ETextPluralType : uint8
 {
 	Cardinal,
@@ -275,27 +291,27 @@ public:
 	 * Construct an instance from an FText.
 	 * The text will be immediately compiled. 
 	 */
-	CORE_API FTextFormat(const FText& InText);
+	CORE_API FTextFormat(const FText& InText, ETextFormatFlags InFormatFlags = ETextFormatFlags::Default);
 
 	/**
 	 * Construct an instance from an FText and custom format pattern definition.
 	 * The text will be immediately compiled.
 	 */
-	CORE_API FTextFormat(const FText& InText, FTextFormatPatternDefinitionConstRef InCustomPatternDef);
+	CORE_API FTextFormat(const FText& InText, FTextFormatPatternDefinitionConstRef InCustomPatternDef, ETextFormatFlags InFormatFlags = ETextFormatFlags::Default);
 
 	/**
 	 * Construct an instance from an FString.
 	 * The string will be immediately compiled.
 	 */
-	static CORE_API FTextFormat FromString(const FString& InString);
-	static CORE_API FTextFormat FromString(FString&& InString);
+	static CORE_API FTextFormat FromString(const FString& InString, ETextFormatFlags InFormatFlags = ETextFormatFlags::Default);
+	static CORE_API FTextFormat FromString(FString&& InString, ETextFormatFlags InFormatFlags = ETextFormatFlags::Default);
 
 	/**
 	 * Construct an instance from an FString and custom format pattern definition.
 	 * The string will be immediately compiled.
 	 */
-	static CORE_API FTextFormat FromString(const FString& InString, FTextFormatPatternDefinitionConstRef InCustomPatternDef);
-	static CORE_API FTextFormat FromString(FString&& InString, FTextFormatPatternDefinitionConstRef InCustomPatternDef);
+	static CORE_API FTextFormat FromString(const FString& InString, FTextFormatPatternDefinitionConstRef InCustomPatternDef, ETextFormatFlags InFormatFlags = ETextFormatFlags::Default);
+	static CORE_API FTextFormat FromString(FString&& InString, FTextFormatPatternDefinitionConstRef InCustomPatternDef, ETextFormatFlags InFormatFlags = ETextFormatFlags::Default);
 
 	/**
 	 * Test to see whether this instance contains valid compiled data.
@@ -324,6 +340,11 @@ public:
 	 */
 	CORE_API EExpressionType GetExpressionType() const;
 
+	/**
+	 * Get the format flags being used.
+	 */
+	CORE_API ETextFormatFlags GetFormatFlags() const;
+
 	/**
 	 * Get the format pattern definition being used.
 	 */
@@ -345,7 +366,7 @@ private:
 	 * Construct an instance from an FString.
 	 * The string will be immediately compiled.
 	 */
-	CORE_API FTextFormat(FString&& InString, FTextFormatPatternDefinitionConstRef InCustomPatternDef);
+	CORE_API FTextFormat(FString&& InString, FTextFormatPatternDefinitionConstRef InCustomPatternDef, ETextFormatFlags InFormatFlags);
 
 	/** Cached compiled expression data */
 	TSharedRef<FTextFormatData, ESPMode::ThreadSafe> TextFormatData;
