Increase Windows crash report timeout when running under Wine

diff --git a/Engine/Source/Runtime/Core/Private/Windows/WindowsPlatformCrashContext.cpp b/Engine/Source/Runtime/Core/Private/Windows/WindowsPlatformCrashContext.cpp
--- a/Engine/Source/Runtime/Core/Private/Windows/WindowsPlatformCrashContext.cpp
+++ b/Engine/Source/Runtime/Core/Private/Windows/WindowsPlatformCrashContext.cpp
@@ -29,6 +29,7 @@
 #include "Serialization/Archive.h"
 #include "Templates/UniquePtr.h"
 #include "Templates/UnrealTemplate.h"
+#include "Windows/WindowsPlatformMisc.h"
 #include "Windows/WindowsPlatformStackWalk.h"
 #include <atomic>
 #include <signal.h>
@@ -151,8 +152,8 @@ namespace {
 	/** Returns the crash timeout in seconds. */
 	FORCEINLINE double GetCrashTimeoutSeconds()
 	{
-		// By default, wait 60s for crash handling. This should normally be enough.
-		double TimeoutSeconds = DefaultCrashHandlingTimeoutSecs;
+		// 60s is plenty of time to generate a crash report under Windows, but generating minidumps under Wine is much slower
+		double TimeoutSeconds = FWindowsPlatformMisc::IsWine() ? 120.0 : DefaultCrashHandlingTimeoutSecs;
 		if (GConfig)
 		{
 			// If available override with configurable value. Negative values are interpreted as infinite wait (not generally recommended)
