fixed FTaskGraphGraphEventTest unit-test

diff --git a/Engine/Source/Runtime/Core/Tests/Async/TaskGraphTest.cpp b/Engine/Source/Runtime/Core/Tests/Async/TaskGraphTest.cpp
--- a/Engine/Source/Runtime/Core/Tests/Async/TaskGraphTest.cpp
+++ b/Engine/Source/Runtime/Core/Tests/Async/TaskGraphTest.cpp
@@ -766,34 +766,17 @@ namespace TaskGraphTests
 		}
 
 		{	// wait for prereq by DontCompleteUntil
-			auto Lambda = [](ENamedThreads::Type CurrentThread, const FGraphEventRef& MyCompletionGraphEvent)
-				{
-					//UE_LOG(LogTemp, Log, TEXT("Main task"));
-
-					FGraphEventRef PrereqHolder = FGraphEvent::CreateGraphEvent();
-					PrereqHolder->SetDebugName(TEXT("PrereqHolder"));
-
-					FGraphEventRef Prereq = FFunctionGraphTask::CreateAndDispatchWhenReady(
-						[PrereqHolder]
-						{
-							//UE_LOG(LogTemp, Log, TEXT("Prereq"));
-
-							PrereqHolder->Wait(); // hold it until it's used for `DontCompleteUntil`
-						}
-					);
-					Prereq->SetDebugName(TEXT("Prereq"));
-
-					MyCompletionGraphEvent->DontCompleteUntil(Prereq);
-					check(!Prereq->IsComplete()); // check that prereq was incomplete during DontCompleteUntil ^^
-
-					// now that Prereq was registered in DontCompleteUntil, unlock it
-					PrereqHolder->DispatchSubsequents();
+			FGraphEventRef Blocker = FGraphEvent::CreateGraphEvent();
+			auto Lambda = [&Blocker](ENamedThreads::Type CurrentThread, const FGraphEventRef& MyCompletionGraphEvent)
+			{
+				MyCompletionGraphEvent->DontCompleteUntil(Blocker);
 			};
 
-			FGraphEventRef Event = FFunctionGraphTask::CreateAndDispatchWhenReady(MoveTemp(Lambda));
-			Event->SetDebugName(TEXT("MainEvent"));
-			check(!Event->IsComplete());
-			Event->Wait(ENamedThreads::GameThread);
+			FGraphEventRef Task = FFunctionGraphTask::CreateAndDispatchWhenReady(MoveTemp(Lambda));
+			FPlatformProcess::Sleep(0.01f);
+			check(!Task->IsComplete());
+			Blocker->DispatchSubsequents();
+			Task->Wait(ENamedThreads::GameThread);
 		}
 
 		{	// prereq is completed before DontCompleteUntil is called
