- Fixed an issue with Base.ini being read over and over, which can trip some FS access checkers

diff --git a/Engine/Source/Runtime/Core/Private/Misc/ConfigCacheIni.cpp b/Engine/Source/Runtime/Core/Private/Misc/ConfigCacheIni.cpp
--- a/Engine/Source/Runtime/Core/Private/Misc/ConfigCacheIni.cpp
+++ b/Engine/Source/Runtime/Core/Private/Misc/ConfigCacheIni.cpp
@@ -471,6 +471,21 @@ static void FixupArrayOfStructKeysForSection(FConfigSection* Section, const FStr
  */
 static bool LoadConfigFileWrapper(const TCHAR* IniFile, FString& Contents, bool bIsOverride = false)
 {
+	// We read the Base.ini file many many times, so cache it
+	static FString BaseIniContents;
+
+	const TCHAR* LastSlash = FCString::Strrchr(IniFile, '/');
+	if (LastSlash == nullptr)
+	{
+		LastSlash = FCString::Strrchr(IniFile, '\\');
+	}
+	bool bIsBaseIni = LastSlash != nullptr && FCString::Stricmp(LastSlash + 1, TEXT("Base.ini")) == 0;
+	if (bIsBaseIni && BaseIniContents.Len() > 0)
+	{
+		Contents = BaseIniContents;
+		return true;
+	}
+
 	// let other systems load the file instead of the standard load below
 	FCoreDelegates::TSPreLoadConfigFileDelegate().Broadcast(IniFile, Contents);
 
@@ -493,7 +508,12 @@ static bool LoadConfigFileWrapper(const TCHAR* IniFile, FString& Contents, bool
 	// note: we don't check if FileOperations are disabled because downloadable content calls this directly (which
 	// needs file ops), and the other caller of this is already checking for disabled file ops
 	// and don't read from the file, if the delegate got anything loaded
-	return FFileHelper::LoadFileToString(Contents, IniFile);
+	bool bResult = FFileHelper::LoadFileToString(Contents, IniFile);
+	if (bResult && bIsBaseIni)
+	{
+		BaseIniContents = Contents;
+	}
+	return bResult;
 }
 
 /**
