Android: Add memory advisor's output to CSV stats. (device free, device used)

diff --git a/Engine/Source/Runtime/Core/Private/Android/AndroidPlatformMemory.cpp b/Engine/Source/Runtime/Core/Private/Android/AndroidPlatformMemory.cpp
--- a/Engine/Source/Runtime/Core/Private/Android/AndroidPlatformMemory.cpp
+++ b/Engine/Source/Runtime/Core/Private/Android/AndroidPlatformMemory.cpp
@@ -117,8 +117,21 @@ static void OnCVarAndroidUseMemoryAdvisorChanged(IConsoleVariable* Var)
 		UE_LOG(LogInit, Warning, TEXT("Cannot disable memory advisor once it has been initialized."));
 	}
 }
+
 #endif
 
+FAndroidPlatformMemory::FMemAdviceStats FAndroidPlatformMemory::GetDeviceMemAdviceStats()
+{
+	FAndroidPlatformMemory::FMemAdviceStats ret;
+#if HAS_ANDROID_MEMORY_ADVICE
+	ret.MemFree = TEMPMemoryAdvice_getAvailableMemory();
+	ret.TotalMem = MemoryAdvice_getTotalMemory();
+	ret.MemUsed = ret.TotalMem - ret.MemFree;
+#endif
+	return ret;
+}
+
+
 static int32 GAndroidAddSwapToTotalPhysical = 1;
 static FAutoConsoleVariableRef CVarAddSwapToTotalPhysical(
 	TEXT("android.AddSwapToTotalPhysical"),
diff --git a/Engine/Source/Runtime/Core/Private/Android/AndroidStats.cpp b/Engine/Source/Runtime/Core/Private/Android/AndroidStats.cpp
--- a/Engine/Source/Runtime/Core/Private/Android/AndroidStats.cpp
+++ b/Engine/Source/Runtime/Core/Private/Android/AndroidStats.cpp
@@ -160,6 +160,10 @@ CSV_DEFINE_STAT(AndroidMemory, TrimMemoryBackgroundLevel);
 static int GTrimMemoryForegroundLevel = 0;
 CSV_DEFINE_STAT(AndroidMemory, TrimMemoryForegroundLevel);
 
+CSV_DEFINE_STAT(AndroidMemory, MemAdvisor_DeviceMemFreeMb);
+CSV_DEFINE_STAT(AndroidMemory, MemAdvisor_DeviceMemUsedMb);
+static FAndroidPlatformMemory::FMemAdviceStats GDeviceMemAdviceStats;
+
 void FAndroidStats::Init(bool bEnableHWCPipe)
 {
 	(void)bEnableHWCPipe;
@@ -253,12 +257,19 @@ void FAndroidStats::UpdateAndroidStats()
 	{
 		LastCollectionTime = CurrentTime;
 		CPUTemp = FAndroidMisc::GetCPUTemperature();
+
+		GDeviceMemAdviceStats = FAndroidPlatformMemory::GetDeviceMemAdviceStats();
 	}
 
 	CSV_CUSTOM_STAT_DEFINED(CPUTemp, CPUTemp, ECsvCustomStatOp::Set);
 	CSV_CUSTOM_STAT_DEFINED(ThermalStatus, GThermalStatus, ECsvCustomStatOp::Set);
 	CSV_CUSTOM_STAT_DEFINED(TrimMemoryBackgroundLevel, GTrimMemoryBackgroundLevel, ECsvCustomStatOp::Set);
-	CSV_CUSTOM_STAT_DEFINED(TrimMemoryForegroundLevel, GTrimMemoryForegroundLevel, ECsvCustomStatOp::Set);
+
+	if(GDeviceMemAdviceStats.IsValid())
+	{
+		CSV_CUSTOM_STAT_DEFINED(MemAdvisor_DeviceMemFreeMb, (int)(GDeviceMemAdviceStats.MemFree/(1024*1024)), ECsvCustomStatOp::Set);
+		CSV_CUSTOM_STAT_DEFINED(MemAdvisor_DeviceMemUsedMb, (int)(GDeviceMemAdviceStats.MemUsed/(1024*1024)), ECsvCustomStatOp::Set);
+	}
 
 	static const uint32 MaxFrequencyGroupStats = 4;
 	const int32 MaxCoresStatsSupport = 16;
diff --git a/Engine/Source/Runtime/Core/Public/Android/AndroidPlatformMemory.h b/Engine/Source/Runtime/Core/Public/Android/AndroidPlatformMemory.h
--- a/Engine/Source/Runtime/Core/Public/Android/AndroidPlatformMemory.h
+++ b/Engine/Source/Runtime/Core/Public/Android/AndroidPlatformMemory.h
@@ -32,6 +32,15 @@ struct FAndroidPlatformMemory : public FGenericPlatformMemory
 	static CORE_API void* BinnedAllocFromOS(SIZE_T Size);
 	static CORE_API void BinnedFreeToOS(void* Ptr, SIZE_T Size);
 
+	CORE_API struct FMemAdviceStats
+	{
+		bool IsValid() const { return MemFree > 0 && MemUsed > 0; };
+		int64 TotalMem = 0; // total memory available to the device
+		int64 MemFree = 0;
+		int64 MemUsed = 0;  // mem used across the device, (simply returns total - free)
+	};
+	static CORE_API FMemAdviceStats GetDeviceMemAdviceStats();
+
 	class FPlatformVirtualMemoryBlock : public FBasicVirtualMemoryBlock
 	{
 	public:
