We want to check for a RenderDoc module loaded by someone else, not load it ourselves by accident.

#rnx

diff --git a/Engine/Source/Runtime/Core/Private/HAL/ThreadingBase.cpp b/Engine/Source/Runtime/Core/Private/HAL/ThreadingBase.cpp
--- a/Engine/Source/Runtime/Core/Private/HAL/ThreadingBase.cpp
+++ b/Engine/Source/Runtime/Core/Private/HAL/ThreadingBase.cpp
@@ -18,7 +18,7 @@
 #include "AutoRTFM/AutoRTFM.h"
 
 #if PLATFORM_WINDOWS
-#include "Microsoft/MinimalWindowsApi.h"
+#include "Windows/WindowsHWrapper.h"
 #endif
 
 #include <atomic>
@@ -94,7 +94,7 @@ FTaskTagScope::FTaskTagScope(bool InTagOnlyIfNone, ETaskTag InTag) : Tag(InTag),
 		// the PE loading process prior to calling the thread's entry point. This
 		// results static initialisation unexpectedly happening on RenderDoc's injection
 		// thread and the ensure below fails.
-		static bool bRenderDocDetected = (Windows::LoadLibraryW(L"renderdoc.dll") != nullptr);
+		static bool bRenderDocDetected = (::GetModuleHandleW(L"renderdoc.dll") != nullptr);
 		if (!bRenderDocDetected)
 #	endif
 		{
