Add Open around GetDisplayString which will return a shared ref, so if we abort during the transaction just reset this shared ref allowing it to be cleaned up if nothing else is ref'ing it

diff --git a/Engine/Source/Runtime/Core/Private/Internationalization/TextHistory.cpp b/Engine/Source/Runtime/Core/Private/Internationalization/TextHistory.cpp
--- a/Engine/Source/Runtime/Core/Private/Internationalization/TextHistory.cpp
+++ b/Engine/Source/Runtime/Core/Private/Internationalization/TextHistory.cpp
@@ -906,7 +906,14 @@ bool FTextHistory_Base::CanUpdateDisplayString()
 void FTextHistory_Base::UpdateDisplayString()
 {
 	check(!TextId.IsEmpty()); // CanUpdateDisplayString should prevent UpdateDisplayString being called
-	LocalizedString = FTextLocalizationManager::Get().GetDisplayString(TextId.GetNamespace(), TextId.GetKey(), SourceString.IsEmpty() ? nullptr : &SourceString);
+
+	// Create a temp to hold the old value in case we abort, in which case we assign out of the OPEN so the old value will be preserved
+	FTextConstDisplayStringPtr NewLocalizedString;
+	UE_AUTORTFM_OPEN({
+		NewLocalizedString = FTextLocalizationManager::Get().GetDisplayString(TextId.GetNamespace(), TextId.GetKey(), SourceString.IsEmpty() ? nullptr : &SourceString);
+	});
+
+	LocalizedString = NewLocalizedString;
 }
 
 bool FTextHistory_Base::StaticShouldReadFromBuffer(const TCHAR* Buffer)
