Add missing rand StdLibSupport self open

diff --git a/Engine/Source/Runtime/Core/Private/AutoRTFM/StdLibSupport.cpp b/Engine/Source/Runtime/Core/Private/AutoRTFM/StdLibSupport.cpp
--- a/Engine/Source/Runtime/Core/Private/AutoRTFM/StdLibSupport.cpp
+++ b/Engine/Source/Runtime/Core/Private/AutoRTFM/StdLibSupport.cpp
@@ -212,6 +212,7 @@ UE_AUTORTFM_REGISTER_SELF_FUNCTION(logf);
 UE_AUTORTFM_REGISTER_SELF_FUNCTION(powf);
 UE_AUTORTFM_REGISTER_SELF_FUNCTION(fmodf);
 UE_AUTORTFM_REGISTER_SELF_FUNCTION(fmodl);
+UE_AUTORTFM_REGISTER_SELF_FUNCTION(rand);
 
 // FIXME: This is only correct when:
 // - Str is newly allocated
