fixed stack-use-after-scope reported by ASan in `WaitAny`: the function can leave the scope while some tasks are being still executed.

diff --git a/Engine/Source/Runtime/Core/Public/Tasks/Task.h b/Engine/Source/Runtime/Core/Public/Tasks/Task.h
--- a/Engine/Source/Runtime/Core/Public/Tasks/Task.h
+++ b/Engine/Source/Runtime/Core/Public/Tasks/Task.h
@@ -416,14 +416,14 @@ namespace UE::Tasks
 		}
 
 		FSharedEventRef Event;
-		std::atomic<int32> CompletedTaskIndex;
+		TSharedPtr<std::atomic<int32>> CompletedTaskIndex = MakeShared<std::atomic<int32>>(INDEX_NONE);
 
 		for (int32 Index = 0; Index != Tasks.Num(); ++Index)
 		{
 			Launch(UE_SOURCE_LOCATION, 
-				[Event, Index, &CompletedTaskIndex] 
+				[Event, Index, CompletedTaskIndex] 
 				{ 
-					CompletedTaskIndex.store(Index, std::memory_order_relaxed);
+					CompletedTaskIndex->store(Index, std::memory_order_relaxed);
 					Event->Trigger(); 
 				}, 
 				Prerequisites(Tasks[Index]),
@@ -432,7 +432,7 @@ namespace UE::Tasks
 			);
 		}
 
-		return Event->Wait(Timeout) ? CompletedTaskIndex.load(std::memory_order_relaxed) : INDEX_NONE;
+		return Event->Wait(Timeout) ? CompletedTaskIndex->load(std::memory_order_relaxed) : INDEX_NONE;
 	}
 
 	// Returns a task that gets completed as soon as any of the given tasks gets completed
