AutomationSpec - Fix hyperlinks that were not able to open the corresponding cpp files from the Test Automation window

diff --git a/Engine/Source/Runtime/Core/Public/Misc/AutomationTest.h b/Engine/Source/Runtime/Core/Public/Misc/AutomationTest.h
--- a/Engine/Source/Runtime/Core/Public/Misc/AutomationTest.h
+++ b/Engine/Source/Runtime/Core/Public/Misc/AutomationTest.h
@@ -3261,12 +3261,12 @@ private:
 		return CompleteId;
 	}
 
-	static TArray<FProgramCounterSymbolInfo> StackWalk()
+	static TArray<FProgramCounterSymbolInfo> StackWalk(int32 IgnoreCount, int32 MaxDepth)
 	{
 		QUICK_SCOPE_CYCLE_COUNTER(STAT_FAutomationSpecBase_StackWalk);
 
 		LLM_SCOPE_BYNAME(TEXT("AutomationTest/Framework"));
-		SAFE_GETSTACK(Stack, 1, 1);
+		SAFE_GETSTACK(Stack, IgnoreCount, MaxDepth);
 		return Stack;
 	}
 
@@ -3288,18 +3288,13 @@ private:
 	{
 		QUICK_SCOPE_CYCLE_COUNTER(STAT_FAutomationSpecBase_GetStack);
 
-		static bool CurrentNeedSkipStackWalk(FAutomationTestFramework::NeedSkipStackWalk());
-		static TSharedRef<TArray<FProgramCounterSymbolInfo>> Stack = MakeShared<TArray<FProgramCounterSymbolInfo>>(
-			CurrentNeedSkipStackWalk ? SkipStackWalk() : StackWalk());
-
-		if (CurrentNeedSkipStackWalk != FAutomationTestFramework::NeedSkipStackWalk())
-		{
-			// This block is to react on changes of corresponding configuration variable */
-			CurrentNeedSkipStackWalk = !CurrentNeedSkipStackWalk;
-			Stack = MakeShared<TArray<FProgramCounterSymbolInfo>>(
-				CurrentNeedSkipStackWalk ? SkipStackWalk() : StackWalk());
-		}
+		const bool NeedSkipStackWalk(FAutomationTestFramework::NeedSkipStackWalk());
+		constexpr int32 IgnoreCount(3);
+		constexpr int32 MaxDepth(1);
 		
+		TSharedRef<TArray<FProgramCounterSymbolInfo>> Stack = MakeShared<TArray<FProgramCounterSymbolInfo>>(
+			NeedSkipStackWalk ? SkipStackWalk() : StackWalk(IgnoreCount, MaxDepth));
+
 		return Stack;
 	}
 
