Allow ForPlatform FConfigContexts for Plugin ini files (ie named the same as the Plugin) to be read from the Plugin config hierarchy instead of the default hierarchy.

diff --git a/Engine/Source/Runtime/Core/Private/Misc/ConfigCacheIni.cpp b/Engine/Source/Runtime/Core/Private/Misc/ConfigCacheIni.cpp
--- a/Engine/Source/Runtime/Core/Private/Misc/ConfigCacheIni.cpp
+++ b/Engine/Source/Runtime/Core/Private/Misc/ConfigCacheIni.cpp
@@ -4128,8 +4128,30 @@ FConfigFile* FConfigCacheIni::FindOrLoadPlatformConfig(FConfigFile& LocalFile, c
 	FConfigFile* File = FindPlatformConfig(IniName, Platform);
 	if (File == nullptr)
 	{
-		FConfigContext Context = FConfigContext::ReadIntoLocalFile(LocalFile, Platform);
-		Context.Load(IniName);
+#if ALLOW_OTHER_PLATFORM_CONFIG
+		// Check if this ini file corresponds to a plugin
+		FConfigPluginDirs* PluginDirs = nullptr;
+		{
+			FScopeLock Lock(&FConfigContext::ConfigToPluginDirsLock);
+			TUniquePtr<FConfigPluginDirs>* PluginDirsPtr = FConfigContext::ConfigToPluginDirs.Find(IniName);
+			if (PluginDirsPtr != nullptr)
+			{
+				// we never remove the item so the raw pointer will not become invalid
+				PluginDirs = PluginDirsPtr->Get();
+			}
+		}
+		if (PluginDirs != nullptr)
+		{
+			// If so, read using the plugin hierarchy
+			FConfigContext Context = FConfigContext::ReadIntoPluginFile(LocalFile, *PluginDirs->PluginPath, PluginDirs->PluginExtensionBaseDirs, Platform);
+			Context.Load(IniName);
+		}
+		else
+#endif
+		{
+			FConfigContext Context = FConfigContext::ReadIntoLocalFile(LocalFile, Platform);
+			Context.Load(IniName);
+		}
 		File = &LocalFile;
 	}
 
diff --git a/Engine/Source/Runtime/Core/Private/Misc/ConfigContext.cpp b/Engine/Source/Runtime/Core/Private/Misc/ConfigContext.cpp
--- a/Engine/Source/Runtime/Core/Private/Misc/ConfigContext.cpp
+++ b/Engine/Source/Runtime/Core/Private/Misc/ConfigContext.cpp
@@ -23,6 +23,10 @@ namespace
 	const TCHAR* SaveAllSectionsKey = TEXT("bCanSaveAllSections");
 }
 
+#if ALLOW_OTHER_PLATFORM_CONFIG
+TMap<FString, TUniquePtr<FConfigPluginDirs>> FConfigContext::ConfigToPluginDirs;
+FCriticalSection FConfigContext::ConfigToPluginDirsLock;
+#endif
 
 FConfigContext::FConfigContext(FConfigCacheIni* InConfigSystem, bool InIsHierarchicalConfig, const FString& InPlatform, FConfigFile* DestConfigFile)
 	: ConfigSystem(InConfigSystem)
@@ -177,7 +181,18 @@ bool FConfigContext::Load(const TCHAR* InBaseIniName, FString& OutFinalFilename)
 	}
 
 	// now load if we need (PrepareForLoad may find an existing file and just use it)
-	return bPerformLoad ? PerformLoad() : true;
+	bool bSuccess = bPerformLoad ? PerformLoad() : true;
+
+#if ALLOW_OTHER_PLATFORM_CONFIG
+	if (bSuccess && bIsForPlugin && Platform == FPlatformProperties::IniPlatformName())
+	{
+		// We have successfuly loaded a plugin ini file for the main platform. Cache the plugin info in case we want to load this ForPlatform.	
+		FScopeLock Lock(&FConfigContext::ConfigToPluginDirsLock);
+		FConfigContext::ConfigToPluginDirs.Add(InBaseIniName, TUniquePtr<FConfigPluginDirs>(new FConfigPluginDirs(PluginRootDir, ChildPluginBaseDirs)));
+	}
+#endif
+
+	return bSuccess;
 }
 
 bool FConfigContext::Load(const TCHAR* InBaseIniName)
diff --git a/Engine/Source/Runtime/Core/Private/Misc/DataDrivenPlatformInfoRegistry.cpp b/Engine/Source/Runtime/Core/Private/Misc/DataDrivenPlatformInfoRegistry.cpp
--- a/Engine/Source/Runtime/Core/Private/Misc/DataDrivenPlatformInfoRegistry.cpp
+++ b/Engine/Source/Runtime/Core/Private/Misc/DataDrivenPlatformInfoRegistry.cpp
@@ -487,15 +487,17 @@ const TArray<const FDataDrivenPlatformInfo*>& FDataDrivenPlatformInfoRegistry::G
 	return PlatformType == EPlatformInfoType::AllPlatformInfos ? AllSortedPlatformInfos : SortedPlatformInfos;
 }
 
-
-const TArray<FString>& FDataDrivenPlatformInfoRegistry::GetValidPlatformDirectoryNames()
+const TArray<FString>& FDataDrivenPlatformInfoRegistry::GetPlatformDirectoryNames(bool bCheckValid)
 {
-	static bool bHasSearchedForPlatforms = false;
-	static TArray<FString> ValidPlatformDirectories;
+	static bool bHasSearchedForPlatforms[2] = { false, false };
+	static TArray<FString> PlatformDirectories[2];
 
-	if (bHasSearchedForPlatforms == false)
+	// Which PlatformDirectories cache to use
+	int32 CacheIndex = bCheckValid ? 1 : 0;
+
+	if (bHasSearchedForPlatforms[CacheIndex] == false)
 	{
-		bHasSearchedForPlatforms = true;
+		bHasSearchedForPlatforms[CacheIndex] = true;
 
 		// look for possible platforms
 		const TMap<FName, FDataDrivenPlatformInfo>& Infos = GetAllPlatformInfos();
@@ -503,27 +505,26 @@ const TArray<FString>& FDataDrivenPlatformInfoRegistry::GetValidPlatformDirector
 		{
 #if DDPI_HAS_EXTENDED_PLATFORMINFO_DATA
 			// if the editor hasn't compiled in support for the platform, it's not "valid"
-			if (!HasCompiledSupportForPlatform(Pair.Key, EPlatformNameType::Ini))
+			if (bCheckValid && !HasCompiledSupportForPlatform(Pair.Key, EPlatformNameType::Ini))
 			{
 				continue;
 			}
 #endif
 
 			// add ourself as valid
-			ValidPlatformDirectories.AddUnique(Pair.Key.ToString());
+			PlatformDirectories[CacheIndex].AddUnique(Pair.Key.ToString());
 
 			// now add additional directories
 			for (FString& AdditionalDir : Pair.Value.AdditionalRestrictedFolders)
 			{
-				ValidPlatformDirectories.AddUnique(AdditionalDir);
+				PlatformDirectories[CacheIndex].AddUnique(AdditionalDir);
 			}
 		}
 	}
 
-	return ValidPlatformDirectories;
+	return PlatformDirectories[CacheIndex];
 }
 
-
 const FDataDrivenPlatformInfo& FDataDrivenPlatformInfoRegistry::GetPlatformInfo(FName PlatformName)
 {
 	const FDataDrivenPlatformInfo* Info = GetAllPlatformInfos().Find(PlatformName);
diff --git a/Engine/Source/Runtime/Core/Private/Misc/Paths.cpp b/Engine/Source/Runtime/Core/Private/Misc/Paths.cpp
--- a/Engine/Source/Runtime/Core/Private/Misc/Paths.cpp
+++ b/Engine/Source/Runtime/Core/Private/Misc/Paths.cpp
@@ -357,31 +357,31 @@ static void AddIfDirectoryExists(TArray<FString>& ExtensionDirs, FString&& Dir)
 	}
 }
 
-static void GetExtensionDirsInternal(TArray<FString>& ExtensionDirs, const FString& BaseDir, const FString& SubDir)
+static void GetExtensionDirsInternal(TArray<FString>& ExtensionDirs, const FString& BaseDir, const FString& SubDir, bool bCheckValid)
 {
 	AddIfDirectoryExists(ExtensionDirs, FPaths::Combine(BaseDir, SubDir));
 
 	FString PlatformExtensionBaseDir = FPaths::Combine(BaseDir, TEXT("Platforms"));
-	for (const FString& PlatformName : FDataDrivenPlatformInfoRegistry::GetValidPlatformDirectoryNames())
+	for (const FString& PlatformName : FDataDrivenPlatformInfoRegistry::GetPlatformDirectoryNames(bCheckValid))
 	{
 		AddIfDirectoryExists(ExtensionDirs, FPaths::Combine(PlatformExtensionBaseDir, PlatformName, SubDir));
 	}
 
 	FString RestrictedBaseDir = FPaths::Combine(BaseDir, TEXT("Restricted"));
-	IFileManager::Get().IterateDirectory(*RestrictedBaseDir, [&ExtensionDirs, SubDir](const TCHAR* FilenameOrDirectory, bool bIsDirectory)  -> bool
+	IFileManager::Get().IterateDirectory(*RestrictedBaseDir, [&ExtensionDirs, SubDir, bCheckValid](const TCHAR* FilenameOrDirectory, bool bIsDirectory)  -> bool
 	{
 		if (bIsDirectory)
 		{
-			GetExtensionDirsInternal(ExtensionDirs, FilenameOrDirectory, SubDir);
+			GetExtensionDirsInternal(ExtensionDirs, FilenameOrDirectory, SubDir, bCheckValid);
 		}
 		return true;
 	});
 }
 
-TArray<FString> FPaths::GetExtensionDirs(const FString& BaseDir, const FString& SubDir)
+TArray<FString> FPaths::GetExtensionDirs(const FString& BaseDir, const FString& SubDir, bool bCheckValid)
 {
 	TArray<FString> ExtensionDirs;
-	GetExtensionDirsInternal(ExtensionDirs, BaseDir, SubDir);
+	GetExtensionDirsInternal(ExtensionDirs, BaseDir, SubDir, bCheckValid);
 	return ExtensionDirs;
 }
 
diff --git a/Engine/Source/Runtime/Core/Public/Misc/ConfigContext.h b/Engine/Source/Runtime/Core/Public/Misc/ConfigContext.h
--- a/Engine/Source/Runtime/Core/Public/Misc/ConfigContext.h
+++ b/Engine/Source/Runtime/Core/Public/Misc/ConfigContext.h
@@ -21,7 +21,21 @@
 #endif
 
 
+#if ALLOW_OTHER_PLATFORM_CONFIG
 
+// Struct to hold the plugin base directory and any child plugin base directories, used for the ConfigToPluginDirs map.
+struct CORE_API FConfigPluginDirs
+{
+	FString PluginPath;
+	TArray<FString> PluginExtensionBaseDirs;
+
+	FConfigPluginDirs(const FString& InPluginPath, const TArray<FString>& InPluginExtensionBaseDirs)
+	:	PluginPath(InPluginPath)
+	,	PluginExtensionBaseDirs(InPluginExtensionBaseDirs)
+	{}
+};
+
+#endif
 
 class FConfigCacheIni;
 class FConfigFile;
@@ -92,9 +106,9 @@ public:
 	/**
 	 * Create a context to read a plugin's ini file named for the plugin. This is not used for inserting, say, Engine.ini into GConfig
 	 */
-	static FConfigContext ReadIntoPluginFile(FConfigFile& DestConfigFile, const FString& PluginRootDir, const TArray<FString>& ChildPluginsBaseDirs)
+	static FConfigContext ReadIntoPluginFile(FConfigFile& DestConfigFile, const FString& PluginRootDir, const TArray<FString>& ChildPluginsBaseDirs, const FString& Platform = FString())
 	{
-		FConfigContext Context(nullptr, true, FString(), &DestConfigFile);
+		FConfigContext Context(nullptr, true, Platform, &DestConfigFile);
 		Context.bIsForPlugin = true;
 		Context.PluginRootDir = PluginRootDir;
 		Context.ChildPluginBaseDirs = ChildPluginsBaseDirs;
@@ -191,6 +205,13 @@ public:
 	// if this is non-null, it contains a set of pre-scanned ini files to use to find files, instead of looking on disk
 	const TSet<FString>* IniCacheSet = nullptr;
 
+#if ALLOW_OTHER_PLATFORM_CONFIG
+	// Map of Plugin config file name to plugin and child directories, for filled in by FPluginManager.ConfigureEnabledPlugins
+	// Used creating a ForPlatform FConfigContext for the Plugin ini.
+	static TMap<FString, TUniquePtr<FConfigPluginDirs>> ConfigToPluginDirs;
+	static FCriticalSection ConfigToPluginDirsLock;
+#endif
+
 protected:
 
 	bool bDoNotResetConfigFile = false;
diff --git a/Engine/Source/Runtime/Core/Public/Misc/DataDrivenPlatformInfoRegistry.h b/Engine/Source/Runtime/Core/Public/Misc/DataDrivenPlatformInfoRegistry.h
--- a/Engine/Source/Runtime/Core/Public/Misc/DataDrivenPlatformInfoRegistry.h
+++ b/Engine/Source/Runtime/Core/Public/Misc/DataDrivenPlatformInfoRegistry.h
@@ -270,11 +270,19 @@ struct FDataDrivenPlatformInfoRegistry
 	*/
 	static CORE_API const TMap<FName, FDataDrivenPlatformInfo>& GetAllPlatformInfos();
 
+	/**
+	 * Gets a set of platform names based on GetAllPlatformInfos, their AdditionalRestrictedFolders
+	 * This is not necessarily the same as IniParents, although there is overlap - IniParents come from chaining DDPIs, so those will be in GetAllPlatformInfos already to be checked 
+	 *
+	 * @param bCheckValid	If true, the result is filtered based on what editor has support compiled for
+	 */
+	static CORE_API const TArray<FString>& GetPlatformDirectoryNames(bool bCheckValid);
+
 	/**
 	 * Gets a set of platform names based on GetAllPlatformInfos, their AdditionalRestrictedFolders, and possibly filtered based on what editor has support compiled for
 	 * This is not necessarily the same as IniParents, although there is overlap - IniParents come from chaining DDPIs, so those will be in GetAllPlatformInfos already to be checked 
 	 */
-	static CORE_API const TArray<FString>& GetValidPlatformDirectoryNames();
+	static CORE_API const TArray<FString>& GetValidPlatformDirectoryNames() { return GetPlatformDirectoryNames(true); }
 
 	/**
 	 * Get the data driven platform info for a given platform. If the platform doesn't have any on disk,
diff --git a/Engine/Source/Runtime/Core/Public/Misc/Paths.h b/Engine/Source/Runtime/Core/Public/Misc/Paths.h
--- a/Engine/Source/Runtime/Core/Public/Misc/Paths.h
+++ b/Engine/Source/Runtime/Core/Public/Misc/Paths.h
@@ -209,11 +209,12 @@ public:
 	static CORE_API FString ConvertPath(const FString& Path, EPathConversion Method, const TCHAR* ExtraData=nullptr, const TCHAR* OverrideProjectDir=nullptr);
 	
 	/**
-	 * Returns platform and restricted extensions that are present and valid (for platforms, it uses FDataDrivePlatformInfo to determine valid platforms, it doesn't just use what's present)
+	 * Returns platform and restricted extensions that are present and if bCheckValid is set, valid
+     * (for platforms, it uses FDataDrivePlatformInfo to determine valid platforms, it doesn't just use what's present)
 	 *
 	 * @return BaseDir and usable extension directories under BaseDir (either Engine or Project)
 	 */
-	static CORE_API TArray<FString> GetExtensionDirs(const FString& BaseDir, const FString& SubDir=FString());
+	static CORE_API TArray<FString> GetExtensionDirs(const FString& BaseDir, const FString& SubDir=FString(), bool bCheckValid=true);
 
 	/**
 	 * Returns the root directory of the engine directory tree
diff --git a/Engine/Source/Runtime/Projects/Private/PluginManager.cpp b/Engine/Source/Runtime/Projects/Private/PluginManager.cpp
--- a/Engine/Source/Runtime/Projects/Private/PluginManager.cpp
+++ b/Engine/Source/Runtime/Projects/Private/PluginManager.cpp
@@ -784,7 +784,7 @@ void FPluginManager::ReadAllPlugins(FDiscoveredPluginMap& Plugins, const TSet<FS
 
 		SlowTask_ReadAll.EnterProgressFrame(0.5f);
 		// Find "built-in" plugins.  That is, plugins situated right within the Engine directory.
-		TArray<FString> EnginePluginDirs = FPaths::GetExtensionDirs(FPaths::EngineDir(), TEXT("Plugins"));
+		TArray<FString> EnginePluginDirs = FPaths::GetExtensionDirs(FPaths::EngineDir(), TEXT("Plugins"), !GIsEditor);
 		{
 			FScopedSlowTask SlowTask_ReadEngine((float)EnginePluginDirs.Num());
 			
@@ -800,7 +800,7 @@ void FPluginManager::ReadAllPlugins(FDiscoveredPluginMap& Plugins, const TSet<FS
 		// assume that the game plugin version is preferred.
 		if (Project != nullptr)
 		{
-			TArray<FString> ProjectPluginDirs = FPaths::GetExtensionDirs(FPaths::GetPath(FPaths::GetProjectFilePath()), TEXT("Plugins"));
+			TArray<FString> ProjectPluginDirs = FPaths::GetExtensionDirs(FPaths::GetPath(FPaths::GetProjectFilePath()), TEXT("Plugins"), !GIsEditor);
 			FScopedSlowTask SlowTask_ReadProject((float)ProjectPluginDirs.Num());
 
 			for (const FString& ProjectPluginDir : ProjectPluginDirs)
@@ -1646,12 +1646,12 @@ bool FPluginManager::ConfigureEnabledPlugins()
 			{
 				TArray<FString> PhysicalFileNames;
 
-				for (const FString& EnginePluginDir : FPaths::GetExtensionDirs(EngineBinariesRootDir, TEXT("Plugins")))
+				for (const FString& EnginePluginDir : FPaths::GetExtensionDirs(EngineBinariesRootDir, TEXT("Plugins"), !GIsEditor))
 				{
 					FindPluginsInDirectory(EnginePluginDir, PhysicalFileNames, IPlatformFile::GetPlatformPhysical());
 				}
 
-				for (const FString& ProjectPluginDir : FPaths::GetExtensionDirs(ProjectBinariesRootDir, TEXT("Plugins")))
+				for (const FString& ProjectPluginDir : FPaths::GetExtensionDirs(ProjectBinariesRootDir, TEXT("Plugins"), !GIsEditor))
 				{
 					FindPluginsInDirectory(ProjectPluginDir, PhysicalFileNames, IPlatformFile::GetPlatformPhysical());
 				}
