- CVar History functionality, mainly so plugins can add/remove cvars and the Value will be updated correctly. However, anything can make use of the history to "unset" a cvar at a certain level and the value will be updated. Note that the full plugin functionality is coming in a future changelist with FConfigBranch support
- Added concept of "array" cvar SetBy priorities, where the value is set at the end of the list, so multiple plugins can set the value, and can be unloaded later
- Added some new priorities (PluginLowPriority and PluginHighPriority, above and below DeviceProfile prio)
- Added SetCVar, UnsetCVar, and UnSetCVarTag console commands that can set/unset via SetBy priority and Tag

diff --git a/Engine/Source/Runtime/Core/Private/HAL/ConsoleManager.cpp b/Engine/Source/Runtime/Core/Private/HAL/ConsoleManager.cpp
--- a/Engine/Source/Runtime/Core/Private/HAL/ConsoleManager.cpp
+++ b/Engine/Source/Runtime/Core/Private/HAL/ConsoleManager.cpp
@@ -10,18 +10,31 @@ ConsoleManager.cpp: console command handling
 #include "Misc/Paths.h"
 #include "Stats/Stats.h"
 #include "Misc/ConfigCacheIni.h"
+#include "Misc/ConfigContext.h"
 #include "Misc/ConfigUtilities.h"
 #include "Misc/CoreDelegates.h"
 #include "Misc/OutputDeviceFile.h"
+#include "Misc/RemoteConfigIni.h"
 #include "Modules/ModuleManager.h"
 #include "HAL/PlatformProcess.h"
-#include "Misc/RemoteConfigIni.h"
 #include "ProfilingDebugging/CsvProfiler.h"
 #include "HAL/FileManager.h"
+#include "Serialization/ArchiveCountMem.h"
 
 DEFINE_LOG_CATEGORY(LogConsoleResponse);
 DEFINE_LOG_CATEGORY_STATIC(LogConsoleManager, Log, All);
 
+#define UE_ALLOW_CVAR_HISTORY 1
+
+#if UE_ALLOW_CVAR_HISTORY
+namespace UE::ConsoleManager::Private
+{
+	// this tracks the cvars that were  added dynamically with a tag (via plugin or similar)
+	// we use this structure to unset the cvars and update the value when the plugin unloads
+	TMap<FName, TSet<IConsoleVariable*>*> TaggedCVars;
+}
+#endif
+
 static inline bool IsWhiteSpace(TCHAR Value) { return Value == TCHAR(' '); }
 
 // @param In must not be 0
@@ -53,26 +66,26 @@ const TCHAR* GetConsoleVariableSetByName(EConsoleVariableFlags ConsoleVariableFl
 {
 	EConsoleVariableFlags SetBy = (EConsoleVariableFlags)((uint32)ConsoleVariableFlags & ECVF_SetByMask);
 
+	#define CASE(x) case ECVF_SetBy##x: return TEXT(#x);
 	switch(SetBy)
 	{
-#define CASE(A) case ECVF_SetBy##A: return TEXT(#A);
-		// Could also be done with enum reflection instead
-		CASE(Constructor)
-		CASE(Scalability)
-		CASE(GameSetting)
-		CASE(ProjectSetting)
-		CASE(SystemSettingsIni)
-		CASE(DeviceProfile)
-		CASE(ConsoleVariablesIni)
-		CASE(Hotfix)
-		CASE(Commandline)
-		CASE(Code)
-		CASE(Console)
-#undef CASE
+		ENUMERATE_SET_BY(CASE)
 	}
+	#undef CASE
+	
 	return TEXT("<UNKNOWN>");
 }
 
+EConsoleVariableFlags GetConsoleVariableSetByValue(const TCHAR* SetByName)
+{
+	#define TEST(x) if (FCString::Stricmp(SetByName, TEXT(#x)) == 0) { return ECVF_SetBy##x; }
+	ENUMERATE_SET_BY(TEST)
+	#undef TEST
+	
+	return ECVF_SetByMask;
+}
+
+
 TArray<const FAutoConsoleObject*>& FAutoConsoleObject::AccessGeneralShaderChangeCvars()
 {
 	// This variable cannot be global because it is accessed from the constructor of other global variables.
@@ -148,6 +161,16 @@ public:
 		return this;
 	}
 	
+	virtual void LogHistory(FOutputDevice& Ar)
+	{
+	
+	}
+
+	virtual SIZE_T GetHistorySize()
+	{
+		return 0;
+	}
+
 	/** Legacy funciton to add old single delegates to the new multicast delegate. */
 	virtual void SetOnChangedCallback(const FConsoleVariableDelegate& Callback) 
 	{
@@ -203,7 +226,7 @@ public:
 		return bRet;
 	}
 
-	void OnChanged(EConsoleVariableFlags SetBy)
+	void OnChanged(EConsoleVariableFlags SetBy, bool bForce)
 	{
 		// SetBy can include set flags. Discard them here
 		SetBy = EConsoleVariableFlags(SetBy & ~ECVF_SetFlagMask);
@@ -212,7 +235,7 @@ public:
 		check(((uint32)SetBy & ECVF_SetByMask) || SetBy == ECVF_Default);
 
 		// double check, if this fires we miss a if(CanChange(SetBy))
-		check(CanChange(SetBy));
+		check(bForce || CanChange(SetBy));
 
 		// only change on main thread
 
@@ -626,6 +649,122 @@ static bool GetConfigValueFromRuntimeSources(FName PlatformName, IConsoleVariabl
 #endif
 
 
+constexpr bool IsArrayPriority(EConsoleVariableFlags Priority)
+{
+	Priority = (EConsoleVariableFlags)(Priority & ECVF_SetByMask);
+	return
+		Priority == ECVF_SetByPluginLowPriority ||
+		Priority == ECVF_SetByPluginHighPriority ||
+		Priority == ECVF_SetByHotfix;
+}
+
+template <class T>
+class FConsoleVariableHistory
+{
+public:
+	using FHistoryData = TConsoleVariableData<T>;
+	using FTaggedHistoryData = TPair<FName, FHistoryData>;
+	TSortedMap<int, TArray<FTaggedHistoryData>> History;
+	
+	bool bHasTaggedArrayData = false;
+	
+	void Track(const TCHAR* InValue, EConsoleVariableFlags SetBy, FName Tag)
+	{
+		T LocalCopy;
+		TTypeFromString<T>::FromString(LocalCopy, InValue);
+		
+		int Priority = (int)SetBy;
+		TArray<FTaggedHistoryData>& ValueArray = History.FindOrAdd(Priority);
+		if (!IsArrayPriority(SetBy))
+		{
+			if (ValueArray.Num() > 0)
+			{
+				// don't do anything if we already are set
+				if (ValueArray[0].Key == Tag && ValueArray[0].Value.GetValueOnGameThread() == LocalCopy)
+				{
+					return;
+				}
+				//remove the entry but leave space
+				ValueArray.Empty(1);
+			}
+		}
+		// add the value to the array (even if it's a non-array type, we store the value in the array)
+		ValueArray.Emplace(Tag, LocalCopy);
+	}
+	
+	void Unset(EConsoleVariableFlags SetBy, FName Tag)
+	{
+		if (!History.Contains(SetBy))
+		{
+			return;
+		}
+		
+		TArray<FTaggedHistoryData>& ValueArray = History[SetBy];
+
+		// first remove the value (if no tag, remove all of them)
+		if (IsArrayPriority(SetBy) && Tag != NAME_None)
+		{
+			// look for the tag in the list
+			for (auto It = ValueArray.CreateIterator(); It; ++It)
+			{
+				if (It->Key == Tag)
+				{
+					It.RemoveCurrent();
+					break;
+				}
+			}
+			
+			// toss it if it's now empty
+			if (ValueArray.IsEmpty())
+			{
+				History.Remove(SetBy);
+			}
+		}
+		else
+		{
+			History.Remove(SetBy);
+		}
+		
+	}
+	
+	const FHistoryData& GetMaxValue(EConsoleVariableFlags& MaxSetBy)
+	{
+		for (typename TSortedMap<int, TArray<FTaggedHistoryData>>::TReverseIterator It(History); It; ++It)
+		{
+			if (It.Value().Num() > 0)
+			{
+				MaxSetBy = (EConsoleVariableFlags)It.Key();
+				return It.Value().Last().Value;
+			}
+		}
+		
+		// indicate failed to find one
+		MaxSetBy = ECVF_SetByMask;
+		static FHistoryData Default(T{});
+		return Default;
+	}
+	
+	void Log(FOutputDevice& Ar)
+	{
+		Ar.Logf(TEXT("\nHISTORY"));
+		int Prio = 0;
+		for (auto& Pri : History)
+		{
+			for (FTaggedHistoryData& Pair : Pri.Value)
+			{
+				FString KeyStr = FString();
+				
+				Ar.Logf(TEXT("%s: %s%s"),
+					GetConsoleVariableSetByName((EConsoleVariableFlags)(Pri.Key)),
+					*TTypeToString<T>::ToString(Pair.Value.GetValueOnGameThread()),
+					(Pair.Key != NAME_None) ? *FString::Printf(TEXT(" [%s]"), *Pair.Key.ToString()) : TEXT(""));
+			}
+			Prio++;
+		}
+	}
+};
+
+
 // an intermediate class between specific typed CVars and FConsoleVariableBase to handle looking up other platform's cvars and caching them
 // if ALLOW_OTHER_PLATFORM_CONFIG is 0, then this is a pass-through class that does nothing
 template <class T>
@@ -635,6 +774,23 @@ public:
 
 	FOtherPlatformValueHelper(const T& DefaultValue, const TCHAR* Help, EConsoleVariableFlags Flags, bool bSaveDefault);
 
+#if UE_ALLOW_CVAR_HISTORY
+	virtual ~FOtherPlatformValueHelper()
+	{
+		// if we had been put into any tagged
+		if (PriorityHistory != nullptr && PriorityHistory->bHasTaggedArrayData)
+		{
+			// make sure we are pulled out of the global list of tagged cvars
+			for (auto Pair : UE::ConsoleManager::Private::TaggedCVars)
+			{
+				Pair.Value->Remove(this);
+			}
+			
+			delete PriorityHistory;
+		}
+	}
+#endif
+	
 #if ALLOW_OTHER_PLATFORM_CONFIG
 	// remember the default value of this cvar, before anything else can assign to it - this way we can know in the editor what the value
 	// would be for a cvar on another platform, if no .ini file on that platform sets it
@@ -656,6 +812,103 @@ protected:
 	FRWLock PlatformValuesLock;
 
 #endif
+
+	virtual void SetInternal(const T& TypedValue, EConsoleVariableFlags SetBy) = 0;
+
+
+#if UE_ALLOW_CVAR_HISTORY
+	friend class FConsoleManager;
+	FConsoleVariableHistory<T>* PriorityHistory = nullptr;
+#endif
+	
+	void TrackHistory(const TCHAR* InValue, EConsoleVariableFlags SetBy, FName Tag)
+	{
+#if UE_ALLOW_CVAR_HISTORY
+		// make a history if we want to
+		if (PriorityHistory == nullptr)
+		{
+			PriorityHistory = new FConsoleVariableHistory<T>();
+			PriorityHistory->Track(*GetString(), ECVF_SetByConstructor, NAME_None);
+		}
+		PriorityHistory->Track(InValue, SetBy, Tag);
+
+		// this cvar so we can remove ourself later
+		if (Tag != NAME_None)
+		{
+			TSet<IConsoleVariable*>* TaggedSet = UE::ConsoleManager::Private::TaggedCVars.FindRef(Tag);
+			if (TaggedSet == nullptr)
+			{
+				TaggedSet = UE::ConsoleManager::Private::TaggedCVars.Add(Tag, new TSet<IConsoleVariable*>());
+			}
+			TaggedSet->Add(this);
+			
+			// set a flag to remember we need to remove ourself from TaggedCVars in destructor
+			PriorityHistory->bHasTaggedArrayData = true;
+		}
+#endif
+	}
+	
+
+public:
+	
+#if UE_ALLOW_CVAR_HISTORY
+	virtual SIZE_T GetHistorySize() override
+	{
+		if (PriorityHistory != nullptr)
+		{
+			return PriorityHistory->History.GetAllocatedSize();
+		}
+		return 0;
+	}
+
+	void LogHistory(FOutputDevice& Ar)
+	{
+		if (PriorityHistory != nullptr)
+		{
+			PriorityHistory->Log(Ar);
+		}
+	}
+#endif
+
+	
+	virtual void Unset(EConsoleVariableFlags SetBy, FName Tag) override
+	{
+#if UE_ALLOW_CVAR_HISTORY
+		
+		if (PriorityHistory == nullptr)
+		{
+			return;
+		}
+		
+		// if we don't know the SetBy, unset from all
+		// this isn't ideal because it could call SetInternal multiple times
+		if (SetBy == ECVF_SetByMask)
+		{
+			#define RECURSE(x) if (IsArrayPriority(ECVF_SetBy##x)) { Unset(ECVF_SetBy##x, Tag); }
+			ENUMERATE_SET_BY(RECURSE)
+			#undef RECURSE
+			
+			return;
+		}
+		
+		uint32 CurrentPri =	(uint32)this->Flags & ECVF_SetByMask;
+		uint32 UnsetPri =	(uint32)SetBy & ECVF_SetByMask;
+
+		PriorityHistory->Unset(SetBy, Tag);
+		
+		// if we are unsetting at the current setby (or maybe in some weird cases, greater than setby) then we need to reset the SetBy and current value
+		if (UnsetPri >= CurrentPri)
+		{
+			// now figure out the new value
+			EConsoleVariableFlags NewSetBy;
+			auto MaxValue = PriorityHistory->GetMaxValue(NewSetBy);
+			
+			// and force it to the new value
+			SetInternal(MaxValue.GetValueOnGameThread(), NewSetBy);
+		}
+#endif
+	}
+
 };
 
 template<class T>
@@ -769,13 +1022,15 @@ public:
 	{
 		delete this; 
 	}
-
-	virtual void Set(const TCHAR* InValue, EConsoleVariableFlags SetBy)
+	
+	virtual void Set(const TCHAR* InValue, EConsoleVariableFlags SetBy, FName Tag)
 	{
+		this->TrackHistory(InValue, SetBy, Tag);
+
 		if (CanChange(SetBy))
 		{
 			TTypeFromString<T>::FromString(Data.ShadowedValue[0], InValue);
-			OnChanged(SetBy);
+			OnChanged(SetBy, false);
 		}
 	}
 
@@ -804,12 +1059,18 @@ private: // ----------------------------------------------------
 		FConsoleVariable<T>* This = (FConsoleVariable<T>*)this;
 		return This->Data.GetReferenceOnAnyThread();
 	}
+	
+	virtual void SetInternal(const T& TypedValue, EConsoleVariableFlags SetBy)
+	{
+		Data.ShadowedValue[0] = TypedValue;
+		OnChanged(SetBy, true);
+	}
 
-	void OnChanged(EConsoleVariableFlags SetBy)
+	virtual void OnChanged(EConsoleVariableFlags SetBy, bool bForce=false)
 	{
 		// propagate from main thread to render thread
 		OnCVarChange(Data.ShadowedValue[1], Data.ShadowedValue[0], Flags, SetBy);
-		FConsoleVariableBase::OnChanged(SetBy);
+		FConsoleVariableBase::OnChanged(SetBy, bForce);
 	}
 };
 
@@ -856,8 +1117,10 @@ template<> TConsoleVariableData<float>* FConsoleVariable<float>::AsVariableFloat
 
 // specialization for FString
 
-template<> void FConsoleVariable<FString>::Set(const TCHAR* InValue, EConsoleVariableFlags SetBy)
+template<> void FConsoleVariable<FString>::Set(const TCHAR* InValue, EConsoleVariableFlags SetBy, FName Tag)
 {
+	TrackHistory(InValue, SetBy, Tag);
+	
 	if(CanChange(SetBy))
 	{
 		Data.ShadowedValue[0] = InValue;
@@ -961,8 +1224,10 @@ public:
 		delete this; 
 	}
 
-	virtual void Set(const TCHAR* InValue, EConsoleVariableFlags SetBy)
+	virtual void Set(const TCHAR* InValue, EConsoleVariableFlags SetBy, FName Tag)
 	{
+		this->TrackHistory(InValue, SetBy, Tag);
+		
 		if(CanChange(SetBy))
 		{
 			TTypeFromString<T>::FromString(MainValue, InValue);
@@ -993,14 +1258,20 @@ private: // ----------------------------------------------------
 		checkSlow(Index < 2);
 		return (Index == 0) ? MainValue : RefValue;
 	}
+	
+	virtual void SetInternal(const T& TypedValue, EConsoleVariableFlags SetBy)
+	{
+		MainValue = TypedValue;
+		OnChanged(SetBy, true);
+	}
 
-	void OnChanged(EConsoleVariableFlags SetBy)
+	void OnChanged(EConsoleVariableFlags SetBy, bool bForce=false)
 	{
 		if(CanChange(SetBy))
 		{
 			// propagate from main thread to render thread or to reference
 			OnCVarChange(RefValue, MainValue, Flags, SetBy);
-			FConsoleVariableBase::OnChanged(SetBy);
+			FConsoleVariableBase::OnChanged(SetBy, bForce);
 		}
 	}
 };
@@ -1026,11 +1297,11 @@ bool FConsoleVariableRef<float>::IsVariableFloat() const
 
 // string version
 
-class FConsoleVariableStringRef : public FConsoleVariableBase
+class FConsoleVariableStringRef : public FOtherPlatformValueHelper<FString>
 {
 public:
 	FConsoleVariableStringRef(FString& InRefValue, const TCHAR* Help, EConsoleVariableFlags Flags)
-		: FConsoleVariableBase(Help, Flags)
+		: FOtherPlatformValueHelper<FString>(FString(), Help, Flags, false)
 		, RefValue(InRefValue)
 		, MainValue(InRefValue)
 	{
@@ -1042,8 +1313,10 @@ public:
 	{
 		delete this;
 	}
-	virtual void Set(const TCHAR* InValue, EConsoleVariableFlags SetBy)
+	virtual void Set(const TCHAR* InValue, EConsoleVariableFlags SetBy, FName Tag)
 	{
+		TrackHistory(InValue, SetBy, Tag);
+		
 		if (CanChange(SetBy))
 		{
 			MainValue = InValue;
@@ -1091,22 +1364,28 @@ private: // ----------------------------------------------------
 		return (Index == 0) ? MainValue : RefValue;
 	}
 
-	void OnChanged(EConsoleVariableFlags SetBy)
+	virtual void SetInternal(const FString& TypedValue, EConsoleVariableFlags SetBy)
+	{
+		MainValue = TypedValue;
+		OnChanged(SetBy);
+	}
+	
+	void OnChanged(EConsoleVariableFlags SetBy, bool bForce=false)
 	{
 		if (CanChange(SetBy))
 		{
 			// propagate from main thread to render thread or to reference
 			OnCVarChange(RefValue, MainValue, Flags, SetBy);
-			FConsoleVariableBase::OnChanged(SetBy);
+			FConsoleVariableBase::OnChanged(SetBy, bForce);
 		}
 	}
 };
 
-class FConsoleVariableBitRef : public FConsoleVariableBase
+class FConsoleVariableBitRef : public FOtherPlatformValueHelper<int>
 {
 public:
 	FConsoleVariableBitRef(const TCHAR* FlagName, uint32 InBitNumber, uint8* InForce0MaskPtr, uint8* InForce1MaskPtr, const TCHAR* Help, EConsoleVariableFlags Flags) 
-		: FConsoleVariableBase(Help, Flags), Force0MaskPtr(InForce0MaskPtr), Force1MaskPtr(InForce1MaskPtr), BitNumber(InBitNumber)
+		: FOtherPlatformValueHelper<int>(0, Help, Flags, false), Force0MaskPtr(InForce0MaskPtr), Force1MaskPtr(InForce1MaskPtr), BitNumber(InBitNumber)
 	{
 	}
 
@@ -1116,8 +1395,10 @@ public:
 	{
 		delete this; 
 	} 
-	virtual void Set(const TCHAR* InValue, EConsoleVariableFlags SetBy)
+	virtual void Set(const TCHAR* InValue, EConsoleVariableFlags SetBy, FName Tag)
 	{
+		TrackHistory(InValue, SetBy, Tag);
+		
 		if(CanChange(SetBy))
 		{
 			int32 Value = FCString::Atoi(InValue);
@@ -1127,7 +1408,7 @@ public:
 			FMath::SetBoolInBitField(Force0MaskPtr, BitNumber, Value == 0);
 			FMath::SetBoolInBitField(Force1MaskPtr, BitNumber, Value == 1);
 
-			OnChanged(SetBy);
+			OnChanged(SetBy, false);
 		}
 	}
 	virtual bool GetBool() const
@@ -1164,6 +1445,14 @@ private: // ----------------------------------------------------
 	uint8* Force0MaskPtr;
 	uint8* Force1MaskPtr;
 	uint32 BitNumber;
+	
+	virtual void SetInternal(const int& TypedValue, EConsoleVariableFlags SetBy)
+	{
+		FMath::SetBoolInBitField(Force0MaskPtr, BitNumber, TypedValue == 0);
+		FMath::SetBoolInBitField(Force1MaskPtr, BitNumber, TypedValue == 1);
+		OnChanged(SetBy, true);
+	}
+	
 };
 
 IConsoleVariable* FConsoleManager::RegisterConsoleVariableBitRef(const TCHAR* CVarName, const TCHAR* FlagName, uint32 BitNumber, uint8* Force0MaskPtr, uint8* Force1MaskPtr, const TCHAR* Help, uint32 Flags)
@@ -1933,6 +2222,98 @@ static void DumpObjects(const TMap<FString, IConsoleObject*>& ConsoleObjects, co
 	delete CSV;
 }
 
+static void SetUnsetCVar(const TMap<FString, IConsoleObject*>& ConsoleObjects, const TCHAR* Params, FOutputDevice& Ar, bool bSet)
+{
+	FString CVarName = FParse::Token(Params, false);
+
+	if (CVarName.Len() == 0)
+	{
+		if (bSet)
+		{
+			Ar.Logf(TEXT("Usage: SetCVar [Platform@]CVarName Value [-setby=Priority] [-tag=SomeTag]"));
+		}
+		else
+		{
+			Ar.Logf(TEXT("Usage: UnsetCVar [Platform@]CVarName [-setby=Priority] [-tag=SomeTag]"));
+		}
+		Ar.Logf(TEXT("   Priority can be one of the following (default is Console):"));
+		
+		#define LOGOP(x) Ar.Logf(TEXT("      %s%s"), TEXT(#x), IsArrayPriority(ECVF_SetBy##x) ? TEXT(" [*]") : TEXT(""));
+		ENUMERATE_SET_BY(LOGOP)
+		#undef LOGOP
+		Ar.Logf(TEXT("      [*] Array type priorities, used for dynamic setting/unsetting"));
+		Ar.Logf(TEXT("   Tag should be set for the ones marked as Array types, for ability to set and unset"));
+
+	}
+	
+	FString PlatformName;
+	int32 PlatformDelim = CVarName.Find(TEXT("@"));
+	if (PlatformDelim > 0)
+	{
+		PlatformName = CVarName.Mid(0, PlatformDelim);
+		CVarName = CVarName.Mid(PlatformDelim + 1);
+	}
+
+	IConsoleObject* Obj = ConsoleObjects.FindRef(CVarName);
+	IConsoleVariable* CVar = Obj ? Obj->AsVariable() : nullptr;
+	
+	if (CVar == nullptr)
+	{
+		Ar.Logf(TEXT("No CVar named %s"), *CVarName);
+		return;
+	}
+	
+#if ALLOW_OTHER_PLATFORM_CONFIG
+	// get platform version
+	if (PlatformName.Len())
+	{
+		CVar = CVar->GetPlatformValueVariable(*PlatformName).Get();
+		if (CVar == nullptr)
+		{
+			Ar.Logf(TEXT("Failed to get CVar for platform %s"), *PlatformName);
+			return;
+		}
+	}
+#else
+	Ar.Logf(TEXT("Unable to lookup a CVar value on another platform in this build"));
+	return;
+#endif
+
+	
+	FString Value;
+	if (bSet)
+	{
+		Value = FParse::Token(Params, false);
+	}
+	
+	EConsoleVariableFlags SetBy = ECVF_SetByConsole;
+	FName Tag = NAME_None;
+	FString Str;
+	if (FParse::Value(Params, TEXT("-setby="), Str))
+	{
+		SetBy = GetConsoleVariableSetByValue(*Str);
+	}
+	if (FParse::Value(Params, TEXT("-tag="), Str))
+	{
+		Tag = *Str;
+	}
+
+	if (bSet)
+	{
+		CVar->Set(*Value, SetBy, Tag);
+	}
+	else
+	{
+		CVar->Unset(SetBy, Tag);
+	}
+}
+
+void UnsetCVarTag(const TCHAR* Params, FOutputDevice& Ar)
+{
+	FString TagName = FParse::Token(Params, false);
+	IConsoleManager::Get().UnsetAllConsoleVariablesWithTag(*TagName);
+}
+
 bool FConsoleManager::ProcessUserConsoleInput(const TCHAR* InInput, FOutputDevice& Ar, UWorld* InWorld)
 {
 	check(InInput);
@@ -1948,6 +2329,21 @@ bool FConsoleManager::ProcessUserConsoleInput(const TCHAR* InInput, FOutputDevic
 		DumpObjects(ConsoleObjects, InInput, Ar, true);
 		return true;
 	}
+	if (FParse::Command(&InInput, TEXT("setcvar")))
+	{
+		SetUnsetCVar(ConsoleObjects, InInput, Ar, true);
+		return true;
+	}
+	if (FParse::Command(&InInput, TEXT("unsetcvar")))
+	{
+		SetUnsetCVar(ConsoleObjects, InInput, Ar, false);
+		return true;
+	}
+	if (FParse::Command(&InInput, TEXT("unsetcvartag")))
+	{
+		UnsetCVarTag(InInput, Ar);
+		return true;
+	}
 
 	const TCHAR* It = InInput;
 
@@ -2110,6 +2506,8 @@ bool FConsoleManager::ProcessUserConsoleInput(const TCHAR* InInput, FOutputDevic
 
 		if(bShowCurrentState)
 		{
+			((FConsoleVariableBase*)CVar)->LogHistory(Ar);
+
 			Ar.Logf(TEXT("%s = \"%s\"      LastSetBy: %s"), *Param1, *CVar->GetString(), GetConsoleVariableSetByName(CVar->GetFlags()));
 		}
 	}
@@ -2432,6 +2830,25 @@ FConsoleVariableMulticastDelegate& FConsoleManager::OnCVarUnregistered()
 	return ConsoleVariableUnregisteredDelegate;
 }
 
+void FConsoleManager::UnsetAllConsoleVariablesWithTag(FName Tag)
+{
+#if UE_ALLOW_CVAR_HISTORY
+	TSet<IConsoleVariable*>* TaggedSet = UE::ConsoleManager::Private::TaggedCVars.FindRef(Tag);
+	if (TaggedSet == nullptr)
+	{
+		return;
+	}
+	
+	for (IConsoleVariable* Var : *TaggedSet)
+	{
+		Var->Unset(EConsoleVariableFlags::ECVF_SetByMask);
+	}
+	
+	UE::ConsoleManager::Private::TaggedCVars.Remove(Tag);
+#endif
+}
+
+
 #if !(UE_BUILD_SHIPPING || UE_BUILD_TEST)
 uint32 GConsoleManagerSinkTestCounter = 0;
 void TestSinkCallback()
diff --git a/Engine/Source/Runtime/Core/Public/HAL/ConsoleManager.h b/Engine/Source/Runtime/Core/Public/HAL/ConsoleManager.h
--- a/Engine/Source/Runtime/Core/Public/HAL/ConsoleManager.h
+++ b/Engine/Source/Runtime/Core/Public/HAL/ConsoleManager.h
@@ -86,6 +86,7 @@ public:
 	CORE_API virtual bool IsNameRegistered(const TCHAR* Name) const override;	
 	CORE_API virtual void RegisterThreadPropagation(uint32 ThreadId, IConsoleThreadPropagation* InCallback) override;
 	CORE_API virtual void UnregisterConsoleObject( IConsoleObject* Object, bool bKeepState) override;
+	CORE_API virtual void UnsetAllConsoleVariablesWithTag(FName Tag) override;
 
 private: // ----------------------------------------------------
 
diff --git a/Engine/Source/Runtime/Core/Public/HAL/IConsoleManager.h b/Engine/Source/Runtime/Core/Public/HAL/IConsoleManager.h
--- a/Engine/Source/Runtime/Core/Public/HAL/IConsoleManager.h
+++ b/Engine/Source/Runtime/Core/Public/HAL/IConsoleManager.h
@@ -142,28 +142,52 @@ enum EConsoleVariableFlags
 	ECVF_SetByProjectSetting =		0x03000000,
 	// Used by the [ConsoleVariables] section of Engine.ini as well as FSystemSettings
 	ECVF_SetBySystemSettingsIni =	0x04000000,
+	// Dyanmically loaded/unloaded plugins, used with the History concept to restore cvars on plugin unload. THIS IS AN ARRAY TYPE meaning multiple vales with this SetBy are stored in the history.
+	ECVF_SetByPluginLowPriority =	0x05000000,
 	// Per device settings using the DeviceProfiles.ini hierarchy (e.g. specific iOS device, higher priority than per project to do device specific settings)
-	ECVF_SetByDeviceProfile =		0x05000000,
+	ECVF_SetByDeviceProfile =		0x06000000,
+	// Dyanmically loaded/unloaded plugins, used with the History concept to restore cvars on plugin unload. THIS IS AN ARRAY TYPE meaning multiple vales with this SetBy are stored in the history.
+	ECVF_SetByPluginHighPriority =	0x07000000,
 	// User settable game overrides, used for GameUserSettings fields that need to override device specific settings
-	ECVF_SetByGameOverride =		0x06000000,
+	ECVF_SetByGameOverride =		0x08000000,
 	// Set by local consolevariables.ini, mostly used for testing multiple projects
-	ECVF_SetByConsoleVariablesIni = 0x07000000,
-	// Set by hotfix
-	ECVF_SetByHotfix =				0x08000000,
+	ECVF_SetByConsoleVariablesIni = 0x09000000,
+	// Set by hotfix. THIS IS AN ARRAY TYPE meaning multiple vales with this SetBy are stored in the history
+	ECVF_SetByHotfix =				0x0A000000,
 	// Used by some command line parameters, others use the Console priority instead
-	ECVF_SetByCommandline =			0x09000000,
-	// Used for high priority temporary debugging or operation modes 
-	ECVF_SetByCode =				0x0A000000,
+	ECVF_SetByCommandline =			0x0B000000,
+	// Used for high priority temporary debugging or operation modes
+	ECVF_SetByCode =				0x0C000000,
 	// Highest priority used via editor UI or or game/editor interactive console
-	ECVF_SetByConsole =				0x0B000000,
+	ECVF_SetByConsole =				0x0D000000,
 
 
 	// ------------------------------------------------
 };
 
+
+#define ENUMERATE_SET_BY(op) \
+	op(Constructor) \
+	op(Scalability) \
+	op(GameSetting) \
+	op(ProjectSetting) \
+	op(SystemSettingsIni) \
+	op(PluginLowPriority) \
+	op(DeviceProfile) \
+	op(PluginHighPriority) \
+	op(ConsoleVariablesIni) \
+	op(Hotfix) \
+	op(Commandline) \
+	op(Code) \
+	op(Console)
+
+
 /** Returns human readable ECVF_SetByMask bits of the console variable flags. */
 extern CORE_API const TCHAR* GetConsoleVariableSetByName(EConsoleVariableFlags ConsoleVariableFlags);
 
+/** Inverse of GetConsoleVariableSetByName() */
+extern CORE_API EConsoleVariableFlags GetConsoleVariableSetByValue(const TCHAR* SetByName);
+
 
 class IConsoleVariable;
 
@@ -447,8 +471,17 @@ public:
 	/**
 	 * Set the internal value from the specified string. 
 	 * @param SetBy anything in ECVF_LastSetMask e.g. ECVF_SetByScalability
+	 * @param Tag optional Tag to set with the value - only useful when UE_ALLOW_CVAR_HISTORY is set, and when setting in an ARRAY type SetBy
 	 **/
-	virtual void Set(const TCHAR* InValue, EConsoleVariableFlags SetBy = ECVF_SetByCode) = 0;
+	virtual void Set(const TCHAR* InValue, EConsoleVariableFlags SetBy = ECVF_SetByCode, FName Tag = NAME_None) = 0;
+
+	/**
+	 * Unsets the value at a certain SetBy priority (this is only useful when UE_ALLOW_CVAR_HISTORY is set). The value of the CVar
+	 * will be recalculated based on remaining History levels
+	 * @param SetBy anything in ECVF_LastSetMask e.g. ECVF_SetByScalability
+	 * @param Tag tag used to remove a setting from an ARRAY type SetBy
+	 **/
+	virtual void Unset(EConsoleVariableFlags SetBy, FName Tag = NAME_None) = 0;
 
 	/**
 	 * Get the internal value as a bool, works on bools, ints and floats.
@@ -528,46 +561,46 @@ public:
 	// convenience methods
 
 	/** Set the internal value from the specified bool. */
-	void Set(bool InValue, EConsoleVariableFlags SetBy = ECVF_SetByCode)
+	void Set(bool InValue, EConsoleVariableFlags SetBy = ECVF_SetByCode, FName Tag=NAME_None)
 	{
 		// NOTE: Bool needs to use 1 and 0 here rather than true/false, as this may be a int32 or something
 		// and eventually this code calls, TTypeFromString<T>::FromString which won't handle the true/false,
 		// but 1 and 0 will work for whatever.
 		// inefficient but no common code path
-		Set(InValue ? TEXT("1") : TEXT("0"), SetBy);
+		Set(InValue ? TEXT("1") : TEXT("0"), SetBy, Tag);
 	}
 	/** Set the internal value from the specified int. */
-	void Set(int32 InValue, EConsoleVariableFlags SetBy = ECVF_SetByCode)
+	void Set(int32 InValue, EConsoleVariableFlags SetBy = ECVF_SetByCode, FName Tag=NAME_None)
 	{
 		// inefficient but no common code path
-		Set(*FString::Printf(TEXT("%d"), InValue), SetBy);
+		Set(*FString::Printf(TEXT("%d"), InValue), SetBy, Tag);
 	}
 	/** Set the internal value from the specified float. */
-	void Set(float InValue, EConsoleVariableFlags SetBy = ECVF_SetByCode)
+	void Set(float InValue, EConsoleVariableFlags SetBy = ECVF_SetByCode, FName Tag=NAME_None)
 	{
 		// inefficient but no common code path
-		Set(*FString::Printf(TEXT("%g"), InValue), SetBy);
+		Set(*FString::Printf(TEXT("%g"), InValue), SetBy, Tag);
 	}
 
-	void SetWithCurrentPriority(bool InValue)
+	void SetWithCurrentPriority(bool InValue, FName Tag=NAME_None)
 	{
 		EConsoleVariableFlags CurFlags = (EConsoleVariableFlags)(GetFlags() & ECVF_SetByMask);
-		Set(InValue, CurFlags);
+		Set(InValue, CurFlags, Tag);
 	}
-	void SetWithCurrentPriority(int32 InValue)
+	void SetWithCurrentPriority(int32 InValue, FName Tag=NAME_None)
 	{
 		EConsoleVariableFlags CurFlags = (EConsoleVariableFlags)(GetFlags() & ECVF_SetByMask);
-		Set(InValue, CurFlags);
+		Set(InValue, CurFlags, Tag);
 	}
-	void SetWithCurrentPriority(float InValue)
+	void SetWithCurrentPriority(float InValue, FName Tag=NAME_None)
 	{
 		EConsoleVariableFlags CurFlags = (EConsoleVariableFlags)(GetFlags() & ECVF_SetByMask);
-		Set(InValue, CurFlags);
+		Set(InValue, CurFlags, Tag);
 	}
-	void SetWithCurrentPriority(const TCHAR* InValue)
+	void SetWithCurrentPriority(const TCHAR* InValue, FName Tag=NAME_None)
 	{
 		EConsoleVariableFlags CurFlags = (EConsoleVariableFlags)(GetFlags() & ECVF_SetByMask);
-		Set(InValue, CurFlags);
+		Set(InValue, CurFlags, Tag);
 	}
 
 
@@ -1047,6 +1080,12 @@ struct IConsoleManager
 	static CORE_API bool VisitPlatformCVarsForEmulation(FName PlatformName, const FString& DeviceProfileName, TFunctionRef<void(const FString& CVarName, const FString& CVarValue, EConsoleVariableFlags SetBy)> Visit);
 #endif
 
+	/**
+	  * When a plugin is unmounted, it needs to unset cvars that it had set when it was mounted. This will unset and fixup all
+	  * variables with the given Tag
+	 */
+	virtual void UnsetAllConsoleVariablesWithTag(FName Tag) = 0;
+	
 	virtual FConsoleVariableMulticastDelegate& OnCVarUnregistered() = 0;
 
 protected:
