autortfm: fix a few issues related to thread-safe lookups and caches: FNamePool::StoreWithNumber, FName::FindOrStoreString, FNameRangeCache::Find, FUObjectAnnotationChunked::GetAnnotation

diff --git a/Engine/Source/Runtime/Core/Private/UObject/UnrealNames.cpp b/Engine/Source/Runtime/Core/Private/UObject/UnrealNames.cpp
--- a/Engine/Source/Runtime/Core/Private/UObject/UnrealNames.cpp
+++ b/Engine/Source/Runtime/Core/Private/UObject/UnrealNames.cpp
@@ -1691,27 +1691,36 @@ FNameEntryId FNamePool::Store(FNameStringView Name)
 #if UE_FNAME_OUTLINE_NUMBER
 FNameEntryId FNamePool::StoreWithNumber(FNameEntryIds StringParts, int32 NumberPart)
 {
-#if WITH_CASE_PRESERVING_NAME
-	// Look for an exact match with the right casing first
-	FNumberedNameDisplayValue DisplayValue(StringParts.DisplayId, NumberPart);
-	if (FNameEntryId Existing = DisplayShards[DisplayValue.Hash.ShardIndex].FindWithNumber(DisplayValue))
+	FNameEntryId Result;
+
+	AutoRTFM::Open([&]
 	{
-		return Existing;
-	}
+#if WITH_CASE_PRESERVING_NAME
+		// Look for an exact match with the right casing first
+		FNumberedNameDisplayValue DisplayValue(StringParts.DisplayId, NumberPart);
+		if (FNameEntryId Existing = DisplayShards[DisplayValue.Hash.ShardIndex].FindWithNumber(DisplayValue))
+		{
+			Result = Existing;
+		}
+		else
+		{
 #endif
+			bool bAdded = false;
 
-	bool bAdded = false;
+			// Insert comparison name first since display value must contain comparison name
+			FNumberedNameComparisonValue ComparisonValue(StringParts.ComparisonId, NumberPart);
+			FNameEntryId ComparisonId = ComparisonShards[ComparisonValue.Hash.ShardIndex].InsertWithNumber(ComparisonValue, bAdded);
 
-	// Insert comparison name first since display value must contain comparison name
-	FNumberedNameComparisonValue ComparisonValue(StringParts.ComparisonId, NumberPart);
-	FNameEntryId ComparisonId = ComparisonShards[ComparisonValue.Hash.ShardIndex].InsertWithNumber(ComparisonValue, bAdded);
-	
 #if WITH_CASE_PRESERVING_NAME
-	DisplayValue.ComparisonId = ComparisonId;
-	return StoreValueWithNumber(DisplayValue, bAdded);
+			DisplayValue.ComparisonId = ComparisonId;
+			Result = StoreValueWithNumber(DisplayValue, bAdded);
+		}
 #else
-	return ComparisonId;
+		Result = ComparisonId;
 #endif
+	});
+
+	return Result;
 }
 
 FNameEntryId FNamePool::FindWithNumber(FNameEntryId StringPart, int32 NumberPart) const
@@ -3020,34 +3029,45 @@ private:
 	// If not found, returns NAME_None for both indices.
 	static FNameEntryIds FindOrStoreString(FNameStringView View, EFindName FindType)
 	{
-		if (View.Len >= NAME_SIZE)
+		FNameEntryIds Result{};
+
+		UE_AUTORTFM_OPEN(
 		{
-			// If we're doing a find, and the string is too long, then clearly we didn't find it
-			if (FindType == FNAME_Find)
+			if (View.Len >= NAME_SIZE)
 			{
-				return {};
-			}
+				// If we're doing a find, and the string is too long, then clearly we didn't find it
+				if (FindType == FNAME_Find)
+				{
+					Result = FNameEntryIds{};
+				}
+				else
+				{
+					checkf(false, TEXT("FName's %d max length exceeded. Got %d characters excluding null-terminator:\n%.*s"), 
+						NAME_SIZE - 1, View.Len, NAME_SIZE, View.IsAnsi() ? ANSI_TO_TCHAR(View.Ansi) : View.Wide)
 
-			checkf(false, TEXT("FName's %d max length exceeded. Got %d characters excluding null-terminator:\n%.*s"), 
-				NAME_SIZE - 1, View.Len, NAME_SIZE, View.IsAnsi() ? ANSI_TO_TCHAR(View.Ansi) : View.Wide)
+					const ANSICHAR* ErrorString = "ERROR_NAME_SIZE_EXCEEDED";
+					Result = FindOrStoreString(FNameStringView(ErrorString, FCStringAnsi::Strlen(ErrorString), false), FNAME_Add);
+				}
+			}
+			else
+			{
+				FNamePool& Pool = GetNamePool();
 
-			const ANSICHAR* ErrorString = "ERROR_NAME_SIZE_EXCEEDED";
-			return FindOrStoreString(FNameStringView(ErrorString, FCStringAnsi::Strlen(ErrorString), false), FNAME_Add);
-		}
-		
-		FNamePool& Pool = GetNamePool();
+				if (FindType == FNAME_Add)
+				{
+					FNameEntryId DisplayId = Pool.Store(View);
+					Result = FNameEntryIds{ ResolveComparisonId(DisplayId), DisplayId };
+				}
+				else
+				{
+					check(FindType == FNAME_Find);
+					FNameEntryId DisplayId = Pool.Find(View);
+					Result = FNameEntryIds{ ResolveComparisonId(DisplayId), DisplayId };
+				}
+			}
+		});
 
-		if (FindType == FNAME_Add)
-		{
-			FNameEntryId DisplayId = Pool.Store(View);
-			return FNameEntryIds{ ResolveComparisonId(DisplayId), DisplayId };
-		}
-		else
-		{
-			check(FindType == FNAME_Find);
-			FNameEntryId DisplayId = Pool.Find(View);
-			return FNameEntryIds{ ResolveComparisonId(DisplayId), DisplayId };
-		}
+		return Result;
 	}
 
 #if UE_FNAME_OUTLINE_NUMBER
diff --git a/Engine/Source/Runtime/CoreUObject/Private/UObject/UObjectGlobals.cpp b/Engine/Source/Runtime/CoreUObject/Private/UObject/UObjectGlobals.cpp
--- a/Engine/Source/Runtime/CoreUObject/Private/UObject/UObjectGlobals.cpp
+++ b/Engine/Source/Runtime/CoreUObject/Private/UObject/UObjectGlobals.cpp
@@ -2662,36 +2662,47 @@ namespace NameReuse
 
 		FName Find(UObject* Parent, FNameEntryId BaseId, FName BaseName)
 		{
-			Lock.ReadLock();
-			FNameRangeEntry* Entry = Find(BaseId);
+			FName Result;
 
-			if (Entry)
+			UE_AUTORTFM_OPEN(
 			{
-				// already allocated name, we can release the shared lock and work on this object directly
-				Lock.ReadUnlock();
-			}
-			else
-			{
-				Lock.ReadUnlock();
+				Lock.ReadLock();
+				FNameRangeEntry* Entry = Find(BaseId);
 
-				// The first time we request a name if we've never created one, don't bother adding to the cache just yet.
-				FName TestName = FName::FindNumberedName(BaseId, FNameRangeEntry::FirstNewNumber + 1);
-				if (TestName.IsNone())
+				if (Entry)
 				{
-					return FName(BaseId, BaseId, FNameRangeEntry::FirstNewNumber + 1);
-				}
+					// already allocated name, we can release the shared lock and work on this object directly
+					Lock.ReadUnlock();
 
-				FWriteScopeLock _(Lock);
-				// We didn't have a name but we may have been preempted as we acquired the write lock
-				Entry = Find(BaseId);
-				if (!Entry)
+					Result = Entry->AllocateName(Parent, BaseId, BaseName);
+				}
+				else
 				{
-					// we were not pre-empted, add a new entry for this name 
-					Entry = Add(BaseId);
+					Lock.ReadUnlock();
+
+					// The first time we request a name if we've never created one, don't bother adding to the cache just yet.
+					FName TestName = FName::FindNumberedName(BaseId, FNameRangeEntry::FirstNewNumber + 1);
+					if (TestName.IsNone())
+					{
+						Result = FName(BaseId, BaseId, FNameRangeEntry::FirstNewNumber + 1);
+					}
+					else
+					{
+						FWriteScopeLock _(Lock);
+						// We didn't have a name but we may have been preempted as we acquired the write lock
+						Entry = Find(BaseId);
+						if (!Entry)
+						{
+							// we were not pre-empted, add a new entry for this name 
+							Entry = Add(BaseId);
+						}
+
+						Result = Entry->AllocateName(Parent, BaseId, BaseName);
+					}
 				}
-			}
+			});
 
-			return Entry->AllocateName(Parent, BaseId, BaseName);
+			return Result;
 		}
 	private:
 		FRWLock Lock;
diff --git a/Engine/Source/Runtime/CoreUObject/Public/UObject/UObjectAnnotation.h b/Engine/Source/Runtime/CoreUObject/Public/UObject/UObjectAnnotation.h
--- a/Engine/Source/Runtime/CoreUObject/Public/UObject/UObjectAnnotation.h
+++ b/Engine/Source/Runtime/CoreUObject/Public/UObject/UObjectAnnotation.h
@@ -827,20 +827,27 @@ public:
 	FORCEINLINE TAnnotation GetAnnotation(int32 Index)
 	{
 		check(Index >= 0);
-		FRWScopeLock AnnotationArrayLock(AnnotationArrayCritical, SLT_ReadOnly);
 
-		const int32 ChunkIndex = Index / NumAnnotationsPerChunk;
-		if (ChunkIndex < Chunks.Num())
+		TAnnotation Result = TAnnotation();
+
+		UE_AUTORTFM_OPEN(
 		{
-			const int32 WithinChunkIndex = Index % NumAnnotationsPerChunk;
+			FRWScopeLock AnnotationArrayLock(AnnotationArrayCritical, SLT_ReadOnly);
 
-			TAnnotationChunk& Chunk = Chunks[ChunkIndex];
-			if (Chunk.Items != nullptr)
+			const int32 ChunkIndex = Index / NumAnnotationsPerChunk;
+			if (ChunkIndex < Chunks.Num())
 			{
-				return Chunk.Items[WithinChunkIndex];
+				const int32 WithinChunkIndex = Index % NumAnnotationsPerChunk;
+
+				TAnnotationChunk& Chunk = Chunks[ChunkIndex];
+				if (Chunk.Items != nullptr)
+				{
+					Result = Chunk.Items[WithinChunkIndex];
+				}
 			}
-		}
-		return TAnnotation();
+		});
+
+		return Result;
 	}
 
 	/**
