Fixed memory tracing for OnDllUnloaded (i.e. MemoryTrace_UnmarkAllocAsHeap needs to match the MemoryTrace_MarkAllocAsHeap).

diff --git a/Engine/Source/Runtime/Core/Private/ProfilingDebugging/Microsoft/WindowsModuleDiagnostics.cpp b/Engine/Source/Runtime/Core/Private/ProfilingDebugging/Microsoft/WindowsModuleDiagnostics.cpp
--- a/Engine/Source/Runtime/Core/Private/ProfilingDebugging/Microsoft/WindowsModuleDiagnostics.cpp
+++ b/Engine/Source/Runtime/Core/Private/ProfilingDebugging/Microsoft/WindowsModuleDiagnostics.cpp
@@ -206,9 +206,9 @@ void FModuleTrace::OnDllLoaded(const UNICODE_STRING& Name, UPTRINT Base)
 	{
 		UE_TRACE_METADATA_CLEAR_SCOPE();
 		LLM(UE_MEMSCOPE(ELLMTag::ProgramSize));
-		MemoryTrace_Alloc(Base, OptionalHeader.SizeOfImage, 4 * 1024);
+		MemoryTrace_Alloc(Base, OptionalHeader.SizeOfImage, 4 * 1024, EMemoryTraceRootHeap::SystemMemory);
 		MemoryTrace_MarkAllocAsHeap(Base, ProgramHeapId);
-		MemoryTrace_Alloc(Base, OptionalHeader.SizeOfImage, 4 * 1024);
+		MemoryTrace_Alloc(Base, OptionalHeader.SizeOfImage, 4 * 1024, EMemoryTraceRootHeap::SystemMemory);
 	}
 #endif // UE_MEMORY_TRACE_ENABLED
 
@@ -221,7 +221,11 @@ void FModuleTrace::OnDllLoaded(const UNICODE_STRING& Name, UPTRINT Base)
 ////////////////////////////////////////////////////////////////////////////////
 void FModuleTrace::OnDllUnloaded(UPTRINT Base)
 {
-	MemoryTrace_Free(Base);
+#if UE_MEMORY_TRACE_ENABLED
+	MemoryTrace_Free(Base, EMemoryTraceRootHeap::SystemMemory);
+	MemoryTrace_UnmarkAllocAsHeap(Base, ProgramHeapId);
+	MemoryTrace_Free(Base, EMemoryTraceRootHeap::SystemMemory);
+#endif // UE_MEMORY_TRACE_ENABLED
 
 	UE_TRACE_LOG(Diagnostics, ModuleUnload, ModuleChannel)
 		<< ModuleUnload.Base(uint64(Base));
