Make TArray::Append consistently use a 'grow' instead of 'reserve' memory growth strategy, so it is amortized linear cost to call many Append()s in a loop for all versions of the function rather than just the Append(Ptr, Count) and Append(InitList) versions

diff --git a/Engine/Source/Runtime/Core/Public/Containers/Array.h b/Engine/Source/Runtime/Core/Public/Containers/Array.h
--- a/Engine/Source/Runtime/Core/Public/Containers/Array.h
+++ b/Engine/Source/Runtime/Core/Public/Containers/Array.h
@@ -1944,10 +1944,8 @@ public:
 		}
 
 		// Allocate memory for the new elements.
-		Reserve(ArrayNum + SourceCount);
-		ConstructItems<ElementType>(GetData() + ArrayNum, Source.GetData(), SourceCount);
-
-		ArrayNum += SourceCount;
+		SizeType Pos = AddUninitialized(SourceCount);
+		ConstructItems<ElementType>(GetData() + Pos, Source.GetData(), SourceCount);
 	}
 
 	/**
@@ -1970,11 +1968,9 @@ public:
 		}
 
 		// Allocate memory for the new elements.
-		Reserve(ArrayNum + SourceCount);
-		RelocateConstructItems<ElementType>(GetData() + ArrayNum, Source.GetData(), SourceCount);
+		SizeType Pos = AddUninitialized(SourceCount);
+		RelocateConstructItems<ElementType>(GetData() + Pos, Source.GetData(), SourceCount);
 		Source.ArrayNum = 0;
-
-		ArrayNum += SourceCount;
 	}
 
 	/**
@@ -2009,10 +2005,8 @@ public:
 		SizeType SourceCount = (SizeType)InCount;
 
 		// Allocate memory for the new elements.
-		Reserve(ArrayNum + SourceCount);
-		ConstructItems<ElementType>(GetData() + ArrayNum, UE4Array_Private::GetDataHelper(Source), SourceCount);
-
-		ArrayNum += SourceCount;
+		SizeType Pos = AddUninitialized(SourceCount);
+		ConstructItems<ElementType>(GetData() + Pos, UE4Array_Private::GetDataHelper(Source), SourceCount);
 	}
 
 	/**
