Fix comments on IsRunningGame because it is true in both client/server and clientonly builds

diff --git a/Engine/Source/Runtime/Core/Public/Misc/CoreMisc.h b/Engine/Source/Runtime/Core/Public/Misc/CoreMisc.h
--- a/Engine/Source/Runtime/Core/Public/Misc/CoreMisc.h
+++ b/Engine/Source/Runtime/Core/Public/Misc/CoreMisc.h
@@ -164,7 +164,8 @@ FORCEINLINE bool IsRunningDedicatedServer()
 }
 
 /**
- * Check to see if this executable was launched as a game process that can be both a client and server (WITH_SERVER_CODE=1).
+ * Check to see if this executable was launched as a game (not editor or dedicated server) process. 
+ * This is true for both client only and client/server packaged games, call IsRunningClientOnly to differentiate.
  * An editor build can be launched with -game to set this to true, but it will be false during single process PlayInEditor mode.
  * This function should not be used for gameplay or networking purposes, check the NetMode instead.
  */
@@ -189,7 +190,8 @@ FORCEINLINE bool IsRunningGame()
 }
 
 /**
- * Check to see if this executable was launched as a a client only process that should not load server data (WITH_SERVER_CODE=0).
+ * Check to see if this executable was launched as a client only game process that should not load server data.
+ * This will be true for packaged builds with a Client target type which will define WITH_SERVER_CODE=0.
  * An editor build can be launched with -clientonly to set this to true, but it will be false during single process PlayInEditor mode.
  * This function should not be used for gameplay or networking purposes, check for NM_Client instead.
  */
