More fixes StdLibSupport

diff --git a/Engine/Source/Runtime/Core/Private/AutoRTFM/StdLibSupport.cpp b/Engine/Source/Runtime/Core/Private/AutoRTFM/StdLibSupport.cpp
--- a/Engine/Source/Runtime/Core/Private/AutoRTFM/StdLibSupport.cpp
+++ b/Engine/Source/Runtime/Core/Private/AutoRTFM/StdLibSupport.cpp
@@ -120,8 +120,14 @@ UE_AUTORTFM_REGISTER_SELF_FUNCTION(static_cast<const char*(*)(const char*, int)>
 UE_AUTORTFM_REGISTER_SELF_FUNCTION(static_cast<const char* (*)(const char*, int)>(&strrchr));
 UE_AUTORTFM_REGISTER_SELF_FUNCTION(static_cast<const char* (*)(const char*, const char*)>(&strstr));
 UE_AUTORTFM_REGISTER_SELF_FUNCTION(strlen);
+UE_AUTORTFM_REGISTER_SELF_FUNCTION(strtol);
 
+UE_AUTORTFM_REGISTER_SELF_FUNCTION(static_cast<const wchar_t* (*)(const wchar_t*, wchar_t)>(&wcschr));
+UE_AUTORTFM_REGISTER_SELF_FUNCTION(static_cast<wchar_t* (*)(wchar_t*, wchar_t)>(&wcschr));
+UE_AUTORTFM_REGISTER_SELF_FUNCTION(static_cast<wchar_t* (*)(wchar_t*, const wchar_t*)>(&wcsstr));
+UE_AUTORTFM_REGISTER_SELF_FUNCTION(static_cast<const wchar_t* (*)(const wchar_t*, const wchar_t*)>(&wcsstr));
 UE_AUTORTFM_REGISTER_SELF_FUNCTION(wcscmp);
+
 UE_AUTORTFM_REGISTER_SELF_FUNCTION(iswupper);
 UE_AUTORTFM_REGISTER_SELF_FUNCTION(iswlower);
 UE_AUTORTFM_REGISTER_SELF_FUNCTION(iswalpha);
@@ -290,6 +296,7 @@ UE_AUTORTFM_REGISTER_SELF_FUNCTION(_tcslen);
 UE_AUTORTFM_REGISTER_SELF_FUNCTION(_isnan);
 UE_AUTORTFM_REGISTER_SELF_FUNCTION(_finite);
 UE_AUTORTFM_REGISTER_SELF_FUNCTION(IsDebuggerPresent);
+UE_AUTORTFM_REGISTER_SELF_FUNCTION(GetSystemTime);
 
 UE_AUTORTFM_REGISTER_SELF_FUNCTION(QueryPerformanceCounter);
 UE_AUTORTFM_REGISTER_SELF_FUNCTION(QueryPerformanceFrequency);
@@ -313,11 +320,14 @@ UE_AUTORTFM_REGISTER_OPEN_FUNCTION(TlsSetValue);
 
 #if PLATFORM_LINUX
 UE_AUTORTFM_REGISTER_SELF_FUNCTION(clock_gettime);
+UE_AUTORTFM_REGISTER_SELF_FUNCTION(bcmp);
 #endif // PLATFORM_LINUX
 
 wchar_t* RTFM_wcsncpy(wchar_t* Dst, const wchar_t* Src, size_t Count)
 {
-	AutoRTFM::Unreachable();
+	FContext* Context = FContext::Get();
+	Context->RecordWrite(Dst, Count * sizeof(wchar_t));
+	return wcsncpy(Dst, Src, Count);
 }
 
 #ifdef _MSC_VER
