Updated comments to SmartPtr

diff --git a/Engine/Source/Runtime/Core/Public/Templates/SmartPtr.h b/Engine/Source/Runtime/Core/Public/Templates/SmartPtr.h
--- a/Engine/Source/Runtime/Core/Public/Templates/SmartPtr.h
+++ b/Engine/Source/Runtime/Core/Public/Templates/SmartPtr.h
@@ -7,12 +7,15 @@
 enum class ESmartPointer
 {
 	Strong,
-	Shared = Strong, // Shared and Strong pointers do the same thing, which is prevent automatic memory deletion using different systems (Garbage Collection vs. Reference Counting)
+	Shared,
 	Weak,
 	// Unique(?) - not sure if needed
 };
 
-/* Smart Pointer Class where the pointer can be statically casted to 'T' */
+/* A container to hold one of the FOUR Smart Pointer types: TStrongObjectPtr, TWeakObjectPtr, TSharedPtr, or TWeakPtr. 
+* During construction, there will be a compile time attempt to static_cast a raw pointer 'From' to 'T'
+* 'From' must derive from 'UObject' xOr 'TSharedFromThis' for this class to be usable
+*/
 template<typename T>
 class TSmartPtr
 {
@@ -36,6 +39,11 @@ public:
 	{
 		if (InType == ESmartPointer::Strong || InType == ESmartPointer::Shared)
 		{
+			ensureAlwaysMsgf(
+				(TIsDerivedFrom<From, UObject>::IsDerived && InType == ESmartPointer::Strong) ||
+				(IsDerivedFromSharedFromThis<From>() && InType == ESmartPointer::Shared),
+				TEXT("Mismatch of Strong and Shared. We can continue on holding a 'strong' reference, but you shouldn't mix UObjects and Shared Pointers"));
+
 			Container = MakeShared<TStrongCastable<From, T>>(InObject);
 		}
 		else if (InType == ESmartPointer::Weak)
