Mac: Change the default of Slate.MacMainMenuInsideUnrealWindow to true

diff --git a/Engine/Source/Runtime/Core/Private/Mac/MacPlatformMisc.cpp b/Engine/Source/Runtime/Core/Private/Mac/MacPlatformMisc.cpp
--- a/Engine/Source/Runtime/Core/Private/Mac/MacPlatformMisc.cpp
+++ b/Engine/Source/Runtime/Core/Private/Mac/MacPlatformMisc.cpp
@@ -77,7 +77,7 @@ static TAutoConsoleVariable<int32> CVarMacPlatformDumpAllThreadsOnHang(
 	TEXT("Mac.DumpAllThreadsOnHang"),
 	1,
 	TEXT("If > 0, then when reporting a hang generate a backtrace for all threads."));
-static bool GMacMainMenuInsideUnrealWindow = false;
+static bool GMacMainMenuInsideUnrealWindow = true;
 static FAutoConsoleVariableRef CVarMacMainMenuInsideUnrealWindow(
 	TEXT("Slate.MacMainMenuInsideUnrealWindow"),
 	GMacMainMenuInsideUnrealWindow,
