Tasks System: adding an opton for tasks to opt out of busy waiting

#udn 0D54z00009iv5weCAA

diff --git a/Engine/Source/Runtime/Core/Private/Tasks/Pipe.cpp b/Engine/Source/Runtime/Core/Private/Tasks/Pipe.cpp
--- a/Engine/Source/Runtime/Core/Private/Tasks/Pipe.cpp
+++ b/Engine/Source/Runtime/Core/Private/Tasks/Pipe.cpp
@@ -55,7 +55,8 @@ namespace UE::Tasks
 				Init(
 					TEXT("Pipe::WaitUntilEmpty() placeholder"),
 					ETaskPriority::Normal, // doesn't matter
-					EExtendedTaskPriority::TaskEvent // no need for scheduling or execution for this dummy task
+					EExtendedTaskPriority::TaskEvent, // no need for scheduling or execution for this dummy task
+					ETaskFlags::None
 				);
 			}
 
diff --git a/Engine/Source/Runtime/Core/Private/Tasks/TaskPrivate.cpp b/Engine/Source/Runtime/Core/Private/Tasks/TaskPrivate.cpp
--- a/Engine/Source/Runtime/Core/Private/Tasks/TaskPrivate.cpp
+++ b/Engine/Source/Runtime/Core/Private/Tasks/TaskPrivate.cpp
@@ -218,7 +218,7 @@ namespace UE::Tasks
 			auto WaitingTaskBody = [CompletionEvent] { CompletionEvent->Trigger(); };
 			using FWaitingTask = TExecutableTask<decltype(WaitingTaskBody)>;
 
-			TRefCountPtr<FWaitingTask> WaitingTask{ FWaitingTask::Create(TEXT("Waiting Task"), MoveTemp(WaitingTaskBody), ETaskPriority::Default /* doesn't matter*/, EExtendedTaskPriority::Inline), /*bAddRef=*/ false };
+			TRefCountPtr<FWaitingTask> WaitingTask{ FWaitingTask::Create(TEXT("Waiting Task"), MoveTemp(WaitingTaskBody), ETaskPriority::Default /* doesn't matter*/, EExtendedTaskPriority::Inline, ETaskFlags::None), /*bAddRef=*/ false };
 			WaitingTask->AddPrerequisites(*this);
 
 			if (WaitingTask->TryLaunch(sizeof(WaitingTask)))
diff --git a/Engine/Source/Runtime/Core/Public/Tasks/Pipe.h b/Engine/Source/Runtime/Core/Public/Tasks/Pipe.h
--- a/Engine/Source/Runtime/Core/Public/Tasks/Pipe.h
+++ b/Engine/Source/Runtime/Core/Public/Tasks/Pipe.h
@@ -54,6 +54,7 @@ namespace UE::Tasks
 		// @param TaskBody a callable with no parameters, usually a lambda but can be also a functor object 
 		// or a pointer to a function. TaskBody can return results.
 		// @Priority - task priority, can affect task scheduling once it's passed the pipe
+		// @param TaskFlags - task config options
 		// @return Task instance that can be used to wait for task completion or to obtain the result of task execution
 		template<typename TaskBodyType>
 		TTask<TInvokeResult_T<TaskBodyType>> Launch
@@ -61,13 +62,14 @@ namespace UE::Tasks
 			const TCHAR* InDebugName, 
 			TaskBodyType&& TaskBody, 
 			ETaskPriority Priority = ETaskPriority::Default,
-			EExtendedTaskPriority ExtendedPriority = EExtendedTaskPriority::None
+			EExtendedTaskPriority ExtendedPriority = EExtendedTaskPriority::None,
+			ETaskFlags Flags = ETaskFlags::DoNotRunInsideBusyWait
 		)
 		{
 			using FResult = TInvokeResult_T<TaskBodyType>;
 			using FExecutableTask = Private::TExecutableTask<std::decay_t<TaskBodyType>>;
 
-			FExecutableTask* Task = FExecutableTask::Create(InDebugName, Forward<TaskBodyType>(TaskBody), Priority, ExtendedPriority);
+			FExecutableTask* Task = FExecutableTask::Create(InDebugName, Forward<TaskBodyType>(TaskBody), Priority, ExtendedPriority, Flags);
 			Task->SetPipe(*this);
 			Task->TryLaunch(sizeof(*Task));
 			return TTask<FResult>{ Task };
@@ -78,6 +80,7 @@ namespace UE::Tasks
 		// @param TaskBody a callable with no parameters, usually a lambda but can be also a functor object 
 		// or a pointer to a function. TaskBody can return results.
 		// @Priority - task priority, can affect task scheduling once it's passed the pipe
+		// @param TaskFlags - task config options
 		// @return Task instance that can be used to wait for task completion or to obtain the result of task execution
 		template<typename TaskBodyType, typename PrerequisitesCollectionType>
 		TTask<TInvokeResult_T<TaskBodyType>> Launch
@@ -86,13 +89,14 @@ namespace UE::Tasks
 			TaskBodyType&& TaskBody, 
 			PrerequisitesCollectionType&& Prerequisites, 
 			ETaskPriority Priority = ETaskPriority::Default,
-			EExtendedTaskPriority ExtendedPriority = EExtendedTaskPriority::None
+			EExtendedTaskPriority ExtendedPriority = EExtendedTaskPriority::None,
+			ETaskFlags Flags = ETaskFlags::DoNotRunInsideBusyWait
 		)
 		{
 			using FResult = TInvokeResult_T<TaskBodyType>;
 			using FExecutableTask = Private::TExecutableTask<std::decay_t<TaskBodyType>>;
 
-			FExecutableTask* Task = FExecutableTask::Create(InDebugName, Forward<TaskBodyType>(TaskBody), Priority, ExtendedPriority);
+			FExecutableTask* Task = FExecutableTask::Create(InDebugName, Forward<TaskBodyType>(TaskBody), Priority, ExtendedPriority, Flags);
 			Task->SetPipe(*this);
 			Task->AddPrerequisites(Forward<PrerequisitesCollectionType>(Prerequisites));
 			Task->TryLaunch(sizeof(*Task));
diff --git a/Engine/Source/Runtime/Core/Public/Tasks/Task.h b/Engine/Source/Runtime/Core/Public/Tasks/Task.h
--- a/Engine/Source/Runtime/Core/Public/Tasks/Task.h
+++ b/Engine/Source/Runtime/Core/Public/Tasks/Task.h
@@ -104,19 +104,21 @@ namespace UE::Tasks
 			// @param DebugName - a unique name for task identification in debugger and profiler, is compiled out in test/shipping builds
 			// @param TaskBody - a functor that will be executed asynchronously
 			// @param Priority - task priority that affects when the task will be executed
+			// @param TaskFlags - task config options
 			// @return a trivially relocatable instance that can be used to wait for task completion or to obtain task execution result
 			template<typename TaskBodyType>
 			void Launch(
 				const TCHAR* DebugName,
 				TaskBodyType&& TaskBody,
 				ETaskPriority Priority = ETaskPriority::Normal,
-				EExtendedTaskPriority ExtendedPriority = EExtendedTaskPriority::None
+				EExtendedTaskPriority ExtendedPriority = EExtendedTaskPriority::None,
+				ETaskFlags Flags = ETaskFlags::DoNotRunInsideBusyWait
 			)
 			{
 				check(!IsValid());
 
 				using FExecutableTask = Private::TExecutableTask<std::decay_t<TaskBodyType>>;
-				FExecutableTask* Task = FExecutableTask::Create(DebugName, Forward<TaskBodyType>(TaskBody), Priority, ExtendedPriority);
+				FExecutableTask* Task = FExecutableTask::Create(DebugName, Forward<TaskBodyType>(TaskBody), Priority, ExtendedPriority, Flags);
 				// this must happen before launching, to support an ability to access the task itself from inside it
 				*Pimpl.GetInitReference() = Task;
 				Task->TryLaunch(sizeof(*Task));
@@ -128,6 +130,7 @@ namespace UE::Tasks
 			// @param Prerequisites - tasks or task events that must be completed before the task being launched can be scheduled, accepts any 
 			// iterable collection (.begin()/.end()), `Tasks::Prerequisites()` helper is recommended to create such collection on the fly
 			// @param Priority - task priority that affects when the task will be executed
+			// @param TaskFlags - task config options
 			// @return a trivially relocatable instance that can be used to wait for task completion or to obtain task execution result
 			template<typename TaskBodyType, typename PrerequisitesCollectionType>
 			void Launch(
@@ -135,13 +138,14 @@ namespace UE::Tasks
 				TaskBodyType&& TaskBody,
 				PrerequisitesCollectionType&& Prerequisites,
 				ETaskPriority Priority = ETaskPriority::Normal,
-				EExtendedTaskPriority ExtendedPriority = EExtendedTaskPriority::None
+				EExtendedTaskPriority ExtendedPriority = EExtendedTaskPriority::None,
+				ETaskFlags Flags = ETaskFlags::DoNotRunInsideBusyWait
 			)
 			{
 				check(!IsValid());
 
 				using FExecutableTask = Private::TExecutableTask<std::decay_t<TaskBodyType>>;
-				FExecutableTask* Task = FExecutableTask::Create(DebugName, Forward<TaskBodyType>(TaskBody), Priority, ExtendedPriority);
+				FExecutableTask* Task = FExecutableTask::Create(DebugName, Forward<TaskBodyType>(TaskBody), Priority, ExtendedPriority, Flags);
 				Task->AddPrerequisites(Forward<PrerequisitesCollectionType>(Prerequisites));
 				// this must happen before launching, to support an ability to access the task itself from inside it
 				*Pimpl.GetInitReference() = Task;
@@ -247,18 +251,20 @@ namespace UE::Tasks
 	// @param DebugName - a unique name for task identification in debugger and profiler, is compiled out in test/shipping builds
 	// @param TaskBody - a functor that will be executed asynchronously
 	// @param Priority - task priority that affects when the task will be executed
+	// @param TaskFlags - task config options
 	// @return a trivially relocatable instance that can be used to wait for task completion or to obtain task execution result
 	template<typename TaskBodyType>
 	TTask<TInvokeResult_T<TaskBodyType>> Launch(
 		const TCHAR* DebugName,
 		TaskBodyType&& TaskBody,
 		ETaskPriority Priority = ETaskPriority::Normal,
-		EExtendedTaskPriority ExtendedPriority = EExtendedTaskPriority::None
+		EExtendedTaskPriority ExtendedPriority = EExtendedTaskPriority::None,
+		ETaskFlags Flags = ETaskFlags::DoNotRunInsideBusyWait
 	)
 	{
 		using FResult = TInvokeResult_T<TaskBodyType>;
 		TTask<FResult> Task;
-		Task.Launch(DebugName, Forward<TaskBodyType>(TaskBody), Priority, ExtendedPriority);
+		Task.Launch(DebugName, Forward<TaskBodyType>(TaskBody), Priority, ExtendedPriority, Flags);
 		return Task;
 	}
 
@@ -268,6 +274,7 @@ namespace UE::Tasks
 	// @param Prerequisites - tasks or task events that must be completed before the task being launched can be scheduled, accepts any 
 	// iterable collection (.begin()/.end()), `Tasks::Prerequisites()` helper is recommended to create such collection on the fly
 	// @param Priority - task priority that affects when the task will be executed
+	// @param TaskFlags - task config options
 	// @return a trivially relocatable instance that can be used to wait for task completion or to obtain task execution result
 	template<typename TaskBodyType, typename PrerequisitesCollectionType>
 	TTask<TInvokeResult_T<TaskBodyType>> Launch(
@@ -275,12 +282,13 @@ namespace UE::Tasks
 		TaskBodyType&& TaskBody,
 		PrerequisitesCollectionType&& Prerequisites,
 		ETaskPriority Priority = ETaskPriority::Normal,
-		EExtendedTaskPriority ExtendedPriority = EExtendedTaskPriority::None
+		EExtendedTaskPriority ExtendedPriority = EExtendedTaskPriority::None,
+		ETaskFlags Flags = ETaskFlags::DoNotRunInsideBusyWait
 	)
 	{
 		using FResult = TInvokeResult_T<TaskBodyType>;
 		TTask<FResult> Task;
-		Task.Launch(DebugName, Forward<TaskBodyType>(TaskBody), Forward<PrerequisitesCollectionType>(Prerequisites), Priority, ExtendedPriority);
+		Task.Launch(DebugName, Forward<TaskBodyType>(TaskBody), Forward<PrerequisitesCollectionType>(Prerequisites), Priority, ExtendedPriority, Flags);
 		return Task;
 	}
 
@@ -332,7 +340,7 @@ namespace UE::Tasks
 		auto WaitingTaskBody = [CompletionEvent] { CompletionEvent->Trigger(); };
 		using FWaitingTask = Private::TExecutableTask<decltype(WaitingTaskBody)>;
 
-		TRefCountPtr<FWaitingTask> WaitingTask{ FWaitingTask::Create(TEXT("Waiting Task"), MoveTemp(WaitingTaskBody), ETaskPriority::Default /* doesn't matter */, EExtendedTaskPriority::Inline), /*bAddRef=*/ false};
+		TRefCountPtr<FWaitingTask> WaitingTask{ FWaitingTask::Create(TEXT("Waiting Task"), MoveTemp(WaitingTaskBody), ETaskPriority::Default /* doesn't matter */, EExtendedTaskPriority::Inline, ETaskFlags::None), /*bAddRef=*/ false};
 		WaitingTask->AddPrerequisites(Tasks);
 
 		if (WaitingTask->TryLaunch(sizeof(WaitingTask)))
diff --git a/Engine/Source/Runtime/Core/Public/Tasks/TaskPrivate.h b/Engine/Source/Runtime/Core/Public/Tasks/TaskPrivate.h
--- a/Engine/Source/Runtime/Core/Public/Tasks/TaskPrivate.h
+++ b/Engine/Source/Runtime/Core/Public/Tasks/TaskPrivate.h
@@ -73,6 +73,12 @@ namespace UE::Tasks
 	const TCHAR* ToString(EExtendedTaskPriority ExtendedPriority);
 	bool ToExtendedTaskPriority(const TCHAR* ExtendedPriorityStr, EExtendedTaskPriority& OutExtendedPriority);
 
+	enum class ETaskFlags
+	{
+		None,
+		DoNotRunInsideBusyWait // do not pick this task for busy-waiting
+	};
+
 	namespace Private
 	{
 		CORE_API void TranslatePriority(ENamedThreads::Type ThreadType, ETaskPriority& OutPriority, EExtendedTaskPriority& OutExtendedPriority);
@@ -141,10 +147,15 @@ namespace UE::Tasks
 			{
 			}
 
-			void Init(const TCHAR* InDebugName, ETaskPriority InPriority, EExtendedTaskPriority InExtendedPriority)
+			void Init(const TCHAR* InDebugName, ETaskPriority InPriority, EExtendedTaskPriority InExtendedPriority, ETaskFlags Flags)
 			{
-				// store debug name, priority and an adaptor for task execution in low-level task. The task body can't be stored as this task implementation needs to do some accounting
-				// before the task is executed (e.g. maintainance of TLS "current task")
+				// store debug name, priority and an adaptor for task execution in low-level task. The task body can't be stored as this task
+				// implementation needs to do some accounting before the task is executed (e.g. maintainance of TLS "current task")
+				LowLevelTasks::ETaskFlags LowLevelTaskFlags = LowLevelTasks::ETaskFlags::DefaultFlags;
+				if (Flags == ETaskFlags::DoNotRunInsideBusyWait)
+				{
+					LowLevelTaskFlags &= ~LowLevelTasks::ETaskFlags::AllowBusyWaiting;
+				}
 				LowLevelTask.Init(InDebugName, InPriority,
 					[
 						this,
@@ -155,7 +166,8 @@ namespace UE::Tasks
 					]
 					{
 						TryExecuteTask();
-					}
+					},
+					LowLevelTaskFlags
 				);
 				ExtendedPriority = InExtendedPriority;
 
@@ -631,10 +643,10 @@ namespace UE::Tasks
 		class TTaskWithResult : public FTaskBase
 		{
 		protected:
-			explicit TTaskWithResult(const TCHAR* InDebugName, ETaskPriority InPriority, EExtendedTaskPriority InExtendedPriority, uint32 InitRefCount)
+			explicit TTaskWithResult(const TCHAR* InDebugName, ETaskPriority InPriority, EExtendedTaskPriority InExtendedPriority, uint32 InitRefCount, ETaskFlags Flags)
 				: FTaskBase(InitRefCount)
 			{
-				Init(InDebugName, InPriority, InExtendedPriority);
+				Init(InDebugName, InPriority, InExtendedPriority, Flags);
 			}
 
 			virtual ~TTaskWithResult() override
@@ -670,8 +682,8 @@ namespace UE::Tasks
 			}
 
 		protected:
-			TExecutableTaskBase(const TCHAR* InDebugName, TaskBodyType&& TaskBody, ETaskPriority InPriority, EExtendedTaskPriority InExtendedPriority) 
-				: TTaskWithResult<ResultType>(InDebugName, InPriority, InExtendedPriority, 2)
+			TExecutableTaskBase(const TCHAR* InDebugName, TaskBodyType&& TaskBody, ETaskPriority InPriority, EExtendedTaskPriority InExtendedPriority, ETaskFlags Flags)
+				: TTaskWithResult<ResultType>(InDebugName, InPriority, InExtendedPriority, 2, Flags)
 				// 2 init refs: one for the initial reference (we don't increment it on passing to `TRefCountPtr`), and one for the internal 
 				// reference that keeps the task alive while it's in the system. is released either on task completion or by the scheduler after
 				// trying to execute the task
@@ -699,12 +711,12 @@ namespace UE::Tasks
 			}
 
 		protected:
-			TExecutableTaskBase(const TCHAR* InDebugName, TaskBodyType&& TaskBody, ETaskPriority InPriority, EExtendedTaskPriority InExtendedPriority) :
+			TExecutableTaskBase(const TCHAR* InDebugName, TaskBodyType&& TaskBody, ETaskPriority InPriority, EExtendedTaskPriority InExtendedPriority, ETaskFlags Flags) :
 				FTaskBase(2) // 2 init refs: one for the initial reference (we don't increment it on passing to `TRefCountPtr`), and one for the internal 
 				// reference that keeps the task alive while it's in the system. is released either on task completion or by the scheduler after
 				// trying to execute the task
 			{
-				Init(InDebugName, InPriority, InExtendedPriority);
+				Init(InDebugName, InPriority, InExtendedPriority, Flags);
 				new(&TaskBodyStorage) TaskBodyType(MoveTemp(TaskBody));
 			}
 
@@ -721,15 +733,15 @@ namespace UE::Tasks
 		class TExecutableTask final : public TExecutableTaskBase<TaskBodyType>
 		{
 		public:
-			TExecutableTask(const TCHAR* InDebugName, TaskBodyType&& TaskBody, ETaskPriority InPriority, EExtendedTaskPriority InExtendedPriority)
-				: TExecutableTaskBase<TaskBodyType>(InDebugName, MoveTemp(TaskBody), InPriority, InExtendedPriority)
+			TExecutableTask(const TCHAR* InDebugName, TaskBodyType&& TaskBody, ETaskPriority InPriority, EExtendedTaskPriority InExtendedPriority, ETaskFlags Flags)
+				: TExecutableTaskBase<TaskBodyType>(InDebugName, MoveTemp(TaskBody), InPriority, InExtendedPriority, Flags)
 			{
 			}
 
 			// a helper that deduces the template argument
-			static TExecutableTask* Create(const TCHAR* InDebugName, TaskBodyType&& TaskBody, ETaskPriority InPriority, EExtendedTaskPriority InExtendedPriority)
+			static TExecutableTask* Create(const TCHAR* InDebugName, TaskBodyType&& TaskBody, ETaskPriority InPriority, EExtendedTaskPriority InExtendedPriority, ETaskFlags Flags)
 			{
-				return new TExecutableTask(InDebugName, MoveTemp(TaskBody), InPriority, InExtendedPriority);
+				return new TExecutableTask(InDebugName, MoveTemp(TaskBody), InPriority, InExtendedPriority, Flags);
 			}
 
 			static void* operator new(size_t Size)
@@ -765,7 +777,7 @@ namespace UE::Tasks
 				: FTaskBase(/*InitRefCount=*/ 1) // for the initial reference (we don't increment it on passing to `TRefCountPtr`)
 			{
 				TaskTrace::Created(GetTraceId(), sizeof(*this));
-				Init(InDebugName, ETaskPriority::Normal, EExtendedTaskPriority::TaskEvent);
+				Init(InDebugName, ETaskPriority::Normal, EExtendedTaskPriority::TaskEvent, ETaskFlags::None);
 			}
 
 			virtual void ExecuteTask() override final
