Remove use ot TFunction causing needless memory allocations in FInterpCurve::FindNearestOnSegment

diff --git a/Engine/Source/Runtime/Core/Public/Math/InterpCurve.h b/Engine/Source/Runtime/Core/Public/Math/InterpCurve.h
--- a/Engine/Source/Runtime/Core/Public/Math/InterpCurve.h
+++ b/Engine/Source/Runtime/Core/Public/Math/InterpCurve.h
@@ -579,6 +579,38 @@ float FInterpCurve<T>::InaccurateFindNearest(const T& PointInSpace, float& OutDi
 	return FindNearest(PointInSpace, OutDistanceSq, OutSegment);
 }
 
+namespace UE::Math::InterpCurve::Private
+{
+	template <int32 MaxIteration, typename BreakBeforeEvaluateFnType, typename BreakAfterEvaluateFnType, typename T>
+	FORCEINLINE static float Newton(float& Value, float& Move, BreakBeforeEvaluateFnType &BreakBeforeEvaluate, BreakAfterEvaluateFnType &BreakAfterEvaluate, float Tolerance, float Diff, const T& PointInSpace, int32 NextPtIdx, int32 PtIdx, const TArray<FInterpCurvePoint<T>>& PointsT)
+	{
+		T FoundPoint;
+
+		for (int32 Iter = 0; Iter < MaxIteration; ++Iter)
+		{
+			Value += Move;
+			Value = FMath::Clamp(Value, 0.0f, 1.0f);
+
+			if (BreakBeforeEvaluate(Value, Tolerance))
+			{
+				break;
+			}
+
+			FoundPoint = FMath::CubicInterp(PointsT[PtIdx].OutVal, PointsT[PtIdx].LeaveTangent * Diff, PointsT[NextPtIdx].OutVal, PointsT[NextPtIdx].ArriveTangent * Diff, Value);
+
+			if (BreakAfterEvaluate(Value, Tolerance))
+			{
+				break;
+			}
+
+			const T Tangent = FMath::CubicInterpDerivative(PointsT[PtIdx].OutVal, PointsT[PtIdx].LeaveTangent * Diff, PointsT[NextPtIdx].OutVal, PointsT[NextPtIdx].ArriveTangent * Diff, Value);
+			const T Delta = (PointInSpace - FoundPoint);
+			Move = static_cast<float>(Tangent.Dot(Delta) / Tangent.SizeSquared());
+		}
+		return static_cast<float>((FoundPoint - PointInSpace).SizeSquared());
+	}
+};
+
 template< class T >
 float FInterpCurve<T>::FindNearestOnSegment(const T& PointInSpace, int32 PtIdx, float& OutSquaredDistance) const
 {
@@ -648,7 +680,7 @@ float FInterpCurve<T>::FindNearestOnSegment(const T& PointInSpace, int32 PtIdx,
 
 		const TArray<FInterpCurvePoint<T>>& PointsT = Points;
 
-		TFunction<bool(float&, const float)> BreakIfConverged = [&LastValue, &MinValue, &MaxValue, &NextValue, &StartValuesT, NewtonIterationCount](float& Value, const float Tolerance) -> bool
+		auto BreakIfConverged = [&LastValue, &MinValue, &MaxValue, &NextValue, &StartValuesT, NewtonIterationCount](float& Value, const float Tolerance) -> bool
 		{
 			if (FMath::IsNearlyEqual(LastValue, Value, Tolerance))
 			{
@@ -676,38 +708,15 @@ float FInterpCurve<T>::FindNearestOnSegment(const T& PointInSpace, int32 PtIdx,
 			return false;
 		};
 
-		TFunction<bool(float&, const float)> NeverBreak = [](float& Value, const float Tolerance) -> bool
-		{
-			return false;
-		};
-
-		TFunction<float(float&, float&, TFunction<bool(float&, const float)>, TFunction<bool(float&, const float)>, const float)> Newton = [&PointsT, &PtIdx, &NextPtIdx, &PointInSpace, &Diff, MaxIteration, InvThree](float& Value, float& Move, TFunction<bool(float&, const float)> BreakBeforeEvaluate, TFunction<bool(float&, const float)> BreakAfterEvaluate, const float Tolerance) -> float
+		// Note: This functor is a struct such that we can FORCEINLINE the function to try to make sure the compiler figures out that it is a compile-time constant.
+		struct FNeverBreak
 		{
-			T FoundPoint;
-
-			for (int32 Iter = 0; Iter < MaxIteration; ++Iter)
+			FORCEINLINE bool operator()(float& Value, const float Tolerance) const
 			{
-				Value += Move;
-				Value = FMath::Clamp(Value, 0.0f, 1.0f);
-
-				if (BreakBeforeEvaluate(Value, Tolerance))
-				{
-					break;
-				}
-
-				FoundPoint = FMath::CubicInterp(PointsT[PtIdx].OutVal, PointsT[PtIdx].LeaveTangent * Diff, PointsT[NextPtIdx].OutVal, PointsT[NextPtIdx].ArriveTangent * Diff, Value);
-
-				if (BreakAfterEvaluate(Value, Tolerance))
-				{
-					break;
-				}
-
-				const T Tangent = FMath::CubicInterpDerivative(PointsT[PtIdx].OutVal, PointsT[PtIdx].LeaveTangent * Diff, PointsT[NextPtIdx].OutVal, PointsT[NextPtIdx].ArriveTangent * Diff, Value);
-				const T Delta = (PointInSpace - FoundPoint);
-				Move = static_cast<float>(Tangent.Dot(Delta) / Tangent.SizeSquared());
+				return false;
 			}
-			return static_cast<float>((FoundPoint - PointInSpace).SizeSquared());
 		};
+		FNeverBreak NeverBreak;
 
 		// A first step starting from the 4 start coordinates with a large tolerance to find a good approximation in few iterations
 		for (int32 Index = 0; Index < NewtonIterationCount; ++Index)
@@ -718,7 +727,7 @@ float FInterpCurve<T>::FindNearestOnSegment(const T& PointInSpace, int32 PtIdx,
 			float& Value = ValuesT[Index];
 			float& Move = NextMovesT[Index];
 
-			DistancesSq[Index] = Newton(Value, Move, BreakIfConverged, NeverBreak, ToleranceStep1);
+			DistancesSq[Index] = UE::Math::InterpCurve::Private::Newton<MaxIteration>(Value, Move, BreakIfConverged, NeverBreak, ToleranceStep1, Diff, PointInSpace, NextPtIdx, PtIdx, PointsT);
 		}
 
 		// Find the index of the best approximation
@@ -733,7 +742,7 @@ float FInterpCurve<T>::FindNearestOnSegment(const T& PointInSpace, int32 PtIdx,
 			MinValue = 0.f;
 			MaxValue = 1.f;
 			float& Value = ValuesT[IndexOfTheSmallest];
-			DistancesSq[IndexOfTheSmallest] = Newton(Value, Move, NeverBreak, BreakIfConverged, ToleranceStep2);
+			DistancesSq[IndexOfTheSmallest] = UE::Math::InterpCurve::Private::Newton<MaxIteration>(Value, Move, NeverBreak, BreakIfConverged, ToleranceStep2, Diff, PointInSpace, NextPtIdx, PtIdx, PointsT);
 		}
 
 		OutSquaredDistance = DistancesSq[IndexOfTheSmallest];
