Add CA_ASSUME to UTEST_NOT_NULL

diff --git a/Engine/Source/Runtime/Core/Public/Misc/AutomationTest.h b/Engine/Source/Runtime/Core/Public/Misc/AutomationTest.h
--- a/Engine/Source/Runtime/Core/Public/Misc/AutomationTest.h
+++ b/Engine/Source/Runtime/Core/Public/Misc/AutomationTest.h
@@ -4000,7 +4000,8 @@ public: \
 	if (!TestNotNull(What, Pointer))\
 	{\
 		return false;\
-	}
+	}\
+	CA_ASSUME(Pointer)
 
 #define UTEST_NOT_NULL_EXPR(Pointer)\
 	if (!TestNotNull(TEXT(#Pointer), Pointer))\
