Add IsValidIndex function to TPagedArray for parity with other containers.

#tests Automated tests via preflight passed.

diff --git a/Engine/Source/Runtime/Core/Public/Containers/PagedArray.h b/Engine/Source/Runtime/Core/Public/Containers/PagedArray.h
--- a/Engine/Source/Runtime/Core/Public/Containers/PagedArray.h
+++ b/Engine/Source/Runtime/Core/Public/Containers/PagedArray.h
@@ -293,6 +293,11 @@ public:
 		return Pages.Num();
 	}
 
+	[[nodiscard]] FORCEINLINE bool IsValidIndex(SizeType Index) const
+	{
+		return Index >= 0 && Index < Count;
+	}
+
 	[[nodiscard]] FORCEINLINE bool IsEmpty() const
 	{
 		return !Count;
diff --git a/Engine/Source/Runtime/Core/Tests/Containers/PagedArrayTest.cpp b/Engine/Source/Runtime/Core/Tests/Containers/PagedArrayTest.cpp
--- a/Engine/Source/Runtime/Core/Tests/Containers/PagedArrayTest.cpp
+++ b/Engine/Source/Runtime/Core/Tests/Containers/PagedArrayTest.cpp
@@ -383,6 +383,46 @@ TEST_CASE_NAMED(
 	CHECK(int32Token::EvenConstructionDestructionCalls(ArrayType::MaxPerPage() * 4 + 1));
 }
 
+TEST_CASE_NAMED(FPagedArrayIsValidIndexTest, "System::Core::Containers::TPagedArray::IsValidIndex", "[SmokeFilter][Core][Containers][PagedArray]")
+{
+	using ArrayType = TPagedArray<int32Token, 128>;
+	int32Token::Reset();
+	{
+		ArrayType PagedArray;
+
+		// Single-page
+		PagedArray.Reserve(ArrayType::MaxPerPage());
+		CHECK(!PagedArray.IsValidIndex(-1));
+		CHECK(!PagedArray.IsValidIndex(0));
+		CHECK(!PagedArray.IsValidIndex(ArrayType::MaxPerPage()));
+		for (int32 i = 0; i < ArrayType::MaxPerPage(); ++i)
+		{
+			PagedArray.Emplace(i);
+			CHECK(PagedArray.Max() == ArrayType::MaxPerPage());
+		}
+		for (int32 i = 0; i < PagedArray.Num(); ++i)
+		{
+			CHECK(PagedArray.IsValidIndex(i));
+		}
+		CHECK(!PagedArray.IsValidIndex(PagedArray.Num() + 1));
+
+		// Multi-page
+		PagedArray.Reserve(ArrayType::MaxPerPage() * 2);
+		CHECK(!PagedArray.IsValidIndex(ArrayType::MaxPerPage() * 2));
+		for (int32 i = 0; i < ArrayType::MaxPerPage(); ++i)
+		{
+			PagedArray.Emplace(i + ArrayType::MaxPerPage());
+			CHECK(PagedArray.Max() == ArrayType::MaxPerPage() * 2);
+		}
+		for (int32 i = 0; i < PagedArray.Num(); ++i)
+		{
+			CHECK(PagedArray.IsValidIndex(i));
+		}
+		CHECK(!PagedArray.IsValidIndex(PagedArray.Num() + 1));
+	}
+	CHECK(int32Token::EvenConstructionDestructionCalls(ArrayType::MaxPerPage() * 2));
+}
+
 TEST_CASE_NAMED(FPagedArrayIndexingTest, "System::Core::Containers::TPagedArray::Indexing", "[SmokeFilter][Core][Containers][PagedArray]")
 {
 	using ArrayType = TPagedArray<int32Token, 128>;
