Fix PVS warnings - V591: Non-void function should return a value.

Add return values for static_assert cases to placate PVS
Fix copy operators not returning this
Fix a handful of other cases

diff --git a/Engine/Plugins/Experimental/GameplayGraph/Source/GameplayGraph/Private/Graph/Algorithms/Search/Search.cpp b/Engine/Plugins/Experimental/GameplayGraph/Source/GameplayGraph/Private/Graph/Algorithms/Search/Search.cpp
--- a/Engine/Plugins/Experimental/GameplayGraph/Source/GameplayGraph/Private/Graph/Algorithms/Search/Search.cpp
+++ b/Engine/Plugins/Experimental/GameplayGraph/Source/GameplayGraph/Private/Graph/Algorithms/Search/Search.cpp
@@ -39,6 +39,7 @@ namespace Graph::Algorithms
 			else
 			{
 				static_assert("Invalid data structure for graph search [GetNextAndAdvance].");
+				return {};
 			}
 		}
 
diff --git a/Engine/Plugins/Experimental/Mutable/Source/MutableRuntime/Private/MuR/ParametersPrivate.h b/Engine/Plugins/Experimental/Mutable/Source/MutableRuntime/Private/MuR/ParametersPrivate.h
--- a/Engine/Plugins/Experimental/Mutable/Source/MutableRuntime/Private/MuR/ParametersPrivate.h
+++ b/Engine/Plugins/Experimental/Mutable/Source/MutableRuntime/Private/MuR/ParametersPrivate.h
@@ -60,11 +60,13 @@ namespace Private
         TIndirectObject& operator=(const T& Object) 
         { 
             Get() = Object; 
+            return *this;
         }
 
         TIndirectObject& operator=(const T&& Object) 
         { 
             Get() = MoveTemp(Object); 
+            return *this;
         }
 
         const T& Get() const 
diff --git a/Engine/Plugins/Runtime/GameplayAbilities/Source/GameplayAbilities/Public/Abilities/Tasks/AbilityTask.h b/Engine/Plugins/Runtime/GameplayAbilities/Source/GameplayAbilities/Public/Abilities/Tasks/AbilityTask.h
--- a/Engine/Plugins/Runtime/GameplayAbilities/Source/GameplayAbilities/Public/Abilities/Tasks/AbilityTask.h
+++ b/Engine/Plugins/Runtime/GameplayAbilities/Source/GameplayAbilities/Public/Abilities/Tasks/AbilityTask.h
@@ -149,6 +149,7 @@ class GAMEPLAYABILITIES_API UAbilityTask : public UGameplayTask
 	FORCEINLINE static T* NewTask(UObject* WorldContextObject, FName InstanceName = FName())
 	{
 		static_assert(DelayedFalse<T>(), "UAbilityTask::NewTask should never be used. Use NewAbilityTask instead");
+		return nullptr;
 	}
 
 	/** Called when the ability task is waiting on remote player data. IF the remote player ends the ability prematurely, and a task with this set is still running, the ability is killed. */
diff --git a/Engine/Source/Runtime/AudioExtensions/Public/IAudioCodec.h b/Engine/Source/Runtime/AudioExtensions/Public/IAudioCodec.h
--- a/Engine/Source/Runtime/AudioExtensions/Public/IAudioCodec.h
+++ b/Engine/Source/Runtime/AudioExtensions/Public/IAudioCodec.h
@@ -111,6 +111,7 @@ namespace Audio
 		else
 		{
 			static_assert(sizeof(From) == 0, "Unsupported conversion");
+			return 0.f;
 		}
 	}
 
diff --git a/Engine/Source/Runtime/Core/Public/Misc/CString.h b/Engine/Source/Runtime/Core/Public/Misc/CString.h
--- a/Engine/Source/Runtime/Core/Public/Misc/CString.h
+++ b/Engine/Source/Runtime/Core/Public/Misc/CString.h
@@ -473,6 +473,7 @@ struct TCString
 		else
 		{
 			static_assert(sizeof(CharType) == 0, "Not supported");
+			return false;
 		}
 	}
 
diff --git a/Engine/Source/Runtime/Core/Public/Misc/Char.h b/Engine/Source/Runtime/Core/Public/Misc/Char.h
--- a/Engine/Source/Runtime/Core/Public/Misc/Char.h
+++ b/Engine/Source/Runtime/Core/Public/Misc/Char.h
@@ -107,6 +107,7 @@ struct TChar : TCharBase<CharType, sizeof(CharType)>
 		else
 		{
 			static_assert(sizeof(CharType) == 0, "Not supported");
+			return false;
 		}
 	}
 
@@ -127,6 +128,7 @@ struct TChar : TCharBase<CharType, sizeof(CharType)>
 		else
 		{
 			static_assert(sizeof(CharType) == 0, "Not supported");
+			return false;
 		}
 	}
 
@@ -147,6 +149,7 @@ struct TChar : TCharBase<CharType, sizeof(CharType)>
 		else
 		{
 			static_assert(sizeof(CharType) == 0, "Not supported");
+			return false;
 		}
 	}
 
@@ -167,6 +170,7 @@ struct TChar : TCharBase<CharType, sizeof(CharType)>
 		else
 		{
 			static_assert(sizeof(CharType) == 0, "Not supported");
+			return false;
 		}
 	}
 
@@ -187,6 +191,7 @@ struct TChar : TCharBase<CharType, sizeof(CharType)>
 		else
 		{
 			static_assert(sizeof(CharType) == 0, "Not supported");
+			return false;
 		}
 	}
 
@@ -207,6 +212,7 @@ struct TChar : TCharBase<CharType, sizeof(CharType)>
 		else
 		{
 			static_assert(sizeof(CharType) == 0, "Not supported");
+			return false;
 		}
 	}
 
@@ -227,6 +233,7 @@ struct TChar : TCharBase<CharType, sizeof(CharType)>
 		else
 		{
 			static_assert(sizeof(CharType) == 0, "Not supported");
+			return false;
 		}
 	}
 
@@ -247,6 +254,7 @@ struct TChar : TCharBase<CharType, sizeof(CharType)>
 		else
 		{
 			static_assert(sizeof(CharType) == 0, "Not supported");
+			return false;
 		}
 	}
 
@@ -267,6 +275,7 @@ struct TChar : TCharBase<CharType, sizeof(CharType)>
 		else
 		{
 			static_assert(sizeof(CharType) == 0, "Not supported");
+			return false;
 		}
 	}
 
@@ -287,6 +296,7 @@ struct TChar : TCharBase<CharType, sizeof(CharType)>
 		else
 		{
 			static_assert(sizeof(CharType) == 0, "Not supported");
+			return false;
 		}
 	}
 
@@ -307,6 +317,7 @@ struct TChar : TCharBase<CharType, sizeof(CharType)>
 		else
 		{
 			static_assert(sizeof(CharType) == 0, "Not supported");
+			return false;
 		}
 	}
 
diff --git a/Engine/Source/Runtime/Core/Public/Serialization/MemoryImage.h b/Engine/Source/Runtime/Core/Public/Serialization/MemoryImage.h
--- a/Engine/Source/Runtime/Core/Public/Serialization/MemoryImage.h
+++ b/Engine/Source/Runtime/Core/Public/Serialization/MemoryImage.h
@@ -576,6 +576,7 @@ public:
 	{
 		this->SafeDelete();
 		this->Ptr = InPtr;
+		return *this;
 	}
 	inline TUniqueMemoryImagePtr& operator=(TUniqueMemoryImagePtr&& Other)
 	{
diff --git a/Engine/Source/Runtime/Core/Public/Templates/NonNullPointer.h b/Engine/Source/Runtime/Core/Public/Templates/NonNullPointer.h
--- a/Engine/Source/Runtime/Core/Public/Templates/NonNullPointer.h
+++ b/Engine/Source/Runtime/Core/Public/Templates/NonNullPointer.h
@@ -98,6 +98,7 @@ public:
 	{
 		// Essentially static_assert(false), but this way prevents GCC/Clang from crying wolf by merely inspecting the function body
 		static_assert(sizeof(ObjectType) == 0, "Tried to assign a null pointer to a TNonNullPtr!");
+		return *this;
 	}
 
 	/**
diff --git a/Engine/Source/Runtime/Core/Public/Templates/TypeHash.h b/Engine/Source/Runtime/Core/Public/Templates/TypeHash.h
--- a/Engine/Source/Runtime/Core/Public/Templates/TypeHash.h
+++ b/Engine/Source/Runtime/Core/Public/Templates/TypeHash.h
@@ -143,6 +143,7 @@ inline uint32 GetTypeHash(ScalarType Value)
 	else
 	{
 		static_assert(sizeof(ScalarType) == 0, "Unsupported scalar type");
+		return 0;
 	}
 }
 
diff --git a/Engine/Source/Runtime/CoreUObject/Public/UObject/ObjectPtr.h b/Engine/Source/Runtime/CoreUObject/Public/UObject/ObjectPtr.h
--- a/Engine/Source/Runtime/CoreUObject/Public/UObject/ObjectPtr.h
+++ b/Engine/Source/Runtime/CoreUObject/Public/UObject/ObjectPtr.h
@@ -1647,6 +1647,7 @@ public:
 	{
 		// Essentially static_assert(false), but this way prevents GCC/Clang from crying wolf by merely inspecting the function body
 		static_assert(sizeof(ObjectType) == 0, "Tried to assign a null pointer to a TNonNullPtr!");
+		return *this;
 	}
 
 	/**
diff --git a/Engine/Source/Runtime/Datasmith/DatasmithCore/Public/DatasmithMaterialElements.h b/Engine/Source/Runtime/Datasmith/DatasmithCore/Public/DatasmithMaterialElements.h
--- a/Engine/Source/Runtime/Datasmith/DatasmithCore/Public/DatasmithMaterialElements.h
+++ b/Engine/Source/Runtime/Datasmith/DatasmithCore/Public/DatasmithMaterialElements.h
@@ -259,6 +259,7 @@ public:
 	template< typename T >
 	T* AddMaterialExpression()
 	{
+		return nullptr;
 	}
 
 	/** Reset all expression to their default values and remove all connections */
diff --git a/Engine/Source/Runtime/Experimental/Chaos/Public/Chaos/Framework/Handles.h b/Engine/Source/Runtime/Experimental/Chaos/Public/Chaos/Framework/Handles.h
--- a/Engine/Source/Runtime/Experimental/Chaos/Public/Chaos/Framework/Handles.h
+++ b/Engine/Source/Runtime/Experimental/Chaos/Public/Chaos/Framework/Handles.h
@@ -541,11 +541,13 @@ namespace Chaos
 		THandleHeap& operator =(const THandleHeap<ElementType, IndexWidth, GenerationWidth>& Other)
 		{
 			CopyFrom(Other);
+			return *this;
 		}
 
 		THandleHeap& operator =(THandleHeap<ElementType, IndexWidth, GenerationWidth>&& Other)
 		{
 			MoveFrom(Other);
+			return *this;
 		}
 
 		THandleHeap()
diff --git a/Engine/Source/Runtime/MovieScene/Public/EntitySystem/MovieScenePropertyComponentHandler.h b/Engine/Source/Runtime/MovieScene/Public/EntitySystem/MovieScenePropertyComponentHandler.h
--- a/Engine/Source/Runtime/MovieScene/Public/EntitySystem/MovieScenePropertyComponentHandler.h
+++ b/Engine/Source/Runtime/MovieScene/Public/EntitySystem/MovieScenePropertyComponentHandler.h
@@ -634,6 +634,7 @@ struct TPropertyDefinitionBuilder
 	TPropertyDefinitionBuilder<PropertyTraits>& SetStat(TStatId InStatID)
 	{
 		Definition->StatID = InStatID;
+		return *this;
 	}
 
 	template<typename BlenderSystemType>
diff --git a/Engine/Source/Runtime/Net/Core/Public/Net/Core/Trace/Private/NetTraceInternal.h b/Engine/Source/Runtime/Net/Core/Public/Net/Core/Trace/Private/NetTraceInternal.h
--- a/Engine/Source/Runtime/Net/Core/Public/Net/Core/Trace/Private/NetTraceInternal.h
+++ b/Engine/Source/Runtime/Net/Core/Public/Net/Core/Trace/Private/NetTraceInternal.h
@@ -281,6 +281,7 @@ template<typename T>
 uint64 GetObjectIdForNetTrace(const T&)
 {
 	static_assert(sizeof(T) == 0, "Not supported type for NetTraceObjectID, implement uint64 GetObjectIdForNetTrace(const T&)");
+	return 0;
 }
 
 inline uint64 GetObjectIdForNetTrace(const FNetworkGUID& NetGUID)
diff --git a/Engine/Source/Runtime/RenderCore/Public/RenderGraphBlackboard.h b/Engine/Source/Runtime/RenderCore/Public/RenderGraphBlackboard.h
--- a/Engine/Source/Runtime/RenderCore/Public/RenderGraphBlackboard.h
+++ b/Engine/Source/Runtime/RenderCore/Public/RenderGraphBlackboard.h
@@ -132,6 +132,7 @@ private:
 	{
 		// Forces the compiler to only evaluate the assert on a concrete type.
 		static_assert(sizeof(StructType) == 0, "Struct has not been registered with the RDG blackboard. Use RDG_REGISTER_BLACKBOARD_STRUCT to do this.");
+		return FString();
 	}
 
 	static RENDERCORE_API FString GetTypeName(const TCHAR* ClassName, const TCHAR* FileName, uint32 LineNumber);
diff --git a/Engine/Source/Runtime/RenderCore/Public/ShaderParameterStructDeclaration.h b/Engine/Source/Runtime/RenderCore/Public/ShaderParameterStructDeclaration.h
--- a/Engine/Source/Runtime/RenderCore/Public/ShaderParameterStructDeclaration.h
+++ b/Engine/Source/Runtime/RenderCore/Public/ShaderParameterStructDeclaration.h
@@ -21,4 +21,5 @@ const FShaderParametersMetadata* GetForwardDeclaredShaderParametersStructMetadat
 	static_assert(!std::is_same<T, T>::value /* true */,
 		TEXT("Partial uniform buffer struct declaration. Use `DECLARE_UNIFORM_BUFFER_STRUCT()` instead of `class T;`.")
 	);
+	return nullptr;
 }
