Fixed FFeedbackContext not displaying its dialog anymore.
- Issue: Through the legacy APIs, FFeedbackContext inner slow task wasn't being ticked and thus wouldn't create the dialog after the delay. Issue introduced by cl 17573843.
- Solution: Replaced the call to FFeedbackContext::RequestUpdateUI with a call to its slowtask->TickProgress. TickProgress will create the dialog if needed and will call RequestUpdateUI on the feedback context.
- This change makes FSlowTask::EnterProgressFrame and FFeedbackContext::UpdateProgress equivalent.

diff --git a/Engine/Source/Runtime/Core/Private/Misc/FeedbackContext.cpp b/Engine/Source/Runtime/Core/Private/Misc/FeedbackContext.cpp
--- a/Engine/Source/Runtime/Core/Private/Misc/FeedbackContext.cpp
+++ b/Engine/Source/Runtime/Core/Private/Misc/FeedbackContext.cpp
@@ -210,7 +210,7 @@ void FFeedbackContext::UpdateProgress( int32 Numerator, int32 Denominator )
 		LegacyAPIScopes.Last()->TotalAmountOfWork = (float)Denominator;
 		LegacyAPIScopes.Last()->CompletedWork = (float)Numerator;
 		LegacyAPIScopes.Last()->CurrentFrameScope = (float)(Denominator - Numerator);
-		RequestUpdateUI();
+		LegacyAPIScopes.Last()->TickProgress();
 	}
 }
 
