Allow OnlineSubsystemNull and ReserveScheduler to start threads in a forked multithread environment

#10979 tag
#rnx

diff --git a/Engine/Plugins/Online/OnlineSubsystemNull/Source/Private/OnlineSubsystemNull.cpp b/Engine/Plugins/Online/OnlineSubsystemNull/Source/Private/OnlineSubsystemNull.cpp
--- a/Engine/Plugins/Online/OnlineSubsystemNull/Source/Private/OnlineSubsystemNull.cpp
+++ b/Engine/Plugins/Online/OnlineSubsystemNull/Source/Private/OnlineSubsystemNull.cpp
@@ -295,7 +295,7 @@ bool FOnlineSubsystemNull::Init()
 		// Create the online async task thread
 		OnlineAsyncTaskThreadRunnable = new FOnlineAsyncTaskManagerNull(this);
 		check(OnlineAsyncTaskThreadRunnable);
-		OnlineAsyncTaskThread = FRunnableThread::Create(OnlineAsyncTaskThreadRunnable, *FString::Printf(TEXT("OnlineAsyncTaskThreadNull %s(%d)"), *InstanceName.ToString(), TaskCounter.Increment()), 128 * 1024, TPri_Normal);
+		OnlineAsyncTaskThread = FForkProcessHelper::CreateForkableThread(OnlineAsyncTaskThreadRunnable, *FString::Printf(TEXT("OnlineAsyncTaskThreadNull %s(%d)"), *InstanceName.ToString(), TaskCounter.Increment()), 128 * 1024, TPri_Normal);
 		check(OnlineAsyncTaskThread);
 		UE_LOG_ONLINE(Verbose, TEXT("Created thread (ID:%d)."), OnlineAsyncTaskThread->GetThreadID());
 
diff --git a/Engine/Source/Runtime/Core/Private/Async/Fundamental/ReserveScheduler.cpp b/Engine/Source/Runtime/Core/Private/Async/Fundamental/ReserveScheduler.cpp
--- a/Engine/Source/Runtime/Core/Private/Async/Fundamental/ReserveScheduler.cpp
+++ b/Engine/Source/Runtime/Core/Private/Async/Fundamental/ReserveScheduler.cpp
@@ -72,8 +72,10 @@ void FReserveScheduler::StartWorkers(uint32 NumWorkers, FThread::EForkable IsFor
 
 	WorkerPriority = GTaskGraphUseDynamicPrioritization ? LowLevelTasks::FScheduler::Get().GetWorkerPriority() : WorkerPriority;
 
+	const bool bSupportsMultithreading = FPlatformProcess::SupportsMultithreading() || FForkProcessHelper::IsForkedMultithreadInstance();
+
 	uint32 OldActiveWorkers = ActiveWorkers.load(std::memory_order_relaxed);
-	if(OldActiveWorkers == 0 && FPlatformProcess::SupportsMultithreading() && ActiveWorkers.compare_exchange_strong(OldActiveWorkers, NumWorkers, std::memory_order_relaxed))
+	if(OldActiveWorkers == 0 && bSupportsMultithreading && ActiveWorkers.compare_exchange_strong(OldActiveWorkers, NumWorkers, std::memory_order_relaxed))
 	{
 		LLM_SCOPE_BYNAME(TEXT("EngineMisc/WorkerThreads"));
 
