Fix missing include

diff --git a/Engine/Source/Runtime/Core/Private/Async/Fundamental/ReserveScheduler.cpp b/Engine/Source/Runtime/Core/Private/Async/Fundamental/ReserveScheduler.cpp
--- a/Engine/Source/Runtime/Core/Private/Async/Fundamental/ReserveScheduler.cpp
+++ b/Engine/Source/Runtime/Core/Private/Async/Fundamental/ReserveScheduler.cpp
@@ -6,6 +6,7 @@
 #include "Async/Fundamental/Scheduler.h"
 #include "HAL/LowLevelMemTracker.h"
 #include "HAL/PlatformProcess.h"
+#include "Misc/Fork.h"
 #include "Misc/ScopeLock.h"
 #include "ProfilingDebugging/CpuProfilerTrace.h"
 #include "Trace/Trace.h"
