Fix for hore error: https://horde.devtools.epicgames.com/job/6536d9cdf5778e18b82d538f

- unecessary check case had a race condition which could lead to failure. Removed case.

diff --git a/Engine/Source/Runtime/Core/Tests/Async/MutexTest.cpp b/Engine/Source/Runtime/Core/Tests/Async/MutexTest.cpp
--- a/Engine/Source/Runtime/Core/Tests/Async/MutexTest.cpp
+++ b/Engine/Source/Runtime/Core/Tests/Async/MutexTest.cpp
@@ -29,7 +29,6 @@ TEST_CASE("Core::Async::Mutex", "[Core][Async]")
 			// Using FThread for now because UE::Tasks::Launch does not always wake a worker thread.
 			Threads[Index] = FThread(TEXT("MutexTest"), [&Mutex, &TasksComplete]
 			{
-				CHECK(Mutex.IsLocked());
 				while (!Mutex.TryLock()) // spin on attempting to acquire the lock
 				{
 				}
