Trace scope

diff --git a/Engine/Source/Runtime/Core/Private/Misc/FeedbackContextMarkup.cpp b/Engine/Source/Runtime/Core/Private/Misc/FeedbackContextMarkup.cpp
--- a/Engine/Source/Runtime/Core/Private/Misc/FeedbackContextMarkup.cpp
+++ b/Engine/Source/Runtime/Core/Private/Misc/FeedbackContextMarkup.cpp
@@ -42,6 +42,7 @@ bool FFeedbackContextMarkup::ParseCommand(const FString& Line, FFeedbackContext*
 
 bool FFeedbackContextMarkup::PipeProcessOutput(const FText& Description, const FString& URL, const FString& Params, FFeedbackContext* Warn, int32* OutExitCode)
 {
+	TRACE_CPUPROFILER_EVENT_SCOPE_STR("FFeedbackContextMarkup::PipeProcessOutput");
 	bool bRes;
 
 	// Create a read and write pipe for the child process
