Android:
 - Add engine access to java thermal headroom api.
 - Add 'thermal stress' to perf reports. (see powermanager.getthermalheadroom)
 - Add rss as separate property of android mem stats

diff --git a/Engine/Build/Android/Java/src/com/epicgames/unreal/GameActivity.java.template b/Engine/Build/Android/Java/src/com/epicgames/unreal/GameActivity.java.template
--- a/Engine/Build/Android/Java/src/com/epicgames/unreal/GameActivity.java.template
+++ b/Engine/Build/Android/Java/src/com/epicgames/unreal/GameActivity.java.template
@@ -3304,6 +3304,10 @@ public class GameActivity extends $${gameActivitySuperClass}$$ implements Surfac
 		if (Build.VERSION.SDK_INT >= 29)
 		{
 			PowerManager powerManager = (PowerManager) getSystemService(POWER_SERVICE);
+			int status = powerManager.getCurrentThermalStatus();
+
+			Log.debug("=== Setting initial thermal status to " + status);			
+			nativeOnThermalStatusChangedListener(status);
 			powerManager.addThermalStatusListener(getMainExecutor(), new PowerManager.OnThermalStatusChangedListener() {
 				@Override 
 				public void onThermalStatusChanged(int status) 
@@ -4345,6 +4349,20 @@ public class GameActivity extends $${gameActivitySuperClass}$$ implements Surfac
 	    return activities != null && activities.size() > 0;
 	}
 
+
+	public float AndroidThunkJava_getThermalHeadroom(int ForecastSeconds)
+	{
+		if (Build.VERSION.SDK_INT >= 30)
+		{
+			PowerManager powerManager = (PowerManager) getSystemService(POWER_SERVICE);
+			return powerManager.getThermalHeadroom(ForecastSeconds);
+		}
+		else
+		{
+			return -1;
+		}
+	}
+
 	public void AndroidThunkJava_OpenIntentAction(String intentAction)
 	{
 		Intent intent = new Intent(intentAction, null);
@@ -6186,6 +6204,14 @@ public class GameActivity extends $${gameActivitySuperClass}$$ implements Surfac
 		if (ANDROID_BUILD_VERSION >= 24)
 		{
 			Log.debug("==================================> SetSustainedPerformanceMode:"+bEnable);
+			if (Build.VERSION.SDK_INT >= 29)
+			{
+				PowerManager powerManager = (PowerManager) getSystemService(POWER_SERVICE);
+				if(powerManager.isSustainedPerformanceModeSupported() == false)
+				{ 
+					Log.warn("SustainedPerformanceMode is unsupported on this device.");
+				}
+			}
 			_activity.runOnUiThread(new Runnable()
 			{
 				public void run()
diff --git a/Engine/Source/Runtime/Core/Private/Android/AndroidJava.cpp b/Engine/Source/Runtime/Core/Private/Android/AndroidJava.cpp
--- a/Engine/Source/Runtime/Core/Private/Android/AndroidJava.cpp
+++ b/Engine/Source/Runtime/Core/Private/Android/AndroidJava.cpp
@@ -5,6 +5,21 @@
 
 #if USE_ANDROID_JNI
 
+FJavaClassObject FJavaClassObject::GetGameActivity()
+{
+	JNIEnv* JEnv = AndroidJavaEnv::GetJavaEnv();
+	jclass MainClass = AndroidJavaEnv::FindJavaClassGlobalRef("com/epicgames/unreal/GameActivity");
+	jobject Object = JEnv->NewGlobalRef(AndroidJavaEnv::GetGameActivityThis());
+
+	return FJavaClassObject(MainClass, Object);
+}
+
+FJavaClassObject::FJavaClassObject(jclass ClassIN, jobject ObjectIN)
+{
+	Object = ObjectIN;
+	Class = ClassIN;
+}
+
 FJavaClassObject::FJavaClassObject(FName ClassName, const char* CtorSig, ...)
 {
 	JNIEnv*	JEnv = AndroidJavaEnv::GetJavaEnv();
@@ -135,6 +150,18 @@ FString FJavaClassObject::CallMethod<FString>(FJavaClassMethod Method, ...)
 	return Result;
 }
 
+template<>
+float FJavaClassObject::CallMethod<float>(FJavaClassMethod Method, ...)
+{
+	JNIEnv* JEnv = AndroidJavaEnv::GetJavaEnv();
+	va_list Params;
+	va_start(Params, Method);
+	float RetVal = JEnv->CallFloatMethodV(Object, Method.Method, Params);
+	va_end(Params);
+	VerifyException();
+	return RetVal;
+}
+
 FScopedJavaObject<jstring> FJavaClassObject::GetJString(const FString& String)
 {
 	JNIEnv*	JEnv = AndroidJavaEnv::GetJavaEnv();
diff --git a/Engine/Source/Runtime/Core/Private/Android/AndroidPlatformMemory.cpp b/Engine/Source/Runtime/Core/Private/Android/AndroidPlatformMemory.cpp
--- a/Engine/Source/Runtime/Core/Private/Android/AndroidPlatformMemory.cpp
+++ b/Engine/Source/Runtime/Core/Private/Android/AndroidPlatformMemory.cpp
@@ -382,6 +382,7 @@ FPlatformMemoryStats FAndroidPlatformMemory::GetStats()
 			else if (strstr(Line, "VmRSS:") == Line)
 			{
 				MemoryStats.UsedPhysical = AndroidPlatformMemory::GetBytesFromStatusLine(Line);
+				MemoryStats.VMRss = MemoryStats.UsedPhysical;
 				++FieldsSetSuccessfully;
 			}
 			else if (strstr(Line, "VmSwap:") == Line)
diff --git a/Engine/Source/Runtime/Core/Private/Android/AndroidPlatformThermal.cpp b/Engine/Source/Runtime/Core/Private/Android/AndroidPlatformThermal.cpp
new file mode 100644
--- /dev/null
+++ b/Engine/Source/Runtime/Core/Private/Android/AndroidPlatformThermal.cpp
@@ -0,0 +1,94 @@
+// Copyright Epic Games, Inc. All Rights Reserved.
+
+/*=============================================================================
+	AndroidPlatformThermal.cpp: Android thermal API queries.
+=============================================================================*/
+#include "Android/AndroidPlatformThermal.h"
+#include <sys/types.h>
+#include <unistd.h>
+#include <stdlib.h>
+#include <limits.h>
+#include <dlfcn.h>
+#include "Containers/StringConv.h"
+#include "CoreGlobals.h"
+#include "Misc/App.h"
+#include "Misc/Guid.h"
+#include "HAL/PlatformStackWalk.h"
+#include "Misc/Paths.h"
+#include "GenericPlatform/GenericPlatformOutputDevices.h"
+#include "HAL/FileManager.h"
+#include "HAL/IConsoleManager.h"
+#include "HAL/ThreadManager.h"
+#include "HAL/RunnableThread.h"
+#include "GenericPlatform/GenericPlatformCrashContext.h"
+#include "Async/TaskGraphInterfaces.h"
+#include "Android/AndroidJavaEnv.h"
+#include <android/thermal.h>
+#include "Android/AndroidJava.h"
+
+namespace UE
+{
+	namespace Android
+	{
+		class FAndroidThermalManagerImpl
+		{
+			static constexpr float CallRateLimitSeconds = 1.0f;
+			static constexpr int NumPeriods = FAndroidPlatformThermal::EForecastPeriod::NUM_FORECASTPERIODS;
+			static constexpr int PeriodEnumToS[NumPeriods] = { 1,5,10 };
+			static_assert(NumPeriods == 3); // update PeriodEnumToS if this changes.
+			double LastCallTime = 0;
+
+			FAndroidPlatformThermal::EForecastPeriod PeriodToRead = FAndroidPlatformThermal::EForecastPeriod::ONE_SEC;
+
+			FJavaClassObject ActivityClass;			// note: ActivityClass must always be accessed from the same thread that created it.
+			FJavaClassMethod ThermalHeadroomMethod;
+
+			FAndroidThermalManagerImpl() : ActivityClass(FJavaClassObject::GetGameActivity())
+			{
+				check(IsInGameThread());
+				ThermalHeadroomMethod = ActivityClass.GetClassMethod("AndroidThunkJava_getThermalHeadroom", "(I)F");
+			}
+
+		public:
+			static FAndroidThermalManagerImpl& Get()
+			{
+				static FAndroidThermalManagerImpl AndroidThermalManager;
+				return AndroidThermalManager;				
+			}
+
+			struct FThermalInfo
+			{
+				float HeadroomForecasts[NumPeriods];
+			}ThermalInfo;
+
+			const FThermalInfo& UpdateAndGetThermalInfo()
+			{
+				check(IsInGameThread());
+				float CurrentTime = (float)FPlatformTime::Seconds();
+				if (CurrentTime > LastCallTime + CallRateLimitSeconds)
+				{
+					/// thermal headroom docs state we should not call this API more than once per second.
+					float RecentThermalReading = ActivityClass.CallMethod<float>(ThermalHeadroomMethod, PeriodEnumToS[PeriodToRead]);
+					ThermalInfo.HeadroomForecasts[PeriodToRead] = RecentThermalReading;
+
+					LastCallTime = CurrentTime;
+					PeriodToRead = (FAndroidPlatformThermal::EForecastPeriod)((int)(PeriodToRead+1) % NumPeriods);
+				}
+				return ThermalInfo;
+			}
+
+		};
+	}
+}
+
+float FAndroidPlatformThermal::GetThermalStress(FAndroidPlatformThermal::EForecastPeriod ForecastPeriod)
+{
+	int Idx = (int)ForecastPeriod;
+	if (ensure(Idx >= 0 && Idx < (int)FAndroidPlatformThermal::EForecastPeriod::NUM_FORECASTPERIODS))
+	{
+		const UE::Android::FAndroidThermalManagerImpl::FThermalInfo& ThermalInfo = UE::Android::FAndroidThermalManagerImpl::Get().UpdateAndGetThermalInfo();
+		return ThermalInfo.HeadroomForecasts[ForecastPeriod];
+	}
+	return -1;
+}
+
diff --git a/Engine/Source/Runtime/Core/Private/Android/AndroidStats.cpp b/Engine/Source/Runtime/Core/Private/Android/AndroidStats.cpp
--- a/Engine/Source/Runtime/Core/Private/Android/AndroidStats.cpp
+++ b/Engine/Source/Runtime/Core/Private/Android/AndroidStats.cpp
@@ -17,6 +17,7 @@
 #if HWCPIPE_SUPPORTED
 #	include "hwcpipe.h"
 #endif
+#include "Android/AndroidPlatformThermal.h"
 
 DECLARE_STATS_GROUP(TEXT("Android CPU stats"), STATGROUP_AndroidCPU, STATCAT_Advanced);
 CSV_DEFINE_CATEGORY(AndroidCPU, true);
@@ -68,6 +69,9 @@ DECLARE_FLOAT_COUNTER_STAT(TEXT("CPU Temperature"), STAT_CPUTemp, STATGROUP_Andr
 CSV_DEFINE_STAT(AndroidCPU, ThermalStatus);
 DECLARE_FLOAT_COUNTER_STAT(TEXT("Thermal Status"), STAT_ThermalStatus, STATGROUP_AndroidCPU);
 
+CSV_DEFINE_STAT(AndroidCPU, ThermalStress);
+DECLARE_FLOAT_COUNTER_STAT(TEXT("Thermal Stress"), STAT_ThermalStress, STATGROUP_AndroidCPU);
+
 #if STATS
 
 #define SET_DWORD_STAT_BY_FNAME(Stat, Amount) \
@@ -157,6 +161,11 @@ static FAutoConsoleVariableRef CVarAndroidCollectCPUStatsRate(
 static int GThermalStatus = 0;
 static int GTrimMemoryBackgroundLevel = 0;
 CSV_DEFINE_STAT(AndroidMemory, TrimMemoryBackgroundLevel);
+
+CSV_DEFINE_STAT(AndroidMemory, Mem_RSS);
+CSV_DEFINE_STAT(AndroidMemory, Mem_Swap);
+CSV_DEFINE_STAT(AndroidMemory, Mem_TotalUsed);
+
 static int GTrimMemoryForegroundLevel = 0;
 CSV_DEFINE_STAT(AndroidMemory, TrimMemoryForegroundLevel);
 
@@ -253,18 +262,28 @@ void FAndroidStats::UpdateAndroidStats()
 	static uint64 LastCollectionTime = FPlatformTime::Cycles64();
 	const uint64 CurrentTime = FPlatformTime::Cycles64();
 	const bool bUpdateStats = ((FPlatformTime::ToSeconds64(CurrentTime - LastCollectionTime) >= GAndroidCPUStatsUpdateRate));
+	static FPlatformMemoryStats MemStats = FAndroidPlatformMemory::GetStats();
 	if (bUpdateStats)
 	{
 		LastCollectionTime = CurrentTime;
 		CPUTemp = FAndroidMisc::GetCPUTemperature();
 
 		GDeviceMemAdviceStats = FAndroidPlatformMemory::GetDeviceMemAdviceStats();
+		MemStats = FAndroidPlatformMemory::GetStats();
 	}
 
 	CSV_CUSTOM_STAT_DEFINED(CPUTemp, CPUTemp, ECsvCustomStatOp::Set);
 	CSV_CUSTOM_STAT_DEFINED(ThermalStatus, GThermalStatus, ECsvCustomStatOp::Set);
-	CSV_CUSTOM_STAT_DEFINED(TrimMemoryBackgroundLevel, GTrimMemoryBackgroundLevel, ECsvCustomStatOp::Set);
 
+
+	float Thermals5S = FAndroidPlatformThermal::GetThermalStress(FAndroidPlatformThermal::EForecastPeriod::FIVE_SEC);
+	CSV_CUSTOM_STAT_DEFINED(ThermalStress, Thermals5S, ECsvCustomStatOp::Set);
+	
+	CSV_CUSTOM_STAT_DEFINED(TrimMemoryBackgroundLevel, GTrimMemoryBackgroundLevel, ECsvCustomStatOp::Set);
+	CSV_CUSTOM_STAT_DEFINED(Mem_Swap, (int)(MemStats.VMSwap / (1024 * 1024)), ECsvCustomStatOp::Set);
+	CSV_CUSTOM_STAT_DEFINED(Mem_RSS, (int)(MemStats.VMRss / (1024 * 1024)), ECsvCustomStatOp::Set);
+	CSV_CUSTOM_STAT_DEFINED(Mem_TotalUsed, (int)(MemStats.UsedPhysical / (1024 * 1024)), ECsvCustomStatOp::Set);
+	
 	if(GDeviceMemAdviceStats.IsValid())
 	{
 		CSV_CUSTOM_STAT_DEFINED(MemAdvisor_DeviceMemFreeMb, (int)(GDeviceMemAdviceStats.MemFree/(1024*1024)), ECsvCustomStatOp::Set);
diff --git a/Engine/Source/Runtime/Core/Public/Android/AndroidJava.h b/Engine/Source/Runtime/Core/Public/Android/AndroidJava.h
--- a/Engine/Source/Runtime/Core/Public/Android/AndroidJava.h
+++ b/Engine/Source/Runtime/Core/Public/Android/AndroidJava.h
@@ -21,7 +21,11 @@ struct FJavaClassMethod
 
 class FJavaClassObject
 {
+	FJavaClassObject(jclass ClassIN, jobject ObjectIN);
+
 public:
+	static FJavaClassObject GetGameActivity();
+
 	// !!  All Java objects returned by JNI functions are local references.
 	FJavaClassObject(FName ClassName, const char* CtorSig, ...);
 
@@ -73,4 +77,7 @@ int64 FJavaClassObject::CallMethod<int64>(FJavaClassMethod Method, ...);
 template<>
 FString FJavaClassObject::CallMethod<FString>(FJavaClassMethod Method, ...);
 
+template<>
+float FJavaClassObject::CallMethod<float>(FJavaClassMethod Method, ...);
+
 #endif
diff --git a/Engine/Source/Runtime/Core/Public/Android/AndroidPlatformMemory.h b/Engine/Source/Runtime/Core/Public/Android/AndroidPlatformMemory.h
--- a/Engine/Source/Runtime/Core/Public/Android/AndroidPlatformMemory.h
+++ b/Engine/Source/Runtime/Core/Public/Android/AndroidPlatformMemory.h
@@ -15,6 +15,7 @@ struct FPlatformMemoryStats : public FGenericPlatformMemoryStats
 {
 	FGenericPlatformMemoryStats::EMemoryPressureStatus GetMemoryPressureStatus() const;
 	uint64 VMSwap = 0;
+	uint64 VMRss = 0;
 };
 
 /**
diff --git a/Engine/Source/Runtime/Core/Public/Android/AndroidPlatformThermal.h b/Engine/Source/Runtime/Core/Public/Android/AndroidPlatformThermal.h
new file mode 100644
--- /dev/null
+++ b/Engine/Source/Runtime/Core/Public/Android/AndroidPlatformThermal.h
@@ -0,0 +1,14 @@
+// Copyright Epic Games, Inc. All Rights Reserved.
+
+#pragma once
+
+struct FAndroidPlatformThermal
+{
+	enum EForecastPeriod { ONE_SEC, FIVE_SEC, TEN_SEC, NUM_FORECASTPERIODS };
+	// returns the 'thermal stress' value for the forecasted period.
+	// This value is a hint for the extent of thermal throttling the device will experience during the time scale provided.
+	// The implementation is provided by PowerManager.getThermalHeadroom(). 
+	// Values of 1.0 or above indicate the current workload can not be sustained and thermal throttling will occur.
+	static float GetThermalStress(EForecastPeriod ForecastPeriod);
+};
+
