Fix an assert in HitSet when accidently inserting an 0x0 Payload into HitSet

Example was a memset(0x0, 0), nullptr with size 0 ended up getting stored causing the HitSet size to become mis-matched

diff --git a/Engine/Source/Runtime/Core/Private/AutoRTFM/TransactionInlines.h b/Engine/Source/Runtime/Core/Private/AutoRTFM/TransactionInlines.h
--- a/Engine/Source/Runtime/Core/Private/AutoRTFM/TransactionInlines.h
+++ b/Engine/Source/Runtime/Core/Private/AutoRTFM/TransactionInlines.h
@@ -34,6 +34,11 @@ AUTORTFM_NO_ASAN UE_AUTORTFM_FORCEINLINE void FTransaction::RecordWriteMaxPageSi
 
 AUTORTFM_NO_ASAN UE_AUTORTFM_FORCEINLINE void FTransaction::RecordWrite(void* LogicalAddress, size_t Size)
 {
+    if (Size == 0)
+    {
+        return;
+    }
+
     // If we are recording a stack address that is relative to our current
     // transactions stack location, we do not need to record the data in the
     // write log because if that transaction aborted, that memory will cease to
