Skip slow tests on CI

diff --git a/Engine/Source/Runtime/Core/Tests/Async/ExternalMutexTest.cpp b/Engine/Source/Runtime/Core/Tests/Async/ExternalMutexTest.cpp
--- a/Engine/Source/Runtime/Core/Tests/Async/ExternalMutexTest.cpp
+++ b/Engine/Source/Runtime/Core/Tests/Async/ExternalMutexTest.cpp
@@ -11,7 +11,7 @@
 namespace UE
 {
 
-TEST_CASE("Core::Async::ExternalMutex", "[Core][Async]")
+TEST_CASE("Core::Async::ExternalMutex", "[Core][Async][Slow]")
 {
 	constexpr static int32 TaskCount = 100;
 
diff --git a/Engine/Source/Runtime/Core/Tests/Async/MutexTest.cpp b/Engine/Source/Runtime/Core/Tests/Async/MutexTest.cpp
--- a/Engine/Source/Runtime/Core/Tests/Async/MutexTest.cpp
+++ b/Engine/Source/Runtime/Core/Tests/Async/MutexTest.cpp
@@ -11,7 +11,7 @@
 namespace UE
 {
 
-TEST_CASE("Core::Async::Mutex", "[Core][Async]")
+TEST_CASE("Core::Async::Mutex", "[Core][Async][Slow]")
 {
 	constexpr static int32 TaskCount = 100;
 
