Reduce number of threads used in Mutex low-level tests to reduce job time on horde.

diff --git a/Engine/Source/Runtime/Core/Tests/Async/ExternalMutexTest.cpp b/Engine/Source/Runtime/Core/Tests/Async/ExternalMutexTest.cpp
--- a/Engine/Source/Runtime/Core/Tests/Async/ExternalMutexTest.cpp
+++ b/Engine/Source/Runtime/Core/Tests/Async/ExternalMutexTest.cpp
@@ -13,7 +13,7 @@ namespace UE
 
 TEST_CASE("Core::Async::ExternalMutex", "[Core][Async][Slow]")
 {
-	constexpr static int32 TaskCount = 100;
+	constexpr static int32 TaskCount = 5;
 
 	SECTION("FExternalMutex IsLocked() and TryLock()")
 	{
diff --git a/Engine/Source/Runtime/Core/Tests/Async/MutexTest.cpp b/Engine/Source/Runtime/Core/Tests/Async/MutexTest.cpp
--- a/Engine/Source/Runtime/Core/Tests/Async/MutexTest.cpp
+++ b/Engine/Source/Runtime/Core/Tests/Async/MutexTest.cpp
@@ -13,7 +13,7 @@ namespace UE
 
 TEST_CASE("Core::Async::Mutex", "[Core][Async][Slow]")
 {
-	constexpr static int32 TaskCount = 100;
+	constexpr static int32 TaskCount = 5;
 
 	SECTION("FMutex IsLocked() and TryLock()")
 	{
