Optimized VectorAnyGreaterThan for NEON

diff --git a/Engine/Source/Runtime/Core/Public/Math/UnrealMathNeon.h b/Engine/Source/Runtime/Core/Public/Math/UnrealMathNeon.h
--- a/Engine/Source/Runtime/Core/Public/Math/UnrealMathNeon.h
+++ b/Engine/Source/Runtime/Core/Public/Math/UnrealMathNeon.h
@@ -2217,26 +2217,17 @@ FORCEINLINE void VectorStoreURGB10A2N(const VectorRegister4Float& Vec, void* Ptr
  * @param Vec2			2nd source vector
  * @return				Non-zero integer if (Vec1.x > Vec2.x) || (Vec1.y > Vec2.y) || (Vec1.z > Vec2.z) || (Vec1.w > Vec2.w)
  */
-FORCEINLINE int32 VectorAnyGreaterThan( VectorRegister4Float Vec1, VectorRegister4Float Vec2 )
+FORCEINLINE int32 VectorAnyGreaterThan(VectorRegister4Float Vec1, VectorRegister4Float Vec2)
 {
-	uint16x8_t u16x8 = (uint16x8_t)vcgtq_f32( Vec1, Vec2 );
-	uint8x8_t u8x8 = (uint8x8_t)vget_low_u16( vuzpq_u16( u16x8, u16x8 ).val[0] );
-	u8x8 = vuzp_u8( u8x8, u8x8 ).val[0];
-	uint32_t buf[2];
-	vst1_u8( (uint8_t *)buf, u8x8 );
-	return (int32)buf[0]; // each byte of output corresponds to a component comparison
+	uint32x4_t Mask = (uint32x4_t)VectorCompareGT(Vec1, Vec2);
+	return vmaxvq_u32(Mask);
 }
 
 FORCEINLINE int32 VectorAnyGreaterThan(VectorRegister4Double Vec1, VectorRegister4Double Vec2)
 {
-	uint16x8_t u16x8_1 = (uint16x8_t)vcgtq_f64(Vec1.XY, Vec2.XY);
-	uint16x8_t u16x8_2 = (uint16x8_t)vcgtq_f64(Vec1.ZW, Vec2.ZW);
-	uint16x8x2_t tmp = vuzpq_u16(u16x8_1, u16x8_2);
-	uint8x8_t u8x8 = (uint8x8_t)vget_low_u16(vuzpq_u16(tmp.val[0], tmp.val[0]).val[0]);
-	u8x8 = vuzp_u8(u8x8, u8x8).val[0];
-	uint32_t buf[2];
-	vst1_u8((uint8_t*)buf, u8x8);
-	return (int32)buf[0]; // each byte of output corresponds to a component comparison
+	uint32x4_t MaskXY = (uint32x4_t)vcgtq_f64(Vec1.XY, Vec2.XY);
+	uint32x4_t MaskZW = (uint32x4_t)vcgtq_f64(Vec1.ZW, Vec2.ZW);
+	return vmaxvq_u32(MaskXY) || vmaxvq_u32(MaskZW);
 }
 
 /**
