Fix CIS issues thrown when android JNI is not enabled

diff --git a/Engine/Source/Runtime/Core/Private/Android/AndroidPlatformThermal.cpp b/Engine/Source/Runtime/Core/Private/Android/AndroidPlatformThermal.cpp
--- a/Engine/Source/Runtime/Core/Private/Android/AndroidPlatformThermal.cpp
+++ b/Engine/Source/Runtime/Core/Private/Android/AndroidPlatformThermal.cpp
@@ -26,6 +26,7 @@
 #include <android/thermal.h>
 #include "Android/AndroidJava.h"
 
+#if USE_ANDROID_JNI
 namespace UE
 {
 	namespace Android
@@ -91,4 +92,6 @@ float FAndroidPlatformThermal::GetThermalStress(FAndroidPlatformThermal::EForeca
 	}
 	return -1;
 }
-
+#else
+float FAndroidPlatformThermal::GetThermalStress(FAndroidPlatformThermal::EForecastPeriod ForecastPeriod) { return -1.0; }
+#endif
