[TSAN] - Fix warning because we were reusing a memory block without acquire semantic

#rnx

diff --git a/Engine/Source/Runtime/Core/Public/Experimental/ConcurrentLinearAllocator.h b/Engine/Source/Runtime/Core/Public/Experimental/ConcurrentLinearAllocator.h
--- a/Engine/Source/Runtime/Core/Public/Experimental/ConcurrentLinearAllocator.h
+++ b/Engine/Source/Runtime/Core/Public/Experimental/ConcurrentLinearAllocator.h
@@ -268,7 +268,7 @@ class TConcurrentLinearAllocator
 
 				// on the allocating side we only need to do a single atomic to reduce contention with the deletions
 				// this will leave the atomic in a state where it only counts the number of live allocations (before it was based of UINT_MAX)
-				if (Header->NumAllocations.fetch_sub(DeltaCount, std::memory_order_release) == DeltaCount)
+				if (Header->NumAllocations.fetch_sub(DeltaCount, std::memory_order_acq_rel) == DeltaCount)
 				{
 					//if all allocations are already freed we can reuse the Block again
 					Header->~FBlockHeader();
@@ -431,7 +431,7 @@ public:
 
 		// on the allocating side we only need to do a single atomic to reduce contention with the deletions
 		// this will leave the atomic in a state where it only counts the number of live allocations (before it was based of UINT_MAX)
-		if (Header->NumAllocations.fetch_sub(DeltaCount, std::memory_order_release) == DeltaCount)
+		if (Header->NumAllocations.fetch_sub(DeltaCount, std::memory_order_acq_rel) == DeltaCount)
 		{
 			//if all allocations are already freed we can reuse the Block again
 			Header->~FBlockHeader();
