Address MonotonicTimeTest failures due to an optimization bug in some compilers.

diff --git a/Engine/Source/Runtime/Core/Tests/Misc/MonotonicTimeTest.cpp b/Engine/Source/Runtime/Core/Tests/Misc/MonotonicTimeTest.cpp
--- a/Engine/Source/Runtime/Core/Tests/Misc/MonotonicTimeTest.cpp
+++ b/Engine/Source/Runtime/Core/Tests/Misc/MonotonicTimeTest.cpp
@@ -6,6 +6,8 @@
 
 #include "Tests/TestHarnessAdapter.h"
 
+UE_DISABLE_OPTIMIZATION_SHIP
+
 namespace UE
 {
 
@@ -139,4 +141,6 @@ TEST_CASE_NAMED(FMonotonicTimePointTest, "System::Core::Time::MonotonicTimePoint
 
 } // namespace UE
 
+UE_ENABLE_OPTIMIZATION_SHIP
+
 #endif // WITH_TESTS
