Fixed undefined behavior in ManualResetEvent

diff --git a/Engine/Source/Runtime/Core/Public/Async/ManualResetEvent.h b/Engine/Source/Runtime/Core/Public/Async/ManualResetEvent.h
--- a/Engine/Source/Runtime/Core/Public/Async/ManualResetEvent.h
+++ b/Engine/Source/Runtime/Core/Public/Async/ManualResetEvent.h
@@ -66,7 +66,7 @@ public:
 				return false;
 			}
 			if ((CurrentState & MaybeWaitingFlag) ||
-				State.compare_exchange_weak(CurrentState, MaybeWaitingFlag, std::memory_order_relaxed, std::memory_order_acquire))
+				State.compare_exchange_weak(CurrentState, MaybeWaitingFlag, std::memory_order_acq_rel))
 			{
 				ParkingLot::WaitUntil(&State, [this] { return !IsNotified(); }, []{}, WaitTime);
 			}
@@ -90,7 +90,7 @@ public:
 	 */
 	inline void Reset()
 	{
-		State.fetch_and(~IsNotifiedFlag, std::memory_order_relaxed);
+		State.fetch_and(~IsNotifiedFlag, std::memory_order_release);
 	}
 
 private:
