Guard FICUInternationalization::PathToCachedFileDataMap access with a critical section. Map can be modified by multiple threads (i.e. game thread and async loading thread)

#rnx

diff --git a/Engine/Source/Runtime/Core/Private/Internationalization/ICUInternationalization.cpp b/Engine/Source/Runtime/Core/Private/Internationalization/ICUInternationalization.cpp
--- a/Engine/Source/Runtime/Core/Private/Internationalization/ICUInternationalization.cpp
+++ b/Engine/Source/Runtime/Core/Private/Internationalization/ICUInternationalization.cpp
@@ -205,6 +205,7 @@ void FICUInternationalization::Terminate()
 
 	u_cleanup();
 
+	FScopeLock Lock(&PathToCachedFileDataMapCS);
 	for (auto& PathToCachedFileDataPair : PathToCachedFileDataMap)
 	{
 		UE_LOG(LogICUInternationalization, Warning, TEXT("ICU data file '%s' (ref count %d) was still referenced after ICU shutdown. This will likely lead to a crash."), *PathToCachedFileDataPair.Key, PathToCachedFileDataPair.Value.ReferenceCount);
@@ -775,44 +776,47 @@ UBool FICUInternationalization::OpenDataFile(const void* InContext, void** OutFi
 	FString PathStr = StringCast<TCHAR>(InPath).Get();
 	FPaths::NormalizeFilename(PathStr);
 
-	FICUCachedFileData* CachedFileData = nullptr;
-
 	// Skip requests for anything outside the ICU data directory
 	const bool bIsWithinDataDirectory = PathStr.StartsWith(This->ICUDataDirectory);
-	if (bIsWithinDataDirectory)
+	if (!bIsWithinDataDirectory)
 	{
-		CachedFileData = This->PathToCachedFileDataMap.Find(PathStr);
+		*OutFileContext = nullptr;
+		*OutContents = nullptr;
+		return false;
+	}
 
-		// If there's no file context, we might have to load the file.
-		if (!CachedFileData)
-		{
+	FScopeLock Lock(&This->PathToCachedFileDataMapCS);
+	FICUCachedFileData* CachedFileData = This->PathToCachedFileDataMap.Find(PathStr);
+
+	// If there's no file context, we might have to load the file.
+	if (!CachedFileData)
+	{
 #if !UE_BUILD_SHIPPING
-			FScopedLoadingState ScopedLoadingState(*PathStr);
+		FScopedLoadingState ScopedLoadingState(*PathStr);
 #endif
 
-			// Attempt to load the file.
-			FArchive* FileAr = IFileManager::Get().CreateFileReader(*PathStr);
-			if (FileAr)
-			{
-				const int64 FileSize = FileAr->TotalSize();
+		// Attempt to load the file.
+		FArchive* FileAr = IFileManager::Get().CreateFileReader(*PathStr);
+		if (FileAr)
+		{
+			const int64 FileSize = FileAr->TotalSize();
 
-				// Create file data.
-				CachedFileData = &This->PathToCachedFileDataMap.Emplace(PathStr, FICUCachedFileData(FileSize));
+			// Create file data.
+			CachedFileData = &This->PathToCachedFileDataMap.Emplace(PathStr, FICUCachedFileData(FileSize));
 
-				// Load file into buffer.
-				FileAr->Serialize(CachedFileData->Buffer, FileSize);
-				delete FileAr;
+			// Load file into buffer.
+			FileAr->Serialize(CachedFileData->Buffer, FileSize);
+			delete FileAr;
 
-				// Stat tracking.
+			// Stat tracking.
 #if STATS
-				DataFileBytesInUseCount += FMemory::GetAllocSize(CachedFileData->Buffer);
-				if (FThreadStats::IsThreadingReady() && CachedDataFileBytesInUseCount != DataFileBytesInUseCount)
-				{
-					SET_MEMORY_STAT(STAT_MemoryICUDataFileAllocationSize, DataFileBytesInUseCount);
-					CachedDataFileBytesInUseCount = DataFileBytesInUseCount;
-				}
-#endif
+			DataFileBytesInUseCount += FMemory::GetAllocSize(CachedFileData->Buffer);
+			if (FThreadStats::IsThreadingReady() && CachedDataFileBytesInUseCount != DataFileBytesInUseCount)
+			{
+				SET_MEMORY_STAT(STAT_MemoryICUDataFileAllocationSize, DataFileBytesInUseCount);
+				CachedDataFileBytesInUseCount = DataFileBytesInUseCount;
 			}
+#endif
 		}
 	}
 
@@ -851,6 +855,7 @@ void FICUInternationalization::CloseDataFile(const void* InContext, void* const
 	check(Path);
 
 	// Look up the cached file data so we can maintain references.
+	FScopeLock Lock(&This->PathToCachedFileDataMapCS);
 	FICUCachedFileData* const CachedFileData = This->PathToCachedFileDataMap.Find(*Path);
 	check(CachedFileData);
 	check(CachedFileData->Buffer == InContents);
diff --git a/Engine/Source/Runtime/Core/Private/Internationalization/ICUInternationalization.h b/Engine/Source/Runtime/Core/Private/Internationalization/ICUInternationalization.h
--- a/Engine/Source/Runtime/Core/Private/Internationalization/ICUInternationalization.h
+++ b/Engine/Source/Runtime/Core/Private/Internationalization/ICUInternationalization.h
@@ -130,6 +130,7 @@ private:
 
 	// Map for associating ICU data file paths with cached file data, to prevent multiple copies of immutable ICU data files from residing in memory.
 	TMap<FString, FICUCachedFileData> PathToCachedFileDataMap;
+	FCriticalSection PathToCachedFileDataMapCS;
 };
 
 #endif
