- Fixed the ProjectDir for programs that are in Restricted

diff --git a/Engine/Source/Runtime/Core/Private/GenericPlatform/GenericPlatformMisc.cpp b/Engine/Source/Runtime/Core/Private/GenericPlatform/GenericPlatformMisc.cpp
--- a/Engine/Source/Runtime/Core/Private/GenericPlatform/GenericPlatformMisc.cpp
+++ b/Engine/Source/Runtime/Core/Private/GenericPlatform/GenericPlatformMisc.cpp
@@ -1277,8 +1277,11 @@ const TCHAR* FGenericPlatformMisc::ProjectDir()
 		ProjectDir.Reserve(FPlatformMisc::GetMaxPathLength());
 		if (FPlatformProperties::IsProgram())
 		{
-			// monolithic, game-agnostic executables, the ini is in Engine/Config/Platform
-			ProjectDir = FString::Printf(TEXT("../../../Engine/Programs/%s/"), FApp::GetProjectName());
+			// programs are under a Programs directory, expected to be two up from the Binaries directory (can't assume Engine)
+			// so Engine/Binaries/Mac/UnrealPak, the ProjectDir would be Engine/Programs/UnrealPak, but for
+			// Engine/Restricted/NoRedist/Binaries/Mac/SomeProgram, the ProjectDir would be
+			// Engine/Restricted/NoRedist/Programs/UnrealPak - so we go up two and then into Programs
+			ProjectDir = FString::Printf(TEXT("../../Programs/%s/"), FApp::GetProjectName());
 
 			// however, if it was staged, that directory won't exist, so look in the normal staged location
 			if (!FPlatformFileManager::Get().GetPlatformFile().DirectoryExists(*ProjectDir))
