- Changes needed to run staged/standalone programs that are in Restricted location

diff --git a/Engine/Source/Runtime/Core/Private/Windows/WindowsPlatformProcess.cpp b/Engine/Source/Runtime/Core/Private/Windows/WindowsPlatformProcess.cpp
--- a/Engine/Source/Runtime/Core/Private/Windows/WindowsPlatformProcess.cpp
+++ b/Engine/Source/Runtime/Core/Private/Windows/WindowsPlatformProcess.cpp
@@ -1119,7 +1119,12 @@ const TCHAR* FWindowsPlatformProcess::BaseDir()
 
 			FString CollapseResult(Result);
 #ifdef UE_RELATIVE_BASE_DIR
-			CollapseResult /= UE_RELATIVE_BASE_DIR;
+			// this may have been defined at compile time because we are in Restricted, but then we have been staged as a program, and then remapped out of Restricted
+			// so if we are already in a Binaries/Win64 directory
+			if (!CollapseResult.EndsWith(TEXT("/Binaries/Win64/")))
+			{
+				CollapseResult /= UE_RELATIVE_BASE_DIR;
+			}
 #endif
 			FPaths::CollapseRelativeDirectories(CollapseResult);
 			FCString::Strcpy(Result, *CollapseResult);
