PR #10999: Fix missing condition in IOSPlatformFile::IsReadOnly

#rnx

diff --git a/Engine/Source/Runtime/Core/Private/IOS/IOSPlatformFile.cpp b/Engine/Source/Runtime/Core/Private/IOS/IOSPlatformFile.cpp
--- a/Engine/Source/Runtime/Core/Private/IOS/IOSPlatformFile.cpp
+++ b/Engine/Source/Runtime/Core/Private/IOS/IOSPlatformFile.cpp
@@ -648,7 +648,7 @@ bool FIOSPlatformFile::IsReadOnly(const TCHAR* Filename)
 
 	if (access(TCHAR_TO_UTF8(*Filepath), W_OK) == -1)
 	{
-		return errno == EACCES;
+		return errno == EPERM || errno == EACCES;
 	}
 	return false;
 }
