[ALT] - Add special case for EmptyDelegate that is often save-serialized from different threads

#rnx

diff --git a/Engine/Source/Runtime/Core/Public/UObject/ScriptDelegates.h b/Engine/Source/Runtime/Core/Public/UObject/ScriptDelegates.h
--- a/Engine/Source/Runtime/Core/Public/UObject/ScriptDelegates.h
+++ b/Engine/Source/Runtime/Core/Public/UObject/ScriptDelegates.h
@@ -837,20 +837,31 @@ public:
 	/** Multi-cast delegate serialization */
 	friend FArchive& operator<<( FArchive& Ar, TMulticastScriptDelegate& D )
 	{
-		FWriteAccessScope WriteScope = D.GetWriteAccessScope();
-
-		if( Ar.IsSaving() )
+		// Special case to avoid taking a lock on empty script delegate.
+		// This is required for avoiding asserts on EmptyDelegate serialization.
+		if (Ar.IsSaving() && D.InvocationList.Num() == 0)
 		{
-			// When saving the delegate, clean up the list to make sure there are no bad object references
-			D.CompactInvocationList();
+			FReadAccessScope ReadScope = D.GetReadAccessScope();
+
+			Ar << D.InvocationList;
 		}
+		else
+		{
+			FWriteAccessScope WriteScope = D.GetWriteAccessScope();
 
-		Ar << D.InvocationList;
+			if( Ar.IsSaving() )
+			{
+				// When saving the delegate, clean up the list to make sure there are no bad object references
+				D.CompactInvocationList();
+			}
 
-		if( Ar.IsLoading() )
-		{
-			// After loading the delegate, clean up the list to make sure there are no bad object references
-			D.CompactInvocationList();
+			Ar << D.InvocationList;
+
+			if( Ar.IsLoading() )
+			{
+				// After loading the delegate, clean up the list to make sure there are no bad object references
+				D.CompactInvocationList();
+			}
 		}
 
 		return Ar;
