[TaskGraph]
  - Add FTaskConcurrencyLimiter which allows to reduce the concurrency for tasks queued through it.

diff --git a/Engine/Source/Runtime/Core/Private/Tasks/TaskConcurrencyLimiter.cpp b/Engine/Source/Runtime/Core/Private/Tasks/TaskConcurrencyLimiter.cpp
new file mode 100644
--- /dev/null
+++ b/Engine/Source/Runtime/Core/Private/Tasks/TaskConcurrencyLimiter.cpp
@@ -0,0 +1,118 @@
+// Copyright Epic Games, Inc. All Rights Reserved.
+
+#include "Tasks/TaskConcurrencyLimiter.h"
+#include "Tasks/Task.h"
+#include "Containers/LockFreeList.h"
+#include "Templates/SharedPointer.h"
+#include "Experimental/ConcurrentLinearAllocator.h"
+#include "ThirdParty/AtomicQueue/AtomicQueue.h"
+#include "CoreTypes.h"
+
+namespace UE::Tasks::TaskConcurrencyLimiter_Private
+{
+	FPimpl::~FPimpl()
+	{
+		if (FEvent* Event = CompletionEvent.exchange(nullptr, std::memory_order_relaxed))
+		{
+			FPlatformProcess::ReturnSynchEventToPool(Event);
+		}
+	}
+
+	void FPimpl::AddWorkItem(LowLevelTasks::FTask* Task)
+	{
+		NumWorkItems.fetch_add(1, std::memory_order_acquire);
+
+		WorkQueue.Push(Task);
+
+		uint32 ConcurrencySlot;
+		if (ConcurrencySlots.Alloc(ConcurrencySlot))
+		{
+			ProcessQueueFromPush(ConcurrencySlot);
+		}
+	}
+
+	bool FPimpl::Wait(FTimespan Timeout)
+	{
+		if (NumWorkItems.load(std::memory_order_relaxed) == 0)
+		{
+			return true;
+		}
+
+		FEvent* LocalCompletionEvent = CompletionEvent.load(std::memory_order_acquire);
+		if (LocalCompletionEvent == nullptr)
+		{
+			FEvent* NewEvent = FPlatformProcess::GetSynchEventFromPool(true);
+			if (!CompletionEvent.compare_exchange_strong(LocalCompletionEvent, NewEvent, std::memory_order_acq_rel, std::memory_order_acquire))
+			{
+				FPlatformProcess::ReturnSynchEventToPool(NewEvent);
+			}
+			else
+			{
+				LocalCompletionEvent = NewEvent;
+			}
+		}
+
+		if (NumWorkItems.load(std::memory_order_acquire) == 0)
+		{
+			return true;
+		}
+
+		return LocalCompletionEvent->Wait(Timeout);
+	}
+
+	void FPimpl::ProcessQueue(uint32 ConcurrencySlot, bool bSkipFirstWakeUp)
+	{
+		bool bWakeUpWorker = !bSkipFirstWakeUp;
+		do
+		{
+			if (LowLevelTasks::FTask* Task = WorkQueue.Pop())
+			{
+				// Now that we know the ConcurrencySlot, set it at launch time so
+				// the executor can retrieve it.
+				Task->SetUserData((void*)(UPTRINT)ConcurrencySlot);
+
+				LowLevelTasks::TryLaunch(*Task, bWakeUpWorker ? LowLevelTasks::EQueuePreference::GlobalQueuePreference : LowLevelTasks::EQueuePreference::LocalQueuePreference, bWakeUpWorker);
+			}
+			else
+			{
+				ConcurrencySlots.Release(ConcurrencySlot);
+				break;
+			}
+
+			// Don't skip wake-up if we launch any additional tasks.
+			bWakeUpWorker = true;
+
+		} while (ConcurrencySlots.Alloc(ConcurrencySlot));
+	}
+
+	void FPimpl::ProcessQueueFromWorker(uint32 ConcurrencySlot)
+	{
+		// Once we are in a worker thread, we want to schedule on the local queue without waking up additional workers
+		// to allow our own worker to pick up the next item and avoid wake-up cost / context switch.
+		static constexpr bool bSkipFirstWakeUp = true;
+
+		ProcessQueue(ConcurrencySlot, bSkipFirstWakeUp);
+	}
+
+	void FPimpl::ProcessQueueFromPush(uint32 ConcurrencySlot)
+	{
+		// When we push new items, we don't want to skip any wake-up.
+		static constexpr bool bSkipFirstWakeUp = false;
+
+		ProcessQueue(ConcurrencySlot, bSkipFirstWakeUp);
+	}
+
+	void FPimpl::CompleteWorkItem(uint32 ConcurrencySlot)
+	{
+		if (NumWorkItems.fetch_sub(1, std::memory_order_release) == 1)
+		{
+			if (FEvent* LocalCompletionEvent = CompletionEvent.load(std::memory_order_acquire))
+			{
+				LocalCompletionEvent->Trigger();
+			}
+		}
+
+		ProcessQueueFromWorker(ConcurrencySlot);
+	}
+
+} // namespace UE::Tasks::TaskConcurrencyLimiter_Private
diff --git a/Engine/Source/Runtime/Core/Public/Tasks/TaskConcurrencyLimiter.h b/Engine/Source/Runtime/Core/Public/Tasks/TaskConcurrencyLimiter.h
new file mode 100644
--- /dev/null
+++ b/Engine/Source/Runtime/Core/Public/Tasks/TaskConcurrencyLimiter.h
@@ -0,0 +1,163 @@
+// Copyright Epic Games, Inc. All Rights Reserved.
+
+#pragma once
+
+#include "Tasks/Task.h"
+#include "Containers/LockFreeList.h"
+#include "Templates/SharedPointer.h"
+#include "Experimental/ConcurrentLinearAllocator.h"
+#include "ThirdParty/AtomicQueue/AtomicQueue.h"
+#include "CoreTypes.h"
+
+namespace UE::Tasks
+{
+	namespace TaskConcurrencyLimiter_Private
+	{
+		// a queue of free slots in range [0 .. max_concurrency). Initially contains all slots in the range.
+		class FConcurrencySlots
+		{
+		public:
+			explicit FConcurrencySlots(uint32 MaxConcurrency)
+				: FreeSlots(MaxConcurrency)
+			{
+				for (uint32 Index = IndexOffset; Index < MaxConcurrency + IndexOffset; ++Index)
+				{
+					FreeSlots.push(Index);
+				}
+			}
+
+			bool Alloc(uint32& Slot)
+			{
+				if (FreeSlots.try_pop(Slot))
+				{
+					Slot -= IndexOffset;
+					return true;
+				}
+
+				return false;
+			}
+
+			void Release(uint32 Slot)
+			{
+				FreeSlots.push(Slot + IndexOffset);
+			}
+
+		private:
+			// this queue uses 0 as a special "null" value. to work around this, slots are shifted by one for storage, thus ending up in 
+			// [1 .. max_concurrency] range
+			static constexpr int32 IndexOffset = 1;
+			atomic_queue::AtomicQueueB<uint32> FreeSlots; // a bounded lock-free FIFO queue
+		};
+
+		// an implementation details of FTaskConcurrenctyLimiter
+		class FPimpl : public TSharedFromThis<FPimpl>
+		{
+		public:
+			explicit FPimpl(uint32 InMaxConcurrency, ETaskPriority InTaskPriority)
+				: ConcurrencySlots(InMaxConcurrency)
+				, TaskPriority(InTaskPriority)
+			{
+			}
+
+			CORE_API ~FPimpl();
+
+			template<typename TaskFunctionType>
+			void Push(const TCHAR* DebugName, TaskFunctionType&& TaskFunction)
+			{
+				TSharedPtr<LowLevelTasks::FTask> Task = MakeShared<LowLevelTasks::FTask>();
+
+				Task->Init(
+					DebugName,
+					TaskPriority,
+					[
+						TaskFunction = MoveTemp(TaskFunction),
+						this,
+						Pimpl = TSharedFromThis<FPimpl>::AsShared(), // to keep it alive
+						Task  // self-destruct
+					]()
+					{
+						// We can't pass the ConcurrencySlot in the lambda during creation as
+						// it's not actually acquired yet. The value will be passed using
+						// the user data when the task is launched.
+						uint32 ConcurrencySlot = (uint32)(UPTRINT)Task->GetUserData();
+
+						TaskFunction(ConcurrencySlot);
+						CompleteWorkItem(ConcurrencySlot);
+					}
+				);
+
+				AddWorkItem(Task.Get());
+			}
+
+			bool CORE_API Wait(FTimespan Timeout);
+
+		private:
+			void CORE_API AddWorkItem(LowLevelTasks::FTask* Task);
+			void CORE_API ProcessQueue(uint32 ConcurrencySlot, bool bSkipFirstWakeUp);
+			void CORE_API ProcessQueueFromWorker(uint32 ConcurrencySlot);
+			void CORE_API ProcessQueueFromPush(uint32 ConcurrencySlot);
+			void CORE_API CompleteWorkItem(uint32 ConcurrencySlot);
+
+			FConcurrencySlots ConcurrencySlots; // free slots queue. used also to limit concurrency
+			ETaskPriority TaskPriority;
+			TLockFreePointerListFIFO<LowLevelTasks::FTask, PLATFORM_CACHE_LINE_SIZE> WorkQueue; // a queue of user-provided task functions
+			std::atomic<uint32> NumWorkItems { 0 };
+			std::atomic<FEvent*> CompletionEvent { nullptr };
+		};
+
+	} // namespace TaskConcurrencyLimiter_Private
+
+	/**
+	* A lightweight construct that limits the concurrency of tasks pushed into it. 
+	*
+	* @note This class supports being destroyed before the tasks it contains are finished.
+	*/
+	class FTaskConcurrencyLimiter
+	{
+	public:
+		/**
+		 * Constructor.
+		 *
+		 * @param MaxConcurrency     How wide the processing can go.
+		 * @param TaskPriority       Priority the tasks will be launched with.
+		 */
+		explicit FTaskConcurrencyLimiter(uint32 MaxConcurrency, ETaskPriority TaskPriority = ETaskPriority::Default)
+			: Pimpl(MakeShared<TaskConcurrencyLimiter_Private::FPimpl>(MaxConcurrency, TaskPriority))
+		{
+		}
+
+		/**
+		 * Push a new task.
+		 *
+		 * @param DebugName    Helps to identify the task in debugger and profiler.
+		 * @param TaskFunction A callable with a slot parameter, usually a lambda but can be also a functor object 
+		 *                     or a pointer to a function. The slot parameter is an index in [0..max_concurrency) range, 
+		 *                     unique at any moment of time, that can be used in user code to index a fixed-size buffer. 
+		 *                     See `TaskConcurrencyLimiterStressTest()` for an example.
+		 */
+		template<typename TaskFunctionType>
+		void Push(const TCHAR* DebugName, TaskFunctionType&& TaskFunction)
+		{
+			Pimpl->Push(DebugName, MoveTemp(TaskFunction));
+		}
+
+		/**
+		 * Waits for task's completion with timeout.
+		 *
+		 * @param  Timeout Maximum amount of time to wait for tasks to finish before returning.
+		 * @return true if all tasks are completed, false otherwise.
+		 * 
+		 * @note   A wait is satisfied once the internal task counter reaches 0 and is never reset
+		 *         afterward when more tasks are added. A new FTaskConcurrencyLimiter can be used
+		 *         for such a use case.
+		 */
+		bool Wait(FTimespan Timeout = FTimespan::MaxValue())
+		{
+			return Pimpl->Wait(Timeout);
+		}
+
+	private:
+		TSharedRef<TaskConcurrencyLimiter_Private::FPimpl> Pimpl;
+	};
+
+} // namespace UE::Tasks
\ No newline at end of file
diff --git a/Engine/Source/Runtime/Core/Tests/Tasks/TasksTest.cpp b/Engine/Source/Runtime/Core/Tests/Tasks/TasksTest.cpp
--- a/Engine/Source/Runtime/Core/Tests/Tasks/TasksTest.cpp
+++ b/Engine/Source/Runtime/Core/Tests/Tasks/TasksTest.cpp
@@ -3,7 +3,9 @@
 #include "CoreTypes.h"
 #include "Misc/AutomationTest.h"
 #include "Tests/Benchmark.h"
+#include "Tasks/Task.h"
 #include "Tasks/Pipe.h"
+#include "Tasks/TaskConcurrencyLimiter.h"
 #include "HAL/Thread.h"
 #include "Async/ParallelFor.h"
 #include "Tests/TestHarnessAdapter.h"
@@ -1519,6 +1521,119 @@ namespace UE { namespace TasksTests
 			Blocker.Trigger();
 		}
 	}
+
+	// produces work for FTaskConcurrencyLimiter from multiple threads to check it's thread-safety, checks that:
+	// * it doesn't go over max concurrency
+	// * can check (not 100% reliably but still useful for local runs) that max concurrency is actually reached
+	// * slots don't overlap
+	template<uint32 MaxConcurrency, uint32 NumItems, uint32 NumPushingTasks>
+	void TaskConcurrencyLimiterStressTest()
+	{
+		static_assert(NumItems % NumPushingTasks == 0);
+
+		std::atomic<uint32> CurrentConcurrency = 0;
+		std::atomic<uint32> ActualMaxConcurrency = 0;
+		std::atomic<uint32> NumProcessed = 0;
+
+		std::atomic<bool> Slots[MaxConcurrency] = {};
+
+		TArray<FTask> PushingTasks;
+		PushingTasks.Reserve(NumPushingTasks);
+
+		FTaskConcurrencyLimiter TaskConcurrencyLimiter(MaxConcurrency);
+
+		for (uint32 i = 0; i != NumPushingTasks; ++i)
+		{
+			PushingTasks.Add(Launch(UE_SOURCE_LOCATION, 
+				[&TaskConcurrencyLimiter, &CurrentConcurrency, &ActualMaxConcurrency, &Slots, &NumProcessed]
+				{
+					TRACE_CPUPROFILER_EVENT_SCOPE(PushTasks);
+					for (int32 i = 0; i != NumItems / NumPushingTasks; ++i)
+					{
+						TaskConcurrencyLimiter.Push(UE_SOURCE_LOCATION,
+							[&CurrentConcurrency, &ActualMaxConcurrency, &Slots, &NumProcessed](uint32 Slot)
+							{
+								TRACE_CPUPROFILER_EVENT_SCOPE(Task);
+								check(Slot < MaxConcurrency);
+								check(!Slots[Slot].load(std::memory_order_relaxed));
+								Slots[Slot].store(true, std::memory_order_relaxed);
+
+								uint32 CurrentConcurrencyLocal = CurrentConcurrency.fetch_add(1, std::memory_order_relaxed) + 1;
+								check(CurrentConcurrencyLocal <= MaxConcurrency);
+								uint32 ActualMaxConcurrencyLocal = ActualMaxConcurrency.load(std::memory_order_relaxed);
+								while (ActualMaxConcurrencyLocal < CurrentConcurrencyLocal &&
+									!ActualMaxConcurrency.compare_exchange_weak(ActualMaxConcurrencyLocal, CurrentConcurrencyLocal, std::memory_order_relaxed, std::memory_order_relaxed))
+								{
+									check(ActualMaxConcurrencyLocal <= MaxConcurrency);
+								}
+
+								FPlatformProcess::YieldCycles(10000);
+
+								CurrentConcurrencyLocal = CurrentConcurrency.fetch_sub(1, std::memory_order_relaxed) - 1;
+								check(CurrentConcurrencyLocal >= 0);
+
+								Slots[Slot].store(false, std::memory_order_relaxed);
+
+								NumProcessed.fetch_add(1, std::memory_order_release);
+							}
+						);
+					}
+				}
+			));
+		}
+
+		Wait(PushingTasks);
+		TaskConcurrencyLimiter.Wait();
+		check(NumProcessed.load(std::memory_order_acquire) == NumItems);
+
+		// unreliable check that MaxConcurrency is actually reached, but reliable enough for testing locally
+		//check(ActualMaxConcurrency.load(std::memory_order_relaxed) == MaxConcurrency);
+	}
+
+	TEST_CASE_NAMED(FTasksConcurrencyLimiterTest, "System::Core::Async::Tasks::TaskConcurrencyLimiter", "[.][ApplicationContextMask][EngineFilter]")
+	{
+		UE_BENCHMARK(5, TaskConcurrencyLimiterStressTest<8, 1'000'000, 10>);
+	}
+
+	template<uint32 RepeatCount>
+	void TaskConcurrencyLimiter_WaitingStressTest()
+	{
+		for (uint32 RepeatIndex = 0; RepeatIndex < RepeatCount; ++RepeatIndex)
+		{
+			FTaskConcurrencyLimiter TaskConcurrencyLimiter{ 1 };
+			FEventRef Blocker;
+			std::atomic<bool> bDone{ false };
+			TaskConcurrencyLimiter.Push(UE_SOURCE_LOCATION, 
+				[&Blocker, &bDone](uint32) 
+				{ 
+					TRACE_CPUPROFILER_EVENT_SCOPE(Blocked);
+					Blocker->Wait(); 
+					bDone.store(true, std::memory_order_relaxed); 
+				}
+			);
+			const uint32 NumWaitingTasks = 10;
+			TArray<FTask> WaitingTasks;
+			WaitingTasks.Reserve(NumWaitingTasks);
+			for (uint32 WaitIndex = 0; WaitIndex < NumWaitingTasks; ++WaitIndex)
+			{
+				WaitingTasks.Add(Launch(UE_SOURCE_LOCATION, 
+					[&TaskConcurrencyLimiter, &bDone] 
+					{
+						TRACE_CPUPROFILER_EVENT_SCOPE(Waiting);
+						TaskConcurrencyLimiter.Wait(); 
+						check(bDone.load(std::memory_order_relaxed)); 
+					}
+				));
+			}
+			Blocker->Trigger();
+			Wait(WaitingTasks);
+		}
+	}
+
+	TEST_CASE_NAMED(FTasksConcurrencyLimiterWaitingTest, "System::Core::Async::Tasks::TaskConcurrencyLimiter::Waiting", "[.][ApplicationContextMask][EngineFilter]")
+	{
+		UE_BENCHMARK(5, TaskConcurrencyLimiter_WaitingStressTest<10'000>);
+	}
 }}
 
 #endif // WITH_TESTS
