Streamline error coming from log capture.
Cleanup Log capture function.
Move all tests under Engine.AssetRegistry into System.AssetRegistry

#rnx

diff --git a/Engine/Source/Runtime/AssetRegistry/Private/AssetRegistryArchive.cpp b/Engine/Source/Runtime/AssetRegistry/Private/AssetRegistryArchive.cpp
--- a/Engine/Source/Runtime/AssetRegistry/Private/AssetRegistryArchive.cpp
+++ b/Engine/Source/Runtime/AssetRegistry/Private/AssetRegistryArchive.cpp
@@ -344,7 +344,7 @@ void FAssetRegistryReader::SerializeTagsAndBundlesOldVersion(FAssetData& Out, FA
 #include "Serialization/MemoryReader.h"
 #include "Serialization/MemoryWriter.h"
 
-IMPLEMENT_SIMPLE_AUTOMATION_TEST(FAssetRegistryTagSerializationTest, "Engine.AssetRegistry.SerializeTagMap", EAutomationTestFlags::ApplicationContextMask | EAutomationTestFlags::SmokeFilter)
+IMPLEMENT_SIMPLE_AUTOMATION_TEST(FAssetRegistryTagSerializationTest, "System.AssetRegistry.SerializeTagMap", EAutomationTestFlags::ApplicationContextMask | EAutomationTestFlags::SmokeFilter)
 
 FAssetDataTagMapSharedView MakeLooseMap(std::initializer_list<TPairInitializer<const char*, FString>> Pairs)
 {
diff --git a/Engine/Source/Runtime/AssetRegistry/Private/AssetRegistryState.cpp b/Engine/Source/Runtime/AssetRegistry/Private/AssetRegistryState.cpp
--- a/Engine/Source/Runtime/AssetRegistry/Private/AssetRegistryState.cpp
+++ b/Engine/Source/Runtime/AssetRegistry/Private/AssetRegistryState.cpp
@@ -3046,7 +3046,7 @@ void FAssetRegistryState::Dump(const TArray<FString>& Arguments, TArray<FString>
 
 #include "Misc/AutomationTest.h"
 
-IMPLEMENT_SIMPLE_AUTOMATION_TEST(FAssetRegistryAssetPathStringsTest, "Engine.AssetRegistry.AssetPathStrings", EAutomationTestFlags::ApplicationContextMask | EAutomationTestFlags::EngineFilter);
+IMPLEMENT_SIMPLE_AUTOMATION_TEST(FAssetRegistryAssetPathStringsTest, "System.AssetRegistry.AssetPathStrings", EAutomationTestFlags::ApplicationContextMask | EAutomationTestFlags::EngineFilter);
 
 // Tests that we can produce correct paths for objects represented by FCachedAssetKey and FAssetData
 // E.g	PackageName.AssetName
diff --git a/Engine/Source/Runtime/AssetRegistry/Private/AssetRegistryStateTest.cpp b/Engine/Source/Runtime/AssetRegistry/Private/AssetRegistryStateTest.cpp
--- a/Engine/Source/Runtime/AssetRegistry/Private/AssetRegistryStateTest.cpp
+++ b/Engine/Source/Runtime/AssetRegistry/Private/AssetRegistryStateTest.cpp
@@ -11,7 +11,7 @@
 
 
 #if WITH_DEV_AUTOMATION_TESTS
-IMPLEMENT_SIMPLE_AUTOMATION_TEST(FEnumerateAssetsPerformance, "Engine.AssetRegistry.AssetRegistryState.EnumerateAssetsPerformance",
+IMPLEMENT_SIMPLE_AUTOMATION_TEST(FEnumerateAssetsPerformance, "System.AssetRegistry.AssetRegistryState.EnumerateAssetsPerformance",
 	EAutomationTestFlags::EditorContext | EAutomationTestFlags::PerfFilter);
 
 bool FEnumerateAssetsPerformance::RunTest(const FString& Parameters)
diff --git a/Engine/Source/Runtime/Core/Private/Misc/AutomationTest.cpp b/Engine/Source/Runtime/Core/Private/Misc/AutomationTest.cpp
--- a/Engine/Source/Runtime/Core/Private/Misc/AutomationTest.cpp
+++ b/Engine/Source/Runtime/Core/Private/Misc/AutomationTest.cpp
@@ -67,6 +67,65 @@ namespace AutomationTest
 
 		return Result;
 	}
+
+	/*
+		Determine the level that a log item should be written to the automation log based on the properties of the current test.
+		only Display/Warning/Error are supported in the automation log so anything with NoLogging/Log will not be shown
+	*/
+	static ELogVerbosity::Type GetAutomationLogLevel(ELogVerbosity::Type LogVerbosity, FName LogCategory, FAutomationTestBase* CurrentTest)
+	{
+		ELogVerbosity::Type EffectiveVerbosity = LogVerbosity;
+
+		static FCriticalSection ActionCS;
+		static FAutomationTestBase* LastTest = nullptr;
+
+		if (AutomationTest::bCaptureLogEvents == false)
+		{
+			return ELogVerbosity::NoLogging;
+		}
+
+		{
+			FScopeLock Lock(&ActionCS);
+			if (CurrentTest != LastTest)
+			{
+				FAutomationTestBase::SuppressedLogCategories.Empty();
+				FAutomationTestBase::LoadDefaultLogSettings();
+				LastTest = CurrentTest;
+			}
+		}
+
+		if (CurrentTest)
+		{
+			if (CurrentTest->SuppressLogs() || CurrentTest->GetSuppressedLogCategories().Contains(LogCategory.ToString()))
+			{
+				EffectiveVerbosity = ELogVerbosity::NoLogging;
+			}
+			else
+			{
+				if (EffectiveVerbosity == ELogVerbosity::Warning)
+				{
+					if (CurrentTest->SuppressLogWarnings())
+					{
+						EffectiveVerbosity = ELogVerbosity::NoLogging;
+					}
+					else if (CurrentTest->ElevateLogWarningsToErrors())
+					{
+						EffectiveVerbosity = ELogVerbosity::Error;
+					}
+				}
+
+				if (EffectiveVerbosity == ELogVerbosity::Error)
+				{
+					if (CurrentTest->SuppressLogErrors())
+					{
+						EffectiveVerbosity = ELogVerbosity::NoLogging;
+					}
+				}
+			}
+		}
+
+		return EffectiveVerbosity;
+	}
 };
 
 bool FAutomationTestBase::bSuppressLogWarnings = false;
@@ -106,69 +165,10 @@ CORE_API const TMap<FString, EAutomationTestFlags::Type>& EAutomationTestFlags::
 	return FlagsMap;
 };
 
-/*
-	Determine the level that a log item should be written to the automation log based on the properties of the current test. 
-	only Display/Warning/Error are supported in the automation log so anything with NoLogging/Log will not be shown
-	(Should be moved under a namespace for 4.27).
-*/
-CORE_API ELogVerbosity::Type GetAutomationLogLevel(ELogVerbosity::Type LogVerbosity, FName LogCategory, FAutomationTestBase* CurrentTest)
-{
-	ELogVerbosity::Type EffectiveVerbosity = LogVerbosity;
-
-	// agrant-todo: these should be controlled by FAutomationTestBase for 4.27 with the same project-level override that
-	// FunctionalTest has. Now that warnings are correctly associated with tests they need to be something all tests
-	// can leverage, not just functional tests
-	static FAutomationTestBase* LastTest = nullptr;
-
-	if (AutomationTest::bCaptureLogEvents == false)
-	{
-		return ELogVerbosity::NoLogging;
-	}
-
-	if (CurrentTest != LastTest) 
-	{
-		FAutomationTestBase::SuppressedLogCategories.Empty();
-		FAutomationTestBase::LoadDefaultLogSettings();
-		LastTest = CurrentTest;
-	}
-
-	if (CurrentTest)
-	{
-		if (CurrentTest->SuppressLogs() || CurrentTest->GetSuppressedLogCategories().Contains(LogCategory.ToString()))
-		{
-			EffectiveVerbosity = ELogVerbosity::NoLogging;
-		}
-		else
-		{
-			if (EffectiveVerbosity == ELogVerbosity::Warning)
-			{
-				if (CurrentTest->SuppressLogWarnings())
-				{
-					EffectiveVerbosity = ELogVerbosity::NoLogging;
-				}
-				else if (CurrentTest->ElevateLogWarningsToErrors())
-				{
-					EffectiveVerbosity = ELogVerbosity::Error;
-				}
-			}
-
-			if (EffectiveVerbosity == ELogVerbosity::Error)
-			{
-				if (CurrentTest->SuppressLogErrors())
-				{
-					EffectiveVerbosity = ELogVerbosity::NoLogging;
-				}
-			}
-		}
-	}
-
-	return EffectiveVerbosity;
-}
-
 void FAutomationTestFramework::FAutomationTestOutputDevice::Serialize( const TCHAR* V, ELogVerbosity::Type Verbosity, const class FName& Category )
 {
 	const int32 STACK_OFFSET = 8;//FMsg::Logf_InternalImpl
-	// TODO would be nice to search for the first stack frame that isn't in outputdevice or other logging files, would be more robust.
+	// TODO would be nice to search for the first stack frame that isn't in output device or other logging files, would be more robust.
 
 	if (!IsRunningCommandlet() && (Verbosity == ELogVerbosity::SetColor))
 	{
@@ -185,20 +185,13 @@ void FAutomationTestFramework::FAutomationTestOutputDevice::Serialize( const TCH
 
 		if (CaptureLog)
 		{
-		
-			ELogVerbosity::Type EffectiveVerbosity = GetAutomationLogLevel(Verbosity, Category, LocalCurTest);
+			ELogVerbosity::Type EffectiveVerbosity = AutomationTest::GetAutomationLogLevel(Verbosity, Category, LocalCurTest);
 
-			FString FormattedMsg = FString::Printf(TEXT("%s: %s"), *Category.ToString(), V);
+			FString FormattedMsg = FString::Printf(TEXT("%s: %s [log]"), *Category.ToString(), V);
 			
 			// Errors
 			if (EffectiveVerbosity == ELogVerbosity::Error)
 			{
-				FScopeLock Lock(&ActionCS);
-				if (!LoggedFailureCause.Contains(LocalCurTest))
-				{
-					LocalCurTest->AddError(FString::Printf(TEXT("%s will be marked as failing due to errors being logged"), *LocalCurTest->GetTestFullName()), STACK_OFFSET);
-					LoggedFailureCause.Add(LocalCurTest);
-				}
 				LocalCurTest->AddError(FormattedMsg, STACK_OFFSET);
 			}
 			// Warnings
diff --git a/Engine/Source/Runtime/Core/Public/Misc/AutomationTest.h b/Engine/Source/Runtime/Core/Public/Misc/AutomationTest.h
--- a/Engine/Source/Runtime/Core/Public/Misc/AutomationTest.h
+++ b/Engine/Source/Runtime/Core/Public/Misc/AutomationTest.h
@@ -1218,12 +1218,6 @@ private:
 	private:
 		/** Associated automation test; all warnings, errors, etc. are routed to the automation test to track */
 		std::atomic<FAutomationTestBase*>CurTest;
-
-		/** Critical section */
-		FCriticalSection ActionCS;
-
-		/** Tests that we've logged the failure cause when an error is involved */
-		TSet<FAutomationTestBase*> LoggedFailureCause;
 	};
 
 	 /** Special feedback context used during automated testing to filter messages that happen during tests */
diff --git a/Engine/Source/Runtime/CoreUObject/Private/AssetRegistry/AssetBundleData.cpp b/Engine/Source/Runtime/CoreUObject/Private/AssetRegistry/AssetBundleData.cpp
--- a/Engine/Source/Runtime/CoreUObject/Private/AssetRegistry/AssetBundleData.cpp
+++ b/Engine/Source/Runtime/CoreUObject/Private/AssetRegistry/AssetBundleData.cpp
@@ -589,7 +589,7 @@ FString FAssetBundleData::ToDebugString() const
 
 // Combine import/export tests
 
-IMPLEMENT_SIMPLE_AUTOMATION_TEST(FAssetBundleEntryImportExportTextTest, "Engine.AssetRegistry.AssetBundleEntry.ImportExportText", EAutomationTestFlags::ApplicationContextMask | EAutomationTestFlags::EngineFilter);
+IMPLEMENT_SIMPLE_AUTOMATION_TEST(FAssetBundleEntryImportExportTextTest, "System.AssetRegistry.AssetBundleEntry.ImportExportText", EAutomationTestFlags::ApplicationContextMask | EAutomationTestFlags::EngineFilter);
 bool FAssetBundleEntryImportExportTextTest::RunTest(const FString& Parameters)
 {
 	FAssetBundleEntry DefaultEntry;
@@ -676,7 +676,7 @@ bool FAssetBundleEntryImportExportTextTest::RunTest(const FString& Parameters)
 }
 
 
-IMPLEMENT_SIMPLE_AUTOMATION_TEST(FLegacyAssetBundleEntryTest, "Engine.AssetRegistry.AssetBundleEntry.LegacyAssetBundleEntry", EAutomationTestFlags::ApplicationContextMask | EAutomationTestFlags::EngineFilter);
+IMPLEMENT_SIMPLE_AUTOMATION_TEST(FLegacyAssetBundleEntryTest, "System.AssetRegistry.AssetBundleEntry.LegacyAssetBundleEntry", EAutomationTestFlags::ApplicationContextMask | EAutomationTestFlags::EngineFilter);
 // Test ImportText for old asset bundles with BundleAssets field
 bool FLegacyAssetBundleEntryTest::RunTest(const FString& Parameters)
 {
@@ -744,7 +744,7 @@ bool FLegacyAssetBundleEntryTest::RunTest(const FString& Parameters)
 	return true;
 }
 
-IMPLEMENT_SIMPLE_AUTOMATION_TEST(FAssetBundlDataImportExportTextTest, "Engine.AssetRegistry.AssetBundleData.ImportExportText", EAutomationTestFlags::ApplicationContextMask | EAutomationTestFlags::EngineFilter);
+IMPLEMENT_SIMPLE_AUTOMATION_TEST(FAssetBundlDataImportExportTextTest, "System.AssetRegistry.AssetBundleData.ImportExportText", EAutomationTestFlags::ApplicationContextMask | EAutomationTestFlags::EngineFilter);
 bool FAssetBundlDataImportExportTextTest::RunTest(const FString& Parameters)
 {
 	UScriptStruct* Struct = TBaseStructure<FAssetBundleData>::Get();
diff --git a/Engine/Source/Runtime/CoreUObject/Private/AssetRegistry/AssetDataTagMap.cpp b/Engine/Source/Runtime/CoreUObject/Private/AssetRegistry/AssetDataTagMap.cpp
--- a/Engine/Source/Runtime/CoreUObject/Private/AssetRegistry/AssetDataTagMap.cpp
+++ b/Engine/Source/Runtime/CoreUObject/Private/AssetRegistry/AssetDataTagMap.cpp
@@ -1638,7 +1638,7 @@ SIZE_T FAssetDataTagMapSharedView::FMemoryCounter::GetFixedSize() const
 #include "Misc/AutomationTest.h"
 #include "Serialization/MemoryReader.h"
 
-IMPLEMENT_SIMPLE_AUTOMATION_TEST(FAssetRegistryExportPathTest, "Engine.AssetRegistry.ExportPath", EAutomationTestFlags::ApplicationContextMask | EAutomationTestFlags::SmokeFilter)
+IMPLEMENT_SIMPLE_AUTOMATION_TEST(FAssetRegistryExportPathTest, "System.AssetRegistry.ExportPath", EAutomationTestFlags::ApplicationContextMask | EAutomationTestFlags::SmokeFilter)
 
 bool FAssetRegistryExportPathTest::RunTest(const FString& Parameters)
 {
@@ -1656,7 +1656,7 @@ bool FAssetRegistryExportPathTest::RunTest(const FString& Parameters)
 namespace FixedTagPrivate
 {
 
-IMPLEMENT_SIMPLE_AUTOMATION_TEST(FCompactExportPathTest, "Engine.AssetRegistry.FixedTag.NumberlessExportPath", EAutomationTestFlags::ApplicationContextMask | EAutomationTestFlags::SmokeFilter)
+IMPLEMENT_SIMPLE_AUTOMATION_TEST(FCompactExportPathTest, "System.AssetRegistry.FixedTag.NumberlessExportPath", EAutomationTestFlags::ApplicationContextMask | EAutomationTestFlags::SmokeFilter)
 
 bool FCompactExportPathTest::RunTest(const FString& Parameters)
 {
@@ -1668,7 +1668,7 @@ bool FCompactExportPathTest::RunTest(const FString& Parameters)
 }
 
 
-IMPLEMENT_SIMPLE_AUTOMATION_TEST(FStoreTest, "Engine.AssetRegistry.FixedTag.Store", EAutomationTestFlags::ApplicationContextMask | EAutomationTestFlags::SmokeFilter)
+IMPLEMENT_SIMPLE_AUTOMATION_TEST(FStoreTest, "System.AssetRegistry.FixedTag.Store", EAutomationTestFlags::ApplicationContextMask | EAutomationTestFlags::SmokeFilter)
 
 static TRefCountPtr<const FStore> MakeTestStore(FStoreData& Data)
 {
