Make static analyzer happy

#rnx

diff --git a/Engine/Source/Runtime/Core/Tests/Tasks/TasksTest.cpp b/Engine/Source/Runtime/Core/Tests/Tasks/TasksTest.cpp
--- a/Engine/Source/Runtime/Core/Tests/Tasks/TasksTest.cpp
+++ b/Engine/Source/Runtime/Core/Tests/Tasks/TasksTest.cpp
@@ -1548,7 +1548,7 @@ namespace UE { namespace TasksTests
 				[&TaskConcurrencyLimiter, &CurrentConcurrency, &ActualMaxConcurrency, &Slots, &NumProcessed]
 				{
 					TRACE_CPUPROFILER_EVENT_SCOPE(PushTasks);
-					for (int32 i = 0; i != NumItems / NumPushingTasks; ++i)
+					for (uint32 i = 0; i < NumItems / NumPushingTasks; ++i)
 					{
 						TaskConcurrencyLimiter.Push(UE_SOURCE_LOCATION,
 							[&CurrentConcurrency, &ActualMaxConcurrency, &Slots, &NumProcessed](uint32 Slot)
