Removed std::enable_if and TEnableIf and other traits from math types and function signatures.

diff --git a/Engine/Source/Runtime/Core/Public/Clang/ClangPlatformMath.h b/Engine/Source/Runtime/Core/Public/Clang/ClangPlatformMath.h
--- a/Engine/Source/Runtime/Core/Public/Clang/ClangPlatformMath.h
+++ b/Engine/Source/Runtime/Core/Public/Clang/ClangPlatformMath.h
@@ -104,7 +104,10 @@ struct FClangPlatformMath : public FGenericPlatformMath
 	 * If there was overflow, it returns false, and OutResult may or may not be written.
 	 * If there wasn't overflow, it returns true, and the result of the addition is written to OutResult.
 	 */
-	template<typename IntType, std::enable_if_t<std::is_integral_v<IntType>>* = nullptr>
+	template <
+		typename IntType
+		UE_REQUIRES(std::is_integral_v<IntType>)
+	>
 	static FORCEINLINE bool AddAndCheckForOverflow(IntType A, IntType B, IntType& OutResult)
 	{
 		return !__builtin_add_overflow(A, B, &OutResult);
@@ -115,7 +118,10 @@ struct FClangPlatformMath : public FGenericPlatformMath
 	 * If there was overflow, it returns false, and OutResult may or may not be written.
 	 * If there wasn't overflow, it returns true, and the result of the subtraction is written to OutResult.
 	 */
-	template<typename IntType, std::enable_if_t<std::is_integral_v<IntType>>* = nullptr>
+	template <
+		typename IntType
+		UE_REQUIRES(std::is_integral_v<IntType>)
+	>
 	static FORCEINLINE bool SubtractAndCheckForOverflow(IntType A, IntType B, IntType& OutResult)
 	{
 		return !__builtin_sub_overflow(A, B, &OutResult);
@@ -126,7 +132,10 @@ struct FClangPlatformMath : public FGenericPlatformMath
 	 * If there was overflow, it returns false, and OutResult may or may not be written.
 	 * If there wasn't overflow, it returns true, and the result of the multiplication is written to OutResult.
 	 */
-	template<typename IntType, std::enable_if_t<std::is_integral_v<IntType>>* = nullptr>
+	template <
+		typename IntType
+		UE_REQUIRES(std::is_integral_v<IntType>)
+	>
 	static FORCEINLINE bool MultiplyAndCheckForOverflow(IntType A, IntType B, IntType& OutResult)
 	{
 		return !__builtin_mul_overflow(A, B, &OutResult);
diff --git a/Engine/Source/Runtime/Core/Public/GenericPlatform/GenericPlatformMath.h b/Engine/Source/Runtime/Core/Public/GenericPlatform/GenericPlatformMath.h
--- a/Engine/Source/Runtime/Core/Public/GenericPlatform/GenericPlatformMath.h
+++ b/Engine/Source/Runtime/Core/Public/GenericPlatform/GenericPlatformMath.h
@@ -1029,7 +1029,10 @@ struct FGenericPlatformMath
 	 * If there was overflow, it returns false, and OutResult may or may not be written.
 	 * If there wasn't overflow, it returns true, and the result of the addition is written to OutResult.
 	 */
-	template<typename IntType, std::enable_if_t<std::is_integral_v<IntType>>* = nullptr>
+	template <
+		typename IntType
+		UE_REQUIRES(std::is_integral_v<IntType>)
+	>
 	static FORCEINLINE bool AddAndCheckForOverflow(IntType A, IntType B, IntType& OutResult)
 	{
 		// This follows Hacker's Delight, Chapter 2-12
@@ -1076,7 +1079,10 @@ struct FGenericPlatformMath
 	 * If there was overflow, it returns false, and OutResult may or may not be written.
 	 * If there wasn't overflow, it returns true, and the result of the subtraction is written to OutResult.
 	 */
-	template<typename IntType, std::enable_if_t<std::is_integral_v<IntType>>* = nullptr>
+	template <
+		typename IntType
+		UE_REQUIRES(std::is_integral_v<IntType>)
+	>
 	static FORCEINLINE bool SubtractAndCheckForOverflow(IntType A, IntType B, IntType& OutResult)
 	{
 		// This follows Hacker's Delight, Chapter 2-12
@@ -1117,7 +1123,10 @@ struct FGenericPlatformMath
 	 * If there was overflow, it returns false, and OutResult may or may not be written.
 	 * If there wasn't overflow, it returns true, and the result of the multiplication is written to OutResult.
 	 */
-	template<typename IntType, std::enable_if_t<std::is_integral_v<IntType>>* = nullptr>
+	template <
+		typename IntType
+		UE_REQUIRES(std::is_integral_v<IntType>)
+	>
 	static FORCEINLINE bool MultiplyAndCheckForOverflow(IntType A, IntType B, IntType& OutResult)
 	{
 		// Handle the case where the second factor is 0 specially (why will become clear in a minute).
diff --git a/Engine/Source/Runtime/Core/Public/Internationalization/FastDecimalFormat.h b/Engine/Source/Runtime/Core/Public/Internationalization/FastDecimalFormat.h
--- a/Engine/Source/Runtime/Core/Public/Internationalization/FastDecimalFormat.h
+++ b/Engine/Source/Runtime/Core/Public/Internationalization/FastDecimalFormat.h
@@ -7,10 +7,8 @@
 #include "Internationalization/Text.h"
 #include "Math/NumericLimits.h"
 #include "Misc/CString.h"
-#include "Templates/EnableIf.h"
-#include "Templates/IsFloatingPoint.h"
-#include "Templates/IsIntegral.h"
-#include "Templates/IsSigned.h"
+
+#include <type_traits>
 
 /** Rules used to format or parse a decimal number */
 struct FDecimalNumberFormattingRules
@@ -79,13 +77,13 @@ struct FDecimalNumberIntegralLimits
 	uint64 NumericLimitMax;
 	bool bIsNumericSigned;
 
-	template<
-		typename IntegralType,
-		typename TEnableIf<TIsIntegral<IntegralType>::Value>::Type * = nullptr
+	template <
+		typename IntegralType
+		UE_REQUIRES(std::is_integral_v<IntegralType>)
 	>
 	static FDecimalNumberIntegralLimits FromNumericLimits()
 	{
-		return FDecimalNumberIntegralLimits(TNumericLimits<IntegralType>::Lowest(), TNumericLimits<IntegralType>::Max(), TIsSigned<IntegralType>::Value);
+		return FDecimalNumberIntegralLimits(TNumericLimits<IntegralType>::Lowest(), TNumericLimits<IntegralType>::Max(), std::is_signed_v<IntegralType>);
 	}
 };
 
@@ -98,9 +96,9 @@ struct FDecimalNumberFractionalLimits
 	double NumericLimitLowest;
 	double NumericLimitMax;
 
-	template<
-		typename FloatingType,
-		typename TEnableIf<TIsFloatingPoint<FloatingType>::Value>::Type * = nullptr
+	template <
+		typename FloatingType
+		UE_REQUIRES(std::is_floating_point_v<FloatingType>)
 	>
 	static FDecimalNumberFractionalLimits FromNumericLimits()
 	{
@@ -206,3 +204,10 @@ CORE_API const FDecimalNumberFormattingRules& GetCultureAgnosticFormattingRules(
 CORE_API uint64 Pow10(const int32 InExponent);
 
 } // namespace FastDecimalFormat
+
+#if UE_ENABLE_INCLUDE_ORDER_DEPRECATED_IN_5_4
+#include "Templates/EnableIf.h"
+#include "Templates/IsFloatingPoint.h"
+#include "Templates/IsIntegral.h"
+#include "Templates/IsSigned.h"
+#endif
diff --git a/Engine/Source/Runtime/Core/Public/Math/GuardedInt.h b/Engine/Source/Runtime/Core/Public/Math/GuardedInt.h
--- a/Engine/Source/Runtime/Core/Public/Math/GuardedInt.h
+++ b/Engine/Source/Runtime/Core/Public/Math/GuardedInt.h
@@ -69,8 +69,8 @@ public:
 
 	/** Construct a TGuardedInt from a regular signed integer value. If it's out of range, it results in an invalid value. */
 	template<
-		typename ValueType,
-		std::enable_if_t<std::is_integral_v<ValueType>>* = nullptr
+		typename ValueType
+		UE_REQUIRES(std::is_integral_v<ValueType>)
 	>
 	explicit TGuardedInt(ValueType InValue)
 		: Value((IntType)InValue), bIsValid(false)
diff --git a/Engine/Source/Runtime/Core/Public/Math/TransformCalculus2D.h b/Engine/Source/Runtime/Core/Public/Math/TransformCalculus2D.h
--- a/Engine/Source/Runtime/Core/Public/Math/TransformCalculus2D.h
+++ b/Engine/Source/Runtime/Core/Public/Math/TransformCalculus2D.h
@@ -8,6 +8,8 @@
 #include "Math/Vector2D.h"
 #include "Math/TransformCalculus.h"
 
+#include <type_traits>
+
 //////////////////////////////////////////////////////////////////////////
 // Transform calculus for 2D types. UE4 already has a 2D Vector class that we
 // will adapt to be interpreted as a translate transform. The rest we create
@@ -907,7 +909,7 @@ namespace TransformCalculusHelper
 }
 
 /** Partial specialization of ConcatenateRules for FTransform2D and any other type via Upcast to FTransform2D first. Requires a conversion ctor on FTransform2D. Funky template logic so we don't hide the default rule for NULL conversions. */
-template<typename TransformType> struct ConcatenateRules<typename TEnableIf<!TransformCalculusHelper::TIsTransform2<TransformType>::Value, FTransform2f>::Type, TransformType	> { typedef FTransform2f ResultType; };
-template<typename TransformType> struct ConcatenateRules<typename TEnableIf<!TransformCalculusHelper::TIsTransform2<TransformType>::Value, FTransform2d>::Type, TransformType	> { typedef FTransform2d ResultType; };
-template<typename TransformType> struct ConcatenateRules<TransformType, typename TEnableIf<!TransformCalculusHelper::TIsTransform2<TransformType>::Value, FTransform2f>::Type	> { typedef FTransform2f ResultType; };
-template<typename TransformType> struct ConcatenateRules<TransformType, typename TEnableIf<!TransformCalculusHelper::TIsTransform2<TransformType>::Value, FTransform2d>::Type	> { typedef FTransform2d ResultType; };
+template<typename TransformType> struct ConcatenateRules<std::enable_if_t<!TransformCalculusHelper::TIsTransform2<TransformType>::Value, FTransform2f>, TransformType> { typedef FTransform2f ResultType; };
+template<typename TransformType> struct ConcatenateRules<std::enable_if_t<!TransformCalculusHelper::TIsTransform2<TransformType>::Value, FTransform2d>, TransformType> { typedef FTransform2d ResultType; };
+template<typename TransformType> struct ConcatenateRules<TransformType, std::enable_if_t<!TransformCalculusHelper::TIsTransform2<TransformType>::Value, FTransform2f>> { typedef FTransform2f ResultType; };
+template<typename TransformType> struct ConcatenateRules<TransformType, std::enable_if_t<!TransformCalculusHelper::TIsTransform2<TransformType>::Value, FTransform2d>> { typedef FTransform2d ResultType; };
diff --git a/Engine/Source/Runtime/Core/Public/Math/UnrealMathNeon.h b/Engine/Source/Runtime/Core/Public/Math/UnrealMathNeon.h
--- a/Engine/Source/Runtime/Core/Public/Math/UnrealMathNeon.h
+++ b/Engine/Source/Runtime/Core/Public/Math/UnrealMathNeon.h
@@ -6,8 +6,6 @@
 
 PRAGMA_DISABLE_SHADOW_VARIABLE_WARNINGS
 
-#include <type_traits>
-
 // Include the intrinsic functions header
 #if ((PLATFORM_WINDOWS || PLATFORM_HOLOLENS) && PLATFORM_64BITS)
 #include <arm64_neon.h>
@@ -531,21 +529,20 @@ FORCEINLINE VectorRegister2Double VectorSetComponentImpl(const VectorRegister2Do
 	return vsetq_lane_f64(Scalar, Vec, ElementIndex);
 }
 
-template<int ElementIndex, typename std::enable_if< (ElementIndex > 1), bool >::type = true >
-FORCEINLINE VectorRegister4Double VectorSetComponentImpl(const VectorRegister4Double& Vec, double Scalar)
-{
-	VectorRegister4Double Result;
-	Result.XY = Vec.XY;
-	Result.ZW = VectorSetComponentImpl<ElementIndex - 2>(Vec.ZW, Scalar);
-	return Result;
-}
-
-template<int ElementIndex, typename std::enable_if < (ElementIndex <= 1), bool >::type = true >
+template <int ElementIndex>
 FORCEINLINE VectorRegister4Double VectorSetComponentImpl(const VectorRegister4Double& Vec, double Scalar)
 {
 	VectorRegister4Double Result;
-	Result.XY = VectorSetComponentImpl<ElementIndex>(Vec.XY, Scalar);
-	Result.ZW = Vec.ZW;
+	if constexpr (ElementIndex > 1)
+	{
+		Result.XY = Vec.XY;
+		Result.ZW = VectorSetComponentImpl<ElementIndex - 2>(Vec.ZW, Scalar);
+	}
+	else
+	{
+		Result.XY = VectorSetComponentImpl<ElementIndex>(Vec.XY, Scalar);
+		Result.ZW = Vec.ZW;
+	}
 	return Result;
 }
 
@@ -745,21 +742,20 @@ FORCEINLINE VectorRegister2Double VectorReplicateImpl(const VectorRegister2Doubl
 	return vdupq_n_f64(vgetq_lane_f64(Vec, ElementIndex));
 }
 
-template <int ElementIndex, typename std::enable_if < (ElementIndex <= 1), bool >::type = true >
-FORCEINLINE VectorRegister4Double VectorReplicateImpl(const VectorRegister4Double& Vec)
-{
-	VectorRegister4Double Result;
-	Result.XY = VectorReplicateImpl<ElementIndex>(Vec.XY);
-	Result.ZW = Result.XY;
-	return Result;
-}
-
-template <int ElementIndex, typename std::enable_if < (ElementIndex > 1), bool >::type = true >
+template <int ElementIndex>
 FORCEINLINE VectorRegister4Double VectorReplicateImpl(const VectorRegister4Double& Vec)
 {
 	VectorRegister4Double Result;
-	Result.ZW = VectorReplicateImpl<ElementIndex - 2>(Vec.ZW);
-	Result.XY = Result.ZW;
+	if constexpr (ElementIndex <= 1)
+	{
+		Result.XY = VectorReplicateImpl<ElementIndex>(Vec.XY);
+		Result.ZW = Result.XY;
+	}
+	else
+	{
+		Result.ZW = VectorReplicateImpl<ElementIndex - 2>(Vec.ZW);
+		Result.XY = Result.ZW;
+	}
 	return Result;
 }
 
@@ -1299,28 +1295,31 @@ FORCEINLINE VectorRegister4Float VectorSwizzleImpl(VectorRegister4Float Vec)
 	return __builtin_shufflevector(Vec, Vec, X, Y, Z, W);
 }
 
-template <int X, int Y, typename std::enable_if < (X <= 1) && (Y <= 1), bool >::type = true>
-FORCEINLINE VectorRegister2Double VectorSwizzleImpl2(VectorRegister4Double Vec)
-{
-	return __builtin_shufflevector(Vec.XY, Vec.XY, X, Y);
-}
-
-template <int X, int Y, typename std::enable_if < (X <= 1) && (Y > 1), bool >::type = true>
-FORCEINLINE VectorRegister2Double VectorSwizzleImpl2(VectorRegister4Double Vec)
-{
-	return __builtin_shufflevector(Vec.XY, Vec.ZW, X, Y);
-}
-
-template <int X, int Y, typename std::enable_if < (X > 1) && (Y <= 1), bool >::type = true>
-FORCEINLINE VectorRegister2Double VectorSwizzleImpl2(VectorRegister4Double Vec)
-{
-	return __builtin_shufflevector(Vec.ZW, Vec.XY, X - 2, Y + 2);
-}
-
-template <int X, int Y, typename std::enable_if < (X > 1) && (Y > 1), bool >::type = true>
+template <int X, int Y>
 FORCEINLINE VectorRegister2Double VectorSwizzleImpl2(VectorRegister4Double Vec)
 {
-	return __builtin_shufflevector(Vec.ZW, Vec.ZW, X - 2, Y);
+	if constexpr (X <= 1)
+	{
+		if constexpr (Y <= 1)
+		{
+			return __builtin_shufflevector(Vec.XY, Vec.XY, X, Y);
+		}
+		else
+		{
+			return __builtin_shufflevector(Vec.XY, Vec.ZW, X, Y);
+		}
+	}
+	else
+	{
+		else if constexpr (Y <= 1)
+		{
+			return __builtin_shufflevector(Vec.ZW, Vec.XY, X - 2, Y + 2);
+		}
+		else
+		{
+			return __builtin_shufflevector(Vec.ZW, Vec.ZW, X - 2, Y);
+		}
+	}
 }
 
 template <int X, int Y, int Z, int W>
@@ -1802,16 +1801,17 @@ FORCEINLINE double VectorGetComponentImpl(VectorRegister2Double Vec)
 	return vgetq_lane_f64(Vec, ElementIndex);
 }
 
-template<int ElementIndex, typename std::enable_if< (ElementIndex > 1), bool >::type = true >
-FORCEINLINE double VectorGetComponentImpl(const VectorRegister4Double& Vec)
-{
-	return VectorGetComponentImpl<ElementIndex - 2>(Vec.ZW);
-}
-
-template<int ElementIndex, typename std::enable_if < (ElementIndex <= 1), bool >::type = true >
+template <int ElementIndex>
 FORCEINLINE double VectorGetComponentImpl(const VectorRegister4Double& Vec)
 {
-	return VectorGetComponentImpl<ElementIndex>(Vec.XY);
+	if constexpr (ElementIndex > 1)
+	{
+		return VectorGetComponentImpl<ElementIndex - 2>(Vec.ZW);
+	}
+	else
+	{
+		return VectorGetComponentImpl<ElementIndex>(Vec.XY);
+	}
 }
 
 #define VectorGetComponent(Vec, ElementIndex) VectorGetComponentImpl<ElementIndex>(Vec)
@@ -2883,3 +2883,7 @@ FORCEINLINE VectorRegister4i VectorIntExpandLow16To32(VectorRegister4i V) {
 // To be continued...
 
 PRAGMA_ENABLE_SHADOW_VARIABLE_WARNINGS
+
+#if UE_ENABLE_INCLUDE_ORDER_DEPRECATED_IN_5_4
+#include <type_traits>
+#endif
\ No newline at end of file
diff --git a/Engine/Source/Runtime/Core/Public/Math/UnrealMathSSE.h b/Engine/Source/Runtime/Core/Public/Math/UnrealMathSSE.h
--- a/Engine/Source/Runtime/Core/Public/Math/UnrealMathSSE.h
+++ b/Engine/Source/Runtime/Core/Public/Math/UnrealMathSSE.h
@@ -843,7 +843,6 @@ FORCEINLINE void VectorStoreFloat1(const VectorRegister4Double& Vec, double* Dst
 
 namespace SSEPermuteHelpers
 {
-	// These need to be #define and not constexpr to make them work with enable_if
 #define InLane0(Index0, Index1)		((Index0) <= 1 && (Index1) <= 1)
 #define InLane1(Index0, Index1)		((Index0) >= 2 && (Index1) >= 2)
 #define InSameLane(Index0, Index1)	(InLane0(Index0, Index1) || InLane1(Index0, Index1))
@@ -856,32 +855,35 @@ namespace SSEPermuteHelpers
 	// Double Swizzle helpers
 	// Templated swizzles required for double shuffles when using __m128d, since we have to break it down in to two separate operations.
 
-	// [0,1]:[0,1]
-	template<int Index0, int Index1, typename std::enable_if< (Index0 <= 1) && (Index1 <= 1), bool >::type = true>
+	template <int Index0, int Index1>
 	FORCEINLINE VectorRegister2Double SelectVectorSwizzle2(const VectorRegister4Double& Vec)
 	{
-		return _mm_shuffle_pd(Vec.GetXY(), Vec.GetXY(), SHUFFLEMASK2(Index0, Index1));
-	}
-
-	// [0,1]:[2,3]
-	template<int Index0, int Index1, typename std::enable_if< (Index0 <= 1) && (Index1 >= 2), bool >::type = true>
-	FORCEINLINE VectorRegister2Double SelectVectorSwizzle2(const VectorRegister4Double& Vec)
-	{
-		return _mm_shuffle_pd(Vec.GetXY(), Vec.GetZW(), SHUFFLEMASK2(Index0, Index1 - 2));
-	}
-
-	// [2,3]:[0,1]
-	template<int Index0, int Index1, typename std::enable_if< (Index0 >= 2) && (Index1 <= 1), bool >::type = true>
-	FORCEINLINE VectorRegister2Double SelectVectorSwizzle2(const VectorRegister4Double& Vec)
-	{
-		return _mm_shuffle_pd(Vec.GetZW(), Vec.GetXY(), SHUFFLEMASK2(Index0 - 2, Index1));
-	}
-
-	// [2,3]:[2,3]
-	template<int Index0, int Index1, typename std::enable_if< (Index0 >= 2) && (Index1 >= 2), bool >::type = true>
-	FORCEINLINE VectorRegister2Double SelectVectorSwizzle2(const VectorRegister4Double& Vec)
-	{
-		return _mm_shuffle_pd(Vec.GetZW(), Vec.GetZW(), SHUFFLEMASK2(Index0 - 2, Index1 - 2));
+		if constexpr (Index0 <= 1)
+		{
+			if constexpr (Index1 <= 1)
+			{
+				// [0,1]:[0,1]
+				return _mm_shuffle_pd(Vec.GetXY(), Vec.GetXY(), SHUFFLEMASK2(Index0, Index1));
+			}
+			else
+			{
+				// [0,1]:[2,3]
+				return _mm_shuffle_pd(Vec.GetXY(), Vec.GetZW(), SHUFFLEMASK2(Index0, Index1 - 2));
+			}
+		}
+		else
+		{
+			if constexpr (Index1 <= 1)
+			{
+				// [2,3]:[0,1]
+				return _mm_shuffle_pd(Vec.GetZW(), Vec.GetXY(), SHUFFLEMASK2(Index0 - 2, Index1));
+			}
+			else
+			{
+				// [2,3]:[2,3]
+				return _mm_shuffle_pd(Vec.GetZW(), Vec.GetZW(), SHUFFLEMASK2(Index0 - 2, Index1 - 2));
+			}
+		}
 	}
 
 	template<> FORCEINLINE VectorRegister2Double SelectVectorSwizzle2<0, 1>(const VectorRegister4Double& Vec) { return Vec.GetXY(); }
@@ -922,49 +924,44 @@ namespace SSEPermuteHelpers
 
 	constexpr int PERMUTE_MASK(int A, int B, int C, int D) { return ((A == 1 ? (1 << 0) : 0) | (B == 1 ? (1 << 1) : 0) | (C == 3 ? (1 << 2) : 0) | (D == 3 ? (1 << 3) : 0)); }
 
-	// [0..1][0..1][2..3][2..3]
-	template<int Index0, int Index1, int Index2, int Index3, typename std::enable_if< InLane0(Index0, Index1) && InLane1(Index2, Index3), bool >::type = true>
-	FORCEINLINE VectorRegister4Double SelectVectorSwizzle(const VectorRegister4Double& Vec)
-	{
-		return _mm256_permute_pd(Vec, PERMUTE_MASK(Index0, Index1, Index2, Index3));
-	}
-
-	// [2..3][2..3][0..1][0..1]
-	template <int Index0, int Index1, int Index2, int Index3, typename std::enable_if< InLane1(Index0, Index1) && InLane0(Index2, Index3), bool >::type = true>
-	FORCEINLINE VectorRegister4Double SelectVectorSwizzle(const VectorRegister4Double& Vec)
-	{
-		// Permute lanes then use [lane0][lane1] swizzle
-		return SelectVectorSwizzle<Index0 - 2, Index1 - 2, Index2 + 2, Index3 + 2>(PermuteLanes<1, 0>(Vec));
-	}
-
-	// [0..1][0..1][0..1][0..1]
-	template <int Index0, int Index1, int Index2, int Index3, typename std::enable_if< InLane0(Index0, Index1) && InLane0(Index2, Index3), bool >::type = true>
-	FORCEINLINE VectorRegister4Double SelectVectorSwizzle(const VectorRegister4Double& Vec)
-	{
-		// Permute lanes then use [lane0][lane1] swizzle
-		return SelectVectorSwizzle<Index0, Index1, Index2 + 2, Index3 + 2>(PermuteLanes<0, 0>(Vec));
-	}
-
-	// [2..3][2..3][2..3][2..3]
-	template <int Index0, int Index1, int Index2, int Index3, typename std::enable_if< InLane1(Index0, Index1) && InLane1(Index2, Index3), bool >::type = true>
-	FORCEINLINE VectorRegister4Double SelectVectorSwizzle(const VectorRegister4Double& Vec)
-	{
-		// Permute lanes then use [lane0][lane1] swizzle
-		return SelectVectorSwizzle<Index0 - 2, Index1 - 2, Index2, Index3>(PermuteLanes<1, 1>(Vec));
-	}
-
-	// Anything with out-of-lane pairs
-	template<int Index0, int Index1, int Index2, int Index3, typename std::enable_if< OutOfLane(Index0, Index1) || OutOfLane(Index2, Index3), bool >::type = true>
+	template <int Index0, int Index1, int Index2, int Index3>
 	FORCEINLINE VectorRegister4Double SelectVectorSwizzle(const VectorRegister4Double& Vec)
 	{
+		if constexpr (InLane0(Index0, Index1) && InLane1(Index2, Index3))
+		{
+			// [0..1][0..1][2..3][2..3]
+			return _mm256_permute_pd(Vec, PERMUTE_MASK(Index0, Index1, Index2, Index3));
+		}
+		else if constexpr (InLane1(Index0, Index1) && InLane0(Index2, Index3))
+		{
+			// [2..3][2..3][0..1][0..1]
+			// Permute lanes then use [lane0][lane1] swizzle
+			return SelectVectorSwizzle<Index0 - 2, Index1 - 2, Index2 + 2, Index3 + 2>(PermuteLanes<1, 0>(Vec));
+		}
+		else if constexpr (InLane0(Index0, Index1) && InLane0(Index2, Index3))
+		{
+			// [0..1][0..1][0..1][0..1]
+			// Permute lanes then use [lane0][lane1] swizzle
+			return SelectVectorSwizzle<Index0, Index1, Index2 + 2, Index3 + 2>(PermuteLanes<0, 0>(Vec));
+		}
+		else if constexpr (InLane1(Index0, Index1) && InLane1(Index2, Index3))
+		{
+			// [2..3][2..3][2..3][2..3]
+			// Permute lanes then use [lane0][lane1] swizzle
+			return SelectVectorSwizzle<Index0 - 2, Index1 - 2, Index2, Index3>(PermuteLanes<1, 1>(Vec));
+		}
+		else
+		{
+			// Anything with out-of-lane pairs
 #if UE_PLATFORM_MATH_USE_AVX_2
-		return _mm256_permute4x64_pd(Vec, SHUFFLEMASK(Index0, Index1, Index2, Index3));
+			return _mm256_permute4x64_pd(Vec, SHUFFLEMASK(Index0, Index1, Index2, Index3));
 #else
-		return VectorRegister4Double(
-			SelectVectorSwizzle2<Index0, Index1>(Vec),
-			SelectVectorSwizzle2<Index2, Index3>(Vec)
-		);
+			return VectorRegister4Double(
+				SelectVectorSwizzle2<Index0, Index1>(Vec),
+				SelectVectorSwizzle2<Index2, Index3>(Vec)
+			);
 #endif
+		}
 	}
 
 	//
@@ -1009,18 +1006,19 @@ namespace SSEPermuteHelpers
 	}
 
 	// Double replicate (4 doubles)
-	template <int Index, typename std::enable_if< (Index <= 1), bool >::type = true >
-	FORCEINLINE VectorRegister4Double VectorReplicateImpl4(const VectorRegister4Double& Vec)
-	{
-		VectorRegister2Double Temp = VectorReplicateImpl2<Index>(Vec.GetXY());
-		return VectorRegister4Double(Temp, Temp);
-	}
-
-	template <int Index, typename std::enable_if< (Index >= 2), bool >::type = true >
+	template <int Index>
 	FORCEINLINE VectorRegister4Double VectorReplicateImpl4(const VectorRegister4Double& Vec)
 	{
-		VectorRegister2Double Temp = VectorReplicateImpl2<Index - 2>(Vec.GetZW());
-		return VectorRegister4Double(Temp, Temp);
+		if constexpr (Index <= 1)
+		{
+			VectorRegister2Double Temp = VectorReplicateImpl2<Index>(Vec.GetXY());
+			return VectorRegister4Double(Temp, Temp);
+		}
+		else
+		{
+			VectorRegister2Double Temp = VectorReplicateImpl2<Index - 2>(Vec.GetZW());
+			return VectorRegister4Double(Temp, Temp);
+		}
 	}
 
 	//
@@ -1071,65 +1069,52 @@ namespace SSEPermuteHelpers
 	// When index pairs are within the same lane, SelectVectorShuffle first efficiently blends elements from the two vectors,
 	// then efficiently swizzles within 128-bit lanes using specializations for indices [0..1][0..1][2..3][2..3]
 	// 
-	// [0..1][0..1][2..3][2..3]
-	template <int Index0, int Index1, int Index2, int Index3, typename std::enable_if< InLane0(Index0, Index1) && InLane1(Index2, Index3), bool >::type = true>
-	FORCEINLINE VectorRegister4Double SelectVectorShuffle(const VectorRegister4Double& Vec1, const VectorRegister4Double& Vec2)
-	{
-		const VectorRegister4Double Blended = ShuffleLanes<0, 1>(Vec1, Vec2);
-		return VectorSwizzleTemplate<Index0, Index1, Index2, Index3>(Blended);
-	}
-
-	// [2..3][2..3][0..1][0..1]
-	template <int Index0, int Index1, int Index2, int Index3, typename std::enable_if< InLane1(Index0, Index1) && InLane0(Index2, Index3), bool >::type = true>
-	FORCEINLINE VectorRegister4Double SelectVectorShuffle(const VectorRegister4Double& Vec1, const VectorRegister4Double& Vec2)
-	{
-		const VectorRegister4Double Blended = ShuffleLanes<1, 0>(Vec1, Vec2);
-		return VectorSwizzleTemplate<Index0 - 2, Index1 - 2, Index2 + 2, Index3 + 2>(Blended);
-	}
-
-	// [0..1][0..1][0..1][0..1]
-	template <int Index0, int Index1, int Index2, int Index3, typename std::enable_if< InLane0(Index0, Index1) && InLane0(Index2, Index3), bool >::type = true>
-	FORCEINLINE VectorRegister4Double SelectVectorShuffle(const VectorRegister4Double& Vec1, const VectorRegister4Double& Vec2)
-	{
-		const VectorRegister4Double Blended = ShuffleLanes<0, 0>(Vec1, Vec2);
-		return VectorSwizzleTemplate<Index0, Index1, Index2 + 2, Index3 + 2>(Blended);
-	}
-
-	// [2..3][2..3][2..3][2..3]
-	template <int Index0, int Index1, int Index2, int Index3, typename std::enable_if< InLane1(Index0, Index1) && InLane1(Index2, Index3), bool >::type = true>
-	FORCEINLINE VectorRegister4Double SelectVectorShuffle(const VectorRegister4Double& Vec1, const VectorRegister4Double& Vec2)
-	{
-		const VectorRegister4Double Blended = ShuffleLanes<1, 1>(Vec1, Vec2);
-		return VectorSwizzleTemplate<Index0 - 2, Index1 - 2, Index2, Index3>(Blended);
-	}
-
-	//
-	// Remaining cases have at least one pair doing a cross-lane swizzle.
-	//
-
-	template <int Index0, int Index1, int Index2, int Index3, typename std::enable_if< InSameLane(Index0, Index1) && OutOfLane(Index2, Index3), bool >::type = true>
-	FORCEINLINE VectorRegister4Double SelectVectorShuffle(const VectorRegister4Double& Vec1, const VectorRegister4Double& Vec2)
-	{
-		const VectorRegister4Double Vec1_XY = VectorSwizzleTemplate<Index0, Index1, 2, 3>(Vec1);
-		const VectorRegister2Double Vec2_ZW = SelectVectorSwizzle2<Index2, Index3>(Vec2);
-		return _mm256_insertf128_pd(Vec1_XY, Vec2_ZW, 0x1);
-	}
-
-	template <int Index0, int Index1, int Index2, int Index3, typename std::enable_if< OutOfLane(Index0, Index1) && InSameLane(Index2, Index3), bool >::type = true>
+	template <int Index0, int Index1, int Index2, int Index3>
 	FORCEINLINE VectorRegister4Double SelectVectorShuffle(const VectorRegister4Double& Vec1, const VectorRegister4Double& Vec2)
 	{
-		const VectorRegister2Double Vec1_XY = SelectVectorSwizzle2<Index0, Index1>(Vec1);
-		const VectorRegister4Double Vec2_ZW = VectorSwizzleTemplate<0, 1, Index2, Index3>(Vec2);
-		return _mm256_insertf128_pd(Vec2_ZW, Vec1_XY, 0x0);
-	}
-
-	template <int Index0, int Index1, int Index2, int Index3, typename std::enable_if< OutOfLane(Index0, Index1) && OutOfLane(Index2, Index3), bool >::type = true>
-	FORCEINLINE VectorRegister4Double SelectVectorShuffle(const VectorRegister4Double& Vec1, const VectorRegister4Double& Vec2)
-	{
-		return VectorRegister4Double(
-			SelectVectorSwizzle2<Index0, Index1>(Vec1),
-			SelectVectorSwizzle2<Index2, Index3>(Vec2)
-		);
+		if constexpr (InLane0(Index0, Index1) && InLane1(Index2, Index3))
+		{
+			// [0..1][0..1][2..3][2..3]
+			const VectorRegister4Double Blended = ShuffleLanes<0, 1>(Vec1, Vec2);
+			return VectorSwizzleTemplate<Index0, Index1, Index2, Index3>(Blended);
+		}
+		else if constexpr (InLane1(Index0, Index1) && InLane0(Index2, Index3))
+		{
+			// [2..3][2..3][0..1][0..1]
+			const VectorRegister4Double Blended = ShuffleLanes<1, 0>(Vec1, Vec2);
+			return VectorSwizzleTemplate<Index0 - 2, Index1 - 2, Index2 + 2, Index3 + 2>(Blended);
+		}
+		else if constexpr (InLane0(Index0, Index1) && InLane0(Index2, Index3))
+		{
+			// [0..1][0..1][0..1][0..1]
+			const VectorRegister4Double Blended = ShuffleLanes<0, 0>(Vec1, Vec2);
+			return VectorSwizzleTemplate<Index0, Index1, Index2 + 2, Index3 + 2>(Blended);
+		}
+		else if constexpr (InLane1(Index0, Index1) && InLane1(Index2, Index3))
+		{
+			// [2..3][2..3][2..3][2..3]
+			const VectorRegister4Double Blended = ShuffleLanes<1, 1>(Vec1, Vec2);
+			return VectorSwizzleTemplate<Index0 - 2, Index1 - 2, Index2, Index3>(Blended);
+		}
+		else if constexpr (InSameLane(Index0, Index1) && OutOfLane(Index2, Index3))
+		{
+			const VectorRegister4Double Vec1_XY = VectorSwizzleTemplate<Index0, Index1, 2, 3>(Vec1);
+			const VectorRegister2Double Vec2_ZW = SelectVectorSwizzle2<Index2, Index3>(Vec2);
+			return _mm256_insertf128_pd(Vec1_XY, Vec2_ZW, 0x1);
+		}
+		else if constexpr (OutOfLane(Index0, Index1) && InSameLane(Index2, Index3))
+		{
+			const VectorRegister2Double Vec1_XY = SelectVectorSwizzle2<Index0, Index1>(Vec1);
+			const VectorRegister4Double Vec2_ZW = VectorSwizzleTemplate<0, 1, Index2, Index3>(Vec2);
+			return _mm256_insertf128_pd(Vec2_ZW, Vec1_XY, 0x0);
+		}
+		else
+		{
+			return VectorRegister4Double(
+				SelectVectorSwizzle2<Index0, Index1>(Vec1),
+				SelectVectorSwizzle2<Index2, Index3>(Vec2)
+			);
+		}
 	}
 
 	// AVX Double Shuffle specializations 
@@ -1214,7 +1199,11 @@ namespace SSEPermuteHelpers
 	template<> FORCEINLINE VectorRegister4Float VectorShuffleTemplate<0, 1, 0, 1>(const VectorRegister4Float& Vec1, const VectorRegister4Float& Vec2) { return _mm_movelh_ps(Vec1, Vec2); }
 	template<> FORCEINLINE VectorRegister4Float VectorShuffleTemplate<2, 3, 2, 3>(const VectorRegister4Float& Vec1, const VectorRegister4Float& Vec2) { return _mm_movehl_ps(Vec2, Vec1); } // Note: movehl copies first from the 2nd argument
 
-}; // namespace SSEPermuteHelpers
+#undef OutOfLane
+#undef InSameLane
+#undef InLane1
+#undef InLane0
+} // namespace SSEPermuteHelpers
 
 /**
  * Replicates one element into all four elements and returns the new vector.
