CIS fix.

diff --git a/Engine/Source/Runtime/Core/Public/Math/UnrealMathNeon.h b/Engine/Source/Runtime/Core/Public/Math/UnrealMathNeon.h
--- a/Engine/Source/Runtime/Core/Public/Math/UnrealMathNeon.h
+++ b/Engine/Source/Runtime/Core/Public/Math/UnrealMathNeon.h
@@ -1311,7 +1311,7 @@ FORCEINLINE VectorRegister2Double VectorSwizzleImpl2(VectorRegister4Double Vec)
 	}
 	else
 	{
-		else if constexpr (Y <= 1)
+		if constexpr (Y <= 1)
 		{
 			return __builtin_shufflevector(Vec.ZW, Vec.XY, X - 2, Y + 2);
 		}
