Replaced a bunch of UE traits and metaprogramming constructs from containers to improve PDB sizes and compile times.

diff --git a/Engine/Plugins/Experimental/LearningAgents/Source/Learning/Public/MultiArray.h b/Engine/Plugins/Experimental/LearningAgents/Source/Learning/Public/MultiArray.h
--- a/Engine/Plugins/Experimental/LearningAgents/Source/Learning/Public/MultiArray.h
+++ b/Engine/Plugins/Experimental/LearningAgents/Source/Learning/Public/MultiArray.h
@@ -5,13 +5,11 @@
 #include "CoreTypes.h"
 #include "Misc/AssertionMacros.h"
 #include "HAL/UnrealMemory.h"
-#include "Templates/IsSigned.h"
 #include "Templates/UnrealTypeTraits.h"
 #include "Templates/UnrealTemplate.h"
 #include "Containers/ContainerAllocationPolicies.h"
 #include "Containers/ContainerElementTypeCompatibility.h"
 
-#include "Templates/AndOrNot.h"
 #include "Templates/IdentityFunctor.h"
 #include "Templates/Invoke.h"
 #include "Templates/Less.h"
@@ -45,7 +43,7 @@ namespace UE::MultiArray::Private
 	// This should be changed to use std::disjunction and std::is_constructible, and the usage
 	// changed to use ::value instead of ::Value, when std::disjunction (C++17) is available everywhere.
 	template <typename DestType, typename SourceType>
-	using TMultiArrayElementsAreCompatible = TOrValue<std::is_same<DestType, std::decay_t<SourceType>>::value, TIsConstructible<DestType, SourceType>>;
+	constexpr bool TMultiArrayElementsAreCompatible_V = std::is_same_v<DestType, std::decay_t<SourceType>> || std::is_constructible_v<DestType, SourceType>;
 }
 
 
@@ -79,7 +77,7 @@ public:
 	>;
 
 	static_assert(DimNum > 1, "TMultiArray requires a positive, non-zero number of dimensions");
-	static_assert(TIsSigned<SizeType>::Value, "TMultiArray only supports signed index types");
+	static_assert(std::is_signed_v<SizeType>, "TMultiArray only supports signed index types");
 
 public:
 
@@ -298,8 +296,8 @@ public:
 	template <
 		typename OtherElementType,
 		typename OtherAllocator,
-		bool bOtherIsChecked,
-		std::enable_if_t<UE::MultiArray::Private::TMultiArrayElementsAreCompatible<ElementType, const OtherElementType&>::Value>* = nullptr
+		bool bOtherIsChecked
+		UE_REQUIRES(UE::MultiArray::Private::TMultiArrayElementsAreCompatible_V<ElementType, const OtherElementType&>)
 	>
 		FORCEINLINE explicit TMultiArray(const TMultiArray<DimNum, OtherElementType, OtherAllocator, bOtherIsChecked>& Other)
 	{
@@ -386,8 +384,8 @@ public:
 	template <
 		typename OtherElementType,
 		typename OtherAllocator,
-		bool bOtherIsChecked,
-		std::enable_if_t<UE::MultiArray::Private::TMultiArrayElementsAreCompatible<ElementType, OtherElementType&&>::Value>* = nullptr
+		bool bOtherIsChecked
+		UE_REQUIRES(UE::MultiArray::Private::TMultiArrayElementsAreCompatible_V<ElementType, OtherElementType&&>)
 	>
 		FORCEINLINE explicit TMultiArray(TMultiArray<DimNum, OtherElementType, OtherAllocator, bOtherIsChecked>&& Other)
 	{
@@ -843,7 +841,7 @@ public:
 		typename AllocatorType::ForAnyElementType
 	>;
 
-	static_assert(TIsSigned<SizeType>::Value, "TMultiArray only supports signed index types");
+	static_assert(std::is_signed_v<SizeType>, "TMultiArray only supports signed index types");
 
 	/**
 	 * Initializer list constructor
@@ -954,8 +952,8 @@ public:
 	template <
 		typename OtherElementType,
 		typename OtherAllocator,
-		bool bOtherIsChecked,
-		std::enable_if_t<UE::MultiArray::Private::TMultiArrayElementsAreCompatible<ElementType, const OtherElementType&>::Value>* = nullptr
+		bool bOtherIsChecked
+		UE_REQUIRES(UE::MultiArray::Private::TMultiArrayElementsAreCompatible_V<ElementType, const OtherElementType&>)
 	>
 		FORCEINLINE explicit TMultiArray(const TMultiArray<DimNum, OtherElementType, OtherAllocator, bOtherIsChecked>& Other)
 	{
@@ -1042,8 +1040,8 @@ public:
 	template <
 		typename OtherElementType,
 		typename OtherAllocator,
-		bool bOtherIsChecked,
-		std::enable_if_t<UE::MultiArray::Private::TMultiArrayElementsAreCompatible<ElementType, OtherElementType&&>::Value>* = nullptr
+		bool bOtherIsChecked
+		UE_REQUIRES(UE::MultiArray::Private::TMultiArrayElementsAreCompatible_V<ElementType, OtherElementType&&>)
 	>
 		FORCEINLINE explicit TMultiArray(TMultiArray<DimNum, OtherElementType, OtherAllocator, bOtherIsChecked>&& Other)
 	{
diff --git a/Engine/Plugins/Experimental/LearningAgents/Source/Learning/Public/MultiArrayView.h b/Engine/Plugins/Experimental/LearningAgents/Source/Learning/Public/MultiArrayView.h
--- a/Engine/Plugins/Experimental/LearningAgents/Source/Learning/Public/MultiArrayView.h
+++ b/Engine/Plugins/Experimental/LearningAgents/Source/Learning/Public/MultiArrayView.h
@@ -4,14 +4,12 @@
 
 #include "CoreTypes.h"
 #include "Containers/ArrayView.h"
-#include "Templates/IsConst.h"
-#include "Templates/IsSigned.h"
-#include "Templates/PointerIsConvertibleFromTo.h"
 #include "Misc/AssertionMacros.h"
 #include "Templates/Invoke.h"
 #include "Templates/UnrealTypeTraits.h"
 #include "Traits/ElementType.h"
 #include "Math/UnrealMathUtility.h"
+#include <type_traits>
 
 namespace UE::MultiArrayView::Private
 {
@@ -19,17 +17,12 @@ namespace UE::MultiArrayView::Private
 	 * Copied from Core/Public/Containers/ArrayView.h
 	 *
 	 * Trait testing whether a type is compatible with the view type
+	 *
+	 * The extra stars here are *IMPORTANT*
+	 * They prevent TMultiArrayView<Base>(TArray<Derived>&) from compiling!
 	 */
 	template <typename T, typename ElementType>
-	struct TIsCompatibleElementType
-	{
-	public:
-		/** NOTE:
-		 * The stars in the TPointerIsConvertibleFromTo test are *IMPORTANT*
-		 * They prevent TArrayView<Base>(TArray<Derived>&) from compiling!
-		 */
-		enum { Value = TPointerIsConvertibleFromTo<T*, ElementType* const>::Value };
-	};
+	constexpr bool TIsCompatibleElementType_V = std::is_convertible_v<T**, ElementType* const*>;
 
 	// Copied from Core/Public/Containers/ArrayView.h
 	//
@@ -48,7 +41,7 @@ namespace UE::MultiArrayView::Private
 	FORCEINLINE decltype(auto) GetReinterpretedDataHelper(T&& Arg)
 	{
 		auto NaturalPtr = GetData(Forward<T>(Arg));
-		using NaturalElementType = typename TRemovePointer<decltype(NaturalPtr)>::Type;
+		using NaturalElementType = std::remove_pointer_t<decltype(NaturalPtr)>;
 
 		auto EndPtr = NaturalPtr + GetNum(Arg);
 		TContainerElementTypeCompatibility<NaturalElementType>::ReinterpretRange(NaturalPtr, EndPtr);
@@ -64,7 +57,7 @@ namespace UE::MultiArrayView::Private
 	template <typename RangeType, typename ElementType>
 	struct TIsCompatibleRangeType
 	{
-		static constexpr bool Value = TIsCompatibleElementType<typename TRemovePointer<decltype(GetData(DeclVal<RangeType&>()))>::Type, ElementType>::Value;
+		static constexpr bool Value = TIsCompatibleElementType_V<std::remove_pointer_t<decltype(GetData(DeclVal<RangeType&>()))>, ElementType>;
 
 		template <typename T>
 		static decltype(auto) GetData(T&& Arg)
@@ -82,13 +75,13 @@ namespace UE::MultiArrayView::Private
 	struct TIsReinterpretableRangeType
 	{
 	private:
-		using NaturalElementType = typename TRemovePointer<decltype(GetData(DeclVal<RangeType&>()))>::Type;
+		using NaturalElementType = std::remove_pointer_t<decltype(GetData(DeclVal<RangeType&>()))>;
 
 	public:
 		static constexpr bool Value =
 			!std::is_same_v<typename TContainerElementTypeCompatibility<NaturalElementType>::ReinterpretType, NaturalElementType>
 			&&
-			TIsCompatibleElementType<typename TContainerElementTypeCompatibility<NaturalElementType>::ReinterpretType, ElementType>::Value;
+			TIsCompatibleElementType_V<typename TContainerElementTypeCompatibility<NaturalElementType>::ReinterpretType, ElementType>;
 
 		template <typename T>
 		static decltype(auto) GetData(T&& Arg)
@@ -106,7 +99,7 @@ public:
 	using SizeType = InSizeType;
 
 	static_assert(DimNum >= 1, "TMultiArrayShape requires a positive, non-zero number of dimensions");
-	static_assert(TIsSigned<SizeType>::Value, "TMultiArrayShape only supports signed index types");
+	static_assert(std::is_signed_v<SizeType>, "TMultiArrayShape only supports signed index types");
 
 public:
 
@@ -186,7 +179,7 @@ public:
 		ElementType*>;
 
 	static_assert(DimNum > 1, "TMultiArrayView requires a positive, non-zero number of dimensions");
-	static_assert(TIsSigned<SizeType>::Value, "TMultiArrayView only supports signed index types");
+	static_assert(std::is_signed_v<SizeType>, "TMultiArrayView only supports signed index types");
 
 public:
 
@@ -617,7 +610,7 @@ public:
 		ElementType* RESTRICT,
 		ElementType*>;
 
-	static_assert(TIsSigned<SizeType>::Value, "TMultiArrayView only supports signed index types");
+	static_assert(std::is_signed_v<SizeType>, "TMultiArrayView only supports signed index types");
 
 	/**
 	 * Constructor.
@@ -625,9 +618,6 @@ public:
 	TMultiArrayView() : DataPtr(nullptr) { }
 
 private:
-	template <typename T>
-	using TIsCompatibleElementType = UE::MultiArrayView::Private::TIsCompatibleElementType<T, ElementType>;
-
 	template <typename T>
 	using TIsCompatibleRangeType = UE::MultiArrayView::Private::TIsCompatibleRangeType<T, ElementType>;
 
@@ -642,8 +632,8 @@ public:
 	 */
 	template <
 		typename OtherRangeType,
-		typename CVUnqualifiedOtherRangeType = std::remove_cv_t<typename TRemoveReference<OtherRangeType>::Type>,
-		typename = typename TEnableIf<
+		typename CVUnqualifiedOtherRangeType = std::remove_cv_t<std::remove_reference_t<OtherRangeType>>
+		UE_REQUIRES(
 			TAnd<
 				TIsContiguousContainer<CVUnqualifiedOtherRangeType>,
 				TOr<
@@ -651,9 +641,9 @@ public:
 					TIsReinterpretableRangeType<OtherRangeType>
 				>
 			>::Value
-		>::Type
+		)
 	>
-		FORCEINLINE TMultiArrayView(OtherRangeType&& Other)
+	FORCEINLINE TMultiArrayView(OtherRangeType&& Other)
 		: DataPtr(std::conditional_t<
 			TIsCompatibleRangeType<OtherRangeType>::Value,
 			TIsCompatibleRangeType<OtherRangeType>,
@@ -684,9 +674,11 @@ public:
 	 * @param InData	The data to view
 	 * @param InNum	The number of elements
 	 */
-	template <typename OtherElementType,
-		typename = typename TEnableIf<TIsCompatibleElementType<OtherElementType>::Value>::Type>
-		FORCEINLINE TMultiArrayView(OtherElementType* InData, SizeType InNum)
+	template <
+		typename OtherElementType
+		UE_REQUIRES(UE::MultiArrayView::Private::TIsCompatibleElementType_V<OtherElementType, ElementType>)
+	>
+	FORCEINLINE TMultiArrayView(OtherElementType* InData, SizeType InNum)
 		: DataPtr(InData)
 	{
 		ArrayShape[0] = InNum;
@@ -1024,3 +1016,9 @@ private:
 
 template<uint8 InDimNum, typename InElementType, bool bInIsChecked = true, bool bInIsRestrict = false, typename InSizeType = FDefaultAllocator::SizeType>
 using TConstMultiArrayView = TMultiArrayView<InDimNum, const InElementType, bInIsChecked, bInIsRestrict, InSizeType>;
+
+#if UE_ENABLE_INCLUDE_ORDER_DEPRECATED_IN_5_4
+#include "Templates/IsConst.h"
+#include "Templates/IsSigned.h"
+#include "Templates/PointerIsConvertibleFromTo.h"
+#endif
\ No newline at end of file
diff --git a/Engine/Plugins/Experimental/StateGraph/Source/StateGraph/Public/StateGraph.h b/Engine/Plugins/Experimental/StateGraph/Source/StateGraph/Public/StateGraph.h
--- a/Engine/Plugins/Experimental/StateGraph/Source/StateGraph/Public/StateGraph.h
+++ b/Engine/Plugins/Experimental/StateGraph/Source/StateGraph/Public/StateGraph.h
@@ -21,6 +21,8 @@
 #include "Misc/CoreMisc.h"
 #include "StateGraphFwd.h"
 #include "Templates/IsClass.h"
+#include "Templates/IsInvocable.h"
+#include "Templates/IsMemberPointer.h"
 #include "Templates/UnrealTemplate.h"
 #include "UObject/NameTypes.h"
 
diff --git a/Engine/Source/Runtime/Core/Public/Containers/Array.h b/Engine/Source/Runtime/Core/Public/Containers/Array.h
--- a/Engine/Source/Runtime/Core/Public/Containers/Array.h
+++ b/Engine/Source/Runtime/Core/Public/Containers/Array.h
@@ -6,7 +6,6 @@
 #include "Misc/AssertionMacros.h"
 #include "Misc/ReverseIterate.h"
 #include "HAL/UnrealMemory.h"
-#include "Templates/IsSigned.h"
 #include "Templates/UnrealTypeTraits.h"
 #include "Templates/UnrealTemplate.h"
 #include "Containers/ContainerAllocationPolicies.h"
@@ -20,14 +19,11 @@
 #include "Algo/Impl/BinaryHeap.h"
 #include "Algo/StableSort.h"
 #include "Concepts/GetTypeHashable.h"
-#include "Templates/AndOrNot.h"
 #include "Templates/IdentityFunctor.h"
 #include "Templates/Invoke.h"
 #include "Templates/Less.h"
 #include "Templates/Sorting.h"
 #include "Templates/AlignmentTemplates.h"
-#include "Templates/IsConstructible.h"
-#include "Templates/MakeUnsigned.h"
 #include "Traits/ElementType.h"
 
 #include <limits>
@@ -333,28 +329,16 @@ namespace UE4Array_Private
 	// {
 	//     TArray<FThing> Arr; // this will cause errors without this workaround
 	// };
-	//
-	// This should be changed to use std::disjunction and std::is_constructible, and the usage
-	// changed to use ::value instead of ::Value, when std::disjunction (C++17) is available everywhere.
 	template <typename DestType, typename SourceType>
-	using TArrayElementsAreCompatible = TOrValue<std::is_same<DestType, std::decay_t<SourceType>>::value, TIsConstructible<DestType, SourceType>>;
-
-	template <typename T>
-	struct TIsTArrayOrDerivedFromTArray
-	{
-		template <typename ElementType, typename AllocatorType>
-		static char (&Resolve(const volatile TArray<ElementType, AllocatorType>*))[2];
+	constexpr bool TArrayElementsAreCompatible_V = std::disjunction_v<std::is_same<DestType, std::decay_t<SourceType>>, std::is_constructible<DestType, SourceType>>;
 
-		static char(&Resolve(...))[1];
+	template <typename ElementType, typename AllocatorType>
+	static char (&ResolveIsTArrayPtr(const volatile TArray<ElementType, AllocatorType>*))[2];
 
-		enum { Value = sizeof(Resolve((T*)nullptr)) == 2 };
-	};
+	static char(&ResolveIsTArrayPtr(...))[1];
 
-	template <typename ElementType, typename RangeType>
-	struct TTypeIsCompatibleWithRangeElementType
-	{
-		enum { Value = TArrayElementsAreCompatible<ElementType, TElementType_T<RangeType>>::Value };
-	};
+	template <typename T>
+	constexpr bool TIsTArrayOrDerivedFromTArray_V = sizeof(ResolveIsTArrayPtr((T*)nullptr)) == 2;
 }
 
 namespace UE::Core::Private
@@ -386,19 +370,16 @@ public:
 	typedef InAllocatorType AllocatorType;
 
 private:
-	using USizeType = typename TMakeUnsigned<SizeType>::Type;
+	using USizeType = typename std::make_unsigned_t<SizeType>;
 
 public:
-	UE_DEPRECATED(5.0, "TArray::Allocator type is deprecated, please use TArray::AllocatorType instead.")
-	typedef InAllocatorType Allocator;
-
 	using ElementAllocatorType = std::conditional_t<
 		AllocatorType::NeedsElementType,
 		typename AllocatorType::template ForElementType<ElementType>,
 		typename AllocatorType::ForAnyElementType
 	>;
 
-	static_assert(TIsSigned<SizeType>::Value, "TArray only supports signed index types");
+	static_assert(std::is_signed_v<SizeType>, "TArray only supports signed index types");
 
 	/**
 	 * Constructor, initializes element number counters.
@@ -449,8 +430,8 @@ public:
 	 */
 	template <
 		typename OtherElementType,
-		typename OtherAllocator,
-		std::enable_if_t<UE4Array_Private::TArrayElementsAreCompatible<ElementType, const OtherElementType&>::Value>* = nullptr
+		typename OtherAllocator
+		UE_REQUIRES(UE4Array_Private::TArrayElementsAreCompatible_V<ElementType, const OtherElementType&>)
 	>
 	FORCEINLINE explicit TArray(const TArray<OtherElementType, OtherAllocator>& Other)
 	{
@@ -638,8 +619,8 @@ public:
 	 */
 	template <
 		typename OtherElementType,
-		typename OtherAllocator,
-		std::enable_if_t<UE4Array_Private::TArrayElementsAreCompatible<ElementType, OtherElementType&&>::Value>* = nullptr
+		typename OtherAllocator
+		UE_REQUIRES(UE4Array_Private::TArrayElementsAreCompatible_V<ElementType, OtherElementType&&>)
 	>
 	FORCEINLINE explicit TArray(TArray<OtherElementType, OtherAllocator>&& Other)
 	{
@@ -654,8 +635,8 @@ public:
 	 *                   at the end of the array in the number of elements.
 	 */
 	template <
-		typename OtherElementType,
-		std::enable_if_t<UE4Array_Private::TArrayElementsAreCompatible<ElementType, OtherElementType&&>::Value>* = nullptr
+		typename OtherElementType
+		UE_REQUIRES(UE4Array_Private::TArrayElementsAreCompatible_V<ElementType, OtherElementType&&>)
 	>
 	TArray(TArray<OtherElementType, AllocatorType>&& Other, SizeType ExtraSlack)
 	{
@@ -837,8 +818,7 @@ public:
 	 * @param bAllowShrinking If this call allows shrinking of the array during element remove.
 	 * @returns Popped element.
 	 */
-	template<typename ET=InElementType>
-	FORCEINLINE typename TEnableIf<!TIsAbstract<ET>::Value, ElementType>::Type Pop(bool bAllowShrinking = true)
+	FORCEINLINE ElementType Pop(bool bAllowShrinking = true)
 	{
 		RangeCheck(0);
 		ElementType Result = MoveTempIfPossible(GetData()[ArrayNum - 1]);
@@ -1980,16 +1960,12 @@ public:
 	 * @see Add, Insert
 	 */
 	template <
-		typename RangeType,
-		typename RangeValueType = std::remove_reference_t<RangeType>,
-		typename RangeElementType = TElementType_T<RangeValueType>,
-		std::enable_if_t<
-			TAnd<
-				TIsContiguousContainer<RangeValueType>,
-				TNot<UE4Array_Private::TIsTArrayOrDerivedFromTArray<RangeValueType>>,
-				UE4Array_Private::TTypeIsCompatibleWithRangeElementType<ElementType, RangeType>
-			>::Value
-		>* = nullptr
+		typename RangeType
+		UE_REQUIRES(
+			TIsContiguousContainer<RangeType>::Value &&
+			!UE4Array_Private::TIsTArrayOrDerivedFromTArray_V<std::remove_reference_t<RangeType>> &&
+			UE4Array_Private::TArrayElementsAreCompatible_V<ElementType, TElementType_T<RangeType>>
+		)
 	>
 	void Append(RangeType&& Source)
 	{
@@ -2295,9 +2271,9 @@ public:
 #if !UE_DEPRECATE_MUTABLE_TOBJECTPTR
 	/** Mutable implicit conversion operator to container of compatible element type. */
 	template <
-		typename AliasElementType = ElementType,
-		std::enable_if_t<TIsContainerElementTypeReinterpretable<AliasElementType>::Value>* = nullptr
-		>
+		typename AliasElementType = ElementType
+		UE_REQUIRES(TIsContainerElementTypeReinterpretable<AliasElementType>::Value)
+	>
 	operator TArray<typename TContainerElementTypeCompatibility<AliasElementType>::ReinterpretType, AllocatorType>& ()
 	{
 		using ElementCompat = TContainerElementTypeCompatibility<ElementType>;
@@ -2308,9 +2284,9 @@ public:
   
 	/** Immutable implicit conversion operator to constant container of compatible element type. */
 	template <
-		typename AliasElementType = ElementType,
-		typename std::enable_if_t<TIsContainerElementTypeReinterpretable<AliasElementType>::Value>* = nullptr
-		>
+		typename AliasElementType = ElementType
+		UE_REQUIRES(TIsContainerElementTypeReinterpretable<AliasElementType>::Value)
+	>
 	operator const TArray<typename TContainerElementTypeCompatibility<AliasElementType>::ReinterpretType, AllocatorType>& () const
 	{
 		using ElementCompat = TContainerElementTypeCompatibility<ElementType>;
@@ -2325,8 +2301,8 @@ public:
 	 * @param Other Array to assign and move from.
 	 */
 	template <
-		typename AliasElementType = ElementType,
-		typename std::enable_if_t<TIsContainerElementTypeCopyable<AliasElementType>::Value>* = nullptr
+		typename AliasElementType = ElementType
+		UE_REQUIRES(TIsContainerElementTypeCopyable<AliasElementType>::Value)
 	>
 	TArray& operator=(TArray<typename TContainerElementTypeCompatibility<ElementType>::CopyFromOtherType, AllocatorType>&& Other)
 	{
@@ -2345,8 +2321,8 @@ public:
 	 */
 	template <
 		typename OtherAllocator,
-		typename AliasElementType = ElementType,
-		typename std::enable_if_t<TIsContainerElementTypeCopyable<AliasElementType>::Value>* = nullptr
+		typename AliasElementType = ElementType
+		UE_REQUIRES(TIsContainerElementTypeCopyable<AliasElementType>::Value)
 	>
 	TArray& operator=(const TArray<typename TContainerElementTypeCompatibility<ElementType>::CopyFromOtherType, OtherAllocator>& Other)
 	{
@@ -2366,8 +2342,8 @@ public:
 	 */
 	template <
 		typename OtherAllocator,
-		typename AliasElementType = ElementType,
-		typename std::enable_if_t<TIsContainerElementTypeCopyable<AliasElementType>::Value>* = nullptr
+		typename AliasElementType = ElementType
+		UE_REQUIRES(TIsContainerElementTypeCopyable<AliasElementType>::Value)
 	>
 	SizeType Insert(const TArray<typename TContainerElementTypeCompatibility<ElementType>::CopyFromOtherType, OtherAllocator>& Items, const SizeType InIndex)
 	{
@@ -2391,8 +2367,8 @@ public:
 	 */
 	template <
 		typename OtherAllocator,
-		typename AliasElementType = ElementType,
-		typename std::enable_if_t<TIsContainerElementTypeCopyable<AliasElementType>::Value>* = nullptr
+		typename AliasElementType = ElementType
+		UE_REQUIRES(TIsContainerElementTypeCopyable<AliasElementType>::Value)
 	>
 	SizeType Insert(TArray<typename TContainerElementTypeCompatibility<ElementType>::CopyFromOtherType, OtherAllocator>&& Items, const SizeType InIndex)
 	{
@@ -2417,8 +2393,8 @@ public:
 	 * @see Add, Insert
 	 */
 	template <
-		typename AliasElementType = ElementType,
-		typename std::enable_if_t<TIsContainerElementTypeCopyable<AliasElementType>::Value>* = nullptr
+		typename AliasElementType = ElementType
+		UE_REQUIRES(TIsContainerElementTypeCopyable<AliasElementType>::Value)
 	>
 	void Append(const typename TContainerElementTypeCompatibility<ElementType>::CopyFromOtherType* Ptr, SizeType Count)
 	{
@@ -3607,3 +3583,10 @@ uint32 GetTypeHash(const TArray<InElementType, InAllocatorType>& A)
 	}
 	return Hash;
 }
+
+#if UE_ENABLE_INCLUDE_ORDER_DEPRECATED_IN_5_4
+#include "Templates/IsSigned.h"
+#include "Templates/AndOrNot.h"
+#include "Templates/IsConstructible.h"
+#include "Templates/MakeUnsigned.h"
+#endif
diff --git a/Engine/Source/Runtime/Core/Public/Containers/ArrayView.h b/Engine/Source/Runtime/Core/Public/Containers/ArrayView.h
--- a/Engine/Source/Runtime/Core/Public/Containers/ArrayView.h
+++ b/Engine/Source/Runtime/Core/Public/Containers/ArrayView.h
@@ -3,9 +3,6 @@
 #pragma once
 
 #include "CoreTypes.h"
-#include "Templates/IsConst.h"
-#include "Templates/IsSigned.h"
-#include "Templates/PointerIsConvertibleFromTo.h"
 #include "Misc/AssertionMacros.h"
 #include "Misc/ReverseIterate.h"
 #include "Templates/Invoke.h"
@@ -13,22 +10,18 @@
 #include "Traits/ElementType.h"
 #include "Containers/Array.h"
 #include "Math/UnrealMathUtility.h"
+#include <type_traits>
 
 namespace UE::Core::ArrayView::Private
 {
 	/**
 	 * Trait testing whether a type is compatible with the view type
+	 *
+	 * The extra stars here are *IMPORTANT*
+	 * They prevent TMultiArrayView<Base>(TArray<Derived>&) from compiling!
 	 */
 	template <typename T, typename ElementType>
-	struct TIsCompatibleElementType
-	{
-	public:
-		/** NOTE:
-		 * The stars in the TPointerIsConvertibleFromTo test are *IMPORTANT*
-		 * They prevent TArrayView<Base>(TArray<Derived>&) from compiling!
-		 */
-		enum { Value = TPointerIsConvertibleFromTo<T*, ElementType* const>::Value };
-	};
+	constexpr bool TIsCompatibleElementType_V =std::is_convertible_v<T**, ElementType* const*>;
 
 	// Simply forwards to an unqualified GetData(), but can be called from within TArrayView
 	// where GetData() is already a member and so hides any others.
@@ -58,7 +51,7 @@ namespace UE::Core::ArrayView::Private
 	template <typename RangeType, typename ElementType>
 	struct TIsCompatibleRangeType
 	{
-		static constexpr bool Value = TIsCompatibleElementType<std::remove_pointer_t<decltype(GetData(DeclVal<RangeType&>()))>, ElementType>::Value;
+		static constexpr bool Value = TIsCompatibleElementType_V<std::remove_pointer_t<decltype(GetData(DeclVal<RangeType&>()))>, ElementType>;
 
 		template <typename T>
 		static decltype(auto) GetData(T&& Arg)
@@ -81,7 +74,7 @@ namespace UE::Core::ArrayView::Private
 		static constexpr bool Value = 
 			!std::is_same_v<typename TypeCompat::ReinterpretType, NaturalElementType>
 			&&
-			TIsCompatibleElementType<typename TypeCompat::ReinterpretType, ElementType>::Value
+			TIsCompatibleElementType_V<typename TypeCompat::ReinterpretType, ElementType>
 			&&
 			(!UE_DEPRECATE_MUTABLE_TOBJECTPTR
 			 || std::is_same_v<ElementType, std::remove_pointer_t<typename TypeCompat::ReinterpretType>* const>
@@ -145,7 +138,7 @@ public:
 	using ElementType = InElementType;
 	using SizeType = InSizeType;
 
-	static_assert(TIsSigned<SizeType>::Value, "TArrayView only supports signed index types");
+	static_assert(std::is_signed_v<SizeType>, "TArrayView only supports signed index types");
 
 	// Defaulted object behavior - we want compiler-generated functions rather than going through the generic range constructor.
 	TArrayView(const TArrayView&) = default;
@@ -162,9 +155,6 @@ public:
 	}
 
 private:
-	template <typename T>
-	using TIsCompatibleElementType = UE::Core::ArrayView::Private::TIsCompatibleElementType<T, ElementType>;
-
 	template <typename T>
 	using TIsCompatibleRangeType = UE::Core::ArrayView::Private::TIsCompatibleRangeType<T, ElementType>;
 
@@ -187,8 +177,9 @@ public:
 					TIsCompatibleRangeType<OtherRangeType>,
 					TIsReinterpretableRangeType<OtherRangeType>
 				>
-			>::Value && 
-			TIsTArrayView_V<CVUnqualifiedOtherRangeType> && !std::is_same_v<CVUnqualifiedOtherRangeType, TArrayView>
+			>::Value &&
+			TIsTArrayView_V<CVUnqualifiedOtherRangeType> &&
+			!std::is_same_v<CVUnqualifiedOtherRangeType, TArrayView>
 		)
 	>
 	FORCEINLINE TArrayView(OtherRangeType&& Other)
@@ -251,7 +242,7 @@ public:
 	 */
 	template <
 		typename OtherElementType
-		UE_REQUIRES(TIsCompatibleElementType<OtherElementType>::Value)
+		UE_REQUIRES(UE::Core::ArrayView::Private::TIsCompatibleElementType_V<OtherElementType, ElementType>)
 	>
 	FORCEINLINE TArrayView(OtherElementType* InData UE_LIFETIMEBOUND, SizeType InCount)
 		: DataPtr(InData)
@@ -915,3 +906,9 @@ FORCEINLINE TArray<InElementType, InAllocatorType>& TArray<InElementType, InAllo
 	CopyToEmpty(Other.GetData(), Other.Num(), ArrayMax);
 	return *this;
 }
+
+#if UE_ENABLE_INCLUDE_ORDER_DEPRECATED_IN_5_4
+#include "Templates/IsConst.h"
+#include "Templates/IsSigned.h"
+#include "Templates/PointerIsConvertibleFromTo.h"
+#endif
\ No newline at end of file
diff --git a/Engine/Source/Runtime/Core/Public/Containers/BinaryHeap.h b/Engine/Source/Runtime/Core/Public/Containers/BinaryHeap.h
--- a/Engine/Source/Runtime/Core/Public/Containers/BinaryHeap.h
+++ b/Engine/Source/Runtime/Core/Public/Containers/BinaryHeap.h
@@ -6,7 +6,8 @@
 #include "CoreTypes.h"
 #include "HAL/UnrealMemory.h"
 #include "Math/UnrealMathUtility.h"
-#include "Templates/IsSigned.h"
+
+#include <type_traits>
 
 /*-----------------------------------------------------------------------------
 	Binary Heap, used to index another data structure.
@@ -18,7 +19,7 @@ template< typename KeyType, typename IndexType = uint32 >
 class FBinaryHeap
 {
 public:
-	static_assert(!TIsSigned<IndexType>::Value, "FBinaryHeap only supports unsigned index types");
+	static_assert(std::is_unsigned_v<IndexType>, "FBinaryHeap only supports unsigned index types");
 
 				FBinaryHeap();
 				FBinaryHeap( uint32 InHeapSize, uint32 InIndexSize );
@@ -400,3 +401,7 @@ FORCEINLINE void FBinaryHeap< KeyType, IndexType >::ResetInternal()
 	Keys = nullptr;
 	HeapIndexes = nullptr;
 }
+
+#if UE_ENABLE_INCLUDE_ORDER_DEPRECATED_IN_5_4
+#include "Templates/IsSigned.h"
+#endif
diff --git a/Engine/Source/Runtime/Core/Public/Containers/EnumAsByte.h b/Engine/Source/Runtime/Core/Public/Containers/EnumAsByte.h
--- a/Engine/Source/Runtime/Core/Public/Containers/EnumAsByte.h
+++ b/Engine/Source/Runtime/Core/Public/Containers/EnumAsByte.h
@@ -4,23 +4,18 @@
 
 #include "CoreTypes.h"
 #include "Templates/IsPODType.h"
-#include "Templates/IsEnumClass.h"
 #include "Templates/TypeHash.h"
 
-template <bool> struct TEnumAsByte_EnumClass;
-template <> struct UE_DEPRECATED(4.15, "TEnumAsByte is not intended for use with enum classes - please derive your enum class from uint8 instead.") TEnumAsByte_EnumClass<true> {};
-template <> struct TEnumAsByte_EnumClass<false> {};
-
 /**
  * Template to store enumeration values as bytes in a type-safe way.
  */
-template<class TEnum>
+template <class InEnumType>
 class TEnumAsByte
 {
-	typedef TEnumAsByte_EnumClass<TIsEnumClass<TEnum>::Value> Check;
+	static_assert(std::is_enum_v<InEnumType> && std::is_convertible_v<InEnumType, int>, "TEnumAsByte is not intended for use with enum classes - please derive your enum class from uint8 instead.");
 
 public:
-	typedef TEnum EnumType;
+	using EnumType = InEnumType;
 
 	TEnumAsByte() = default;
 	TEnumAsByte(const TEnumAsByte&) = default;
@@ -31,7 +26,7 @@ public:
 	 *
 	 * @param InValue value to construct with.
 	 */
-	FORCEINLINE TEnumAsByte( TEnum InValue )
+	FORCEINLINE TEnumAsByte( EnumType InValue )
 		: Value(static_cast<uint8>(InValue))
 	{ }
 
@@ -60,9 +55,9 @@ public:
 	 * @param InValue The value to compare with.
 	 * @return true if the two values are equal, false otherwise.
 	 */
-	bool operator==( TEnum InValue ) const
+	bool operator==( EnumType InValue ) const
 	{
-		return static_cast<TEnum>(Value) == InValue;
+		return static_cast<EnumType>(Value) == InValue;
 	}
 
 	/**
@@ -76,10 +71,10 @@ public:
 		return Value == InValue.Value;
 	}
 
-	/** Implicit conversion to TEnum. */
-	operator TEnum() const
+	/** Implicit conversion to EnumType. */
+	operator EnumType() const
 	{
-		return (TEnum)Value;
+		return (EnumType)Value;
 	}
 
 public:
@@ -89,9 +84,9 @@ public:
 	 *
 	 * @return The enumeration value.
 	 */
-	TEnum GetValue() const
+	EnumType GetValue() const
 	{
-		return (TEnum)Value;
+		return (EnumType)Value;
 	}
 
 	/**
@@ -128,4 +123,8 @@ struct TIsTEnumAsByte
 template <typename T> struct TIsTEnumAsByte<               TEnumAsByte<T>> { static constexpr bool Value = true; };
 template <typename T> struct TIsTEnumAsByte<const          TEnumAsByte<T>> { static constexpr bool Value = true; };
 template <typename T> struct TIsTEnumAsByte<      volatile TEnumAsByte<T>> { static constexpr bool Value = true; };
-template <typename T> struct TIsTEnumAsByte<const volatile TEnumAsByte<T>> { static constexpr bool Value = true; };
\ No newline at end of file
+template <typename T> struct TIsTEnumAsByte<const volatile TEnumAsByte<T>> { static constexpr bool Value = true; };
+
+#if UE_ENABLE_INCLUDE_ORDER_DEPRECATED_IN_5_4
+#include "Templates/IsEnumClass.h"
+#endif
diff --git a/Engine/Source/Runtime/Core/Public/Containers/Map.h b/Engine/Source/Runtime/Core/Public/Containers/Map.h
--- a/Engine/Source/Runtime/Core/Public/Containers/Map.h
+++ b/Engine/Source/Runtime/Core/Public/Containers/Map.h
@@ -27,8 +27,8 @@ template <typename KeyInitType, typename ValueInitType>
 class TPairInitializer
 {
 public:
-	typename TRValueToLValueReference<KeyInitType  >::Type Key;
-	typename TRValueToLValueReference<ValueInitType>::Type Value;
+	std::conditional_t<std::is_rvalue_reference_v<KeyInitType>,   KeyInitType&,   KeyInitType>   Key;
+	std::conditional_t<std::is_rvalue_reference_v<ValueInitType>, ValueInitType&, ValueInitType> Value;
 
 	/** Initialization constructor. */
 	FORCEINLINE TPairInitializer(KeyInitType InKey, ValueInitType InValue)
@@ -58,12 +58,13 @@ template <typename KeyInitType>
 class TKeyInitializer
 {
 public:
-	typename TRValueToLValueReference<KeyInitType>::Type Key;
+	std::conditional_t<std::is_rvalue_reference_v<KeyInitType>, KeyInitType&, KeyInitType> Key;
 
 	/** Initialization constructor. */
 	FORCEINLINE explicit TKeyInitializer(KeyInitType InKey)
 		: Key(InKey)
-	{ }
+	{
+	}
 
 	template <typename KeyType, typename ValueType>
 	operator TPair<KeyType, ValueType>() const
diff --git a/Engine/Source/Runtime/Core/Public/Containers/PagedArray.h b/Engine/Source/Runtime/Core/Public/Containers/PagedArray.h
--- a/Engine/Source/Runtime/Core/Public/Containers/PagedArray.h
+++ b/Engine/Source/Runtime/Core/Public/Containers/PagedArray.h
@@ -390,7 +390,10 @@ public:
 		Append(InList.begin(), static_cast<SizeType>(InList.size()));
 	}
 
-	template <typename ContainerType, std::enable_if_t<TIsContiguousContainer<ContainerType>::Value>* = nullptr>
+	template <
+		typename ContainerType
+		UE_REQUIRES(TIsContiguousContainer<ContainerType>::Value)
+	>
 	FORCEINLINE void Append(ContainerType&& Container)
 	{
 		Append(GetData(Container), GetNum(Container));
@@ -428,7 +431,10 @@ public:
 		Append(InList.begin(), static_cast<SizeType>(InList.size()));
 	}
 
-	template <typename ContainerType, std::enable_if_t<TIsContiguousContainer<ContainerType>::Value>* = nullptr>
+	template <
+		typename ContainerType
+		UE_REQUIRES(TIsContiguousContainer<ContainerType>::Value)
+	>
 	void Assign(ContainerType&& Container)
 	{
 		Reset();
diff --git a/Engine/Source/Runtime/Core/Public/Containers/RingBuffer.h b/Engine/Source/Runtime/Core/Public/Containers/RingBuffer.h
--- a/Engine/Source/Runtime/Core/Public/Containers/RingBuffer.h
+++ b/Engine/Source/Runtime/Core/Public/Containers/RingBuffer.h
@@ -6,12 +6,11 @@
 #include "Misc/AssertionMacros.h"
 #include "Containers/ArrayView.h"
 #include "Templates/IsPODType.h"
-#include "Templates/IsTriviallyDestructible.h"
-#include "Templates/MakeSigned.h"
-#include "Templates/MakeUnsigned.h"
 #include "Templates/MemoryOps.h"
 #include "Templates/UnrealTemplate.h"
 
+#include <type_traits>
+
 template< typename ContainerType, typename ElementType, typename SizeType>
 class TRingBufferIterator
 {
@@ -140,9 +139,9 @@ public:
 	/* The Allocator type being used */
 	typedef AllocatorT Allocator;
 	/** Type used to request values at a given index in the container. */
-	typedef typename TMakeSigned<typename Allocator::SizeType>::Type IndexType;
+	typedef std::make_signed_t<typename Allocator::SizeType> IndexType;
 	/** Type used to communicate size and capacity and counts */
-	typedef typename TMakeUnsigned<typename Allocator::SizeType>::Type SizeType;
+	typedef std::make_unsigned_t<typename Allocator::SizeType> SizeType;
 	/** Iterator type used for ranged-for traversal. */
 	typedef TRingBufferIterator<TRingBuffer, ElementType, typename Allocator::SizeType> TIterator;
 	typedef TRingBufferIterator<const TRingBuffer, const ElementType, typename Allocator::SizeType> TConstIterator;
@@ -153,7 +152,7 @@ private:
 	 * They may have added on multiples of the capacity due to wrapping around, but will be interpreted as pointing to the value at (value & IndexMask).
 	 * StorageModuloTypes are also allowed to underflow/overflow their integer storage type; the only constraint is that X - Front <= Capacity for all valid indexes and for AfterBack.
 	 */
-	typedef typename TMakeUnsigned<typename Allocator::SizeType>::Type StorageModuloType;
+	typedef std::make_unsigned_t<typename Allocator::SizeType> StorageModuloType;
 public:
 
 	/** Construct Empty Queue with capacity 0. */
@@ -771,7 +770,7 @@ public:
 private:
 	static constexpr bool NeedsDestructElements()
 	{
-		return !TIsTriviallyDestructible<T>::Value;
+		return !std::is_trivially_destructible_v<T>;
 	}
 
 	/** Set the capacity to the given value and move or copy all elements from the old storage into a new storage with the given capacity.  Assumes the capacity has already been normalized and is greater than or equal to the number of elements in the RingBuffer. */
@@ -992,3 +991,9 @@ private:
 	 */
 	StorageModuloType AfterBack;
 };
+
+#if UE_ENABLE_INCLUDE_ORDER_DEPRECATED_IN_5_4
+#include "Templates/IsTriviallyDestructible.h"
+#include "Templates/MakeSigned.h"
+#include "Templates/MakeUnsigned.h"
+#endif
diff --git a/Engine/Source/Runtime/Core/Public/Containers/Set.h b/Engine/Source/Runtime/Core/Public/Containers/Set.h
--- a/Engine/Source/Runtime/Core/Public/Containers/Set.h
+++ b/Engine/Source/Runtime/Core/Public/Containers/Set.h
@@ -16,7 +16,6 @@
 #include <initializer_list>
 #include "Templates/TypeHash.h"
 #include "Containers/SparseArray.h"
-#include "Templates/Decay.h"
 #include "Serialization/StructuredArchive.h"
 #include "Serialization/MemoryImageWriter.h"
 #include "ContainersFwd.h"
@@ -165,8 +164,8 @@ public:
 
 	/** Initialization constructor. */
 	template <
-		typename InitType,
-		typename = std::enable_if_t<!std::is_same_v<TSetElementBase, typename TDecay<InitType>::Type>>
+		typename InitType
+		UE_REQUIRES(!std::is_same_v<TSetElementBase, std::decay_t<InitType>>)
 	>
 	explicit FORCEINLINE TSetElementBase(InitType&& InValue)
 		: Value(Forward<InitType>(InValue))
@@ -199,8 +198,8 @@ public:
 
 	/** Initialization constructor. */
 	template <
-		typename InitType,
-		typename = std::enable_if_t<!std::is_same_v<TSetElementBase, typename TDecay<InitType>::Type>>
+		typename InitType
+		UE_REQUIRES(!std::is_same_v<TSetElementBase, std::decay_t<InitType>>)
 	>
 	explicit FORCEINLINE TSetElementBase(InitType&& InValue)
 		: Value(Forward<InitType>(InValue))
@@ -234,8 +233,8 @@ public:
 
 	/** Initialization constructor. */
 	template <
-		typename InitType,
-		typename = std::enable_if_t<!std::is_same_v<TSetElement, typename TDecay<InitType>::Type>>
+		typename InitType
+		UE_REQUIRES(!std::is_same_v<TSetElement, std::decay_t<InitType>>)
 	>
 	explicit FORCEINLINE TSetElement(InitType&& InValue)
 		: Super(Forward<InitType>(InValue))
@@ -1315,8 +1314,8 @@ public:
 	 */
 	template <
 		typename OtherKeyFuncs,
-		typename AliasElementType = ElementType,
-		typename std::enable_if_t<TIsContainerElementTypeCopyable<AliasElementType>::Value>* = nullptr
+		typename AliasElementType = ElementType
+		UE_REQUIRES(TIsContainerElementTypeCopyable<AliasElementType>::Value)
 	>
 	TSet& operator=(TSet<typename TContainerElementTypeCompatibility<ElementType>::CopyFromOtherType, OtherKeyFuncs, Allocator>&& Other)
 	{
@@ -1336,8 +1335,8 @@ public:
 	template <
 		typename OtherKeyFuncs,
 		typename OtherAllocator,
-		typename AliasElementType = ElementType,
-		typename std::enable_if_t<TIsContainerElementTypeCopyable<AliasElementType>::Value>* = nullptr
+		typename AliasElementType = ElementType
+		UE_REQUIRES(TIsContainerElementTypeCopyable<AliasElementType>::Value)
 	>
 	TSet& operator=(const TSet<typename TContainerElementTypeCompatibility<ElementType>::CopyFromOtherType, OtherKeyFuncs, OtherAllocator>& Other)
 	{
@@ -1355,8 +1354,8 @@ public:
 	template <
 		typename OtherKeyFuncs,
 		typename OtherAllocator,
-		typename AliasElementType = ElementType,
-		typename std::enable_if_t<TIsContainerElementTypeCopyable<AliasElementType>::Value>* = nullptr
+		typename AliasElementType = ElementType
+		UE_REQUIRES(TIsContainerElementTypeCopyable<AliasElementType>::Value)
 	>
 	void Append(const TSet<typename TContainerElementTypeCompatibility<ElementType>::CopyFromOtherType, OtherKeyFuncs, OtherAllocator>& OtherSet)
 	{
@@ -1375,8 +1374,8 @@ public:
 	 */
 	template <
 		typename OtherKeyFuncs,
-		typename AliasElementType = ElementType,
-		typename std::enable_if_t<TIsContainerElementTypeCopyable<AliasElementType>::Value>* = nullptr
+		typename AliasElementType = ElementType
+		UE_REQUIRES(TIsContainerElementTypeCopyable<AliasElementType>::Value)
 	>
 	void Append(TSet<typename TContainerElementTypeCompatibility<ElementType>::CopyFromOtherType, OtherKeyFuncs, Allocator>&& OtherSet)
 	{
@@ -1633,7 +1632,7 @@ private:
 	public:
 		using KeyArgumentType =
 			std::conditional_t<
-				std::is_reference<ReferenceOrValueType>::value,
+				std::is_reference_v<ReferenceOrValueType>,
 				TRetainedRef<std::remove_reference_t<ReferenceOrValueType>>,
 				KeyInitType
 			>;
@@ -2275,3 +2274,7 @@ bool LegacyCompareNotEqual(const TSet<ElementType, KeyFuncs, Allocator>& A,const
 {
 	return !TSetPrivateFriend::LegacyCompareEqual(A, B);
 }
+
+#if UE_ENABLE_INCLUDE_ORDER_DEPRECATED_IN_5_4
+#include "Templates/Decay.h"
+#endif
diff --git a/Engine/Source/Runtime/Core/Public/Containers/SparseArray.h b/Engine/Source/Runtime/Core/Public/Containers/SparseArray.h
--- a/Engine/Source/Runtime/Core/Public/Containers/SparseArray.h
+++ b/Engine/Source/Runtime/Core/Public/Containers/SparseArray.h
@@ -5,10 +5,8 @@
 #include "CoreTypes.h"
 #include "Misc/AssertionMacros.h"
 #include "HAL/UnrealMemory.h"
-#include "Templates/IsTriviallyCopyConstructible.h"
 #include "Templates/UnrealTypeTraits.h"
 #include "Templates/UnrealTemplate.h"
-#include "Templates/IsTriviallyDestructible.h"
 #include "Containers/ContainerAllocationPolicies.h"
 #include "Templates/Less.h"
 #include "Containers/Array.h"
@@ -333,7 +331,7 @@ public:
 	/** Removes Count elements from the array, starting from Index. */
 	void RemoveAt(int32 Index,int32 Count = 1)
 	{
-		if constexpr (!TIsTriviallyDestructible<ElementType>::Value)
+		if constexpr (!std::is_trivially_destructible_v<ElementType>)
 		{
 			FElementOrFreeListLink* DataPtr = (FElementOrFreeListLink*)Data.GetData();
 			for (int32 It = Index, ItCount = Count; ItCount; ++It, --ItCount)
@@ -376,9 +374,9 @@ public:
 	void Empty(int32 ExpectedNumElements = 0)
 	{
 		// Destruct the allocated elements.
-		if constexpr ( !TIsTriviallyDestructible<ElementType>::Value )
+		if constexpr (!std::is_trivially_destructible_v<ElementType>)
 		{
-			for(TIterator It(*this);It;++It)
+			for (TIterator It(*this); It; ++It)
 			{
 				ElementType& Element = *It;
 				Element.~ElementType();
@@ -396,9 +394,9 @@ public:
 	void Reset()
 	{
 		// Destruct the allocated elements.
-		if constexpr ( !TIsTriviallyDestructible<ElementType>::Value )
+		if constexpr (!std::is_trivially_destructible_v<ElementType>)
 		{
-			for(TIterator It(*this);It;++It)
+			for (TIterator It(*this); It; ++It)
 			{
 				ElementType& Element = *It;
 				Element.~ElementType();
@@ -774,7 +772,7 @@ public:
 			const FElementOrFreeListLink* SrcData  = (const FElementOrFreeListLink*)InCopy.Data.GetData();
 
 			// Determine whether we need per element construction or bulk copy is fine
-			if constexpr (!TIsTriviallyCopyConstructible<ElementType>::Value)
+			if constexpr (!std::is_trivially_copy_constructible_v<ElementType>)
 			{
 				// Use the inplace new to copy the element to an array element
 				for (int32 Index = 0; Index < SrcMax; ++Index)
@@ -809,7 +807,7 @@ private:
 	FORCEINLINE static void Move(SparseArrayType& ToArray, SparseArrayType& FromArray)
 	{
 		// Destruct the allocated elements.
-		if constexpr ( !TIsTriviallyDestructible<ElementType>::Value )
+		if constexpr (!std::is_trivially_destructible_v<ElementType>)
 		{
 			for (ElementType& Element : ToArray)
 			{
@@ -1534,3 +1532,8 @@ void operator<<(FStructuredArchive::FSlot Slot, TSparseArray<ElementType, Alloca
 		}
 	}
 }
+
+#if UE_ENABLE_INCLUDE_ORDER_DEPRECATED_IN_5_4
+#include "Templates/IsTriviallyCopyConstructible.h"
+#include "Templates/IsTriviallyDestructible.h"
+#endif
diff --git a/Engine/Source/Runtime/Core/Public/Containers/StringConv.h b/Engine/Source/Runtime/Core/Public/Containers/StringConv.h
--- a/Engine/Source/Runtime/Core/Public/Containers/StringConv.h
+++ b/Engine/Source/Runtime/Core/Public/Containers/StringConv.h
@@ -10,10 +10,7 @@
 #include "Containers/ContainerAllocationPolicies.h"
 #include "Containers/Array.h"
 #include "Misc/CString.h"
-#include "Templates/AndOrNot.h"
-#include "Templates/EnableIf.h"
 #include "Templates/IsArray.h"
-#include "Templates/RemoveReference.h"
 #include "Templates/UnrealTemplate.h"
 #include "Templates/UnrealTypeTraits.h"
 #include "Traits/ElementType.h"
@@ -31,8 +28,8 @@ public:
 	typedef To   ToType;
 
 	template <
-		typename CharType,
-		std::enable_if_t<TIsCharEncodingCompatibleWith_V<CharType, FromType>>* = nullptr
+		typename CharType
+		UE_REQUIRES(TIsCharEncodingCompatibleWith_V<CharType, FromType>)
 	>
 	FORCEINLINE static void Convert(To* Dest, int32 DestLen, const CharType* Source, int32 SourceLen)
 	{
@@ -41,8 +38,8 @@ public:
 	}
 
 	template <
-		typename CharType,
-		std::enable_if_t<TIsCharEncodingCompatibleWith_V<CharType, FromType>>* = nullptr
+		typename CharType
+		UE_REQUIRES(TIsCharEncodingCompatibleWith_V<CharType, FromType>)
 	>
 	static int32 ConvertedLength(const CharType* Source, int32 SourceLen)
 	{
@@ -291,8 +288,8 @@ namespace UE::Core::Private
 		 * @return          The number of bytes written to Dest, up to DestLen, or -1 if the entire Source string could did not fit in DestLen bytes.
 		 */
 		template <
-			typename SrcBufferType,
-			std::enable_if_t<TIsCharEncodingCompatibleWith_V<SrcBufferType, FromType>>* = nullptr
+			typename SrcBufferType
+			UE_REQUIRES(TIsCharEncodingCompatibleWith_V<SrcBufferType, FromType>)
 		>
 		static FORCEINLINE int32 Convert(IntendedToType* Dest, int32 DestLen, const SrcBufferType* Source, int32 SourceLen)
 		{
@@ -305,8 +302,8 @@ namespace UE::Core::Private
 			return (int32)(Result - Dest);
 		}
 		template <
-			typename SrcBufferType,
-			std::enable_if_t<TIsCharEncodingCompatibleWith_V<SrcBufferType, FromType>>* = nullptr
+			typename SrcBufferType
+			UE_REQUIRES(TIsCharEncodingCompatibleWith_V<SrcBufferType, FromType>)
 		>
 		static FORCEINLINE int32 Convert(ToType* Dest, int32 DestLen, const SrcBufferType* Source, int32 SourceLen)
 		{
@@ -320,8 +317,8 @@ namespace UE::Core::Private
 		 * @return The length of the string in UTF-8 code units.
 		 */
 		template <
-			typename SrcBufferType,
-			std::enable_if_t<TIsCharEncodingCompatibleWith_V<SrcBufferType, FromType>>* = nullptr
+			typename SrcBufferType
+			UE_REQUIRES(TIsCharEncodingCompatibleWith_V<SrcBufferType, FromType>)
 		>
 		static FORCEINLINE int32 ConvertedLength(const SrcBufferType* Source, int32 SourceLen)
 		{
@@ -354,8 +351,8 @@ public:
 	 * @param SourceLen The length of the source string.
 	 */
 	template <
-		typename SrcBufferType,
-		std::enable_if_t<TIsCharEncodingCompatibleWith_V<SrcBufferType, FromType>>* = nullptr
+		typename SrcBufferType
+		UE_REQUIRES(TIsCharEncodingCompatibleWith_V<SrcBufferType, FromType>)
 	>
 	static FORCEINLINE void Convert(ToType* Dest, const int32 DestLen, const SrcBufferType* Source, const int32 SourceLen)
 	{
@@ -374,8 +371,8 @@ public:
 	 * @return The length of the string in UTF-16 or UTF-32 characters.
 	 */
 	template <
-		typename SrcBufferType,
-		std::enable_if_t<TIsCharEncodingCompatibleWith_V<SrcBufferType, FromType>>* = nullptr
+		typename SrcBufferType
+		UE_REQUIRES(TIsCharEncodingCompatibleWith_V<SrcBufferType, FromType>)
 	>
 	static int32 ConvertedLength(const SrcBufferType* Source, const int32 SourceLen)
 	{
@@ -462,8 +459,8 @@ public:
 	 * @param SourceLen The length of the source string.
 	 */
 	template <
-		typename SrcBufferType,
-		std::enable_if_t<TIsCharEncodingCompatibleWith_V<SrcBufferType, FromType>>* = nullptr
+		typename SrcBufferType
+		UE_REQUIRES(TIsCharEncodingCompatibleWith_V<SrcBufferType, FromType>)
 	>
 	static FORCEINLINE void Convert(ToType* Dest, int32 DestLen, const SrcBufferType* Source, int32 SourceLen)
 	{
@@ -476,8 +473,8 @@ public:
 	 * @return The length of the string in UTF-16 code units.
 	 */
 	template <
-		typename SrcBufferType,
-		std::enable_if_t<TIsCharEncodingCompatibleWith_V<SrcBufferType, FromType>>* = nullptr
+		typename SrcBufferType
+		UE_REQUIRES(TIsCharEncodingCompatibleWith_V<SrcBufferType, FromType>)
 	>
 	static FORCEINLINE int32 ConvertedLength(const SrcBufferType* Source, int32 SourceLen)
 	{
@@ -538,8 +535,8 @@ public:
 	 * @param SourceLen The length of the source string.
 	 */
 	template <
-		typename SrcBufferType,
-		std::enable_if_t<TIsCharEncodingCompatibleWith_V<SrcBufferType, FromType>>* = nullptr
+		typename SrcBufferType
+		UE_REQUIRES(TIsCharEncodingCompatibleWith_V<SrcBufferType, FromType>)
 	>
 	static FORCEINLINE void Convert(ToType* Dest, const int32 DestLen, const SrcBufferType* Source, const int32 SourceLen)
 	{
@@ -554,8 +551,8 @@ public:
 	 * @return The length of the string in UTF-32 characters.
 	 */
 	template <
-		typename SrcBufferType,
-		std::enable_if_t<TIsCharEncodingCompatibleWith_V<SrcBufferType, FromType>>* = nullptr
+		typename SrcBufferType
+		UE_REQUIRES(TIsCharEncodingCompatibleWith_V<SrcBufferType, FromType>)
 	>
 	static int32 ConvertedLength(const SrcBufferType* Source, const int32 SourceLen)
 	{
@@ -678,8 +675,8 @@ public:
 	TStringConversion& operator=(const TStringConversion&) = delete;
 
 	template <
-		typename SrcBufferType,
-		std::enable_if_t<TIsCharEncodingCompatibleWith_V<SrcBufferType, FromType>>* = nullptr
+		typename SrcBufferType
+		UE_REQUIRES(TIsCharEncodingCompatibleWith_V<SrcBufferType, FromType>)
 	>
 	explicit TStringConversion(const SrcBufferType* Source)
 	{
@@ -699,8 +696,8 @@ public:
 	}
 
 	template <
-		typename SrcBufferType,
-		std::enable_if_t<TIsCharEncodingCompatibleWith_V<SrcBufferType, FromType>>* = nullptr
+		typename SrcBufferType
+		UE_REQUIRES(TIsCharEncodingCompatibleWith_V<SrcBufferType, FromType>)
 	>
 	TStringConversion(const SrcBufferType* Source, int32 SourceLen)
 	{
@@ -732,14 +729,12 @@ public:
 	 */
 	template <
 		typename FromRangeType,
-		typename FromRangeCharType = std::remove_cv_t<std::remove_pointer_t<decltype(GetData(DeclVal<FromRangeType>()))>>,
-		std::enable_if_t<
-			TAnd<
-				TIsContiguousContainer<FromRangeType>,
-				TNot<TIsArray<typename TRemoveReference<FromRangeType>::Type>>,
-				TIsCharEncodingCompatibleWith<FromRangeCharType, FromType>
-			>::Value
-		>* = nullptr
+		typename FromRangeCharType = std::remove_cv_t<std::remove_pointer_t<decltype(GetData(std::declval<FromRangeType>()))>>
+		UE_REQUIRES(
+			TIsContiguousContainer<FromRangeType>::Value &&
+			!TIsArray<std::remove_reference_t<FromRangeType>>::Value &&
+			TIsCharEncodingCompatibleWith<FromRangeCharType, FromType>::Value
+		)
 	>
 	TStringConversion(FromRangeType&& Source)
 		: TStringConversion((const FromType*)GetData(Source), GetNum(Source))
@@ -812,8 +807,8 @@ public:
 
 public:
 	template <
-		typename SrcBufferType,
-		std::enable_if_t<TIsCharEncodingCompatibleWith_V<SrcBufferType, FromType>>* = nullptr
+		typename SrcBufferType
+		UE_REQUIRES(TIsCharEncodingCompatibleWith_V<SrcBufferType, FromType>)
 	>
 	explicit TStringPointer(const SrcBufferType* Source)
 	{
@@ -830,8 +825,8 @@ public:
 	}
 
 	template <
-		typename SrcBufferType,
-		std::enable_if_t<TIsCharEncodingCompatibleWith_V<SrcBufferType, FromType>>* = nullptr
+		typename SrcBufferType
+		UE_REQUIRES(TIsCharEncodingCompatibleWith_V<SrcBufferType, FromType>)
 	>
 	TStringPointer(const SrcBufferType* Source, int32 SourceLen)
 	{
@@ -858,14 +853,12 @@ public:
 	 */
 	template <
 		typename FromRangeType,
-		typename FromRangeCharType = std::remove_cv_t<std::remove_pointer_t<decltype(GetData(DeclVal<FromRangeType>()))>>,
-		std::enable_if_t<
-			TAnd<
-				TIsContiguousContainer<FromRangeType>,
-				TNot<TIsArray<typename TRemoveReference<FromRangeType>::Type>>,
-				TIsCharEncodingCompatibleWith<FromRangeCharType, FromType>
-			>::Value
-		>* = nullptr
+		typename FromRangeCharType = std::remove_cv_t<std::remove_pointer_t<decltype(GetData(std::declval<FromRangeType>()))>>
+		UE_REQUIRES(
+			TIsContiguousContainer<FromRangeType>::Value &&
+			!TIsArray<std::remove_reference_t<FromRangeType>>::Value &&
+			TIsCharEncodingCompatibleWith<FromRangeCharType, FromType>::Value
+		)
 	>
 	TStringPointer(FromRangeType&& Source)
 		: TStringPointer((const FromType*)GetData(Source), GetNum(Source))
@@ -1295,3 +1288,9 @@ FORCEINLINE TArray<ToType> StringToArray(const FromType* Str)
 {
 	return StringToArray<ToType>(Str, TCString<FromType>::Strlen(Str) + 1);
 }
+
+#if UE_ENABLE_INCLUDE_ORDER_DEPRECATED_IN_5_4
+#include "Templates/AndOrNot.h"
+#include "Templates/EnableIf.h"
+#include "Templates/RemoveReference.h"
+#endif
diff --git a/Engine/Source/Runtime/Core/Public/Containers/StringView.h b/Engine/Source/Runtime/Core/Public/Containers/StringView.h
--- a/Engine/Source/Runtime/Core/Public/Containers/StringView.h
+++ b/Engine/Source/Runtime/Core/Public/Containers/StringView.h
@@ -101,13 +101,13 @@ public:
 
 	/** Construct a view of the null-terminated string pointed to by InData. */
 	template <
-		typename OtherCharType,
-		std::enable_if_t<
+		typename OtherCharType
+		UE_REQUIRES(
 			TAnd<
 				TIsCharType<OtherCharType>,
 				TIsCharEncodingCompatibleWith<OtherCharType, CharType>
 			>::Value
-		>* = nullptr
+		)
 	>
 	constexpr inline TStringView(const OtherCharType* InData UE_LIFETIMEBOUND)
 		: DataPtr((const CharType*)InData)
@@ -117,13 +117,13 @@ public:
 
 	/** Construct a view of InSize characters beginning at InData. */
 	template <
-		typename OtherCharType,
-		std::enable_if_t<
+		typename OtherCharType
+		UE_REQUIRES(
 			TAnd<
 				TIsCharType<OtherCharType>,
 				TIsCharEncodingCompatibleWith<OtherCharType, CharType>
 			>::Value
-		>* = nullptr
+		)
 	>
 	constexpr inline TStringView(const OtherCharType* InData UE_LIFETIMEBOUND, int32 InSize)
 		: DataPtr((const CharType*)InData)
@@ -132,8 +132,9 @@ public:
 	}
 
 	/** Construct a view from a contiguous range of characters, such as FString or TStringBuilder. */
-	template <typename CharRangeType,
-		std::enable_if_t<
+	template <
+		typename CharRangeType
+		UE_REQUIRES(
 			TAnd<
 				TIsContiguousContainer<CharRangeType>,
 				TIsCharType<TElementType_T<CharRangeType>>,
@@ -141,7 +142,8 @@ public:
 			>::Value &&
 			!std::is_array_v<std::remove_reference_t<CharRangeType>> &&
 			!std::is_same_v<CharRangeType, ViewType>
-		>* = nullptr>
+		)
+	>
 	constexpr inline TStringView(const CharRangeType& InRange UE_LIFETIMEBOUND)
 		: DataPtr((const CharType*)UE::Core::Private::StringViewGetData(InRange))
 		, Size(IntCastChecked<int32>(GetNum(InRange)))
@@ -242,9 +244,14 @@ public:
 	 * @param Other        A string that is comparable with the character type of this view.
 	 * @param SearchCase Whether the comparison should ignore case.
 	 */
-	template <typename OtherCharType,
-		std::enable_if_t<TIsCharType<OtherCharType>::Value>* = nullptr>
-	[[nodiscard]] inline bool Equals(const OtherCharType* Other, ESearchCase::Type SearchCase = ESearchCase::CaseSensitive) const;
+	template <
+		typename OtherCharType
+		UE_REQUIRES(TIsCharType_V<OtherCharType>)
+	>
+	[[nodiscard]] inline bool Equals(const OtherCharType* Other, ESearchCase::Type SearchCase = ESearchCase::CaseSensitive) const
+	{
+		return PrivateEquals(Other, SearchCase);
+	}
 
 	/**
 	 * Compare this view with a character range.
@@ -266,9 +273,14 @@ public:
 	 * @param SearchCase   Whether the comparison should ignore case.
 	 * @return 0 is equal, negative if this view is less, positive if this view is greater.
 	 */
-	template <typename OtherCharType,
-		std::enable_if_t<TIsCharType<OtherCharType>::Value>* = nullptr>
-		[[nodiscard]] inline int32 Compare(TStringView<OtherCharType> Other, ESearchCase::Type SearchCase = ESearchCase::CaseSensitive) const;
+	template <
+		typename OtherCharType
+		UE_REQUIRES(TIsCharType_V<OtherCharType>)
+	>
+	[[nodiscard]] inline int32 Compare(TStringView<OtherCharType> Other, ESearchCase::Type SearchCase = ESearchCase::CaseSensitive) const
+	{
+		return PrivateCompare(Other, SearchCase);
+	}
 
 	/**
 	 * Compare this view with a null-terminated string.
@@ -347,6 +359,12 @@ private:
 	static bool PrivateEquals(TStringView Lhs, TStringView Rhs);
 	static bool PrivateLess(TStringView Lhs, TStringView Rhs);
 
+	template <typename OtherCharType>
+	inline bool PrivateEquals(const OtherCharType* Other, ESearchCase::Type SearchCase) const;
+
+	template <typename OtherCharType>
+	inline int32 PrivateCompare(TStringView<OtherCharType> Other, ESearchCase::Type SearchCase) const;
+
 public:
 	friend constexpr inline auto GetNum(TStringView String)
 	{
@@ -605,9 +623,8 @@ inline TStringView<CharType> TStringView<CharType>::TrimEnd() const
 }
 
 template <typename CharType>
-template <typename OtherCharType,
-	std::enable_if_t<TIsCharType<OtherCharType>::Value>*>
-inline bool TStringView<CharType>::Equals(const OtherCharType* const Other, ESearchCase::Type SearchCase) const
+template <typename OtherCharType>
+inline bool TStringView<CharType>::PrivateEquals(const OtherCharType* const Other, ESearchCase::Type SearchCase) const
 {
 	if (IsEmpty())
 	{
@@ -625,9 +642,8 @@ inline bool TStringView<CharType>::Equals(const OtherCharType* const Other, ESea
 }
 
 template <typename CharType>
-template <typename OtherCharType,
-	std::enable_if_t<TIsCharType<OtherCharType>::Value>*>
-inline int32 TStringView<CharType>::Compare(TStringView<OtherCharType> OtherView, ESearchCase::Type SearchCase) const
+template <typename OtherCharType>
+inline int32 TStringView<CharType>::PrivateCompare(TStringView<OtherCharType> OtherView, ESearchCase::Type SearchCase) const
 {
 	const int32 SelfLen = Len();
 	const int32 OtherLen = OtherView.Len();
diff --git a/Engine/Source/Runtime/Core/Public/Containers/UnrealString.h b/Engine/Source/Runtime/Core/Public/Containers/UnrealString.h
--- a/Engine/Source/Runtime/Core/Public/Containers/UnrealString.h
+++ b/Engine/Source/Runtime/Core/Public/Containers/UnrealString.h
@@ -133,17 +133,23 @@ UE_NODISCARD inline const uint8 TCharToNibble(const TCHAR Hex)
 }
 
 /** Convert numeric types to a string */
-template <typename StringType = FString, typename T>
-UE_NODISCARD typename TEnableIf<TIsArithmetic<T>::Value, StringType>::Type
-LexToString(const T& Value)
+template <
+	typename StringType = FString,
+	typename T
+	UE_REQUIRES(std::is_arithmetic_v<T>)
+>
+UE_NODISCARD StringType LexToString(const T& Value)
 {
 	// std::remove_cv_t to remove potential volatile decorations. Removing const is pointless, but harmless because it's specified in the param declaration.
 	return StringType::Printf(TFormatSpecifier<std::remove_cv_t<T>>::GetFormatSpecifier(), Value);
 }
 
-template <typename StringType = FString, typename CharType>
-UE_NODISCARD typename TEnableIf<TIsCharType<CharType>::Value, StringType>::Type
-LexToString(const CharType* Ptr)
+template <
+	typename StringType = FString,
+	typename CharType
+	UE_REQUIRES(TIsCharType_V<CharType>)
+>
+UE_NODISCARD StringType LexToString(const CharType* Ptr)
 {
 	return StringType(Ptr);
 }
@@ -195,9 +201,12 @@ struct TTypeToString
 
 /** Parse a string into this type, returning whether it was successful */
 /** Specialization for arithmetic types */
-template <typename T, typename CharType>
-typename TEnableIf<TIsArithmetic<T>::Value, bool>::Type
-LexTryParseString(T& OutValue, const CharType* Buffer)
+template <
+	typename T,
+	typename CharType
+	UE_REQUIRES(std::is_arithmetic_v<T>)
+>
+bool LexTryParseString(T& OutValue, const CharType* Buffer)
 {
 	if (Buffer[0] == '\0')
 	{
@@ -215,7 +224,7 @@ LexTryParseString(T& OutValue, const CharType* Buffer)
 			C = *(++Buffer);
 		}
 
-		while (C != CHARTEXT(CharType, '\0') && !TChar<CharType>::IsWhitespace(C) && (TIsFloatingPoint<T>::Value || C != CHARTEXT(CharType, '.')))
+		while (C != CHARTEXT(CharType, '\0') && !TChar<CharType>::IsWhitespace(C) && (std::is_floating_point_v<T> || C != CHARTEXT(CharType, '.')))
 		{
 			bSawZero = bSawZero || (C == CHARTEXT(CharType, '0'));
 			if (!bSawZero && C != CHARTEXT(CharType, '.'))
diff --git a/Engine/Source/Runtime/Core/Public/Containers/UnrealString.h.inl b/Engine/Source/Runtime/Core/Public/Containers/UnrealString.h.inl
--- a/Engine/Source/Runtime/Core/Public/Containers/UnrealString.h.inl
+++ b/Engine/Source/Runtime/Core/Public/Containers/UnrealString.h.inl
@@ -52,46 +52,6 @@ private:
 	typedef TArray<ElementType, AllocatorType> DataType;
 	DataType Data;
 
-	template <typename RangeType>
-	using TRangeElementType = std::remove_cv_t<typename TRemovePointer<decltype(GetData(DeclVal<RangeType>()))>::Type>;
-
-	template <typename CharRangeType>
-	struct TIsRangeOfCharType : TIsCharType<TRangeElementType<CharRangeType>>
-	{
-	};
-
-	template <typename CharRangeType>
-	struct TIsRangeOfElementType
-	{
-		enum { Value = std::is_same_v<ElementType, TRangeElementType<CharRangeType>> };
-	};
-
-	/** Trait testing whether a type is a contiguous range of characters, and not CharType[]. */
-	template <typename CharRangeType>
-	using TIsCharRangeNotCArray = TAnd<
-		TIsContiguousContainer<CharRangeType>,
-		TNot<TIsArray<typename TRemoveReference<CharRangeType>::Type>>,
-		TIsRangeOfCharType<CharRangeType>>;
-
-	/** Trait testing whether a type is a contiguous range of characters, and not CharType[] and not a string class type. */
-	template <typename CharRangeType>
-	using TIsCharRangeNotCArrayNotStringClass = TAnd<
-		TIsCharRangeNotCArray<CharRangeType>,
-		TNot<TIsDerivedFrom<typename TDecay<CharRangeType>::Type, UE_STRING_CLASS>>>;
-
-	/** Trait testing whether a type is a contiguous range of ElementType, and not CharType[]. */
-	template <typename CharRangeType>
-	using TIsTCharRangeNotCArray = TAnd<
-		TIsContiguousContainer<CharRangeType>,
-		TNot<TIsArray<typename TRemoveReference<CharRangeType>::Type>>,
-		TIsRangeOfElementType<CharRangeType>>;
-
-	/** Trait testing whether a type is a contiguous range of ElementType, and not CharType[] and not a string class type. */
-	template <typename CharRangeType>
-	using TIsTCharRangeNotCArrayNotStringClass = TAnd<
-		TIsTCharRangeNotCArray<CharRangeType>,
-		TNot<TIsDerivedFrom<typename TDecay<CharRangeType>::Type, UE_STRING_CLASS>>>;
-
 	/** Like the TIsCharEncodingCompatibleWithTCHAR trait, but for the element type of the string */
 	template <typename SrcEncoding>
 	using TIsCharEncodingCompatibleWithElementType = TIsCharEncodingCompatibleWith<SrcEncoding, ElementType>;
@@ -144,13 +104,31 @@ public:
 	CORE_API UE_STRING_CLASS(const UCS2CHAR* Str, int32 ExtraSlack);
 
 	/** Construct from contiguous range of characters such as a string view or string builder */
-	template <typename CharRangeType, typename TEnableIf<TIsCharRangeNotCArrayNotStringClass<CharRangeType>::Value>::Type* = nullptr>
+	template <
+		typename CharRangeType,
+		typename CharRangeElementType = TElementType_T<CharRangeType>
+		UE_REQUIRES(
+			TIsContiguousContainer<CharRangeType>::Value &&
+			!std::is_array_v<std::remove_reference_t<CharRangeType>> &&
+			TIsCharType_V<CharRangeElementType> &&
+			!std::is_base_of_v<UE_STRING_CLASS, std::decay_t<CharRangeType>>
+		)
+	>
 	FORCEINLINE explicit UE_STRING_CLASS(CharRangeType&& Str) : UE_STRING_CLASS(GetNum(Str), GetData(Forward<CharRangeType>(Str)))
 	{
 	}
 
 	/** Construct from contiguous range of characters with extra slack on top of original string length */
-	template <typename CharRangeType, typename TEnableIf<TIsCharRangeNotCArrayNotStringClass<CharRangeType>::Value>::Type* = nullptr>
+	template <
+		typename CharRangeType,
+		typename CharRangeElementType = TElementType_T<CharRangeType>
+		UE_REQUIRES(
+			TIsContiguousContainer<CharRangeType>::Value &&
+			!std::is_array_v<std::remove_reference_t<CharRangeType>>&&
+			TIsCharType_V<CharRangeElementType> &&
+			!std::is_base_of_v<UE_STRING_CLASS, std::decay_t<CharRangeType>>
+		)
+	>
 	explicit UE_STRING_CLASS(CharRangeType&& Str, int32 ExtraSlack)
 	{
 		uint32 InLen = GetNum(Str);
@@ -189,7 +167,15 @@ public:
 
 	CORE_API UE_STRING_CLASS& operator=(const ElementType* Str);
 
-	template <typename CharRangeType, typename TEnableIf<TIsTCharRangeNotCArray<CharRangeType>::Value>::Type* = nullptr>
+	template <
+		typename CharRangeType,
+		typename CharRangeElementType = TElementType_T<CharRangeType>
+		UE_REQUIRES(
+			TIsContiguousContainer<CharRangeType>::Value &&
+			!std::is_array_v<std::remove_reference_t<CharRangeType>> &&
+			std::is_same_v<ElementType, CharRangeElementType>
+		)
+	>
 	FORCEINLINE UE_STRING_CLASS& operator=(CharRangeType&& Range)
 	{
 		AssignRange(GetData(Range), GetNum(Range));
@@ -379,7 +365,11 @@ public:
 	}
 
 	/** Append a string and return a reference to this */
-	template <typename CharRangeType, typename TEnableIf<TIsCharRangeNotCArray<CharRangeType>::Value>::Type* = nullptr>
+	template <
+		typename CharRangeType,
+		typename CharRangeElementType = TElementType_T<CharRangeType>
+		UE_REQUIRES(TIsContiguousContainer<CharRangeType>::Value && TIsCharType_V<CharRangeElementType>)
+	>
 	FORCEINLINE UE_STRING_CLASS& Append(CharRangeType&& Str)
 	{
 		AppendChars(GetData(Str), GetNum(Str));
@@ -398,8 +388,8 @@ public:
 
 	/** Append a single character and return a reference to this */
 	template <
-		typename AppendedCharType,
-		std::enable_if_t<TIsCharType<AppendedCharType>::Value>* = nullptr
+		typename AppendedCharType
+		UE_REQUIRES(TIsCharType_V<AppendedCharType>)
 	>
 	FORCEINLINE UE_STRING_CLASS& operator+=(AppendedCharType Char)
 	{
@@ -432,10 +422,19 @@ public:
 	 * @param InPrefix the prefix to search for at the start of the string to remove.
 	 * @return true if the prefix was removed, otherwise false.
 	 */
-	template <typename TCharRangeType, std::enable_if_t<TIsCharRangeNotCArrayNotStringClass<TCharRangeType>::Value>* = nullptr>
-	bool RemoveFromStart(TCharRangeType&& InPrefix, ESearchCase::Type SearchCase = ESearchCase::IgnoreCase)
+	template <
+		typename CharRangeType,
+		typename CharRangeElementType = TElementType_T<CharRangeType>
+		UE_REQUIRES(
+			TIsContiguousContainer<CharRangeType>::Value &&
+			!std::is_array_v<std::remove_reference_t<CharRangeType>>&&
+			TIsCharType_V<CharRangeElementType> &&
+			!std::is_base_of_v<UE_STRING_CLASS, std::decay_t<CharRangeType>>
+		)
+	>
+	bool RemoveFromStart(CharRangeType&& InPrefix, ESearchCase::Type SearchCase = ESearchCase::IgnoreCase)
 	{
-		static_assert(std::is_same_v<typename TElementType<TCharRangeType>::Type, ElementType>, "Expected a range of ElementType");
+		static_assert(std::is_same_v<CharRangeElementType, ElementType>, "Expected a range of ElementType");
 		return RemoveFromStart(GetData(InPrefix), GetNum(InPrefix), SearchCase);
 	}
 
@@ -476,10 +475,19 @@ public:
 	 * @param InSuffix the suffix to search for at the end of the string to remove.
 	 * @return true if the suffix was removed, otherwise false.
 	 */
-	template <typename TCharRangeType, std::enable_if_t<TIsCharRangeNotCArrayNotStringClass<TCharRangeType>::Value>* = nullptr>
-	bool RemoveFromEnd(TCharRangeType&& InSuffix, ESearchCase::Type SearchCase = ESearchCase::IgnoreCase)
+	template <
+		typename CharRangeType,
+		typename CharRangeElementType = TElementType_T<CharRangeType>
+		UE_REQUIRES(
+			TIsContiguousContainer<CharRangeType>::Value &&
+			!std::is_array_v<std::remove_reference_t<CharRangeType>>&&
+			TIsCharType_V<CharRangeElementType> &&
+			!std::is_base_of_v<UE_STRING_CLASS, std::decay_t<CharRangeType>>
+		)
+	>
+	bool RemoveFromEnd(CharRangeType&& InSuffix, ESearchCase::Type SearchCase = ESearchCase::IgnoreCase)
 	{
-		static_assert(std::is_same_v<typename TElementType<TCharRangeType>::Type, ElementType>, "Expected a range of ElementType");
+		static_assert(std::is_same_v<CharRangeElementType, ElementType>, "Expected a range of ElementType");
 		return RemoveFromEnd(GetData(InSuffix), GetNum(InSuffix), SearchCase);
 	}
 
@@ -530,8 +538,11 @@ public:
 	 *
 	 * @return The concatenated string.
 	 */
-	template <typename CharType>
-	UE_NODISCARD FORCEINLINE friend typename TEnableIf<TIsCharType<CharType>::Value, UE_STRING_CLASS>::Type operator+(const UE_STRING_CLASS& Lhs, CharType Rhs)
+	template <
+		typename CharType
+		UE_REQUIRES(TIsCharType_V<CharType>)
+	>
+	UE_NODISCARD FORCEINLINE friend UE_STRING_CLASS operator+(const UE_STRING_CLASS& Lhs, CharType Rhs)
 	{
 		Lhs.CheckInvariants();
 
@@ -549,8 +560,11 @@ public:
 	 *
 	 * @return The concatenated string.
 	 */
-	template <typename CharType>
-	UE_NODISCARD FORCEINLINE friend typename TEnableIf<TIsCharType<CharType>::Value, UE_STRING_CLASS>::Type operator+(UE_STRING_CLASS&& Lhs, CharType Rhs)
+	template <
+		typename CharType
+		UE_REQUIRES(TIsCharType_V<CharType>)
+	>
+	UE_NODISCARD FORCEINLINE friend UE_STRING_CLASS operator+(UE_STRING_CLASS&& Lhs, CharType Rhs)
 	{
 		Lhs.CheckInvariants();
 
@@ -584,14 +598,62 @@ public:
 	UE_NODISCARD FORCEINLINE friend UE_STRING_CLASS operator+(const UE_STRING_CLASS& Lhs, const ElementType* Rhs)		{ return ConcatFC(Lhs, Rhs); }
 	UE_NODISCARD FORCEINLINE friend UE_STRING_CLASS operator+(UE_STRING_CLASS&& Lhs, const ElementType* Rhs)			{ return ConcatFC(MoveTemp(Lhs), Rhs); }
 
-	template <typename T, typename TEnableIf<TIsTCharRangeNotCArrayNotStringClass<T>::Value>::Type* = nullptr>
-	UE_NODISCARD FORCEINLINE friend UE_STRING_CLASS operator+(T&& Lhs, const UE_STRING_CLASS& Rhs)						{ return ConcatRF(GetData(Lhs), GetNum(Lhs), Rhs); }
-	template <typename T, typename TEnableIf<TIsTCharRangeNotCArrayNotStringClass<T>::Value>::Type* = nullptr>
-	UE_NODISCARD FORCEINLINE friend UE_STRING_CLASS operator+(T&& Lhs, UE_STRING_CLASS&& Rhs)							{ return ConcatRF(GetData(Lhs), GetNum(Lhs), MoveTemp(Rhs)); }
-	template <typename T, typename TEnableIf<TIsTCharRangeNotCArrayNotStringClass<T>::Value>::Type* = nullptr>
-	UE_NODISCARD FORCEINLINE friend UE_STRING_CLASS operator+(const UE_STRING_CLASS& Lhs, T&& Rhs)						{ return ConcatFR(Lhs, GetData(Rhs), GetNum(Rhs)); }
-	template <typename T, typename TEnableIf<TIsTCharRangeNotCArrayNotStringClass<T>::Value>::Type* = nullptr>
-	UE_NODISCARD FORCEINLINE friend UE_STRING_CLASS operator+(UE_STRING_CLASS&& Lhs, T&& Rhs)							{ return ConcatFR(MoveTemp(Lhs), GetData(Rhs), GetNum(Rhs)); }
+	template <
+		typename CharRangeType,
+		typename CharRangeElementType = TElementType_T<CharRangeType>
+		UE_REQUIRES(
+			TIsContiguousContainer<CharRangeType>::Value &&
+			!std::is_array_v<std::remove_reference_t<CharRangeType>> &&
+			std::is_same_v<ElementType, CharRangeElementType> &&
+			!std::is_base_of_v<UE_STRING_CLASS, std::decay_t<CharRangeType>>
+		)
+	>
+	UE_NODISCARD FORCEINLINE friend UE_STRING_CLASS operator+(CharRangeType&& Lhs, const UE_STRING_CLASS& Rhs)
+	{
+		return ConcatRF(GetData(Lhs), GetNum(Lhs), Rhs);
+	}
+	template <
+		typename CharRangeType,
+		typename CharRangeElementType = TElementType_T<CharRangeType>
+		UE_REQUIRES(
+			TIsContiguousContainer<CharRangeType>::Value &&
+			!std::is_array_v<std::remove_reference_t<CharRangeType>>&&
+			std::is_same_v<ElementType, CharRangeElementType> &&
+			!std::is_base_of_v<UE_STRING_CLASS, std::decay_t<CharRangeType>>
+		)
+	>
+	UE_NODISCARD FORCEINLINE friend UE_STRING_CLASS operator+(CharRangeType&& Lhs, UE_STRING_CLASS&& Rhs)
+	{
+		return ConcatRF(GetData(Lhs), GetNum(Lhs), MoveTemp(Rhs));
+	}
+	template <
+		typename CharRangeType,
+		typename CharRangeElementType = TElementType_T<CharRangeType>
+		UE_REQUIRES(
+			TIsContiguousContainer<CharRangeType>::Value &&
+			!std::is_array_v<std::remove_reference_t<CharRangeType>>&&
+			std::is_same_v<ElementType, CharRangeElementType> &&
+			!std::is_base_of_v<UE_STRING_CLASS, std::decay_t<CharRangeType>>
+		)
+	>
+	UE_NODISCARD FORCEINLINE friend UE_STRING_CLASS operator+(const UE_STRING_CLASS& Lhs, CharRangeType&& Rhs)
+	{
+		return ConcatFR(Lhs, GetData(Rhs), GetNum(Rhs));
+	}
+	template <
+		typename CharRangeType,
+		typename CharRangeElementType = TElementType_T<CharRangeType>
+		UE_REQUIRES(
+			TIsContiguousContainer<CharRangeType>::Value &&
+			!std::is_array_v<std::remove_reference_t<CharRangeType>>&&
+			std::is_same_v<ElementType, CharRangeElementType> &&
+			!std::is_base_of_v<UE_STRING_CLASS, std::decay_t<CharRangeType>>
+		)
+	>
+	UE_NODISCARD FORCEINLINE friend UE_STRING_CLASS operator+(UE_STRING_CLASS&& Lhs, CharRangeType&& Rhs)
+	{
+		return ConcatFR(MoveTemp(Lhs), GetData(Rhs), GetNum(Rhs));
+	}
 
 	/**
 	 * Concatenate this path with given path ensuring the / character is used between them
@@ -613,7 +675,15 @@ public:
 	* @param Str path CharRangeType (string class/string view/string builder etc) to be concatenated onto the end of this
 	* @return reference to path
 	*/
-	template <typename CharRangeType, typename TEnableIf<TIsTCharRangeNotCArray <CharRangeType>::Value>::Type* = nullptr>
+	template <
+		typename CharRangeType,
+		typename CharRangeElementType = TElementType_T<CharRangeType>
+		UE_REQUIRES(
+			TIsContiguousContainer<CharRangeType>::Value &&
+			!std::is_array_v<std::remove_reference_t<CharRangeType>>&&
+			std::is_same_v<ElementType, CharRangeElementType>
+		)
+	>
 	FORCEINLINE UE_STRING_CLASS& operator/=(CharRangeType&& Str)
 	{
 		PathAppend(GetData(Str), GetNum(Str));
@@ -626,7 +696,10 @@ public:
 	* @param Str path array of CharType (that needs converting) to be concatenated onto the end of this
 	* @return reference to path
 	*/
-	template <typename CharType,typename TEnableIf<TIsCharType<CharType>::Value>::Type* = nullptr>
+	template <
+		typename CharType
+		UE_REQUIRES(TIsCharType_V<CharType>)
+	>
 	FORCEINLINE UE_STRING_CLASS& operator/=(const CharType* Str)
 	{
 		UE_STRING_CLASS Temp = Str;
@@ -1079,11 +1152,20 @@ public:
 	 *
 	 *        Consider using UE::String::FindLast() as an alternative.
 	 */
-	template <typename TCharRangeType, std::enable_if_t<TIsCharRangeNotCArrayNotStringClass<TCharRangeType>::Value>* = nullptr>
-	UE_NODISCARD int32 Find(TCharRangeType&& SubStr, ESearchCase::Type SearchCase = ESearchCase::IgnoreCase,
+	template <
+		typename CharRangeType,
+		typename CharRangeElementType = TElementType_T<CharRangeType>
+		UE_REQUIRES(
+			TIsContiguousContainer<CharRangeType>::Value &&
+			!std::is_array_v<std::remove_reference_t<CharRangeType>>&&
+			TIsCharType_V<CharRangeElementType> &&
+			!std::is_base_of_v<UE_STRING_CLASS, std::decay_t<CharRangeType>>
+		)
+	>
+	UE_NODISCARD int32 Find(CharRangeType&& SubStr, ESearchCase::Type SearchCase = ESearchCase::IgnoreCase,
 		ESearchDir::Type SearchDir = ESearchDir::FromStart, int32 StartPosition = INDEX_NONE) const
 	{
-		static_assert(std::is_same_v<typename TElementType<TCharRangeType>::Type, ElementType>, "Expected a range of ElementType");
+		static_assert(std::is_same_v<CharRangeElementType, ElementType>, "Expected a range of ElementType");
 		return Find(GetData(SubStr), GetNum(SubStr), SearchCase, SearchDir, StartPosition);
 	}
 
@@ -1153,12 +1235,21 @@ public:
 	 * @param SearchDir			Indicates whether the search starts at the beginning or at the end ( defaults to ESearchDir::FromStart )
 	 * @return					Returns whether the string contains the substring. If the substring is empty, returns true.
 	 **/
-	template <typename TCharRangeType, std::enable_if_t<TIsCharRangeNotCArrayNotStringClass<TCharRangeType>::Value>* = nullptr>
-	UE_NODISCARD FORCEINLINE bool Contains(TCharRangeType&& SubStr, ESearchCase::Type SearchCase = ESearchCase::IgnoreCase,
+	template <
+		typename CharRangeType,
+		typename CharRangeElementType = TElementType_T<CharRangeType>
+		UE_REQUIRES(
+			TIsContiguousContainer<CharRangeType>::Value &&
+			!std::is_array_v<std::remove_reference_t<CharRangeType>>&&
+			TIsCharType_V<CharRangeElementType> &&
+			!std::is_base_of_v<UE_STRING_CLASS, std::decay_t<CharRangeType>>
+		)
+	>
+	UE_NODISCARD FORCEINLINE bool Contains(CharRangeType&& SubStr, ESearchCase::Type SearchCase = ESearchCase::IgnoreCase,
 		ESearchDir::Type SearchDir = ESearchDir::FromStart) const
 	{
-		static_assert(std::is_same_v<typename TElementType<TCharRangeType>::Type, ElementType>, "Expected a range of characters");
-		return Find(Forward<TCharRangeType>(SubStr), SearchCase, SearchDir) != INDEX_NONE;
+		static_assert(std::is_same_v<CharRangeElementType, ElementType>, "Expected a range of characters");
+		return Find(Forward<CharRangeType>(SubStr), SearchCase, SearchDir) != INDEX_NONE;
 	}
 
 	/** 
@@ -1442,10 +1533,19 @@ public:
 	 * @param SearchCase		Indicates whether the search is case sensitive or not ( defaults to ESearchCase::IgnoreCase )
 	 * @return true if this string begins with specified text, false otherwise
 	 */
-	template <typename TCharRangeType, std::enable_if_t<TIsCharRangeNotCArrayNotStringClass<TCharRangeType>::Value>* = nullptr>
-	UE_NODISCARD bool StartsWith(TCharRangeType&& InPrefix, ESearchCase::Type SearchCase = ESearchCase::IgnoreCase) const
+	template <
+		typename CharRangeType,
+		typename CharRangeElementType = TElementType_T<CharRangeType>
+		UE_REQUIRES(
+			TIsContiguousContainer<CharRangeType>::Value &&
+			!std::is_array_v<std::remove_reference_t<CharRangeType>>&&
+			TIsCharType_V<CharRangeElementType> &&
+			!std::is_base_of_v<UE_STRING_CLASS, std::decay_t<CharRangeType>>
+		)
+	>
+	UE_NODISCARD bool StartsWith(CharRangeType&& InPrefix, ESearchCase::Type SearchCase = ESearchCase::IgnoreCase) const
 	{
-		static_assert(std::is_same_v<typename TElementType<TCharRangeType>::Type, ElementType>, "Expected a range of characters");
+		static_assert(std::is_same_v<CharRangeElementType, ElementType>, "Expected a range of characters");
 		return StartsWith(GetData(InPrefix), GetNum(InPrefix), SearchCase);
 	}
 
@@ -1485,10 +1585,19 @@ public:
 	 * @param SearchCase		Indicates whether the search is case sensitive or not ( defaults to ESearchCase::IgnoreCase )
 	 * @return true if this string ends with specified text, false otherwise
 	 */
-	template <typename TCharRangeType, std::enable_if_t<TIsCharRangeNotCArrayNotStringClass<TCharRangeType>::Value>* = nullptr>
-	UE_NODISCARD bool EndsWith(TCharRangeType&& InSuffix, ESearchCase::Type SearchCase = ESearchCase::IgnoreCase) const
+	template <
+		typename CharRangeType,
+		typename CharRangeElementType = TElementType_T<CharRangeType>
+		UE_REQUIRES(
+			TIsContiguousContainer<CharRangeType>::Value &&
+			!std::is_array_v<std::remove_reference_t<CharRangeType>>&&
+			TIsCharType_V<CharRangeElementType> &&
+			!std::is_base_of_v<UE_STRING_CLASS, std::decay_t<CharRangeType>>
+		)
+	>
+	UE_NODISCARD bool EndsWith(CharRangeType&& InSuffix, ESearchCase::Type SearchCase = ESearchCase::IgnoreCase) const
 	{
-		static_assert(std::is_same_v<typename TElementType<TCharRangeType>::Type, ElementType>, "Expected a range of characters");
+		static_assert(std::is_same_v<CharRangeElementType, ElementType>, "Expected a range of characters");
 		return EndsWith(GetData(InSuffix), GetNum(InSuffix), SearchCase);
 	}
 
@@ -1530,10 +1639,19 @@ public:
 	 * @return true if this string matches the *?-type wildcard given. 
 	 * @warning This is a simple, SLOW routine. Use with caution
 	 */
-	template <typename TCharRangeType, std::enable_if_t<TIsCharRangeNotCArrayNotStringClass<TCharRangeType>::Value>* = nullptr>
-	UE_NODISCARD bool MatchesWildcard(TCharRangeType&& Wildcard, ESearchCase::Type SearchCase = ESearchCase::IgnoreCase) const
+	template <
+		typename CharRangeType,
+		typename CharRangeElementType = TElementType_T<CharRangeType>
+		UE_REQUIRES(
+			TIsContiguousContainer<CharRangeType>::Value &&
+			!std::is_array_v<std::remove_reference_t<CharRangeType>>&&
+			TIsCharType_V<CharRangeElementType> &&
+			!std::is_base_of_v<UE_STRING_CLASS, std::decay_t<CharRangeType>>
+		)
+	>
+	UE_NODISCARD bool MatchesWildcard(CharRangeType&& Wildcard, ESearchCase::Type SearchCase = ESearchCase::IgnoreCase) const
 	{
-		static_assert(std::is_same_v<typename TElementType<TCharRangeType>::Type, ElementType>, "Expected a range of characters");
+		static_assert(std::is_same_v<CharRangeElementType, ElementType>, "Expected a range of characters");
 		return MatchesWildcard(GetData(Wildcard), GetNum(Wildcard), SearchCase);
 	}
 
diff --git a/Engine/Source/Runtime/Core/Public/Templates/Function.h b/Engine/Source/Runtime/Core/Public/Templates/Function.h
--- a/Engine/Source/Runtime/Core/Public/Templates/Function.h
+++ b/Engine/Source/Runtime/Core/Public/Templates/Function.h
@@ -6,19 +6,12 @@
 #include "Misc/AssertionMacros.h"
 #include "Misc/IntrusiveUnsetOptionalState.h"
 #include "HAL/UnrealMemory.h"
-#include "Templates/AndOrNot.h"
 #include "Templates/UnrealTypeTraits.h"
-#include "Templates/RemoveReference.h"
-#include "Templates/Decay.h"
 #include "Templates/Invoke.h"
-#include "Templates/IsConstructible.h"
-#include "Templates/IsInvocable.h"
-#include "Templates/IsMemberPointer.h"
-#include "Templates/IsPointer.h"
 #include "Templates/UnrealTemplate.h"
 #include "Math/UnrealMathUtility.h"
 #include <new> // IWYU pragma: export
-
+#include <type_traits>
 
 // Disable visualization hack for shipping or test builds.
 #ifndef UE_ENABLE_TFUNCTIONREF_VISUALIZATION
@@ -40,9 +33,6 @@
 	#define TFUNCTION_INLINE_ALIGNMENT    16
 #endif
 
-template <typename FuncType> class TFunction;
-template <typename FuncType> class TFunctionRef;
-
 /**
  * TFunction<FuncType>
  *
@@ -247,20 +237,10 @@ namespace UE::Core::Private::Function
 		}
 	};
 
-	template <typename T>
-	struct TIsNullableBinding :
-		TOr<
-			TIsPointer<T>,
-			TIsMemberPointer<T>,
-			TIsTFunction<T>
-		>
-	{
-	};
-
 	template <typename T>
 	FORCEINLINE bool IsBound(const T& Func)
 	{
-		if constexpr (TIsNullableBinding<T>::Value)
+		if constexpr (std::is_pointer_v<T> || std::is_member_pointer_v<T> || TIsTFunction<T>::Value)
 		{
 			// Function pointers, data member pointers, member function pointers and TFunctions
 			// can all be null/unbound, so test them using their boolean state.
@@ -279,13 +259,13 @@ namespace UE::Core::Private::Function
 	template <typename FunctorType, bool bOnHeap>
 	struct TStorageOwnerType<FunctorType, true, bOnHeap>
 	{
-		using Type = TFunction_UniqueOwnedObject<typename TDecay<FunctorType>::Type, bOnHeap>;
+		using Type = TFunction_UniqueOwnedObject<std::decay_t<FunctorType>, bOnHeap>;
 	};
 
 	template <typename FunctorType, bool bOnHeap>
 	struct TStorageOwnerType<FunctorType, false, bOnHeap>
 	{
-		using Type = TFunction_CopyableOwnedObject<typename TDecay<FunctorType>::Type, bOnHeap>;
+		using Type = TFunction_CopyableOwnedObject<std::decay_t<FunctorType>, bOnHeap>;
 	};
 
 	template <typename FunctorType, bool bUnique, bool bOnHeap>
@@ -377,7 +357,7 @@ namespace UE::Core::Private::Function
 		}
 
 		template <typename FunctorType>
-		typename TDecay<FunctorType>::Type* Bind(FunctorType&& InFunc)
+		std::decay_t<FunctorType>* Bind(FunctorType&& InFunc)
 		{
 			if (!IsBound(InFunc))
 			{
@@ -670,60 +650,12 @@ namespace UE::Core::Private::Function
 		#endif
 	};
 
-	template <typename FunctorType, typename Ret, typename... ParamTypes>
-	struct TFunctorReturnTypeIsCompatible
-		: TIsConstructible<Ret, decltype(DeclVal<FunctorType>()(DeclVal<ParamTypes>()...))>
-	{
-	};
-
-	template <typename MemberRet, typename Class, typename Ret, typename... ParamTypes>
-	struct TFunctorReturnTypeIsCompatible<MemberRet Class::*, Ret, ParamTypes...>
-		: TIsConstructible<Ret, MemberRet>
-	{
-	};
-
-	template <typename MemberRet, typename Class, typename Ret, typename... ParamTypes>
-	struct TFunctorReturnTypeIsCompatible<MemberRet Class::* const, Ret, ParamTypes...>
-		: TIsConstructible<Ret, MemberRet>
-	{
-	};
-
-	template <typename MemberRet, typename Class, typename... MemberParamTypes, typename Ret, typename... ParamTypes>
-	struct TFunctorReturnTypeIsCompatible<MemberRet (Class::*)(MemberParamTypes...), Ret, ParamTypes...>
-		: TIsConstructible<Ret, MemberRet>
-	{
-	};
-
-	template <typename MemberRet, typename Class, typename... MemberParamTypes, typename Ret, typename... ParamTypes>
-	struct TFunctorReturnTypeIsCompatible<MemberRet (Class::*)(MemberParamTypes...) const, Ret, ParamTypes...>
-		: TIsConstructible<Ret, MemberRet>
-	{
-	};
-
-	template <typename FuncType, typename FunctorType>
-	struct TFuncCanBindToFunctor;
-
-	template <typename FunctorType, typename Ret, typename... ParamTypes>
-	struct TFuncCanBindToFunctor<Ret(ParamTypes...), FunctorType> :
-		TAnd<
-			TIsInvocable<FunctorType, ParamTypes...>,
-			TFunctorReturnTypeIsCompatible<FunctorType, Ret, ParamTypes...>
-		>
-	{
-	};
-
-	template <typename FunctorType, typename... ParamTypes>
-	struct TFuncCanBindToFunctor<void(ParamTypes...), FunctorType> :
-		TIsInvocable<FunctorType, ParamTypes...>
-	{
-	};
-
 	struct FFunctionRefStoragePolicy
 	{
 		constexpr static bool bCanBeNull = false;
 
 		template <typename FunctorType>
-		typename TRemoveReference<FunctorType>::Type* Bind(FunctorType&& InFunc)
+		std::remove_reference_t<FunctorType>* Bind(FunctorType&& InFunc)
 		{
 			checkf(IsBound(InFunc), TEXT("Cannot bind a null/unbound callable to a TFunctionRef"));
 
@@ -811,23 +743,21 @@ namespace UE::Core::Private::Function
  *     });
  * }
  */
-template <typename FuncType>
-class TFunctionRef final : public UE::Core::Private::Function::TFunctionRefBase<UE::Core::Private::Function::FFunctionRefStoragePolicy, FuncType>
+template <typename Ret, typename... ParamTypes>
+class TFunctionRef<Ret(ParamTypes...)> final : public UE::Core::Private::Function::TFunctionRefBase<UE::Core::Private::Function::FFunctionRefStoragePolicy, Ret(ParamTypes...)>
 {
-	using Super = UE::Core::Private::Function::TFunctionRefBase<UE::Core::Private::Function::FFunctionRefStoragePolicy, FuncType>;
+	using Super = UE::Core::Private::Function::TFunctionRefBase<UE::Core::Private::Function::FFunctionRefStoragePolicy, Ret(ParamTypes...)>;
 
 public:
 	/**
 	 * Constructor which binds a TFunctionRef to a callable object.
 	 */
 	template <
-		typename FunctorType,
-		typename = typename TEnableIf<
-			TAnd<
-				TNot<TIsTFunctionRef<typename TDecay<FunctorType>::Type>>,
-				UE::Core::Private::Function::TFuncCanBindToFunctor<FuncType, typename TDecay<FunctorType>::Type>
-			>::Value
-		>::Type
+		typename FunctorType
+		UE_REQUIRES(
+			!TIsTFunctionRef<std::decay_t<FunctorType>>::Value &&
+			std::is_invocable_r_v<Ret, std::decay_t<FunctorType>, ParamTypes...>
+		)
 	>
 	TFunctionRef(FunctorType&& InFunc UE_LIFETIMEBOUND)
 		: Super(Forward<FunctorType>(InFunc))
@@ -902,10 +832,10 @@ public:
  *     FString Result = Transform(5); // "Hello: 25"
  * }
  */
-template <typename FuncType>
-class TFunction final : public UE::Core::Private::Function::TFunctionRefBase<UE::Core::Private::Function::TFunctionStorage<false>, FuncType>
+template <typename Ret, typename... ParamTypes>
+class TFunction<Ret(ParamTypes...)> final : public UE::Core::Private::Function::TFunctionRefBase<UE::Core::Private::Function::TFunctionStorage<false>, Ret(ParamTypes...)>
 {
-	using Super = UE::Core::Private::Function::TFunctionRefBase<UE::Core::Private::Function::TFunctionStorage<false>, FuncType>;
+	using Super = UE::Core::Private::Function::TFunctionRefBase<UE::Core::Private::Function::TFunctionStorage<false>, Ret(ParamTypes...)>;
 
 public:
 	/**
@@ -919,13 +849,11 @@ public:
 	 * Constructor which binds a TFunction to any function object.
 	 */
 	template <
-		typename FunctorType,
-		typename = typename TEnableIf<
-			TAnd<
-				TNot<TIsTFunction<typename TDecay<FunctorType>::Type>>,
-				UE::Core::Private::Function::TFuncCanBindToFunctor<FuncType, FunctorType>
-			>::Value
-		>::Type
+		typename FunctorType
+		UE_REQUIRES(
+			!TIsTFunction<std::decay_t<FunctorType>>::Value &&
+			std::is_invocable_r_v<Ret, std::decay_t<FunctorType>, ParamTypes...>
+		)
 	>
 	TFunction(FunctorType&& InFunc)
 		: Super(Forward<FunctorType>(InFunc))
@@ -940,7 +868,7 @@ public:
 		// reference):
 		//
 		// TFunction<int32(float)> MyFunction = [MyFunctionRef](float F) { return MyFunctionRef(F); };
-		static_assert(!TIsTFunctionRef<typename TDecay<FunctorType>::Type>::Value, "Cannot construct a TFunction from a TFunctionRef");
+		static_assert(!TIsTFunctionRef<std::decay_t<FunctorType>>::Value, "Cannot construct a TFunction from a TFunctionRef");
 	}
 
 	TFunction(TFunction&&) = default;
@@ -1015,10 +943,10 @@ public:
  * Foo(MovableFunc);           // error - TUniqueFunction is not copyable
  * Foo(MoveTemp(MovableFunc)); // ok
  */
-template <typename FuncType>
-class TUniqueFunction final : public UE::Core::Private::Function::TFunctionRefBase<UE::Core::Private::Function::TFunctionStorage<true>, FuncType>
+template <typename Ret, typename... ParamTypes>
+class TUniqueFunction<Ret(ParamTypes...)> final : public UE::Core::Private::Function::TFunctionRefBase<UE::Core::Private::Function::TFunctionStorage<true>, Ret(ParamTypes...)>
 {
-	using Super = UE::Core::Private::Function::TFunctionRefBase<UE::Core::Private::Function::TFunctionStorage<true>, FuncType>;
+	using Super = UE::Core::Private::Function::TFunctionRefBase<UE::Core::Private::Function::TFunctionStorage<true>, Ret(ParamTypes...)>;
 
 public:
 	/**
@@ -1032,13 +960,12 @@ public:
 	 * Constructor which binds a TFunction to any function object.
 	 */
 	template <
-		typename FunctorType,
-		typename = typename TEnableIf<
-			TAnd<
-				TNot<TOr<TIsTUniqueFunction<typename TDecay<FunctorType>::Type>, TIsTFunction<typename TDecay<FunctorType>::Type>>>,
-				UE::Core::Private::Function::TFuncCanBindToFunctor<FuncType, FunctorType>
-			>::Value
-		>::Type
+		typename FunctorType
+		UE_REQUIRES(
+			!TIsTUniqueFunction<std::decay_t<FunctorType>>::Value &&
+			!TIsTFunction      <std::decay_t<FunctorType>>::Value &&
+			std::is_invocable_r_v<Ret, std::decay_t<FunctorType>, ParamTypes...>
+		)
 	>
 	TUniqueFunction(FunctorType&& InFunc)
 		: Super(Forward<FunctorType>(InFunc))
@@ -1053,22 +980,22 @@ public:
 		// reference):
 		//
 		// TFunction<int32(float)> MyFunction = [MyFunctionRef](float F) { return MyFunctionRef(F); };
-		static_assert(!TIsTFunctionRef<typename TDecay<FunctorType>::Type>::Value, "Cannot construct a TUniqueFunction from a TFunctionRef");
+		static_assert(!TIsTFunctionRef<std::decay_t<FunctorType>>::Value, "Cannot construct a TUniqueFunction from a TFunctionRef");
 	}
 
 	/**
 	 * Constructor which takes ownership of a TFunction's functor.
 	 */
-	TUniqueFunction(TFunction<FuncType>&& Other)
-		: Super(MoveTemp(*(UE::Core::Private::Function::TFunctionRefBase<UE::Core::Private::Function::TFunctionStorage<false>, FuncType>*)&Other))
+	TUniqueFunction(TFunction<Ret(ParamTypes...)>&& Other)
+		: Super(MoveTemp(*(UE::Core::Private::Function::TFunctionRefBase<UE::Core::Private::Function::TFunctionStorage<false>, Ret(ParamTypes...)>*)&Other))
 	{
 	}
 
 	/**
 	 * Constructor which takes ownership of a TFunction's functor.
 	 */
-	TUniqueFunction(const TFunction<FuncType>& Other)
-		: Super(*(const UE::Core::Private::Function::TFunctionRefBase<UE::Core::Private::Function::TFunctionStorage<false>, FuncType>*)&Other)
+	TUniqueFunction(const TFunction<Ret(ParamTypes...)>& Other)
+		: Super(*(const UE::Core::Private::Function::TFunctionRefBase<UE::Core::Private::Function::TFunctionStorage<false>, Ret(ParamTypes...)>*)&Other)
 	{
 	}
 
@@ -1122,4 +1049,14 @@ FORCEINLINE bool operator!=(TYPE_OF_NULLPTR, const TFunction<FuncType>& Func)
 {
 	return (bool)Func;
 }
+#endif
+
+#if UE_ENABLE_INCLUDE_ORDER_DEPRECATED_IN_5_4
+#include "Templates/AndOrNot.h"
+#include "Templates/Decay.h"
+#include "Templates/IsConstructible.h"
+#include "Templates/IsInvocable.h"
+#include "Templates/IsPointer.h"
+#include "Templates/IsMemberPointer.h"
+#include "Templates/RemoveReference.h"
 #endif
\ No newline at end of file
diff --git a/Engine/Source/Runtime/Core/Public/Templates/Tuple.h b/Engine/Source/Runtime/Core/Public/Templates/Tuple.h
--- a/Engine/Source/Runtime/Core/Public/Templates/Tuple.h
+++ b/Engine/Source/Runtime/Core/Public/Templates/Tuple.h
@@ -4,13 +4,11 @@
 
 #include "CoreTypes.h"
 #include "Templates/UnrealTemplate.h"
-#include "Templates/Decay.h"
 #include "Delegates/IntegerSequence.h"
 #include "Templates/Invoke.h"
 #include "Serialization/StructuredArchive.h"
 #include "Serialization/MemoryLayout.h"
 #include "Templates/TypeHash.h"
-#include "Templates/IsConstructible.h"
 #include <tuple>
 
 // This workaround exists because Visual Studio causes false positives for code like this during static analysis:
@@ -120,7 +118,7 @@ namespace UE::Core::Private::Tuple
 			decltype(auto) Result = (ForwardAsBase<TupleType, TTupleBaseElement<DeducedType, Index, TupleSize>>(Tuple).Value);
 
 			// Keep tuple rvalue references to rvalue reference elements as rvalues, because that's how std::get() works, not how C++ struct member access works.
-			return static_cast<std::conditional_t<TAnd<TNot<TIsReferenceType<TupleType>>, TIsRValueReferenceType<DeducedType>>::Value, DeducedType, decltype(Result)>>(Result);
+			return static_cast<std::conditional_t<!std::is_reference_v<TupleType> && std::is_rvalue_reference_v<DeducedType>, DeducedType, decltype(Result)>>(Result);
 		}
 
 		template <typename TupleType>
@@ -141,7 +139,7 @@ namespace UE::Core::Private::Tuple
 			decltype(auto) Result = (Forward<TupleType>(Tuple).Key);
 
 			// Keep tuple rvalue references to rvalue reference elements as rvalues, because that's how std::get() works, not how C++ struct member access works.
-			return static_cast<std::conditional_t<TAnd<TNot<TIsReferenceType<TupleType>>, TIsRValueReferenceType<decltype(Tuple.Key)>>::Value, decltype(Tuple.Key), decltype(Result)>>(Result);
+			return static_cast<std::conditional_t<!std::is_reference_v<TupleType> && std::is_rvalue_reference_v<decltype(Tuple.Key)>, decltype(Tuple.Key), decltype(Result)>>(Result);
 		}
 	};
 	template <>
@@ -154,7 +152,7 @@ namespace UE::Core::Private::Tuple
 			decltype(auto) Result = (Forward<TupleType>(Tuple).Value);
 
 			// Keep tuple rvalue references to rvalue reference elements as rvalues, because that's how std::get() works, not how C++ struct member access works.
-			return static_cast<std::conditional_t<TAnd<TNot<TIsReferenceType<TupleType>>, TIsRValueReferenceType<decltype(Tuple.Value)>>::Value, decltype(Tuple.Value), decltype(Result)>>(Result);
+			return static_cast<std::conditional_t<!std::is_reference_v<TupleType> && std::is_rvalue_reference_v<decltype(Tuple.Value)>, decltype(Tuple.Value), decltype(Result)>>(Result);
 		}
 	};
 #else
@@ -168,7 +166,7 @@ namespace UE::Core::Private::Tuple
 			decltype(auto) Result = (ForwardAsBase<TupleType, TTupleBaseElement<decltype(Tuple.Key), 0, 2>>(Tuple).Key);
 
 			// Keep tuple rvalue references to rvalue reference elements as rvalues, because that's how std::get() works, not how C++ struct member access works.
-			return static_cast<std::conditional_t<TAnd<TNot<TIsReferenceType<TupleType>>, TIsRValueReferenceType<decltype(Tuple.Key)>>::Value, decltype(Tuple.Key), decltype(Result)>>(Result);
+			return static_cast<std::conditional_t<!std::is_reference_v<TupleType> && std::is_rvalue_reference_v<decltype(Tuple.Key)>, decltype(Tuple.Key), decltype(Result)>>(Result);
 		}
 	};
 #endif
@@ -996,3 +994,8 @@ FORCEINLINE void operator<<(FStructuredArchive::FSlot Slot, TTuple<Types...>& Tu
 {
 	TTuple<Types...>::SerializeStructured(Slot, Tuple);
 }
+
+#if UE_ENABLE_INCLUDE_ORDER_DEPRECATED_IN_5_4
+#include "Templates/Decay.h"
+#include "Templates/IsConstructible.h"
+#endif
\ No newline at end of file
