Add warning on all critical sections about double unlock

diff --git a/Engine/Source/Runtime/Core/Public/HAL/PThreadCriticalSection.h b/Engine/Source/Runtime/Core/Public/HAL/PThreadCriticalSection.h
--- a/Engine/Source/Runtime/Core/Public/HAL/PThreadCriticalSection.h
+++ b/Engine/Source/Runtime/Core/Public/HAL/PThreadCriticalSection.h
@@ -61,8 +61,11 @@ public:
 		return 0 == pthread_mutex_trylock(&Mutex);
 	}
 
+
 	/**
-	 * Releases the lock on the critical seciton
+	 * Releases the lock on the critical section
+	 *
+	 * Calling this when not locked is undefined behavior & may cause indefinite waiting on next lock.
 	 */
 	FORCEINLINE void Unlock(void)
 	{
diff --git a/Engine/Source/Runtime/Core/Public/Windows/WindowsCriticalSection.h b/Engine/Source/Runtime/Core/Public/Windows/WindowsCriticalSection.h
--- a/Engine/Source/Runtime/Core/Public/Windows/WindowsCriticalSection.h
+++ b/Engine/Source/Runtime/Core/Public/Windows/WindowsCriticalSection.h
@@ -60,6 +60,9 @@ public:
 
 	/**
 	 * Releases the lock on the critical section
+	 * 
+	 * Calling this when not locked is undefined behavior & may cause indefinite waiting on next lock.
+	 * See: https://learn.microsoft.com/en-us/windows/win32/api/synchapi/nf-synchapi-leavecriticalsection#remarks
 	 */
 	FORCEINLINE void Unlock()
 	{
