Work around a clang static analyzer assert ([REDACTED_PROFILE_URL]/llvm-project/issues/69922) by not using __builtin_bit_cast when __clang_analyzer__ is defined.

diff --git a/Engine/Source/Runtime/Core/Public/Templates/TypeCompatibleBytes.h b/Engine/Source/Runtime/Core/Public/Templates/TypeCompatibleBytes.h
--- a/Engine/Source/Runtime/Core/Public/Templates/TypeCompatibleBytes.h
+++ b/Engine/Source/Runtime/Core/Public/Templates/TypeCompatibleBytes.h
@@ -38,7 +38,9 @@ inline ToType BitCast(const FromType& From)
 // Ensure we can use this builtin - seems to be present on Clang 9, GCC 11 and MSVC 19.26,
 // but gives spurious "non-void function 'BitCast' should return a value" errors on some
 // Mac and Android toolchains when building PCHs, so avoid those.
-#if (PLATFORM_APPLE && __clang_major__ >= 13) || (PLATFORM_ANDROID && __clang_major__ >= 13) || (PLATFORM_SWITCH && __clang_major__ >= 13) || (!PLATFORM_APPLE && !PLATFORM_ANDROID && !PLATFORM_SWITCH && defined(__clang__) && __clang_major__ >= 9) || (defined(__GNUC__) && __GNUC__ >= 11) || (defined(_MSC_VER) && _MSC_VER >= 1926)
+// However, there is a bug in the Clang static analyzer with this builtin: [REDACTED_PROFILE_URL]/llvm-project/issues/69922
+// Don't use it when performing static analysis until the bug is fixed.
+#if !defined(__clang_analyzer__) && ((PLATFORM_APPLE && __clang_major__ >= 13) || (PLATFORM_ANDROID && __clang_major__ >= 13) || (PLATFORM_SWITCH && __clang_major__ >= 13) || (!PLATFORM_APPLE && !PLATFORM_ANDROID && !PLATFORM_SWITCH && defined(__clang__) && __clang_major__ >= 9) || (defined(__GNUC__) && __GNUC__ >= 11) || (defined(_MSC_VER) && _MSC_VER >= 1926))
 	return __builtin_bit_cast(ToType, From);
 #else
 	TTypeCompatibleBytes<ToType> Result;
