Fixed TChunkedArray ending iteration prematurely

Prior to this change the begin/end() iterator objects would be compared by their the element pointer. The equivalent of "begin.elements + num != end.element". In cases were the number of elements was an exact multiple of a chunks count, "elements + num" would be an out-of-bounds address. It is never dereferenced so this is fine. However, depending on the allocator, one chunk could immediately follow another in memory such that "addr(B.elements + num) == addr(A.elements)". If A was the first chunk iterated, it would have the same element address as end() in B and iteration would finish prematurely.

#rnx
#tests analysis of good/bad data from Horde, created a unit test to induce and prove scenario

diff --git a/Engine/Source/Runtime/Core/Public/Containers/ChunkedArray.h b/Engine/Source/Runtime/Core/Public/Containers/ChunkedArray.h
--- a/Engine/Source/Runtime/Core/Public/Containers/ChunkedArray.h
+++ b/Engine/Source/Runtime/Core/Public/Containers/ChunkedArray.h
@@ -23,6 +23,7 @@ namespace UE4ChunkedArray_Private
 		ElementType* Elem;
 		ChunkType**  Chunk;
 		ChunkType**  LastChunk;
+		uint32		 Count = 0;
 
 		ElementType& operator*() const
 		{
@@ -31,6 +32,7 @@ namespace UE4ChunkedArray_Private
 
 		void operator++()
 		{
+			++Count;
 			++Elem;
 			if (Chunk != LastChunk && Elem == (*Chunk)->Elements + NumElementsPerChunk)
 			{
@@ -41,7 +43,7 @@ namespace UE4ChunkedArray_Private
 
 		bool operator!=(const TChunkedArrayIterator& Rhs) const
 		{
-			return Elem != Rhs.Elem;
+			return Count != Rhs.Count;
 		}
 	};
 }
@@ -315,7 +317,9 @@ public:
 		bool bBeyondLastChunk = Num && (Num % NumElementsPerChunk) == 0;
 		FChunk** ChunkPtr = Chunks.GetData() + (Num / NumElementsPerChunk) + (bBeyondLastChunk ? -1 : 0); // do not read off the end of the chunk array!
 		FChunk** LastChunkPtr = Chunks.GetData() + (Num ? Num - 1 : 0) / NumElementsPerChunk;
-		return FIterType(ChunkPtr, LastChunkPtr, ChunkPtr ? (*ChunkPtr)->Elements + (bBeyondLastChunk ? NumElementsPerChunk : (Num % NumElementsPerChunk))  : nullptr);
+		FIterType Ret(ChunkPtr, LastChunkPtr, ChunkPtr ? (*ChunkPtr)->Elements + (bBeyondLastChunk ? NumElementsPerChunk : (Num % NumElementsPerChunk))  : nullptr);
+		Ret.Count = Num;
+		return Ret;
 	}
 
 	FConstIterType end() const
@@ -324,7 +328,9 @@ public:
 		bool bBeyondLastChunk = Num && Num % NumElementsPerChunk == 0;
 		const FChunk** ChunkPtr = Chunks.GetData() + (Num / NumElementsPerChunk) + (bBeyondLastChunk ? -1 : 0); // do not read off the end of the chunk array!
 		const FChunk** LastChunkPtr = Chunks.GetData() + (Num ? Num - 1 : 0) / NumElementsPerChunk;
-		return FConstIterType(ChunkPtr, LastChunkPtr, ChunkPtr ? (*ChunkPtr)->Elements + (bBeyondLastChunk ? NumElementsPerChunk : (Num % NumElementsPerChunk)) : nullptr);
+		FConstIterType Ret(ChunkPtr, LastChunkPtr, ChunkPtr ? (*ChunkPtr)->Elements + (bBeyondLastChunk ? NumElementsPerChunk : (Num % NumElementsPerChunk)) : nullptr);
+		Ret.Count = Num;
+		return Ret;
 	}
 };
 
