Add unit test for TMap::TIterator::RemoveCurrent

* Add new low-level unit test file for TMap tests, currently with a single unit test called "System::Core::Containers::TMap::TIterator::RemoveCurrent". This test checks that TMap::TIterator::RemoveCurrent removes the pointed-at pair from the map but that subsequent increments of the iterator do not miss any element.
* Update the documentation comment of TMap::TIterator::RemoveCurrent to point out that the method can be called without interrupting subsequent increments.

diff --git a/Engine/Source/Runtime/Core/Public/Containers/Map.h b/Engine/Source/Runtime/Core/Public/Containers/Map.h
--- a/Engine/Source/Runtime/Core/Public/Containers/Map.h
+++ b/Engine/Source/Runtime/Core/Public/Containers/Map.h
@@ -917,7 +917,10 @@ public:
 			}
 		}
 
-		/** Removes the current pair from the map. */
+		/** Removes the current pair from the map without losing the iteration
+		 * position. Increment before using the iterator again, but after that
+		 * it will point at the element that was after the removed element.
+		 */
 		FORCEINLINE void RemoveCurrent()
 		{
 			TBaseIterator<false>::PairIt.RemoveCurrent();
diff --git a/Engine/Source/Runtime/Core/Tests/Containers/MapTest.cpp b/Engine/Source/Runtime/Core/Tests/Containers/MapTest.cpp
new file mode 100644
--- /dev/null
+++ b/Engine/Source/Runtime/Core/Tests/Containers/MapTest.cpp
@@ -0,0 +1,39 @@
+﻿// Copyright Epic Games, Inc. All Rights Reserved.
+
+#if WITH_LOW_LEVEL_TESTS
+
+#include "Containers/Map.h"
+
+#include "Tests/TestHarnessAdapter.h"
+
+TEST_CASE_NAMED(FMapIteratorTest, "System::Core::Containers::TMap::TIterator::RemoveCurrent", "[Core][Containers][TMap]")
+{
+	TMap<int, int> MapUnderTest;
+
+	MapUnderTest.Add(5, 55);
+	MapUnderTest.Add(1, 11);
+	MapUnderTest.Add(3, 33);
+	MapUnderTest.Add(4, 44);
+
+	int VisitedItemsCount = 0;
+	for (auto It = MapUnderTest.CreateIterator(); It; ++It)
+	{
+		++VisitedItemsCount;
+		if (It->Key == 1)
+		{
+			It.RemoveCurrent();
+		}
+	}
+
+	CHECK(MapUnderTest.Num() == 3);
+	
+	CHECK(VisitedItemsCount == 4);
+	
+	CHECK(MapUnderTest.Find(5) != nullptr);
+	CHECK(MapUnderTest.Find(1) == nullptr);
+	CHECK(MapUnderTest.Find(3) != nullptr);
+	CHECK(MapUnderTest.Find(4) != nullptr);
+}
+
+
+#endif //WITH_LOW_LEVEL_TESTS
