Logging: Removed the UE_LOG_JSON global mutex on Windows that enabled structured debug output

Integration with other tools can rely on the -JsonDebugOutput command line argument.

diff --git a/Engine/Source/Runtime/Core/Private/Windows/WindowsPlatformMisc.cpp b/Engine/Source/Runtime/Core/Private/Windows/WindowsPlatformMisc.cpp
--- a/Engine/Source/Runtime/Core/Private/Windows/WindowsPlatformMisc.cpp
+++ b/Engine/Source/Runtime/Core/Private/Windows/WindowsPlatformMisc.cpp
@@ -998,13 +998,6 @@ void FWindowsPlatformMisc::LocalPrint( const TCHAR *Message )
 #endif
 }
 
-bool FWindowsPlatformMisc::IsLowLevelOutputDebugStringStructured()
-{
-	HANDLE Mutex = OpenMutexW(SYNCHRONIZE, /*bInheritHandle*/ false, L"UE_LOG_JSON");
-	ON_SCOPE_EXIT { CloseHandle(Mutex); };
-	return !!Mutex || FGenericPlatformMisc::IsLowLevelOutputDebugStringStructured();
-}
-
 void FWindowsPlatformMisc::RequestExit( bool Force, const TCHAR* CallSite )
 {
 	UE_LOG(LogWindows, Log,  TEXT("FPlatformMisc::RequestExit(%i, %s)"),
diff --git a/Engine/Source/Runtime/Core/Public/Windows/WindowsPlatformMisc.h b/Engine/Source/Runtime/Core/Public/Windows/WindowsPlatformMisc.h
--- a/Engine/Source/Runtime/Core/Public/Windows/WindowsPlatformMisc.h
+++ b/Engine/Source/Runtime/Core/Public/Windows/WindowsPlatformMisc.h
@@ -90,8 +90,6 @@ struct FWindowsPlatformMisc
 	static CORE_API void SetUTF8Output();
 	static CORE_API void LocalPrint(const TCHAR *Message);
 
-	static CORE_API bool IsLowLevelOutputDebugStringStructured();
-
 	static bool IsLocalPrintThreadSafe()
 	{ 
 		//returning true when the debugger is attached is to allow
