[UE] Expose FGameThreadHitchHeartBeatThreaded::bStartSuspended to correctly track state of the heartbeat thread.

#tests replay
#rnx

diff --git a/Engine/Source/Runtime/Core/Private/HAL/ThreadHeartBeat.cpp b/Engine/Source/Runtime/Core/Private/HAL/ThreadHeartBeat.cpp
--- a/Engine/Source/Runtime/Core/Private/HAL/ThreadHeartBeat.cpp
+++ b/Engine/Source/Runtime/Core/Private/HAL/ThreadHeartBeat.cpp
@@ -954,7 +954,7 @@ void FGameThreadHitchHeartBeatThreaded::InitSettings()
 		CmdLine_StackWalk = FParse::Param(FCommandLine::Get(), TEXT("hitchdetectionstackwalk"));
 
 		// Determine whether to start suspended
-		bool bStartSuspended = false;
+		bStartSuspended = false;
 		if (GConfig)
 		{
 			GConfig->GetBool(TEXT("Core.System"), TEXT("GameThreadHeartBeatStartSuspended"), bStartSuspended, GEngineIni);
@@ -1176,6 +1176,15 @@ void FGameThreadHitchHeartBeatThreaded::ResumeHeartBeat()
 #endif
 }
 
+bool FGameThreadHitchHeartBeatThreaded::IsStartedSuspended()
+{
+#if USE_HITCH_DETECTION
+	return bStartSuspended;
+#else
+	return true;
+#endif
+}
+
 double FGameThreadHitchHeartBeatThreaded::GetFrameStartTime()
 {
 #if USE_HITCH_DETECTION
diff --git a/Engine/Source/Runtime/Core/Public/HAL/ThreadHeartBeat.h b/Engine/Source/Runtime/Core/Public/HAL/ThreadHeartBeat.h
--- a/Engine/Source/Runtime/Core/Public/HAL/ThreadHeartBeat.h
+++ b/Engine/Source/Runtime/Core/Public/HAL/ThreadHeartBeat.h
@@ -340,6 +340,7 @@ class FGameThreadHitchHeartBeatThreaded : public FRunnable
 	double FirstStartTime;
 	double FrameStartTime;
 	int32 SuspendedCount;
+	bool bStartSuspended = false;
 
 #if WALK_STACK_ON_HITCH_DETECTED
 #if LOOKUP_SYMBOLS_IN_HITCH_STACK_WALK
@@ -388,6 +389,11 @@ public:
 	* Resume heartbeat hitch detection. Call only after first calling SuspendHeartBeat.
 	*/
 	CORE_API void ResumeHeartBeat();
+	
+	/**
+	* Check if started suspended.
+	*/
+	CORE_API bool IsStartedSuspended();
 
 	// No-op, used in FUnixSignalGameHitchHeartBeat
 	void Restart() {}
