[LLM] Suspend HitchDetector when LLM is activated

LLM adds a lot of CPU overhead, there is no need to add more overhead by spamming logs with hitchdetector.

#tests local

diff --git a/Engine/Source/Runtime/Core/Private/HAL/LowLevelMemTracker.cpp b/Engine/Source/Runtime/Core/Private/HAL/LowLevelMemTracker.cpp
--- a/Engine/Source/Runtime/Core/Private/HAL/LowLevelMemTracker.cpp
+++ b/Engine/Source/Runtime/Core/Private/HAL/LowLevelMemTracker.cpp
@@ -7,6 +7,7 @@
 #include "HAL/LowLevelMemStats.h"
 #include "HAL/FileManager.h"
 #include "HAL/IConsoleManager.h"
+#include "HAL/ThreadHeartBeat.h"
 #include "LowLevelMemTrackerPrivate.h"
 #include "MemPro/MemProProfiler.h"
 #include "Misc/CString.h"
@@ -1262,6 +1263,13 @@ void FLowLevelMemTracker::ProcessCommandLine(const TCHAR* CmdLine)
 	bIsDisabled = false;
 	bCsvWriterEnabled = bLocalCsvWriterEnabled;
 	bTraceWriterEnabled = bLocalTraceWriterEnabled;
+	if (!FParse::Param(CmdLine, TEXT("unsuspend_hitchdetector")))
+	{
+		// Disable hang heatbeat when LLM is enabled
+		FThreadHeartBeat::Get().SuspendHeartBeat(true);
+		// Disable hitchdetector when LLM is enabled
+		FGameThreadHitchHeartBeat::Get().SuspendHeartBeat();
+	}
 	BootstrapInitialise();
 	FinishInitialise();
 
