WindowsPlatformFile: Move FWindowsPlatformFile into a private header so we can make future changes with its functionality split into subclasses.
#rnx

diff --git a/Engine/Source/Runtime/Core/Private/Windows/WindowsPlatformFile.cpp b/Engine/Source/Runtime/Core/Private/Windows/WindowsPlatformFile.cpp
--- a/Engine/Source/Runtime/Core/Private/Windows/WindowsPlatformFile.cpp
+++ b/Engine/Source/Runtime/Core/Private/Windows/WindowsPlatformFile.cpp
@@ -1,32 +1,37 @@
 // Copyright Epic Games, Inc. All Rights Reserved.
 
 #include "Windows/WindowsPlatformFile.h"
+
+#include "Windows/WindowsPlatformFilePrivate.h"
+
+#include "Algo/Replace.h"
+#include "Async/AsyncFileHandle.h"
+#include "Async/AsyncWork.h"
+#include "Async/MappedFileHandle.h"
+#include "Containers/LockFreeList.h"
+#include "Containers/UnrealString.h"
+#include "CoreGlobals.h"
 #include "CoreTypes.h"
-#include "Misc/DateTime.h"
-#include "Misc/AssertionMacros.h"
-#include "Misc/PathViews.h"
-#include "Misc/StringBuilder.h"
+#include "HAL/IConsoleManager.h"
+#include "HAL/IPlatformFileManagedStorageWrapper.h"
+#include "HAL/UnrealMemory.h"
 #include "Logging/LogMacros.h"
 #include "Math/UnrealMathUtility.h"
-#include "HAL/UnrealMemory.h"
-#include "GenericPlatform/GenericPlatformFile.h"
-#include "Containers/UnrealString.h"
-#include "Algo/Replace.h"
-#include "Templates/Function.h"
+#include "Misc/AssertionMacros.h"
 #include "Misc/Paths.h"
-#include "CoreGlobals.h"
-#include <sys/utime.h>
-#include "Containers/LockFreeList.h"
-#include "Async/AsyncFileHandle.h"
-#include "Async/AsyncWork.h"
+#include "Misc/PathViews.h"
 #include "Misc/ScopeLock.h"
-#include "HAL/IConsoleManager.h"
 #include "ProfilingDebugging/PlatformFileTrace.h"
-#include "HAL/IPlatformFileManagedStorageWrapper.h"
-#include "Windows/AllowWindowsPlatformTypes.h"
+#include "Templates/Function.h"
+#include <sys/utime.h>
 
+#include "Windows/AllowWindowsPlatformTypes.h" // Required immediately before MicrosoftAsyncIO.h and the block below it
 #include "Microsoft/MicrosoftAsyncIO.h"
-#include "Async/MappedFileHandle.h"
+namespace FileConstants
+{
+uint32 WIN_INVALID_SET_FILE_POINTER = INVALID_SET_FILE_POINTER;
+} // namespace FileConstants
+#include "Windows/HideWindowsPlatformTypes.h"
 
 TLockFreePointerListUnordered<void, PLATFORM_CACHE_LINE_SIZE> MicrosoftAsyncIOEventPool;
 bool GTriggerFailedMicrosoftRead = false;
@@ -44,78 +49,78 @@ static FAutoConsoleCommand TriggerFailedMicrosoftReadCmd(
 );
 #endif
 
+namespace FileConstants
+{
 
-	namespace FileConstants
-	{
-		uint32 WIN_INVALID_SET_FILE_POINTER = INVALID_SET_FILE_POINTER;
+/** Special suffix to access Alternate Data Stream "Zone.Identifier" for NTFS files */
+constexpr FStringView ZoneIdentifierStreamSuffix(TEXTVIEW(":Zone.Identifier"));
 
-		/** Special suffix to access Alternate Data Stream "Zone.Identifier" for NTFS files */
-		constexpr FStringView ZoneIdentifierStreamSuffix(TEXTVIEW(":Zone.Identifier"));
-	}
-#include "Windows/HideWindowsPlatformTypes.h"
+} // namespace FileConstants
 
-namespace
+namespace UE::WindowsPlatformFile::Private
 {
-	FORCEINLINE int32 UEDayOfWeekToWindowsSystemTimeDayOfWeek(const EDayOfWeek InDayOfWeek)
-	{
-		switch (InDayOfWeek)
-		{
-			case EDayOfWeek::Monday:
-				return 1;
-			case EDayOfWeek::Tuesday:
-				return 2;
-			case EDayOfWeek::Wednesday:
-				return 3;
-			case EDayOfWeek::Thursday:
-				return 4;
-			case EDayOfWeek::Friday:
-				return 5;
-			case EDayOfWeek::Saturday:
-				return 6;
-			case EDayOfWeek::Sunday:
-				return 0;
-			default:
-				break;
-		}
 
-		return 0;
-	}
+int32 UEDayOfWeekToWindowsSystemTimeDayOfWeek(const EDayOfWeek InDayOfWeek)
+{
+	switch (InDayOfWeek)
+	{
+		case EDayOfWeek::Monday:
+			return 1;
+		case EDayOfWeek::Tuesday:
+			return 2;
+		case EDayOfWeek::Wednesday:
+			return 3;
+		case EDayOfWeek::Thursday:
+			return 4;
+		case EDayOfWeek::Friday:
+			return 5;
+		case EDayOfWeek::Saturday:
+			return 6;
+		case EDayOfWeek::Sunday:
+			return 0;
+		default:
+			break;
+	}
+
+	return 0;
+}
 
-	FORCEINLINE FDateTime WindowsFileTimeToUEDateTime(const FILETIME& InFileTime)
+FDateTime WindowsFileTimeToUEDateTime(const FILETIME& InFileTime)
+{
+	// This roundabout conversion clamps the precision of the returned time value to match that of time_t (1 second precision)
+	// This avoids issues when sending files over the network via cook-on-the-fly
+	SYSTEMTIME SysTime;
+	if (FileTimeToSystemTime(&InFileTime, &SysTime))
 	{
-		// This roundabout conversion clamps the precision of the returned time value to match that of time_t (1 second precision)
-		// This avoids issues when sending files over the network via cook-on-the-fly
-		SYSTEMTIME SysTime;
-		if (FileTimeToSystemTime(&InFileTime, &SysTime))
-		{
-			return FDateTime(SysTime.wYear, SysTime.wMonth, SysTime.wDay, SysTime.wHour, SysTime.wMinute, SysTime.wSecond);
-		}
-
-		// Failed to convert
-		return FDateTime::MinValue();
+		return FDateTime(SysTime.wYear, SysTime.wMonth, SysTime.wDay, SysTime.wHour, SysTime.wMinute, SysTime.wSecond);
 	}
 
-	FORCEINLINE FILETIME UEDateTimeToWindowsFileTime(const FDateTime& InDateTime)
-	{
-		// This roundabout conversion clamps the precision of the returned time value to match that of time_t (1 second precision)
-		// This avoids issues when sending files over the network via cook-on-the-fly
-		SYSTEMTIME SysTime;
-		SysTime.wYear = (WORD)InDateTime.GetYear();
-		SysTime.wMonth = (WORD)InDateTime.GetMonth();
-		SysTime.wDay = (WORD)InDateTime.GetDay();
-		SysTime.wDayOfWeek = (WORD)UEDayOfWeekToWindowsSystemTimeDayOfWeek(InDateTime.GetDayOfWeek());
-		SysTime.wHour = (WORD)InDateTime.GetHour();
-		SysTime.wMinute = (WORD)InDateTime.GetMinute();
-		SysTime.wSecond = (WORD)InDateTime.GetSecond();
-		SysTime.wMilliseconds = 0;
-
-		FILETIME FileTime;
-		SystemTimeToFileTime(&SysTime, &FileTime);
+	// Failed to convert
+	return FDateTime::MinValue();
+}
 
-		return FileTime;
-	}
+FILETIME UEDateTimeToWindowsFileTime(const FDateTime& InDateTime)
+{
+	// This roundabout conversion clamps the precision of the returned time value to match that of time_t (1 second precision)
+	// This avoids issues when sending files over the network via cook-on-the-fly
+	SYSTEMTIME SysTime;
+	SysTime.wYear = (WORD)InDateTime.GetYear();
+	SysTime.wMonth = (WORD)InDateTime.GetMonth();
+	SysTime.wDay = (WORD)InDateTime.GetDay();
+	SysTime.wDayOfWeek = (WORD)UE::WindowsPlatformFile::Private::UEDayOfWeekToWindowsSystemTimeDayOfWeek(InDateTime.GetDayOfWeek());
+	SysTime.wHour = (WORD)InDateTime.GetHour();
+	SysTime.wMinute = (WORD)InDateTime.GetMinute();
+	SysTime.wSecond = (WORD)InDateTime.GetSecond();
+	SysTime.wMilliseconds = 0;
+
+	FILETIME FileTime;
+	SystemTimeToFileTime(&SysTime, &FileTime);
+
+	return FileTime;
 }
 
+} // namespace UE::WindowsPlatformFile::Private
+
 /**
  * This file reader uses overlapped i/o and double buffering to asynchronously read from files
  */
@@ -916,640 +921,624 @@ FMappedFileRegionWindows::~FMappedFileRegionWindows()
 	Parent->UnMap(this);
 }
 
-/**
- * Windows File I/O implementation
-**/
-class CORE_API FWindowsPlatformFile : public IPhysicalPlatformFile
+void FWindowsPlatformFile::NormalizeWindowsPath(FStringBuilderBase& Path, bool bIsFilename)
 {
-	/**
-	  * Convert from a valid Unreal Path to a canonical and strict-valid Windows Path.
-	  * An Unreal Path may have either \ or / and may have empty directories (two / in a row), and may have .. and may be relative
-	  * A canonical and strict-valid Windows Path has only \, does not have .., does not have empty directories, and is an absolute path, either \\UNC or D:\
-	  * We need to use strict-valid Windows Paths when calling Windows API calls so that we can support the long-path prefix \\?\
-	  */
-	static void NormalizeWindowsPath(FStringBuilderBase& Path, bool bIsFilename)
-	{
-		FPathViews::ToAbsolutePathInline(Path);
+	FPathViews::ToAbsolutePathInline(Path);
 
-		// NormalizeFilename was already called by ConvertRelativePathToFull, but we still need to do the extra steps in NormalizeDirectoryName if it is a directory
-		if (!bIsFilename)
-		{
-			FPathViews::NormalizeDirectoryName(Path);
-		}
+	// NormalizeFilename was already called by ConvertRelativePathToFull, but we still need to do the extra steps in NormalizeDirectoryName if it is a directory
+	if (!bIsFilename)
+	{
+		FPathViews::NormalizeDirectoryName(Path);
+	}
 
-		// Remove duplicate slashes
-		const bool bIsUNCPath = Path.ToView().StartsWith(TEXTVIEW("//"));
+	// Remove duplicate slashes
+	const bool bIsUNCPath = Path.ToView().StartsWith(TEXTVIEW("//"));
 		
-		FPathViews::RemoveDuplicateSlashes(Path);
+	FPathViews::RemoveDuplicateSlashes(Path);
 
-		if (bIsUNCPath)
-		{
-			// Keep // at the beginning.  If There are more than two / at the beginning, replace them with just //.
-			Path.Prepend(TEXTVIEW("/"));
-		}
-
-		// We now have a canonical, strict-valid, absolute Unreal Path.  Convert it to a Windows Path.
-		for (TCHAR& Char : TArrayView<TCHAR>(Path.GetData(), Path.Len()))
-		{
-			if (Char == TEXT('/'))
-			{
-				Char = TEXT('\\');
-			}
-		}
+	if (bIsUNCPath)
+	{
+		// Keep // at the beginning.  If There are more than two / at the beginning, replace them with just //.
+		Path.Prepend(TEXTVIEW("/"));
+	}
 
-		// Handle Windows paths with null-terminated length over MAX_PATH
-		if (Path.Len() >= MAX_PATH)
+	// We now have a canonical, strict-valid, absolute Unreal Path.  Convert it to a Windows Path.
+	for (TCHAR& Char : TArrayView<TCHAR>(Path.GetData(), Path.Len()))
+	{
+		if (Char == TEXT('/'))
 		{
-			if (bIsUNCPath)
-			{
-				Path.ReplaceAt(0, 1, TEXTVIEW("\\\\?\\UNC"));
-			}
-			else
-			{
-				Path.Prepend(TEXTVIEW("\\\\?\\"));
-			}
+			Char = TEXT('\\');
 		}
 	}
-	
-	class FNormalizedFilename : public TStringBuilder<256>
+
+	// Handle Windows paths with null-terminated length over MAX_PATH
+	if (Path.Len() >= MAX_PATH)
 	{
-	public:
-		explicit FNormalizedFilename(const TCHAR* Filename)
+		if (bIsUNCPath)
 		{
-			Append(Filename);
-			NormalizeWindowsPath(*this, /* bIsFilename */ true);
+			Path.ReplaceAt(0, 1, TEXTVIEW("\\\\?\\UNC"));
 		}
-
-		explicit FNormalizedFilename(FStringView Filename)
+		else
 		{
-			Append(Filename);
-			NormalizeWindowsPath(*this, /* bIsFilename */ true);
+			Path.Prepend(TEXTVIEW("\\\\?\\"));
 		}
+	}
+}
+	
+FWindowsPlatformFile::FNormalizedFilename::FNormalizedFilename(const TCHAR* Filename)
+{
+	Append(Filename);
+	NormalizeWindowsPath(*this, /* bIsFilename */ true);
+}
 
-		explicit FNormalizedFilename(FStringView Dir, FStringView Filename)
-		{
-			Append(Dir);
-			FPathViews::Append(*this, Filename);
-			NormalizeWindowsPath(*this, /* bIsFilename */ true);
-		}
-	};
+FWindowsPlatformFile::FNormalizedFilename::FNormalizedFilename(FStringView Filename)
+{
+	Append(Filename);
+	NormalizeWindowsPath(*this, /* bIsFilename */ true);
+}
 
-	class FNormalizedDirectory : public TStringBuilder<256>
-	{
-	public:
-		explicit FNormalizedDirectory(const TCHAR* Directory)
-		{
-			Append(Directory);
-			NormalizeWindowsPath(*this, /* bIsFilename */ false);
-		}
-	};
+FWindowsPlatformFile::FNormalizedFilename::FNormalizedFilename(FStringView Dir, FStringView Filename)
+{
+	Append(Dir);
+	FPathViews::Append(*this, Filename);
+	NormalizeWindowsPath(*this, /* bIsFilename */ true);
+}
 
-public:
-	//~ For visibility of overloads we don't override
-	using IPhysicalPlatformFile::IterateDirectory;
-	using IPhysicalPlatformFile::IterateDirectoryStat;
+FWindowsPlatformFile::FNormalizedDirectory::FNormalizedDirectory(const TCHAR* Directory)
+{
+	Append(Directory);
+	NormalizeWindowsPath(*this, /* bIsFilename */ false);
+}
 
-	virtual bool FileExists(const TCHAR* Filename) override
-	{
-		uint32 Result = GetFileAttributesW(*FNormalizedFilename(Filename));
-		if (Result != 0xFFFFFFFF && !(Result & FILE_ATTRIBUTE_DIRECTORY))
-		{
-			return true;
-		}
-		return false;
-	}
-	virtual int64 FileSize(const TCHAR* Filename) override
-	{
-		WIN32_FILE_ATTRIBUTE_DATA Info;
-		if (!!GetFileAttributesExW(*FNormalizedFilename(Filename), GetFileExInfoStandard, &Info))
-		{
-			if ((Info.dwFileAttributes & FILE_ATTRIBUTE_DIRECTORY) == 0)
-			{
-				LARGE_INTEGER li;
-				li.HighPart = Info.nFileSizeHigh;
-				li.LowPart = Info.nFileSizeLow;
-				return li.QuadPart;
-			}
-		}
-		return -1;
-	}
-	virtual bool DeleteFile(const TCHAR* Filename) override
+bool FWindowsPlatformFile::FileExists(const TCHAR* Filename)
+{
+	uint32 Result = GetFileAttributesW(*FNormalizedFilename(Filename));
+	if (Result != 0xFFFFFFFF && !(Result & FILE_ATTRIBUTE_DIRECTORY))
 	{
-		FNormalizedFilename NormalizedFilename(Filename);
-		return !!DeleteFileW(*NormalizedFilename);
+		return true;
 	}
-	virtual bool IsReadOnly(const TCHAR* Filename) override
+	return false;
+}
+
+int64 FWindowsPlatformFile::FileSize(const TCHAR* Filename)
+{
+	WIN32_FILE_ATTRIBUTE_DATA Info;
+	if (!!GetFileAttributesExW(*FNormalizedFilename(Filename), GetFileExInfoStandard, &Info))
 	{
-		uint32 Result = GetFileAttributesW(*FNormalizedFilename(Filename));
-		if (Result != 0xFFFFFFFF)
+		if ((Info.dwFileAttributes & FILE_ATTRIBUTE_DIRECTORY) == 0)
 		{
-			return !!(Result & FILE_ATTRIBUTE_READONLY);
+			LARGE_INTEGER li;
+			li.HighPart = Info.nFileSizeHigh;
+			li.LowPart = Info.nFileSizeLow;
+			return li.QuadPart;
 		}
-		return false;
-	}
-	virtual bool MoveFile(const TCHAR* To, const TCHAR* From) override
-	{
-		return !!MoveFileW(*FNormalizedFilename(From), *FNormalizedFilename(To));
 	}
-	virtual bool SetReadOnly(const TCHAR* Filename, bool bNewReadOnlyValue) override
+	return -1;
+}
+
+bool FWindowsPlatformFile::DeleteFile(const TCHAR* Filename)
+{
+	FNormalizedFilename NormalizedFilename(Filename);
+	return !!DeleteFileW(*NormalizedFilename);
+}
+
+bool FWindowsPlatformFile::IsReadOnly(const TCHAR* Filename)
+{
+	uint32 Result = GetFileAttributesW(*FNormalizedFilename(Filename));
+	if (Result != 0xFFFFFFFF)
 	{
-		return !!SetFileAttributesW(*FNormalizedFilename(Filename), bNewReadOnlyValue ? FILE_ATTRIBUTE_READONLY : FILE_ATTRIBUTE_NORMAL);
+		return !!(Result & FILE_ATTRIBUTE_READONLY);
 	}
+	return false;
+}
 
-	virtual FDateTime GetTimeStamp(const TCHAR* Filename) override
-	{
-		WIN32_FILE_ATTRIBUTE_DATA Info;
-		if (GetFileAttributesExW(*FNormalizedFilename(Filename), GetFileExInfoStandard, &Info))
-		{
-			return WindowsFileTimeToUEDateTime(Info.ftLastWriteTime);
-		}
+bool FWindowsPlatformFile::MoveFile(const TCHAR* To, const TCHAR* From)
+{
+	return !!MoveFileW(*FNormalizedFilename(From), *FNormalizedFilename(To));
+}
+
+bool FWindowsPlatformFile::SetReadOnly(const TCHAR* Filename, bool bNewReadOnlyValue)
+{
+	return !!SetFileAttributesW(*FNormalizedFilename(Filename), bNewReadOnlyValue ? FILE_ATTRIBUTE_READONLY : FILE_ATTRIBUTE_NORMAL);
+}
 
-		return FDateTime::MinValue();
+FDateTime FWindowsPlatformFile::GetTimeStamp(const TCHAR* Filename)
+{
+	WIN32_FILE_ATTRIBUTE_DATA Info;
+	if (GetFileAttributesExW(*FNormalizedFilename(Filename), GetFileExInfoStandard, &Info))
+	{
+		return UE::WindowsPlatformFile::Private::WindowsFileTimeToUEDateTime(Info.ftLastWriteTime);
 	}
 
-	virtual void SetTimeStamp(const TCHAR* Filename, FDateTime DateTime) override
+	return FDateTime::MinValue();
+}
+
+void FWindowsPlatformFile::SetTimeStamp(const TCHAR* Filename, FDateTime DateTime)
+{
+	TRACE_PLATFORMFILE_BEGIN_OPEN(Filename);
+	HANDLE Handle = CreateFileW(*FNormalizedFilename(Filename), FILE_WRITE_ATTRIBUTES, FILE_SHARE_READ, nullptr, OPEN_EXISTING, FILE_ATTRIBUTE_NORMAL | FILE_FLAG_OVERLAPPED, nullptr);
+	if (Handle != INVALID_HANDLE_VALUE)
 	{
-		TRACE_PLATFORMFILE_BEGIN_OPEN(Filename);
-		HANDLE Handle = CreateFileW(*FNormalizedFilename(Filename), FILE_WRITE_ATTRIBUTES, FILE_SHARE_READ, nullptr, OPEN_EXISTING, FILE_ATTRIBUTE_NORMAL | FILE_FLAG_OVERLAPPED, nullptr);
-		if (Handle != INVALID_HANDLE_VALUE)
-		{
-			TRACE_PLATFORMFILE_END_OPEN(Handle);
-			const FILETIME ModificationFileTime = UEDateTimeToWindowsFileTime(DateTime);
-			SetFileTime(Handle, nullptr, nullptr, &ModificationFileTime);
-			TRACE_PLATFORMFILE_BEGIN_CLOSE(Handle);
+		TRACE_PLATFORMFILE_END_OPEN(Handle);
+		const FILETIME ModificationFileTime = UE::WindowsPlatformFile::Private::UEDateTimeToWindowsFileTime(DateTime);
+		SetFileTime(Handle, nullptr, nullptr, &ModificationFileTime);
+		TRACE_PLATFORMFILE_BEGIN_CLOSE(Handle);
 #if PLATFORMFILETRACE_ENABLED
-			// MSVC static analysis has a rule that reports the argument to CloseHandle as uninitialized memory after the call to CloseHandle, so we have to save it ahead of time
-			uint64 SavedHandle = uint64(Handle);
+		// MSVC static analysis has a rule that reports the argument to CloseHandle as uninitialized memory after the call to CloseHandle, so we have to save it ahead of time
+		uint64 SavedHandle = uint64(Handle);
 #endif
-			BOOL CloseResult = CloseHandle(Handle);
+		BOOL CloseResult = CloseHandle(Handle);
 #if PLATFORMFILETRACE_ENABLED
-			if (CloseResult)
-			{
-				TRACE_PLATFORMFILE_END_CLOSE(SavedHandle);
-			}
-			else
-			{
-				TRACE_PLATFORMFILE_FAIL_CLOSE(SavedHandle);
-			}
-#else
-			(void)CloseResult;
-#endif
+		if (CloseResult)
+		{
+			TRACE_PLATFORMFILE_END_CLOSE(SavedHandle);
 		}
 		else
 		{
-			TRACE_PLATFORMFILE_FAIL_OPEN(Filename);
+			TRACE_PLATFORMFILE_FAIL_CLOSE(SavedHandle);
 		}
+#else
+		(void)CloseResult;
+#endif
 	}
-
-	virtual FDateTime GetAccessTimeStamp(const TCHAR* Filename) override
+	else
 	{
-		WIN32_FILE_ATTRIBUTE_DATA Info;
-		if (GetFileAttributesExW(*FNormalizedFilename(Filename), GetFileExInfoStandard, &Info))
-		{
-			return WindowsFileTimeToUEDateTime(Info.ftLastAccessTime);
-		}
+		TRACE_PLATFORMFILE_FAIL_OPEN(Filename);
+	}
+}
 
-		return FDateTime::MinValue();
+FDateTime FWindowsPlatformFile::GetAccessTimeStamp(const TCHAR* Filename)
+{
+	WIN32_FILE_ATTRIBUTE_DATA Info;
+	if (GetFileAttributesExW(*FNormalizedFilename(Filename), GetFileExInfoStandard, &Info))
+	{
+		return UE::WindowsPlatformFile::Private::WindowsFileTimeToUEDateTime(Info.ftLastAccessTime);
 	}
 
-	virtual FString GetFilenameOnDisk(const TCHAR* Filename) override
+	return FDateTime::MinValue();
+}
+
+FString FWindowsPlatformFile::GetFilenameOnDisk(const TCHAR* Filename)
+{
+	FString NormalizedFileName = *FNormalizedFilename(Filename);
+	TRACE_PLATFORMFILE_BEGIN_OPEN(Filename);
+	HANDLE hFile = CreateFile(*NormalizedFileName, FILE_READ_ATTRIBUTES, FILE_SHARE_READ | FILE_SHARE_WRITE | FILE_SHARE_DELETE, NULL, OPEN_EXISTING, 0, NULL);
+	// If the file exists on disk, read the capitalization from the path on disk, otherwise just return the (normalized) input filename
+	if (hFile != INVALID_HANDLE_VALUE)
 	{
-		FString NormalizedFileName = *FNormalizedFilename(Filename);
-		TRACE_PLATFORMFILE_BEGIN_OPEN(Filename);
-		HANDLE hFile = CreateFile(*NormalizedFileName, FILE_READ_ATTRIBUTES, FILE_SHARE_READ | FILE_SHARE_WRITE | FILE_SHARE_DELETE, NULL, OPEN_EXISTING, 0, NULL);
-		// If the file exists on disk, read the capitalization from the path on disk, otherwise just return the (normalized) input filename
-		if (hFile != INVALID_HANDLE_VALUE)
+		TRACE_PLATFORMFILE_END_OPEN(hFile);
+		for (uint32 Length = NormalizedFileName.Len() + 10;;)
 		{
-			TRACE_PLATFORMFILE_END_OPEN(hFile);
-			for (uint32 Length = NormalizedFileName.Len() + 10;;)
-			{
-				TArray<TCHAR, FString::AllocatorType>& CharArray = NormalizedFileName.GetCharArray();
-				CharArray.SetNum(Length);
+			TArray<TCHAR, FString::AllocatorType>& CharArray = NormalizedFileName.GetCharArray();
+			CharArray.SetNum(Length);
 
-				Length = GetFinalPathNameByHandle(hFile, CharArray.GetData(), CharArray.Num(), FILE_NAME_NORMALIZED);
-				if (Length == 0)
-				{
-					NormalizedFileName = *FNormalizedFilename(Filename);
-					break;
-				}
-				if (Length < (uint32)CharArray.Num())
-				{
-					CharArray.SetNum(Length + 1);
-					break;
-				}
-			}
-			TRACE_PLATFORMFILE_BEGIN_CLOSE(hFile);
-#if PLATFORMFILETRACE_ENABLED
-			// MSVC static analysis has a rule that reports the argument to CloseHandle as uninitialized memory after the call to CloseHandle, so we have to save it ahead of time
-			uint64 SavedHFile = uint64(hFile);
-#endif
-			BOOL CloseResult = CloseHandle(hFile);
-#if PLATFORMFILETRACE_ENABLED
-			if (CloseResult)
+			Length = GetFinalPathNameByHandle(hFile, CharArray.GetData(), CharArray.Num(), FILE_NAME_NORMALIZED);
+			if (Length == 0)
 			{
-				TRACE_PLATFORMFILE_END_CLOSE(SavedHFile);
+				NormalizedFileName = *FNormalizedFilename(Filename);
+				break;
 			}
-			else
+			if (Length < (uint32)CharArray.Num())
 			{
-				TRACE_PLATFORMFILE_FAIL_CLOSE(SavedHFile);
+				CharArray.SetNum(Length + 1);
+				break;
 			}
-#else
-			(void)CloseResult;
-#endif
-		}
-		else
-		{
-			TRACE_PLATFORMFILE_FAIL_OPEN(Filename);
 		}
-
-		// Remove the Windows device path prefix.
-		if (NormalizedFileName.StartsWith(TEXT("\\\\?\\UNC\\"), ESearchCase::CaseSensitive))
+		TRACE_PLATFORMFILE_BEGIN_CLOSE(hFile);
+#if PLATFORMFILETRACE_ENABLED
+		// MSVC static analysis has a rule that reports the argument to CloseHandle as uninitialized memory after the call to CloseHandle, so we have to save it ahead of time
+		uint64 SavedHFile = uint64(hFile);
+#endif
+		BOOL CloseResult = CloseHandle(hFile);
+#if PLATFORMFILETRACE_ENABLED
+		if (CloseResult)
 		{
-			NormalizedFileName.RemoveAt(2, 6); // remove ?\UNC\ to convert \\?\UNC\Path\... to \\Path\...
+			TRACE_PLATFORMFILE_END_CLOSE(SavedHFile);
 		}
 		else
 		{
-			NormalizedFileName.RemoveFromStart(TEXT("\\\\?\\"), ESearchCase::CaseSensitive);
+			TRACE_PLATFORMFILE_FAIL_CLOSE(SavedHFile);
 		}
+#else
+		(void)CloseResult;
+#endif
+	}
+	else
+	{
+		TRACE_PLATFORMFILE_FAIL_OPEN(Filename);
+	}
 
-		// Convert the result back into an UnrealPath (\\ -> /)
-		NormalizedFileName.ReplaceCharInline(TEXT('\\'), TEXT('/'), ESearchCase::CaseSensitive);
-
-		return NormalizedFileName;
+	// Remove the Windows device path prefix.
+	if (NormalizedFileName.StartsWith(TEXT("\\\\?\\UNC\\"), ESearchCase::CaseSensitive))
+	{
+		NormalizedFileName.RemoveAt(2, 6); // remove ?\UNC\ to convert \\?\UNC\Path\... to \\Path\...
 	}
+	else
+	{
+		NormalizedFileName.RemoveFromStart(TEXT("\\\\?\\"), ESearchCase::CaseSensitive);
+	}
+
+	// Convert the result back into an UnrealPath (\\ -> /)
+	NormalizedFileName.ReplaceCharInline(TEXT('\\'), TEXT('/'), ESearchCase::CaseSensitive);
+
+	return NormalizedFileName;
+}
 
-	virtual ESymlinkResult IsSymlink(const TCHAR* Filename) override
+ESymlinkResult FWindowsPlatformFile::IsSymlink(const TCHAR* Filename)
+{
+	const uint32 Result = GetFileAttributesW(Filename);
+	if (Result != 0xFFFFFFFF && (Result & FILE_ATTRIBUTE_REPARSE_POINT))
 	{
-		const uint32 Result = GetFileAttributesW(Filename);
-		if (Result != 0xFFFFFFFF && (Result & FILE_ATTRIBUTE_REPARSE_POINT))
-		{
-			return ESymlinkResult::Symlink;
-		}
-		return ESymlinkResult::NonSymlink;
+		return ESymlinkResult::Symlink;
 	}
+	return ESymlinkResult::NonSymlink;
+}
+
+bool FWindowsPlatformFile::HasMarkOfTheWeb(FStringView Filename, FString* OutSourceURL)
+{
+	FNormalizedFilename StreamPath(Filename);
+	StreamPath.Append(FileConstants::ZoneIdentifierStreamSuffix);
 
-	virtual bool HasMarkOfTheWeb(FStringView Filename, FString* OutSourceURL = nullptr) override
+	const uint32 ZoneId = GetPrivateProfileIntW(TEXT("ZoneTransfer"), TEXT("ZoneId"), -1 /* URLZONE_INVALID */, *StreamPath);
+	if (ZoneId != 2 /* URLZONE_TRUSTED */ && ZoneId != 3 /* URLZONE_INTERNET */ && ZoneId != 4 /* URLZONE_UNTRUSTED */)
 	{
-		FNormalizedFilename StreamPath(Filename);
-		StreamPath.Append(FileConstants::ZoneIdentifierStreamSuffix);
+		return false;
+	}
 
-		const uint32 ZoneId = GetPrivateProfileIntW(TEXT("ZoneTransfer"), TEXT("ZoneId"), -1 /* URLZONE_INVALID */, *StreamPath);
-		if (ZoneId != 2 /* URLZONE_TRUSTED */ && ZoneId != 3 /* URLZONE_INTERNET */ && ZoneId != 4 /* URLZONE_UNTRUSTED */)
+	if (OutSourceURL != nullptr)
+	{
+		// Can't use FConfigFile to parse ini because it treats double slashes in the url (e.g. HostUrl=http://...) as the beginning of a comment.
+		TCHAR HostUrl[2048]; // @todo: remove hardcoded url length limit by using a simple standalone ini parser
+		const int32 HostUrlLen = static_cast<int32>(GetPrivateProfileStringW(TEXT("ZoneTransfer"), TEXT("HostUrl"), nullptr, HostUrl, static_cast<DWORD>(GetNum(HostUrl)), *StreamPath));
+		if (0 < HostUrlLen && HostUrlLen <= GetNum(HostUrl))
 		{
-			return false;
+			*OutSourceURL = FString(HostUrlLen, HostUrl);
 		}
+	}
 
-		if (OutSourceURL != nullptr)
-		{
-			// Can't use FConfigFile to parse ini because it treats double slashes in the url (e.g. HostUrl=http://...) as the beginning of a comment.
-			TCHAR HostUrl[2048]; // @todo: remove hardcoded url length limit by using a simple standalone ini parser
-			const int32 HostUrlLen = static_cast<int32>(GetPrivateProfileStringW(TEXT("ZoneTransfer"), TEXT("HostUrl"), nullptr, HostUrl, static_cast<DWORD>(GetNum(HostUrl)), *StreamPath));
-			if (0 < HostUrlLen && HostUrlLen <= GetNum(HostUrl))
-			{
-				*OutSourceURL = FString(HostUrlLen, HostUrl);
-			}
-		}
+	return true;
+}
 
-		return true;
-	}
+bool FWindowsPlatformFile::SetMarkOfTheWeb(FStringView Filename, bool bNewStatus, const FString* InSourceURL)
+{
+	FNormalizedFilename StreamPath(Filename);
+	StreamPath.Append(FileConstants::ZoneIdentifierStreamSuffix);
 
-	virtual bool SetMarkOfTheWeb(FStringView Filename, bool bNewStatus, const FString* InSourceURL = nullptr) override
+	if (!bNewStatus)
 	{
-		FNormalizedFilename StreamPath(Filename);
-		StreamPath.Append(FileConstants::ZoneIdentifierStreamSuffix);
+		return !!DeleteFileW(*StreamPath) || GetLastError() == ERROR_FILE_NOT_FOUND;
+	}
 
-		if (!bNewStatus)
-		{
-			return !!DeleteFileW(*StreamPath) || GetLastError() == ERROR_FILE_NOT_FOUND;
-		}
+	if (!WritePrivateProfileStringW(TEXT("ZoneTransfer"), TEXT("ZoneId"), TEXT("3") /* URLZONE_INTERNET */, *StreamPath))
+	{
+		return false;
+	}
 
-		if (!WritePrivateProfileStringW(TEXT("ZoneTransfer"), TEXT("ZoneId"), TEXT("3") /* URLZONE_INTERNET */, *StreamPath))
+	if (InSourceURL != nullptr)
+	{
+		if (!WritePrivateProfileStringW(TEXT("ZoneTransfer"), TEXT("HostUrl"), **InSourceURL, *StreamPath))
 		{
 			return false;
 		}
-
-		if (InSourceURL != nullptr)
-		{
-			if (!WritePrivateProfileStringW(TEXT("ZoneTransfer"), TEXT("HostUrl"), **InSourceURL, *StreamPath))
-			{
-				return false;
-			}
-		}
-
-		return true;
 	}
 
-#define USE_WINDOWS_ASYNC_IMPL 0
+	return true;
+}
+
 #if USE_WINDOWS_ASYNC_IMPL
-	virtual IAsyncReadFileHandle* OpenAsyncRead(const TCHAR* Filename) override
-	{
-		uint32  Access = GENERIC_READ;
-		uint32  WinFlags = FILE_SHARE_READ;
-		uint32  Create = OPEN_EXISTING;
+IAsyncReadFileHandle* FWindowsPlatformFile::OpenAsyncRead(const TCHAR* Filename)
+{
+	uint32  Access = GENERIC_READ;
+	uint32  WinFlags = FILE_SHARE_READ;
+	uint32  Create = OPEN_EXISTING;
 
 
-		FString NormalizedFilename = WindowsNormalizedFilename(Filename);
-		TRACE_PLATFORMFILE_BEGIN_OPEN(Filename);
-		HANDLE Handle = CreateFileW(*NormalizedFilename, Access, WinFlags, NULL, Create, FILE_ATTRIBUTE_NORMAL | FILE_FLAG_OVERLAPPED, NULL);
+	FString NormalizedFilename = WindowsNormalizedFilename(Filename);
+	TRACE_PLATFORMFILE_BEGIN_OPEN(Filename);
+	HANDLE Handle = CreateFileW(*NormalizedFilename, Access, WinFlags, NULL, Create, FILE_ATTRIBUTE_NORMAL | FILE_FLAG_OVERLAPPED, NULL);
 #if PLATFORMFILETRACE_ENABLED
-		if (Handle != INVALID_HANDLE_VALUE)
-		{
-			TRACE_PLATFORMFILE_END_OPEN(Handle);
-		}
-		else
-		{
-			TRACE_PLATFORMFILE_FAIL_OPEN(Filename);
-		}
+	if (Handle != INVALID_HANDLE_VALUE)
+	{
+		TRACE_PLATFORMFILE_END_OPEN(Handle);
+	}
+	else
+	{
+		TRACE_PLATFORMFILE_FAIL_OPEN(Filename);
+	}
 #endif
 
-		// we can't really fail here because this is intended to be an async open
-		return new FMicrosoftAsyncReadFileHandle(Handle, *NormalizedFilename);
+	// we can't really fail here because this is intended to be an async open
+	return new FMicrosoftAsyncReadFileHandle(Handle, *NormalizedFilename);
 
-	}
+}
 #endif
 
-	virtual IFileHandle* OpenRead(const TCHAR* Filename, bool bAllowWrite = false) override
-	{
-		uint32  Access    = GENERIC_READ;
-		uint32  WinFlags  = FILE_SHARE_READ | (bAllowWrite ? FILE_SHARE_WRITE : 0);
-		uint32  Create    = OPEN_EXISTING;
+IFileHandle* FWindowsPlatformFile::OpenRead(const TCHAR* Filename, bool bAllowWrite)
+{
+	uint32  Access    = GENERIC_READ;
+	uint32  WinFlags  = FILE_SHARE_READ | (bAllowWrite ? FILE_SHARE_WRITE : 0);
+	uint32  Create    = OPEN_EXISTING;
 #define USE_OVERLAPPED_IO (!IS_PROGRAM && !WITH_EDITOR)		// Use straightforward synchronous I/O in cooker/editor
 
-		TRACE_PLATFORMFILE_BEGIN_OPEN(Filename);
+	TRACE_PLATFORMFILE_BEGIN_OPEN(Filename);
 #if USE_OVERLAPPED_IO
-		HANDLE Handle    = CreateFileW(*FNormalizedFilename(Filename), Access, WinFlags, NULL, Create, FILE_ATTRIBUTE_NORMAL | FILE_FLAG_OVERLAPPED, NULL);
-		if (Handle != INVALID_HANDLE_VALUE)
-		{
-			TRACE_PLATFORMFILE_END_OPEN(Handle);
-			return new FAsyncBufferedFileReaderWindows(Handle, Access, WinFlags, FILE_FLAG_OVERLAPPED);
-		}
+	HANDLE Handle    = CreateFileW(*FNormalizedFilename(Filename), Access, WinFlags, NULL, Create, FILE_ATTRIBUTE_NORMAL | FILE_FLAG_OVERLAPPED, NULL);
+	if (Handle != INVALID_HANDLE_VALUE)
+	{
+		TRACE_PLATFORMFILE_END_OPEN(Handle);
+		return new FAsyncBufferedFileReaderWindows(Handle, Access, WinFlags, FILE_FLAG_OVERLAPPED);
+	}
 #else
-		HANDLE Handle = CreateFileW(*FNormalizedFilename(Filename), Access, WinFlags, NULL, Create, FILE_ATTRIBUTE_NORMAL, NULL);
-		if (Handle != INVALID_HANDLE_VALUE)
-		{
-			TRACE_PLATFORMFILE_END_OPEN(Handle);
-
-			FFileHandleWindows* FileHandle = new FFileHandleWindows(Handle, Access, WinFlags, 0);
+	HANDLE Handle = CreateFileW(*FNormalizedFilename(Filename), Access, WinFlags, NULL, Create, FILE_ATTRIBUTE_NORMAL, NULL);
+	if (Handle != INVALID_HANDLE_VALUE)
+	{
+		TRACE_PLATFORMFILE_END_OPEN(Handle);
 
-			// Some operations can fail during the handle initialization, so we
-			// double check that the handle is valid before returning it
-			if (FileHandle->IsValid())
-			{
-				return FileHandle;
-			}
-			else
-			{
-				delete FileHandle;
+		FFileHandleWindows* FileHandle = new FFileHandleWindows(Handle, Access, WinFlags, 0);
 
-				return nullptr;
-			}
+		// Some operations can fail during the handle initialization, so we
+		// double check that the handle is valid before returning it
+		if (FileHandle->IsValid())
+		{
+			return FileHandle;
 		}
-#endif
 		else
 		{
-			TRACE_PLATFORMFILE_FAIL_OPEN(Filename);
+			delete FileHandle;
+
 			return nullptr;
 		}
 	}
-
-	virtual IFileHandle* OpenReadNoBuffering(const TCHAR* Filename, bool bAllowWrite = false) override
+#endif
+	else
 	{
-		uint32  Access = GENERIC_READ;
-		uint32  WinFlags = FILE_SHARE_READ | (bAllowWrite ? FILE_SHARE_WRITE : 0);
-		uint32  Create = OPEN_EXISTING;
-		TRACE_PLATFORMFILE_BEGIN_OPEN(Filename);
-		HANDLE Handle = CreateFileW(*FNormalizedFilename(Filename), Access, WinFlags, NULL, Create, FILE_ATTRIBUTE_NORMAL | FILE_FLAG_OVERLAPPED, NULL);
-		if (Handle != INVALID_HANDLE_VALUE)
-		{
-			TRACE_PLATFORMFILE_END_OPEN(Handle);
+		TRACE_PLATFORMFILE_FAIL_OPEN(Filename);
+		return nullptr;
+	}
+}
 
-			FFileHandleWindows* FileHandle = new FFileHandleWindows(Handle, Access, WinFlags, FILE_FLAG_OVERLAPPED);
+IFileHandle* FWindowsPlatformFile::OpenReadNoBuffering(const TCHAR* Filename, bool bAllowWrite)
+{
+	uint32  Access = GENERIC_READ;
+	uint32  WinFlags = FILE_SHARE_READ | (bAllowWrite ? FILE_SHARE_WRITE : 0);
+	uint32  Create = OPEN_EXISTING;
+	TRACE_PLATFORMFILE_BEGIN_OPEN(Filename);
+	HANDLE Handle = CreateFileW(*FNormalizedFilename(Filename), Access, WinFlags, NULL, Create, FILE_ATTRIBUTE_NORMAL | FILE_FLAG_OVERLAPPED, NULL);
+	if (Handle != INVALID_HANDLE_VALUE)
+	{
+		TRACE_PLATFORMFILE_END_OPEN(Handle);
 
-			// Some operations can fail during the handle initialization, so we
-			// double check that the handle is valid before returning it
-			if (FileHandle->IsValid())
-			{
-				return FileHandle;
-			}
-			else
-			{
-				delete FileHandle;
+		FFileHandleWindows* FileHandle = new FFileHandleWindows(Handle, Access, WinFlags, FILE_FLAG_OVERLAPPED);
 
-				return nullptr;
-			}
+		// Some operations can fail during the handle initialization, so we
+		// double check that the handle is valid before returning it
+		if (FileHandle->IsValid())
+		{
+			return FileHandle;
 		}
 		else
 		{
-			TRACE_PLATFORMFILE_FAIL_OPEN(Filename);
+			delete FileHandle;
+
 			return nullptr;
 		}
 	}
+	else
+	{
+		TRACE_PLATFORMFILE_FAIL_OPEN(Filename);
+		return nullptr;
+	}
+}
 
-	virtual IFileHandle* OpenWrite(const TCHAR* Filename, bool bAppend = false, bool bAllowRead = false) override
+IFileHandle* FWindowsPlatformFile::OpenWrite(const TCHAR* Filename, bool bAppend, bool bAllowRead)
+{
+	uint32  Access    = GENERIC_WRITE | (bAllowRead ? GENERIC_READ : 0);
+	uint32  WinFlags  = bAllowRead ? FILE_SHARE_READ : 0;
+	uint32  Create    = bAppend ? OPEN_ALWAYS : CREATE_ALWAYS;
+	TRACE_PLATFORMFILE_BEGIN_OPEN(Filename);
+	HANDLE Handle    = CreateFileW(*FNormalizedFilename(Filename), Access, WinFlags, NULL, Create, FILE_ATTRIBUTE_NORMAL, NULL);
+	if(Handle != INVALID_HANDLE_VALUE)
 	{
-		uint32  Access    = GENERIC_WRITE | (bAllowRead ? GENERIC_READ : 0);
-		uint32  WinFlags  = bAllowRead ? FILE_SHARE_READ : 0;
-		uint32  Create    = bAppend ? OPEN_ALWAYS : CREATE_ALWAYS;
-		TRACE_PLATFORMFILE_BEGIN_OPEN(Filename);
-		HANDLE Handle    = CreateFileW(*FNormalizedFilename(Filename), Access, WinFlags, NULL, Create, FILE_ATTRIBUTE_NORMAL, NULL);
-		if(Handle != INVALID_HANDLE_VALUE)
-		{
-			TRACE_PLATFORMFILE_END_OPEN(Handle);
-			FFileHandleWindows* PlatformFileHandle = new FFileHandleWindows(Handle, Access, WinFlags, 0);
+		TRACE_PLATFORMFILE_END_OPEN(Handle);
+		FFileHandleWindows* PlatformFileHandle = new FFileHandleWindows(Handle, Access, WinFlags, 0);
 
-			// Some operations can fail during the handle initialization, so we
-			// double check that the handle is valid before returning it
-			if (PlatformFileHandle->IsValid())
-			{
-				if (bAppend)
-				{
-					PlatformFileHandle->SeekFromEnd(0);
-				}
-				return PlatformFileHandle;
-			}
-			else
+		// Some operations can fail during the handle initialization, so we
+		// double check that the handle is valid before returning it
+		if (PlatformFileHandle->IsValid())
+		{
+			if (bAppend)
 			{
-				delete PlatformFileHandle;
-
-				return nullptr;
+				PlatformFileHandle->SeekFromEnd(0);
 			}
+			return PlatformFileHandle;
 		}
 		else
 		{
-			TRACE_PLATFORMFILE_FAIL_OPEN(Filename);
+			delete PlatformFileHandle;
+
 			return nullptr;
 		}
 	}
+	else
+	{
+		TRACE_PLATFORMFILE_FAIL_OPEN(Filename);
+		return nullptr;
+	}
+}
 
-
-	virtual IMappedFileHandle* OpenMapped(const TCHAR* Filename) override
+IMappedFileHandle* FWindowsPlatformFile::OpenMapped(const TCHAR* Filename)
+{
+	int64 Size = FileSize(Filename);
+	if (Size < 1)
 	{
-		int64 Size = FileSize(Filename);
-		if (Size < 1)
-		{
-			return nullptr;
-		}
-		uint32  Access = GENERIC_READ;
-		uint32  WinFlags = FILE_SHARE_READ;
-		uint32  Create = OPEN_EXISTING;
-		TRACE_PLATFORMFILE_BEGIN_OPEN(Filename);
-		HANDLE Handle = CreateFileW(*FNormalizedFilename(Filename), Access, WinFlags, NULL, Create, FILE_ATTRIBUTE_NORMAL, NULL);
-		if (Handle != INVALID_HANDLE_VALUE)
+		return nullptr;
+	}
+	uint32  Access = GENERIC_READ;
+	uint32  WinFlags = FILE_SHARE_READ;
+	uint32  Create = OPEN_EXISTING;
+	TRACE_PLATFORMFILE_BEGIN_OPEN(Filename);
+	HANDLE Handle = CreateFileW(*FNormalizedFilename(Filename), Access, WinFlags, NULL, Create, FILE_ATTRIBUTE_NORMAL, NULL);
+	if (Handle != INVALID_HANDLE_VALUE)
+	{
+		TRACE_PLATFORMFILE_END_OPEN(Handle);
+	}
+	else
+	{
+		TRACE_PLATFORMFILE_FAIL_OPEN(Filename);
+		return nullptr;
+	}
+	HANDLE MappingHandle = CreateFileMapping(Handle, NULL, PAGE_READONLY, 0, 0, NULL);
+	if (MappingHandle == NULL)
+	{
+		TRACE_PLATFORMFILE_BEGIN_CLOSE(Handle);
+#if PLATFORMFILETRACE_ENABLED
+		// MSVC static analysis has a rule that reports the argument to CloseHandle as uninitialized memory after the call to CloseHandle, so we have to save it ahead of time
+		uint64 SavedHandle = uint64(Handle);
+#endif
+		BOOL CloseResult = CloseHandle(Handle);
+#if PLATFORMFILETRACE_ENABLED
+		if (CloseResult)
 		{
-			TRACE_PLATFORMFILE_END_OPEN(Handle);
+			TRACE_PLATFORMFILE_END_CLOSE(SavedHandle);
 		}
 		else
 		{
-			TRACE_PLATFORMFILE_FAIL_OPEN(Filename);
-			return nullptr;
+			TRACE_PLATFORMFILE_FAIL_CLOSE(SavedHandle);
 		}
-		HANDLE MappingHandle = CreateFileMapping(Handle, NULL, PAGE_READONLY, 0, 0, NULL);
-		if (MappingHandle == NULL)
-		{
-			TRACE_PLATFORMFILE_BEGIN_CLOSE(Handle);
-#if PLATFORMFILETRACE_ENABLED
-			// MSVC static analysis has a rule that reports the argument to CloseHandle as uninitialized memory after the call to CloseHandle, so we have to save it ahead of time
-			uint64 SavedHandle = uint64(Handle);
-#endif
-			BOOL CloseResult = CloseHandle(Handle);
-#if PLATFORMFILETRACE_ENABLED
-			if (CloseResult)
-			{
-				TRACE_PLATFORMFILE_END_CLOSE(SavedHandle);
-			}
-			else
-			{
-				TRACE_PLATFORMFILE_FAIL_CLOSE(SavedHandle);
-			}
 #else
-			(void)CloseResult;
+		(void)CloseResult;
 #endif
-			return nullptr;
-		}
-		return new FMappedFileHandleWindows(Handle, MappingHandle, Size, Filename);
+		return nullptr;
 	}
-	virtual bool DirectoryExists(const TCHAR* Directory) override
+	return new FMappedFileHandleWindows(Handle, MappingHandle, Size, Filename);
+}
+
+bool FWindowsPlatformFile::DirectoryExists(const TCHAR* Directory)
+{
+	// Empty Directory is the current directory so assume it always exists.
+	bool bExists = !FCString::Strlen(Directory);
+	if (!bExists) 
 	{
-		// Empty Directory is the current directory so assume it always exists.
-		bool bExists = !FCString::Strlen(Directory);
-		if (!bExists) 
-		{
-			uint32 Result = GetFileAttributesW(*FNormalizedDirectory(Directory));
-			bExists = (Result != 0xFFFFFFFF && (Result & FILE_ATTRIBUTE_DIRECTORY));
-		}
-		return bExists;
+		uint32 Result = GetFileAttributesW(*FNormalizedDirectory(Directory));
+		bExists = (Result != 0xFFFFFFFF && (Result & FILE_ATTRIBUTE_DIRECTORY));
 	}
-	virtual bool CreateDirectory(const TCHAR* Directory) override
+	return bExists;
+}
+
+bool FWindowsPlatformFile::CreateDirectory(const TCHAR* Directory)
+{
+	return CreateDirectoryW(*FNormalizedDirectory(Directory), NULL) || GetLastError() == ERROR_ALREADY_EXISTS;
+}
+
+bool FWindowsPlatformFile::DeleteDirectory(const TCHAR* Directory)
+{
+	RemoveDirectoryW(*FNormalizedDirectory(Directory));
+	uint32 LastError = GetLastError();
+	const bool bSucceeded = !DirectoryExists(Directory);
+	if (!bSucceeded)
 	{
-		return CreateDirectoryW(*FNormalizedDirectory(Directory), NULL) || GetLastError() == ERROR_ALREADY_EXISTS;
+		SetLastError(LastError);
 	}
-	virtual bool DeleteDirectory(const TCHAR* Directory) override
+	return bSucceeded;
+}
+
+FFileStatData FWindowsPlatformFile::GetStatData(const TCHAR* FilenameOrDirectory)
+{
+	WIN32_FILE_ATTRIBUTE_DATA Info;
+	if (GetFileAttributesExW(*FNormalizedFilename(FilenameOrDirectory), GetFileExInfoStandard, &Info))
 	{
-		RemoveDirectoryW(*FNormalizedDirectory(Directory));
-		uint32 LastError = GetLastError();
-		const bool bSucceeded = !DirectoryExists(Directory);
-		if (!bSucceeded)
+		const bool bIsDirectory = !!(Info.dwFileAttributes & FILE_ATTRIBUTE_DIRECTORY);
+
+		int64 FileSize = -1;
+		if (!bIsDirectory)
 		{
-			SetLastError(LastError);
+			LARGE_INTEGER li;
+			li.HighPart = Info.nFileSizeHigh;
+			li.LowPart = Info.nFileSizeLow;
+			FileSize = static_cast<int64>(li.QuadPart);
 		}
-		return bSucceeded;
+
+		return FFileStatData(
+			UE::WindowsPlatformFile::Private::WindowsFileTimeToUEDateTime(Info.ftCreationTime),
+			UE::WindowsPlatformFile::Private::WindowsFileTimeToUEDateTime(Info.ftLastAccessTime),
+			UE::WindowsPlatformFile::Private::WindowsFileTimeToUEDateTime(Info.ftLastWriteTime),
+			FileSize, 
+			bIsDirectory,
+			!!(Info.dwFileAttributes & FILE_ATTRIBUTE_READONLY)
+			);
 	}
-	virtual FFileStatData GetStatData(const TCHAR* FilenameOrDirectory) override
-	{
-		WIN32_FILE_ATTRIBUTE_DATA Info;
-		if (GetFileAttributesExW(*FNormalizedFilename(FilenameOrDirectory), GetFileExInfoStandard, &Info))
-		{
-			const bool bIsDirectory = !!(Info.dwFileAttributes & FILE_ATTRIBUTE_DIRECTORY);
 
-			int64 FileSize = -1;
-			if (!bIsDirectory)
-			{
-				LARGE_INTEGER li;
-				li.HighPart = Info.nFileSizeHigh;
-				li.LowPart = Info.nFileSizeLow;
-				FileSize = static_cast<int64>(li.QuadPart);
-			}
+	return FFileStatData();
+}
 
-			return FFileStatData(
-				WindowsFileTimeToUEDateTime(Info.ftCreationTime),
-				WindowsFileTimeToUEDateTime(Info.ftLastAccessTime),
-				WindowsFileTimeToUEDateTime(Info.ftLastWriteTime),
-				FileSize, 
-				bIsDirectory,
-				!!(Info.dwFileAttributes & FILE_ATTRIBUTE_READONLY)
-				);
-		}
+bool FWindowsPlatformFile::IterateDirectory(const TCHAR* Directory, FDirectoryVisitor& Visitor)
+{
+	TStringBuilder<256> DirTemp;
+	DirTemp << Directory;
+	const int32 DirLen = DirTemp.Len();
+	return IterateDirectoryCommon(Directory, [&](const WIN32_FIND_DATAW& InData) -> bool
+	{
+		DirTemp.RemoveSuffix(DirTemp.Len() - DirLen);
+		FPathViews::Append(DirTemp, InData.cFileName);
+		const bool bIsDirectory = !!(InData.dwFileAttributes & FILE_ATTRIBUTE_DIRECTORY);
+		return Visitor.CallShouldVisitAndVisit(*DirTemp, bIsDirectory);
+	});
+}
 
-		return FFileStatData();
-	}
-	virtual bool IterateDirectory(const TCHAR* Directory, FDirectoryVisitor& Visitor) override
-	{
-		TStringBuilder<256> DirTemp;
-		DirTemp << Directory;
-		const int32 DirLen = DirTemp.Len();
-		return IterateDirectoryCommon(Directory, [&](const WIN32_FIND_DATAW& InData) -> bool
-		{
-			DirTemp.RemoveSuffix(DirTemp.Len() - DirLen);
-			FPathViews::Append(DirTemp, InData.cFileName);
-			const bool bIsDirectory = !!(InData.dwFileAttributes & FILE_ATTRIBUTE_DIRECTORY);
-			return Visitor.CallShouldVisitAndVisit(*DirTemp, bIsDirectory);
-		});
-	}
-	virtual bool IterateDirectoryStat(const TCHAR* Directory, FDirectoryStatVisitor& Visitor) override
+bool FWindowsPlatformFile::IterateDirectoryStat(const TCHAR* Directory, FDirectoryStatVisitor& Visitor)
+{
+	const FString DirectoryStr = Directory;
+	return IterateDirectoryCommon(Directory, [&](const WIN32_FIND_DATAW& InData) -> bool
 	{
-		const FString DirectoryStr = Directory;
-		return IterateDirectoryCommon(Directory, [&](const WIN32_FIND_DATAW& InData) -> bool
-		{
-			const bool bIsDirectory = !!(InData.dwFileAttributes & FILE_ATTRIBUTE_DIRECTORY);
+		const bool bIsDirectory = !!(InData.dwFileAttributes & FILE_ATTRIBUTE_DIRECTORY);
 
-			int64 FileSize = -1;
-			if (!bIsDirectory)
-			{
-				LARGE_INTEGER li;
-				li.HighPart = InData.nFileSizeHigh;
-				li.LowPart = InData.nFileSizeLow;
-				FileSize = static_cast<int64>(li.QuadPart);
-			}
+		int64 FileSize = -1;
+		if (!bIsDirectory)
+		{
+			LARGE_INTEGER li;
+			li.HighPart = InData.nFileSizeHigh;
+			li.LowPart = InData.nFileSizeLow;
+			FileSize = static_cast<int64>(li.QuadPart);
+		}
 
-			return Visitor.CallShouldVisitAndVisit(
-				*(DirectoryStr / InData.cFileName), 
-				FFileStatData(
-					WindowsFileTimeToUEDateTime(InData.ftCreationTime),
-					WindowsFileTimeToUEDateTime(InData.ftLastAccessTime),
-					WindowsFileTimeToUEDateTime(InData.ftLastWriteTime),
-					FileSize, 
-					bIsDirectory,
-					!!(InData.dwFileAttributes & FILE_ATTRIBUTE_READONLY)
-					)
-				);
-		});
-	}
+		return Visitor.CallShouldVisitAndVisit(
+			*(DirectoryStr / InData.cFileName), 
+			FFileStatData(
+				UE::WindowsPlatformFile::Private::WindowsFileTimeToUEDateTime(InData.ftCreationTime),
+				UE::WindowsPlatformFile::Private::WindowsFileTimeToUEDateTime(InData.ftLastAccessTime),
+				UE::WindowsPlatformFile::Private::WindowsFileTimeToUEDateTime(InData.ftLastWriteTime),
+				FileSize, 
+				bIsDirectory,
+				!!(InData.dwFileAttributes & FILE_ATTRIBUTE_READONLY)
+				)
+			);
+	});
+}
 	
-	// Outline to reduce stack space usage since IterateDirectoryCommon might be recursive
-	FORCENOINLINE static HANDLE FindFirstFileWithWildcard(const TCHAR* Directory, WIN32_FIND_DATAW& OutData)
-	{
-		return FindFirstFileW(*(FNormalizedFilename(Directory, TEXTVIEW("*.*"))), &OutData);
-	}
+// Forced not inline to reduce stack space usage since IterateDirectoryCommon might be recursive
+FORCENOINLINE HANDLE FWindowsPlatformFile::FindFirstFileWithWildcard(const TCHAR* Directory, WIN32_FIND_DATAW& OutData)
+{
+	return FindFirstFileW(*(FNormalizedFilename(Directory, TEXTVIEW("*.*"))), &OutData);
+}
 
-	bool IterateDirectoryCommon(const TCHAR* Directory, const TFunctionRef<bool(const WIN32_FIND_DATAW&)>& Visitor)
+bool FWindowsPlatformFile::IterateDirectoryCommon(const TCHAR* Directory, const TFunctionRef<bool(const WIN32_FIND_DATAW&)>& Visitor)
+{
+	bool bResult = true;
+	WIN32_FIND_DATAW Data;
+	HANDLE Handle = FindFirstFileWithWildcard(Directory, /* Out */ Data);
+	if (Handle != INVALID_HANDLE_VALUE)
 	{
-		bool bResult = true;
-		WIN32_FIND_DATAW Data;
-		HANDLE Handle = FindFirstFileWithWildcard(Directory, /* Out */ Data);
-		if (Handle != INVALID_HANDLE_VALUE)
+		do
 		{
-			do
+			if (FCString::Strcmp(Data.cFileName, TEXT(".")) && FCString::Strcmp(Data.cFileName, TEXT("..")))
 			{
-				if (FCString::Strcmp(Data.cFileName, TEXT(".")) && FCString::Strcmp(Data.cFileName, TEXT("..")))
-				{
-					bResult = Visitor(Data);
-				}
-			} while (bResult && FindNextFileW(Handle, &Data));
-			FindClose(Handle);
-		}
-		return bResult;
+				bResult = Visitor(Data);
+			}
+		} while (bResult && FindNextFileW(Handle, &Data));
+		FindClose(Handle);
 	}
-};
+	return bResult;
+}
 
 IPlatformFile& IPlatformFile::GetPlatformPhysical()
 {
diff --git a/Engine/Source/Runtime/Core/Private/Windows/WindowsPlatformFilePrivate.h b/Engine/Source/Runtime/Core/Private/Windows/WindowsPlatformFilePrivate.h
new file mode 100644
--- /dev/null
+++ b/Engine/Source/Runtime/Core/Private/Windows/WindowsPlatformFilePrivate.h
@@ -0,0 +1,90 @@
+// Copyright Epic Games, Inc. All Rights Reserved.
+
+#pragma once
+
+#include "Windows/WindowsPlatformFile.h"
+
+#include "Containers/StringView.h"
+#include "GenericPlatform/GenericPlatformFile.h"
+#include "Misc/DateTime.h"
+#include "Misc/StringBuilder.h"
+#include "Windows/WindowsHWrapper.h"
+
+#define USE_WINDOWS_ASYNC_IMPL 0
+#define USE_OVERLAPPED_IO (!IS_PROGRAM && !WITH_EDITOR)		// Use straightforward synchronous I/O in cooker/editor
+
+/**
+ * Windows File I/O implementation
+**/
+class FWindowsPlatformFile : public IPhysicalPlatformFile
+{
+private:
+	/**
+	  * Convert from a valid Unreal Path to a canonical and strict-valid Windows Path.
+	  * An Unreal Path may have either \ or / and may have empty directories (two / in a row), and may have .. and may be relative
+	  * A canonical and strict-valid Windows Path has only \, does not have .., does not have empty directories, and is an absolute path, either \\UNC or D:\
+	  * We need to use strict-valid Windows Paths when calling Windows API calls so that we can support the long-path prefix \\?\
+	  */
+	static void NormalizeWindowsPath(FStringBuilderBase& Path, bool bIsFilename);
+
+	class FNormalizedFilename : public TStringBuilder<256>
+	{
+	public:
+		explicit FNormalizedFilename(const TCHAR* Filename);
+		explicit FNormalizedFilename(FStringView Filename);
+		explicit FNormalizedFilename(FStringView Dir, FStringView Filename);
+	};
+
+	class FNormalizedDirectory : public TStringBuilder<256>
+	{
+	public:
+		explicit FNormalizedDirectory(const TCHAR* Directory);
+	};
+
+public:
+	//~ For visibility of overloads we don't override
+	using IPhysicalPlatformFile::IterateDirectory;
+	using IPhysicalPlatformFile::IterateDirectoryStat;
+
+	virtual bool FileExists(const TCHAR* Filename) override;
+	virtual int64 FileSize(const TCHAR* Filename) override;
+	virtual bool DeleteFile(const TCHAR* Filename) override;
+	virtual bool IsReadOnly(const TCHAR* Filename) override;
+	virtual bool MoveFile(const TCHAR* To, const TCHAR* From) override;
+	virtual bool SetReadOnly(const TCHAR* Filename, bool bNewReadOnlyValue) override;
+	virtual FDateTime GetTimeStamp(const TCHAR* Filename) override;
+	virtual void SetTimeStamp(const TCHAR* Filename, FDateTime DateTime) override;
+	virtual FDateTime GetAccessTimeStamp(const TCHAR* Filename) override;
+	virtual FString GetFilenameOnDisk(const TCHAR* Filename) override;
+	virtual ESymlinkResult IsSymlink(const TCHAR* Filename) override;
+	virtual bool HasMarkOfTheWeb(FStringView Filename, FString* OutSourceURL = nullptr) override;
+	virtual bool SetMarkOfTheWeb(FStringView Filename, bool bNewStatus, const FString* InSourceURL = nullptr) override;
+
+#if USE_WINDOWS_ASYNC_IMPL
+	virtual IAsyncReadFileHandle* OpenAsyncRead(const TCHAR* Filename) override;
+#endif
+
+	virtual IFileHandle* OpenRead(const TCHAR* Filename, bool bAllowWrite = false) override;
+	virtual IFileHandle* OpenReadNoBuffering(const TCHAR* Filename, bool bAllowWrite = false) override;
+	virtual IFileHandle* OpenWrite(const TCHAR* Filename, bool bAppend = false, bool bAllowRead = false) override;
+	virtual IMappedFileHandle* OpenMapped(const TCHAR* Filename) override;
+	virtual bool DirectoryExists(const TCHAR* Directory) override;
+	virtual bool CreateDirectory(const TCHAR* Directory) override;
+	virtual bool DeleteDirectory(const TCHAR* Directory) override;
+	virtual FFileStatData GetStatData(const TCHAR* FilenameOrDirectory) override;
+	virtual bool IterateDirectory(const TCHAR* Directory, FDirectoryVisitor& Visitor) override;
+	virtual bool IterateDirectoryStat(const TCHAR* Directory, FDirectoryStatVisitor& Visitor) override;
+
+	// Forced not inline to reduce stack space usage since IterateDirectoryCommon might be recursive
+	FORCENOINLINE static HANDLE FindFirstFileWithWildcard(const TCHAR* Directory, WIN32_FIND_DATAW& OutData);
+	bool IterateDirectoryCommon(const TCHAR* Directory, const TFunctionRef<bool(const WIN32_FIND_DATAW&)>& Visitor);
+};
+
+namespace UE::WindowsPlatformFile::Private
+{
+
+int32 UEDayOfWeekToWindowsSystemTimeDayOfWeek(const EDayOfWeek InDayOfWeek);
+FDateTime WindowsFileTimeToUEDateTime(const FILETIME& InFileTime);
+FILETIME UEDateTimeToWindowsFileTime(const FDateTime& InDateTime);
+
+} // namespace UE::WindowsPlatformFile::Private
\ No newline at end of file
