WindowsMemoryTrace: Fix compilation errors in _WIN32_WINNT>=0xA00.
#rnx

diff --git a/Engine/Source/Runtime/Core/Private/ProfilingDebugging/Microsoft/WindowsMemoryTrace.cpp b/Engine/Source/Runtime/Core/Private/ProfilingDebugging/Microsoft/WindowsMemoryTrace.cpp
--- a/Engine/Source/Runtime/Core/Private/ProfilingDebugging/Microsoft/WindowsMemoryTrace.cpp
+++ b/Engine/Source/Runtime/Core/Private/ProfilingDebugging/Microsoft/WindowsMemoryTrace.cpp
@@ -21,6 +21,10 @@
 #include <winternl.h>
 #include <intrin.h>
 
+#if (NTDDI_VERSION >= NTDDI_WIN10_RS4)
+#pragma comment(lib, "mincore.lib") // VirtualAlloc2
+#endif
+
 ////////////////////////////////////////////////////////////////////////////////
 FMalloc* MemoryTrace_CreateInternal(FMalloc*, int32, const WIDECHAR* const*);
 
@@ -233,23 +237,35 @@ private:
 	static bool				bLight;
 	static LPVOID WINAPI	VmAlloc(LPVOID Address, SIZE_T Size, DWORD Type, DWORD Protect);
 	static LPVOID WINAPI	VmAllocEx(HANDLE Process, LPVOID Address, SIZE_T Size, DWORD Type, DWORD Protect);
-	static LPVOID WINAPI	VmAlloc2(HANDLE Process, LPVOID BaseAddress, SIZE_T Size, ULONG AllocationType, ULONG PageProtection, /*MEM_EXTENDED_PARAMETER* */ void* ExtendedParameters, ULONG ParameterCount);
+#if (NTDDI_VERSION >= NTDDI_WIN10_RS4)
+	static PVOID WINAPI		VmAlloc2(HANDLE Process, PVOID BaseAddress, SIZE_T Size, ULONG AllocationType,
+								ULONG PageProtection, MEM_EXTENDED_PARAMETER* ExtendedParameters, ULONG ParameterCount);
+	static PVOID(WINAPI* VmAlloc2Orig)(HANDLE, PVOID, SIZE_T, ULONG, ULONG, MEM_EXTENDED_PARAMETER*, ULONG);
+	typedef PVOID(__stdcall* FnVirtualAlloc2)(HANDLE, PVOID, SIZE_T, ULONG, ULONG, MEM_EXTENDED_PARAMETER*, ULONG);
+#else
+	static LPVOID WINAPI	VmAlloc2(HANDLE Process, LPVOID BaseAddress, SIZE_T Size, ULONG AllocationType,
+								ULONG PageProtection, void* ExtendedParameters, ULONG ParameterCount);
+	static LPVOID(WINAPI* VmAlloc2Orig)(HANDLE, LPVOID, SIZE_T, ULONG, ULONG, /*MEM_EXTENDED_PARAMETER* */ void*, ULONG);
+	typedef LPVOID(__stdcall* FnVirtualAlloc2)(HANDLE, LPVOID, SIZE_T, ULONG, ULONG, /* MEM_EXTENDED_PARAMETER* */ void*, ULONG);
+#endif
 	static BOOL WINAPI		VmFree(LPVOID Address, SIZE_T Size, DWORD Type);
 	static BOOL WINAPI		VmFreeEx(HANDLE Process, LPVOID Address, SIZE_T Size, DWORD Type);
 	static LPVOID			(WINAPI *VmAllocOrig)(LPVOID, SIZE_T, DWORD, DWORD);
 	static LPVOID			(WINAPI *VmAllocExOrig)(HANDLE, LPVOID, SIZE_T, DWORD, DWORD);
-	static LPVOID			(WINAPI *VmAlloc2Orig)(HANDLE, LPVOID, SIZE_T, ULONG, ULONG, /*MEM_EXTENDED_PARAMETER* */ void*, ULONG);
 	static BOOL				(WINAPI *VmFreeOrig)(LPVOID, SIZE_T, DWORD);
 	static BOOL				(WINAPI *VmFreeExOrig)(HANDLE, LPVOID, SIZE_T, DWORD);
 
-	typedef LPVOID(__stdcall* FnVirtualAlloc2)(HANDLE, LPVOID, SIZE_T, ULONG, ULONG, /* MEM_EXTENDED_PARAMETER* */ void*, ULONG);
 };
 
 ////////////////////////////////////////////////////////////////////////////////
 bool	FVirtualWinApiHooks::bLight;
 LPVOID	(WINAPI *FVirtualWinApiHooks::VmAllocOrig)(LPVOID, SIZE_T, DWORD, DWORD);
 LPVOID	(WINAPI *FVirtualWinApiHooks::VmAllocExOrig)(HANDLE, LPVOID, SIZE_T, DWORD, DWORD);
-LPVOID	(WINAPI *FVirtualWinApiHooks::VmAlloc2Orig)(HANDLE, LPVOID, SIZE_T, ULONG, ULONG, /*MEM_EXTENDED_PARAMETER* */ void*, ULONG);
+#if (NTDDI_VERSION >= NTDDI_WIN10_RS4)
+PVOID	(WINAPI* FVirtualWinApiHooks::VmAlloc2Orig)(HANDLE, PVOID, SIZE_T, ULONG, ULONG, MEM_EXTENDED_PARAMETER*, ULONG);
+#else
+LPVOID	(WINAPI* FVirtualWinApiHooks::VmAlloc2Orig)(HANDLE, LPVOID, SIZE_T, ULONG, ULONG, /*MEM_EXTENDED_PARAMETER* */ void*, ULONG);
+#endif
 BOOL	(WINAPI *FVirtualWinApiHooks::VmFreeOrig)(LPVOID, SIZE_T, DWORD);
 BOOL	(WINAPI *FVirtualWinApiHooks::VmFreeExOrig)(HANDLE, LPVOID, SIZE_T, DWORD);
 
@@ -353,7 +369,13 @@ BOOL WINAPI FVirtualWinApiHooks::VmFreeEx(HANDLE Process, LPVOID Address, SIZE_T
 }
 
 ////////////////////////////////////////////////////////////////////////////////
-LPVOID WINAPI FVirtualWinApiHooks::VmAlloc2(HANDLE Process, LPVOID BaseAddress, SIZE_T Size, ULONG Type, ULONG PageProtection, /*MEM_EXTENDED_PARAMETER* */ void* ExtendedParameters, ULONG ParameterCount)
+#if (NTDDI_VERSION >= NTDDI_WIN10_RS4)
+PVOID WINAPI FVirtualWinApiHooks::VmAlloc2(HANDLE Process, PVOID BaseAddress, SIZE_T Size, ULONG Type,
+	ULONG PageProtection, MEM_EXTENDED_PARAMETER* ExtendedParameters, ULONG ParameterCount)
+#else
+LPVOID WINAPI FVirtualWinApiHooks::VmAlloc2(HANDLE Process, LPVOID BaseAddress, SIZE_T Size, ULONG Type,
+	ULONG PageProtection, /*MEM_EXTENDED_PARAMETER* */ void* ExtendedParameters, ULONG ParameterCount)
+#endif
 {
 	LPVOID Ret = VmAlloc2Orig(Process, BaseAddress, Size, Type, PageProtection, ExtendedParameters, ParameterCount);
 
