PathViews: Add IsDriveSpecifierWithoutRoot and SplitVolumeSpecifier.
#rnx

diff --git a/Engine/Source/Runtime/Core/Private/Misc/PathViews.cpp b/Engine/Source/Runtime/Core/Private/Misc/PathViews.cpp
--- a/Engine/Source/Runtime/Core/Private/Misc/PathViews.cpp
+++ b/Engine/Source/Runtime/Core/Private/Misc/PathViews.cpp
@@ -692,6 +692,165 @@ void FPathViews::SplitFirstComponent(FStringView InPath, FStringView& OutFirstCo
 	}
 }
 
+bool FPathViews::IsDriveSpecifierWithoutRoot(FStringView InPath)
+{
+	const int32 PathLen = InPath.Len();
+	const TCHAR* const PathData = InPath.GetData();
+	for (int32 Index = 0; Index < PathLen; ++Index)
+	{
+		TCHAR C = PathData[Index];
+		if (UE4PathViews_Private::IsSlashOrBackslash(C))
+		{
+			// '/' or '/root' or 'root/' or 'root/remainder'
+			return false;
+		}
+		if (C == ':')
+		{
+			++Index;
+			if (Index == PathLen)
+			{
+				// 'D:'
+				return true;
+			}
+			C = PathData[Index];
+			if (C == ':')
+			{
+				// 'D::root'
+				// Path is even more invalid: two colons are not allowed. Arbitrarily return false for this case.
+				return false;
+			}
+			if (UE4PathViews_Private::IsSlashOrBackslash(C))
+			{
+				// 'D:/root'
+				return false;
+			}
+			// 'D:root'
+			return true;
+		}
+	}
+	// 'root'
+	return false;
+}
+
+void FPathViews::SplitVolumeSpecifier(FStringView InPath, FStringView& OutVolumeSpecifier, FStringView& OutRemainder)
+{
+	const int32 PathLen = InPath.Len();
+	const TCHAR* const PathData = InPath.GetData();
+
+	if (PathLen == 0)
+	{
+		OutVolumeSpecifier.Reset();
+		OutRemainder.Reset();
+		return;
+	}
+
+	if (PathData[0] == ':')
+	{
+		OutVolumeSpecifier = InPath.Left(1);
+		if (PathLen == 1)
+		{
+			// ':'
+			OutRemainder.Reset();
+		}
+		else
+		{
+			// ':remainder'
+			OutRemainder = InPath.RightChop(1);
+		}
+		return;
+	}
+
+	if (PathLen == 1)
+	{
+		// '/' or 'D'
+		OutVolumeSpecifier.Reset();
+		OutRemainder = InPath;
+		return;
+	}
+
+	if (UE4PathViews_Private::IsSlashOrBackslash(PathData[0]))
+	{
+		if (UE4PathViews_Private::IsSlashOrBackslash(PathData[1]))
+		{
+			// '//' '///' or '//volume or //volume/remainder
+			if (PathLen == 2)
+			{
+				OutVolumeSpecifier = InPath;
+				OutRemainder.Reset();
+				return;
+			}
+
+			// //////volume/remainder -> { '//////volume', 'remainder' }
+			int32 SlashSlashEnd = 2;
+			while (SlashSlashEnd < PathLen && UE4PathViews_Private::IsSlashOrBackslash(PathData[SlashSlashEnd]))
+			{
+				++SlashSlashEnd;
+			}
+
+			const TCHAR* VolumeEnd = Algo::FindByPredicate(FStringView(PathData + SlashSlashEnd, PathLen-SlashSlashEnd),
+				UE4PathViews_Private::IsSlashOrBackslash);
+			if (!VolumeEnd)
+			{
+				// '//' or '//volume'
+				OutVolumeSpecifier = InPath;
+				OutRemainder.Reset();
+				return;
+			}
+
+			// '//' or '/////' or '//volume/remainder
+			OutVolumeSpecifier = InPath.Left(static_cast<int32>(VolumeEnd - PathData));
+			OutRemainder = InPath.RightChop(OutVolumeSpecifier.Len());
+			return;
+		}
+
+		// '/remainder' or '/:'
+		OutVolumeSpecifier.Reset();
+		OutRemainder = InPath;
+		return;
+	}
+
+	if (PathData[1] == ':')
+	{
+		// 'D:' or 'D:remainder'
+		OutVolumeSpecifier = InPath.Left(2);
+		OutRemainder = InPath.RightChop(2);
+		return;
+	}
+
+	// root or root/remainder or drive:remainder
+	for (int32 Index = 2; Index < PathLen; ++Index)
+	{
+		TCHAR C = PathData[Index];
+		if (UE4PathViews_Private::IsSlashOrBackslash(C))
+		{
+			// 'root/remainder'
+			OutVolumeSpecifier.Reset();
+			OutRemainder = InPath;
+			return;
+		}
+		if (C == ':')
+		{
+			++Index;
+			OutVolumeSpecifier = InPath.Left(Index);
+			if (Index == PathLen)
+			{
+				// 'drive:'
+				OutRemainder.Reset();
+			}
+			else
+			{
+				// 'drive:remainder'
+				OutRemainder = InPath.RightChop(Index);
+			}
+			return;
+		}
+	}
+
+	// 'root'
+	OutVolumeSpecifier.Reset();
+	OutRemainder = InPath;
+}
+
 void FPathViews::AppendPath(FStringBuilderBase& InOutPath, FStringView AppendPath)
 {
 	using namespace UE4PathViews_Private;
diff --git a/Engine/Source/Runtime/Core/Public/Misc/PathViews.h b/Engine/Source/Runtime/Core/Public/Misc/PathViews.h
--- a/Engine/Source/Runtime/Core/Public/Misc/PathViews.h
+++ b/Engine/Source/Runtime/Core/Public/Misc/PathViews.h
@@ -117,6 +117,25 @@ public:
 	 */
 	static CORE_API bool IsPathLeaf(FStringView InPath);
 
+	/**
+	 * Report whether the given path is an invalid path because it has a drive specifier (':') without a following
+	 * path separator indicating the root of that drive. In general usage on Windows, drive specifiers without a
+	 * following path separator are interpreted to mean the current working directory of the given drive. But that
+	 * context is not always applicable and will need to be handled in system-specific ways.
+	 * @see SplitVolumeSpecifier for callers who want to followup by modifing the path.
+	 * 
+	 * D:		-> true
+	 * D:/		-> false
+	 * D:root	-> true
+	 * D:/root	-> false
+	 * //volume -> false
+	 * /root    -> false
+	 * relpath  -> false
+	 * 
+	 * See FPathViewsVolumeSpecifierTest for further edgecases.
+	 */
+	static CORE_API bool IsDriveSpecifierWithoutRoot(FStringView InPath);
+
 	/**
 	 * Splits InPath into individual directory components, and calls ComponentVisitor on each.
 	 *
@@ -341,6 +360,24 @@ public:
 	 */
 	static CORE_API void SplitFirstComponent(FStringView InPath, FStringView& OutFirstComponent, FStringView& OutRemainder);
 
+	/**
+	 * Split the path into a volume specifier and the rest of the path.
+	 * This function handles drive specifiers without a root in a specific way, see the examples and
+	 * @see IsDriveSpecifierWithoutRoot
+	 *
+	 * <emptystring>      -> { '', '' }
+	 * D:		          -> { 'D:', '' }
+	 * D:root/path        -> { 'D:', 'root/path' }
+	 * D:/root/path       -> { 'D:', '/root/path' }
+	 * //volume           -> { '//volume', '' }
+	 * //volume/root/path -> { '//volume', '/root/path' }
+	 * root/path          -> { '', 'root/path' }
+	 * /root/path         -> { '', '/root/path' }
+	 * 
+	 * See FPathViewsVolumeSpecifierTest for further edgecases.
+	 */
+	static CORE_API void SplitVolumeSpecifier(FStringView InPath, FStringView& OutVolumeSpecifier, FStringView& OutRemainder);
+
 	/**
 	 * If AppendPath is a relative path, append it as a relative path onto InOutPath.
 	 * If AppendPath is absolute, reset InOutPath and replace it with RelPath.
diff --git a/Engine/Source/Runtime/Core/Tests/Misc/PathViewsTest.cpp b/Engine/Source/Runtime/Core/Tests/Misc/PathViewsTest.cpp
--- a/Engine/Source/Runtime/Core/Tests/Misc/PathViewsTest.cpp
+++ b/Engine/Source/Runtime/Core/Tests/Misc/PathViewsTest.cpp
@@ -856,4 +856,60 @@ TEST_CASE_NAMED(FPathViewsToAbsoluteTest, "System::Core::Misc::PathViews::ToAbso
 
 }
 
+TEST_CASE_NAMED(FPathViewsVolumeSpecifierTest, "System::Core::Misc::PathViews::VolumeSpecifier", "[ApplicationContextMask][SmokeFilter]")
+{
+	using namespace PathTest;
+
+
+	struct FTestCase
+	{
+		FStringView Input;
+		bool bDriveSpecifier;
+		FStringView Volume;
+		FStringView Remainder;
+	};
+	FTestCase TestCases[] = {
+		{ TEXTVIEW(""),					false,	TEXTVIEW(""),			TEXTVIEW("") },
+		{ TEXTVIEW("D:"),				true,	TEXTVIEW("D:"),			TEXTVIEW("") },
+		{ TEXTVIEW("D:/"),				false,	TEXTVIEW("D:"),			TEXTVIEW("/") },
+		{ TEXTVIEW("D:\\"),				false,	TEXTVIEW("D:"),			TEXTVIEW("\\") },
+		{ TEXTVIEW("D:root/path"),		true,	TEXTVIEW("D:"),			TEXTVIEW("root/path") },
+		{ TEXTVIEW("D:/root/path"),		false,	TEXTVIEW("D:"),			TEXTVIEW("/root/path") },
+		{ TEXTVIEW("D:\\root\\path"),	false,	TEXTVIEW("D:"),			TEXTVIEW("\\root\\path") },
+		{ TEXTVIEW("//volume"),			false,	TEXTVIEW("//volume"),	TEXTVIEW("") },
+		{ TEXTVIEW("\\\\volume"),		false,	TEXTVIEW("\\\\volume"),	TEXTVIEW("") },
+		{ TEXTVIEW("/\\volume"),		false,	TEXTVIEW("/\\volume"),	TEXTVIEW("") }, // Poorly defined case, somewhat arbitrary
+		{ TEXTVIEW("\\/volume"),		false,	TEXTVIEW("\\/volume"),	TEXTVIEW("") }, // Poorly defined case, somewhat arbitrary
+		{ TEXTVIEW("//volume/"),		false,	TEXTVIEW("//volume"),	TEXTVIEW("/") },
+		{ TEXTVIEW("//volume/root"),	false,	TEXTVIEW("//volume"),	TEXTVIEW("/root") },
+		{ TEXTVIEW("/root/path"),		false,	TEXTVIEW(""),			TEXTVIEW("/root/path") },
+		{ TEXTVIEW("\\root\\path"),		false,	TEXTVIEW(""),			TEXTVIEW("\\root\\path") },
+		{ TEXTVIEW("root/path"),		false,	TEXTVIEW(""),			TEXTVIEW("root/path") },
+		{ TEXTVIEW("/"),				false,	TEXTVIEW(""),			TEXTVIEW("/") },
+		{ TEXTVIEW("\\"),				false,	TEXTVIEW(""),			TEXTVIEW("\\") },
+		{ TEXTVIEW("//"),				false,	TEXTVIEW("//"),			TEXTVIEW("") }, // Poorly defined case, somewhat arbitrary
+		{ TEXTVIEW("\\\\"),				false,	TEXTVIEW("\\\\"),		TEXTVIEW("") }, // Poorly defined case, somewhat arbitrary
+		{ TEXTVIEW("/\\"),				false,	TEXTVIEW("/\\"),		TEXTVIEW("") }, // Poorly defined case, somewhat arbitrary
+		{ TEXTVIEW("\\/"),				false,	TEXTVIEW("\\/"),		TEXTVIEW("") }, // Poorly defined case, somewhat arbitrary
+		{ TEXTVIEW("/:"),				false,	TEXTVIEW(""),			TEXTVIEW("/:") }, // Poorly defined case, somewhat arbitrary
+		{ TEXTVIEW(":"),				true,	TEXTVIEW(":"),			TEXTVIEW("") }, // Poorly defined case, somewhat arbitrary
+		{ TEXTVIEW(":/"),				false,	TEXTVIEW(":"),			TEXTVIEW("/") }, // Poorly defined case, somewhat arbitrary
+		{ TEXTVIEW(":root"),			true,	TEXTVIEW(":"),			TEXTVIEW("root") }, // Poorly defined case, somewhat arbitrary
+		{ TEXTVIEW("////volume/path"),	false,	TEXTVIEW("////volume"),	TEXTVIEW("/path") }, // Poorly defined case, somewhat arbitrary, @see RemoveDuplicateSlashes
+	};
+
+	for (const FTestCase& TestCase : TestCases)
+	{
+		bool bDriveSpecifier;
+		FStringView Volume;
+		FStringView Remainder;
+		bDriveSpecifier = FPathViews::IsDriveSpecifierWithoutRoot(TestCase.Input);
+		FPathViews::SplitVolumeSpecifier(TestCase.Input, Volume, Remainder);
+		CHECK_EQUALS(FString(TestCase.Input), bDriveSpecifier, TestCase.bDriveSpecifier);
+		CHECK_EQUALS(FString(TestCase.Input), Volume, TestCase.Volume);
+		CHECK_EQUALS(FString(TestCase.Input), Remainder, TestCase.Remainder);
+	};
+}
+
+
 #endif //WITH_TESTS
