Fixed uninitialized variable in FInterpCurve<T>::FindNearestOnSegment() when BreakBeforeEvaluate() returns true on the first iteration.

diff --git a/Engine/Source/Runtime/Core/Public/Math/InterpCurve.h b/Engine/Source/Runtime/Core/Public/Math/InterpCurve.h
--- a/Engine/Source/Runtime/Core/Public/Math/InterpCurve.h
+++ b/Engine/Source/Runtime/Core/Public/Math/InterpCurve.h
@@ -685,7 +685,7 @@ float FInterpCurve<T>::FindNearestOnSegment(const T& PointInSpace, int32 PtIdx,
 			[&PointsT, &PtIdx, &NextPtIdx, &PointInSpace, &Diff, MaxIteration, InvThree]
 			(float& Value, float& Move, TFunctionRef<bool(float&, const float)> BreakBeforeEvaluate, TFunctionRef<bool(float&, const float)> BreakAfterEvaluate, const float Tolerance) -> float
 		{
-			T FoundPoint;
+			T FoundPoint = PointInSpace;
 
 			for (int32 Iter = 0; Iter < MaxIteration; ++Iter)
 			{
