Improve timing logs

diff --git a/Engine/Source/Runtime/Core/Private/GenericPlatform/GenericPlatformTime.cpp b/Engine/Source/Runtime/Core/Private/GenericPlatform/GenericPlatformTime.cpp
--- a/Engine/Source/Runtime/Core/Private/GenericPlatform/GenericPlatformTime.cpp
+++ b/Engine/Source/Runtime/Core/Private/GenericPlatform/GenericPlatformTime.cpp
@@ -132,7 +132,9 @@ FString FGenericPlatformTime::PrettyTime( double Seconds )
 	}
 	else if ( Seconds < 1.0 )
 	{
-		return FString::Printf( TEXT("%d ms"), FMath::TruncToInt((float)(Seconds*1000)) );
+		int32 Ms = FMath::TruncToInt((float)(Seconds * 1000));
+		int32 Us = FMath::TruncToInt((float)(Seconds * 1000000)) - Ms*1000;
+		return FString::Printf( TEXT("%d.%03d ms"), Ms, Us );
 	}
 	else if ( Seconds < 10.0 )
 	{
