Fix 30s being displayed as 0.5 min in UE
In CalculateDisplayUnit, if a unit is lower than 1 and the other greater than 1, choose the greater than 1. If both at lower, use the biggest one.
Fix loop always using the last 2 elements of the array
Always call FixedDisplayUnits, else the unit is not updated until the next refresh when we switch between 59s and 60s
#rnx
#tests PIE

diff --git a/Engine/Source/Editor/PropertyEditor/Private/UserInterface/PropertyEditor/SPropertyEditorNumeric.h b/Engine/Source/Editor/PropertyEditor/Private/UserInterface/PropertyEditor/SPropertyEditorNumeric.h
--- a/Engine/Source/Editor/PropertyEditor/Private/UserInterface/PropertyEditor/SPropertyEditorNumeric.h
+++ b/Engine/Source/Editor/PropertyEditor/Private/UserInterface/PropertyEditor/SPropertyEditorNumeric.h
@@ -506,7 +506,7 @@ private:
 			EPropertyValueSetFlags::Type Flags = (EPropertyValueSetFlags::InteractiveChange | EPropertyValueSetFlags::NotTransactable);
 			PropertyHandle->SetValue( NewValue, Flags );
 
-			if (TypeInterface.IsValid() && !TypeInterface->FixedDisplayUnits.IsSet())
+			if (TypeInterface.IsValid())
 			{
 				TypeInterface->SetupFixedDisplay(NewValue);
 			}
@@ -528,7 +528,7 @@ private:
 				LastSliderCommittedValue = NewValue;
 			}
 
-			if (TypeInterface.IsValid() && !TypeInterface->FixedDisplayUnits.IsSet())
+			if (TypeInterface.IsValid())
 			{
 				TypeInterface->SetupFixedDisplay(NewValue);
 			}
diff --git a/Engine/Source/Runtime/Core/Public/Math/UnitConversion.inl b/Engine/Source/Runtime/Core/Public/Math/UnitConversion.inl
--- a/Engine/Source/Runtime/Core/Public/Math/UnitConversion.inl
+++ b/Engine/Source/Runtime/Core/Public/Math/UnitConversion.inl
@@ -190,7 +190,7 @@ EUnit FUnitConversion::CalculateDisplayUnit(T Value, EUnit InUnits)
 	{
 		return QuantizeUnitsToBestFit(Value, InUnits).Units;
 	}
-	else if (DisplayUnits.Num() == 1)
+	if (DisplayUnits.Num() == 1)
 	{
 		return DisplayUnits[0];
 	}
@@ -204,16 +204,25 @@ EUnit FUnitConversion::CalculateDisplayUnit(T Value, EUnit InUnits)
 	int32 BestIndex = 0;
 	for (int32 Index = 0; Index < DisplayUnits.Num() - 1; ++Index)
 	{
-		double This = Convert(Value, InUnits, DisplayUnits[Index]);
-		double Next = Convert(Value, InUnits, DisplayUnits[Index + 1]);
-
-		if (FMath::Abs(FMath::LogX(10.0f, (float)This)) < FMath::Abs(FMath::LogX(10.0f, (float)Next)))
+		const T Best = FMath::Abs(Convert(Value, InUnits, DisplayUnits[BestIndex]));
+		const T Next = FMath::Abs(Convert(Value, InUnits, DisplayUnits[Index + 1]));
+		if (Best < 1.0 && Next >= 1.0)
 		{
-			BestIndex = Index;
+			BestIndex = Index + 1;
 		}
-		else
+		else if (Best < 1.0 && Next < 1.0)
 		{
-			BestIndex = Index + 1;
+			if (Next > Best)
+			{
+				BestIndex = Index + 1;
+			}
+		}
+		else if (Best >= 1.0 && Next >= 1.0)
+		{
+			if (FMath::LogX(10.0f, Next) < FMath::LogX(10.0f, Best))
+			{
+				BestIndex = Index + 1;
+			}
 		}
 	}
 
