Do not use FMath::Abs if T is not signed
#rnx

diff --git a/Engine/Source/Runtime/Core/Public/Math/UnitConversion.inl b/Engine/Source/Runtime/Core/Public/Math/UnitConversion.inl
--- a/Engine/Source/Runtime/Core/Public/Math/UnitConversion.inl
+++ b/Engine/Source/Runtime/Core/Public/Math/UnitConversion.inl
@@ -204,8 +204,18 @@ EUnit FUnitConversion::CalculateDisplayUnit(T Value, EUnit InUnits)
 	int32 BestIndex = 0;
 	for (int32 Index = 0; Index < DisplayUnits.Num() - 1; ++Index)
 	{
-		const T Best = FMath::Abs(Convert(Value, InUnits, DisplayUnits[BestIndex]));
-		const T Next = FMath::Abs(Convert(Value, InUnits, DisplayUnits[Index + 1]));
+		T Best, Next;
+		if constexpr (std::is_signed_v<T>)
+		{
+			Best = FMath::Abs(Convert(Value, InUnits, DisplayUnits[BestIndex]));
+			Next = FMath::Abs(Convert(Value, InUnits, DisplayUnits[Index + 1]));
+		}
+		else
+		{
+			Best = Convert(Value, InUnits, DisplayUnits[BestIndex]);
+			Next = Convert(Value, InUnits, DisplayUnits[Index + 1]);
+		}
+
 		if (Best < 1.0 && Next >= 1.0)
 		{
 			BestIndex = Index + 1;
