- Made EditorSettings.ini and EditorLayout.ini be fully saved, instead of needing per-section overrides (Editor.ini is setill per-section controlled)
- Added BasePlugin.ini (plugin .ini files now use that instead of Base.ini) so that plugin .ini's are fully saved, but a plugin can override it in the plugin's Default .ini file

diff --git a/Engine/Config/BaseEditorKeyBindings.ini b/Engine/Config/BaseEditorKeyBindings.ini
--- a/Engine/Config/BaseEditorKeyBindings.ini
+++ b/Engine/Config/BaseEditorKeyBindings.ini
@@ -1,6 +1,3 @@
-[SectionsToSave]
-+Section=UserDefinedChords
-
 [/Script/UnrealEd.UnrealEdKeyBindings]
 ; --- Curve Editor Bindings
 +KeyBindings=(bCtrlDown=False,bAltDown=False,bShiftDown=False,Key="one",CommandName="CurveEditor_ChangeInterpModeAUTO")
diff --git a/Engine/Config/BaseEditorSettings.ini b/Engine/Config/BaseEditorSettings.ini
--- a/Engine/Config/BaseEditorSettings.ini
+++ b/Engine/Config/BaseEditorSettings.ini
@@ -2,9 +2,6 @@
 ; These settings are overridden by a project's DefaultEditorSettings.ini file, per-user editor setting defaults are in *EditorPerProjectUserSettings.ini
 ; Most of these options are modifiable from Project Settings in the editor, and are not available in a packaged build
 
-[SectionsToSave]
-+Section=/Script/UnrealEd.ContentBrowserSettings
-
 [Internationalization]
 ShouldUseLocalizedNumericInput=True
 ShouldUseLocalizedPropertyNames=True
diff --git a/Engine/Config/PluginBase.ini b/Engine/Config/PluginBase.ini
new file mode 100644
--- /dev/null
+++ b/Engine/Config/PluginBase.ini
@@ -0,0 +1,8 @@
+; Unlike normal config files like Engine, etc, Plugins defaults to saving everything
+; There are usually far fewer sections, and old plugins that aren't updated will
+; still work. See Base.ini for more info
+[SectionsToSave]
+bCanSaveAllSections=true
+
+; always save this, even if a plugin sets bCanSaveAllSections to false
++Section=CurrentIniVersion
diff --git a/Engine/Source/Runtime/Core/Private/Misc/ConfigContext.cpp b/Engine/Source/Runtime/Core/Private/Misc/ConfigContext.cpp
--- a/Engine/Source/Runtime/Core/Private/Misc/ConfigContext.cpp
+++ b/Engine/Source/Runtime/Core/Private/Misc/ConfigContext.cpp
@@ -389,10 +389,14 @@ bool FConfigContext::PerformLoad()
 		ConfigFile->bHasPlatformName = true;
 
 		// chcek if the config file wants to save all sections
-		bool bLocalSaveAllSections;
-		ConfigFile->bCanSaveAllSections = ConfigFile->GetBool(SectionsToSaveString, SaveAllSectionsKey, bLocalSaveAllSections) && bLocalSaveAllSections;
-		// we can always save all sections of a User config file
-		ConfigFile->bCanSaveAllSections = ConfigFile->bCanSaveAllSections || BaseIniName.Contains(TEXT("User"));
+		bool bLocalSaveAllSections = false;
+		ConfigFile->GetBool(SectionsToSaveString, SaveAllSectionsKey, bLocalSaveAllSections);
+
+		// we can always save all sections of a User config file, Editor* (not Editor.ini tho, that is already handled in the normal method)
+		bool bIsUserFile = BaseIniName.Contains(TEXT("User"));
+		bool bIsEditorSettingsFile = BaseIniName.Contains(TEXT("Editor")) && BaseIniName != TEXT("Editor");
+
+		ConfigFile->bCanSaveAllSections = bLocalSaveAllSections || bIsUserFile || bIsEditorSettingsFile;
 
 		// don't write anything to disk in cooked builds - we will always use re-generated INI files anyway.
 		// Note: Unfortunately bAllowGeneratedIniWhenCooked is often true even in shipping builds with cooked data
diff --git a/Engine/Source/Runtime/Core/Public/Misc/ConfigHierarchy.h b/Engine/Source/Runtime/Core/Public/Misc/ConfigHierarchy.h
--- a/Engine/Source/Runtime/Core/Public/Misc/ConfigHierarchy.h
+++ b/Engine/Source/Runtime/Core/Public/Misc/ConfigHierarchy.h
@@ -50,7 +50,7 @@ inline FConfigLayer GConfigLayers[] =
 inline FConfigLayer GPluginLayers[] =
 {
 	// Engine/Base.ini
-	{ TEXT("AbsoluteBase"),				TEXT("{ENGINE}/Config/Base.ini"), EConfigLayerFlags::NoExpand},
+	{ TEXT("AbsoluteBase"),				TEXT("{ENGINE}/Config/PluginBase.ini"), EConfigLayerFlags::NoExpand},
 
 	// Plugin/Base*.ini
 	{ TEXT("PluginBase"),				TEXT("{PLUGIN}/Config/Base{TYPE}.ini") },
