Attempted fix for issue https://jira.it.epicgames.com/browse/UE-199751

diff --git a/Engine/Source/Runtime/Core/Tests/Async/ExternalMutexTest.cpp b/Engine/Source/Runtime/Core/Tests/Async/ExternalMutexTest.cpp
--- a/Engine/Source/Runtime/Core/Tests/Async/ExternalMutexTest.cpp
+++ b/Engine/Source/Runtime/Core/Tests/Async/ExternalMutexTest.cpp
@@ -5,6 +5,7 @@
 #include "Async/ExternalMutex.h"
 
 #include "Async/UniqueLock.h"
+#include "HAL/PlatformProcess.h"
 #include "HAL/Thread.h"
 #include "TestHarness.h"
 
@@ -35,6 +36,7 @@ TEST_CASE("Core::Async::ExternalMutex", "[Core][Async][Slow]")
 				FExternalMutex Mutex(ExternalState);
 				while (!Mutex.TryLock()) // spin on attempting to acquire the lock
 				{
+					FPlatformProcess::YieldThread();
 				}
 				CHECK(Mutex.IsLocked());
 				CHECK(!Mutex.TryLock());
diff --git a/Engine/Source/Runtime/Core/Tests/Async/MutexTest.cpp b/Engine/Source/Runtime/Core/Tests/Async/MutexTest.cpp
--- a/Engine/Source/Runtime/Core/Tests/Async/MutexTest.cpp
+++ b/Engine/Source/Runtime/Core/Tests/Async/MutexTest.cpp
@@ -5,6 +5,7 @@
 #include "Async/Mutex.h"
 
 #include "Async/UniqueLock.h"
+#include "HAL/PlatformProcess.h"
 #include "HAL/Thread.h"
 #include "TestHarness.h"
 
@@ -31,6 +32,7 @@ TEST_CASE("Core::Async::Mutex", "[Core][Async][Slow]")
 			{
 				while (!Mutex.TryLock()) // spin on attempting to acquire the lock
 				{
+					FPlatformProcess::YieldThread();
 				}
 				CHECK(Mutex.IsLocked());
 				CHECK(!Mutex.TryLock());
