Adding support for querying CPU features. Exit the game in PreInit on Windows when an unsupported CPU is found. I haven't spotted a common x86 place to put this, I expect it'll get moved over to a more common place soon enough.

diff --git a/Engine/Source/Runtime/Core/Private/Windows/WindowsPlatformMisc.cpp b/Engine/Source/Runtime/Core/Private/Windows/WindowsPlatformMisc.cpp
--- a/Engine/Source/Runtime/Core/Private/Windows/WindowsPlatformMisc.cpp
+++ b/Engine/Source/Runtime/Core/Private/Windows/WindowsPlatformMisc.cpp
@@ -587,6 +587,15 @@ void FWindowsPlatformMisc::PlatformPreInit()
 
 	// initialize the file SHA hash mapping
 	InitSHAHashes();
+
+	// Check for SSE42 or better. This is now minspec and there is a high likelihood
+	// of crashing on an invalid instruction on unsupported processors as we use these
+	// instructions now.
+	if (CheckFeatureBit_X86(ECPUFeatureBits_X86::SSE42) == false)
+	{
+		FMessageDialog::Open(EAppMsgType::Ok, NSLOCTEXT("Launch", "Error_CPUNotSupported", "This CPU does not support a required feature (SSE4.2)."));
+		FPlatformMisc::RequestExit(false, TEXT("FWindowsPlatformMisc::PlatformPreInit.CPUNotSupported"));
+	}
 }
 
 
@@ -2837,54 +2846,119 @@ FString FWindowsPlatformMisc::GetCPUBrand()
 	return FCPUIDQueriedData::GetBrand();
 }
 
-bool FWindowsPlatformMisc::HasAVX2InstructionSupport()
-{
-#if PLATFORM_CPU_ARM_FAMILY
-	return false;
+#if PLATFORM_CPU_X86_FAMILY
+
+#ifdef _MSC_VER
+	#define CpuIdEx __cpuidex
+	#define CpuId __cpuid
 #else
-	if (!HasCPUIDInstruction())
-	{
-		return false;
-	}
+	// GCC/Clang
 
-	int flags[4];
-	/* CPUID.(EAX=01H, ECX=0H):ECX.FMA[bit 12]==1   &&
-	   CPUID.(EAX=01H, ECX=0H):ECX.MOVBE[bit 22]==1 &&
-	   CPUID.(EAX=01H, ECX=0H):ECX.XSAVE[bit 26]==1 &&
-	   CPUID.(EAX=01H, ECX=0H):ECX.OSXSAVE[bit 27]==1 &&
-	   CPUID.(EAX=01H, ECX=0H):ECX.AVX[bit 28]==1 */
-	const int FMA_MOVBE_XSAVE_OSXSAVE_AVX_BITS = (1 << 12) | (1 << 22) | (1 << 26) | (1 << 27) | (1 << 28);
-	__cpuidex(flags, 1, 0);
-	if ((flags[2] & FMA_MOVBE_XSAVE_OSXSAVE_AVX_BITS) != FMA_MOVBE_XSAVE_OSXSAVE_AVX_BITS)
-	{
-		return false;
-	}
+	// 64-bit: GCC/Clang won't let us use "=b" constraint on Mac64, and we need to preserve RBX
+	// (PIC/PIE base)
+	#define CpuIdEx(out, leaf_id, subleaf_id)\
+			asm("xchgq %%rbx,%q1\n" \
+				"cpuid\n" \
+				"xchgq %%rbx,%q1\n" \
+				: "=a" (out[0]), "=&r" (out[1]), "=c" (out[2]), "=d" (out[3]): "0" (leaf_id), "2"(subleaf_id));
 
-	/*  CPUID.(EAX=07H, ECX=0H):EBX.AVX2[bit 5]==1  &&
-		CPUID.(EAX=07H, ECX=0H):EBX.BMI1[bit 3]==1  &&
-		CPUID.(EAX=07H, ECX=0H):EBX.BMI2[bit 8]==1  */
-	const int AVX2_BMI1_BMI2_BITS = (1 << 5) | (1 << 3) | (1 << 8);
-	__cpuidex(flags, 7, 0);
-	if ((flags[1] & AVX2_BMI1_BMI2_BITS) != AVX2_BMI1_BMI2_BITS)
-	{
-		return false;
-	}
+	#define CpuId(out, leaf_id) CpuIdEx(out, leaf_id, 0)
 
-	/* CPUID.(EAX=80000001H):ECX.LZCNT[bit 5]==1 */
-	const int LZCNT_BITS = (1 << 5);
-	__cpuidex(flags, 0x80000001, 0);
-	if ((flags[2] & LZCNT_BITS) != LZCNT_BITS)
+#endif // if not msc
+
+static std::atomic_uint32_t CachedX86FeatureBits = 0;
+uint32 FWindowsPlatformMisc::GetFeatureBits_X86()
+{
+	//
+	// Note we are 64bit+ now so we know we have cpuid.
+	//
+
+	uint32 FeatureBits = CachedX86FeatureBits.load(std::memory_order_relaxed);
+	if (FeatureBits)
 	{
-		return false;
+		return FeatureBits;
 	}
 
-	// OS must save YMM registers between context switch
-	if ((_xgetbv(0) & 6) != 6)
+	int CpuInfo[4];
+	uint32 MaxLeaf;
+
+	// Basic CPUID information
+	CpuId(CpuInfo, 0);
+	MaxLeaf = CpuInfo[0];
+
+	// Basic feature flags
+	CpuId(CpuInfo, 1);
+
+	FeatureBits |= (CpuInfo[3] & (1u << 26)) ? ECPUFeatureBits_X86::SSE2 : 0;
+	FeatureBits |= (CpuInfo[2] & (1u << 9)) ? ECPUFeatureBits_X86::SSSE3 : 0;
+	FeatureBits |= (CpuInfo[2] & (1u << 20)) ? ECPUFeatureBits_X86::SSE42 : 0;
+	FeatureBits |= (CpuInfo[2] & (1u << 28)) ? ECPUFeatureBits_X86::AVX : 0;
+	FeatureBits |= (CpuInfo[2] & (1u << 29)) ? ECPUFeatureBits_X86::F16C : 0;
+
+	// We don't have a feature flag we report for this, but we do use it later
+	bool has_popcnt = (CpuInfo[2] & (1u << 23)) != 0;
+
+	if (MaxLeaf >= 7)
 	{
-		return false;
+		// "Structured extended feature flags enumeration"
+		CpuIdEx(CpuInfo, 7, 0);
+
+		// Some (Celeron) Skylakes erroneously report BMI1/BMI2 even though they don't have it.
+		// These Celerons also don't have AVX.
+		//
+		// All CPUs that actually have BMI1/BMI2 (as of this writing, 2016-05-11) have AVX.
+		// (The ones we care about, anyway.) So only report BMI1/BMI2 if AVX is present.
+		// Also only report AVX or the BMIs if POPCNT is present; all processors I know of
+		// have either both or neither, and it's convenient for us to be able to assume
+		// that either BMI1/BMI2 or AVX2 implies POPCNT.
+		if ((FeatureBits & ECPUFeatureBits_X86::AVX) && has_popcnt)
+		{
+			if (CpuInfo[1] & (1u << 3))	FeatureBits |= ECPUFeatureBits_X86::BMI1;
+			if (CpuInfo[1] & (1u << 8))	FeatureBits |= ECPUFeatureBits_X86::BMI2;
+
+			// OS must save YMM registers between context switch
+			bool OsSavesAvxRegs = (_xgetbv(0) & 6) == 6;
+
+			// In addition to the above, only report AVX2 if BMI1 (and thus LZCNT/TZCNT)
+			// are also reported present; finally VC++ with /arch:AVX2 will emit BMI2
+			// instructions for things like variable shifts so we require BMI2 for AVX2
+			// as well.
+			//
+			// In practice this is not a limitation, AVX2 and BMI2 are a package deal on
+			// all uArchs I'm aware of.
+			const uint32 Avx2Bits = (1u << 3) /* BMI1 */ | (1u << 5) /* AVX2 */ | (1u << 8) /* BMI2 */;
+			if (((CpuInfo[1] & Avx2Bits) == Avx2Bits) && OsSavesAvxRegs)
+				FeatureBits |= ECPUFeatureBits_X86::AVX2;
+
+			// For us to report AVX512, we want the Skylake feature set
+			const uint32 Avx512Bits = (1u << 31) /* AVX512VL */ | (1u << 30) /* AVX512BW */ | (1u << 17) /* AVX512DQ */ | (1u << 16) /* AVX512F */;
+			if ((CpuInfo[1] & Avx512Bits) == Avx512Bits)
+				FeatureBits |= ECPUFeatureBits_X86::AVX512;
+
+			// Use the VBMI2 bit (set on ICL+) to set the NOCAVEATS flag. This is available
+			// on a generation of cores where AVX-512 has no major clock penalty anymore so
+			// whether to use AVX-512 or not is a much more straightforward calculation,
+			// and not so dependent on what else is running at the same time.
+			if (CpuInfo[2] & (1u << 6))
+				FeatureBits |= ECPUFeatureBits_X86::AVX512_NOCAVEATS;
+		}
 	}
 
-	return true;
+	// write detected features
+	// only write value once at end of the function!
+	FeatureBits |= 1; // initialized flag
+
+	CachedX86FeatureBits.store(FeatureBits, std::memory_order_release);
+	return FeatureBits;
+}
+#endif
+
+bool FWindowsPlatformMisc::HasAVX2InstructionSupport()
+{
+#if PLATFORM_CPU_ARM_FAMILY
+	return false;
+#else
+	return CheckFeatureBit_X86(ECPUFeatureBits_X86::AVX2);
 #endif
 }
 
diff --git a/Engine/Source/Runtime/Core/Public/Windows/WindowsPlatformMisc.h b/Engine/Source/Runtime/Core/Public/Windows/WindowsPlatformMisc.h
--- a/Engine/Source/Runtime/Core/Public/Windows/WindowsPlatformMisc.h
+++ b/Engine/Source/Runtime/Core/Public/Windows/WindowsPlatformMisc.h
@@ -12,6 +12,22 @@ class GenericApplication;
 struct FGuid;
 class IPlatformChunkInstall;
 
+#if PLATFORM_CPU_X86_FAMILY
+namespace ECPUFeatureBits_X86
+{
+	constexpr uint32 SSE2 = 1U << 2;
+	constexpr uint32 SSSE3 = 1U << 3;
+	constexpr uint32 SSE42 = 1U << 4;
+	constexpr uint32 AVX = 1U << 5;
+	constexpr uint32 BMI1 = 1U << 6; // Bit Manipulation Instructions - 1
+	constexpr uint32 BMI2 = 1U << 7; // Bit Manipulation Instructions - 2
+	constexpr uint32 AVX2 = 1U << 8;
+	constexpr uint32 F16C = 1U << 9; // Float16 conversion instructions
+	constexpr uint32 AVX512 = 1U << 10; // Skylake feature set : AVXF512{ F,VL,BW,DQ}.
+	constexpr uint32 AVX512_NOCAVEATS = 1U << 11; // Set when we have AVX512 without caveats like throttling.
+}
+#endif
+
 /** Helper struct used to get the string version of the Windows version. */
 struct FWindowsOSVersionHelper
 {
@@ -178,6 +194,14 @@ struct FWindowsPlatformMisc
 	 */
 	static CORE_API bool HasCPUIDInstruction();
 
+#if PLATFORM_CPU_X86_FAMILY
+	// Query the CPUID and parse out various feature bits. This is safe to call multiple times and caches the result internally for rapid access.
+	// Bits are all from the ECPUFeatureBits_X86 namespace.
+	static CORE_API uint32 GetFeatureBits_X86();
+	static CORE_API bool CheckFeatureBit_X86(uint32 FeatureBit_X86) { return (GetFeatureBits_X86() & FeatureBit_X86) != 0; }
+	static CORE_API bool CheckAllFeatureBits_X86(uint32 FeatureBits_X86) { return (GetFeatureBits_X86() & FeatureBits_X86) == FeatureBits_X86; }
+#endif
+
 	/**
 	 * Determines if AVX2 instruction set is supported on this platform
 	 *
