Don't align large tasks allocated via GMalloc to PLATFORM_CACHE_LINE_SIZE. Alignments larger than 16 bytes prevent the small block allocator from being used in MallocBinned2/3, which means these allocs will go to the OS. This is particularly inefficient on platforms with a 64KB allocation granularity, and resulted in a 230MB regression recently.
#localization none
#tests replayrun tested on multiple platforms in PerfB and 28.00

diff --git a/Engine/Source/Runtime/Core/Public/Tasks/TaskPrivate.h b/Engine/Source/Runtime/Core/Public/Tasks/TaskPrivate.h
--- a/Engine/Source/Runtime/Core/Public/Tasks/TaskPrivate.h
+++ b/Engine/Source/Runtime/Core/Public/Tasks/TaskPrivate.h
@@ -725,6 +725,7 @@ namespace UE::Tasks
 		};
 
 		inline constexpr int32 SmallTaskSize = 256;
+		inline constexpr int32 LargeTaskAlignment = 16; // Larger than this will result in very wasteful allocations with MallocBinned2/3
 		using FExecutableTaskAllocator = TLockFreeFixedSizeAllocator_TLSCache<SmallTaskSize, PLATFORM_CACHE_LINE_SIZE>;
 		CORE_API extern FExecutableTaskAllocator SmallTaskAllocator;
 
@@ -746,7 +747,7 @@ namespace UE::Tasks
 
 			static void* operator new(size_t Size)
 			{
-				return Size <= SmallTaskSize ? SmallTaskAllocator.Allocate() : GMalloc->Malloc(sizeof(TExecutableTask), PLATFORM_CACHE_LINE_SIZE);
+				return Size <= SmallTaskSize ? SmallTaskAllocator.Allocate() : GMalloc->Malloc(sizeof(TExecutableTask), LargeTaskAlignment);
 			}
 
 			static void operator delete(void* Ptr, size_t Size)
