Fix/silence remaining V591 warnings and enable the warning

diff --git a/Engine/Plugins/Experimental/Mutable/Source/MutableRuntime/Private/MuR/ParametersPrivate.h b/Engine/Plugins/Experimental/Mutable/Source/MutableRuntime/Private/MuR/ParametersPrivate.h
--- a/Engine/Plugins/Experimental/Mutable/Source/MutableRuntime/Private/MuR/ParametersPrivate.h
+++ b/Engine/Plugins/Experimental/Mutable/Source/MutableRuntime/Private/MuR/ParametersPrivate.h
@@ -55,6 +55,7 @@ namespace Private
         TIndirectObject& operator=(const TIndirectObject<T>& Other) 
         { 
             Get() = Other; 
+			return *this;
         }
 
         TIndirectObject& operator=(const T& Object) 
diff --git a/Engine/Plugins/FX/Niagara/Source/Niagara/Public/NiagaraTypes.h b/Engine/Plugins/FX/Niagara/Source/Niagara/Public/NiagaraTypes.h
--- a/Engine/Plugins/FX/Niagara/Source/Niagara/Public/NiagaraTypes.h
+++ b/Engine/Plugins/FX/Niagara/Source/Niagara/Public/NiagaraTypes.h
@@ -1139,7 +1139,7 @@ const FNiagaraTypeDefinition& FNiagaraTypeDefinition::Get()
 	else if constexpr (std::is_same_v<T, FNiagaraID>) { return GetIDDef(); }
 	else if constexpr (std::is_same_v<T, FNiagaraRandInfo>) { return GetRandInfoDef(); }
 	else { static_assert(sizeof(T) == 0, "Unsupported type"); }
-}
+} //-V591
 
 //////////////////////////////////////////////////////////////////////////
 
diff --git a/Engine/Source/Runtime/Core/Public/Windows/WindowsPlatformCodeAnalysis.h b/Engine/Source/Runtime/Core/Public/Windows/WindowsPlatformCodeAnalysis.h
--- a/Engine/Source/Runtime/Core/Public/Windows/WindowsPlatformCodeAnalysis.h
+++ b/Engine/Source/Runtime/Core/Public/Windows/WindowsPlatformCodeAnalysis.h
@@ -103,7 +103,7 @@
 	#pragma warning(disable : 6240) // warning C6240 : (<expression> && <non-zero constant>) always evaluates to the result of <expression>. Did you intend to use the bitwise-and operator?
 
 	//PVS-Studio settings:
-	//-V::505,542,581,591,601,623,668,677,686,688,690,703,704,711,719,720,728,730,735,751,1002,1008,1062
+	//-V::505,542,581,601,623,668,677,686,688,690,703,704,711,719,720,728,730,735,751,1002,1008,1062
 	//-V:TRYCOMPRESSION:519,547
 	//-V:check(:501,547,560,605
 	//-V:checkf(:510
