PathViews: Convert HasRedundantTerminatingSeparatorto public, for callers that don't want the expense of calling the complete NormalizeDirectoryName function.
#rnx

diff --git a/Engine/Source/Runtime/Core/Private/Misc/PathViews.cpp b/Engine/Source/Runtime/Core/Private/Misc/PathViews.cpp
--- a/Engine/Source/Runtime/Core/Private/Misc/PathViews.cpp
+++ b/Engine/Source/Runtime/Core/Private/Misc/PathViews.cpp
@@ -26,55 +26,6 @@ namespace UE4PathViews_Private
 
 	static bool IsSlashOrBackslashOrPeriod(TCHAR C) { return C == TEXT('/') || C == TEXT('\\') || C == TEXT('.'); }
 
-	static bool HasRedundantTerminatingSeparator(FStringView A)
-	{
-		if (A.Len() <= 2)
-		{
-			if (A.Len() <= 1)
-			{
-				// "", "c", or "/", All of these are either not slash terminating or not redundant
-				return false;
-			}
-			else if (!IsSlashOrBackslash(A[1]))
-			{
-				// "/c" or "cd"
-				return false;
-			}
-			else if (IsSlashOrBackslash(A[0]))
-			{
-				// "//"
-				return false;
-			}
-			else if (A[0] == ':')
-			{
-				// ":/", which is an invalid path, and we arbitrarily decide its terminating slash is not redundant
-				return false;
-			}
-			else
-			{
-				// "c/"
-				return true;
-			}
-		}
-		else if (!IsSlashOrBackslash(A[A.Len() - 1]))
-		{
-			// "/Some/Path"
-			return false;
-		}
-		else if (A[A.Len() - 2] == ':')
-		{
-			// "Volume:/",  "Volume:/Some/Path:/"
-			// The first case is the root dir of the volume and is not redundant
-			// The second case is an invalid path (at most one colon), and we arbitrarily decide its terminating slash is not redundant
-			return false;
-		}
-		else
-		{
-			// /Some/Path/
-			return true;
-		}
-	}
-
 	static bool StringEqualsIgnoreCaseIgnoreSeparator(FStringView A, FStringView B)
 	{
 		if (A.Len() != B.Len())
@@ -198,6 +149,57 @@ FStringView FPathViews::GetPathLeaf(const FStringView& InPath)
 	return FStringView();
 }
 
+bool FPathViews::HasRedundantTerminatingSeparator(FStringView A)
+{
+	using namespace UE4PathViews_Private;
+
+	if (A.Len() <= 2)
+	{
+		if (A.Len() <= 1)
+		{
+			// "", "c", or "/", All of these are either not slash terminating or not redundant
+			return false;
+		}
+		else if (!IsSlashOrBackslash(A[1]))
+		{
+			// "/c" or "cd"
+			return false;
+		}
+		else if (IsSlashOrBackslash(A[0]))
+		{
+			// "//"
+			return false;
+		}
+		else if (A[0] == ':')
+		{
+			// ":/", which is an invalid path, and we arbitrarily decide its terminating slash is not redundant
+			return false;
+		}
+		else
+		{
+			// "c/"
+			return true;
+		}
+	}
+	else if (!IsSlashOrBackslash(A[A.Len() - 1]))
+	{
+		// "/Some/Path"
+		return false;
+	}
+	else if (A[A.Len() - 2] == ':')
+	{
+		// "Volume:/",  "Volume:/Some/Path:/"
+		// The first case is the root dir of the volume and is not redundant
+		// The second case is an invalid path (at most one colon), and we arbitrarily decide its terminating slash is not redundant
+		return false;
+	}
+	else
+	{
+		// /Some/Path/
+		return true;
+	}
+}
+
 bool FPathViews::IsPathLeaf(FStringView InPath)
 {
 	using namespace UE4PathViews_Private;
@@ -513,15 +515,10 @@ void FPathViews::NormalizeFilename(FStringBuilderBase& InOutPath)
 	FPlatformMisc::NormalizePath(InOutPath);
 }
 
-static bool ShouldRemoveTrailingSlash(FStringView Dir)
-{
-	return Dir.EndsWith(TEXT('/')) && !Dir.EndsWith(TEXTVIEW("//"), ESearchCase::CaseSensitive) && !Dir.EndsWith(TEXTVIEW(":/"), ESearchCase::CaseSensitive);
-}
-
 void FPathViews::NormalizeDirectoryName(FStringBuilderBase& InOutPath)
 {
 	Algo::Replace(MakeArrayView(InOutPath), TEXT('\\'), TEXT('/'));
-	InOutPath.RemoveSuffix(ShouldRemoveTrailingSlash(InOutPath.ToView()) ? 1 : 0);
+	InOutPath.RemoveSuffix(HasRedundantTerminatingSeparator(InOutPath.ToView()) ? 1 : 0);
 	FPlatformMisc::NormalizePath(InOutPath);
 }
 
diff --git a/Engine/Source/Runtime/Core/Public/Misc/PathViews.h b/Engine/Source/Runtime/Core/Public/Misc/PathViews.h
--- a/Engine/Source/Runtime/Core/Public/Misc/PathViews.h
+++ b/Engine/Source/Runtime/Core/Public/Misc/PathViews.h
@@ -324,6 +324,19 @@ public:
 	 */
 	static CORE_API bool IsRelativePath(FStringView InPath);
 
+	/**
+	 * Report whether the path has an unneeded trailing slash. 
+	 * /root/path	-> false
+	 * /root/		-> true
+	 * /root//		-> true
+	 * /			-> false
+	 * //			-> false
+	 * ///			-> true
+	 * d:/			-> false
+	 * d://			-> true
+	 */
+	static CORE_API bool HasRedundantTerminatingSeparator(FStringView A);
+
 	/** Convert to absolute using process BaseDir(), normalize and append. FPaths::ConvertRelativePathToFull() equivalent. */
 	static CORE_API void ToAbsolutePath(FStringView InPath, FStringBuilderBase& OutPath);
 	
diff --git a/Engine/Source/Runtime/Core/Tests/Misc/PathViewsTest.cpp b/Engine/Source/Runtime/Core/Tests/Misc/PathViewsTest.cpp
--- a/Engine/Source/Runtime/Core/Tests/Misc/PathViewsTest.cpp
+++ b/Engine/Source/Runtime/Core/Tests/Misc/PathViewsTest.cpp
@@ -631,6 +631,28 @@ TEST_CASE_NAMED(FPathViewsIsRelativePathTest, "System::Core::Misc::PathViews::Is
 	Instance.IsRelativePathTest();
 }
 
+TEST_CASE_NAMED(FPathViewsHasRedundantTerminatingSeparatorTest, "System::Core::Misc::PathViews::HasRedundantTerminatingSeparator", "[ApplicationContextMask][SmokeFilter]")
+{
+	CHECK_EQUALS(TEXT(""), FPathViews::HasRedundantTerminatingSeparator(TEXTVIEW("")), false);
+	CHECK_EQUALS(TEXT("/"), FPathViews::HasRedundantTerminatingSeparator(TEXTVIEW("/")), false);
+	CHECK_EQUALS(TEXT("//"), FPathViews::HasRedundantTerminatingSeparator(TEXTVIEW("//")), false);
+	CHECK_EQUALS(TEXT("///"), FPathViews::HasRedundantTerminatingSeparator(TEXTVIEW("///")), true);
+	CHECK_EQUALS(TEXT("\\"), FPathViews::HasRedundantTerminatingSeparator(TEXTVIEW("\\")), false);
+	CHECK_EQUALS(TEXT("\\\\"), FPathViews::HasRedundantTerminatingSeparator(TEXTVIEW("\\\\")), false);
+	CHECK_EQUALS(TEXT("\\\\\\"), FPathViews::HasRedundantTerminatingSeparator(TEXTVIEW("\\\\\\")), true);
+	CHECK_EQUALS(TEXT("text"), FPathViews::HasRedundantTerminatingSeparator(TEXTVIEW("text")), false);
+	CHECK_EQUALS(TEXT("text/"), FPathViews::HasRedundantTerminatingSeparator(TEXTVIEW("text/")), true);
+	CHECK_EQUALS(TEXT("text//"), FPathViews::HasRedundantTerminatingSeparator(TEXTVIEW("text//")), true);
+	CHECK_EQUALS(TEXT("text\\"), FPathViews::HasRedundantTerminatingSeparator(TEXTVIEW("text\\")), true);
+	CHECK_EQUALS(TEXT("text\\\\"), FPathViews::HasRedundantTerminatingSeparator(TEXTVIEW("text\\\\")), true);
+	CHECK_EQUALS(TEXT("D:"), FPathViews::HasRedundantTerminatingSeparator(TEXTVIEW("D:")), false);
+	CHECK_EQUALS(TEXT("D:/"), FPathViews::HasRedundantTerminatingSeparator(TEXTVIEW("D:/")), false);
+	CHECK_EQUALS(TEXT("D://"), FPathViews::HasRedundantTerminatingSeparator(TEXTVIEW("D://")), true);
+	CHECK_EQUALS(TEXT("D:\\"), FPathViews::HasRedundantTerminatingSeparator(TEXTVIEW("D:\\")), false);
+	CHECK_EQUALS(TEXT("D:\\\\"), FPathViews::HasRedundantTerminatingSeparator(TEXTVIEW("D:\\\\")), true);
+}
+
+
 TEST_CASE_NAMED(FPathViewsSplitFirstComponentTest, "System::Core::Misc::PathViews::SplitFirstComponent", "[ApplicationContextMask][SmokeFilter]")
 {
 	auto RunSplitFirstTest = [](const TCHAR* FullPath, const TCHAR* ExpectedFirst, const TCHAR* ExpectedRemaining)
