LLM: Added "-Assets" and "-AssetClasses" params for the "DumpLLM" console command. This will dump the respective tag sets (instead of default tag set). Requires LLM_ALLOW_ASSETS_TAGS and game/editor to be started with "-llm -llmtagsets=Assets,AssetClasses" in command line.

diff --git a/Engine/Source/Runtime/Core/Private/HAL/LowLevelMemTracker.cpp b/Engine/Source/Runtime/Core/Private/HAL/LowLevelMemTracker.cpp
--- a/Engine/Source/Runtime/Core/Private/HAL/LowLevelMemTracker.cpp
+++ b/Engine/Source/Runtime/Core/Private/HAL/LowLevelMemTracker.cpp
@@ -82,19 +82,34 @@ FAutoConsoleCommand DumpLLM(
 	{
 		FString Command = FString::Join(Args, TEXT(" "));
 
+		FLowLevelMemTracker::EDumpFormat DumpFormat = FLowLevelMemTracker::EDumpFormat::PlainText;
 		bool bCSV = FParse::Param(*Command, TEXT("CSV"));
-		bool bSnapshot = FParse::Param(*Command, TEXT("SNAPSHOT"));
+		if (bCSV)
+		{
+			DumpFormat = FLowLevelMemTracker::EDumpFormat::CSV;
+		}
 
-		UE::LLM::ESizeParams SizeParams(UE::LLM::ESizeParams::Default);
+		UE::LLM::ESizeParams SizeParams = UE::LLM::ESizeParams::Default;
+		bool bSnapshot = FParse::Param(*Command, TEXT("SNAPSHOT"));
 		if (bSnapshot)
 		{
 			EnumAddFlags(SizeParams, UE::LLM::ESizeParams::RelativeToSnapshot);
 		}
-		FLowLevelMemTracker::Get().DumpToLog(
-			bCSV ? FLowLevelMemTracker::EDumpFormat::CSV : FLowLevelMemTracker::EDumpFormat::PlainText, &Ar, SizeParams);
+
+		ELLMTagSet TagSet = ELLMTagSet::None;
+		if (FParse::Param(*Command, TEXT("Assets")))
+		{
+			TagSet = ELLMTagSet::Assets;
+		}
+		else if (FParse::Param(*Command, TEXT("AssetClasses")))
+		{
+			TagSet = ELLMTagSet::AssetClasses;
+		}
+
+		FLowLevelMemTracker::Get().DumpToLog(DumpFormat, &Ar, SizeParams, TagSet);
 	}));
 
-void FLowLevelMemTracker::DumpToLog(EDumpFormat DumpFormat, FOutputDevice* OutputDevice, UE::LLM::ESizeParams SizeParams)
+void FLowLevelMemTracker::DumpToLog(EDumpFormat DumpFormat, FOutputDevice* OutputDevice, UE::LLM::ESizeParams SizeParams, ELLMTagSet TagSet)
 {
 	if (!IsEnabled())
 	{
@@ -127,7 +142,7 @@ void FLowLevelMemTracker::DumpToLog(EDumpFormat DumpFormat, FOutputDevice* Outpu
 		EnumAddFlags(PeakSizeParams, UE::LLM::ESizeParams::ReportPeak);
 
 		TArray<FTagLine> TagLines;
-		for (const UE::LLMPrivate::FTagData* TagData : GetTrackedTags(TrackerType))
+		for (const UE::LLMPrivate::FTagData* TagData : GetTrackedTags(TrackerType, TagSet))
 		{
 			int64 CurrentAmount = GetTagAmountForTracker(TrackerType, TagData, SizeParams);
 			int64 PeakAmount = GetTagAmountForTracker(TrackerType, TagData, PeakSizeParams);
diff --git a/Engine/Source/Runtime/Core/Public/HAL/LowLevelMemTracker.h b/Engine/Source/Runtime/Core/Public/HAL/LowLevelMemTracker.h
--- a/Engine/Source/Runtime/Core/Public/HAL/LowLevelMemTracker.h
+++ b/Engine/Source/Runtime/Core/Public/HAL/LowLevelMemTracker.h
@@ -712,7 +712,7 @@ public:
 		PlainText,
 		CSV,
 	};
-	CORE_API void DumpToLog(EDumpFormat DumpFormat = EDumpFormat::PlainText, FOutputDevice* OutputDevice = nullptr, UE::LLM::ESizeParams SizeParams = UE::LLM::ESizeParams::Default);
+	CORE_API void DumpToLog(EDumpFormat DumpFormat = EDumpFormat::PlainText, FOutputDevice* OutputDevice = nullptr, UE::LLM::ESizeParams SizeParams = UE::LLM::ESizeParams::Default, ELLMTagSet TagSet = ELLMTagSet::None);
 
 	CORE_API void OnPreFork();
 
