AssetRegistry: AssetDataGather: Use NTFSJournal for the directory scan.
#rnx

diff --git a/Engine/Config/BaseEngine.ini b/Engine/Config/BaseEngine.ini
--- a/Engine/Config/BaseEngine.ini
+++ b/Engine/Config/BaseEngine.ini
@@ -2644,6 +2644,14 @@ bFilterSearchableNames=true
 +CommandletsUsingAR=DataValidation
 +CommandletsUsingAR=GatherText
 
+; AssetDataGathering can be cached in the editor; there are two different caches, one for the directory scan to find all
+; packages and one for tags and other data read from the packages. The directory scan cache by default is only enabled if
+; the current machine supports FileJournal (IPlatformFile::FileJournalIsAvailable), but the cache can also
+; be used on-demand without checking it for out-of-date via commandline -SkipAssetRegistryCacheInvalidate.
+; To support that use case set AssetRegistryDiscoveryCache=Always to always create the cache.
+; Options are Never,Default,Always. Default only creates the cache if the platform supports invalidation via FileJournal.
+AssetRegistryDiscoveryCache=Default
+
 [AutomationTesting]
 ImportTestPath=Content/EditorAutomation/
 ImportTestPackagePath=/Engine/Content/EditorAutomation
diff --git a/Engine/Source/Runtime/AssetRegistry/Private/AssetDataGatherer.cpp b/Engine/Source/Runtime/AssetRegistry/Private/AssetDataGatherer.cpp
--- a/Engine/Source/Runtime/AssetRegistry/Private/AssetDataGatherer.cpp
+++ b/Engine/Source/Runtime/AssetRegistry/Private/AssetDataGatherer.cpp
@@ -23,6 +23,7 @@
 #include "Misc/AsciiSet.h"
 #include "Misc/Char.h"
 #include "Misc/CommandLine.h"
+#include "Misc/ConfigCacheIni.h"
 #include "Misc/CoreDelegates.h"
 #include "Misc/DelayedAutoRegister.h"
 #include "Misc/Parse.h"
@@ -91,12 +92,42 @@ void FPreloadSettings::Initialize()
 	bForceDependsGathering = FParse::Param(FCommandLine::Get(), TEXT("ForceDependsGathering"));
 	bGatherDependsData = (GIsEditor && !FParse::Param(FCommandLine::Get(), TEXT("NoDependsGathering"))) || bForceDependsGathering;
 	bool bNoAssetRegistryCache = FParse::Param(FCommandLine::Get(), TEXT("NoAssetRegistryCache"));
+	bool bNoAssetRegistryDiscoveryCache = bNoAssetRegistryCache || FParse::Param(FCommandLine::Get(), TEXT("NoAssetRegistryDiscoveryCache"));
 	bool bNoAssetRegistryCacheRead = FParse::Param(FCommandLine::Get(), TEXT("NoAssetRegistryCacheRead"));
 	bool bNoAssetRegistryCacheWrite = FParse::Param(FCommandLine::Get(), TEXT("NoAssetRegistryCacheWrite"));
 	uint32 MultiprocessId = UE::GetMultiprocessId();
 	bool bMultiprocess = MultiprocessId > 0 || FParse::Param(FCommandLine::Get(), TEXT("multiprocess"));
-	bCacheReadEnabled = !bNoAssetRegistryCache && !bNoAssetRegistryCacheRead;
-	bCacheWriteEnabled = !bNoAssetRegistryCache && !bNoAssetRegistryCacheWrite && !bMultiprocess;
+	bGatherCacheReadEnabled = !bNoAssetRegistryCache && !bNoAssetRegistryCacheRead;
+	bGatherCacheWriteEnabled = !bNoAssetRegistryCache && !bNoAssetRegistryCacheWrite && !bMultiprocess;
+	bool bPlatformSupportsDiscoveryCache = FPlatformFileManager::Get().GetPlatformFile().FileJournalIsAvailable();
+
+	bool bSkipInvalidate = FParse::Param(FCommandLine::Get(), TEXT("AssetRegistryCacheSkipInvalidate"));
+	FString AssetRegistryDiscoveryWriteCacheStr;
+	GConfig->GetString(TEXT("AssetRegistry"), TEXT("AssetRegistryDiscoveryWriteCache"),
+		AssetRegistryDiscoveryWriteCacheStr, GEngineIni);
+	FParse::Value(FCommandLine::Get(), TEXT("AssetRegistryDiscoveryWriteCache="), AssetRegistryDiscoveryWriteCacheStr);
+
+	bDiscoveryCacheReadEnabled = (bPlatformSupportsDiscoveryCache || bSkipInvalidate) && !bNoAssetRegistryDiscoveryCache
+		&& !bNoAssetRegistryCacheRead;
+	if (bNoAssetRegistryDiscoveryCache || bNoAssetRegistryCacheWrite || bMultiprocess ||
+		AssetRegistryDiscoveryWriteCacheStr == TEXT("Never") || AssetRegistryDiscoveryWriteCacheStr == TEXT("false")
+		|| AssetRegistryDiscoveryWriteCacheStr == TEXT("0"))
+	{
+		DiscoveryCacheWriteEnabled = EFeatureEnabled::Never;
+	}
+	else if ((AssetRegistryDiscoveryWriteCacheStr.IsEmpty() || AssetRegistryDiscoveryWriteCacheStr == TEXT("Default"))
+		&& !bSkipInvalidate)
+	{
+		// Precalculate IfPlatformSupported -> Never if we already know the platform doesn't support it
+		DiscoveryCacheWriteEnabled = bPlatformSupportsDiscoveryCache ?
+			EFeatureEnabled::IfPlatformSupported : EFeatureEnabled::Never;
+	}
+	else
+	{
+		DiscoveryCacheWriteEnabled = EFeatureEnabled::Always;
+	}
+	bDiscoveryCacheInvalidateEnabled = !bSkipInvalidate && MultiprocessId == 0;
+
 	bool bAsyncEnabled = FPlatformProcess::SupportsMultithreading() && FTaskGraphInterface::IsRunning();
 
 	MonolithicCacheBaseFilename = AssetRegistryCacheRootFolder / (bGatherDependsData ? TEXT("CachedAssetRegistry") : TEXT("CachedAssetRegistryNoDeps"));
@@ -106,13 +137,25 @@ void FPreloadSettings::Initialize()
 	bMonolithicCacheActivatedDuringPreload = false;
 #endif
 }
-bool FPreloadSettings::IsCacheReadEnabled() const
+bool FPreloadSettings::IsGatherCacheReadEnabled() const
 {
-	return bCacheReadEnabled;
+	return bGatherCacheReadEnabled;
 }
-bool FPreloadSettings::IsCacheWriteEnabled() const
+bool FPreloadSettings::IsGatherCacheWriteEnabled() const
 {
-	return bCacheWriteEnabled;
+	return bGatherCacheWriteEnabled;
+}
+bool FPreloadSettings::IsDiscoveryCacheReadEnabled() const
+{
+	return bDiscoveryCacheReadEnabled;
+}
+EFeatureEnabled FPreloadSettings::IsDiscoveryCacheWriteEnabled() const
+{
+	return DiscoveryCacheWriteEnabled;
+}
+bool FPreloadSettings::IsDiscoveryCacheInvalidateEnabled() const
+{
+	return bDiscoveryCacheInvalidateEnabled;
 }
 bool FPreloadSettings::IsMonolithicCacheActivatedDuringPreload() const
 {
@@ -120,7 +163,7 @@ bool FPreloadSettings::IsMonolithicCacheActivatedDuringPreload() const
 }
 bool FPreloadSettings::IsPreloadMonolithicCache() const
 {
-	return bCacheReadEnabled && bMonolithicCacheActivatedDuringPreload;
+	return bGatherCacheReadEnabled && bMonolithicCacheActivatedDuringPreload;
 }
 bool FPreloadSettings::IsGatherDependsData() const
 {
@@ -1499,6 +1542,11 @@ void FAssetDataDiscovery::EnsureCompletion()
 	}
 }
 
+void FAssetDataDiscovery::OnInitialSearchCompleted()
+{
+	Cache.SaveCache();
+	Cache.Shutdown();
+}
 
 void FAssetDataDiscovery::TickInternal(bool bTickAll)
 {
@@ -1506,9 +1554,15 @@ void FAssetDataDiscovery::TickInternal(bool bTickAll)
 	CHECK_IS_NOT_LOCKED_CURRENT_THREAD(ResultsLock);
 	check(TickOwner.IsOwnedByCurrentThread());
 
+	if (!Cache.IsInitialized())
+	{
+		Cache.LoadAndUpdateCache();
+	}
+
 	TArray<FScanDirAndParentData> ScanRequests;
 	TStringBuilder<128> DirMountRelPath;
 
+	int32 LocalNumCachedDirectories = 0;
 	int32 DirToScanDatasNum = 0;
 	bool bUpdatedPriorityData = false;
 	double TickStartTime = FPlatformTime::Seconds();
@@ -1526,12 +1580,14 @@ void FAssetDataDiscovery::TickInternal(bool bTickAll)
 			FGathererScopeLock TreeScopeLock(&TreeLock);
 
 			// Process scanned directories from the previous iteration of the for(;;) loop.
+			int32 NumScanned = 0;
 			if (DirToScanDatasNum > 0)
 			{
 				for (FDirToScanData& Data : TArrayView<FDirToScanData>(DirToScanDatas.GetData(), DirToScanDatasNum))
 				{
 					if (Data.bScanned)
 					{
+						++NumScanned;
 						TArrayView<FDiscoveredPathData> LocalSubDirs(Data.IteratedSubDirs.GetData(), Data.NumIteratedDirs);
 						TArrayView<FDiscoveredPathData> LocalDiscoveredFiles(Data.IteratedFiles.GetData(), Data.NumIteratedFiles);
 						if (!Data.ScanDir->IsValid())
@@ -1556,7 +1612,17 @@ void FAssetDataDiscovery::TickInternal(bool bTickAll)
 						Data.ScanDir.SafeRelease();
 					}
 				}
+				// Rather than collecting LocalNumUncachedDirectories from inside the parallel for that scans them,
+				// just calculate it from the number of Data.bScanned and the number of LocalNumCachedDirectories
+				int32 LocalNumUncachedDirectories = NumScanned - LocalNumCachedDirectories;
 				DirToScanDatasNum = 0;
+				{
+					FGathererScopeLock ResultsScopeLock(&ResultsLock);
+					NumCachedDirectories += LocalNumCachedDirectories;
+					NumUncachedDirectories += LocalNumUncachedDirectories;
+				}
+				LocalNumCachedDirectories = 0;
+				LocalNumUncachedDirectories = 0;
 			}
 
 			// Look for new dirs to scan, break out of the for (;;) loop if we don't find any
@@ -1580,6 +1646,16 @@ void FAssetDataDiscovery::TickInternal(bool bTickAll)
 				FPathViews::AppendPath(ScanData.DirLocalAbsPath, DirMountRelPath);
 				ScanData.DirLongPackageName << MountDir->GetLongPackageName();
 				FPathViews::AppendPath(ScanData.DirLongPackageName, DirMountRelPath);
+				// The DirLocalAbsPath and DirLongPackageName need to be normalized. They are already mostly
+				// normalized, but might have a redundant terminating separator
+				while (FPathViews::HasRedundantTerminatingSeparator(ScanData.DirLocalAbsPath))
+				{
+					ScanData.DirLocalAbsPath.RemoveSuffix(1);
+				}
+				while (FPathViews::HasRedundantTerminatingSeparator(ScanData.DirLongPackageName))
+				{
+					ScanData.DirLongPackageName.RemoveSuffix(1);
+				}
 				ScanData.ScanDir = MoveTemp(ScanRequest.ScanDir);
 				ScanData.ParentData = MoveTemp(ScanRequest.ParentData);
 			};
@@ -1689,6 +1765,87 @@ void FAssetDataDiscovery::TickInternal(bool bTickAll)
 		}
 
 		// Outside of the TreeLock critical section, scan the directories
+
+		// Process on a single thread any of the DirToScans that we find in the cache.
+		for (int32 DirToScanDatasIndex = 0; DirToScanDatasIndex < DirToScanDatasNum; ++DirToScanDatasIndex)
+		{
+			FDirToScanData& Data = DirToScanDatas[DirToScanDatasIndex];
+			FCachedDirScanDir* PathData = Cache.FindDir(Data.DirLocalAbsPath);
+			if (!PathData || !PathData->bCacheValid)
+			{
+				continue;
+			}
+			Data.bScanned = true;
+			++LocalNumCachedDirectories;
+
+			for (const FString& RelPath : PathData->SubDirRelPaths)
+			{
+				// Don't enter directories that contain invalid packagepath characters (including '.';
+				// extensions are not valid in content directories because '.' is not valid in a packagepath)
+				if (!FPackageName::DoesPackageNameContainInvalidCharacters(RelPath))
+				{
+					int32 DirLongPackageRootNameLen = Data.DirLongPackageName.Len();
+					int32 DirLocalAbsPathLen = Data.DirLocalAbsPath.Len();
+					ON_SCOPE_EXIT
+					{
+						Data.DirLongPackageName.RemoveSuffix(Data.DirLongPackageName.Len() - DirLongPackageRootNameLen);
+						Data.DirLocalAbsPath.RemoveSuffix(Data.DirLocalAbsPath.Len() - DirLocalAbsPathLen);
+					};
+
+					FPathViews::AppendPath(Data.DirLongPackageName, RelPath);
+					FPathViews::AppendPath(Data.DirLocalAbsPath, RelPath);
+					if (Data.IteratedSubDirs.Num() < Data.NumIteratedDirs + 1)
+					{
+						check(Data.IteratedSubDirs.Num() == Data.NumIteratedDirs);
+						Data.IteratedSubDirs.Emplace();
+					}
+					Data.IteratedSubDirs[Data.NumIteratedDirs++].Assign(Data.DirLocalAbsPath, Data.DirLongPackageName, RelPath,
+						EGatherableFileType::Directory, false /* bBlocked */);
+				}
+			}
+			for (const FCachedDirScanFile& FileData : PathData->Files)
+			{
+				FStringView RelPath(FileData.RelPath);
+
+				EGatherableFileType FileType = GetFileType(RelPath);
+				// Don't record files that contain invalid packagepath characters (not counting their extension)
+				// or that do not end with a recognized extension
+				if (FileType != EGatherableFileType::Invalid)
+				{
+					FStringView BaseName = FPathViews::GetBaseFilename(RelPath);
+					if (!FPackageName::DoesPackageNameContainInvalidCharacters(BaseName))
+					{
+						int32 DirLongPackageRootNameLen = Data.DirLongPackageName.Len();
+						int32 DirLocalAbsPathLen = Data.DirLocalAbsPath.Len();
+						ON_SCOPE_EXIT
+						{
+							Data.DirLongPackageName.RemoveSuffix(Data.DirLongPackageName.Len() - DirLongPackageRootNameLen);
+							Data.DirLocalAbsPath.RemoveSuffix(Data.DirLocalAbsPath.Len() - DirLocalAbsPathLen);
+						};
+
+						if (Data.IteratedFiles.Num() < Data.NumIteratedFiles + 1)
+						{
+							check(Data.IteratedFiles.Num() == Data.NumIteratedFiles);
+							Data.IteratedFiles.Emplace();
+						}
+						FPathViews::AppendPath(Data.DirLongPackageName, BaseName);
+						FPathViews::AppendPath(Data.DirLocalAbsPath, RelPath);
+						const bool bBlocked = FileType == EGatherableFileType::PackageFile && IsPackageBlocked(Data.DirLocalAbsPath);
+						Data.IteratedFiles[Data.NumIteratedFiles++].Assign(Data.DirLocalAbsPath, Data.DirLongPackageName, RelPath,
+							FileData.ModificationTime, FileType, bBlocked);
+					}
+				}
+			}
+		}
+
+		// If we found any cached directories, keep looking in their children before we start querying
+		// the disk for uncached
+		if (LocalNumCachedDirectories > 0 && !bUpdatedPriorityData)
+		{
+			continue;
+		}
+
+		// Otherwise look on disk in parallel for all of the DirToScans
 		int32 NumThreads = FMath::Max(FTaskGraphInterface::Get().GetNumWorkerThreads(), 1);
 		if (AssetDataGathererConstants::GARDiscoverThreads > 0)
 		{
@@ -1709,6 +1866,10 @@ void FAssetDataDiscovery::TickInternal(bool bTickAll)
 		[this](FDirToScanBuffer& ScanBuffer, int32 DirToScanDatasIndex)
 		{
 			FDirToScanData& Data = DirToScanDatas[DirToScanDatasIndex];
+			if (Data.bScanned)
+			{
+				return;
+			}
 			if (ScanBuffer.bAbort)
 			{
 				return;
@@ -1719,32 +1880,33 @@ void FAssetDataDiscovery::TickInternal(bool bTickAll)
 				return;
 			}
 
-			IFileManager::Get().IterateDirectoryStat(Data.DirLocalAbsPath.ToString(),
-			[this, &Data]
-			(const TCHAR* InPackageFilename, const FFileStatData& InPackageStatData)
+			FCachedDirScanDir CacheDataToAdd;
+			FPlatformFileManager::Get().GetPlatformFile().FileJournalIterateDirectory(Data.DirLocalAbsPath.ToString(),
+			[this, &Data, &CacheDataToAdd]
+			(const TCHAR* IterFilename, const FFileJournalData& IterData)
 			{
-				FStringView LocalAbsPath(InPackageFilename);
+				FStringView LocalAbsPath(IterFilename);
 				FStringView RelPath;
 				FString Buffer;
-				if (!FPathViews::TryMakeChildPathRelativeTo(InPackageFilename, Data.DirLocalAbsPath, RelPath))
+				if (!FPathViews::TryMakeChildPathRelativeTo(IterFilename, Data.DirLocalAbsPath, RelPath))
 				{
 					// Try again with the path converted to the absolute path format that we passed in; some
 					// IFileManagers can send relative paths to the visitor even though the search path is absolute
-					Buffer = FPaths::ConvertRelativePathToFull(FString(InPackageFilename));
+					Buffer = FPaths::ConvertRelativePathToFull(FString(IterFilename));
 					LocalAbsPath = Buffer;
 					if (!FPathViews::TryMakeChildPathRelativeTo(Buffer, Data.DirLocalAbsPath, RelPath))
 					{
 						UE_LOG(LogAssetRegistry, Warning,
-							TEXT("IterateDirectoryStat returned unexpected result %s which is not a child of the requested path %s."),
-							InPackageFilename, Data.DirLocalAbsPath.ToString());
+							TEXT("IterateDirectory returned unexpected result %s which is not a child of the requested path %s."),
+							IterFilename, Data.DirLocalAbsPath.ToString());
 						return true;
 					}
 				}
 				if (FPathViews::GetPathLeaf(RelPath).Len() != RelPath.Len())
 				{
 					UE_LOG(LogAssetRegistry, Warning,
-						TEXT("IterateDirectoryStat returned unexpected result %s which is not a direct child of the requested path %s."),
-						InPackageFilename, Data.DirLocalAbsPath.ToString());
+						TEXT("IterateDirectory returned unexpected result %s which is not a direct child of the requested path %s."),
+						IterFilename, Data.DirLocalAbsPath.ToString());
 					return true;
 				}
 				int32 DirLongPackageRootNameLen = Data.DirLongPackageName.Len();
@@ -1753,8 +1915,14 @@ void FAssetDataDiscovery::TickInternal(bool bTickAll)
 					Data.DirLongPackageName.RemoveSuffix(Data.DirLongPackageName.Len() - DirLongPackageRootNameLen);
 				};
 
-				if (InPackageStatData.bIsDirectory)
+				if (IterData.bIsDirectory)
 				{
+					if (Cache.IsWriteEnabled() != EFeatureEnabled::Never)
+					{
+						CacheDataToAdd.SubDirRelPaths.Add(FString(RelPath));
+						Cache.QueueAdd(FString(LocalAbsPath), IterData.JournalHandle);
+					}
+
 					FPathViews::AppendPath(Data.DirLongPackageName, RelPath);
 					// Don't enter directories that contain invalid packagepath characters (including '.';
 					// extensions are not valid in content directories because '.' is not valid in a packagepath)
@@ -1771,6 +1939,12 @@ void FAssetDataDiscovery::TickInternal(bool bTickAll)
 				}
 				else
 				{
+					if (Cache.IsWriteEnabled() != EFeatureEnabled::Never)
+					{
+						FCachedDirScanFile& FileData = CacheDataToAdd.Files.Emplace_GetRef();
+						FileData.RelPath = FString(RelPath);
+						FileData.ModificationTime = IterData.ModificationTime;
+					}
 					EGatherableFileType FileType = GetFileType(RelPath);
 					// Don't record files that contain invalid packagepath characters (not counting their extension)
 					// or that do not end with a recognized extension
@@ -1787,14 +1961,25 @@ void FAssetDataDiscovery::TickInternal(bool bTickAll)
 							FPathViews::AppendPath(Data.DirLongPackageName, BaseName);
 							const bool bBlocked = FileType == EGatherableFileType::PackageFile && IsPackageBlocked(LocalAbsPath);
 							Data.IteratedFiles[Data.NumIteratedFiles++].Assign(LocalAbsPath, Data.DirLongPackageName, RelPath,
-								InPackageStatData.ModificationTime, FileType, bBlocked);
+								IterData.ModificationTime, FileType, bBlocked);
 						}
 					}
 				}
 				return true;
 			});
+
+			if (Cache.IsWriteEnabled() != EFeatureEnabled::Never)
+			{
+				Cache.QueueAdd(FString(Data.DirLocalAbsPath), MoveTemp(CacheDataToAdd));
+			}
+
 			Data.bScanned = true;
 		});
+
+		if (Cache.IsWriteEnabled() != EFeatureEnabled::Never)
+		{
+			Cache.QueueConsume();
+		}
 	}
 }
 
@@ -1891,10 +2076,13 @@ void FAssetDataDiscovery::GetAndTrimSearchResults(bool& bOutIsComplete, TArray<F
 	}
 }
 
-void FAssetDataDiscovery::GetDiagnostics(float& OutCumulativeDiscoveryTime)
+void FAssetDataDiscovery::GetDiagnostics(float& OutCumulativeDiscoveryTime, int32& OutNumCachedDirectories,
+	int32& OutNumUncachedDirectories)
 {
 	FGathererScopeLock ResultsScopeLock(&ResultsLock);
 	OutCumulativeDiscoveryTime = CumulativeDiscoveryTime;
+	OutNumCachedDirectories = NumCachedDirectories;
+	OutNumUncachedDirectories = NumUncachedDirectories;
 }
 
 
@@ -3216,8 +3404,8 @@ FAssetDataGatherer::FAssetDataGatherer(const TArray<FString>& InLongPackageNames
 	GPreloadSettings.Initialize();
 	bGatherAssetPackageData = GIsEditor || GPreloadSettings.IsForceDependsGathering();
 	bGatherDependsData = GPreloadSettings.IsGatherDependsData();
-	bCacheReadEnabled = GPreloadSettings.IsCacheReadEnabled();
-	bCacheWriteEnabled = GPreloadSettings.IsCacheWriteEnabled();
+	bCacheReadEnabled = GPreloadSettings.IsGatherCacheReadEnabled();
+	bCacheWriteEnabled = GPreloadSettings.IsGatherCacheWriteEnabled();
 	// If IsMonolithicCacheActivatedDuringPreload is true, we are already instructed to use the MonolithicCache.
 	// Otherwise it may be set to true later if game/commandlet calls SearchAllAssets.
 	bReadMonolithicCache = bCacheReadEnabled && GPreloadSettings.IsMonolithicCacheActivatedDuringPreload();
@@ -3264,6 +3452,14 @@ FAssetDataGatherer::~FAssetDataGatherer()
 	DiskCachedAssetBlocks.Empty();
 }
 
+void FAssetDataGatherer::OnInitialSearchCompleted()
+{
+	if (Discovery)
+	{
+		Discovery->OnInitialSearchCompleted();
+	}
+}
+
 void FAssetDataGatherer::ActivateMonolithicCache()
 {
 	FGathererScopeLock ResultsScopeLock(&ResultsLock);
@@ -4069,7 +4265,7 @@ void FAssetDataGatherer::GetAndTrimSearchResults(FResults& InOutResults, FResult
 FAssetGatherDiagnostics FAssetDataGatherer::GetDiagnostics()
 {
 	FAssetGatherDiagnostics Diag;
-	Discovery->GetDiagnostics(Diag.DiscoveryTimeSeconds);
+	Discovery->GetDiagnostics(Diag.DiscoveryTimeSeconds, Diag.NumCachedDirectories, Diag.NumUncachedDirectories);
 	FGathererScopeLock ResultsScopeLock(&ResultsLock);
 	Diag.GatherTimeSeconds = CumulativeGatherTime;
 	Diag.NumCachedAssetFiles = NumCachedAssetFiles;
diff --git a/Engine/Source/Runtime/AssetRegistry/Private/AssetDataGatherer.h b/Engine/Source/Runtime/AssetRegistry/Private/AssetDataGatherer.h
--- a/Engine/Source/Runtime/AssetRegistry/Private/AssetDataGatherer.h
+++ b/Engine/Source/Runtime/AssetRegistry/Private/AssetDataGatherer.h
@@ -58,6 +58,10 @@ struct FAssetGatherDiagnostics
 	float DiscoveryTimeSeconds;
 	/** Time spent reading asset files on disk / from cache */
 	float GatherTimeSeconds;
+	/** How many directories in the search results were read from the cache. */
+	int32 NumCachedDirectories;
+	/** How many directories in the search results were not in the cache and were read by scanning the disk. */
+	int32 NumUncachedDirectories;
 	/** How many files in the search results were read from the cache. */
 	int32 NumCachedAssetFiles;
 	/** How many files in the search results were not in the cache and were read by parsing the file. */
@@ -74,6 +78,7 @@ public:
 		const TArray<FString>& InMountRelativePathsDenyList, bool bInAsyncEnabled);
 	virtual ~FAssetDataGatherer();
 
+	void OnInitialSearchCompleted();
 
 	// Extra at-construction configuration 
 
diff --git a/Engine/Source/Runtime/AssetRegistry/Private/AssetDataGathererDiscoveryCache.cpp b/Engine/Source/Runtime/AssetRegistry/Private/AssetDataGathererDiscoveryCache.cpp
new file mode 100644
--- /dev/null
+++ b/Engine/Source/Runtime/AssetRegistry/Private/AssetDataGathererDiscoveryCache.cpp
@@ -0,0 +1,547 @@
+// Copyright Epic Games, Inc. All Rights Reserved.
+
+#include "AssetDataGathererDiscoveryCache.h"
+
+#include "AssetDataGathererPrivate.h"
+#include "AssetRegistryPrivate.h"
+#include "Compression/CompressedBuffer.h"
+#include "HAL/FileManager.h"
+#include "HAL/PlatformFileManager.h"
+#include "Logging/LogMacros.h"
+#include "Memory/SharedBuffer.h"
+#include "Misc/Paths.h"
+#include "Misc/PathViews.h"
+#include "ProfilingDebugging/CpuProfilerTrace.h"
+#include "Serialization/Archive.h"
+#include "Serialization/LargeMemoryWriter.h"
+#include "Serialization/MemoryReader.h"
+#include "Templates/UniquePtr.h"
+
+FArchive& operator<<(FArchive& Ar, FFileJournalFileHandle& H)
+{
+	Ar.Serialize(H.Bytes, sizeof(H.Bytes));
+	return Ar;
+}
+
+namespace AssetDataGathererConstants
+{
+const FGuid DiscoveryCacheVersion(TEXT("62267A3E1D2144EF9C8734271DD28742"));
+}
+
+namespace UE::AssetDataGather::Private
+{
+
+FString FAssetDataDiscoveryCache::GetCacheFileName() const
+{
+	return FPaths::Combine(GPreloadSettings.GetAssetRegistryCacheRootFolder(),
+		TEXT("CachedAssetRegistryDiscovery.bin"));
+}
+
+void FAssetDataDiscoveryCache::LoadAndUpdateCache()
+{
+	if (bInitialized)
+	{
+		return;
+	}
+	TRACE_CPUPROFILER_EVENT_SCOPE(AssetDataGatherLoadDiscoveryCache);
+	bInitialized = true;
+
+	WriteEnabled = GPreloadSettings.IsDiscoveryCacheWriteEnabled();
+	CachedVolumes.Empty();
+	if (!GPreloadSettings.IsDiscoveryCacheReadEnabled() && WriteEnabled == EFeatureEnabled::Never)
+	{
+		return;
+	}
+
+	FString TestError;
+	FFileJournalEntryHandle TestLatestJournalEntry;
+	FString ProjectDir = FPaths::ProjectDir();
+	IPlatformFile& PlatformFile = FPlatformFileManager::Get().GetPlatformFile();
+	FString TestVolumeName = PlatformFile.FileJournalGetVolumeName(ProjectDir);
+	EFileJournalResult TestResult = PlatformFile.FileJournalGetLatestEntry(*TestVolumeName,	TestLatestJournalEntry,
+		&TestError);
+	bool bPlatformSupported = TestResult == EFileJournalResult::Success;
+
+	bool bInvalidateEnabled = GPreloadSettings.IsDiscoveryCacheInvalidateEnabled();
+	bool bReadEnabled = GPreloadSettings.IsDiscoveryCacheReadEnabled() && (bPlatformSupported || !bInvalidateEnabled);
+	// Precalculate IfPlatformSupported -> Never if we already know the project doesn't support it
+	if (WriteEnabled == EFeatureEnabled::IfPlatformSupported && !bPlatformSupported)
+	{
+		WriteEnabled = EFeatureEnabled::Never;
+	}
+
+	if ((!bReadEnabled && GPreloadSettings.IsDiscoveryCacheReadEnabled()) ||
+		(WriteEnabled == EFeatureEnabled::Never && GPreloadSettings.IsDiscoveryCacheWriteEnabled() != EFeatureEnabled::Never))
+	{
+		const TCHAR* MissingOperation = (!bReadEnabled && WriteEnabled == EFeatureEnabled::Never) ? TEXT("read or written") :
+			(!bReadEnabled ? TEXT("read") : TEXT("written"));
+		UE_LOG(LogAssetRegistry, Display,
+			TEXT("PlatformFileJournal query failed for volume '%s' of project directory '%s', so AssetDiscovery cache will not be %s. PlatformFileJournal error information:\n\t%s"),
+			*TestVolumeName, *ProjectDir, MissingOperation, *TestError);
+	}
+
+	if (!bReadEnabled)
+	{
+		return;
+	}
+
+	bool bReadSucceeded = TryReadCacheFile();
+	if (!bReadSucceeded)
+	{
+		return;
+	}
+
+	if (!bInvalidateEnabled)
+	{
+		return;
+	}
+
+	for (TPair<FString, FCachedVolumeInfo>& VolumePair : CachedVolumes)
+	{
+		FString& VolumeName = VolumePair.Key;
+		FCachedVolumeInfo& VolumeInfo = VolumePair.Value;
+
+		TMap<FFileJournalFileHandle, FString> KnownDirectories;
+		TSet<FString> ModifiedDirectories;
+		for (TMap<FString, FCachedDirScanDir>::TIterator DirIter(VolumeInfo.Dirs); DirIter; ++DirIter)
+		{
+			FString& DirName = DirIter->Key;
+			FCachedDirScanDir& DirData = DirIter->Value;
+			// If not valid, we cannot remove the DirData yet because we will need it to get the list of
+			// its child directories that exist in the tree, at the point when we rescan it
+			if (DirData.bCacheValid && DirData.JournalHandle != FileJournalFileHandleInvalid)
+			{
+				KnownDirectories.Add(DirData.JournalHandle, DirName);
+			}
+		}
+
+		EFileJournalResult ReadModifiedResult = EFileJournalResult::Success;
+		bool bReadModifiedSucceeded = false;
+		if (!VolumeInfo.bJournalAvailable)
+		{
+			UE_LOG(LogAssetRegistry, Warning,
+				TEXT("PlatformFileJournal query failed for volume '%s'. AssetRegistry discovery of files on this volume will be uncached. PlatformFileJournal error information:\n\t%s"),
+				*VolumeName, *VolumeInfo.LastError);
+		}
+		else
+		{
+			EFileJournalResult Result = PlatformFile.FileJournalReadModified(*VolumeInfo.VolumeName,
+				VolumeInfo.NextJournalEntryToScan, KnownDirectories, ModifiedDirectories,
+				VolumeInfo.NextJournalEntryToScan, &VolumeInfo.LastError);
+			switch (Result)
+			{
+			case EFileJournalResult::Success:
+				bReadModifiedSucceeded = true;
+				break;
+			case EFileJournalResult::JournalWrapped:
+				UE_LOG(LogAssetRegistry, Warning,
+					TEXT("PlatformFileJournal journal has wrapped for volume '%s'. AssetRegistry discovery of files on this volume will be uncached. Notes on wrapping:")
+					TEXT("\r\n%s"),
+					*VolumeName, *VolumeInfo.LastError);
+				break;
+			default:
+				UE_LOG(LogAssetRegistry, Warning,
+					TEXT("PlatformFileJournal query failed for volume '%s'. AssetRegistry discovery of files on this volume will be uncached. PlatformFileJournal error information:")
+					TEXT("\n\t%s"),
+					*VolumeName, *VolumeInfo.LastError);
+				break;
+			}
+		}
+
+		if (!bReadModifiedSucceeded)
+		{
+			VolumeInfo.NextJournalEntryToScan = VolumeInfo.NextJournalEntryOnDisk;
+			VolumeInfo.Dirs.Empty();
+		}
+		else
+		{
+			for (const FString& ModifiedDirectory : ModifiedDirectories)
+			{
+				FCachedDirScanDir* DirData = VolumeInfo.Dirs.Find(ModifiedDirectory);
+				if (DirData)
+				{
+					// We cannot remove the DirData yet because we will need it to get the list of its child
+					// directories that exist in the tree, at the point when we rescan it
+					DirData->bCacheValid = false;
+				}
+			}
+		}
+	}
+}
+
+void FAssetDataDiscoveryCache::SaveCache()
+{
+	if (WriteEnabled == EFeatureEnabled::Never)
+	{
+		return;
+	}
+
+	TRACE_CPUPROFILER_EVENT_SCOPE(AssetDataGatherSaveDiscoveryCache);
+	for (TPair<FString, FCachedVolumeInfo>& VolumePair : CachedVolumes)
+	{
+		VolumePair.Value.PreSave();
+	}
+
+	FString Filename = GetCacheFileName();
+
+	FLargeMemoryWriter Writer;
+	SerializeWriteCacheFile(Writer);
+	FCompressedBuffer Compressed = FCompressedBuffer::Compress(FSharedBuffer::MakeView(Writer.GetView()));
+
+	TUniquePtr<FArchive> Ar(IFileManager::Get().CreateFileWriter(*Filename));
+	if (!Ar)
+	{
+		UE_LOG(LogAssetRegistry, Warning, TEXT("Could not write to DiscoveryCacheFile %s."), *Filename);
+		return;
+	}
+	Compressed.Save(*Ar);
+}
+
+bool FAssetDataDiscoveryCache::TryReadCacheFile()
+{
+	FSharedBuffer RawBuffer;
+	bool bCacheValid = true;
+	bool bCacheCorrupt = false;
+	FString Filename = GetCacheFileName();
+	{
+		TUniquePtr<FArchive> CompressedFileArchive(IFileManager::Get().CreateFileReader(*Filename));
+		if (!CompressedFileArchive || CompressedFileArchive->TotalSize() == 0)
+		{
+			bCacheValid = false;
+		}
+		else
+		{
+			FCompressedBufferReader CompressedBuffer(*CompressedFileArchive);
+			if (CompressedBuffer.GetRawSize() == 0)
+			{
+				bCacheValid = false;
+				bCacheCorrupt = true;
+			}
+			else
+			{
+				RawBuffer = CompressedBuffer.Decompress();
+				if (RawBuffer.GetSize() != CompressedBuffer.GetRawSize())
+				{
+					bCacheValid = false;
+					bCacheCorrupt = true;
+				}
+			}
+		}
+	}
+	if (bCacheValid)
+	{
+		FMemoryReaderView Ar(RawBuffer.GetView());
+		SerializeReadCacheFile(Ar);
+		if (Ar.IsError())
+		{
+			bCacheValid = false;
+			bCacheCorrupt = true;
+		}
+	}
+
+	if (!bCacheValid)
+	{
+		if (bCacheCorrupt)
+		{
+			UE_LOG(LogAssetRegistry, Warning,
+				TEXT("Corrupt AssetDiscovery cache %s. AssetRegistry discovery of files will be uncached."),
+				*Filename);
+		}
+		else
+		{
+			UE_LOG(LogAssetRegistry, Display,
+				TEXT("No AssetDiscovery cache present at %s. AssetRegistry discovery of files will be uncached."),
+				*Filename);
+		}
+		CachedVolumes.Empty();
+		return false;
+	}
+
+	for (TPair<FString, FCachedVolumeInfo>& Pair : CachedVolumes)
+	{
+		Pair.Value.InitializePlatformData();
+	}
+	return true;
+}
+
+void FAssetDataDiscoveryCache::SerializeReadCacheFile(FArchive& Ar)
+{
+	FGuid Version;
+	Ar << Version;
+	if (Version != AssetDataGathererConstants::DiscoveryCacheVersion)
+	{
+		Ar.SetError();
+		return;
+	}
+
+	CachedVolumes.Empty();
+	Ar << CachedVolumes;
+}
+
+void FAssetDataDiscoveryCache::SerializeWriteCacheFile(FArchive& Ar)
+{
+	FGuid Version = AssetDataGathererConstants::DiscoveryCacheVersion;
+	Ar << Version;
+	Ar << CachedVolumes;
+}
+
+void FAssetDataDiscoveryCache::Shutdown()
+{
+	WriteEnabled = EFeatureEnabled::Never;
+
+	CachedVolumes.Empty();
+	while (!ScanQueueDirFullDatas.IsEmpty())
+	{
+		TPair<FString, FCachedDirScanDir> Pair;
+		ScanQueueDirFullDatas.Dequeue(Pair);
+	}
+	while (!ScanQueueDirHandles.IsEmpty())
+	{
+		TPair<FString, FFileJournalFileHandle> Pair;
+		ScanQueueDirHandles.Dequeue(Pair);
+	}
+}
+
+FCachedVolumeInfo& FAssetDataDiscoveryCache::FindOrAddVolume(FStringView PathOrVolumeName)
+{
+	FStringView VolumeNameView;
+	FStringView Remainder;
+	FPathViews::SplitVolumeSpecifier(PathOrVolumeName, VolumeNameView, Remainder);
+	FString VolumeName(VolumeNameView);
+	if (VolumeName.IsEmpty())
+	{
+		VolumeName = GEmptyVolumeName;
+	}
+	FCachedVolumeInfo& Volume = CachedVolumes.FindOrAdd(VolumeName);
+	Volume.ConditionalConstruct(VolumeName);
+	return Volume;
+}
+
+FCachedDirScanDir& FAssetDataDiscoveryCache::FindOrAddDir(FStringView Path, FCachedVolumeInfo** OutVolume)
+{
+	FCachedVolumeInfo& Volume = FindOrAddVolume(Path);
+	if (OutVolume)
+	{
+		*OutVolume = &Volume;
+	}
+	return Volume.FindOrAddDir(Path);
+}
+
+void FAssetDataDiscoveryCache::RemoveDir(FStringView Path)
+{
+	FStringView VolumeNameView;
+	FStringView Remainder;
+	FPathViews::SplitVolumeSpecifier(Path, VolumeNameView, Remainder);
+	FString VolumeName(VolumeNameView);
+	FCachedVolumeInfo* Info = CachedVolumes.Find(VolumeName);
+	if (!Info)
+	{
+		return;
+	}
+	Info->RemoveDirs({ FString(Path) });
+}
+
+FCachedVolumeInfo* FAssetDataDiscoveryCache::FindVolume(FStringView PathOrVolumeName)
+{
+	FStringView VolumeNameView;
+	FStringView Remainder;
+	FPathViews::SplitVolumeSpecifier(PathOrVolumeName, VolumeNameView, Remainder);
+	FString VolumeName(VolumeNameView);
+
+	return CachedVolumes.Find(VolumeName);
+}
+
+FCachedDirScanDir* FAssetDataDiscoveryCache::FindDir(FStringView Path, FCachedVolumeInfo** OutVolume)
+{
+	FCachedVolumeInfo* Volume = FindVolume(Path);
+	if (OutVolume)
+	{
+		*OutVolume = Volume;
+	}
+	if (!Volume)
+	{
+		return nullptr;
+	}
+	return Volume->FindDir(Path);
+}
+
+void FAssetDataDiscoveryCache::QueueConsume()
+{
+	if (WriteEnabled == EFeatureEnabled::Never)
+	{
+		return;
+	}
+	TPair<FString, FFileJournalFileHandle> HandlePair;
+	while (ScanQueueDirHandles.Dequeue(HandlePair))
+	{
+		FString& DirName = HandlePair.Key;
+		FCachedVolumeInfo& Volume = FindOrAddVolume(DirName);
+		if (!Volume.bJournalAvailable && WriteEnabled != EFeatureEnabled::Always)
+		{
+			continue;
+		}
+		FCachedDirScanDir& Existing = Volume.FindOrAddDir(DirName);
+		Existing.JournalHandle = MoveTemp(HandlePair.Value);
+	}
+
+	TPair<FString, FCachedDirScanDir> ScanDirPair;
+	while (ScanQueueDirFullDatas.Dequeue(ScanDirPair))
+	{
+		FString& DirName = ScanDirPair.Key;
+		FCachedDirScanDir& ScanDir = ScanDirPair.Value;
+		FCachedVolumeInfo& Volume = FindOrAddVolume(DirName);
+		if (!Volume.bJournalAvailable && WriteEnabled != EFeatureEnabled::Always)
+		{
+			continue;
+		}
+
+		FCachedDirScanDir& Existing = Volume.FindOrAddDir(DirName);
+
+		// Mark for removal any subpaths in the cache that no longer exist on disk
+		if (!Existing.SubDirRelPaths.IsEmpty())
+		{
+			TSet<FString> StillExisting(ScanDir.SubDirRelPaths);
+			for (FString& OldRelPath : Existing.SubDirRelPaths)
+			{
+				if (!StillExisting.Contains(OldRelPath))
+				{
+					Volume.DirsToRemove.Add(FPaths::Combine(DirName, OldRelPath));
+				}
+			}
+		}
+
+		// If neither the new entry nor the existing entry have the JournalHandle,
+		// initialize it now.
+		if (ScanDir.JournalHandle == FileJournalFileHandleInvalid)
+		{
+			if (Existing.JournalHandle != FileJournalFileHandleInvalid)
+			{
+				ScanDir.JournalHandle = Existing.JournalHandle;
+			}
+			else
+			{
+				FFileJournalData PlatformData =
+					FPlatformFileManager::Get().GetPlatformFile().FileJournalGetFileData(*DirName);
+				ScanDir.JournalHandle = PlatformData.JournalHandle;
+			}
+		}
+
+		Existing = MoveTemp(ScanDir);
+		Existing.bCacheValid = true;
+	}
+}
+
+void FAssetDataDiscoveryCache::QueueAdd(FString DirName, FCachedDirScanDir DirData)
+{
+	if (WriteEnabled == EFeatureEnabled::Never)
+	{
+		return;
+	}
+	ScanQueueDirFullDatas.Enqueue(TPair<FString, FCachedDirScanDir>(MoveTemp(DirName), MoveTemp(DirData)));
+}
+
+void FAssetDataDiscoveryCache::QueueAdd(FString DirName, FFileJournalFileHandle JournalHandle)
+{
+	if (WriteEnabled == EFeatureEnabled::Never)
+	{
+		return;
+	}
+	ScanQueueDirHandles.Enqueue(TPair<FString, FFileJournalFileHandle>(MoveTemp(DirName), MoveTemp(JournalHandle)));
+}
+
+void FCachedVolumeInfo::ConditionalConstruct(const FString& InVolumeName)
+{
+	if (!VolumeName.IsEmpty())
+	{
+		return;
+	}
+	VolumeName = InVolumeName;
+	InitializePlatformData();
+}
+
+void FCachedVolumeInfo::PreSave()
+{
+	RemoveDirs(MoveTemp(DirsToRemove));
+}
+
+void FCachedVolumeInfo::InitializePlatformData()
+{
+	if (VolumeName.IsEmpty() || VolumeName == GEmptyVolumeName)
+	{
+		bJournalAvailable = false;
+		NextJournalEntryOnDisk = FileJournalEntryHandleInvalid;
+		NextJournalEntryToScan = FileJournalEntryHandleInvalid;
+	}
+	else
+	{
+		EFileJournalResult Result = FPlatformFileManager::Get().GetPlatformFile().FileJournalGetLatestEntry(
+			*VolumeName, NextJournalEntryOnDisk, &LastError);
+		bJournalAvailable = Result == EFileJournalResult::Success;
+		if (NextJournalEntryToScan == FileJournalEntryHandleInvalid)
+		{
+			NextJournalEntryToScan = NextJournalEntryOnDisk;
+		}
+	}
+}
+
+FCachedDirScanDir& FCachedVolumeInfo::FindOrAddDir(FStringView Path)
+{
+	uint32 PathHash = GetTypeHash(Path);
+	FCachedDirScanDir* ScanDir = Dirs.FindByHash(PathHash, Path);
+	if (!ScanDir)
+	{
+		ScanDir = &Dirs.FindOrAddByHash(PathHash, FString(Path));
+	}
+	return *ScanDir;
+}
+
+void FCachedVolumeInfo::RemoveDirs(TArray<FString>&& InPaths)
+{
+	// Recursively remove RemoveDir directories; iterate by popping InPaths and pushing child dirs back on
+	while (!InPaths.IsEmpty())
+	{
+		FString RemoveDir = InPaths.Pop(false /* bAllowShrinking */);
+		FCachedDirScanDir DirData;
+		if (Dirs.RemoveAndCopyValue(RemoveDir, DirData))
+		{
+			for (FString& RelPath : DirData.SubDirRelPaths)
+			{
+				InPaths.Add(FPaths::Combine(RemoveDir, RelPath));
+			}
+		}
+	}
+	InPaths.Empty(); // Free allocated memory
+}
+
+FCachedDirScanDir* FCachedVolumeInfo::FindDir(FStringView Path)
+{
+	uint32 PathHash = GetTypeHash(Path);
+	return Dirs.FindByHash(PathHash, Path);
+}
+
+FArchive& operator<<(FArchive& Ar, FCachedVolumeInfo& Data)
+{
+	Ar << Data.Dirs;
+	Ar << Data.VolumeName;
+	Ar << Data.NextJournalEntryToScan;
+	return Ar;
+}
+
+FArchive& operator<<(FArchive& Ar, FCachedDirScanDir& Data)
+{
+	Ar << Data.JournalHandle;
+	Ar << Data.SubDirRelPaths;
+	Ar << Data.Files;
+	Ar << Data.bCacheValid;
+	return Ar;
+}
+
+FArchive& operator<<(FArchive& Ar, FCachedDirScanFile& Data)
+{
+	Ar << Data.RelPath;
+	Ar << Data.ModificationTime;
+	return Ar;
+}
+
+
+} // namespace UE::AssetDataGather::Private
\ No newline at end of file
diff --git a/Engine/Source/Runtime/AssetRegistry/Private/AssetDataGathererDiscoveryCache.h b/Engine/Source/Runtime/AssetRegistry/Private/AssetDataGathererDiscoveryCache.h
new file mode 100644
--- /dev/null
+++ b/Engine/Source/Runtime/AssetRegistry/Private/AssetDataGathererDiscoveryCache.h
@@ -0,0 +1,128 @@
+// Copyright Epic Games, Inc. All Rights Reserved.
+
+#pragma once
+
+#include "Containers/Array.h"
+#include "Containers/Map.h"
+#include "Containers/MpscQueue.h"
+#include "Containers/SortedMap.h"
+#include "Containers/StringView.h"
+#include "Containers/UnrealString.h"
+#include "HAL/PlatformFile.h"
+#include "Misc/DateTime.h"
+
+class FArchive;
+
+namespace UE::AssetDataGather::Private
+{
+
+enum class EFeatureEnabled : uint8
+{
+	Never,
+	IfPlatformSupported,
+	Always
+};
+
+/** Data about a file in a cached directory used by FAssetDataDiscoveryCache to avoid the IO cost of rescanning. */
+struct FCachedDirScanFile
+{
+	FString RelPath;
+	FDateTime ModificationTime;
+};
+FArchive& operator<<(FArchive& Ar, FCachedDirScanFile& Data);
+
+/** Data about a directory used by FAssetDataDiscoveryCache to avoid the IO cost of rescanning. */
+struct FCachedDirScanDir
+{
+	TArray<FString> SubDirRelPaths;
+	TArray<FCachedDirScanFile> Files;
+	FFileJournalFileHandle JournalHandle = FileJournalFileHandleInvalid;
+	bool bCacheValid = false;
+};
+FArchive& operator<<(FArchive& Ar, FCachedDirScanDir& Data);
+
+/** Data about a volume used by FAssetDataDiscoveryCache to avoid the IO cost of rescanning directories on the volume. */
+class FCachedVolumeInfo
+{
+public:
+	void ConditionalConstruct(const FString& InVolumeName);
+	void PreSave();
+	void InitializePlatformData();
+	
+	/** Assumes Path has already been normalized. */
+	FCachedDirScanDir& FindOrAddDir(FStringView Path);
+	void RemoveDirs(TArray<FString>&& InPaths);
+
+	/** Assumes Path has already been normalized. */
+	FCachedDirScanDir* FindDir(FStringView Path);
+
+public:
+	TMap<FString, FCachedDirScanDir> Dirs;
+	FString VolumeName;
+	FFileJournalEntryHandle NextJournalEntryToScan = FileJournalEntryHandleInvalid;
+
+	// Transient
+	/** Directories that we marked for recursive removal from CachedVolumes; we consume this during WriteCacheFile. */
+	TArray<FString> DirsToRemove;
+	FString LastError;
+	FFileJournalEntryHandle NextJournalEntryOnDisk = FileJournalEntryHandleInvalid;
+	bool bJournalAvailable = false;
+};
+FArchive& operator<<(FArchive& Ar, FCachedVolumeInfo& Data);
+
+/**
+ * Keeps a record of the directories, files, and file timestamps that were found by the directory
+ * scan the last time it ran, and invalidates the records that have been reported modified by
+ * the PlatformFileJournal.
+ * 
+ * ThreadSafety: Not threadsafe, with the exception of QueueAdd functions, which can be called
+ * at any time from any thread.
+ * 
+ */
+class FAssetDataDiscoveryCache
+{
+public:
+	FString GetCacheFileName() const;
+
+	bool IsInitialized() { return bInitialized; }
+	EFeatureEnabled IsWriteEnabled() { return WriteEnabled; }
+	void Shutdown();
+
+	void LoadAndUpdateCache();
+	void SaveCache();
+
+	FCachedVolumeInfo& FindOrAddVolume(FStringView PathOrVolumeName);
+	/** Assumes Path has already been normalized. */
+	FCachedDirScanDir& FindOrAddDir(FStringView Path, FCachedVolumeInfo** OutVolume = nullptr);
+	void RemoveDir(FStringView Path);
+
+	FCachedVolumeInfo* FindVolume(FStringView PathOrVolumeName);
+	/** Assumes Path has already been normalized. */
+	FCachedDirScanDir* FindDir(FStringView Path, FCachedVolumeInfo** OutVolume = nullptr);
+
+	void QueueConsume();
+	void QueueAdd(FString DirName, FCachedDirScanDir DirData);
+	void QueueAdd(FString DirName, FFileJournalFileHandle JournalHandle);
+
+private:
+	bool TryReadCacheFile();
+	void SerializeReadCacheFile(FArchive& Ar);
+	void SerializeWriteCacheFile(FArchive& Ar);
+
+private:
+	/** Cached information about each volume on disk that has mounted directories. */
+	TSortedMap<FString, FCachedVolumeInfo> CachedVolumes;
+	/** Childpath information about directories collected from ParallelFor during scanning. */
+	TMpscQueue<TPair<FString, FCachedDirScanDir>> ScanQueueDirFullDatas;
+	/** JournalHandle information about directories collected during ParallelFor during scanning. */
+	TMpscQueue<TPair<FString, FFileJournalFileHandle>> ScanQueueDirHandles;
+
+	/** Used to implement initialization on demand. */
+	bool bInitialized = false;
+	/** Whether writing the cache is enabled for no volumes, all volumes, or ones with journal available. */
+	EFeatureEnabled WriteEnabled = EFeatureEnabled::Never;
+};
+
+constexpr const TCHAR* GEmptyVolumeName = TEXT("<EmptyVolume>");
+
+} // UE::AssetDataGather::Private
\ No newline at end of file
diff --git a/Engine/Source/Runtime/AssetRegistry/Private/AssetDataGathererPrivate.h b/Engine/Source/Runtime/AssetRegistry/Private/AssetDataGathererPrivate.h
--- a/Engine/Source/Runtime/AssetRegistry/Private/AssetDataGathererPrivate.h
+++ b/Engine/Source/Runtime/AssetRegistry/Private/AssetDataGathererPrivate.h
@@ -4,6 +4,7 @@
 
 #include "AssetDataGatherer.h"
 
+#include "AssetDataGathererDiscoveryCache.h"
 #include "Containers/Set.h"
 #include "Misc/Optional.h"
 #include "Misc/StringBuilder.h"
@@ -598,10 +599,12 @@ public:
 	void GetAndTrimSearchResults(bool& bOutIsComplete, TArray<FString>& OutDiscoveredPaths,
 		FFilesToSearch& OutFilesToSearch, int32& OutNumPathsToSearch);
 	/** Get diagnostics for telemetry or logging. */
-	void GetDiagnostics(float& OutCumulativeDiscoveryTime);
+	void GetDiagnostics(float& OutCumulativeDiscoveryTime, int32& OutNumCachedDirectories,
+		int32& OutNumUncachedDirectories);
 	/** Wait (joining in on the tick) until all currently monitored paths have been scanned. */
 	void WaitForIdle(double EndTimeSeconds);
 	bool IsIdle() const;
+	void OnInitialSearchCompleted();
 
 	/** Optionally set some scan properties for the given paths and then wait for their scans to finish. */
 	void SetPropertiesAndWait(TArrayView<FPathExistence> QueryPaths, bool bAddToAllowList, bool bForceRescan,
@@ -777,6 +780,10 @@ private:
 	int32 CumulativeDiscoveredFiles = 0;
 	/** Cumulative total of time spent in all non-idle periods. */
 	float CumulativeDiscoveryTime = 0.f;
+	/** The total number of directories in the search results that were read from the cache. */
+	int32 NumCachedDirectories = 0;
+	/** The total number of directories in the search results that were not in the cache and were read by iterating the disk. */
+	int32 NumUncachedDirectories = 0;
 
 	// Variable section for variables that are read/writable only within TreeLock.
 
@@ -832,20 +839,25 @@ private:
 	};
 	TArray<FDirToScanBuffer> DirToScanBuffers;
 
+	FAssetDataDiscoveryCache Cache;
+
 	friend class FMountDir;
 	friend class FScanDir;
 };
 
 /**
  * Settings about whether to use cache data for the AssetDataGatherer; these settings are shared by
- * FPreloader and the FAssetDataGatherer.
+ * FPreloader, FAssetDataGatherer, and FAssetDataDiscovery.
  */
 struct FPreloadSettings
 {
 public:
 	void Initialize();
-	bool IsCacheReadEnabled() const;
-	bool IsCacheWriteEnabled() const;
+	bool IsGatherCacheReadEnabled() const;
+	bool IsGatherCacheWriteEnabled() const;
+	bool IsDiscoveryCacheReadEnabled() const;
+	EFeatureEnabled IsDiscoveryCacheWriteEnabled() const;
+	bool IsDiscoveryCacheInvalidateEnabled() const;
 	bool IsMonolithicCacheActivatedDuringPreload() const;
 	bool IsPreloadMonolithicCache() const;
 	bool IsGatherDependsData() const;
@@ -860,8 +872,11 @@ private:
 	FString AssetRegistryCacheRootFolder;
 	bool bForceDependsGathering = false;
 	bool bGatherDependsData = false;
-	bool bCacheReadEnabled = false;
-	bool bCacheWriteEnabled = false;
+	bool bGatherCacheReadEnabled = false;
+	bool bGatherCacheWriteEnabled = false;
+	bool bDiscoveryCacheReadEnabled = false;
+	EFeatureEnabled DiscoveryCacheWriteEnabled = EFeatureEnabled::Never;
+	bool bDiscoveryCacheInvalidateEnabled = false;
 	bool bMonolithicCacheActivatedDuringPreload = false;
 	bool bInitialized = false;
 };
diff --git a/Engine/Source/Runtime/AssetRegistry/Private/AssetRegistry.cpp b/Engine/Source/Runtime/AssetRegistry/Private/AssetRegistry.cpp
--- a/Engine/Source/Runtime/AssetRegistry/Private/AssetRegistry.cpp
+++ b/Engine/Source/Runtime/AssetRegistry/Private/AssetRegistry.cpp
@@ -4340,25 +4340,32 @@ Impl::EGatherStatus FAssetRegistryImpl::TickGatherer(Impl::FEventContext& EventC
 		// startup up until that point, and we need to wait for all the plugins to load before declaring completion.
 		if (!bInitialSearchCompleted && bPreloadingComplete && IsEngineStartupModuleLoadingComplete())
 		{
+			RecordTimer(); // OnInitialSearchComplete reads data set by RecordTimer
+			OnInitialSearchCompleted(EventContext);
+		}
+	}
+
+	return OutStatus;
+}
+
+void FAssetRegistryImpl::OnInitialSearchCompleted(Impl::FEventContext& EventContext)
+{
+	bInitialSearchCompleted = true;
+
 #if WITH_EDITOR
-			// update redirectors
-			UpdateRedirectCollector();
+	// update redirectors
+	UpdateRedirectCollector();
 #endif
 
-			// Handle any deferred loading operations
-			SetPerformanceMode(Impl::EPerformanceMode::MostlyStatic);
+	// Handle any deferred loading operations
+	SetPerformanceMode(Impl::EPerformanceMode::MostlyStatic);
 
-			RecordTimer();
-			LogSearchDiagnostics(InitialSearchStartTime);
-			TRACE_END_REGION(TEXT("Asset Registry Scan"));
+	LogSearchDiagnostics(InitialSearchStartTime);
+	TRACE_END_REGION(TEXT("Asset Registry Scan"));
 
-			bInitialSearchCompleted = true;
+	GlobalGatherer->OnInitialSearchCompleted();
 
-			EventContext.bFileLoadedEventBroadcast = true;
-		}
-	}
-
-	return OutStatus;
+	EventContext.bFileLoadedEventBroadcast = true;
 }
 
 void FAssetRegistryImpl::LogSearchDiagnostics(double StartTime)
@@ -4371,11 +4378,16 @@ void FAssetRegistryImpl::LogSearchDiagnostics(double StartTime)
 	Telemetry.DiscoveryTimeSeconds = Diagnostics.DiscoveryTimeSeconds;
 	Telemetry.GatherTimeSeconds = Diagnostics.GatherTimeSeconds;
 	Telemetry.StoreTimeSeconds = StoreGatherResultsTimeSeconds;
+	Telemetry.NumCachedDirectories = Diagnostics.NumCachedDirectories;
+	Telemetry.NumUncachedDirectories = Diagnostics.NumUncachedDirectories;
 	Telemetry.NumCachedAssetFiles = Diagnostics.NumCachedAssetFiles;
 	Telemetry.NumUncachedAssetFiles = Diagnostics.NumUncachedAssetFiles;
 	FTelemetryRouter::Get().ProvideTelemetry(Telemetry);
-	UE_LOG(LogAssetRegistry, Log, TEXT("AssetRegistryGather time %.4fs: AssetDataDiscovery %0.4fs, AssetDataGather %0.4fs, StoreResults %0.4fs."),
-		Total, Diagnostics.DiscoveryTimeSeconds, Diagnostics.GatherTimeSeconds, StoreGatherResultsTimeSeconds);
+	UE_LOG(LogAssetRegistry, Log, TEXT("AssetRegistryGather time %.4fs: AssetDataDiscovery %0.4fs, AssetDataGather %0.4fs, StoreResults %0.4fs.")
+		TEXT("\n\tNumCachedDirectories %d. NumUncachedDirectories %d. NumCachedFiles %d. NumUncachedFiles %d."),
+		Total, Diagnostics.DiscoveryTimeSeconds, Diagnostics.GatherTimeSeconds, StoreGatherResultsTimeSeconds,
+		Diagnostics.NumCachedDirectories, Diagnostics.NumUncachedDirectories, Diagnostics.NumCachedAssetFiles,
+		Diagnostics.NumUncachedAssetFiles);
 
 #if !NO_LOGGING
 	if (bVerboseLogging)
diff --git a/Engine/Source/Runtime/AssetRegistry/Private/AssetRegistryImpl.h b/Engine/Source/Runtime/AssetRegistry/Private/AssetRegistryImpl.h
--- a/Engine/Source/Runtime/AssetRegistry/Private/AssetRegistryImpl.h
+++ b/Engine/Source/Runtime/AssetRegistry/Private/AssetRegistryImpl.h
@@ -367,6 +367,7 @@ private:
 		Premade::ELoadResult LoadResult, FAssetRegistryState&& ARState);
 	/** Add MountPoints of all AssetDatas currently registered in this->State to the list of PersistentMountPoints. */
 	void UpdatePersistentMountPoints();
+	void OnInitialSearchCompleted(Impl::FEventContext& EventContext);
 
 private:
 
diff --git a/Engine/Source/Runtime/AssetRegistry/Public/AssetRegistry/AssetRegistryTelemetry.h b/Engine/Source/Runtime/AssetRegistry/Public/AssetRegistry/AssetRegistryTelemetry.h
--- a/Engine/Source/Runtime/AssetRegistry/Public/AssetRegistry/AssetRegistryTelemetry.h
+++ b/Engine/Source/Runtime/AssetRegistry/Public/AssetRegistry/AssetRegistryTelemetry.h
@@ -58,6 +58,10 @@ namespace UE::Telemetry::AssetRegistry
         double GatherTimeSeconds = 0.0;
         // Time spent storing asset data in the asset registry for searching
         double StoreTimeSeconds = 0.0;
+        // Number of directories read from cache
+        int32 NumCachedDirectories = 0;
+        // Number of directories scanned from disk
+        int32 NumUncachedDirectories = 0;
         // Number of asset files read from cache
         int32 NumCachedAssetFiles = 0;
         // Number of asset files read loose from disk
diff --git a/Engine/Source/Runtime/Core/Private/GenericPlatform/GenericPlatformFile.cpp b/Engine/Source/Runtime/Core/Private/GenericPlatform/GenericPlatformFile.cpp
--- a/Engine/Source/Runtime/Core/Private/GenericPlatform/GenericPlatformFile.cpp
+++ b/Engine/Source/Runtime/Core/Private/GenericPlatform/GenericPlatformFile.cpp
@@ -12,6 +12,7 @@
 #include "Templates/UniquePtr.h"
 #include "Misc/ScopeLock.h"
 #include "HAL/LowLevelMemTracker.h"
+#include "String/BytesToHex.h"
 
 #include "Async/AsyncFileHandle.h"
 #include "Async/MappedFileHandle.h"
@@ -964,6 +965,108 @@ bool IPlatformFile::CreateDirectoryTree(const TCHAR* Directory)
 	return InternalCreateDirectoryTree(*this, LocalDirname);
 }
 
+FString FFileJournalFileHandle::ToString()
+{
+	FString Output;
+	TArray<TCHAR, FString::AllocatorType>& CharArray = Output.GetCharArray();
+	CharArray.AddUninitialized(sizeof(FFileJournalFileHandle) * 2 + 3);
+	TCHAR* Data = CharArray.GetData();
+	Data[0] = '0';
+	Data[1] = 'x';
+	UE::String::BytesToHexLower(Bytes, Data + 2);
+	CharArray.Last() = TCHAR('\0');
+	return Output;
+}
+
+namespace UE::PlatformFileJournal::Private
+{
+
+FFileJournalFileHandle CreateInvalidFileHandle()
+{
+	FFileJournalFileHandle Result;
+	for (uint8& Byte : Result.Bytes)
+	{
+		Byte = 0;
+	}
+	return Result;
+}
+
+FFileJournalData ToJournalData(const FFileStatData& StatData)
+{
+	FFileJournalData JournalData;
+	JournalData.ModificationTime = StatData.ModificationTime;
+	JournalData.JournalHandle = FileJournalFileHandleInvalid;
+	JournalData.bIsValid = StatData.bIsValid;
+	JournalData.bIsDirectory = StatData.bIsDirectory;
+	return JournalData;
+}
+
+constexpr const TCHAR* PlatformNotAvailableMessage = TEXT("PlatformFileJournal is not implemented on the current platform.");
+
+} // namespace UE::PlatformFileJournal::Private
+
+const FFileJournalFileHandle FileJournalFileHandleInvalid = UE::PlatformFileJournal::Private::CreateInvalidFileHandle();
+
+bool IPlatformFile::FileJournalIsAvailable(const TCHAR* VolumeOrPath, ELogVerbosity::Type* OutErrorLevel,
+	FString* OutError)
+{
+	if (OutErrorLevel)
+	{
+		*OutErrorLevel = ELogVerbosity::Display;
+	}
+	if (OutError)
+	{
+		*OutError = UE::PlatformFileJournal::Private::PlatformNotAvailableMessage;
+	}
+	return false;
+}
+
+EFileJournalResult IPlatformFile::FileJournalGetLatestEntry(const TCHAR* VolumeOrPath,
+	FFileJournalEntryHandle& OutEntryHandle, FString* OutError)
+{
+	if (OutError)
+	{
+		*OutError = UE::PlatformFileJournal::Private::PlatformNotAvailableMessage;
+	}
+	OutEntryHandle = 0;
+	return EFileJournalResult::InvalidPlatform;
+}
+
+bool IPlatformFile::FileJournalIterateDirectory(const TCHAR* Directory, FDirectoryJournalVisitorFunc Visitor)
+{
+	return IFileManager::Get().IterateDirectoryStat(Directory,
+		[&Visitor](const TCHAR* InPackageFilename, const FFileStatData& StatData)
+		{
+			return Visitor(InPackageFilename, UE::PlatformFileJournal::Private::ToJournalData(StatData));
+		});
+}
+
+FFileJournalData IPlatformFile::FileJournalGetFileData(const TCHAR* FilenameOrDirectory)
+{
+	return UE::PlatformFileJournal::Private::ToJournalData(IFileManager::Get().GetStatData(FilenameOrDirectory));
+}
+
+EFileJournalResult IPlatformFile::FileJournalReadModified(const TCHAR* VolumeName,
+	const FFileJournalEntryHandle& StartingJournalEntry, TMap<FFileJournalFileHandle, FString>& KnownDirectories,
+	TSet<FString>& OutModifiedDirectories, FFileJournalEntryHandle& OutNextJournalEntry, FString* OutError)
+{
+	OutNextJournalEntry = StartingJournalEntry;
+	if (OutError)
+	{
+		*OutError = UE::PlatformFileJournal::Private::PlatformNotAvailableMessage;
+	}
+	return EFileJournalResult::InvalidPlatform;
+}
+
+FString IPlatformFile::FileJournalGetVolumeName(FStringView Path)
+{
+	FString FullPath = FPaths::ConvertRelativePathToFull(FString(Path));
+	FStringView VolumeName;
+	FStringView Remainder;
+	FPathViews::SplitVolumeSpecifier(FullPath, VolumeName, Remainder);
+	return FString(VolumeName);
+}
+
 bool IPhysicalPlatformFile::Initialize(IPlatformFile* Inner, const TCHAR* CmdLine)
 {
 	// Physical platform file should never wrap anything.
diff --git a/Engine/Source/Runtime/Core/Private/Windows/WindowsPlatformFileJournal.cpp b/Engine/Source/Runtime/Core/Private/Windows/WindowsPlatformFileJournal.cpp
new file mode 100644
--- /dev/null
+++ b/Engine/Source/Runtime/Core/Private/Windows/WindowsPlatformFileJournal.cpp
@@ -0,0 +1,888 @@
+// Copyright Epic Games, Inc. All Rights Reserved.
+
+/**
+ * The FileJournal API on Windows is implemented using Windows' USN Journal (aka NTFS Journal).
+ *
+ * USN stands for update sequence number.
+ * The Journal has to be turned on as an fsutil feature on each drive. When it is enabled on a drive,
+ * all writes to files on the drive cause the recording of an entry describing the modification into the journal.
+ * The journal has a fixed size (specified at creation) and will drop older entries to make room for new ones.
+ * 
+ * We only enable it on editor because it has not been requested for in-game use and we want to reduce executable size.
+ * 
+ * Some features of the Journal require structs and enums defined in _WIN32_WINNT_WIN8 or higher.
+ * If we are compiling against a lower version of the windows library, we define the structs and enums
+ * we need. The API will still work (assuming the machine running our process supports it) because the library
+ * just forwards our function call on to the dll on the machine. If the current machine is running Windows7
+ * or lower, we gracefully fail the journal operations.
+ *
+ */
+
+#include "Windows/WindowsPlatformFilePrivate.h"
+
+#include "Containers/StringView.h"
+#include "Logging/LogVerbosity.h"
+#include "Misc/Paths.h"
+#include "Misc/PathViews.h"
+
+#define UE_WINDOWS_PLATFORM_FILEJOURNAL_ENABLED WITH_EDITOR 
+
+#if UE_WINDOWS_PLATFORM_FILEJOURNAL_ENABLED
+#include "Windows/WindowsPlatform.h" // Defines LogWindows, must come before CoreGlobals.h
+
+#include "Containers/Map.h"
+#include "Containers/Set.h"
+#include "CoreGlobals.h"
+#include "HAL/PlatformMisc.h"
+#include "Logging/LogMacros.h"
+#include "Misc/ScopeExit.h"
+#include "Misc/StringBuilder.h"
+#include "Windows/WindowsHWrapper.h"
+
+// windows.h -> winscard.h -> winioctl.h is excluded by both UE_MINIMAL_WINDOWS_INCLUDE and NOCRYPT, which we define in
+// WindowsHWrapper -> MinWindows.h before including windows.h. Include winioctl.h directly since it was excluded
+#if defined(WINDOWS_H_WRAPPER_GUARD)
+#error WINDOWS_H_WRAPPER_GUARD already defined
+#endif
+#define WINDOWS_H_WRAPPER_GUARD
+#include "Windows/AllowWindowsPlatformTypes.h"
+#include "Windows/PreWindowsApi.h"
+#include <winioctl.h>
+#include "Windows/PostWindowsApi.h"
+#include "Windows/HideWindowsPlatformTypes.h"
+#undef WINDOWS_H_WRAPPER_GUARD
+
+#endif
+
+FString FWindowsPlatformFile::FileJournalGetVolumeName(FStringView Path)
+{
+	FString FullPath = FPaths::ConvertRelativePathToFull(FString(Path));
+	FStringView VolumeName;
+	FStringView Remainder;
+	FPathViews::SplitVolumeSpecifier(FullPath, VolumeName, Remainder);
+	if (VolumeName.EndsWith(TEXT(":")))
+	{
+		return FString(VolumeName);
+	}
+	return FString(); // Volume names of the form \\volumename are not supported by USNJournal
+}
+
+#if UE_WINDOWS_PLATFORM_FILEJOURNAL_ENABLED 
+
+#if _WIN32_WINNT < _WIN32_WINNT_WIN8
+typedef struct _FILE_ID_EXTD_DIR_INFO {
+	ULONG NextEntryOffset;
+	ULONG FileIndex;
+	LARGE_INTEGER CreationTime;
+	LARGE_INTEGER LastAccessTime;
+	LARGE_INTEGER LastWriteTime;
+	LARGE_INTEGER ChangeTime;
+	LARGE_INTEGER EndOfFile;
+	LARGE_INTEGER AllocationSize;
+	ULONG FileAttributes;
+	ULONG FileNameLength;
+	ULONG EaSize;
+	ULONG ReparsePointTag;
+	FILE_ID_128 FileId;
+	WCHAR FileName[1];
+} FILE_ID_EXTD_DIR_INFO, * PFILE_ID_EXTD_DIR_INFO;
+
+constexpr _FILE_INFO_BY_HANDLE_CLASS FileIdExtdDirectoryInfo = (_FILE_INFO_BY_HANDLE_CLASS)(((int32)FileFullDirectoryRestartInfo) + 4);
+constexpr _FILE_INFO_BY_HANDLE_CLASS FileIdExtdDirectoryRestartInfo = (_FILE_INFO_BY_HANDLE_CLASS)(((int32)FileIdExtdDirectoryInfo) + 1);
+
+#endif
+
+namespace UE::WindowsPlatformFileJournal::Private
+{
+
+bool FileJournalIsAvailableSystemWide();
+bool FileJournalIsAvailableSystemWide(ELogVerbosity::Type* OutErrorLevel, FString* OutError);
+EFileJournalResult CreateJournalWindowsHandle(const TCHAR* VolumeOrPath, HANDLE& OutHandle,
+	FStringBuilderBase& OutVolumeName, FString* OutError);
+EFileJournalResult CreateJournalWindowsHandle(const TCHAR* VolumeName, HANDLE& OutHandle, FString* OutError);
+EFileJournalResult GetJournalDescriptor(HANDLE JournalWindowsHandle, const TCHAR* VolumeName,
+	USN_JOURNAL_DATA_V1& OutJournalDescriptor, FString* OutError);
+FString JournalGetLastErrorString(uint32 LastError);
+
+FStringView GetFileName(FILE_ID_EXTD_DIR_INFO* Info);
+FStringView GetFileName(USN_RECORD_V3* Record);
+
+FFileJournalFileHandle ToFileHandle(uint64 ReferenceNumber);
+FFileJournalFileHandle ToFileHandle(FILE_ID_128& ReferenceNumber);
+
+uint8 GJournalAvailable = 0;
+
+};
+
+bool FWindowsPlatformFile::FileJournalIsAvailable(const TCHAR* VolumeOrPath, ELogVerbosity::Type* OutErrorLevel,
+	FString* OutError)
+{
+	using namespace UE::WindowsPlatformFileJournal::Private;
+
+	if (!FileJournalIsAvailableSystemWide(OutErrorLevel, OutError))
+	{
+		return false;
+	}
+	if (!VolumeOrPath)
+	{
+		return true;
+	}
+
+	FFileJournalEntryHandle Handle;
+	EFileJournalResult Result = FileJournalGetLatestEntry(VolumeOrPath, Handle, OutError);
+	if (Result == EFileJournalResult::Success)
+	{
+		if (OutErrorLevel)
+		{
+			*OutErrorLevel = ELogVerbosity::Verbose;
+		}
+		return true;
+	}
+	else
+	{
+		if (OutErrorLevel)
+		{
+			*OutErrorLevel = ELogVerbosity::Error;
+		}
+		return false;
+	}
+}
+
+
+EFileJournalResult FWindowsPlatformFile::FileJournalGetLatestEntry(const TCHAR* VolumeName,
+	FFileJournalEntryHandle& OutEntryHandle, FString* OutError)
+{
+	using namespace UE::WindowsPlatformFileJournal::Private;
+
+	OutEntryHandle = 0;
+
+	HANDLE JournalHandle;
+	TStringBuilder<16> NormalizedVolumeName;
+	EFileJournalResult Result = CreateJournalWindowsHandle(VolumeName, JournalHandle, NormalizedVolumeName,
+		OutError);
+	if (Result != EFileJournalResult::Success)
+	{
+		return Result;
+	}
+	ON_SCOPE_EXIT{ ::CloseHandle(JournalHandle); };
+
+	USN_JOURNAL_DATA_V1 JournalDescriptor;
+	Result = GetJournalDescriptor(JournalHandle, *NormalizedVolumeName, JournalDescriptor, OutError);
+	if (Result != EFileJournalResult::Success)
+	{
+		return Result;
+	}
+	static_assert(sizeof(uint64) >= sizeof(USN), "We are storing USNs as uint64");
+	OutEntryHandle = static_cast<FFileJournalEntryHandle>(JournalDescriptor.NextUsn);
+	return EFileJournalResult::Success;
+}
+
+bool FWindowsPlatformFile::FileJournalIterateDirectory(const TCHAR* Directory,
+	FDirectoryJournalVisitorFunc Visitor)
+{
+	using namespace UE::WindowsPlatformFileJournal::Private;
+	using namespace UE::WindowsPlatformFile::Private;
+
+	if (!FileJournalIsAvailableSystemWide())
+	{
+		return IPlatformFile::FileJournalIterateDirectory(Directory, Visitor);
+	}
+
+	FNormalizedDirectory Normalized(Directory);
+	HANDLE Handle = CreateFileW(*Normalized,
+		GENERIC_READ,
+		FILE_SHARE_READ,
+		nullptr,
+		OPEN_EXISTING,
+		FILE_FLAG_BACKUP_SEMANTICS, // Required when opening a directory
+		nullptr);
+	if (Handle == INVALID_HANDLE_VALUE)
+	{
+		return false;
+	}
+	ON_SCOPE_EXIT{ ::CloseHandle(Handle); };
+
+	// Add overflow to BufferSize to allow check for a null terminator by reading past the end of FileNameLength
+	constexpr int BufferSize = 4096;
+	constexpr int BufferSizeOverflow = 1;
+	uint8 Buffer[BufferSize + BufferSizeOverflow];
+
+	BOOL Result = GetFileInformationByHandleEx(Handle, ::FileIdExtdDirectoryRestartInfo, Buffer, BufferSize);
+	if (!Result)
+	{
+		// TODO: Need to handle filenames with size > buffersize by checking GetLastError and passing in a dynamically sized buffer.
+		return false;
+	}
+
+	TStringBuilder<1024> LocalAbsPath;
+	LocalAbsPath << Directory;
+	int32 SavedLocalAbsPathLen = LocalAbsPath.Len();
+
+	for (;;) // For each call to GetFileInformationByHandleEx
+	{
+		DWORD CurrentOffset = 0;
+		for (;;) // For every FILE_ID_EXTD_DIR_INFO in the buffer returned by GetFileInformationByHandleEx
+		{
+			FILE_ID_EXTD_DIR_INFO* Info = reinterpret_cast<FILE_ID_EXTD_DIR_INFO*>(Buffer + CurrentOffset);
+			FStringView RelPath = GetFileName(Info);
+
+			// GetFileInformationByHandleEx returns entries for . and .., but our contract is that we only
+			// return entries for child paths
+			if (RelPath != TEXTVIEW(".") && RelPath != TEXTVIEW(".."))
+			{
+				check(!RelPath.Contains(TEXT("/")) && !RelPath.Contains(TEXT("\\"))); // RemoveSuffix would break if we called AppendPath with an absolute path
+				LocalAbsPath.RemoveSuffix(LocalAbsPath.Len() - SavedLocalAbsPathLen);
+				FPathViews::AppendPath(LocalAbsPath, RelPath);
+
+				FFileJournalData FileData;
+				FILETIME LastWriteTime;
+				LastWriteTime.dwLowDateTime = Info->LastWriteTime.LowPart;
+				LastWriteTime.dwHighDateTime = Info->LastWriteTime.HighPart;
+				FileData.ModificationTime = WindowsFileTimeToUEDateTime(LastWriteTime);
+				FileData.bIsValid = true;
+				FileData.bIsDirectory = !!(Info->FileAttributes & FILE_ATTRIBUTE_DIRECTORY);
+				FileData.JournalHandle = ToFileHandle(Info->FileId);
+
+				if (!Visitor(*LocalAbsPath, FileData))
+				{
+					return false;
+				}
+			}
+
+			if (Info->NextEntryOffset == 0)
+			{
+				break;
+			}
+			CurrentOffset += Info->NextEntryOffset;
+		}
+
+		BOOL NextResult = GetFileInformationByHandleEx(Handle, ::FileIdExtdDirectoryInfo, Buffer, BufferSize);
+		if (!NextResult)
+		{
+			// TODO: Need to handle filenames with size > buffersize by checking GetLastError and passing in a dynamically sized buffer.
+			return true;
+		}
+	}
+	checkNoEntry(); // loop body does not break, it only returns
+	return true;
+}
+
+FFileJournalData FWindowsPlatformFile::FileJournalGetFileData(const TCHAR* FilenameOrDirectory)
+{
+	using namespace UE::WindowsPlatformFileJournal::Private;
+	using namespace UE::WindowsPlatformFile::Private;
+
+	if (!FileJournalIsAvailableSystemWide())
+	{
+		return IPlatformFile::FileJournalGetFileData(FilenameOrDirectory);
+	}
+
+	FFileJournalData OutResult;
+	FNormalizedFilename Normalized(FilenameOrDirectory);
+	// We don't know whether FilenameOrDirectory is a directory, so we can't use FNormalizedDirectory and have to use
+	// FNormalizedFile. Unlike FNormalizedDirectory, FNormalizedFile does not remove redundant terminating separators.
+	// But CreateFileW requires that they be removed, so remove them.
+	while (FPathViews::HasRedundantTerminatingSeparator(Normalized.ToView()))
+	{
+		Normalized.RemoveSuffix(1);
+	}
+
+	HANDLE Handle = CreateFileW(*Normalized,
+		GENERIC_READ,
+		FILE_SHARE_READ,
+		nullptr,
+		OPEN_EXISTING,
+		FILE_FLAG_BACKUP_SEMANTICS, // Required when opening a directory
+		nullptr);
+	if (Handle == INVALID_HANDLE_VALUE)
+	{
+		return OutResult;
+	}
+	ON_SCOPE_EXIT{ ::CloseHandle(Handle); };
+
+	// GetFileInformationByHandleEx returns information for as many paths in the directory as fit in the buffer
+	// we provide. The first one is the directory itself, with filename '.'.
+	// That's the only one we want to fetch for this call; we don't want to waste time copying data into our buffer
+	// about the other files. Since the filename is a single character, and sizeof(FILE_ID_EXTD_DIR_INFO) includes the
+	// first character of the filename, then the record for this filename will fit into a buffer of exactly
+	// sizeof(FILE_ID_EXTD_DIR_INFO). But make the buffer bigger than that (but smaller than 2*FILE_ID_EXTD_DIR_INFO)
+	// just to be on the safe side.
+	static_assert(sizeof(FILE_ID_EXTD_DIR_INFO) > 20); // Our extra space should be smaller than half of a buffer
+	uint8 Buffer[sizeof(FILE_ID_EXTD_DIR_INFO) + 10];
+
+	BOOL Result = GetFileInformationByHandleEx(Handle, ::FileIdExtdDirectoryRestartInfo, Buffer, sizeof(Buffer));
+	if (!Result)
+	{
+		uint32 LastError = FPlatformMisc::GetLastError();
+		// ERROR_BAD_LENGTH means our buffer is too short, which indicates our comment above is wrong and we
+		// need to fix the code. Other failures can happen and are not necessarily errors.
+		if (LastError == ERROR_BAD_LENGTH)
+		{
+			UE_LOG(LogWindows, Error, TEXT("Unexpected failure of GetFileInformationByHandleEx('%s') with error code %u (0x%08x) (%s)."),
+				*Normalized, LastError, LastError, *JournalGetLastErrorString(LastError));
+		}
+		return OutResult;
+	}
+
+	FILE_ID_EXTD_DIR_INFO* Info = reinterpret_cast<FILE_ID_EXTD_DIR_INFO*>(Buffer);
+	FStringView RelPath = GetFileName(Info);
+	if (RelPath != TEXTVIEW("."))
+	{
+		UE_LOG(LogWindows, Error, TEXT("Unexpected failure of GetFileInformationByHandleEx('%s'). Expected it to return results for filename '.', but it returned filename %.*s."),
+			*Normalized, RelPath.Len(), RelPath.GetData());
+		return OutResult;
+	}
+
+	OutResult.bIsDirectory = !!(Info->FileAttributes & FILE_ATTRIBUTE_DIRECTORY);
+	FILETIME LastWriteTime;
+	LastWriteTime.dwLowDateTime = Info->LastWriteTime.LowPart;
+	LastWriteTime.dwHighDateTime = Info->LastWriteTime.HighPart;
+	OutResult.ModificationTime = WindowsFileTimeToUEDateTime(LastWriteTime);
+	OutResult.bIsValid = true;
+	OutResult.JournalHandle = ToFileHandle(Info->FileId);
+
+	return OutResult;
+}
+
+EFileJournalResult FWindowsPlatformFile::FileJournalReadModified(const TCHAR* VolumeName,
+	const FFileJournalEntryHandle& StartingJournalEntry, TMap<FFileJournalFileHandle, FString>& KnownDirectories,
+	TSet<FString>& OutModifiedDirectories, FFileJournalEntryHandle& OutNextJournalEntry, FString* OutError)
+{
+	using namespace UE::WindowsPlatformFileJournal::Private;
+
+	if (!FileJournalIsAvailableSystemWide(nullptr, OutError))
+	{
+		return IPlatformFile::FileJournalReadModified(VolumeName, StartingJournalEntry, KnownDirectories,
+			OutModifiedDirectories, OutNextJournalEntry, nullptr);
+	}
+
+	EFileJournalResult Result;
+	OutNextJournalEntry = StartingJournalEntry;
+
+	uint8 ReadBuffer[4096];
+	DWORD ReadBufferSize = UE_ARRAY_COUNT(ReadBuffer);
+	// For detection of when we have reached the end of the journal, we need buffer size large enough to read more
+	// records at once than we expect can be written in the time it takes us to call it twice in a row.
+	// sizeof(USN_RECORD_V3) == 80
+	static_assert(sizeof(ReadBuffer) > 10 * sizeof(USN_RECORD_V3), "ReadBuffer too small");
+
+	HANDLE JournalWindowsHandle = NULL;
+	TStringBuilder<16> NormalizedVolumeName;
+	Result = CreateJournalWindowsHandle(VolumeName, JournalWindowsHandle, NormalizedVolumeName, OutError);
+	if (Result != EFileJournalResult::Success)
+	{
+		return Result;
+	}
+	ON_SCOPE_EXIT{ ::CloseHandle(JournalWindowsHandle); };
+	VolumeName = *NormalizedVolumeName;
+
+	USN_JOURNAL_DATA_V1 JournalDescriptor;
+	Result = GetJournalDescriptor(JournalWindowsHandle, VolumeName, JournalDescriptor,
+		OutError);
+	if (Result != EFileJournalResult::Success)
+	{
+		return Result;
+	}
+
+	READ_USN_JOURNAL_DATA_V1 ReadJournalRequest;
+	// The USN at which to begin reading the change journal, or 0 to start from beginning
+	// ReadJournalRequest.StartUsn; // Set below in the loop
+	// Filter for types of changes that should be returned, e.g. USN_REASON_DATA_EXTEND: The file or directory is added to
+	ReadJournalRequest.ReasonMask = ~0;
+	// Changes to the filesystem can be reported when made, or when all processes drop their handle to the file.
+	// ReturnOnlyOnClose = 0 specifies as soon as made, non-zero specifies wait for handles to drop.
+	ReadJournalRequest.ReturnOnlyOnClose = 0;
+	// Timeout to use, in seconds, if BytesToWaitFor requests more data than exists in the change journal.
+	ReadJournalRequest.Timeout = 0;
+	// If FSCTL_READ_USN_JOURNAL operation requests more bytes than exist, how many bytes we should wait for before returning
+	ReadJournalRequest.BytesToWaitFor = 0;
+	// Journal to use, id is from JournalDescriptor returned by FSCTL_QUERY_USN_JOURNAL
+	ReadJournalRequest.UsnJournalID = JournalDescriptor.UsnJournalID;
+	// Minimum major version supported by our application, can be read from JournalDescriptor returned by FSCTL_QUERY_USN_JOURNAL
+	ReadJournalRequest.MinMajorVersion = JournalDescriptor.MinSupportedMajorVersion;
+	// Maximum major version supported by our application, can be read from JournalDescriptor returned by FSCTL_QUERY_USN_JOURNAL
+	ReadJournalRequest.MaxMajorVersion = JournalDescriptor.MaxSupportedMajorVersion;
+
+	DWORD ReadBufferResultsSize;
+	uint8* RecordBytesEnd = ReadBuffer;
+
+	// bTryShowFilenames is a debug feature, and is only available if process is privileged. The extra data it returns
+	// is logged but otherwise unused. If set to true and the process is not privileged, this function will run more
+	// slowly but will not fail.
+	const bool bTryShowFilenames = false;
+	bool bShowFilenames = bTryShowFilenames;
+	bool bTriedAndFailedShowFilenames = false;
+	bool bFirstCall = true;
+	for (;;)
+	{
+		// FSCTL_READ_UNPRIVILEGED_USN_JOURNAL can read without elevation but does not show filenames.
+		// FCSTL_READ_USN_JOURNAL requires elevation and shows filenames
+		DWORD dwIoControlCode = bShowFilenames ? FSCTL_READ_USN_JOURNAL : FSCTL_READ_UNPRIVILEGED_USN_JOURNAL;
+		bool bWasFirstCall = bFirstCall;
+		bFirstCall = false;
+		USN CurrentJournalUSN = static_cast<USN>(OutNextJournalEntry);
+		ReadJournalRequest.StartUsn = CurrentJournalUSN;
+
+		BOOL ReadJournalResult = ::DeviceIoControl(
+			JournalWindowsHandle,
+			dwIoControlCode,
+			&ReadJournalRequest,
+			sizeof(ReadJournalRequest),
+			ReadBuffer,
+			sizeof(ReadBuffer),
+			&ReadBufferResultsSize,
+			NULL /* Overlapped */);
+		if (!ReadJournalResult)
+		{
+			uint32 LastError = FPlatformMisc::GetLastError();
+			if (bWasFirstCall && dwIoControlCode == FSCTL_READ_USN_JOURNAL && LastError == ERROR_ACCESS_DENIED)
+			{
+				bTriedAndFailedShowFilenames = true;
+				bShowFilenames = false;
+				continue; // Try again with the unprivileged version
+			}
+			else
+			{
+				switch (LastError)
+				{
+				case ERROR_JOURNAL_ENTRY_DELETED:
+					if (OutError)
+					{
+						*OutError = FString::Printf(TEXT("NTFS Journal has wrapped for volume '%s'; the number of files written to this volume exceeded the buffer size of the journal and some entries were dropped.")
+							TEXT("\r\nIf this happens frequently increase the buffer size of the journal:")
+							TEXT("\r\nLaunch cmd.exe as admin and run command `fsutil usn createJournal %s m=<SizeInBytes>`")
+							TEXT("\r\nSizeInBytes defaults to 33554432 (32MB). You can see the current value via `fsutil usn queryJournal %s`, 'Maximum Size' field."),
+							VolumeName, VolumeName, VolumeName);
+					}
+					return EFileJournalResult::JournalWrapped;
+				default:
+					if (OutError)
+					{
+						*OutError = FString::Printf(
+							TEXT("FSCTL_READ_UNPRIVILEGED_USN_JOURNAL failed for volume '%s'. GetLastError == %u (0x%08x) (%s)."),
+							VolumeName, LastError, LastError, *JournalGetLastErrorString(LastError));
+					}
+					return EFileJournalResult::FailedReadJournal;
+				}
+			}
+		}
+		if (ReadBufferResultsSize < sizeof(USN))
+		{
+			if (OutError)
+			{
+				*OutError = FString::Printf(
+					TEXT("FSCTL_READ_UNPRIVILEGED_USN_JOURNAL returned an unexpected value; ReadBufferResultsSize == %u, and the minimum expected is %d."),
+					ReadBufferResultsSize, sizeof(USN));
+			}
+			return EFileJournalResult::JournalInternalError;
+		}
+		if (bTriedAndFailedShowFilenames)
+		{
+			UE_LOG(LogWindows, Error, TEXT("FileJournalReadModified cannot show filenames because the current process is not elevated."));
+			bTriedAndFailedShowFilenames = false;
+		}
+		uint8* RecordBytes = ReadBuffer;
+		RecordBytesEnd = ReadBuffer + ReadBufferResultsSize;
+		USN NextJournalUSN = *reinterpret_cast<USN*>(ReadBuffer);
+		RecordBytes += sizeof(USN);
+		while (RecordBytes < RecordBytesEnd)
+		{
+			USN_RECORD_V3* Record = reinterpret_cast<USN_RECORD_V3*>(RecordBytes);
+			if (Record->MajorVersion != 3)
+			{
+				// If we need to support earlier version 2, we will need to intepret the data as USN_RECORD_V2
+				if (OutError)
+				{
+					*OutError = FString::Printf(
+						TEXT("FSCTL_READ_UNPRIVILEGED_USN_JOURNAL for volume '%s' returned a record of version %d.%d, but we only handle version 3.*."),
+						VolumeName, (int)Record->MajorVersion, (int)Record->MinorVersion);
+				}
+				return EFileJournalResult::UnhandledJournalVersion;
+			}
+
+			FFileJournalFileHandle ParentJournalHandle = ToFileHandle(Record->ParentFileReferenceNumber);
+			FString* ParentDirectoryName = KnownDirectories.Find(ParentJournalHandle);
+			if (ParentDirectoryName)
+			{
+				OutModifiedDirectories.Add(*ParentDirectoryName);
+
+				if (bShowFilenames)
+				{
+					FStringView Filename = GetFileName(Record);
+					UE_LOG(LogWindows, Display, TEXT("FileJournalReadModified: ParentDirectory %s with file %.*s"),
+						*ParentJournalHandle.ToString(), Filename.Len(), Filename.GetData());
+				}
+			}
+
+			RecordBytes += Record->RecordLength;
+		}
+		static_assert(sizeof(uint64) >= sizeof(USN), "We are storing USNs as uint64");
+		OutNextJournalEntry = static_cast<uint64>(NextJournalUSN);
+
+		// Keep going until we reach the end, which we detect by  Journal did not have enough results to fill the buffer
+		if (RecordBytesEnd <= ReadBuffer + ReadBufferSize / 2)
+		{
+			break;
+		}
+		// Prevent a possible journal bug from causing an infinite loop
+		if (NextJournalUSN <= CurrentJournalUSN)
+		{
+			if (OutError)
+			{
+				*OutError = FString::Printf(
+					TEXT("FSCTL_READ_UNPRIVILEGED_USN_JOURNAL returned an unexpected value; ReadBufferResultsSize == %u, but NextJournalUSN did not increase.")
+					TEXT("Previous value = %" UINT64_FMT ". New value = %" UINT64_FMT "."),
+					ReadBufferResultsSize, static_cast<uint64>(CurrentJournalUSN), static_cast<uint64>(NextJournalUSN));
+			}
+			return EFileJournalResult::JournalInternalError;
+		}
+	}
+	return EFileJournalResult::Success;
+}
+
+namespace UE::WindowsPlatformFileJournal::Private
+{
+
+bool FileJournalIsAvailableSystemWide()
+{
+	switch (GJournalAvailable)
+	{
+	case 1:
+		return true;
+	case 2:
+		return false;
+	default:
+		break;
+	}
+	return FileJournalIsAvailableSystemWide(nullptr, nullptr);
+}
+
+bool FileJournalIsAvailableSystemWide(ELogVerbosity::Type* OutErrorLevel, FString* OutError)
+{
+	switch (GJournalAvailable)
+	{
+	case 1:
+		if (OutErrorLevel)
+		{
+			*OutErrorLevel = ELogVerbosity::Verbose;
+		}
+		if (OutError)
+		{
+			OutError->Empty();
+		}
+		return true;
+	case 2:
+		if (!OutErrorLevel && !OutError)
+		{
+			return false;
+		}
+		// Fall through and recompute so we can provide the information about why it failed
+		break;
+	default:
+		break;
+	}
+
+	auto ReturnResult = [OutErrorLevel, OutError](bool bResult, ELogVerbosity::Type ErrorLevel, FString&& Error)
+		{
+			if (OutErrorLevel)
+			{
+				*OutErrorLevel = ErrorLevel;
+			}
+			if (OutError)
+			{
+				*OutError = MoveTemp(Error);
+			}
+			GJournalAvailable = bResult ? 0x1 : 0x2;
+			return bResult;
+		};
+
+	FWindowsPlatformFile::FNormalizedDirectory Normalized(*FPaths::LaunchDir());
+	HANDLE Handle = CreateFileW(*Normalized,
+		GENERIC_READ,
+		FILE_SHARE_READ,
+		nullptr,
+		OPEN_EXISTING,
+		FILE_FLAG_BACKUP_SEMANTICS, // Required when opening a directory
+		nullptr);
+	if (Handle == INVALID_HANDLE_VALUE)
+	{
+		// Unexpected failure
+		uint32 LastError = GetLastError();
+		return ReturnResult(false, ELogVerbosity::Error, FString::Printf(
+			TEXT("FileJournal: CreateFile('%s') failed; FileJournal feature is unavailable. LastError=%u (0x%08x) (%s)."),
+			*Normalized, LastError, LastError, *JournalGetLastErrorString(LastError)));
+	}
+	ON_SCOPE_EXIT{ ::CloseHandle(Handle); };
+
+	static_assert(sizeof(FILE_ID_EXTD_DIR_INFO) > 20); // Our extra space should be smaller than half of a buffer
+	uint8 Buffer[sizeof(FILE_ID_EXTD_DIR_INFO) + 10];
+	BOOL Result = GetFileInformationByHandleEx(Handle, ::FileIdExtdDirectoryRestartInfo, Buffer, sizeof(Buffer));
+	if (Result)
+	{
+		// Available
+		return ReturnResult(true, ELogVerbosity::Verbose, FString());
+	}
+
+	uint32 LastError = GetLastError();
+	switch (LastError)
+	{
+	case ERROR_INVALID_PARAMETER:
+		// Failed, but expected if the windows version is too low. FileIdExtdDirectoryRestartInfo is not recognized before windows 8.
+		return ReturnResult(false, ELogVerbosity::Display, FString::Printf(
+			TEXT("FileJournal requires Windows 8 or later; it is unavailable on the current platform. LastError=%u (0x%08x) (%s)."),
+			LastError, LastError, *JournalGetLastErrorString(LastError)));
+	default:
+		// Unexpected failure
+		return ReturnResult(false, ELogVerbosity::Error, FString::Printf(
+			TEXT("FileJournal: FileIdExtdDirectoryRestartInfo('%s') failed; FileJournal feature is unavailable. LastError=%u (0x%08x) (%s)."),
+			*Normalized, LastError, LastError, *JournalGetLastErrorString(LastError)));
+	}
+}
+
+EFileJournalResult CreateJournalWindowsHandle(const TCHAR* VolumeOrPath, HANDLE& OutHandle,
+	FStringBuilderBase& OutVolumeName, FString* OutError)
+{
+	OutHandle = NULL;
+	OutVolumeName.Reset();
+
+	TStringBuilder<256> AbsPath;
+	FPathViews::ToAbsolutePath(VolumeOrPath, AbsPath);
+	FStringView VolumeName;
+	FStringView Remainder;
+	FPathViews::SplitVolumeSpecifier(AbsPath, VolumeName, Remainder);
+	if (!VolumeName.EndsWith(':'))
+	{
+		if (OutError)
+		{
+			*OutError = FString::Printf(
+				TEXT("Invalid drive name in path '%s'. FileJournal is only supported for drives assigned to a drive letter (e.g. C:, D:)."),
+				*AbsPath);
+		}
+		return EFileJournalResult::InvalidVolumeName;
+	}
+	OutVolumeName << VolumeName;
+
+	return CreateJournalWindowsHandle(*OutVolumeName, OutHandle, OutError);
+}
+
+EFileJournalResult CreateJournalWindowsHandle(const TCHAR* VolumeName, HANDLE& OutHandle, FString* OutError)
+{
+	OutHandle = NULL;
+
+	TStringBuilder<32> JournalName(InPlace, TEXT("\\\\.\\"), VolumeName);
+	OutHandle = ::CreateFileW(
+		*JournalName,
+		FILE_TRAVERSE,
+		FILE_SHARE_READ | FILE_SHARE_WRITE,
+		NULL /* lpSecurityAttributes */,
+		OPEN_EXISTING,
+		FILE_ATTRIBUTE_NORMAL,
+		NULL /* hTemplateFile */);
+
+	if (OutHandle == INVALID_HANDLE_VALUE)
+	{
+		if (OutError)
+		{
+			uint32 LastError = FPlatformMisc::GetLastError();
+			*OutError = FString::Printf(TEXT("CreateFile failed to open for reading FileJournal on volume '%s'. GetLastError == %d (0x%08x) (%s)"),
+				VolumeName, LastError, LastError, *JournalGetLastErrorString(LastError));
+		}
+		return EFileJournalResult::FailedOpenJournal;
+	}
+	return EFileJournalResult::Success;
+}
+
+EFileJournalResult GetJournalDescriptor(HANDLE JournalWindowsHandle, const TCHAR* VolumeName,
+	USN_JOURNAL_DATA_V1& OutJournalDescriptor, FString* OutError)
+{
+	DWORD SizeOfOutputData;
+	BOOL DescribeFileJournalResult = ::DeviceIoControl(
+		JournalWindowsHandle,
+		FSCTL_QUERY_USN_JOURNAL,
+		NULL /* InBuffer */,
+		0 /* InBufferSize */,
+		&OutJournalDescriptor,
+		sizeof(OutJournalDescriptor),
+		&SizeOfOutputData,
+		NULL /* Overlapped */);
+	if (!DescribeFileJournalResult)
+	{
+		uint32 LastError = FPlatformMisc::GetLastError();
+		switch (LastError)
+		{
+		case ERROR_JOURNAL_NOT_ACTIVE:
+			if (OutError)
+			{
+				// Can also create using DeviceIoControl with FSCTL_CREATE_USN_JOURNAL
+				*OutError = FString::Printf(TEXT("NTFS Journal is not active for volume '%s'. ")
+					TEXT("Launch cmd.exe as admin and run command `fsutil usn createJournal %s m=<SizeInBytes>`."),
+					TEXT("Recommended <SizeInBytes> is 100000000 (100MB)."),
+					VolumeName, VolumeName);
+			}
+			return EFileJournalResult::JournalNotActive;
+		default:
+			if (OutError)
+			{
+				*OutError = FString::Printf(
+					TEXT("FSCTL_QUERY_USN_JOURNAL failed for volume '%s'. GetLastError == %u (0x%08x) (%s)."),
+					VolumeName, LastError, LastError, *JournalGetLastErrorString(LastError));
+			}
+			return EFileJournalResult::FailedDescribeJournal;
+		}
+	}
+	return EFileJournalResult::Success;
+}
+
+FStringView GetFileName(FILE_ID_EXTD_DIR_INFO* Info)
+{
+	int32 FileNameLength = Info->FileNameLength / sizeof(Info->FileName[0]);
+	return FStringView(Info->FileName, FileNameLength);
+};
+
+FStringView GetFileName(USN_RECORD_V3* Record)
+{
+	int32 FileNameLength = Record->FileNameLength / sizeof(Record->FileName[0]);
+	return FStringView(Record->FileName, FileNameLength);
+};
+
+FString JournalGetLastErrorString(uint32 LastError)
+{
+	const TCHAR* ErrorName = nullptr;
+	switch (LastError)
+	{
+	case ERROR_ACCESS_DENIED: ErrorName = TEXT("ERROR_ACCESS_DENIED"); break;
+	case ERROR_BAD_LENGTH: ErrorName = TEXT("ERROR_BAD_LENGTH"); break;
+	case ERROR_JOURNAL_DELETE_IN_PROGRESS: ErrorName = TEXT("ERROR_JOURNAL_DELETE_IN_PROGRESS"); break;
+	case ERROR_JOURNAL_ENTRY_DELETED: ErrorName = TEXT("ERROR_JOURNAL_ENTRY_DELETED"); break;
+	case ERROR_POTENTIAL_FILE_FOUND: ErrorName = TEXT("ERROR_POTENTIAL_FILE_FOUND"); break;
+	case ERROR_JOURNAL_NOT_ACTIVE: ErrorName = TEXT("ERROR_JOURNAL_NOT_ACTIVE"); break;
+	default: break;
+	}
+
+	TCHAR ErrorBuffer[256];
+	FPlatformMisc::GetSystemErrorMessage(ErrorBuffer, UE_ARRAY_COUNT(ErrorBuffer), static_cast<int32>(LastError));
+	if (ErrorName && ErrorBuffer[0] != '\0')
+	{
+		return FString::Printf(TEXT("%s: %s"), ErrorName, ErrorBuffer);
+	}
+	else if (ErrorName)
+	{
+		return ErrorName;
+	}
+	else if (ErrorBuffer[0] != '\0')
+	{
+		return FString(ErrorBuffer);
+	}
+
+	return TEXT("<UnknownErrorCode>");
+}
+
+FFileJournalFileHandle ToFileHandle(FILE_ID_128& ReferenceNumber)
+{
+	static_assert(sizeof(ReferenceNumber) == 16);
+	static_assert((int)ExtendedFileIdType < 256);
+	// In version 3, the FileReferenceNumbers are FILE_ID_128, and their FILE_ID_DESCRIPTOR has Type ExtendedFileType
+	// and ExtendedFileId equal to the FileReferenceNumber
+	FFileJournalFileHandle Result;
+	Result.Bytes[0] = (uint8)ExtendedFileIdType;
+	Result.Bytes[1] = 0;
+	Result.Bytes[2] = 0;
+	Result.Bytes[3] = 1 << 7; // Valid = true
+	FMemory::Memcpy(&Result.Bytes[4], &ReferenceNumber, sizeof(ReferenceNumber));
+	return Result;
+}
+
+FFileJournalFileHandle ToFileHandle(uint64 ReferenceNumber)
+{
+	static_assert(sizeof(ReferenceNumber) == 8);
+	static_assert((int)FileIdType < 256);
+	// In version 2, the FileReferenceNumbers are DWORDDLONG (aka uint64), and their FILE_ID_DESCRIPTOR has type FileIdType
+	// and FileId.QuadPart equal to their FileReferenceNumber
+	FFileJournalFileHandle Result;
+	Result.Bytes[0] = (uint8)FileIdType;
+	Result.Bytes[1] = 0;
+	Result.Bytes[2] = 0;
+	Result.Bytes[3] = 1 << 7; // Valid = true
+	FMemory::Memcpy(&Result.Bytes[4], &ReferenceNumber, sizeof(ReferenceNumber));
+	Result.Bytes[12] = 0;
+	Result.Bytes[13] = 0;
+	Result.Bytes[14] = 0;
+	Result.Bytes[15] = 0;
+	return Result;
+}
+
+} // namespace UE::WindowsPlatformFileJournal::Private
+
+#else // UE_WINDOWS_PLATFORM_FILEJOURNAL_ENABLED 
+
+#if !WITH_EDITOR
+namespace UE::WindowsPlatformFileJournal::Private
+{
+constexpr const TCHAR* OnlyAvailableInEditorMessage = TEXT("FileJournal is only available in editor.");
+}
+#endif
+
+bool FWindowsPlatformFile::FileJournalIsAvailable(const TCHAR* VolumeOrPath, ELogVerbosity::Type* OutErrorLevel,
+	FString* OutError)
+{
+	using namespace UE::WindowsPlatformFileJournal::Private;
+
+#if !WITH_EDITOR
+	if (OutErrorLevel)
+	{
+		*OutErrorLevel = ELogVerbosity::Display;
+		OutErrorLevel = nullptr;
+	}
+	if (OutError)
+	{
+		*OutError = OnlyAvailableInEditorMessage;
+		OutError = nullptr;
+	}
+#endif
+	return IPlatformFile::FileJournalIsAvailable(VolumeOrPath, OutErrorLevel, OutError);
+}
+
+EFileJournalResult  FWindowsPlatformFile::FileJournalGetLatestEntry(const TCHAR* VolumeName,
+	FFileJournalEntryHandle& OutEntryHandle, FString* OutError)
+{
+	using namespace UE::WindowsPlatformFileJournal::Private;
+
+#if !WITH_EDITOR
+	if (OutError)
+	{
+		*OutError = OnlyAvailableInEditorMessage;
+		OutError = nullptr;
+	}
+#endif
+	return IPlatformFile::FileJournalGetLatestEntry(VolumeName, OutEntryHandle, OutError);
+}
+
+bool FWindowsPlatformFile::FileJournalIterateDirectory(const TCHAR* Directory,
+	FDirectoryJournalVisitorFunc Visitor)
+{
+	return IPlatformFile::FileJournalIterateDirectory(Directory, Visitor);
+}
+
+FFileJournalData FWindowsPlatformFile::FileJournalGetFileData(const TCHAR* FilenameOrDirectory)
+{
+	return IPlatformFile::FileJournalGetFileData(FilenameOrDirectory);
+}
+
+EFileJournalResult FWindowsPlatformFile::FileJournalReadModified(const TCHAR* VolumeName,
+	const FFileJournalEntryHandle& StartingJournalEntry, TMap<FFileJournalFileHandle, FString>& KnownDirectories,
+	TSet<FString>& OutModifiedDirectories, FFileJournalEntryHandle& OutNextJournalEntry, FString* OutError)
+{
+	using namespace UE::WindowsPlatformFileJournal::Private;
+
+#if !WITH_EDITOR
+	if (OutError)
+	{
+		*OutError = OnlyAvailableInEditorMessage;
+		OutError = nullptr;
+	}
+#endif
+	return IPlatformFile::FileJournalReadModified(VolumeName, StartingJournalEntry, KnownDirectories, OutModifiedDirectories,
+		OutNextJournalEntry, OutError);
+}
+
+#endif // else !UE_WINDOWS_PLATFORM_FILEJOURNAL_ENABLED
\ No newline at end of file
diff --git a/Engine/Source/Runtime/Core/Private/Windows/WindowsPlatformFilePrivate.h b/Engine/Source/Runtime/Core/Private/Windows/WindowsPlatformFilePrivate.h
--- a/Engine/Source/Runtime/Core/Private/Windows/WindowsPlatformFilePrivate.h
+++ b/Engine/Source/Runtime/Core/Private/Windows/WindowsPlatformFilePrivate.h
@@ -18,7 +18,7 @@
 **/
 class FWindowsPlatformFile : public IPhysicalPlatformFile
 {
-private:
+public:
 	/**
 	  * Convert from a valid Unreal Path to a canonical and strict-valid Windows Path.
 	  * An Unreal Path may have either \ or / and may have empty directories (two / in a row), and may have .. and may be relative
@@ -78,6 +78,18 @@ public:
 	// Forced not inline to reduce stack space usage since IterateDirectoryCommon might be recursive
 	FORCENOINLINE static HANDLE FindFirstFileWithWildcard(const TCHAR* Directory, WIN32_FIND_DATAW& OutData);
 	bool IterateDirectoryCommon(const TCHAR* Directory, const TFunctionRef<bool(const WIN32_FIND_DATAW&)>& Visitor);
+
+	virtual bool FileJournalIsAvailable(const TCHAR* VolumeOrPath = nullptr, ELogVerbosity::Type* OutErrorLevel = nullptr,
+		FString* OutError = nullptr) override;
+	virtual EFileJournalResult FileJournalGetLatestEntry(const TCHAR* VolumeName, FFileJournalEntryHandle& OutEntryHandle,
+		FString* OutError = nullptr) override;
+	virtual bool FileJournalIterateDirectory(const TCHAR* Directory, FDirectoryJournalVisitorFunc Visitor) override;
+	virtual FFileJournalData FileJournalGetFileData(const TCHAR* FilenameOrDirectory) override;
+	virtual EFileJournalResult FileJournalReadModified(const TCHAR* VolumeName,
+		const FFileJournalEntryHandle& StartingJournalEntry, TMap<FFileJournalFileHandle, FString>& KnownDirectories,
+		TSet<FString>& OutModifiedDirectories, FFileJournalEntryHandle& OutNextJournalEntry,
+		FString* OutError = nullptr) override;
+	virtual FString FileJournalGetVolumeName(FStringView InPath) override;
 };
 
 namespace UE::WindowsPlatformFile::Private
diff --git a/Engine/Source/Runtime/Core/Public/GenericPlatform/GenericPlatformFile.h b/Engine/Source/Runtime/Core/Public/GenericPlatform/GenericPlatformFile.h
--- a/Engine/Source/Runtime/Core/Public/GenericPlatform/GenericPlatformFile.h
+++ b/Engine/Source/Runtime/Core/Public/GenericPlatform/GenericPlatformFile.h
@@ -11,6 +11,8 @@
 #include "Containers/StringView.h"
 #include "Containers/UnrealString.h"
 #include "CoreTypes.h"
+#include "HAL/PlatformMemory.h"
+#include "Math/NumericLimits.h"
 #include "Misc/AssertionMacros.h"
 #include "Misc/DateTime.h"
 #include "Misc/EnumClassFlags.h"
@@ -19,6 +21,7 @@
 class FArchive;
 class IAsyncReadFileHandle;
 class IMappedFileHandle;
+namespace ELogVerbosity { enum Type : uint8; }
 
 /**
 * Enum for async IO priorities. 
@@ -89,6 +92,21 @@ enum class ESymlinkResult : int8
 
 ENUM_CLASS_FLAGS(EDirectoryVisitorFlags);
 
+/** Results that can be returned from IPlatformFile FileJournal API. */
+enum class EFileJournalResult
+{
+	Success,
+	InvalidPlatform,
+	InvalidVolumeName,
+	JournalNotActive,
+	JournalWrapped,
+	FailedOpenJournal,
+	FailedDescribeJournal,
+	FailedReadJournal,
+	JournalInternalError,
+	UnhandledJournalVersion,
+};
+
 /** 
  * File handle interface. 
 **/
@@ -212,6 +230,39 @@ struct FFileStatData
 };
 
 
+/**
+ * A handle used by the FileJournal API. Represents an entry for an action on a file in the FileJournal.
+ */
+typedef uint64 FFileJournalEntryHandle;
+constexpr FFileJournalEntryHandle FileJournalEntryHandleInvalid = static_cast<FFileJournalEntryHandle>(MAX_uint64);
+
+/**
+ * A handle used by the FileJournal API. Uniquely represents a file on disk without needing to use the filename.
+ */
+struct FFileJournalFileHandle
+{
+	bool operator==(const FFileJournalFileHandle& Other) const;
+	bool operator!=(const FFileJournalFileHandle& Other) const;
+	FString ToString();
+
+	uint8 Bytes[20];
+};
+uint32 GetTypeHash(const FFileJournalFileHandle&);
+CORE_API extern const FFileJournalFileHandle FileJournalFileHandleInvalid;
+
+/**
+ * Contains the information that's returned from FileJournalGetFileData for a file.
+ */
+struct FFileJournalData
+{
+	FFileJournalData();
+
+	FDateTime ModificationTime;
+	FFileJournalFileHandle JournalHandle = FileJournalFileHandleInvalid;
+	bool bIsValid : 1;
+	bool bIsDirectory : 1;
+};
+
 /**
 * File I/O Interface
 **/
@@ -480,6 +531,80 @@ public:
 	**/
 	virtual bool		IterateDirectoryStat(const TCHAR* Directory, FDirectoryStatVisitor& Visitor) = 0;
 
+	/**
+	 * Return whether FileJournal functionality is available on the current platform if VolumeName is nullptr or for
+	 * the given Volume if VolumeName is non-null. Optionally returns a user-displayable string for why it is not
+	 * available and a severity level for the reason. VolumeName may be a VolumeName as returned by
+	 * FileJournalGetVolumeName or any path to a file or directory on the Volume.
+	 */
+	CORE_API virtual bool FileJournalIsAvailable(const TCHAR* VolumeOrPath = nullptr,
+		ELogVerbosity::Type* OutErrorLevel = nullptr, FString* OutError = nullptr);
+
+	/**
+	 * Report the current end of the journal for the given volume, to be used as the StartingJournalEntry
+	 * in FileJournalGetModifiedDirectories. If !FileJournalIsAvaiable for the given volume, sets OutEntryHandle
+	 * to FileJournalEntryHandleInvalid. Returns EFileJournalResult::Success if successful, otherwise an error code
+	 * and optionally a user-displayable explanation for the error code.
+	 */
+	CORE_API virtual EFileJournalResult FileJournalGetLatestEntry(const TCHAR* VolumeName,
+		FFileJournalEntryHandle& OutEntryHandle, FString* OutError = nullptr);
+
+	/** File and directory visitor function that takes FileJournal data. */
+	typedef TFunctionRef<bool(const TCHAR*, const FFileJournalData&)> FDirectoryJournalVisitorFunc;
+
+	/**
+	 * Iterate the given directory as with IterateDirectoryStat, but report a FFileJournalData for each 
+	 * file and directory, which notably includes the FFileJournalFileHandle for the file/directory.
+	 * 
+	 * The paths returned as the first argument of the visitor function are the combined paths produced
+	 * by combining the input directory with the relative path of the child file or directory.
+	 * 
+	 * If the FileJournal is unavailable on the current system the iteration will still succeed but the 
+	 * FFileJournalFileHandle for each child path will be set to FileJournalFileHandleInvalid.
+	 * 
+	 * If the FileJournal is available on the current system but not on the volume of the given directory,
+	 * it is arbitrary whether the FFileJournalFileHandle will be validly set; if not valid they will be
+	 * set to FileJournalFileHandleInvalid.
+	 * 
+	 * @return	false if the directory did not exist or if the visitor returned false.
+	 */
+	CORE_API virtual bool FileJournalIterateDirectory(const TCHAR* Directory, FDirectoryJournalVisitorFunc Visitor);
+
+	/**
+	 * Return the data for the given path as with GetStatData, but report a FFileJournalData instead, which
+	 * notably includes the FFileJournalFileHandle for the file/directory.
+	 * Check the FFileJournalData::bIsValid member before using the returned data 
+	 */
+	CORE_API virtual FFileJournalData FileJournalGetFileData(const TCHAR* FilenameOrDirectory);
+
+	/**
+	 * Query the FileJournal to find a list of all directories on the given volume with files that have been added,
+	 * deleted, or modified in the specified time range. The beginning of the time range is specified by
+	 * StartingJournalEntry, which came from FileJournalGetLatestEntry or a previous call to FileJournalReadModified.
+	 * The end of the range is the latest modification on the volume. VolumeName can be the return value from
+	 * FileJournalGetVolumeName, or any path on the desired volume.
+	 * 
+	 * The caller must provide the mapping from FFileJournalFileHandle to DirectoryName; the FFileJournalFileHandle 
+	 * for each Directory can be found from FileJournalGetFileData or FileJournalIterateDirectory.
+	 *
+	 * Modified directories are appended into OutModifiedDirectories, and the next FileJournal entry to scan is written
+	 * into OutNextJournalEntry.
+	 * 
+	 * Returns EFileJournalResult::Success if successful, otherwise an error code and optionally a user-displayable
+	 * explanation for the error code. In an error case, partial results may still be written into the output.
+	 */
+	CORE_API virtual EFileJournalResult FileJournalReadModified(const TCHAR* VolumeName,
+		const FFileJournalEntryHandle& StartingJournalEntry, TMap<FFileJournalFileHandle, FString>& KnownDirectories,
+		TSet<FString>& OutModifiedDirectories, FFileJournalEntryHandle& OutNextJournalEntry,
+		FString* OutError = nullptr);
+
+	/**
+	 * Return the VolumeSpecifier present in the given path. Returns empty string if path does not have a valid
+	 * volume specifier for use by the FileJournal (e.g. some platforms do not support \\paths for FileJournal).
+	 * @see FPathViews::SplitVolumeSpecifier for a more general function that can return specifiers not usable
+	 * by the FileJournal.
+	 */
+	CORE_API virtual FString FileJournalGetVolumeName(FStringView InPath);
 
 	/////////////////////////////////////////////////////////////////////////////////////////////////////////
 	/////////// Utility Functions. These have a default implementation that uses the pure virtual operations.
@@ -744,3 +869,31 @@ public:
 
 	virtual int32 GetNumFiles() const = 0;
 };
+
+inline uint32 GetTypeHash(const FFileJournalFileHandle& A)
+{
+	constexpr uint32 Mult = 103;
+	uint32 Hash = 0;
+	const uint8* EndA = reinterpret_cast<const uint8*>(&A) + sizeof(FFileJournalFileHandle);
+	for (const uint8* PA = reinterpret_cast<const uint8*>(&A); PA < EndA; ++PA)
+	{
+		Hash = Hash * Mult + *PA;
+	}
+	return Hash;
+}
+
+inline bool FFileJournalFileHandle::operator==(const FFileJournalFileHandle& Other) const
+{
+	return 0 == FPlatformMemory::Memcmp(this, &Other, sizeof(FFileJournalFileHandle));
+}
+
+inline bool FFileJournalFileHandle::operator!=(const FFileJournalFileHandle& Other) const
+{
+	return 0 != FPlatformMemory::Memcmp(this, &Other, sizeof(FFileJournalFileHandle));
+}
+
+inline FFileJournalData::FFileJournalData()
+	: bIsValid(false)
+	, bIsDirectory(false)
+{
+}
