Move the global stereo test variants CVar out of ScreenshotFunctionalTest.cpp and into AutomationTest.cpp as "Automation.EnableStereoTestVariants", where it can now be controlled via the automation commandline
(e.g. "Automation Runtest Anisotropy;EnableStereoTests;Quit;").

diff --git a/Engine/Source/Developer/AutomationController/Private/AutomationCommandline.cpp b/Engine/Source/Developer/AutomationController/Private/AutomationCommandline.cpp
--- a/Engine/Source/Developer/AutomationController/Private/AutomationCommandline.cpp
+++ b/Engine/Source/Developer/AutomationController/Private/AutomationCommandline.cpp
@@ -691,6 +691,14 @@ protected:
 						CVar->Set(false);
 					}					
 				}
+				else if (FParse::Command(&TempCmd, TEXT("EnableStereoTests")))
+				{
+					if (IConsoleVariable* CVar = IConsoleManager::Get().FindConsoleVariable(TEXT("Automation.EnableStereoTestVariants")))
+					{
+						Ar.Logf(TEXT("Automation: Enabling Stereo Test Variants"));
+						CVar->Set(true);
+					}
+				}
 				else
 				{
 					Ar.Logf(TEXT("Incorrect automation command syntax! Supported commands are: "));
@@ -705,6 +713,8 @@ protected:
 					Ar.Logf(TEXT("\tAutomation Now"));
 					Ar.Logf(TEXT("\tAutomation Quit"));
 					Ar.Logf(TEXT("\tAutomation SoftQuit"));
+					Ar.Logf(TEXT("\tAutomation IgnoreLogEvents"));
+					Ar.Logf(TEXT("\tAutomation EnableStereoTests"));
 					bHandled = false;
 				}
 			}
diff --git a/Engine/Source/Developer/FunctionalTesting/Private/ScreenshotFunctionalTest.cpp b/Engine/Source/Developer/FunctionalTesting/Private/ScreenshotFunctionalTest.cpp
--- a/Engine/Source/Developer/FunctionalTesting/Private/ScreenshotFunctionalTest.cpp
+++ b/Engine/Source/Developer/FunctionalTesting/Private/ScreenshotFunctionalTest.cpp
@@ -19,11 +19,6 @@
 
 #include UE_INLINE_GENERATED_CPP_BY_NAME(ScreenshotFunctionalTest)
 
-static TAutoConsoleVariable<int32> CVarEnableStereoTestVariants(
-	TEXT("r.EnableStereoScreenshotTestVariants"), 0,
-	TEXT("Allows screenshot comparison tests with \"Support Stereo Test Variants\" checked to also test stereo rendering."),
-	ECVF_Default);
-
 AScreenshotFunctionalTest::AScreenshotFunctionalTest( const FObjectInitializer& ObjectInitializer )
 	: AScreenshotFunctionalTestBase(ObjectInitializer)
 	, bCameraCutOnScreenshotPrep(true)
@@ -46,7 +41,7 @@ void AScreenshotFunctionalTest::Serialize(FArchive& Ar)
 
 void AScreenshotFunctionalTest::PrepareTest()
 {
-	bShouldDoViewRectOffsetVariant = bSupportStereoTestVariants && CVarEnableStereoTestVariants.GetValueOnAnyThread() > 0;
+	bShouldDoViewRectOffsetVariant = bSupportStereoTestVariants && FAutomationTestFramework::NeedPerformStereoTestVariants();
 
 	// Pre-prep flush to allow rendering to temporary targets and other test resources
 	UAutomationBlueprintFunctionLibrary::FinishLoadingBeforeScreenshot();
diff --git a/Engine/Source/Runtime/Core/Private/Misc/AutomationTest.cpp b/Engine/Source/Runtime/Core/Private/Misc/AutomationTest.cpp
--- a/Engine/Source/Runtime/Core/Private/Misc/AutomationTest.cpp
+++ b/Engine/Source/Runtime/Core/Private/Misc/AutomationTest.cpp
@@ -51,6 +51,12 @@ namespace AutomationTest
 		bLogTestStateTrace,
 		TEXT("Whether to enable or disable logging of test state trace"));
 
+	static bool bEnableStereoTestVariants = false;
+	static FAutoConsoleVariableRef CVarAutomationEnableStereoTestVariants(
+		TEXT("Automation.EnableStereoTestVariants"),
+		bEnableStereoTestVariants,
+		TEXT("Whether to enable stereo test variants for screenshot functional tests"));
+
 	// The method prepares the filename and LineNumber to be placed in the form that could be extracted by SAutomationWindow widget if it is additionally eclosed into []
 	// The result format is filename(line)
 	static FString CreateFileLineDescription(const FString& Filename, const int32 LineNumber)
@@ -301,6 +307,11 @@ bool FAutomationTestFramework::NeedLogBPTestMetadata()
 	return AutomationTest::bLogBPTestMetadata;
 }
 
+bool FAutomationTestFramework::NeedPerformStereoTestVariants()
+{
+	return AutomationTest::bEnableStereoTestVariants;
+}
+
 bool FAutomationTestFramework::RegisterAutomationTest( const FString& InTestNameToRegister, FAutomationTestBase* InTestToRegister )
 {
 	const bool bAlreadyRegistered = AutomationTestClassNameToInstanceMap.Contains( InTestNameToRegister );
diff --git a/Engine/Source/Runtime/Core/Public/Misc/AutomationTest.h b/Engine/Source/Runtime/Core/Public/Misc/AutomationTest.h
--- a/Engine/Source/Runtime/Core/Public/Misc/AutomationTest.h
+++ b/Engine/Source/Runtime/Core/Public/Misc/AutomationTest.h
@@ -1159,6 +1159,11 @@ public:
 	 */
 	static CORE_API bool NeedLogBPTestMetadata();
 
+	/**
+	 * Whether to also run stereo test variants for screenshot functional testts
+	 */
+	static CORE_API bool NeedPerformStereoTestVariants();
+
 	/**
 	 * Notify that the screenshot comparison has completed
 	 */
