Fixed return address for check(f) and verify(f) macros
Reduced code bloat in Debug\Dev builds when these macros are enabled

diff --git a/Engine/Source/Runtime/Core/Private/Misc/AssertionMacros.cpp b/Engine/Source/Runtime/Core/Private/Misc/AssertionMacros.cpp
--- a/Engine/Source/Runtime/Core/Private/Misc/AssertionMacros.cpp
+++ b/Engine/Source/Runtime/Core/Private/Misc/AssertionMacros.cpp
@@ -669,15 +669,15 @@ FORCENOINLINE bool FDebug::OptionallyLogFormattedEnsureMessageReturningFalseImpl
 }
 #endif
 
-FORCENOINLINE void UE_DEBUG_SECTION VARARGS LowLevelFatalErrorHandler(const ANSICHAR* File, int32 Line, void* ProgramCounter, const TCHAR* Format, ...)
+FORCENOINLINE void UE_DEBUG_SECTION VARARGS LowLevelFatalErrorHandler(const ANSICHAR* File, int32 Line, const TCHAR* Format, ...)
 {
 	va_list Args;
 	va_start(Args, Format);
-	StaticFailDebugV(TEXT("LowLevelFatalError"), "", File, Line, /*bIsEnsure*/ false, ProgramCounter, Format, Args);
+	StaticFailDebugV(TEXT("LowLevelFatalError"), "", File, Line, /*bIsEnsure*/ false, PLATFORM_RETURN_ADDRESS(), Format, Args);
 	va_end(Args);
 
 	UE_DEBUG_BREAK_AND_PROMPT_FOR_REMOTE();
-	FDebug::ProcessFatalError(ProgramCounter);
+	FDebug::ProcessFatalError(PLATFORM_RETURN_ADDRESS());
 }
 
 void FDebug::DumpStackTraceToLog(const ELogVerbosity::Type LogVerbosity)
diff --git a/Engine/Source/Runtime/Core/Public/Misc/AssertionMacros.h b/Engine/Source/Runtime/Core/Public/Misc/AssertionMacros.h
--- a/Engine/Source/Runtime/Core/Public/Misc/AssertionMacros.h
+++ b/Engine/Source/Runtime/Core/Public/Misc/AssertionMacros.h
@@ -233,10 +233,17 @@ RetType FORCENOINLINE UE_DEBUG_SECTION DispatchCheckVerify(InnerType&& Inner, Ar
 		{ \
 			if(UNLIKELY(!(expr))) \
 			{ \
-				if (FDebug::CheckVerifyFailedImpl(#expr, __FILE__, __LINE__, PLATFORM_RETURN_ADDRESS(), TEXT(""))) \
+				struct Impl \
 				{ \
-					PLATFORM_BREAK(); \
-				} \
+					static void FORCENOINLINE UE_DEBUG_SECTION ExecCheckImplInternal() \
+					{ \
+						if (FDebug::CheckVerifyFailedImpl(#expr, __FILE__, __LINE__, PLATFORM_RETURN_ADDRESS(), TEXT(""))) \
+						{ \
+							PLATFORM_BREAK(); \
+						} \
+					} \
+				}; \
+				Impl::ExecCheckImplInternal(); \
 				CA_ASSUME(false); \
 			} \
 		}
@@ -254,13 +261,13 @@ RetType FORCENOINLINE UE_DEBUG_SECTION DispatchCheckVerify(InnerType&& Inner, Ar
 
 	#define UE_CHECK_F_IMPL(expr, format, ...) \
 		{ \
-			UE_VALIDATE_FORMAT_STRING(format, ##__VA_ARGS__); \
 			if(UNLIKELY(!(expr))) \
 			{ \
-				if (FDebug::CheckVerifyFailedImpl(#expr, __FILE__, __LINE__, PLATFORM_RETURN_ADDRESS(), format, ##__VA_ARGS__)) \
+				if (DispatchCheckVerify<bool>([&] () UE_DEBUG_SECTION \
 				{ \
-					PLATFORM_BREAK(); \
-				} \
+					UE_VALIDATE_FORMAT_STRING(format, ##__VA_ARGS__); \
+					return FDebug::CheckVerifyFailedImpl(#expr, __FILE__, __LINE__, PLATFORM_RETURN_ADDRESS(), format, ##__VA_ARGS__); \
+				})) PLATFORM_BREAK(); \
 				CA_ASSUME(false); \
 			} \
 		}
@@ -455,11 +462,11 @@ namespace UEAsserts_Private
 ----------------------------------------------------------------------------*/
 
 /** low level fatal error handler. */
-CORE_API void UE_DEBUG_SECTION VARARGS LowLevelFatalErrorHandler(const ANSICHAR* File, int32 Line, void* ProgramCounter, const TCHAR* Format=TEXT(""), ... );
+CORE_API void UE_DEBUG_SECTION VARARGS LowLevelFatalErrorHandler(const ANSICHAR* File, int32 Line, const TCHAR* Format=TEXT(""), ... );
 
 #define LowLevelFatalError(Format, ...) \
 	{ \
 		static_assert(TIsArrayOrRefOfTypeByPredicate<decltype(Format), TIsCharEncodingCompatibleWithTCHAR>::Value, "Formatting string must be a TCHAR array."); \
-		LowLevelFatalErrorHandler(__FILE__, __LINE__, PLATFORM_RETURN_ADDRESS(), (const TCHAR*)Format, ##__VA_ARGS__); \
+		LowLevelFatalErrorHandler(__FILE__, __LINE__, (const TCHAR*)Format, ##__VA_ARGS__); \
 	}
 
