Changed FCriticalSection to UE::FMutex in FTextFormatData to avoid using an OS system object (Some platforms have a limitation on number of mutex that can be created).

#rnx

diff --git a/Engine/Source/Runtime/Core/Private/Internationalization/TextFormatter.cpp b/Engine/Source/Runtime/Core/Private/Internationalization/TextFormatter.cpp
--- a/Engine/Source/Runtime/Core/Private/Internationalization/TextFormatter.cpp
+++ b/Engine/Source/Runtime/Core/Private/Internationalization/TextFormatter.cpp
@@ -29,6 +29,7 @@
 #include "Templates/ValueOrError.h"
 #include "Trace/Detail/Channel.h"
 #include "AutoRTFM/AutoRTFM.h"
+#include "Async/Mutex.h"
 
 #define LOCTEXT_NAMESPACE "TextFormatter"
 
@@ -380,7 +381,7 @@ public:
 	 */
 	FORCEINLINE bool IsValid() const
 	{
-		FScopeLock Lock(&CompiledDataCS);
+		UE::TScopeLock Lock(CompiledDataMutex);
 		return IsValid_NoLock();
 	}
 
@@ -395,7 +396,7 @@ public:
 	 */
 	FORCEINLINE bool ValidatePattern(const FCulturePtr& InCulture, TArray<FString>& OutValidationErrors)
 	{
-		FScopeLock Lock(&CompiledDataCS);
+		UE::TScopeLock Lock(CompiledDataMutex);
 		return ValidatePattern_NoLock(InCulture, OutValidationErrors);
 	}
 
@@ -404,7 +405,7 @@ public:
 	 */
 	FORCEINLINE FString Format(const FPrivateTextFormatArguments& InFormatArgs)
 	{
-		FScopeLock Lock(&CompiledDataCS);
+		UE::TScopeLock Lock(CompiledDataMutex);
 		return Format_NoLock(InFormatArgs);
 	}
 
@@ -413,7 +414,7 @@ public:
 	 */
 	FORCEINLINE void GetFormatArgumentNames(TArray<FString>& OutArgumentNames)
 	{
-		FScopeLock Lock(&CompiledDataCS);
+		UE::TScopeLock Lock(CompiledDataMutex);
 		return GetFormatArgumentNames_NoLock(OutArgumentNames);
 	}
 
@@ -440,7 +441,7 @@ public:
 	 */
 	FORCEINLINE FTextFormat::EExpressionType GetExpressionType() const
 	{
-		FScopeLock Lock(&CompiledDataCS);
+		UE::TScopeLock Lock(CompiledDataMutex);
 		return CompiledExpressionType;
 	}
 
@@ -524,52 +525,52 @@ private:
 	FText SourceText;
 
 	/**
-	 * Critical section protecting the compiled data from being modified concurrently.
+	 * Mutex protecting the compiled data from being modified concurrently.
 	 */
-	mutable FCriticalSection CompiledDataCS;
+	mutable UE::FMutex CompiledDataMutex;
 
 	/**
 	 * Copy of the string that was last compiled.
 	 * This allows the text to update via a culture change without immediately invalidating our compiled tokens.
 	 * If the data was constructed from an FString rather than an FText, then this is the string we were given and shouldn't be updated once the initial construction has happened.
-	 * Concurrent access protected by CompiledDataCS.
+	 * Concurrent access protected by CompiledDataMutex.
 	 */
 	FString SourceExpression;
 
 	/**
 	 * Lexed expression tokens generated from, and referencing, SourceExpression.
-	 * Concurrent access protected by CompiledDataCS.
+	 * Concurrent access protected by CompiledDataMutex.
 	 */
 	TArray<FExpressionToken> LexedExpression;
 
 	/**
 	 * Snapshot of the text that last time it was compiled into a format expression.
 	 * This is used to detect when the source text was changed and allow a re-compile.
-	 * Concurrent access protected by CompiledDataCS.
+	 * Concurrent access protected by CompiledDataMutex.
 	 */
 	FTextSnapshot CompiledTextSnapshot;
 
 	/**
 	 * The type of expression currently compiled.
-	 * Concurrent access protected by CompiledDataCS.
+	 * Concurrent access protected by CompiledDataMutex.
 	 */
 	FTextFormat::EExpressionType CompiledExpressionType;
 
 	/**
 	 * Holds the last compilation error (if any, when CompiledExpressionType == Invalid).
-	 * Concurrent access protected by CompiledDataCS.
+	 * Concurrent access protected by CompiledDataMutex.
 	 */
 	FString LastCompileError;
 
 	/**
 	 * The base length of the string that will go into the formatted string (no including any argument substitutions.
-	 * Concurrent access protected by CompiledDataCS.
+	 * Concurrent access protected by CompiledDataMutex.
 	 */
 	int32 BaseFormatStringLength;
 
 	/**
 	 * A multiplier to apply to the given argument count (base is 1, and 1 is added for every argument modifier that may make use of the arguments).
-	 * Concurrent access protected by CompiledDataCS.
+	 * Concurrent access protected by CompiledDataMutex.
 	 */
 	int32 FormatArgumentEstimateMultiplier;
 };
