Add trace regions for slow tasks

diff --git a/Engine/Source/Runtime/Core/Private/Misc/SlowTask.cpp b/Engine/Source/Runtime/Core/Private/Misc/SlowTask.cpp
--- a/Engine/Source/Runtime/Core/Private/Misc/SlowTask.cpp
+++ b/Engine/Source/Runtime/Core/Private/Misc/SlowTask.cpp
@@ -13,6 +13,15 @@
 #include "Misc/AssertionMacros.h"
 #include "Misc/FeedbackContext.h"
 #include "Misc/SlowTaskStack.h"
+#include "ProfilingDebugging/MiscTrace.h"
+
+static int32 GSlowTaskMaxTraceRegionDepth = 2;
+static FAutoConsoleVariableRef CVarSlowTaskMaxTraceRegionDepth (
+	TEXT("Trace.SlowTaskMaxRegionDepth"),
+	GSlowTaskMaxTraceRegionDepth,
+	TEXT("Maximum depth of nested slow tasks to create as trace regions in insights"),
+	ECVF_Default
+);
 
 bool FSlowTask::ShouldCreateThrottledSlowTask()
 {
@@ -91,6 +100,10 @@ void FSlowTask::Initialize()
 	if (bEnabled)
 	{
 		Context.ScopeStack.Push(this);
+		if (Context.ScopeStack.Num() <= GSlowTaskMaxTraceRegionDepth)
+		{
+			TRACE_BEGIN_REGION(*DefaultMessage.ToString());
+		}
 	}
 }
 
@@ -119,6 +132,11 @@ void FSlowTask::Destroy()
 		FSlowTaskStack& Stack = Context.ScopeStack;
 		if (ensure(Stack.Num() != 0))
 		{
+			if (Context.ScopeStack.Num() <= GSlowTaskMaxTraceRegionDepth)
+			{
+				TRACE_END_REGION(*DefaultMessage.ToString());
+			}
+
 			FSlowTask* Task = Stack.Last();
 			if (ensureMsgf(Task == this, TEXT("Out-of-order slow task construction/destruction: destroying '%s' but '%s' is at the top of the stack"), *DefaultMessage.ToString(), *Task->DefaultMessage.ToString()))
 			{
