[UE] Skip UE_TRACE metadata tag for non-default tagset so that it doesn't override LLMTag metadata.

#tests PIE

diff --git a/Engine/Source/Runtime/Core/Public/HAL/LowLevelMemTracker.h b/Engine/Source/Runtime/Core/Public/HAL/LowLevelMemTracker.h
--- a/Engine/Source/Runtime/Core/Public/HAL/LowLevelMemTracker.h
+++ b/Engine/Source/Runtime/Core/Public/HAL/LowLevelMemTracker.h
@@ -350,7 +350,10 @@ extern FName LLMGetTagStat(ELLMTag Tag);
 				{ \
 					SCOPE_NAME.TryAddTagAndActivate(UniqueNameEvaluated, Constructor); \
 				} \
-				UE_MEMSCOPE_ACTIVATE(__LINE__, UniqueNameEvaluated); \
+				if (TagSet == ELLMTagSet::None && Tracker == ELLMTracker::Default) \
+				{ \
+					UE_MEMSCOPE_ACTIVATE(__LINE__, UniqueNameEvaluated); \
+				} \
 			} \
 		} \
 	} while (false) /* do/while is added to require a semicolon */
