Fixing a crash when an empty string is passed to CompareWithNumericSuffix

diff --git a/Engine/Source/Runtime/Core/Private/Misc/ComparisonUtility.cpp b/Engine/Source/Runtime/Core/Private/Misc/ComparisonUtility.cpp
--- a/Engine/Source/Runtime/Core/Private/Misc/ComparisonUtility.cpp
+++ b/Engine/Source/Runtime/Core/Private/Misc/ComparisonUtility.cpp
@@ -19,6 +19,7 @@ int32 CompareWithNumericSuffix(FStringView A, FStringView B)
 	{
 		int32 NumDigits = 0;
 		OutSuffix = 0;
+		if (!Full.IsEmpty())
 		{
 			int32 Magnitude = 1;
 			const TCHAR* FullStart = Full.GetData();
