Add a test to cover the FChunkedArray iterator bug that was fixed in CL 29283139

#rnx

- This is a cutdown version of the test that was used to identify the bug from CL 29283139.
- Moved the existing tests for FChunkedArray from "System::Core" to "System::Core::Containers"

diff --git a/Engine/Source/Runtime/Core/Tests/Containers/ChunkedArrayTest.cpp b/Engine/Source/Runtime/Core/Tests/Containers/ChunkedArrayTest.cpp
--- a/Engine/Source/Runtime/Core/Tests/Containers/ChunkedArrayTest.cpp
+++ b/Engine/Source/Runtime/Core/Tests/Containers/ChunkedArrayTest.cpp
@@ -19,7 +19,7 @@ public:
 	}
 };
 
-TEST_CASE_NAMED(TChunkedArrayAddAllocatesProperAmountOfChunksTest, "System::Core::ChunkedArray::Add", "[Core][ChunkedArray][SmokeFilter]")
+TEST_CASE_NAMED(TChunkedArrayAddAllocatesProperAmountOfChunksTest, "System::Core::Containers::ChunkedArray::Add", "[Core][Containers][ChunkedArray][SmokeFilter]")
 {
 	SECTION("One element per chunk")
 	{
@@ -72,6 +72,63 @@ TEST_CASE_NAMED(TChunkedArrayAddAllocatesProperAmountOfChunksTest, "System::Core
 	}
 }
 
+TEST_CASE_NAMED(TChunkedArrayTIteratorTest, "System::Core::Containers::ChunkedArray::TIterator", "[Core][Containers][ChunkedArray]")
+{
+	struct FDataTest
+	{
+	public:
+		FDataTest() = default;
+		~FDataTest() = default;
+
+		FDataTest(int32 InValue)
+			:Data{ InValue, InValue, InValue, InValue }
+		{}
+
+		int32 GetValue() const
+		{
+			return Data[0];
+		}
+
+	private:
+		int32 Data[4];
+	};
+
+	auto TestArray = [](int32 NumToTest)
+	{
+		TChunkedArray<FDataTest, 16384> ChunkedData;
+		for (int32 Index = 0; Index < NumToTest; ++Index)
+		{
+			ChunkedData.AddElement(FDataTest(Index));
+		}
+
+		int32 Count = 0;
+		for (FDataTest& Data : ChunkedData)
+		{
+			REQUIRE(Count++ == Data.GetValue());
+		}
+
+		REQUIRE(Count == ChunkedData.Num());
+	};
+
+	// This test was originally written to test the theory that FChunkedArray had a bug
+	// when used in a ranged for loop (it did) which was done via:
+	// MinBytes = 1;
+	// MaxBytes = 16384;
+	// which created enough allocation churn to repreoduce the problem. However that test
+	// took a long time to run and is unsuitable for a fast test. The following values are
+	// a cutdown version of the original test which seem to reproduce the original issue
+	// that as fixed in CL 29283139
+
+	const int32 MinBytes = 2048;
+	const int32 MaxBytes = 4096;
+
+	for (int32 Index = MinBytes; Index < MaxBytes; Index++)
+	{
+		TestArray(Index);
+	}
+}
+
+
 }
 
 #endif // WITH_TESTS
