Use structured logging for OutOfMemory messages

Structured logging allows to surface information contained in log lines more easily for automated processes. Out of memory errors are a common thing we'd like to detect and the current detection code has already been broken once.

This CL also contains a change to the structured log for fatal errors. Those were erroneously noted as assertion failures due to an unintended call to Debug::AssertFailed originating via copy-paste from FMsg::LogfImpl. Removing this call ensures that the new log matches the old behavior from unstructured logging (which I have manually verified).

diff --git a/Engine/Source/Runtime/Core/Private/GenericPlatform/GenericPlatformMemory.cpp b/Engine/Source/Runtime/Core/Private/GenericPlatform/GenericPlatformMemory.cpp
--- a/Engine/Source/Runtime/Core/Private/GenericPlatform/GenericPlatformMemory.cpp
+++ b/Engine/Source/Runtime/Core/Private/GenericPlatform/GenericPlatformMemory.cpp
@@ -8,6 +8,7 @@
 #include "Containers/StringConv.h"
 #include "UObject/NameTypes.h"
 #include "Logging/LogMacros.h"
+#include "Logging/StructuredLog.h"
 #include "Stats/Stats.h"
 #include "Containers/Ticker.h"
 #include "Misc/FeedbackContext.h"
@@ -248,7 +249,11 @@ void FGenericPlatformMemory::OnOutOfMemory(uint64 Size, uint32 Alignment)
 		FCoreDelegates::GetOutOfMemoryDelegate().Broadcast();
 
 		// ErrorMsg might be unrelated to OoM error in some cases as the code that calls OnOutOfMemory could have called other system functions that modified errno
-		UE_LOG(LogMemory, Fatal, TEXT("Ran out of memory allocating %llu bytes with alignment %u. Last error msg: %s."), Size, Alignment, ErrorMsg);
+		UE_LOGFMT_NSLOC(LogMemory, Fatal, "Memory", "OutOfMemoryError", "Ran out of memory allocating {Bytes} bytes with alignment {Alignment}. Last error msg: {LastError}.",
+			("Bytes", Size),
+			("Alignment", Alignment),
+			("LastError", ErrorMsg)
+		);
 	};
 	
 	UE_CALL_ONCE(HandleOOM);
diff --git a/Engine/Source/Runtime/Core/Private/IOS/IOSPlatformMemory.cpp b/Engine/Source/Runtime/Core/Private/IOS/IOSPlatformMemory.cpp
--- a/Engine/Source/Runtime/Core/Private/IOS/IOSPlatformMemory.cpp
+++ b/Engine/Source/Runtime/Core/Private/IOS/IOSPlatformMemory.cpp
@@ -10,6 +10,7 @@
 #include "HAL/MallocAnsi.h"
 #include "GenericPlatform/GenericPlatformMemoryPoolStats.h"
 #include "Misc/CoreDelegates.h"
+#include "Logging/StructuredLog.h"
 
 void FIOSPlatformMemory::OnOutOfMemory(uint64 Size, uint32 Alignment)
 {
@@ -52,7 +53,11 @@ void FIOSPlatformMemory::OnOutOfMemory(uint64 Size, uint32 Alignment)
 		FCoreDelegates::GetMemoryTrimDelegate().Broadcast();
     
 		// ErrorMsg might be unrelated to OoM error in some cases as the code that calls OnOutOfMemory could have called other system functions that modified errno
-		UE_LOG(LogMemory, Warning, TEXT("Ran out of memory allocating %llu bytes with alignment %u. Last error msg: %s."), Size, Alignment, ErrorMsg);
+		UE_LOGFMT_NSLOC(LogMemory, Warning, "Memory", "OutOfMemoryError", "Ran out of memory allocating {Bytes} bytes with alignment {Alignment}. Last error msg: {LastError}.",
+			("Bytes", Size),
+			("Alignment", Alignment),
+			("LastError", ErrorMsg)
+		);
     
 		// make this a fatal error that ends here not in the log
 		// changed to 3 from NULL because clang noticed writing to NULL and warned about it
diff --git a/Engine/Source/Runtime/Core/Private/Logging/StructuredLog.cpp b/Engine/Source/Runtime/Core/Private/Logging/StructuredLog.cpp
--- a/Engine/Source/Runtime/Core/Private/Logging/StructuredLog.cpp
+++ b/Engine/Source/Runtime/Core/Private/Logging/StructuredLog.cpp
@@ -887,7 +887,6 @@ void FatalLogWithFieldArray(const FLogCategoryBase& Category, const FStaticLogRe
 	Record.FormatMessageTo(Message);
 
 	StaticFailDebug(TEXT("Fatal error:"), "", Log.File, Log.Line, /*bIsEnsure*/ false, PLATFORM_RETURN_ADDRESS(), TEXT("%s"), *Message);
-	FDebug::AssertFailed("", Log.File, Log.Line, TEXT("%s"), *Message);
 
 	UE_DEBUG_BREAK_AND_PROMPT_FOR_REMOTE();
 	FDebug::ProcessFatalError(PLATFORM_RETURN_ADDRESS());
@@ -911,7 +910,6 @@ void FatalLogWithFieldArray(const FLogCategoryBase& Category, const FStaticLocal
 	Record.FormatMessageTo(Message);
 
 	StaticFailDebug(TEXT("Fatal error:"), "", Log.File, Log.Line, /*bIsEnsure*/ false, PLATFORM_RETURN_ADDRESS(), TEXT("%s"), *Message);
-	FDebug::AssertFailed("", Log.File, Log.Line, TEXT("%s"), *Message);
 
 	UE_DEBUG_BREAK_AND_PROMPT_FOR_REMOTE();
 	FDebug::ProcessFatalError(PLATFORM_RETURN_ADDRESS());
diff --git a/Engine/Source/Runtime/Core/Private/Unix/UnixPlatformMemory.cpp b/Engine/Source/Runtime/Core/Private/Unix/UnixPlatformMemory.cpp
--- a/Engine/Source/Runtime/Core/Private/Unix/UnixPlatformMemory.cpp
+++ b/Engine/Source/Runtime/Core/Private/Unix/UnixPlatformMemory.cpp
@@ -14,6 +14,7 @@
 #include "Templates/UnrealTemplate.h"
 #include "Containers/UnrealString.h"
 #include "Logging/LogMacros.h"
+#include "Logging/StructuredLog.h"
 #include "HAL/MallocAnsi.h"
 #include "HAL/MallocMimalloc.h"
 #include "HAL/MallocJemalloc.h"
@@ -1194,7 +1195,11 @@ void FUnixPlatformMemory::OnOutOfMemory(uint64 Size, uint32 Alignment)
 		FCoreDelegates::GetOutOfMemoryDelegate().Broadcast();
 
 		// ErrorMsg might be unrelated to OoM error in some cases as the code that calls OnOutOfMemory could have called other system functions that modified errno
-		UE_LOG(LogMemory, Fatal, TEXT("Ran out of memory allocating %llu bytes with alignment %u. Last error msg: %s."), Size, Alignment, ErrorMsg);
+		UE_LOGFMT_NSLOC(LogMemory, Fatal, "Memory", "OutOfMemoryError", "Ran out of memory allocating {Bytes} bytes with alignment {Alignment}. Last error msg: {LastError}.",
+			("Bytes", Size),
+			("Alignment", Alignment),
+			("LastError", ErrorMsg)
+		);
 	};
 	
 	UE_CALL_ONCE(HandleOOM);
