Do not execute CheckVerifyImpl inside ensure macro if it has been already executed or ensure is not allowed as this evaluation might result in an expensive function calls

diff --git a/Engine/Source/Runtime/Core/Private/Misc/AssertionMacros.cpp b/Engine/Source/Runtime/Core/Private/Misc/AssertionMacros.cpp
--- a/Engine/Source/Runtime/Core/Private/Misc/AssertionMacros.cpp
+++ b/Engine/Source/Runtime/Core/Private/Misc/AssertionMacros.cpp
@@ -714,28 +714,23 @@ FORCENOINLINE void FDebug::DumpStackTraceToLog(const TCHAR* Heading, const ELogV
 #if DO_ENSURE && !USING_CODE_ANALYSIS
 bool UE_DEBUG_SECTION VARARGS CheckVerifyImpl(bool& InOutExecuted, bool Always, const ANSICHAR* File, int32 Line, void* ProgramCounter, const ANSICHAR* Expr, const TCHAR* Format, ...)
 {
-	if ((!InOutExecuted || Always) && FPlatformMisc::IsEnsureAllowed())
-	{
-		InOutExecuted = true;
-		va_list Args;
-		va_start(Args, Format);
-		FDebug::OptionallyLogFormattedEnsureMessageReturningFalse(true, Expr, File, Line, ProgramCounter, Format, Args);
-		va_end(Args);
+	InOutExecuted = true;
+	va_list Args;
+	va_start(Args, Format);
+	FDebug::OptionallyLogFormattedEnsureMessageReturningFalse(true, Expr, File, Line, ProgramCounter, Format, Args);
+	va_end(Args);
 
-		if (!FPlatformMisc::IsDebuggerPresent())
-		{
-			FPlatformMisc::PromptForRemoteDebugging(true);
-			return false;
-		}
+	if (!FPlatformMisc::IsDebuggerPresent())
+	{
+		FPlatformMisc::PromptForRemoteDebugging(true);
+		return false;
+	}
 
 #if UE_BUILD_SHIPPING
-		return true;
+	return true;
 #else
-		return !GIgnoreDebugger;
+	return !GIgnoreDebugger;
 #endif
-	}
-
-	return false;
 }
 #endif
 
diff --git a/Engine/Source/Runtime/Core/Public/Misc/AssertionMacros.h b/Engine/Source/Runtime/Core/Public/Misc/AssertionMacros.h
--- a/Engine/Source/Runtime/Core/Public/Misc/AssertionMacros.h
+++ b/Engine/Source/Runtime/Core/Public/Misc/AssertionMacros.h
@@ -377,7 +377,7 @@ RetType FORCENOINLINE UE_DEBUG_SECTION DispatchCheckVerify(InnerType&& Inner, Ar
 		(LIKELY(!!(InExpression)) || (DispatchCheckVerify<bool>([Capture] () UE_DEBUG_SECTION \
 		{ \
 			static bool bExecuted = false; \
-			return CheckVerifyImpl(bExecuted, Always, __FILE__, __LINE__, PLATFORM_RETURN_ADDRESS(), #InExpression, InFormat); \
+			return (!bExecuted || Always) && FPlatformMisc::IsEnsureAllowed() && CheckVerifyImpl(bExecuted, Always, __FILE__, __LINE__, PLATFORM_RETURN_ADDRESS(), #InExpression, InFormat); \
 		}) && [] () { PLATFORM_BREAK(); return false; } ()))
 
 	#define UE_ENSURE_IMPL2(Capture, Always, InExpression, InFormat, ...) \
@@ -385,7 +385,7 @@ RetType FORCENOINLINE UE_DEBUG_SECTION DispatchCheckVerify(InnerType&& Inner, Ar
 		{ \
 			static bool bExecuted = false; \
 			UE_VALIDATE_FORMAT_STRING(InFormat, ##__VA_ARGS__); \
-			return CheckVerifyImpl(bExecuted, Always, __FILE__, __LINE__, PLATFORM_RETURN_ADDRESS(), #InExpression, InFormat, ##__VA_ARGS__); \
+			return (!bExecuted || Always) && FPlatformMisc::IsEnsureAllowed() && CheckVerifyImpl(bExecuted, Always, __FILE__, __LINE__, PLATFORM_RETURN_ADDRESS(), #InExpression, InFormat, ##__VA_ARGS__); \
 		}) && [] () { PLATFORM_BREAK(); return false; } ()))
 
 	#define ensure(           InExpression                ) UE_ENSURE_IMPL ( , false, InExpression, TEXT(""))
