Replaced TEnableIf and std::enable_if and some other traits.

diff --git a/Engine/Source/Runtime/Core/Public/Containers/Array.h b/Engine/Source/Runtime/Core/Public/Containers/Array.h
--- a/Engine/Source/Runtime/Core/Public/Containers/Array.h
+++ b/Engine/Source/Runtime/Core/Public/Containers/Array.h
@@ -2272,7 +2272,7 @@ public:
 	/** Mutable implicit conversion operator to container of compatible element type. */
 	template <
 		typename AliasElementType = ElementType
-		UE_REQUIRES(TIsContainerElementTypeReinterpretable<AliasElementType>::Value)
+		UE_REQUIRES(TIsContainerElementTypeReinterpretable_V<AliasElementType>)
 	>
 	operator TArray<typename TContainerElementTypeCompatibility<AliasElementType>::ReinterpretType, AllocatorType>& ()
 	{
@@ -2285,7 +2285,7 @@ public:
 	/** Immutable implicit conversion operator to constant container of compatible element type. */
 	template <
 		typename AliasElementType = ElementType
-		UE_REQUIRES(TIsContainerElementTypeReinterpretable<AliasElementType>::Value)
+		UE_REQUIRES(TIsContainerElementTypeReinterpretable_V<AliasElementType>)
 	>
 	operator const TArray<typename TContainerElementTypeCompatibility<AliasElementType>::ReinterpretType, AllocatorType>& () const
 	{
@@ -2302,7 +2302,7 @@ public:
 	 */
 	template <
 		typename AliasElementType = ElementType
-		UE_REQUIRES(TIsContainerElementTypeCopyable<AliasElementType>::Value)
+		UE_REQUIRES(TIsContainerElementTypeCopyable_V<AliasElementType>)
 	>
 	TArray& operator=(TArray<typename TContainerElementTypeCompatibility<ElementType>::CopyFromOtherType, AllocatorType>&& Other)
 	{
@@ -2322,7 +2322,7 @@ public:
 	template <
 		typename OtherAllocator,
 		typename AliasElementType = ElementType
-		UE_REQUIRES(TIsContainerElementTypeCopyable<AliasElementType>::Value)
+		UE_REQUIRES(TIsContainerElementTypeCopyable_V<AliasElementType>)
 	>
 	TArray& operator=(const TArray<typename TContainerElementTypeCompatibility<ElementType>::CopyFromOtherType, OtherAllocator>& Other)
 	{
@@ -2343,7 +2343,7 @@ public:
 	template <
 		typename OtherAllocator,
 		typename AliasElementType = ElementType
-		UE_REQUIRES(TIsContainerElementTypeCopyable<AliasElementType>::Value)
+		UE_REQUIRES(TIsContainerElementTypeCopyable_V<AliasElementType>)
 	>
 	SizeType Insert(const TArray<typename TContainerElementTypeCompatibility<ElementType>::CopyFromOtherType, OtherAllocator>& Items, const SizeType InIndex)
 	{
@@ -2368,7 +2368,7 @@ public:
 	template <
 		typename OtherAllocator,
 		typename AliasElementType = ElementType
-		UE_REQUIRES(TIsContainerElementTypeCopyable<AliasElementType>::Value)
+		UE_REQUIRES(TIsContainerElementTypeCopyable_V<AliasElementType>)
 	>
 	SizeType Insert(TArray<typename TContainerElementTypeCompatibility<ElementType>::CopyFromOtherType, OtherAllocator>&& Items, const SizeType InIndex)
 	{
@@ -2394,7 +2394,7 @@ public:
 	 */
 	template <
 		typename AliasElementType = ElementType
-		UE_REQUIRES(TIsContainerElementTypeCopyable<AliasElementType>::Value)
+		UE_REQUIRES(TIsContainerElementTypeCopyable_V<AliasElementType>)
 	>
 	void Append(const typename TContainerElementTypeCompatibility<ElementType>::CopyFromOtherType* Ptr, SizeType Count)
 	{
diff --git a/Engine/Source/Runtime/Core/Public/Containers/ContainerElementTypeCompatibility.h b/Engine/Source/Runtime/Core/Public/Containers/ContainerElementTypeCompatibility.h
--- a/Engine/Source/Runtime/Core/Public/Containers/ContainerElementTypeCompatibility.h
+++ b/Engine/Source/Runtime/Core/Public/Containers/ContainerElementTypeCompatibility.h
@@ -34,9 +34,12 @@ struct TContainerElementTypeCompatibility
  * These types and the container methods that leverage them will be removed in a future release.
  */
 template <typename ElementType>
-struct TIsContainerElementTypeReinterpretable
+constexpr bool TIsContainerElementTypeReinterpretable_V = !std::is_same_v<typename TContainerElementTypeCompatibility<ElementType>::ReinterpretType, ElementType>;
+
+template <typename ElementType>
+struct UE_DEPRECATED(5.4, "TIsContainerElementTypeReinterpretable has been deprecated - please use TIsContainerElementTypeReinterpretable_V instead") TIsContainerElementTypeReinterpretable
 {
-	enum { Value = !std::is_same_v<typename TContainerElementTypeCompatibility<ElementType>::ReinterpretType, ElementType> };
+	enum { Value = TIsContainerElementTypeReinterpretable_V<ElementType> };
 };
 
 /**
@@ -44,7 +47,10 @@ struct TIsContainerElementTypeReinterpretable
  * These types and the container methods that leverage them will be removed in a future release.
  */
 template <typename ElementType>
-struct TIsContainerElementTypeCopyable
+constexpr bool TIsContainerElementTypeCopyable_V = !std::is_same_v<typename TContainerElementTypeCompatibility<ElementType>::CopyFromOtherType, ElementType>;
+
+template <typename ElementType>
+struct UE_DEPRECATED(5.4, "TIsContainerElementTypeCopyable has been deprecated - please use TIsContainerElementTypeCopyable_V instead") TIsContainerElementTypeCopyable
 {
-	enum { Value = !std::is_same_v<typename TContainerElementTypeCompatibility<ElementType>::CopyFromOtherType, ElementType> };
+	enum { Value = TIsContainerElementTypeCopyable_V<ElementType> };
 };
diff --git a/Engine/Source/Runtime/Core/Public/Containers/Set.h b/Engine/Source/Runtime/Core/Public/Containers/Set.h
--- a/Engine/Source/Runtime/Core/Public/Containers/Set.h
+++ b/Engine/Source/Runtime/Core/Public/Containers/Set.h
@@ -1315,7 +1315,7 @@ public:
 	template <
 		typename OtherKeyFuncs,
 		typename AliasElementType = ElementType
-		UE_REQUIRES(TIsContainerElementTypeCopyable<AliasElementType>::Value)
+		UE_REQUIRES(TIsContainerElementTypeCopyable_V<AliasElementType>)
 	>
 	TSet& operator=(TSet<typename TContainerElementTypeCompatibility<ElementType>::CopyFromOtherType, OtherKeyFuncs, Allocator>&& Other)
 	{
@@ -1336,7 +1336,7 @@ public:
 		typename OtherKeyFuncs,
 		typename OtherAllocator,
 		typename AliasElementType = ElementType
-		UE_REQUIRES(TIsContainerElementTypeCopyable<AliasElementType>::Value)
+		UE_REQUIRES(TIsContainerElementTypeCopyable_V<AliasElementType>)
 	>
 	TSet& operator=(const TSet<typename TContainerElementTypeCompatibility<ElementType>::CopyFromOtherType, OtherKeyFuncs, OtherAllocator>& Other)
 	{
@@ -1355,7 +1355,7 @@ public:
 		typename OtherKeyFuncs,
 		typename OtherAllocator,
 		typename AliasElementType = ElementType
-		UE_REQUIRES(TIsContainerElementTypeCopyable<AliasElementType>::Value)
+		UE_REQUIRES(TIsContainerElementTypeCopyable_V<AliasElementType>)
 	>
 	void Append(const TSet<typename TContainerElementTypeCompatibility<ElementType>::CopyFromOtherType, OtherKeyFuncs, OtherAllocator>& OtherSet)
 	{
@@ -1375,7 +1375,7 @@ public:
 	template <
 		typename OtherKeyFuncs,
 		typename AliasElementType = ElementType
-		UE_REQUIRES(TIsContainerElementTypeCopyable<AliasElementType>::Value)
+		UE_REQUIRES(TIsContainerElementTypeCopyable_V<AliasElementType>)
 	>
 	void Append(TSet<typename TContainerElementTypeCompatibility<ElementType>::CopyFromOtherType, OtherKeyFuncs, Allocator>&& OtherSet)
 	{
diff --git a/Engine/Source/Runtime/Core/Public/Templates/NonNullPointer.h b/Engine/Source/Runtime/Core/Public/Templates/NonNullPointer.h
--- a/Engine/Source/Runtime/Core/Public/Templates/NonNullPointer.h
+++ b/Engine/Source/Runtime/Core/Public/Templates/NonNullPointer.h
@@ -5,16 +5,13 @@
 #include "CoreTypes.h"
 #include "Misc/AssertionMacros.h"
 #include "Misc/OptionalFwd.h"
-#include "Templates/EnableIf.h"
-#include "Templates/PointerIsConvertibleFromTo.h"
+
+#include <type_traits>
 
 class FArchive;
 enum class EDefaultConstructNonNullPtr { UnsafeDoNotUse }; // So we can construct TNonNullPtrs
 
 namespace UE::Core::Private::NonNullPtr {
-template <typename...>
-using TVoid = void;
-
 /**
  * Version of `::TPointerIsConvertibleFromTo` that produces an incomplete type
  * when either `From` or `To` are incomplete types
@@ -28,9 +25,9 @@ struct TPointerIsConvertibleFromTo;
  * non-function types
  */
 template <typename From, typename To>
-struct TPointerIsConvertibleFromTo<From, To, TVoid<decltype(sizeof(From)), decltype(sizeof(To))>>
-	: ::TPointerIsConvertibleFromTo<From, To>
+struct TPointerIsConvertibleFromTo<From, To, std::void_t<decltype(sizeof(From)), decltype(sizeof(To))>>
 {
+	static constexpr bool Value = std::is_convertible_v<From*, To*>;
 };
 
 /**
@@ -40,8 +37,8 @@ struct TPointerIsConvertibleFromTo<From, To, TVoid<decltype(sizeof(From)), declt
  */
 template <typename Result1, typename... Args1, typename Result2, typename... Args2>
 struct TPointerIsConvertibleFromTo<Result1(Args1...), Result2(Args2...)>
-	: ::TPointerIsConvertibleFromTo<Result1(Args1...), Result2(Args2...)>
 {
+	static constexpr bool Value = std::is_convertible_v<Result1(*)(Args1...), Result2(*)(Args2...)>;
 };
 }
 
@@ -83,8 +80,8 @@ public:
 	 * Constructs a non-null pointer from another non-null pointer
 	 */
 	template <
-		typename OtherObjectType,
-		typename = typename TEnableIf<UE::Core::Private::NonNullPtr::TPointerIsConvertibleFromTo<OtherObjectType, ObjectType>::Value>::Type
+		typename OtherObjectType
+		UE_REQUIRES(UE::Core::Private::NonNullPtr::TPointerIsConvertibleFromTo<OtherObjectType, ObjectType>::Value)
 	>
 	FORCEINLINE TNonNullPtr(const TNonNullPtr<OtherObjectType>& Other)
 		: Object(Other.Object)
@@ -114,8 +111,11 @@ public:
 	/**
 	 * Assignment operator taking another TNonNullPtr
 	 */
-	template <typename OtherObjectType>
-	FORCEINLINE typename TEnableIf<UE::Core::Private::NonNullPtr::TPointerIsConvertibleFromTo<OtherObjectType, ObjectType>::Value, TNonNullPtr&>::Type operator=(const TNonNullPtr<OtherObjectType>& Other)
+	template <
+		typename OtherObjectType
+		UE_REQUIRES(UE::Core::Private::NonNullPtr::TPointerIsConvertibleFromTo<OtherObjectType, ObjectType>::Value)
+	>
+	FORCEINLINE TNonNullPtr& operator=(const TNonNullPtr<OtherObjectType>& Other)
 	{
 		Object = Other.Object;
 		return *this;
@@ -240,3 +240,8 @@ public:
 private:
 	OptionalType* Pointer;
 };
+
+#if UE_ENABLE_INCLUDE_ORDER_DEPRECATED_IN_5_4
+#include "Templates/EnableIf.h"
+#include "Templates/PointerIsConvertibleFromTo.h"
+#endif
diff --git a/Engine/Source/Runtime/Core/Public/Templates/SharedPointer.h b/Engine/Source/Runtime/Core/Public/Templates/SharedPointer.h
--- a/Engine/Source/Runtime/Core/Public/Templates/SharedPointer.h
+++ b/Engine/Source/Runtime/Core/Public/Templates/SharedPointer.h
@@ -146,12 +146,11 @@ namespace UE::Core::Private
 // A type which does not inherit TSharedFromThis at all is ok.
 // We only check this inside the constructor because we don't necessarily have the full type of T when we declare a TSharedPtr<T>.
 #define UE_TSHAREDPTR_STATIC_ASSERT_VALID_MODE(ObjectType, Mode) \
-	enum \
-	{ \
-		ObjectTypeHasSameModeSharedFromThis = TPointerIsConvertibleFromTo<ObjectType, TSharedFromThis<ObjectType, Mode>>::Value, \
-		ObjectTypeHasOppositeModeSharedFromThis = TPointerIsConvertibleFromTo<ObjectType, TSharedFromThis<ObjectType, (Mode == ESPMode::NotThreadSafe) ? ESPMode::ThreadSafe : ESPMode::NotThreadSafe>>::Value \
-	}; \
-	static_assert(ObjectTypeHasSameModeSharedFromThis || !ObjectTypeHasOppositeModeSharedFromThis, "You cannot use a TSharedPtr of one mode with a type which inherits TSharedFromThis of another mode.");
+	static_assert( \
+		std::is_convertible_v<ObjectType*, TSharedFromThis<ObjectType, Mode>*> || \
+		!std::is_convertible_v<ObjectType*, TSharedFromThis<ObjectType, (Mode == ESPMode::NotThreadSafe) ? ESPMode::ThreadSafe : ESPMode::NotThreadSafe>*>, \
+		"You cannot use a TSharedPtr of one mode with a type which inherits TSharedFromThis of another mode." \
+	);
 
 /**
  * TSharedRef is a non-nullable, non-intrusive reference-counted authoritative object reference.
@@ -683,30 +682,6 @@ struct TCallTraits<TSharedRef<ObjectType, Mode>> : public TCallTraitsBase<TShare
 };
 
 
-/**
- * Wrapper for a type that yields a reference to that type.
- */
-template<class T>
-struct FMakeReferenceTo
-{
-	typedef T& Type;
-};
-
-
-/**
- * Specialization for FMakeReferenceTo<void>.
- */
-template<>
-struct FMakeReferenceTo<void>
-{
-	typedef void Type;
-};
-template<>
-struct FMakeReferenceTo<const void>
-{
-	typedef void Type;
-};
-
 /**
  * TSharedPtr is a non-intrusive reference-counted authoritative object pointer.  This shared pointer
  * will be conditionally thread-safe when the optional Mode template argument is set to ThreadSafe.
@@ -1123,7 +1098,7 @@ public:
 	 *
 	 * @return  Reference to the object
 	 */
-	[[nodiscard]] FORCEINLINE typename FMakeReferenceTo<ObjectType>::Type operator*() const
+	[[nodiscard]] FORCEINLINE decltype(auto) operator*() const
 	{
 		check( IsValid() );
 		return *Object;
@@ -1209,8 +1184,8 @@ private:
 	 *       pointer to a shared pointer.  Use the weak pointer's Pin() method instead!
 	 */
 	template <
-		typename OtherType,
-		typename = typename TEnableIf<TPointerIsConvertibleFromTo<OtherType, ObjectType>::Value>::Type
+		typename OtherType
+		UE_REQUIRES(std::is_convertible_v<OtherType*, ObjectType*>)
 	>
 	FORCEINLINE explicit TSharedPtr(TWeakPtr< OtherType, Mode >&& InWeakPtr)
 		: Object(nullptr)
diff --git a/Engine/Source/Runtime/Core/Public/Templates/UniquePtr.h b/Engine/Source/Runtime/Core/Public/Templates/UniquePtr.h
--- a/Engine/Source/Runtime/Core/Public/Templates/UniquePtr.h
+++ b/Engine/Source/Runtime/Core/Public/Templates/UniquePtr.h
@@ -796,8 +796,12 @@ template <
  *
  * @return A TUniquePtr which points to a newly-constructed T with the specified Args.
  */
-template <typename T, typename... TArgs>
-[[nodiscard]] FORCEINLINE typename TEnableIf<!std::is_array_v<T>, TUniquePtr<T>>::Type MakeUniqueForOverwrite()
+template <
+	typename T,
+	typename... TArgs
+	UE_REQUIRES(!std::is_array_v<T>)
+>
+[[nodiscard]] FORCEINLINE TUniquePtr<T> MakeUniqueForOverwrite()
 {
 	return TUniquePtr<T>(new T);
 }
@@ -810,10 +814,13 @@ template <typename T, typename... TArgs>
  *
  * @return A TUniquePtr which points to a newly-constructed T array of the specified Size.
  */
-template <typename T>
-[[nodiscard]] FORCEINLINE typename TEnableIf<TIsUnboundedArray<T>::Value, TUniquePtr<T>>::Type MakeUnique(SIZE_T Size)
+template <
+	typename T
+	UE_REQUIRES(TIsUnboundedArray<T>::Value)
+>
+[[nodiscard]] FORCEINLINE TUniquePtr<T> MakeUnique(SIZE_T Size)
 {
-	typedef typename TRemoveExtent<T>::Type ElementType;
+	using ElementType = std::remove_extent_t<T>;
 	return TUniquePtr<T>(new ElementType[Size]());
 }
 
@@ -825,18 +832,25 @@ template <typename T>
  *
  * @return A TUniquePtr which points to a newly-constructed T array of the specified Size.
  */
-template <typename T>
-[[nodiscard]] FORCEINLINE typename TEnableIf<TIsUnboundedArray<T>::Value, TUniquePtr<T>>::Type MakeUniqueForOverwrite(SIZE_T Size)
+template <
+	typename T
+	UE_REQUIRES(TIsUnboundedArray<T>::Value)
+>
+[[nodiscard]] FORCEINLINE TUniquePtr<T> MakeUniqueForOverwrite(SIZE_T Size)
 {
-	typedef typename TRemoveExtent<T>::Type ElementType;
+	using ElementType = std::remove_extent_t<T>;
 	return TUniquePtr<T>(new ElementType[Size]);
 }
 
 /**
  * Overload to cause a compile error when MakeUnique<T[N]> is attempted.  Use MakeUnique<T>(N) instead.
  */
-template <typename T, typename... TArgs>
-typename TEnableIf<TIsBoundedArray<T>::Value, TUniquePtr<T>>::Type MakeUnique(TArgs&&... Args) = delete;
+template <
+	typename T,
+	typename... TArgs
+	UE_REQUIRES(TIsBoundedArray<T>::Value)
+>
+TUniquePtr<T> MakeUnique(TArgs&&... Args) = delete;
 
 /**
  * Trait which determines whether or not a type is a TUniquePtr.
diff --git a/Engine/Source/Runtime/Core/Public/Templates/UnrealTemplate.h b/Engine/Source/Runtime/Core/Public/Templates/UnrealTemplate.h
--- a/Engine/Source/Runtime/Core/Public/Templates/UnrealTemplate.h
+++ b/Engine/Source/Runtime/Core/Public/Templates/UnrealTemplate.h
@@ -5,10 +5,7 @@
 #include "CoreTypes.h"
 #include "Templates/IsPointer.h"
 #include "HAL/UnrealMemory.h"
-#include "Templates/EnableIf.h"
-#include "Templates/AndOrNot.h"
 #include "Templates/CopyQualifiersAndRefsFromTo.h"
-#include "Templates/IsArithmetic.h"
 #include "Templates/UnrealTypeTraits.h"
 #include "Templates/RemoveReference.h"
 #include "Templates/TypeCompatibleBytes.h"
@@ -70,7 +67,10 @@ FORCEINLINE void Move(T& A,typename TMoveSupportTraits<T>::Move B)
 /**
  * Generically gets the data pointer of a contiguous container
  */
-template<typename T, typename = typename TEnableIf<TIsContiguousContainer<T>::Value>::Type>
+template <
+	typename T
+	UE_REQUIRES(TIsContiguousContainer<T>::Value)
+>
 constexpr auto GetData(T&& Container) -> decltype(Container.GetData())
 {
 	return Container.GetData();
@@ -90,7 +90,10 @@ constexpr const T* GetData(std::initializer_list<T> List)
 /**
 * Generically gets the number of items in a contiguous container
 */
-template<typename T, typename = typename TEnableIf<TIsContiguousContainer<T>::Value>::Type>
+template <
+	typename T
+	UE_REQUIRES(TIsContiguousContainer<T>::Value)
+>
 constexpr auto GetNum(T&& Container) -> decltype(Container.Num())
 {
 	return Container.Num();
@@ -195,8 +198,11 @@ OutType FloatCastChecked(InType In, InType Precision)
 ----------------------------------------------------------------------------*/
 
 #ifdef __clang__
-	template <typename T>
-	auto UEArrayCountHelper(T& t) -> typename TEnableIf<__is_array(T), char(&)[sizeof(t) / sizeof(t[0]) + 1]>::Type;
+	template <
+		typename T
+		UE_REQUIRES(__is_array(T))
+	>
+	auto UEArrayCountHelper(T& t) -> char(&)[sizeof(t) / sizeof(t[0]) + 1];
 #else
 	template <typename T, uint32 N>
 	char (&UEArrayCountHelper(const T (&)[N]))[N + 1];
@@ -513,12 +519,12 @@ template <typename T> struct TRemovePointer<T*> { typedef T Type; };
  * const object, because we would prefer to be informed when MoveTemp will have no effect.
  */
 template <typename T>
-UE_INTRINSIC_CAST FORCEINLINE typename TRemoveReference<T>::Type&& MoveTemp(T&& Obj)
+UE_INTRINSIC_CAST FORCEINLINE std::remove_reference_t<T>&& MoveTemp(T&& Obj)
 {
-	typedef typename TRemoveReference<T>::Type CastType;
+	using CastType = std::remove_reference_t<T>;
 
 	// Validate that we're not being passed an rvalue or a const object - the former is redundant, the latter is almost certainly a mistake
-	static_assert(TIsLValueReferenceType<T>::Value, "MoveTemp called on an rvalue");
+	static_assert(std::is_lvalue_reference_v<T>, "MoveTemp called on an rvalue");
 	static_assert(!std::is_same_v<CastType&, const CastType&>, "MoveTemp called on a const object");
 
 	return (CastType&&)Obj;
@@ -531,9 +537,9 @@ UE_INTRINSIC_CAST FORCEINLINE typename TRemoveReference<T>::Type&& MoveTemp(T&&
  * where you can but not stop compilation.
  */
 template <typename T>
-UE_INTRINSIC_CAST FORCEINLINE typename TRemoveReference<T>::Type&& MoveTempIfPossible(T&& Obj)
+UE_INTRINSIC_CAST FORCEINLINE std::remove_reference_t<T>&& MoveTempIfPossible(T&& Obj)
 {
-	typedef typename TRemoveReference<T>::Type CastType;
+	using CastType = std::remove_reference_t<T>;
 	return (CastType&&)Obj;
 }
 
@@ -558,22 +564,6 @@ FORCEINLINE T CopyTemp(const T& Val)
 	return Val;
 }
 
-template <typename T>
-UE_DEPRECATED(5.2, "CopyTemp on an rvalue is deprecated and should be removed or replaced with CopyTempIfNecessary when the argument is unknown")
-FORCEINLINE T&& CopyTemp(T&& Val)
-{
-// Compile this block back in after removing the deprecation, rather than deleting the function entirely.
-// Also change the return type to `T`.
-#if 0
-	// Comment this in rather than deleting the function when removing the deprecation
-	static_assert(sizeof(T) == 0, "CopyTemp called on an rvalue");
-#endif
-
-	// Create a prvalue by move-constructing from the xvalue - wasteful if Val
-	// already refers to a prvalue, but we can't differentiate those.
-	return MoveTemp(Val);
-}
-
 /**
  * CopyTempIfNecessary will enforce the creation of a prvalue.
  * This is UE's equivalent of the exposition std::decay-copy:
@@ -593,13 +583,13 @@ FORCEINLINE std::decay_t<T> CopyTempIfNecessary(T&& Val)
  * This is UE's equivalent of std::forward.
  */
 template <typename T>
-UE_INTRINSIC_CAST FORCEINLINE T&& Forward(typename TRemoveReference<T>::Type& Obj)
+UE_INTRINSIC_CAST FORCEINLINE T&& Forward(std::remove_reference_t<T>& Obj)
 {
 	return (T&&)Obj;
 }
 
 template <typename T>
-UE_INTRINSIC_CAST FORCEINLINE T&& Forward(typename TRemoveReference<T>::Type&& Obj)
+UE_INTRINSIC_CAST FORCEINLINE T&& Forward(std::remove_reference_t<T>&& Obj)
 {
 	return (T&&)Obj;
 }
@@ -661,13 +651,16 @@ template <typename T> struct TRValueToLValueReference<T&&> { typedef T& Type; };
 
 /**
  * Reverses the order of the bits of a value.
- * This is an TEnableIf'd template to ensure that no undesirable conversions occur.  Overloads for other types can be added in the same way.
+ * This is a constrained template to ensure that no undesirable conversions occur.  Overloads for other types can be added in the same way.
  *
  * @param Bits - The value to bit-swap.
  * @return The bit-swapped value.
  */
-template <typename T>
-FORCEINLINE typename TEnableIf<std::is_same_v<T, uint32>, T>::Type ReverseBits( T Bits )
+template <
+	typename T
+	UE_REQUIRES(std::is_same_v<T, uint32>)
+>
+FORCEINLINE T ReverseBits( T Bits )
 {
 	Bits = ( Bits << 16) | ( Bits >> 16);
 	Bits = ( (Bits & 0x00ff00ff) << 8 ) | ( (Bits & 0xff00ff00) >> 8 );
@@ -752,7 +745,7 @@ T&& DeclVal();
 template <typename T>
 FORCEINLINE T ImplicitConv(typename TIdentity<T>::Type Obj)
 {
-    return Obj;
+	return Obj;
 }
 
 /**
@@ -762,9 +755,9 @@ FORCEINLINE T ImplicitConv(typename TIdentity<T>::Type Obj)
 template <
 	typename T,
 	typename Base,
-	decltype(ImplicitConv<const volatile Base*>((typename TRemoveReference<T>::Type*)nullptr))* = nullptr
+	decltype(ImplicitConv<const volatile Base*>((std::remove_reference_t<T>*)nullptr))* = nullptr
 >
-UE_INTRINSIC_CAST FORCEINLINE decltype(auto) ForwardAsBase(typename TRemoveReference<T>::Type& Obj)
+UE_INTRINSIC_CAST FORCEINLINE decltype(auto) ForwardAsBase(std::remove_reference_t<T>& Obj)
 {
 	return (TCopyQualifiersAndRefsFromTo_T<T&&, Base>)Obj;
 }
@@ -772,9 +765,15 @@ UE_INTRINSIC_CAST FORCEINLINE decltype(auto) ForwardAsBase(typename TRemoveRefer
 template <
 	typename T,
 	typename Base,
-	decltype(ImplicitConv<const volatile Base*>((typename TRemoveReference<T>::Type*)nullptr))* = nullptr
+	decltype(ImplicitConv<const volatile Base*>((std::remove_reference_t<T>*)nullptr))* = nullptr
 >
-UE_INTRINSIC_CAST FORCEINLINE decltype(auto) ForwardAsBase(typename TRemoveReference<T>::Type&& Obj)
+UE_INTRINSIC_CAST FORCEINLINE decltype(auto) ForwardAsBase(std::remove_reference_t<T>&& Obj)
 {
 	return (TCopyQualifiersAndRefsFromTo_T<T&&, Base>)Obj;
-}
\ No newline at end of file
+}
+
+#if UE_ENABLE_INCLUDE_ORDER_DEPRECATED_IN_5_4
+#include "Templates/AndOrNot.h"
+#include "Templates/EnableIf.h"
+#include "Templates/IsArithmetic.h"
+#endif
diff --git a/Engine/Source/Runtime/Core/Public/UObject/WeakObjectPtrTemplates.h b/Engine/Source/Runtime/Core/Public/UObject/WeakObjectPtrTemplates.h
--- a/Engine/Source/Runtime/Core/Public/UObject/WeakObjectPtrTemplates.h
+++ b/Engine/Source/Runtime/Core/Public/UObject/WeakObjectPtrTemplates.h
@@ -3,9 +3,6 @@
 #pragma once
 
 #include "CoreTypes.h"
-#include "Templates/IsPointer.h"
-#include "Templates/PointerIsConvertibleFromTo.h"
-#include "Templates/AndOrNot.h"
 #include "Templates/LosesQualifiersFromTo.h"
 #include "Containers/Map.h"
 #include "UObject/WeakObjectPtrTemplatesFwd.h"
@@ -56,7 +53,7 @@ public:
 	{
 		// This static assert is in here rather than in the body of the class because we want
 		// to be able to define TWeakObjectPtr<UUndefinedClass>.
-		static_assert(TPointerIsConvertibleFromTo<T, const volatile UObject>::Value, "TWeakObjectPtr can only be constructed with UObject types");
+		static_assert(std::is_convertible_v<T*, const volatile UObject*>, "TWeakObjectPtr can only be constructed with UObject types");
 	}
 
 	/**
@@ -84,8 +81,11 @@ public:
 	 * Copy from an object pointer
 	 * @param Object object to create a weak pointer to
 	 */
-	template<class U>
-	FORCEINLINE typename TEnableIf<!TLosesQualifiersFromTo<U, T>::Value, TWeakObjectPtr&>::Type operator=(U* Object)
+	template <
+		typename U
+		UE_REQUIRES(!TLosesQualifiersFromTo<U, T>::Value)
+	>
+	FORCEINLINE TWeakObjectPtr& operator=(U* Object)
 	{
 		T* TempObject = Object;
 		TWeakObjectPtrBase::operator=(TempObject);
@@ -350,52 +350,6 @@ struct TWeakObjectPtrMapKeyFuncs : public TDefaultMapKeyFuncs<KeyType, ValueType
 	}
 };
 
-/**
- * Automatic version of the weak object pointer
- */
-template<class T> 
-class TAutoWeakObjectPtr : public TWeakObjectPtr<T>
-{
-public:
-	/** NULL constructor **/
-	UE_DEPRECATED(4.15, "TAutoWeakObjectPtr has been deprecated - use TWeakObjectPtr instead")
-	FORCEINLINE TAutoWeakObjectPtr()
-	{
-	}
-	/** Construct from a raw pointer **/
-	UE_DEPRECATED(4.15, "TAutoWeakObjectPtr has been deprecated - use TWeakObjectPtr instead")
-	FORCEINLINE TAutoWeakObjectPtr(const T* Target)
-		: TWeakObjectPtr<T>(Target)
-	{
-	}
-	/**  Construct from the base type **/
-	UE_DEPRECATED(4.15, "TAutoWeakObjectPtr has been deprecated - use TWeakObjectPtr instead")
-	FORCEINLINE TAutoWeakObjectPtr(const TWeakObjectPtr<T>& Other) 
-		: TWeakObjectPtr<T>(Other)
-	{
-	}
-	UE_DEPRECATED(4.15, "Implicit conversion from TAutoWeakObjectPtr to the pointer type has been deprecated - use Get() instead")
-	FORCEINLINE operator T* () const
-	{
-		return this->Get();
-	}
-	UE_DEPRECATED(4.15, "Implicit conversion from TAutoWeakObjectPtr to the pointer type has been deprecated - use Get() instead")
-	FORCEINLINE operator const T* () const
-	{
-		return (const T*)this->Get();
-	}
-
-	UE_DEPRECATED(4.15, "Implicit conversion from TAutoWeakObjectPtr to the pointer type has been deprecated - use Get() instead")
-	FORCEINLINE explicit operator bool() const
-	{
-		return this->Get() != nullptr;
-	}
-};
-
-template<class T> struct TIsPODType<TAutoWeakObjectPtr<T> > { enum { Value = true }; };
-template<class T> struct TIsZeroConstructType<TAutoWeakObjectPtr<T> > { enum { Value = true }; };
-template<class T> struct TIsWeakPointerType<TAutoWeakObjectPtr<T> > { enum { Value = true }; };
-
 template <typename T>
 struct TCallTraits<TWeakObjectPtr<T>> : public TCallTraitsBase<TWeakObjectPtr<T>>
 {
@@ -433,3 +387,9 @@ FArchive& operator<<( FArchive& Ar, TWeakObjectPtr<T, TWeakObjectPtrBase>& WeakO
 	WeakObjectPtr.Serialize(Ar);
 	return Ar;
 }
+
+#if UE_ENABLE_INCLUDE_ORDER_DEPRECATED_IN_5_4
+#include "Templates/AndOrNot.h"
+#include "Templates/IsPointer.h"
+#include "Templates/PointerIsConvertibleFromTo.h"
+#endif
