Enable memory tracing on Android

diff --git a/Engine/Build/Android/Java/src/com/epicgames/unreal/GameActivity.java.template b/Engine/Build/Android/Java/src/com/epicgames/unreal/GameActivity.java.template
--- a/Engine/Build/Android/Java/src/com/epicgames/unreal/GameActivity.java.template
+++ b/Engine/Build/Android/Java/src/com/epicgames/unreal/GameActivity.java.template
@@ -80,6 +80,8 @@ import android.app.ProgressDialog;
 import android.os.Bundle;
 import android.os.PersistableBundle;
 import android.util.Log;
+import android.system.Os;
+import android.system.ErrnoException;
 
 import android.hardware.input.InputManager;
 import android.os.Vibrator;
@@ -7063,6 +7065,27 @@ public class GameActivity extends $${gameActivitySuperClass}$$ implements Surfac
 
 	static
 	{
+		// We need to decide on enabling memory tracing before loading native code.
+		// It's not trivial to get a full command line before GameActivity is constructed, so fallback to a sentinel file.
+		try
+		{
+			Context context = GameApplication.getAppContext();
+			Bundle bundle = context.getPackageManager().getApplicationInfo(context.getPackageName(), PackageManager.GET_META_DATA).metaData;
+			boolean bShipping =
+				bundle.containsKey("com.epicgames.unreal.GameActivity.BuildConfiguration") &&
+				bundle.getString("com.epicgames.unreal.GameActivity.BuildConfiguration").equals("Shipping");
+			if (!bShipping)
+			{ 
+				boolean bEnableMemoryTracing = new File(GameApplication.getAppContext().getFilesDir() + "/UEEnableMemoryTracing.txt").exists();
+				Os.setenv("UEEnableMemoryTracing", bEnableMemoryTracing ? "1" : "0", true);
+			}
+		}
+		catch (ErrnoException | NameNotFoundException e)
+		{
+			Log.debug(e.toString());
+			Log.debug("Cannot set \"UEEnableMemoryTracing\" environment variable");
+		}
+
 //$${soLoadUnreal}$$
 //$${soLoadLibrary}$$
 	}
diff --git a/Engine/Source/Runtime/Core/Core.Build.cs b/Engine/Source/Runtime/Core/Core.Build.cs
--- a/Engine/Source/Runtime/Core/Core.Build.cs
+++ b/Engine/Source/Runtime/Core/Core.Build.cs
@@ -149,6 +149,14 @@ public class Core : ModuleRules
 				"zlib",
 				"libunwind"
 				);
+
+			if (Target.Configuration != UnrealTargetConfiguration.Shipping && Target.Type != TargetType.Program)
+			{
+				PublicDefinitions.Add("UE_MEMORY_TRACE_AVAILABLE=1");
+				PublicDefinitions.Add("UE_MEMORY_TAGS_TRACE_ENABLED=1");
+				PublicDefinitions.Add("UE_CALLSTACK_TRACE_ENABLED=1");
+				PublicDefinitions.Add("UE_CALLSTACK_TRACE_ANDROID_USE_STACK_FRAMES_WALKING=1");
+			}
 		}
 		else if (Target.IsInPlatformGroup(UnrealPlatformGroup.Unix))
 		{
diff --git a/Engine/Source/Runtime/Core/Private/Android/AndroidCallstackTrace.cpp b/Engine/Source/Runtime/Core/Private/Android/AndroidCallstackTrace.cpp
new file mode 100644
--- /dev/null
+++ b/Engine/Source/Runtime/Core/Private/Android/AndroidCallstackTrace.cpp
@@ -0,0 +1,146 @@
+﻿// Copyright Epic Games, Inc. All Rights Reserved.
+
+#include "ProfilingDebugging/CallstackTrace.h"
+
+#if UE_CALLSTACK_TRACE_ENABLED
+
+#ifndef UE_CALLSTACK_TRACE_ANDROID_USE_STACK_FRAMES_WALKING
+	#define UE_CALLSTACK_TRACE_ANDROID_USE_STACK_FRAMES_WALKING 0
+#endif
+
+#include "CoreTypes.h"
+#include "HAL/Platform.h"
+#include "HAL/PlatformStackWalk.h"
+#include "ProfilingDebugging/CallstackTracePrivate.h"
+
+////////////////////////////////////////////////////////////////////////////////
+class FBacktracer
+{
+public:
+	FBacktracer(FMalloc* InMalloc);
+	~FBacktracer();
+	static FBacktracer*	Get();
+	FORCEINLINE uint32 GetBacktraceId(void* ReturnAddress);
+
+private:
+	static FBacktracer* Instance;
+	FMalloc* Malloc;
+	FCallstackTracer CallstackTracer;
+};
+
+////////////////////////////////////////////////////////////////////////////////
+FBacktracer* FBacktracer::Instance = nullptr;
+
+////////////////////////////////////////////////////////////////////////////////
+FBacktracer::FBacktracer(FMalloc* InMalloc)
+	: Malloc(InMalloc)
+	, CallstackTracer(InMalloc)
+{
+	Instance = this;
+}
+
+////////////////////////////////////////////////////////////////////////////////
+FBacktracer::~FBacktracer()
+{
+}
+
+////////////////////////////////////////////////////////////////////////////////
+FBacktracer* FBacktracer::Get()
+{
+	return Instance;
+}
+
+////////////////////////////////////////////////////////////////////////////////
+
+FORCEINLINE uint32 FBacktracer::GetBacktraceId(void* ReturnAddress)
+{
+#if !UE_BUILD_SHIPPING
+	uint64 StackFrames[64];
+
+#if UE_CALLSTACK_TRACE_ANDROID_USE_STACK_FRAMES_WALKING
+	uint32 NumStackFrames = FPlatformStackWalk::CaptureStackBackTraceViaFramePointerWalking(StackFrames, UE_ARRAY_COUNT(StackFrames));
+#else
+	uint32 NumStackFrames = FPlatformStackWalk::CaptureStackBackTrace(StackFrames, UE_ARRAY_COUNT(StackFrames));
+#endif
+
+	if (NumStackFrames > 0)
+	{
+		FCallstackTracer::FBacktraceEntry BacktraceEntry;
+		uint64 BacktraceId = 0;
+		uint32 FrameIdx = 0;
+		bool bUseAddress = false;
+		for (int32 Index = 0; Index < NumStackFrames; Index++)
+		{
+			if (!bUseAddress)
+			{
+				// start using backtrace only after ReturnAddress
+				if (StackFrames[Index] == (uint64)ReturnAddress)
+				{
+					bUseAddress = true;
+				}
+			}
+			if (bUseAddress || NumStackFrames == 1)
+			{
+				uint64 RetAddr = StackFrames[Index];
+				StackFrames[FrameIdx++] = RetAddr;
+
+				// This is a simple order-dependent LCG. Should be sufficient enough
+				BacktraceId += RetAddr;
+				BacktraceId *= 0x30be8efa499c249dull;
+			}
+		}
+
+		// Save the collected id
+		BacktraceEntry.Hash = BacktraceId;
+		BacktraceEntry.FrameCount = FrameIdx;
+		BacktraceEntry.Frames = StackFrames;
+
+		// Add to queue to be processed. This might block until there is room in the
+		// queue (i.e. the processing thread has caught up processing).
+		return CallstackTracer.AddCallstack(BacktraceEntry);
+	}
+#endif
+
+	return 0;
+}
+
+////////////////////////////////////////////////////////////////////////////////
+void Modules_Initialize();
+
+////////////////////////////////////////////////////////////////////////////////
+void CallstackTrace_CreateInternal(FMalloc* Malloc)
+{
+	if (FBacktracer::Get() != nullptr)
+	{
+		return;
+	}
+
+	// Allocate, construct and intentionally leak backtracer
+	void* Alloc = Malloc->Malloc(sizeof(FBacktracer), alignof(FBacktracer));
+	new (Alloc) FBacktracer(Malloc);
+}
+
+////////////////////////////////////////////////////////////////////////////////
+void CallstackTrace_InitializeInternal()
+{
+	if (FBacktracer* Instance = FBacktracer::Get())
+	{
+#if !UE_BUILD_SHIPPING
+		Modules_Initialize();
+#endif
+	}
+}
+
+////////////////////////////////////////////////////////////////////////////////
+uint32 CallstackTrace_GetCurrentId()
+{
+	void* ReturnAddress = PLATFORM_RETURN_ADDRESS_FOR_CALLSTACKTRACING();
+	if (FBacktracer* Instance = FBacktracer::Get())
+	{
+		return Instance->GetBacktraceId(ReturnAddress);
+	}
+
+	return 0;
+}
+
+#endif
\ No newline at end of file
diff --git a/Engine/Source/Runtime/Core/Private/Android/AndroidMemoryTrace.cpp b/Engine/Source/Runtime/Core/Private/Android/AndroidMemoryTrace.cpp
new file mode 100644
--- /dev/null
+++ b/Engine/Source/Runtime/Core/Private/Android/AndroidMemoryTrace.cpp
@@ -0,0 +1,22 @@
+// Copyright Epic Games, Inc. All Rights Reserved.
+
+#include "ProfilingDebugging/MemoryTrace.h"
+
+#if UE_MEMORY_TRACE_ENABLED
+
+#include <android/log.h>
+
+////////////////////////////////////////////////////////////////////////////////
+FMalloc* MemoryTrace_CreateInternal(FMalloc*);
+
+////////////////////////////////////////////////////////////////////////////////
+FMalloc* MemoryTrace_Create(FMalloc* InMalloc)
+{
+	const char* UEEnableMemoryTracing = getenv("UEEnableMemoryTracing");
+	__android_log_print(ANDROID_LOG_DEBUG, "UE", "getenv(\"UEEnableMemoryTracing\") == \"%s\"", UEEnableMemoryTracing ? UEEnableMemoryTracing : "nullptr");
+
+	const bool bEnableMemoryTracing = UEEnableMemoryTracing != nullptr && !strncmp(UEEnableMemoryTracing, "1", 1);
+	return bEnableMemoryTracing ? MemoryTrace_CreateInternal(InMalloc) :  InMalloc;
+}
+
+#endif // UE_MEMORY_TRACE_ENABLED
diff --git a/Engine/Source/Runtime/Core/Private/Android/AndroidModuleDiagnostics.cpp b/Engine/Source/Runtime/Core/Private/Android/AndroidModuleDiagnostics.cpp
new file mode 100644
--- /dev/null
+++ b/Engine/Source/Runtime/Core/Private/Android/AndroidModuleDiagnostics.cpp
@@ -0,0 +1,94 @@
+﻿// Copyright Epic Games, Inc. All Rights Reserved.
+
+#include "CoreTypes.h"
+#include "HAL/LowLevelMemTracker.h"
+#include "HAL/PlatformProcess.h"
+#include "Math/UnrealMathUtility.h"
+#include "Misc/Paths.h"
+#include "ProfilingDebugging/ModuleDiagnostics.h"
+#include "ProfilingDebugging/MemoryTrace.h"
+#include "ProfilingDebugging/MetadataTrace.h"
+
+#if !UE_BUILD_SHIPPING
+#include <link.h>
+#endif
+
+////////////////////////////////////////////////////////////////////////////////
+void Modules_Initialize()
+{
+#if !UE_BUILD_SHIPPING
+	using namespace UE::Trace;
+
+	constexpr uint32 SizeOfSymbolFormatString = 5;
+	UE_TRACE_LOG(Diagnostics, ModuleInit, ModuleChannel, sizeof(ANSICHAR) * SizeOfSymbolFormatString)
+		<< ModuleInit.SymbolFormat("dwarf", SizeOfSymbolFormatString)
+		<< ModuleInit.ModuleBaseShift(uint8(0));
+
+	HeapId ProgramHeapId = MemoryTrace_HeapSpec(EMemoryTraceRootHeap::SystemMemory, TEXT("Program"), EMemoryTraceHeapFlags::NeverFrees);
+
+	auto IterateCallback = [](struct dl_phdr_info *Info, size_t /*size*/, void *Data)
+	{
+		int32 TotalMemSize = 0;
+		uint64 RealBase = Info->dlpi_addr;
+		bool bRealBaseSet = false;
+
+		constexpr uint32 GNUSectionNameSize = 4;
+		constexpr uint32 BuildIdSize = 20;
+		uint8 BuildId[BuildIdSize] = {0};
+		FMemory::Memzero(BuildId, BuildIdSize);
+		bool bBuildIdSet = false;
+		for (int SectionIdx = 0; SectionIdx < Info->dlpi_phnum; SectionIdx++)
+		{
+			TotalMemSize += Info->dlpi_phdr[SectionIdx].p_memsz;
+
+			uint64 Offset = Info->dlpi_addr + Info->dlpi_phdr[SectionIdx].p_vaddr;
+			uint32 Type = Info->dlpi_phdr[SectionIdx].p_type;
+
+			if (!bRealBaseSet && Type == PT_LOAD)
+			{
+				RealBase = Offset;
+				bRealBaseSet = true;
+			}
+			if (!bBuildIdSet && Type == PT_NOTE)
+			{
+				ElfW(Nhdr)* Note = (ElfW(Nhdr)*)Offset;
+				char* NoteName = (char*)Note + sizeof(ElfW(Nhdr));
+				uint8* NoteDesc = (uint8*)Note + sizeof(ElfW(Nhdr)) + Note->n_namesz;
+				if (Note->n_namesz == GNUSectionNameSize && Note->n_type == NT_GNU_BUILD_ID && strcmp(NoteName, ELF_NOTE_GNU) == 0)
+				{
+					FMemory::Memcpy(BuildId, NoteDesc, FMath::Min(BuildIdSize, Note->n_descsz));
+					bBuildIdSet = true;
+				}
+			}
+		}
+
+		FString ImageName = FString(Info->dlpi_name);
+		if (ImageName.IsEmpty())
+		{
+			constexpr bool bRemoveExtension = false;
+			ImageName = FPlatformProcess::ExecutableName(bRemoveExtension);
+		}
+		ImageName = FPaths::GetCleanFilename(ImageName);
+
+		UE_TRACE_LOG(Diagnostics, ModuleLoad, ModuleChannel, sizeof(TCHAR) * ImageName.Len() + BuildIdSize)
+			<< ModuleLoad.Name(*ImageName, ImageName.Len())
+			<< ModuleLoad.Base(RealBase)
+			<< ModuleLoad.Size(TotalMemSize)
+			<< ModuleLoad.ImageId(BuildId, BuildIdSize);
+
+#if UE_MEMORY_TRACE_ENABLED
+		{
+			UE_TRACE_METADATA_CLEAR_SCOPE();
+			LLM(UE_MEMSCOPE(ELLMTag::ProgramSize));
+			MemoryTrace_Alloc(RealBase, TotalMemSize, 1);
+			HeapId ProgramHeapId = *(HeapId*)Data;
+			MemoryTrace_MarkAllocAsHeap(RealBase, ProgramHeapId);
+			MemoryTrace_Alloc(RealBase, TotalMemSize, 1);
+		}
+#endif // UE_MEMORY_TRACE_ENABLED
+
+		return 0;
+	};
+	dl_iterate_phdr(IterateCallback, &ProgramHeapId);
+#endif
+}
diff --git a/Engine/Source/Runtime/Core/Private/Android/AndroidPlatformStackWalk.cpp b/Engine/Source/Runtime/Core/Private/Android/AndroidPlatformStackWalk.cpp
--- a/Engine/Source/Runtime/Core/Private/Android/AndroidPlatformStackWalk.cpp
+++ b/Engine/Source/Runtime/Core/Private/Android/AndroidPlatformStackWalk.cpp
@@ -13,6 +13,7 @@
 #include <stdio.h>
 #include <signal.h>
 #include <setjmp.h>
+#include <pthread.h>
 #include <android/log.h>
 #include "Android/AndroidSignals.h"
 
@@ -257,6 +258,74 @@ uint32 FAndroidPlatformStackWalk::CaptureStackBackTrace(uint64* BackTrace, uint3
 	return Depth;
 }
 
+uint32 FAndroidPlatformStackWalk::CaptureStackBackTraceViaFramePointerWalking(uint64* BackTrace, uint32 MaxDepth)
+{
+#if PLATFORM_ANDROID_ARM64 || PLATFORM_ANDROID_X64
+	uint64 StackTopPtr = 0;
+
+	// pthread_getattr_np/pthread_attr_getstack are slow for main thread, so let's cache the stack top 
+	const bool bIsMainThread = gettid() == getpid();
+	static uint64 StackTopMainThread = 0;
+	if (bIsMainThread && StackTopMainThread != 0)
+	{
+		StackTopPtr = StackTopMainThread;
+	}
+	else
+	{
+		pthread_attr_t ThreadAttr;
+		pthread_getattr_np(pthread_self(), &ThreadAttr);
+
+		void* StackBase;
+		size_t StackSize;
+		pthread_attr_getstack(&ThreadAttr, &StackBase, &StackSize);
+
+		StackTopPtr = (uint64)StackBase + StackSize;
+
+		if (bIsMainThread)
+		{
+			StackTopMainThread = StackTopPtr;
+		}
+	}
+
+	struct StackFrame
+	{
+		StackFrame* NextFrame;
+		UPTRINT ReturnPtr;
+
+		inline UPTRINT GetReturnPtr()
+		{
+#if PLATFORM_ANDROID_ARM64
+			register uintptr_t Ptr = ReturnPtr;
+			asm("xpaclri" : "+r"(Ptr)); // this instruction is mapped to NOP on pre-PAC architectures
+			return Ptr;
+#else
+			return ReturnPtr;
+#endif
+		}
+	};
+
+	StackFrame* FrameEnd = (StackFrame*)StackTopPtr;
+	StackFrame* FrameStart = (StackFrame*)__builtin_frame_address(0);
+
+	uint32 NumStackFrames = 0;
+
+	for (StackFrame* CurrentFrame = FrameStart;
+		NumStackFrames < MaxDepth &&
+		CurrentFrame->NextFrame > FrameStart &&
+		CurrentFrame->NextFrame <= FrameEnd &&
+		((UPTRINT)CurrentFrame->NextFrame & (sizeof(StackFrame) - 1)) == 0 && // stop at unaligned frame
+		CurrentFrame->GetReturnPtr() != 0; // stop if function ptr is 0
+		CurrentFrame = CurrentFrame->NextFrame, NumStackFrames++)
+	{
+		BackTrace[NumStackFrames] = CurrentFrame->GetReturnPtr();
+	}
+
+	return NumStackFrames;
+#else
+	return 0;
+#endif
+}
+
 bool FAndroidPlatformStackWalk::SymbolInfoToHumanReadableString(const FProgramCounterSymbolInfo& SymbolInfo, ANSICHAR* HumanReadableString, SIZE_T HumanReadableStringSize)
 {
 	const int32 MAX_TEMP_SPRINTF = 256;
diff --git a/Engine/Source/Runtime/Core/Private/ProfilingDebugging/MemoryTrace.cpp b/Engine/Source/Runtime/Core/Private/ProfilingDebugging/MemoryTrace.cpp
--- a/Engine/Source/Runtime/Core/Private/ProfilingDebugging/MemoryTrace.cpp
+++ b/Engine/Source/Runtime/Core/Private/ProfilingDebugging/MemoryTrace.cpp
@@ -327,7 +327,7 @@ static bool MemoryTrace_ShouldEnable(int32 ArgC, const ArgCharType* const* ArgV)
 }
 
 ////////////////////////////////////////////////////////////////////////////////
-static FMalloc* MemoryTrace_CreateInternal(FMalloc* InMalloc)
+FMalloc* MemoryTrace_CreateInternal(FMalloc* InMalloc)
 {
 	// Some OSes (i.e. Windows) will terminate all threads except the main
 	// one as part of static deinit. However we may receive more memory
diff --git a/Engine/Source/Runtime/Core/Public/Android/AndroidPlatform.h b/Engine/Source/Runtime/Core/Public/Android/AndroidPlatform.h
--- a/Engine/Source/Runtime/Core/Public/Android/AndroidPlatform.h
+++ b/Engine/Source/Runtime/Core/Public/Android/AndroidPlatform.h
@@ -65,6 +65,8 @@ typedef FAndroidTypes FPlatformTypes;
 #define PLATFORM_CODE_SECTION(Name)					__attribute__((section(Name)))
 #define PLATFORM_ALLOW_ALLOCATIONS_IN_FASYNCWRITER_SERIALIZEBUFFERTOARCHIVE	0
 
+#define PLATFORM_RETURN_ADDRESS_FOR_CALLSTACKTRACING	PLATFORM_RETURN_ADDRESS
+
 #define PLATFORM_GLOBAL_LOG_CATEGORY				LogAndroid
 
 #ifndef RUNNING_WITH_ASAN
diff --git a/Engine/Source/Runtime/Core/Public/Android/AndroidPlatformStackWalk.h b/Engine/Source/Runtime/Core/Public/Android/AndroidPlatformStackWalk.h
--- a/Engine/Source/Runtime/Core/Public/Android/AndroidPlatformStackWalk.h
+++ b/Engine/Source/Runtime/Core/Public/Android/AndroidPlatformStackWalk.h
@@ -28,6 +28,9 @@ struct FAndroidPlatformStackWalk : public FGenericPlatformStackWalk
 
 	static CORE_API void ProgramCounterToSymbolInfo(uint64 ProgramCounter, FProgramCounterSymbolInfo& out_SymbolInfo);
 	static CORE_API uint32 CaptureStackBackTrace(uint64* BackTrace, uint32 MaxDepth, void* Context = nullptr);
+	// Fast stack backtrace, only for tracing.
+	// Will NOT work from signal handlers, is NOT suitable for crashreporting. Will not see inline functions.
+	static CORE_API uint32 CaptureStackBackTraceViaFramePointerWalking(uint64* BackTrace, uint32 MaxDepth);
 	static CORE_API bool SymbolInfoToHumanReadableString(const FProgramCounterSymbolInfo& SymbolInfo, ANSICHAR* HumanReadableString, SIZE_T HumanReadableStringSize);
 
 	static CORE_API uint32 CaptureThreadStackBackTrace(uint64 ThreadId, uint64* BackTrace, uint32 MaxDepth, void* Context = nullptr);
