- Added a ECVF_SetByPreview priority so we can preview in editor more safely, while other systems may modify the cvar
- Previewing now no longer needs to keep backup values, if can just use the Unset functionality, and activating a preview can take a Tag, so multiple previews can happen at once, and turn off
- Removed the PlatformIndependentValue member, because the History system has enough information without needing the extra storage
- Other-platform CVar variables now have the history filled out as it walks the cvar priorities with VisitPlatformCVarsForEmulation
- Any combination of platform/device profile is now cached in the other-platform CVar variables
- Added IConsoleManager functions, which will load up all known cvars for a platform/device profile, cache them, and can preview them, which means pushing them into the ECVF_SetByPreview priority in the active "real" cvars:
  - LoadAllPlatformCVars()
  - PreviewPlatformCVars()
  - ClearAllPlatformCVars()

diff --git a/Engine/Plugins/FX/Niagara/Source/Niagara/Private/NiagaraPlatformSet.cpp b/Engine/Plugins/FX/Niagara/Source/Niagara/Private/NiagaraPlatformSet.cpp
--- a/Engine/Plugins/FX/Niagara/Source/Niagara/Private/NiagaraPlatformSet.cpp
+++ b/Engine/Plugins/FX/Niagara/Source/Niagara/Private/NiagaraPlatformSet.cpp
@@ -1439,14 +1439,9 @@ const FNiagaraCVarValues& FDeviceProfileValueCache::GetValues(const UDeviceProfi
 				{
 					DefaultStr = PlatformCVar->GetString();
 				}
-				else if (IConsoleVariable* DefaultCVar = CVar->GetDefaultValueVariable())
-				{
-					DefaultStr = DefaultCVar->GetString();
-				}
 				else
 				{
-					//Failing all that we just take the current value.
-					DefaultStr = CVar->GetString();
+					DefaultStr = CVar->GetDefaultValue();
 				}
 			}
 		}
diff --git a/Engine/Source/Runtime/Core/Private/HAL/ConsoleManager.cpp b/Engine/Source/Runtime/Core/Private/HAL/ConsoleManager.cpp
--- a/Engine/Source/Runtime/Core/Private/HAL/ConsoleManager.cpp
+++ b/Engine/Source/Runtime/Core/Private/HAL/ConsoleManager.cpp
@@ -24,16 +24,12 @@ ConsoleManager.cpp: console command handling
 DEFINE_LOG_CATEGORY(LogConsoleResponse);
 DEFINE_LOG_CATEGORY_STATIC(LogConsoleManager, Log, All);
 
-#define UE_ALLOW_CVAR_HISTORY 1
-
-#if UE_ALLOW_CVAR_HISTORY
 namespace UE::ConsoleManager::Private
 {
 	// this tracks the cvars that were  added dynamically with a tag (via plugin or similar)
 	// we use this structure to unset the cvars and update the value when the plugin unloads
 	TMap<FName, TSet<IConsoleVariable*>*> TaggedCVars;
 }
-#endif
 
 static inline bool IsWhiteSpace(TCHAR Value) { return Value == TCHAR(' '); }
 
@@ -161,16 +157,6 @@ public:
 		return this;
 	}
 	
-	virtual void LogHistory(FOutputDevice& Ar)
-	{
-	
-	}
-
-	virtual SIZE_T GetHistorySize()
-	{
-		return 0;
-	}
-
 	/** Legacy funciton to add old single delegates to the new multicast delegate. */
 	virtual void SetOnChangedCallback(const FConsoleVariableDelegate& Callback) 
 	{
@@ -245,6 +231,25 @@ public:
 	}
 
 	
+	// ------
+	// Helper accessors to get to FConsoleVariableExtendedData, when we don't have a Type
+
+	/**
+	 * Print the history to a log
+	 */
+	virtual void LogHistory(FOutputDevice& Ar) = 0;
+	/**
+	 * Track memory used by history data
+	 */
+	virtual SIZE_T GetHistorySize() = 0;
+
+#if ALLOW_OTHER_PLATFORM_CONFIG
+	/**
+	 * Internal function for caching data for other platforms - this is like Set(), but won't do any callbacks, or log any errors/warnings
+	 */
+	virtual void SetOtherPlatformValue(const TCHAR* InValue, EConsoleVariableFlags SetBy, FName Tag) = 0;
+#endif
+
 protected: // -----------------------------------------
 
 	// not using TCHAR* to allow chars support reloading of modules (otherwise we would keep a pointer into the module)
@@ -619,33 +624,6 @@ bool IConsoleManager::VisitPlatformCVarsForEmulation(FName PlatformName, const F
 }
 
 
-static bool GetConfigValueFromRuntimeSources(FName PlatformName, IConsoleVariable* CVar, FString& OutValue, EConsoleVariableFlags& LastSetBy)
-{
-	FString VariableName = IConsoleManager::Get().FindConsoleObjectName(CVar);
-	if (VariableName.Len() == 0)
-	{
-		return false;
-	}
-
-	OutValue = CVar->GetDefaultValueVariable()->GetString();
-	LastSetBy = ECVF_SetByConstructor;
-
-	// use the platform's base DeviceProfile for emulation
-	IConsoleManager::VisitPlatformCVarsForEmulation(PlatformName, PlatformName.ToString(),
-		[&VariableName, &OutValue, &LastSetBy](const FString& CVarName, const FString& CVarValue, EConsoleVariableFlags SetByAndPreview)
-		{
-			// if this key is the variable, set it at the current SetBy level
-			if (CVarName == VariableName)
-			{
-				OutValue = CVarValue;
-				LastSetBy = (EConsoleVariableFlags)(SetByAndPreview & ECVF_SetByMask);
-			}
-		});
-
-	// by this point, the value will be set, with the DefaultValue if nothing else from the visit function
-	return true;
-}
-
 #endif
 
 
@@ -655,7 +633,8 @@ constexpr bool IsArrayPriority(EConsoleVariableFlags Priority)
 	return
 		Priority == ECVF_SetByPluginLowPriority ||
 		Priority == ECVF_SetByPluginHighPriority ||
-		Priority == ECVF_SetByHotfix;
+		Priority == ECVF_SetByHotfix ||
+		Priority == ECVF_SetByPreview;
 }
 
 template <class T>
@@ -765,17 +744,17 @@ public:
 };
 
 
-// an intermediate class between specific typed CVars and FConsoleVariableBase to handle looking up other platform's cvars and caching them
-// if ALLOW_OTHER_PLATFORM_CONFIG is 0, then this is a pass-through class that does nothing
+// an intermediate class between specific typed CVars and FConsoleVariableBase to handle history and (in some configurations)
+// cached values of the CVar on other platforms/device profiles. It is expected that all CVar classes extend from this (if not
+// you will get abtract class compilation errors). We need the store type T, so we cannot put this up into FConsoleVariableBase
 template <class T>
-class FOtherPlatformValueHelper : public FConsoleVariableBase
+class FConsoleVariableExtendedData : public FConsoleVariableBase
 {
 public:
 
-	FOtherPlatformValueHelper(const T& DefaultValue, const TCHAR* Help, EConsoleVariableFlags Flags, bool bSaveDefault);
+	FConsoleVariableExtendedData(const T& DefaultValue, const TCHAR* Help, EConsoleVariableFlags Flags);
 
-#if UE_ALLOW_CVAR_HISTORY
-	virtual ~FOtherPlatformValueHelper()
+	virtual ~FConsoleVariableExtendedData()
 	{
 		// if we had been put into any tagged
 		if (PriorityHistory != nullptr && PriorityHistory->bHasTaggedArrayData)
@@ -789,41 +768,90 @@ public:
 			delete PriorityHistory;
 		}
 	}
-#endif
 	
-#if ALLOW_OTHER_PLATFORM_CONFIG
-	// remember the default value of this cvar, before anything else can assign to it - this way we can know in the editor what the value
-	// would be for a cvar on another platform, if no .ini file on that platform sets it
-	virtual IConsoleVariable* GetDefaultValueVariable()
+	/**
+	 * Subclasses implement this to simply set TypedValue as the current value for the CVar. It must not call any callbacks
+	 */
+	virtual void SetInternal(const T& TypedValue, EConsoleVariableFlags SetBy) = 0;
+
+	/**
+	 * Similar to the Set function, except that it must always work, independent of priority. This is used by the History system
+	 * to set the value wth a lower priority, and then update the priority.
+	 */
+	void SetInternalAndUpdateState(const T& TypedValue, EConsoleVariableFlags SetBy)
 	{
-		return PlatformIndependentDefault.Get();
+		SetInternal(TypedValue, SetBy);
+		OnChanged(SetBy, true);
 	}
 
-	virtual TSharedPtr<IConsoleVariable> GetPlatformValueVariable(FName PlatformName) override;
-	virtual void ClearPlatformVariables(FName PlatformName) override;
+#if ALLOW_OTHER_PLATFORM_CONFIG
 
-protected:
-	// cache of the default value - we need to remember this when calculating the value on another platform if that platform doesn't override it in any ini file
-	// even if this platform does (in which case the default is lost)
-	TSharedPtr<IConsoleVariable> PlatformIndependentDefault;
+	virtual TSharedPtr<IConsoleVariable> GetPlatformValueVariable(FName PlatformName, const FString& DeviceProfileName) override;
+	virtual bool HasPlatformValueVariable(FName PlatformName, const FString& DeviceProfileName) override;
+	virtual void ClearPlatformVariables(FName PlatformName) override;
 
 	// cache of the values of this cvar on other platforms
 	TMap<FName, TSharedPtr<IConsoleVariable> > PlatformValues;
 	FRWLock PlatformValuesLock;
+	friend FConsoleManager;
 
 #endif
 
-	virtual void SetInternal(const T& TypedValue, EConsoleVariableFlags SetBy) = 0;
+	virtual FString GetDefaultValue() override
+	{
+		// simply convert the default typed value to a string
+		return TTypeToString<T>::ToString(GetDefaultTypedValue());
+	}
 
+	T GetDefaultTypedValue()
+	{
+		// pull our constructed value out of the history if it exists, (if it doesn't, than our current value is the constructed value!)
+		if (PriorityHistory != nullptr)
+		{
+			// constructor will never have more than one value
+			return PriorityHistory->History[ECVF_SetByConstructor][0].Value.GetValueOnAnyThread(true);
+		}
+		
+		// if we have no history at all, that means we never called Set, so the current value must be the value we were constructed with
+		T ConstructorValue;
+		GetValue(ConstructorValue);
+		return ConstructorValue;
+	}
 
-#if UE_ALLOW_CVAR_HISTORY
-	friend class FConsoleManager;
+protected:
+	// A history object will be created and set here the first time Set() is called to change the value
 	FConsoleVariableHistory<T>* PriorityHistory = nullptr;
+	friend class FConsoleManager;
+
+#if ALLOW_OTHER_PLATFORM_CONFIG
+	virtual void SetOtherPlatformValue(const TCHAR* InValue, EConsoleVariableFlags SetBy, FName Tag) override
+	{
+		// always track it
+		TrackHistory(InValue, SetBy, Tag);
+
+		// set it if we are equal to or higher than before
+		uint32 CurrentSetBy = GetFlags() & ECVF_SetByMask;
+		uint32 NewSetBy = SetBy & ECVF_SetByMask;
+
+		if (NewSetBy >= CurrentSetBy)
+		{
+			// update value
+			T ConvertedValue;
+			TTypeFromString<T>::FromString(ConvertedValue, InValue);
+			SetInternal(ConvertedValue, SetBy);
+			
+			// update the setby
+			SetFlags((EConsoleVariableFlags)((GetFlags() & ~ECVF_SetByMask) | NewSetBy));
+		}
+	}
 #endif
-	
+
+	/**
+	 * This is the key function that subclasses need to call in their Set() implementation. This will track the values of cvars at
+	 * different priorities, so priorities/plugins/etc can be unset later, and the CVar will updatae state correctly (in Unset())
+	 */
 	void TrackHistory(const TCHAR* InValue, EConsoleVariableFlags SetBy, FName Tag)
 	{
-#if UE_ALLOW_CVAR_HISTORY
 		// make a history if we want to
 		if (PriorityHistory == nullptr)
 		{
@@ -845,13 +873,8 @@ protected:
 			// set a flag to remember we need to remove ourself from TaggedCVars in destructor
 			PriorityHistory->bHasTaggedArrayData = true;
 		}
-#endif
 	}
 	
-
-public:
-	
-#if UE_ALLOW_CVAR_HISTORY
 	virtual SIZE_T GetHistorySize() override
 	{
 		if (PriorityHistory != nullptr)
@@ -868,12 +891,13 @@ public:
 			PriorityHistory->Log(Ar);
 		}
 	}
-#endif
-
 	
+	/**
+	 * Removes the value at the given SetBy (and potentially Tag for the Array type SetBy priorities). This will update the
+	 * current value of the CVar as needed (if the Unset is at the current prio)
+	 */
 	virtual void Unset(EConsoleVariableFlags SetBy, FName Tag) override
 	{
-#if UE_ALLOW_CVAR_HISTORY
 		
 		if (PriorityHistory == nullptr)
 		{
@@ -891,11 +915,11 @@ public:
 			return;
 		}
 		
+		PriorityHistory->Unset(SetBy, Tag);
+		
 		uint32 CurrentPri =	(uint32)this->Flags & ECVF_SetByMask;
 		uint32 UnsetPri =	(uint32)SetBy & ECVF_SetByMask;
 
-		PriorityHistory->Unset(SetBy, Tag);
-		
 		// if we are unsetting at the current setby (or maybe in some weird cases, greater than setby) then we need to reset the SetBy and current value
 		if (UnsetPri >= CurrentPri)
 		{
@@ -903,10 +927,9 @@ public:
 			EConsoleVariableFlags NewSetBy;
 			auto MaxValue = PriorityHistory->GetMaxValue(NewSetBy);
 			
-			// and force it to the new value
-			SetInternal(MaxValue.GetValueOnGameThread(), NewSetBy);
+			// and force it to the new value and call any set callbacks
+			SetInternalAndUpdateState(MaxValue.GetValueOnGameThread(), NewSetBy);
 		}
-#endif
 	}
 
 };
@@ -1002,7 +1025,7 @@ template<> FString FConsoleVariableConversionHelper<FString>::GetString(FString
 
 // T: bool, int32, float, FString
 template <class T>
-class FConsoleVariable : public FOtherPlatformValueHelper<T>
+class FConsoleVariable : public FConsoleVariableExtendedData<T>
 {
 // help find functions without needing this-> prefixes
 	using FConsoleVariableBase::GetShadowIndex;
@@ -1010,8 +1033,8 @@ class FConsoleVariable : public FOtherPlatformValueHelper<T>
 	using FConsoleVariableBase::Flags;
 
 public:
-	FConsoleVariable(T DefaultValue, const TCHAR* Help, EConsoleVariableFlags Flags, bool bSaveDefault=true) 
-		: FOtherPlatformValueHelper<T>(DefaultValue, Help, Flags, bSaveDefault)
+	FConsoleVariable(T DefaultValue, const TCHAR* Help, EConsoleVariableFlags Flags)
+		: FConsoleVariableExtendedData<T>(DefaultValue, Help, Flags)
 		, Data(DefaultValue)
 	{
 	}
@@ -1063,7 +1086,6 @@ private: // ----------------------------------------------------
 	virtual void SetInternal(const T& TypedValue, EConsoleVariableFlags SetBy)
 	{
 		Data.ShadowedValue[0] = TypedValue;
-		OnChanged(SetBy, true);
 	}
 
 	virtual void OnChanged(EConsoleVariableFlags SetBy, bool bForce=false)
@@ -1136,51 +1158,101 @@ template<> TConsoleVariableData<FString>* FConsoleVariable<FString>::AsVariableS
 
 
 template<class T>
-FOtherPlatformValueHelper<T>::FOtherPlatformValueHelper(const T& DefaultValue, const TCHAR* Help, EConsoleVariableFlags Flags, bool bSaveDefault)
+FConsoleVariableExtendedData<T>::FConsoleVariableExtendedData(const T& DefaultValue, const TCHAR* Help, EConsoleVariableFlags Flags)
 	: FConsoleVariableBase(Help, Flags)
 {
+
+}
+
 #if ALLOW_OTHER_PLATFORM_CONFIG
-	if (bSaveDefault)
+
+/**
+ * Helper for FindOrCreatePlatformCVar that will either find the other-platform FConsoleVariable, or create an appropriate one and put it into the
+ * cached set of other-platform CVars inside the passed in CVar
+ */
+template<typename T>
+FConsoleVariable<T>* FindOrCreateTypedPlatformCVar(FConsoleVariableExtendedData<T>* CVar, FName PlatformKey)
+{
+	FRWScopeLock Lock(CVar->PlatformValuesLock, SLT_Write);
+	TSharedPtr<IConsoleVariable> PlatformCVar = CVar->PlatformValues.FindRef(PlatformKey);
+	if (!PlatformCVar.IsValid())
 	{
-		PlatformIndependentDefault = TSharedPtr<IConsoleVariable>(new FConsoleVariable<T>(DefaultValue, TEXT(""), (EConsoleVariableFlags)(Flags | ECVF_ReadOnly), false));
+		PlatformCVar = TSharedPtr<IConsoleVariable>(new FConsoleVariable(CVar->GetDefaultTypedValue(), TEXT("Platform CVar copy"), CVar->GetFlags()));
+		
+		// cache it
+		CVar->PlatformValues.Add(PlatformKey, PlatformCVar);
 	}
-#endif
+
+	return static_cast<FConsoleVariable<T>*>(PlatformCVar.Get());
 }
 
-#if ALLOW_OTHER_PLATFORM_CONFIG
-template<class T>
-TSharedPtr<IConsoleVariable> FOtherPlatformValueHelper<T>::GetPlatformValueVariable(FName PlatformName)
+/**
+ * Will find an existing other-platform FConsoleVariable, or create one based on the type of variable that was passed in
+ */
+FConsoleVariableBase* FindOrCreatePlatformCVar(IConsoleVariable* CVar, FName PlatformKey)
 {
-	// cheap lock here, contention is very rare
-	FRWScopeLock Lock(PlatformValuesLock, SLT_Write);
-	if (!PlatformValues.Contains(PlatformName))
-	{
-		FString ConfigValue;
-		EConsoleVariableFlags LastSetBy;
-		if (!GetConfigValueFromRuntimeSources(PlatformName, this, ConfigValue, LastSetBy))
-		{
-			return nullptr;
-		}
+	if (CVar->IsVariableBool())		return FindOrCreateTypedPlatformCVar(static_cast<FConsoleVariableExtendedData<bool>*>(CVar), PlatformKey);
+	if (CVar->IsVariableInt())		return FindOrCreateTypedPlatformCVar(static_cast<FConsoleVariableExtendedData<int32>*>(CVar), PlatformKey);
+	if (CVar->IsVariableFloat())	return FindOrCreateTypedPlatformCVar(static_cast<FConsoleVariableExtendedData<float>*>(CVar), PlatformKey);
+	if (CVar->IsVariableString())	return FindOrCreateTypedPlatformCVar(static_cast<FConsoleVariableExtendedData<FString>*>(CVar), PlatformKey);
 
-		T TypedValue;
-		TTypeFromString<T>::FromString(TypedValue, UE::ConfigUtilities::ConvertValueFromHumanFriendlyValue(*ConfigValue));
+	unimplemented();
+	return nullptr;
+}
 
-		// clear the existing setby mask
-		int NewFlags = GetFlags() & ~ECVF_SetFlagMask;
-		// add in new LastSetBy and make it readonly
-		NewFlags |= LastSetBy | ECVF_ReadOnly;
+static const FString GSpecialDPNameForPremadePlatformKey(TEXT("/"));
+/**
+ * Helper function to make a single key used in the PlatformValues set, that combines Platform and a DP name. If the DP name is empty, this will use
+ * Platform name as the DP name. This matches up with how GetPlatformValueVariable() was generally used in the past - get the CVar using the platform-named
+ * DeviceProfile. 
+ * If DeviceProfileName is the GSpecialDPNameForPremadePlatformKey, then that indicates PlatformName is already a PlatformKey that
+ * was created with this function before, so just return the PlatformName untouched. This makes is so we don't need two versions of functions
+ * that take a PlatformName and DeviceProfileName, and we don't need to keep re-creating PlatformKeys
+ */
+static FName MakePlatformKey(FName PlatformName, const FString& DeviceProfileName)
+{
+	// a bit of a hack to say PlatformName is already a Key
+	if (DeviceProfileName == GSpecialDPNameForPremadePlatformKey || PlatformName == NAME_None)
+	{
+		return PlatformName;
+	}
+	return *(PlatformName.ToString() + TEXT("/") + (DeviceProfileName.Len() ? DeviceProfileName : PlatformName.ToString()));
+}
 
-		// make a new cvar to hold the value from the other platform
-		TSharedPtr<IConsoleVariable> PlatformCVar = TSharedPtr<IConsoleVariable>(new FConsoleVariable<T>(TypedValue, TEXT(""), (EConsoleVariableFlags)NewFlags, false));
+template<class T>
+bool FConsoleVariableExtendedData<T>::HasPlatformValueVariable(FName PlatformName, const FString& DeviceProfileName)
+{
+	// cheap lock here, contention is very rare
+	FRWScopeLock Lock(PlatformValuesLock, SLT_Write);
+	return PlatformValues.Contains(MakePlatformKey(PlatformName, DeviceProfileName));
+}
 
-		PlatformValues.Add(PlatformName, PlatformCVar);
+template<class T>
+TSharedPtr<IConsoleVariable> FConsoleVariableExtendedData<T>::GetPlatformValueVariable(FName PlatformName, const FString& DeviceProfileName)
+{
+	// if we have GSpecialDPNameForPremadePlatformKey passed in, we have already gone through the
+	// Load and we have a premade key in PlatformName
+	if (DeviceProfileName != GSpecialDPNameForPremadePlatformKey)
+	{
+		// make sure we have cached this platform/DP
+		IConsoleManager::Get().LoadAllPlatformCVars(PlatformName, DeviceProfileName);
 	}
 
-	return PlatformValues.FindRef(PlatformName);
+	// we have assumed in the past that we would return at least the constructor version, so create one if we are explicitly asking
+	// this can happen when .ini files don't give a value to a cvar, but we are asking for a platform's value anyway - in
+	// which case we want the constructor value, not the current platform's value
+	if (!HasPlatformValueVariable(PlatformName, DeviceProfileName))
+	{
+		FindOrCreatePlatformCVar(this, MakePlatformKey(PlatformName, DeviceProfileName));
+	}
+	
+	// cheap lock here, contention is very rare
+	FRWScopeLock Lock(PlatformValuesLock, SLT_Write);
+	return PlatformValues.FindRef(MakePlatformKey(PlatformName, DeviceProfileName));
 }
 
 template<class T>
-void FOtherPlatformValueHelper<T>::ClearPlatformVariables(FName PlatformName)
+void FConsoleVariableExtendedData<T>::ClearPlatformVariables(FName PlatformName)
 {
 	FRWScopeLock Lock(PlatformValuesLock, SLT_Write);
 
@@ -1202,7 +1274,7 @@ void FOtherPlatformValueHelper<T>::ClearPlatformVariables(FName PlatformName)
 
 // T: int32, float, bool
 template <class T>
-class FConsoleVariableRef : public FOtherPlatformValueHelper<T>
+class FConsoleVariableRef : public FConsoleVariableExtendedData<T>
 {
 	// help find functions without needing this-> prefixes
 	using FConsoleVariableBase::GetShadowIndex;
@@ -1211,7 +1283,7 @@ class FConsoleVariableRef : public FOtherPlatformValueHelper<T>
 
 public:
 	FConsoleVariableRef(T& InRefValue, const TCHAR* Help, EConsoleVariableFlags Flags) 
-		: FOtherPlatformValueHelper<T>(InRefValue, Help, Flags, true)
+		: FConsoleVariableExtendedData<T>(InRefValue, Help, Flags)
 		, RefValue(InRefValue)
 		, MainValue(InRefValue)
 	{
@@ -1262,7 +1334,6 @@ private: // ----------------------------------------------------
 	virtual void SetInternal(const T& TypedValue, EConsoleVariableFlags SetBy)
 	{
 		MainValue = TypedValue;
-		OnChanged(SetBy, true);
 	}
 
 	void OnChanged(EConsoleVariableFlags SetBy, bool bForce=false)
@@ -1297,11 +1368,11 @@ bool FConsoleVariableRef<float>::IsVariableFloat() const
 
 // string version
 
-class FConsoleVariableStringRef : public FOtherPlatformValueHelper<FString>
+class FConsoleVariableStringRef : public FConsoleVariableExtendedData<FString>
 {
 public:
 	FConsoleVariableStringRef(FString& InRefValue, const TCHAR* Help, EConsoleVariableFlags Flags)
-		: FOtherPlatformValueHelper<FString>(FString(), Help, Flags, false)
+		: FConsoleVariableExtendedData<FString>(FString(), Help, Flags)
 		, RefValue(InRefValue)
 		, MainValue(InRefValue)
 	{
@@ -1367,7 +1438,6 @@ private: // ----------------------------------------------------
 	virtual void SetInternal(const FString& TypedValue, EConsoleVariableFlags SetBy)
 	{
 		MainValue = TypedValue;
-		OnChanged(SetBy);
 	}
 	
 	void OnChanged(EConsoleVariableFlags SetBy, bool bForce=false)
@@ -1381,11 +1451,11 @@ private: // ----------------------------------------------------
 	}
 };
 
-class FConsoleVariableBitRef : public FOtherPlatformValueHelper<int>
+class FConsoleVariableBitRef : public FConsoleVariableExtendedData<int>
 {
 public:
 	FConsoleVariableBitRef(const TCHAR* FlagName, uint32 InBitNumber, uint8* InForce0MaskPtr, uint8* InForce1MaskPtr, const TCHAR* Help, EConsoleVariableFlags Flags) 
-		: FOtherPlatformValueHelper<int>(0, Help, Flags, false), Force0MaskPtr(InForce0MaskPtr), Force1MaskPtr(InForce1MaskPtr), BitNumber(InBitNumber)
+		: FConsoleVariableExtendedData<int>(0, Help, Flags), Force0MaskPtr(InForce0MaskPtr), Force1MaskPtr(InForce1MaskPtr), BitNumber(InBitNumber)
 	{
 	}
 
@@ -1450,7 +1520,6 @@ private: // ----------------------------------------------------
 	{
 		FMath::SetBoolInBitField(Force0MaskPtr, BitNumber, TypedValue == 0);
 		FMath::SetBoolInBitField(Force1MaskPtr, BitNumber, TypedValue == 1);
-		OnChanged(SetBy, true);
 	}
 	
 };
@@ -2362,8 +2431,9 @@ bool FConsoleManager::ProcessUserConsoleInput(const TCHAR* InInput, FOutputDevic
 		Param1.MidInline(0, Param1.Len() - 1, false);
 	}
 
-	// look for the <cvar>@<platform> syntax
+	// look for the <cvar>@<platform[/deviceprofile]> syntax
 	FName PlatformName;
+	FString DeviceProfileName;
 	if (Param1.Contains(TEXT("@")))
 	{
 		FString Left, Right;
@@ -2372,7 +2442,16 @@ bool FConsoleManager::ProcessUserConsoleInput(const TCHAR* InInput, FOutputDevic
 		if (Left.Len() && Right.Len())
 		{
 			Param1 = Left;
-			PlatformName = *Right;
+			if (Right.Contains(TEXT("/")))
+			{
+				FString Plat;
+				Right.Split(TEXT("/"), &Plat, &DeviceProfileName);
+				PlatformName = *Plat;
+			}
+			else
+			{
+				PlatformName = *Right;
+			}
 		}
 	}
 
@@ -2410,7 +2489,7 @@ bool FConsoleManager::ProcessUserConsoleInput(const TCHAR* InInput, FOutputDevic
 		else
 		{
 #if ALLOW_OTHER_PLATFORM_CONFIG
-			PlatformCVar = CVar->GetPlatformValueVariable(PlatformName);
+			PlatformCVar = CVar->GetPlatformValueVariable(PlatformName, DeviceProfileName);
 			CVar = PlatformCVar.Get();
 			if (!CVar)
 			{
@@ -2837,9 +2916,8 @@ FConsoleObjectWithNameMulticastDelegate& FConsoleManager::OnConsoleObjectUnregis
 	return ConsoleObjectUnregisteredDelegate;
 }
 
-void FConsoleManager::UnsetAllConsoleVariablesWithTag(FName Tag)
+void FConsoleManager::UnsetAllConsoleVariablesWithTag(FName Tag, EConsoleVariableFlags Priority)
 {
-#if UE_ALLOW_CVAR_HISTORY
 	TSet<IConsoleVariable*>* TaggedSet = UE::ConsoleManager::Private::TaggedCVars.FindRef(Tag);
 	if (TaggedSet == nullptr)
 	{
@@ -2848,13 +2926,92 @@ void FConsoleManager::UnsetAllConsoleVariablesWithTag(FName Tag)
 	
 	for (IConsoleVariable* Var : *TaggedSet)
 	{
-		Var->Unset(EConsoleVariableFlags::ECVF_SetByMask);
+		Var->Unset(Priority);
 	}
 	
 	UE::ConsoleManager::Private::TaggedCVars.Remove(Tag);
-#endif
 }
 
+#if ALLOW_OTHER_PLATFORM_CONFIG
+
+void FConsoleManager::LoadAllPlatformCVars(FName PlatformName, const FString& DeviceProfileName)
+{
+	FName PlatformKey = MakePlatformKey(PlatformName, DeviceProfileName);
+	
+	if (CachedPlatformsAndDeviceProfiles.Contains(PlatformKey))
+	{
+		return;
+	}
+	CachedPlatformsAndDeviceProfiles.Add(PlatformKey);
+	
+	// use the platform's base DeviceProfile for emulation
+	VisitPlatformCVarsForEmulation(PlatformName, DeviceProfileName,
+		[PlatformKey](const FString& CVarName, const FString& CVarValue, EConsoleVariableFlags SetByAndPreview)
+	{
+		// make sure the named cvar exists
+		IConsoleVariable* CVar = IConsoleManager::Get().FindConsoleVariable(*CVarName);
+		if (CVar == nullptr)
+		{
+			return;
+		}
+		
+		// find or make the cvar for this platformkey
+		FConsoleVariableBase* PlatformCVar = FindOrCreatePlatformCVar(CVar, PlatformKey);
+
+		// now cache the passed in value
+		int32 SetBy = SetByAndPreview & ECVF_SetByMask;
+		PlatformCVar->SetOtherPlatformValue(*CVarValue, (EConsoleVariableFlags)SetBy, NAME_None);
+	});
+}
+
+void FConsoleManager::PreviewPlatformCVars(FName PlatformName, const FString& DeviceProfileName, FName PreviewModeTag)
+{
+	UE_LOG(LogConsoleManager, Display, TEXT("Previewing Platform '%s', DeviceProfile '%s', ModeTag '%s'"), *PlatformName.ToString(), *DeviceProfileName, *PreviewModeTag.ToString());
+	
+	LoadAllPlatformCVars(PlatformName, DeviceProfileName.Len() ? DeviceProfileName : PlatformName.ToString());
+	
+	FName PlatformKey = MakePlatformKey(PlatformName, DeviceProfileName);
+
+	for (auto Pair : ConsoleObjects)
+	{
+		if (IConsoleVariable* CVar = Pair.Value->AsVariable())
+		{
+			// we want Preview but not Scalability or Cheat
+			if ((CVar->GetFlags() & (ECVF_Preview | ECVF_ScalabilityGroup | ECVF_Cheat)) == ECVF_Preview)
+			{
+				// if we have a value for the platform, then set it in the real CVar
+				if (CVar->HasPlatformValueVariable(PlatformKey, TEXT("/")))
+				{
+					TSharedPtr<IConsoleVariable> PlatformCVar = CVar->GetPlatformValueVariable(PlatformKey, TEXT("/"));
+					CVar->Set(*PlatformCVar->GetString(), ECVF_SetByPreview, PreviewModeTag);
+					
+					UE_LOG(LogConsoleManager, Display, TEXT("  |-> %s = %s"), *Pair.Key, *PlatformCVar->GetString());
+				}
+			}
+		}
+	}
+}
+
+void FConsoleManager::ClearAllPlatformCVars(FName PlatformName, const FString& DeviceProfileName)
+{
+	FName PlatformKey = MakePlatformKey(PlatformName, DeviceProfileName);
+	if (!CachedPlatformsAndDeviceProfiles.Contains(PlatformKey))
+	{
+		return;
+	}
+	CachedPlatformsAndDeviceProfiles.Remove(PlatformKey);
+	
+	for (auto Pair : ConsoleObjects)
+	{
+		if (IConsoleVariable* CVar = Pair.Value->AsVariable())
+		{
+			// clear any cached values for this key
+			CVar->ClearPlatformVariables(PlatformKey);
+		}
+	}
+}
+
+#endif
 
 #if !(UE_BUILD_SHIPPING || UE_BUILD_TEST)
 uint32 GConsoleManagerSinkTestCounter = 0;
diff --git a/Engine/Source/Runtime/Core/Public/HAL/ConsoleManager.h b/Engine/Source/Runtime/Core/Public/HAL/ConsoleManager.h
--- a/Engine/Source/Runtime/Core/Public/HAL/ConsoleManager.h
+++ b/Engine/Source/Runtime/Core/Public/HAL/ConsoleManager.h
@@ -87,8 +87,13 @@ public:
 	CORE_API virtual bool IsNameRegistered(const TCHAR* Name) const override;	
 	CORE_API virtual void RegisterThreadPropagation(uint32 ThreadId, IConsoleThreadPropagation* InCallback) override;
 	CORE_API virtual void UnregisterConsoleObject( IConsoleObject* Object, bool bKeepState) override;
-	CORE_API virtual void UnsetAllConsoleVariablesWithTag(FName Tag) override;
-
+	CORE_API virtual void UnsetAllConsoleVariablesWithTag(FName Tag, EConsoleVariableFlags Priority) override;
+#if ALLOW_OTHER_PLATFORM_CONFIG
+	CORE_API virtual void LoadAllPlatformCVars(FName PlatformName, const FString& DeviceProfileName=FString()) override;
+	CORE_API virtual void ClearAllPlatformCVars(FName PlatformName=NAME_None, const FString& DeviceProfileName=FString()) override;
+	CORE_API virtual void PreviewPlatformCVars(FName PlatformName, const FString& DeviceProfileName, FName PreviewModeTag) override;
+#endif
+	
 private: // ----------------------------------------------------
 
 	/** Map of console variables and commands, indexed by the name of that command or variable */
@@ -103,6 +108,8 @@ private: // ----------------------------------------------------
 	FConsoleObjectWithNameMulticastDelegate ConsoleObjectUnregisteredDelegate;
 
 	IConsoleThreadPropagation* ThreadPropagationCallback;
+	
+	TSet<FName> CachedPlatformsAndDeviceProfiles;
 
 	// if true the next call to CallAllConsoleVariableSinks() we will call all registered sinks
 	bool bCallAllConsoleVariableSinks;
diff --git a/Engine/Source/Runtime/Core/Public/HAL/IConsoleManager.h b/Engine/Source/Runtime/Core/Public/HAL/IConsoleManager.h
--- a/Engine/Source/Runtime/Core/Public/HAL/IConsoleManager.h
+++ b/Engine/Source/Runtime/Core/Public/HAL/IConsoleManager.h
@@ -154,12 +154,14 @@ enum EConsoleVariableFlags
 	ECVF_SetByConsoleVariablesIni = 0x09000000,
 	// Set by hotfix. THIS IS AN ARRAY TYPE meaning multiple vales with this SetBy are stored in the history
 	ECVF_SetByHotfix =				0x0A000000,
+	// Set for previewing in editor. THIS IS AN ARRAY TYPE meaning multiple vales with this SetBy are stored in the history
+	ECVF_SetByPreview =				0x0B000000,
 	// Used by some command line parameters, others use the Console priority instead
-	ECVF_SetByCommandline =			0x0B000000,
+	ECVF_SetByCommandline =			0x0C000000,
 	// Used for high priority temporary debugging or operation modes
-	ECVF_SetByCode =				0x0C000000,
+	ECVF_SetByCode =				0x0D000000,
 	// Highest priority used via editor UI or or game/editor interactive console
-	ECVF_SetByConsole =				0x0D000000,
+	ECVF_SetByConsole =				0x0E000000,
 
 
 	// ------------------------------------------------
@@ -177,6 +179,7 @@ enum EConsoleVariableFlags
 	op(PluginHighPriority) \
 	op(ConsoleVariablesIni) \
 	op(Hotfix) \
+	op(Preview) \
 	op(Commandline) \
 	op(Code) \
 	op(Console)
@@ -535,22 +538,29 @@ public:
 	/**
 	 * Get the saved off default value, in a cvar variable, if one was created
 	 */
+	UE_DEPRECATED(5.4, "Use GetDefaultValue() instead")
 	virtual IConsoleVariable* GetDefaultValueVariable()
 	{
 		return nullptr;
 	}
 
+	/**
+	 * Get the value this CVar was constructed with
+	 */
+	virtual FString GetDefaultValue() = 0;
+
 #if ALLOW_OTHER_PLATFORM_CONFIG
 
 	/**
-	 * Get a CVar opject that matches this cvar, but contains the value of the platform given.
+	 * Get a CVar opject that matches this cvar, but contains the value of the platform given. This will trigger a lof of all cvars for this platform/DP if it doesn't exist
+	 * Note: If this causes a compile error due to pure virtual, make sure your IConsoleVariable subclass inherits from FConsoleVariableExtendedData
 	 */
-	virtual TSharedPtr<IConsoleVariable> GetPlatformValueVariable(FName PlatformName)
-	{
-		// this could be called for some special subclass like FConsoleVariableBitRef that don't implement this (yet)
-		unimplemented();
-		return TSharedPtr<IConsoleVariable>();
-	}
+	virtual TSharedPtr<IConsoleVariable> GetPlatformValueVariable(FName PlatformName, const FString& DeviceProfileName=FString()) = 0;
+
+	/**
+	 * Checks if the CVar has a cached value for the given platform
+	 */
+	virtual bool HasPlatformValueVariable(FName PlatformName, const FString& DeviceProfileName=FString()) = 0;
 
 	/**
 	 * Used only for debugging/iterating, this will clear all of the other platform's cvar objects, which will 
@@ -559,6 +569,7 @@ public:
 	virtual void ClearPlatformVariables(FName PlatformName=NAME_None)
 	{
 	}
+	
 #endif
 
 	// convenience methods
@@ -1072,22 +1083,41 @@ struct IConsoleManager
 #if ALLOW_OTHER_PLATFORM_CONFIG
 	/**
 	 * Walks over the best set of ini sections for another platform that can be used to emulate cvars as the platform
-	 * will have set - WITH THE NOTABLE EXCEPTION OF specific DeviceProfiles. Use the DeviceProfileManager code to get a DP
-	 * for a platform, as this code cannot easily access DeviceProfile logic.
-	 * It also won't include any UserSettings or ConsoleVariables.ini settings.
-	 * 
+	 * will have set
+	 * It also won't include any UserSettings
+	 *
 	 * @param PlatformName The platform name (the ini name, so Windows, not Win64)
 	 * @param DeviceProfileName If this is non-empty, the given deviceprofile will be loaded from the platform's inis and inserted into the CVars
 	 * @param Visit the callback to run for each CVar found
 	 */
 	static CORE_API bool VisitPlatformCVarsForEmulation(FName PlatformName, const FString& DeviceProfileName, TFunctionRef<void(const FString& CVarName, const FString& CVarValue, EConsoleVariableFlags SetBy)> Visit);
+	
+	/**
+	 * Loads all platform cvars, which can be retrieved with IConsoleVariable::GetPlatformValueVariable. If DeviceProfileName is empty, it will use the
+	 *  DevicePlatform named the same as the Platform
+	 */
+	virtual void LoadAllPlatformCVars(FName PlatformName, const FString& DeviceProfileName=FString()) = 0;
+	
+	/**
+	 * Applies the cvars from the platform/DeviceProfile pair (can be blank to use the platform's named DP) and sets into the SetByPreview  priority, with the PreviewModeTag
+	 * to be unset later
+	 */
+	virtual void PreviewPlatformCVars(FName PlatformName, const FString& DeviceProfileName, FName PreviewModeTag) = 0;
+
+	/**
+	 * Empties the cache for the given Platform/DP for all CVars. If using NAME_None for Platofrm, this will wipe every cached platform/DP value. If DeviceProfileName
+	 * is empty, it will use the PlatfomName as the DeviceProfile name
+	 */
+	virtual void ClearAllPlatformCVars(FName PlatformName=NAME_None, const FString& DeviceProfileName=FString()) = 0;
+
 #endif
 
 	/**
 	  * When a plugin is unmounted, it needs to unset cvars that it had set when it was mounted. This will unset and fixup all
 	  * variables with the given Tag
+	  * @param Priority If set, then the internal search for CVars is restricted to this SetBy. This is an optimization, only to be used if you _know_ that all CVars set with this tag were set at this priority. Any other priorites cannot be unset
 	 */
-	virtual void UnsetAllConsoleVariablesWithTag(FName Tag) = 0;
+	virtual void UnsetAllConsoleVariablesWithTag(FName Tag, EConsoleVariableFlags Priority=ECVF_SetByMask) = 0;
 	
 	virtual FConsoleVariableMulticastDelegate& OnCVarUnregistered() = 0;
 	virtual FConsoleObjectWithNameMulticastDelegate& OnConsoleObjectUnregistered() = 0;
diff --git a/Engine/Source/Runtime/Engine/Classes/DeviceProfiles/DeviceProfileManager.h b/Engine/Source/Runtime/Engine/Classes/DeviceProfiles/DeviceProfileManager.h
--- a/Engine/Source/Runtime/Engine/Classes/DeviceProfiles/DeviceProfileManager.h
+++ b/Engine/Source/Runtime/Engine/Classes/DeviceProfiles/DeviceProfileManager.h
@@ -25,7 +25,7 @@ DECLARE_MULTICAST_DELEGATE( FOnActiveDeviceProfileChanged );
 struct FPushedCVarSetting
 {
 	FPushedCVarSetting() : SetBy(ECVF_Default)	{}
-	FPushedCVarSetting(const FString& InValue, EConsoleVariableFlags InFlags) 
+	FPushedCVarSetting(const FString& InValue, EConsoleVariableFlags InFlags)
 		: Value(InValue)
 		, SetBy(EConsoleVariableFlags(InFlags & ECVF_SetByMask))
 	{}
@@ -126,12 +126,12 @@ public:
 	* It does not change the actual DP. 
 	* Use RestorePreviewDeviceProfile to revert.
 	*/
-	ENGINE_API void SetPreviewDeviceProfile(UDeviceProfile* DeviceProfile);
+	ENGINE_API void SetPreviewDeviceProfile(UDeviceProfile* DeviceProfile, FName PreviewModeTag=NAME_None);
 
 	/**
 	* Revert the preview state.
 	*/
-	ENGINE_API void RestorePreviewDeviceProfile();
+	ENGINE_API void RestorePreviewDeviceProfile(FName PreviewModeTag=NAME_None);
 #endif
 
 	/**
@@ -335,10 +335,6 @@ private:
 	// Used to undo the DP before applying new state.
 	static ENGINE_API TMap<FString, FPushedCVarSetting> PushedSettings;
 
-#if ALLOW_OTHER_PLATFORM_CONFIG
-	// Original values of the CVars modified by SetPreviewDeviceProfile.
-	TMap<FString, FPushedCVarSetting> PreviewPushedSettings;
-#endif
 	// Holds the device profile that has been overridden, null no override active.
 	UDeviceProfile* BaseDeviceProfile = nullptr;
 
diff --git a/Engine/Source/Runtime/Engine/Private/DeviceProfiles/DeviceProfileManager.cpp b/Engine/Source/Runtime/Engine/Private/DeviceProfiles/DeviceProfileManager.cpp
--- a/Engine/Source/Runtime/Engine/Private/DeviceProfiles/DeviceProfileManager.cpp
+++ b/Engine/Source/Runtime/Engine/Private/DeviceProfiles/DeviceProfileManager.cpp
@@ -1038,52 +1038,35 @@ void UDeviceProfileManager::SaveProfiles(bool bSaveToDefaults)
 }
 
 #if ALLOW_OTHER_PLATFORM_CONFIG
-void UDeviceProfileManager::SetPreviewDeviceProfile(UDeviceProfile* DeviceProfile)
+void UDeviceProfileManager::SetPreviewDeviceProfile(UDeviceProfile* DeviceProfile, FName PreviewModeTag)
 {
-	// we're applying a preview mode on top of an overridden DP?
-	check(BaseDeviceProfile == nullptr);
-
-	RestorePreviewDeviceProfile();
-
-	PreviewDeviceProfile = DeviceProfile;
-
-	UE_LOG(LogDeviceProfileManager, Log, TEXT("SetPreviewDeviceProfile preview to %s"), *DeviceProfile->GetName());
-	// apply the preview DP cvars.
-	for (const auto& Pair : DeviceProfile->GetAllPreviewCVars())
+	if (PreviewModeTag == NAME_None)
 	{
-		IConsoleVariable* CVar = IConsoleManager::Get().FindConsoleVariable(*Pair.Key);
-		// skip over scalability group cvars (maybe they shouldn't be left in the AllExpandedCVars?)
-		if (CVar != nullptr && !CVar->TestFlags(EConsoleVariableFlags::ECVF_ScalabilityGroup))
-		{
-			// remember the previous value and priority so we can restore
-			FPushedCVarSetting OldSetting = FPushedCVarSetting(CVar->GetString(), CVar->GetFlags());
-			PreviewPushedSettings.Add(Pair.Key, OldSetting);
-			//UE_LOG(LogDeviceProfileManager, Log, TEXT("Pushing Device Profile CVar: [[%s:%s]]"), *Pair.Key, *Pair.Value);
-			// cheat CVar can only be set in ConsoleVariables.ini
-			if (!CVar->TestFlags(EConsoleVariableFlags::ECVF_Cheat))
-			{
-				// set the cvar to the new value, with same priority that it was before (SetByMask means current priority)
-				CVar->SetWithCurrentPriority(UE::ConfigUtilities::ConvertValueFromHumanFriendlyValue(*Pair.Value));
-			}
-		}
+		PreviewModeTag = "UnknownPreviewMode";
 	}
+	
+	RestorePreviewDeviceProfile(PreviewModeTag);
+
+	PreviewDeviceProfile = DeviceProfile;
+	
+	// walk over all cvars and apply them
+	IConsoleManager::Get().PreviewPlatformCVars(*DeviceProfile->DeviceType, DeviceProfile->GetName(), PreviewModeTag);
 
 	// broadcast cvar sinks now that we are done
 	IConsoleManager::Get().CallAllConsoleVariableSinks();
 }
 
 
-void UDeviceProfileManager::RestorePreviewDeviceProfile()
+void UDeviceProfileManager::RestorePreviewDeviceProfile(FName PreviewModeTag)
 {
-	PreviewDeviceProfile = nullptr;
-	if (PreviewPushedSettings.Num())
+	if (PreviewModeTag == NAME_None)
 	{
-		checkf(BaseDeviceProfile == nullptr, TEXT("call to RestorePreviewDeviceProfile while both preview and BaseDeviceProfile has been set?"));
-
-		UE_LOG(LogDeviceProfileManager, Log, TEXT("Restoring Preview DP "));
-		// this sets us back to non-preview state.
-		RestorePushedState(PreviewPushedSettings);
+		PreviewModeTag = "UnknownPreviewMode";
 	}
+
+	PreviewDeviceProfile = nullptr;
+
+	IConsoleManager::Get().UnsetAllConsoleVariablesWithTag(PreviewModeTag, ECVF_SetByPreview);
 }
 #endif
 
@@ -1092,11 +1075,6 @@ void UDeviceProfileManager::RestorePreviewDeviceProfile()
 */
 void UDeviceProfileManager::SetOverrideDeviceProfile(UDeviceProfile* DeviceProfile)
 {
-#if ALLOW_OTHER_PLATFORM_CONFIG
-	// we have an active preview running but we're changing the actual device's DP too?
-	check(PreviewPushedSettings.Num() == 0);
-#endif
-
 	// If we're not already overriding record the BaseDeviceProfile
 	if(!BaseDeviceProfile)
 	{
@@ -1122,11 +1100,6 @@ void UDeviceProfileManager::SetOverrideDeviceProfile(UDeviceProfile* DeviceProfi
 */
 void UDeviceProfileManager::RestoreDefaultDeviceProfile()
 {
-#if ALLOW_OTHER_PLATFORM_CONFIG
-	// We're restoring overridden DP while a preview is active?
-	check(PreviewPushedSettings.Num() == 0);
-#endif
-
 	// have we been overridden?
 	if (BaseDeviceProfile)
 	{
@@ -1446,7 +1419,7 @@ static bool GetCVarForDeviceProfile( FOutputDevice& Ar, FString DPName, FString
 			return false;
 		}
 
-		Value = CVar->GetDefaultValueVariable()->GetString();
+		Value = CVar->GetDefaultValue();
 	}
 
 	Ar.Logf(TEXT("%s@%s = \"%s\""), *DPName, *CVarName, *Value);
