- Fixed a CVar GetValueOnGameThread to AnyThread so that asset validator works

diff --git a/Engine/Source/Runtime/Core/Private/HAL/ConsoleManager.cpp b/Engine/Source/Runtime/Core/Private/HAL/ConsoleManager.cpp
--- a/Engine/Source/Runtime/Core/Private/HAL/ConsoleManager.cpp
+++ b/Engine/Source/Runtime/Core/Private/HAL/ConsoleManager.cpp
@@ -659,7 +659,7 @@ public:
 			if (ValueArray.Num() > 0)
 			{
 				// don't do anything if we already are set
-				if (ValueArray[0].Key == Tag && ValueArray[0].Value.GetValueOnGameThread() == LocalCopy)
+				if (ValueArray[0].Key == Tag && ValueArray[0].Value.GetValueOnAnyThread() == LocalCopy)
 				{
 					return;
 				}
