[Core]
* Fixed non unity compile error for android

diff --git a/Engine/Source/Runtime/Core/Private/Android/AndroidMemoryTrace.cpp b/Engine/Source/Runtime/Core/Private/Android/AndroidMemoryTrace.cpp
--- a/Engine/Source/Runtime/Core/Private/Android/AndroidMemoryTrace.cpp
+++ b/Engine/Source/Runtime/Core/Private/Android/AndroidMemoryTrace.cpp
@@ -5,6 +5,7 @@
 #if UE_MEMORY_TRACE_ENABLED
 
 #include <android/log.h>
+#include <string.h>
 
 ////////////////////////////////////////////////////////////////////////////////
 FMalloc* MemoryTrace_CreateInternal(FMalloc*);
