Add GetTypeHash() overload for TOptional.

Required to have TOptional as key of TMap.

Implement by using the value type's GetTypeHash() overload when set (return 0 otherwise).

#rnx

diff --git a/Engine/Source/Runtime/Core/Public/Misc/Optional.h b/Engine/Source/Runtime/Core/Public/Misc/Optional.h
--- a/Engine/Source/Runtime/Core/Public/Misc/Optional.h
+++ b/Engine/Source/Runtime/Core/Public/Misc/Optional.h
@@ -417,6 +417,12 @@ FArchive& operator<<(FArchive& Ar, TOptional<OptionalType>& Optional)
 	return Ar;
 }
 
+template<typename OptionalType>
+inline auto GetTypeHash(const TOptional<OptionalType>& Optional) -> decltype(GetTypeHash(*Optional))
+{
+	return Optional.IsSet() ? GetTypeHash(*Optional) : 0;
+}
+
 /**
  * Trait which determines whether or not a type is a TOptional.
  */
