Added new static StringType::Construct* functions to replace the dubious set of string constructors whose behaviour aren't obvious.
Added deprecation option for the string constructors, which are enabled for FUtf8String and FAnsiString and will be used to incrementally replace usage for FString.

diff --git a/Engine/Source/Runtime/Core/Private/Containers/String.cpp.inl b/Engine/Source/Runtime/Core/Private/Containers/String.cpp.inl
--- a/Engine/Source/Runtime/Core/Private/Containers/String.cpp.inl
+++ b/Engine/Source/Runtime/Core/Private/Containers/String.cpp.inl
@@ -160,40 +160,57 @@ namespace UE::Core::Private
 	}
 
 	template <typename SrcCharType>
-	FORCEINLINE void ConstructFromCString(/* Out */ TArray<UE_STRING_CHARTYPE>& Data, const SrcCharType* Src)
+	FORCEINLINE void ConstructFromCString(/* Out */ TArray<UE_STRING_CHARTYPE>& Data, const SrcCharType* Src, int32 ExtraSlack)
 	{
-		if (Src && *Src)
+		checkf(Src,             TEXT("Unable to construct string from a null pointer"));
+		checkf(ExtraSlack >= 0, TEXT("Unable to construct string with negative slack"));
+
+		if (*Src)
 		{
 			int32 SrcLen  = TCString<SrcCharType>::Strlen(Src) + 1;
 			int32 DestLen = FPlatformString::ConvertedLength<UE_STRING_CHARTYPE>(Src, SrcLen);
-			Data.Reserve(DestLen);
+			Data.Reserve(DestLen + ExtraSlack);
 			Data.AddUninitialized(DestLen);
 
 			FPlatformString::Convert(Data.GetData(), DestLen, Src, SrcLen);
 		}
+		else if (ExtraSlack > 0)
+		{
+			Data.Reserve(ExtraSlack + 1);
+		}
 	}
 
 	template <typename SrcCharType>
-	FORCEINLINE void ConstructWithLength(/* Out */ TArray<UE_STRING_CHARTYPE>& Data, int32 InCount, const SrcCharType* InSrc)
+	FORCEINLINE void ConstructWithLength(/* Out */ TArray<UE_STRING_CHARTYPE>& Data, const SrcCharType* InSrc, int32 InCount, int32 ExtraSlack)
 	{
-		if (InSrc)
+		checkf(InSrc || InCount == 0, TEXT("Unable to construct string from a null pointer"));
+		checkf(InCount >= 0,          TEXT("Unable to construct string with a negative size"));
+		checkf(ExtraSlack >= 0,       TEXT("Unable to construct string with negative slack"));
+
+		int32 DestLen = FPlatformString::ConvertedLength<UE_STRING_CHARTYPE>(InSrc, InCount);
+		if (DestLen > 0)
 		{
-			int32 DestLen = FPlatformString::ConvertedLength<UE_STRING_CHARTYPE>(InSrc, InCount);
-			if (DestLen > 0 && *InSrc)
-			{
-				Data.Reserve(DestLen + 1);
-				Data.AddUninitialized(DestLen + 1);
+			Data.Reserve(DestLen + ExtraSlack + 1);
+			Data.AddUninitialized(DestLen + 1);
 
-				FPlatformString::Convert(Data.GetData(), DestLen, InSrc, InCount);
-				*(Data.GetData() + Data.Num() - 1) = CHARTEXT(UE_STRING_CHARTYPE, '\0');
-			}
+			UE_STRING_CHARTYPE* DataPtr = Data.GetData();
+
+			FPlatformString::Convert(DataPtr, DestLen, InSrc, InCount);
+			DataPtr[DestLen] = CHARTEXT(UE_STRING_CHARTYPE, '\0');
+		}
+		else if (ExtraSlack > 0)
+		{
+			Data.Reserve(ExtraSlack + 1);
 		}
 	}
 
 	template <typename SrcCharType>
 	FORCEINLINE void ConstructWithSlack(/* Out */ TArray<UE_STRING_CHARTYPE>& Data, const SrcCharType* Src, int32 ExtraSlack)
 	{
-		if (Src && *Src)
+		checkf(Src,             TEXT("Unable to construct string from a null pointer"));
+		checkf(ExtraSlack >= 0, TEXT("Unable to construct string with negative slack"));
+
+		if (*Src)
 		{
 			int32 SrcLen = TCString<SrcCharType>::Strlen(Src) + 1;
 			int32 DestLen = FPlatformString::ConvertedLength<UE_STRING_CHARTYPE>(Src, SrcLen);
@@ -209,18 +226,34 @@ namespace UE::Core::Private
 	}
 } // namespace UE::Core::Private
 
-UE_STRING_CLASS::UE_STRING_CLASS(const ANSICHAR* Str)								{ UE::Core::Private::ConstructFromCString(Data, Str); }
-UE_STRING_CLASS::UE_STRING_CLASS(const WIDECHAR* Str)								{ UE::Core::Private::ConstructFromCString(Data, Str); }
-UE_STRING_CLASS::UE_STRING_CLASS(const UTF8CHAR* Str)								{ UE::Core::Private::ConstructFromCString(Data, Str); }
-UE_STRING_CLASS::UE_STRING_CLASS(const UCS2CHAR* Str)								{ UE::Core::Private::ConstructFromCString(Data, Str); }
-UE_STRING_CLASS::UE_STRING_CLASS(int32 Len, const ANSICHAR* Str)					{ UE::Core::Private::ConstructWithLength(Data, Len, Str); }
-UE_STRING_CLASS::UE_STRING_CLASS(int32 Len, const WIDECHAR* Str)					{ UE::Core::Private::ConstructWithLength(Data, Len, Str); }
-UE_STRING_CLASS::UE_STRING_CLASS(int32 Len, const UTF8CHAR* Str)					{ UE::Core::Private::ConstructWithLength(Data, Len, Str); }
-UE_STRING_CLASS::UE_STRING_CLASS(int32 Len, const UCS2CHAR* Str)					{ UE::Core::Private::ConstructWithLength(Data, Len, Str); }
-UE_STRING_CLASS::UE_STRING_CLASS(const ANSICHAR* Str, int32 ExtraSlack)				{ UE::Core::Private::ConstructWithSlack(Data, Str, ExtraSlack); }
-UE_STRING_CLASS::UE_STRING_CLASS(const WIDECHAR* Str, int32 ExtraSlack)				{ UE::Core::Private::ConstructWithSlack(Data, Str, ExtraSlack); }
-UE_STRING_CLASS::UE_STRING_CLASS(const UTF8CHAR* Str, int32 ExtraSlack)				{ UE::Core::Private::ConstructWithSlack(Data, Str, ExtraSlack); }
-UE_STRING_CLASS::UE_STRING_CLASS(const UCS2CHAR* Str, int32 ExtraSlack)				{ UE::Core::Private::ConstructWithSlack(Data, Str, ExtraSlack); }
+UE_STRING_CLASS::UE_STRING_CLASS(const ANSICHAR* Str)								{ if (Str) { UE::Core::Private::ConstructFromCString(Data, Str, 0); } }
+UE_STRING_CLASS::UE_STRING_CLASS(const WIDECHAR* Str)								{ if (Str) { UE::Core::Private::ConstructFromCString(Data, Str, 0); } }
+UE_STRING_CLASS::UE_STRING_CLASS(const UTF8CHAR* Str)								{ if (Str) { UE::Core::Private::ConstructFromCString(Data, Str, 0); } }
+UE_STRING_CLASS::UE_STRING_CLASS(const UCS2CHAR* Str)								{ if (Str) { UE::Core::Private::ConstructFromCString(Data, Str, 0); } }
+
+// Deprecated
+UE_STRING_CLASS::UE_STRING_CLASS(int32 Len, const ANSICHAR* Str)					{ if (Str && *Str) { UE::Core::Private::ConstructWithLength(Data, Str, Len, 0); } }
+UE_STRING_CLASS::UE_STRING_CLASS(int32 Len, const WIDECHAR* Str)					{ if (Str && *Str) { UE::Core::Private::ConstructWithLength(Data, Str, Len, 0); } }
+UE_STRING_CLASS::UE_STRING_CLASS(int32 Len, const UTF8CHAR* Str)					{ if (Str && *Str) { UE::Core::Private::ConstructWithLength(Data, Str, Len, 0); } }
+UE_STRING_CLASS::UE_STRING_CLASS(int32 Len, const UCS2CHAR* Str)					{ if (Str && *Str) { UE::Core::Private::ConstructWithLength(Data, Str, Len, 0); } }
+UE_STRING_CLASS::UE_STRING_CLASS(const ANSICHAR* Str, int32 ExtraSlack)				{ if (Str) { UE::Core::Private::ConstructFromCString(Data, Str, ExtraSlack); } else if (ExtraSlack > 0) { Data.Reserve(ExtraSlack + 1); } }
+UE_STRING_CLASS::UE_STRING_CLASS(const WIDECHAR* Str, int32 ExtraSlack)				{ if (Str) { UE::Core::Private::ConstructFromCString(Data, Str, ExtraSlack); } else if (ExtraSlack > 0) { Data.Reserve(ExtraSlack + 1); } }
+UE_STRING_CLASS::UE_STRING_CLASS(const UTF8CHAR* Str, int32 ExtraSlack)				{ if (Str) { UE::Core::Private::ConstructFromCString(Data, Str, ExtraSlack); } else if (ExtraSlack > 0) { Data.Reserve(ExtraSlack + 1); } }
+UE_STRING_CLASS::UE_STRING_CLASS(const UCS2CHAR* Str, int32 ExtraSlack)				{ if (Str) { UE::Core::Private::ConstructFromCString(Data, Str, ExtraSlack); } else if (ExtraSlack > 0) { Data.Reserve(ExtraSlack + 1); } }
+// Deprecated
+
+UE_STRING_CLASS UE_STRING_CLASS::ConstructFromPtrSize(const ANSICHAR* Str, int32 Len)								{ UE_STRING_CLASS Result; UE::Core::Private::ConstructWithLength(Result.Data, Str, Len, 0); return Result; }
+UE_STRING_CLASS UE_STRING_CLASS::ConstructFromPtrSize(const WIDECHAR* Str, int32 Len)								{ UE_STRING_CLASS Result; UE::Core::Private::ConstructWithLength(Result.Data, Str, Len, 0); return Result; }
+UE_STRING_CLASS UE_STRING_CLASS::ConstructFromPtrSize(const UTF8CHAR* Str, int32 Len)								{ UE_STRING_CLASS Result; UE::Core::Private::ConstructWithLength(Result.Data, Str, Len, 0); return Result; }
+UE_STRING_CLASS UE_STRING_CLASS::ConstructFromPtrSize(const UCS2CHAR* Str, int32 Len)								{ UE_STRING_CLASS Result; UE::Core::Private::ConstructWithLength(Result.Data, Str, Len, 0); return Result; }
+UE_STRING_CLASS UE_STRING_CLASS::ConstructWithSlack(const ANSICHAR* Str, int32 ExtraSlack)							{ UE_STRING_CLASS Result; UE::Core::Private::ConstructFromCString(Result.Data, Str, ExtraSlack); return Result; }
+UE_STRING_CLASS UE_STRING_CLASS::ConstructWithSlack(const WIDECHAR* Str, int32 ExtraSlack)							{ UE_STRING_CLASS Result; UE::Core::Private::ConstructFromCString(Result.Data, Str, ExtraSlack); return Result; }
+UE_STRING_CLASS UE_STRING_CLASS::ConstructWithSlack(const UTF8CHAR* Str, int32 ExtraSlack)							{ UE_STRING_CLASS Result; UE::Core::Private::ConstructFromCString(Result.Data, Str, ExtraSlack); return Result; }
+UE_STRING_CLASS UE_STRING_CLASS::ConstructWithSlack(const UCS2CHAR* Str, int32 ExtraSlack)							{ UE_STRING_CLASS Result; UE::Core::Private::ConstructFromCString(Result.Data, Str, ExtraSlack); return Result; }
+UE_STRING_CLASS UE_STRING_CLASS::ConstructFromPtrSizeWithSlack(const ANSICHAR* Str, int32 Len, int32 ExtraSlack)	{ UE_STRING_CLASS Result; UE::Core::Private::ConstructWithLength(Result.Data, Str, Len, ExtraSlack); return Result; }
+UE_STRING_CLASS UE_STRING_CLASS::ConstructFromPtrSizeWithSlack(const WIDECHAR* Str, int32 Len, int32 ExtraSlack)	{ UE_STRING_CLASS Result; UE::Core::Private::ConstructWithLength(Result.Data, Str, Len, ExtraSlack); return Result; }
+UE_STRING_CLASS UE_STRING_CLASS::ConstructFromPtrSizeWithSlack(const UTF8CHAR* Str, int32 Len, int32 ExtraSlack)	{ UE_STRING_CLASS Result; UE::Core::Private::ConstructWithLength(Result.Data, Str, Len, ExtraSlack); return Result; }
+UE_STRING_CLASS UE_STRING_CLASS::ConstructFromPtrSizeWithSlack(const UCS2CHAR* Str, int32 Len, int32 ExtraSlack)	{ UE_STRING_CLASS Result; UE::Core::Private::ConstructWithLength(Result.Data, Str, Len, ExtraSlack); return Result; }
 
 UE_STRING_CLASS& UE_STRING_CLASS::operator=( const ElementType* Other )
 {
@@ -398,7 +431,7 @@ int32 UE_STRING_CLASS::Find(const ElementType* SubStr, int32 SubStrLen, ESearchC
 		// times in the loop below
 		if ( SearchCase == ESearchCase::IgnoreCase)
 		{
-			return ToUpper().Find(UE_STRING_CLASS(SubStrLen, SubStr).ToUpper(), ESearchCase::CaseSensitive, SearchDir, StartPosition);
+			return ToUpper().Find(UE_STRING_CLASS::ConstructFromPtrSize(SubStr, SubStrLen).ToUpper(), ESearchCase::CaseSensitive, SearchDir, StartPosition);
 		}
 		else
 		{
@@ -781,7 +814,7 @@ UE_STRING_CLASS UE_STRING_CLASS::RightChop(int32 Count) const &
 {
 	const int32 Length = Len();
 	const int32 Skip = FMath::Clamp(Count, 0, Length);
-	return UE_STRING_CLASS(Length - Skip, **this + Skip);
+	return UE_STRING_CLASS::ConstructFromPtrSize(**this + Skip, Length - Skip);
 }
 
 UE_STRING_CLASS UE_STRING_CLASS::Mid(int32 Start, int32 Count) const &
@@ -792,7 +825,7 @@ UE_STRING_CLASS UE_STRING_CLASS::Mid(int32 Start, int32 Count) const &
 		const int32 RequestedStart = Start;
 		Start = FMath::Clamp(Start, 0, Length);
 		const int32 End = (int32)FMath::Clamp((int64)Count + RequestedStart, (int64)Start, (int64)Length);
-		return UE_STRING_CLASS(End-Start, **this + Start);
+		return UE_STRING_CLASS::ConstructFromPtrSize(**this + Start, End-Start);
 	}
 
 	return UE_STRING_CLASS();
@@ -1374,7 +1407,7 @@ int32 UE_STRING_CLASS::ParseIntoArray(TArray<UE_STRING_CLASS>& OutArray, const E
 				if(!InCullEmpty || SubstringLength != 0)
 				{
 					// ... add new string from substring beginning up to the beginning of this delimiter.
-					OutArray.Emplace(SubstringEndIndex - SubstringBeginIndex, Start + SubstringBeginIndex);
+					OutArray.Add(UE_STRING_CLASS::ConstructFromPtrSize(Start + SubstringBeginIndex, SubstringEndIndex - SubstringBeginIndex));
 				}
 				// Next substring begins at the end of the discovered delimiter.
 				SubstringBeginIndex = SubstringEndIndex + DelimiterLength;
diff --git a/Engine/Source/Runtime/Core/Private/Containers/StringTest.cpp b/Engine/Source/Runtime/Core/Private/Containers/StringTest.cpp
--- a/Engine/Source/Runtime/Core/Private/Containers/StringTest.cpp
+++ b/Engine/Source/Runtime/Core/Private/Containers/StringTest.cpp
@@ -381,8 +381,8 @@ bool FStringFromStringViewTest::RunTest(const FString& Parameters)
 	return true;
 }
 
-IMPLEMENT_SIMPLE_AUTOMATION_TEST(FStringConstructWithSlackTest, "System.Core.String.ConstructWithSlack", EAutomationTestFlags::ApplicationContextMask | EAutomationTestFlags::SmokeFilter)
-bool FStringConstructWithSlackTest::RunTest(const FString& Parameters)
+IMPLEMENT_SIMPLE_AUTOMATION_TEST(FStringConstructorWithSlackTest, "System.Core.String.ConstructorWithSlack", EAutomationTestFlags::ApplicationContextMask | EAutomationTestFlags::SmokeFilter)
+bool FStringConstructorWithSlackTest::RunTest(const FString& Parameters)
 {
 	// Note that the total capacity of a string might be greater than the string length + slack + a null terminator due to
 	// underlying malloc implementations which is why we poll FMemory to see what size of allocation we should be expecting.
@@ -790,4 +790,143 @@ bool FStringFindAndContainsTest::RunTest(const FString& Parameters)
 	return true;
 }
 
+IMPLEMENT_SIMPLE_AUTOMATION_TEST(FStringConstructorWithLengthTest, "System.Core.String.ConstructorWithLength", EAutomationTestFlags::ApplicationContextMask | EAutomationTestFlags::SmokeFilter)
+bool FStringConstructorWithLengthTest::RunTest(const FString& Parameters)
+{
+	auto DoTest = [this](const TCHAR* Ptr, int32 Size, const TArray<TCHAR>& Expected)
+	{
+		FString Str(Size, Ptr);
+
+		const TArray<TCHAR>& StrArr = Str.GetCharArray();
+		if (StrArr != Expected)
+		{
+			AddError(
+				FString::Printf(
+					TEXT("FString(%d, TEXT(\"%s\")) failure: result '%s' (expected '%.*s')"),
+					Size,
+					Ptr,
+					*Str,
+					Expected.Num(),
+					Expected.GetData()
+				)
+			);
+		}
+	};
+
+	DoTest(TEXT("\0abc"),    4, {});
+	DoTest(TEXT("abc\0def"), 3, { TEXT('a'), TEXT('b'), TEXT('c'),                                                          TEXT('\0') });
+	DoTest(TEXT("abc\0def"), 4, { TEXT('a'), TEXT('b'), TEXT('c'), TEXT('\0'),                                              TEXT('\0') });
+	DoTest(TEXT("abc\0def"), 7, { TEXT('a'), TEXT('b'), TEXT('c'), TEXT('\0'), TEXT('d'), TEXT('e'), TEXT('f'),             TEXT('\0') });
+	DoTest(TEXT("abc\0def"), 8, { TEXT('a'), TEXT('b'), TEXT('c'), TEXT('\0'), TEXT('d'), TEXT('e'), TEXT('f'), TEXT('\0'), TEXT('\0') });
+
+	return true;
+}
+
+IMPLEMENT_SIMPLE_AUTOMATION_TEST(FStringConstructWithSlackTest, "System.Core.String.ConstructWithSlack", EAutomationTestFlags::ApplicationContextMask | EAutomationTestFlags::SmokeFilter)
+bool FStringConstructWithSlackTest::RunTest(const FString& Parameters)
+{
+	auto DoTest = [this](const TCHAR* Ptr, int32 ExtraSlack, const TArray<TCHAR>& Expected)
+	{
+		FString Str = FString::ConstructWithSlack(Ptr, ExtraSlack);
+
+		const TArray<TCHAR>& StrArr = Str.GetCharArray();
+		int32 ActualSlack = StrArr.Max() - StrArr.Num();
+		bool  bValidSlack = (ExtraSlack == 0 && *Ptr == TEXT('\0')) ? (ActualSlack == 0) : (ActualSlack >= ExtraSlack);
+		if (StrArr != Expected || !bValidSlack)
+		{
+			AddError(
+				FString::Printf(
+					TEXT("FString::ConstructWithSlack(TEXT(\"%s\"), %d) failure: result '%s' (expected '%.*s'), slack '%d' (expected '%d')"),
+					Ptr,
+					ExtraSlack,
+					*Str,
+					Expected.Num(),
+					Expected.GetData(),
+					ActualSlack,
+					ExtraSlack
+				)
+			);
+		}
+	};
+
+	DoTest(TEXT("\0abc"), 0,  {});
+	DoTest(TEXT("\0abc"), 47, {});
+	DoTest(TEXT("abc"),   0,  { TEXT('a'), TEXT('b'), TEXT('c'), TEXT('\0') });
+	DoTest(TEXT("abc"),   47, { TEXT('a'), TEXT('b'), TEXT('c'), TEXT('\0') });
+
+	return true;
+}
+
+IMPLEMENT_SIMPLE_AUTOMATION_TEST(FStringConstructFromPtrSizeTest, "System.Core.String.ConstructFromPtrSize", EAutomationTestFlags::ApplicationContextMask | EAutomationTestFlags::SmokeFilter)
+bool FStringConstructFromPtrSizeTest::RunTest(const FString& Parameters)
+{
+	auto DoTest = [this](const TCHAR* Ptr, int32 Size, const TArray<TCHAR>& Expected)
+	{
+		FString Str = FString::ConstructFromPtrSize(Ptr, Size);
+
+		const TArray<TCHAR>& StrArr = Str.GetCharArray();
+		if (StrArr != Expected)
+		{
+			AddError(
+				FString::Printf(
+					TEXT("FString::ConstructFromPtrSize(TEXT(\"%s\"), %d) failure: result '%s' (expected '%.*s')"),
+					Ptr,
+					Size,
+					*Str,
+					Expected.Num(),
+					Expected.GetData())
+			);
+		}
+	};
+
+	DoTest(TEXT("\0abc"),    4, { TEXT('\0'), TEXT('a'), TEXT('b'), TEXT('c'),                                               TEXT('\0') });
+	DoTest(TEXT("abc\0def"), 3, { TEXT('a'),  TEXT('b'), TEXT('c'),                                                          TEXT('\0') });
+	DoTest(TEXT("abc\0def"), 4, { TEXT('a'),  TEXT('b'), TEXT('c'), TEXT('\0'),                                              TEXT('\0') });
+	DoTest(TEXT("abc\0def"), 7, { TEXT('a'),  TEXT('b'), TEXT('c'), TEXT('\0'), TEXT('d'), TEXT('e'), TEXT('f'),             TEXT('\0') });
+	DoTest(TEXT("abc\0def"), 8, { TEXT('a'),  TEXT('b'), TEXT('c'), TEXT('\0'), TEXT('d'), TEXT('e'), TEXT('f'), TEXT('\0'), TEXT('\0') });
+
+	return true;
+}
+
+IMPLEMENT_SIMPLE_AUTOMATION_TEST(FStringConstructFromPtrSizeWithSlackTest, "System.Core.String.ConstructFromPtrSizeWithSlack", EAutomationTestFlags::ApplicationContextMask | EAutomationTestFlags::SmokeFilter)
+bool FStringConstructFromPtrSizeWithSlackTest::RunTest(const FString& Parameters)
+{
+	auto DoTest = [this](const TCHAR* Ptr, int32 Size, int32 ExtraSlack, const TArray<TCHAR>& Expected)
+	{
+		FString Str = FString::ConstructFromPtrSizeWithSlack(Ptr, Size, ExtraSlack);
+
+		const TArray<TCHAR>& StrArr = Str.GetCharArray();
+		int32 ActualSlack = StrArr.Max() - StrArr.Num();
+		bool  bValidSlack = (ExtraSlack == 0 && *Ptr == TEXT('\0')) ? (ActualSlack == 0) : (ActualSlack >= ExtraSlack);
+		if (StrArr != Expected || !bValidSlack)
+		{
+			AddError(
+				FString::Printf(
+					TEXT("FString::ConstructFromPtrSizeWithSlack(TEXT(\"%s\"), %d) failure: result '%s' (expected '%.*s'), slack '%d' (expected '%d')"),
+					Ptr,
+					Size,
+					*Str,
+					Expected.Num(),
+					Expected.GetData(),
+					ActualSlack,
+					ExtraSlack
+				)
+			);
+		}
+	};
+
+	DoTest(TEXT("\0abc"),    4, 0,  { TEXT('\0'), TEXT('a'), TEXT('b'), TEXT('c'),                                               TEXT('\0') });
+	DoTest(TEXT("\0abc"),    4, 47, { TEXT('\0'), TEXT('a'), TEXT('b'), TEXT('c'),                                               TEXT('\0') });
+	DoTest(TEXT("abc\0def"), 3, 0,  { TEXT('a'),  TEXT('b'), TEXT('c'),                                                          TEXT('\0') });
+	DoTest(TEXT("abc\0def"), 3, 47, { TEXT('a'),  TEXT('b'), TEXT('c'),                                                          TEXT('\0') });
+	DoTest(TEXT("abc\0def"), 4, 0,  { TEXT('a'),  TEXT('b'), TEXT('c'), TEXT('\0'),                                              TEXT('\0') });
+	DoTest(TEXT("abc\0def"), 4, 47, { TEXT('a'),  TEXT('b'), TEXT('c'), TEXT('\0'),                                              TEXT('\0') });
+	DoTest(TEXT("abc\0def"), 7, 0,  { TEXT('a'),  TEXT('b'), TEXT('c'), TEXT('\0'), TEXT('d'), TEXT('e'), TEXT('f'),             TEXT('\0') });
+	DoTest(TEXT("abc\0def"), 7, 47, { TEXT('a'),  TEXT('b'), TEXT('c'), TEXT('\0'), TEXT('d'), TEXT('e'), TEXT('f'),             TEXT('\0') });
+	DoTest(TEXT("abc\0def"), 8, 0,  { TEXT('a'),  TEXT('b'), TEXT('c'), TEXT('\0'), TEXT('d'), TEXT('e'), TEXT('f'), TEXT('\0'), TEXT('\0') });
+	DoTest(TEXT("abc\0def"), 8, 47, { TEXT('a'),  TEXT('b'), TEXT('c'), TEXT('\0'), TEXT('d'), TEXT('e'), TEXT('f'), TEXT('\0'), TEXT('\0') });
+
+	return true;
+}
+
 #endif // WITH_DEV_AUTOMATION_TESTS
diff --git a/Engine/Source/Runtime/Core/Public/Containers/AnsiString.h b/Engine/Source/Runtime/Core/Public/Containers/AnsiString.h
--- a/Engine/Source/Runtime/Core/Public/Containers/AnsiString.h
+++ b/Engine/Source/Runtime/Core/Public/Containers/AnsiString.h
@@ -5,10 +5,12 @@
 // Include UnrealString.h.inl's includes before defining the macros, in case the macros 'poison' other headers or there are re-entrant includes.
 #include "Containers/UnrealStringIncludes.h.inl"
 
-#define UE_STRING_CLASS             FAnsiString
-#define UE_STRING_CHARTYPE          ANSICHAR
-#define UE_STRING_CHARTYPE_IS_TCHAR 0
+#define UE_STRING_CLASS                        FAnsiString
+#define UE_STRING_CHARTYPE                     ANSICHAR
+#define UE_STRING_CHARTYPE_IS_TCHAR            0
+#define UE_STRING_DEPRECATED(Version, Message) UE_DEPRECATED(Version, Message)
 	#include "Containers/UnrealString.h.inl"
+#undef UE_STRING_DEPRECATED
 #undef UE_STRING_CHARTYPE_IS_TCHAR
 #undef UE_STRING_CHARTYPE
 #undef UE_STRING_CLASS
diff --git a/Engine/Source/Runtime/Core/Public/Containers/UnrealString.h b/Engine/Source/Runtime/Core/Public/Containers/UnrealString.h
--- a/Engine/Source/Runtime/Core/Public/Containers/UnrealString.h
+++ b/Engine/Source/Runtime/Core/Public/Containers/UnrealString.h
@@ -8,10 +8,12 @@
 // Include UnrealString.h.inl's includes before defining the macros, in case the macros 'poison' other headers or there are re-entrant includes.
 #include "Containers/UnrealStringIncludes.h.inl"
 
-#define UE_STRING_CLASS             FString
-#define UE_STRING_CHARTYPE          TCHAR
-#define UE_STRING_CHARTYPE_IS_TCHAR 1
+#define UE_STRING_CLASS                        FString
+#define UE_STRING_CHARTYPE                     TCHAR
+#define UE_STRING_CHARTYPE_IS_TCHAR            1
+#define UE_STRING_DEPRECATED(Version, Message)
 	#include "Containers/UnrealString.h.inl"
+#undef UE_STRING_DEPRECATED
 #undef UE_STRING_CHARTYPE_IS_TCHAR
 #undef UE_STRING_CHARTYPE
 #undef UE_STRING_CLASS
diff --git a/Engine/Source/Runtime/Core/Public/Containers/UnrealString.h.inl b/Engine/Source/Runtime/Core/Public/Containers/UnrealString.h.inl
--- a/Engine/Source/Runtime/Core/Public/Containers/UnrealString.h.inl
+++ b/Engine/Source/Runtime/Core/Public/Containers/UnrealString.h.inl
@@ -20,6 +20,9 @@
 #ifndef UE_STRING_CHARTYPE_IS_TCHAR
 	#error "UnrealString.h.inl should only be included after defining UE_STRING_CHARTYPE_IS_TCHAR"
 #endif
+#ifndef UE_STRING_DEPRECATED
+	#error "UnrealString.h.inl should only be included after defining UE_STRING_DEPRECATED"
+#endif
 
 struct PREPROCESSOR_JOIN(UE_STRING_CLASS, FormatArg);
 template<typename InKeyType,typename InValueType,typename SetAllocator ,typename KeyFuncs > class TMap;
@@ -92,16 +95,34 @@ public:
 	CORE_API UE_STRING_CLASS(const UCS2CHAR* Str);
 
 	/** Construct from null-terminated C substring or nullptr */
-	CORE_API UE_STRING_CLASS(int32 Len, const ANSICHAR* Str);
-	CORE_API UE_STRING_CLASS(int32 Len, const WIDECHAR* Str);
-	CORE_API UE_STRING_CLASS(int32 Len, const UTF8CHAR* Str);
-	CORE_API UE_STRING_CLASS(int32 Len, const UCS2CHAR* Str);
+	UE_STRING_DEPRECATED(5.4, "This constructor has been deprecated - please use " PREPROCESSOR_TO_STRING(UE_STRING_CLASS) "::ConstructFromPtrSize(Ptr, Size) instead.") CORE_API UE_STRING_CLASS(int32 Len, const ANSICHAR* Str);
+	UE_STRING_DEPRECATED(5.4, "This constructor has been deprecated - please use " PREPROCESSOR_TO_STRING(UE_STRING_CLASS) "::ConstructFromPtrSize(Ptr, Size) instead.") CORE_API UE_STRING_CLASS(int32 Len, const WIDECHAR* Str);
+	UE_STRING_DEPRECATED(5.4, "This constructor has been deprecated - please use " PREPROCESSOR_TO_STRING(UE_STRING_CLASS) "::ConstructFromPtrSize(Ptr, Size) instead.") CORE_API UE_STRING_CLASS(int32 Len, const UTF8CHAR* Str);
+	UE_STRING_DEPRECATED(5.4, "This constructor has been deprecated - please use " PREPROCESSOR_TO_STRING(UE_STRING_CLASS) "::ConstructFromPtrSize(Ptr, Size) instead.") CORE_API UE_STRING_CLASS(int32 Len, const UCS2CHAR* Str);
 
 	/** Construct from null-terminated C string or nullptr with extra slack on top of original string length */
-	CORE_API UE_STRING_CLASS(const ANSICHAR* Str, int32 ExtraSlack);
-	CORE_API UE_STRING_CLASS(const WIDECHAR* Str, int32 ExtraSlack);
-	CORE_API UE_STRING_CLASS(const UTF8CHAR* Str, int32 ExtraSlack);
-	CORE_API UE_STRING_CLASS(const UCS2CHAR* Str, int32 ExtraSlack);
+	UE_STRING_DEPRECATED(5.4, "This constructor has been deprecated - please use " PREPROCESSOR_TO_STRING(UE_STRING_CLASS) "::ConstructWithSlack(Ptr, Size) instead.") CORE_API UE_STRING_CLASS(const ANSICHAR* Str, int32 ExtraSlack);
+	UE_STRING_DEPRECATED(5.4, "This constructor has been deprecated - please use " PREPROCESSOR_TO_STRING(UE_STRING_CLASS) "::ConstructWithSlack(Ptr, Size) instead.") CORE_API UE_STRING_CLASS(const WIDECHAR* Str, int32 ExtraSlack);
+	UE_STRING_DEPRECATED(5.4, "This constructor has been deprecated - please use " PREPROCESSOR_TO_STRING(UE_STRING_CLASS) "::ConstructWithSlack(Ptr, Size) instead.") CORE_API UE_STRING_CLASS(const UTF8CHAR* Str, int32 ExtraSlack);
+	UE_STRING_DEPRECATED(5.4, "This constructor has been deprecated - please use " PREPROCESSOR_TO_STRING(UE_STRING_CLASS) "::ConstructWithSlack(Ptr, Size) instead.") CORE_API UE_STRING_CLASS(const UCS2CHAR* Str, int32 ExtraSlack);
+
+	/** Construct from null-terminated C string with extra slack on top of original string length. */
+	static CORE_API UE_STRING_CLASS ConstructWithSlack(const ANSICHAR* Str, int32 ExtraSlack);
+	static CORE_API UE_STRING_CLASS ConstructWithSlack(const WIDECHAR* Str, int32 ExtraSlack);
+	static CORE_API UE_STRING_CLASS ConstructWithSlack(const UTF8CHAR* Str, int32 ExtraSlack);
+	static CORE_API UE_STRING_CLASS ConstructWithSlack(const UCS2CHAR* Str, int32 ExtraSlack);
+
+	/** Construct from a buffer.  If the buffer contains zeros, these will be present in the constructed string. */
+	static CORE_API UE_STRING_CLASS ConstructFromPtrSize(const ANSICHAR* Str, int32 Size);
+	static CORE_API UE_STRING_CLASS ConstructFromPtrSize(const WIDECHAR* Str, int32 Size);
+	static CORE_API UE_STRING_CLASS ConstructFromPtrSize(const UTF8CHAR* Str, int32 Size);
+	static CORE_API UE_STRING_CLASS ConstructFromPtrSize(const UCS2CHAR* Str, int32 Size);
+
+	/** Construct from a buffer with extra slack on top of original string length.  If the buffer contains zeros, these will be present in the constructed string. */
+	static CORE_API UE_STRING_CLASS ConstructFromPtrSizeWithSlack(const ANSICHAR* Str, int32 Size, int32 ExtraSlack);
+	static CORE_API UE_STRING_CLASS ConstructFromPtrSizeWithSlack(const WIDECHAR* Str, int32 Size, int32 ExtraSlack);
+	static CORE_API UE_STRING_CLASS ConstructFromPtrSizeWithSlack(const UTF8CHAR* Str, int32 Size, int32 ExtraSlack);
+	static CORE_API UE_STRING_CLASS ConstructFromPtrSizeWithSlack(const UCS2CHAR* Str, int32 Size, int32 ExtraSlack);
 
 	/** Construct from contiguous range of characters such as a string view or string builder */
 	template <
@@ -114,8 +135,9 @@ public:
 			!std::is_base_of_v<UE_STRING_CLASS, std::decay_t<CharRangeType>>
 		)
 	>
-	FORCEINLINE explicit UE_STRING_CLASS(CharRangeType&& Str) : UE_STRING_CLASS(GetNum(Str), GetData(Forward<CharRangeType>(Str)))
+	FORCEINLINE explicit UE_STRING_CLASS(CharRangeType&& Str)
 	{
+		*this = UE_STRING_CLASS::ConstructFromPtrSize(GetData(Forward<CharRangeType>(Str)), GetNum(Str));
 	}
 
 	/** Construct from contiguous range of characters with extra slack on top of original string length */
@@ -1046,7 +1068,7 @@ public:
 	/** Returns the left most given number of characters */
 	UE_NODISCARD FORCEINLINE UE_STRING_CLASS Left( int32 Count ) const &
 	{
-		return UE_STRING_CLASS( FMath::Clamp(Count,0,Len()), **this );
+		return UE_STRING_CLASS::ConstructFromPtrSize(**this, FMath::Clamp(Count,0,Len()) );
 	}
 
 	UE_NODISCARD FORCEINLINE UE_STRING_CLASS Left(int32 Count) &&
@@ -1067,7 +1089,7 @@ public:
 	UE_NODISCARD FORCEINLINE UE_STRING_CLASS LeftChop( int32 Count ) const &
 	{
 		const int32 Length = Len();
-		return UE_STRING_CLASS( FMath::Clamp(Length-Count,0, Length), **this );
+		return UE_STRING_CLASS::ConstructFromPtrSize( **this, FMath::Clamp(Length-Count,0, Length) );
 	}
 
 	UE_NODISCARD FORCEINLINE UE_STRING_CLASS LeftChop(int32 Count)&&
diff --git a/Engine/Source/Runtime/Core/Public/Containers/Utf8String.h b/Engine/Source/Runtime/Core/Public/Containers/Utf8String.h
--- a/Engine/Source/Runtime/Core/Public/Containers/Utf8String.h
+++ b/Engine/Source/Runtime/Core/Public/Containers/Utf8String.h
@@ -5,10 +5,12 @@
 // Include UnrealString.h.inl's includes before defining the macros, in case the macros 'poison' other headers or there are re-entrant includes.
 #include "Containers/UnrealStringIncludes.h.inl"
 
-#define UE_STRING_CLASS             FUtf8String
-#define UE_STRING_CHARTYPE          UTF8CHAR
-#define UE_STRING_CHARTYPE_IS_TCHAR 0
+#define UE_STRING_CLASS                        FUtf8String
+#define UE_STRING_CHARTYPE                     UTF8CHAR
+#define UE_STRING_CHARTYPE_IS_TCHAR            0
+#define UE_STRING_DEPRECATED(Version, Message) UE_DEPRECATED(Version, Message)
 	#include "Containers/UnrealString.h.inl"
+#undef UE_STRING_DEPRECATED
 #undef UE_STRING_CHARTYPE_IS_TCHAR
 #undef UE_STRING_CHARTYPE
 #undef UE_STRING_CLASS
