Core: Added another operator+ for FAsciiSet to allow for combining sets

diff --git a/Engine/Source/Runtime/Core/Public/Misc/AsciiSet.h b/Engine/Source/Runtime/Core/Public/Misc/AsciiSet.h
--- a/Engine/Source/Runtime/Core/Public/Misc/AsciiSet.h
+++ b/Engine/Source/Runtime/Core/Public/Misc/AsciiSet.h
@@ -56,6 +56,18 @@ public:
 		return FAsciiSet(Bitset);
 	}
 
+	/** Create new set containing the combination of this set and another */
+	constexpr FORCEINLINE FAsciiSet operator|(FAsciiSet OtherSet) const
+	{
+		return FAsciiSet(LoMask | OtherSet.LoMask, HiMask | OtherSet.HiMask);
+	}
+
+	/** Create new set containing only the characters that exist in both this and another specified set */
+	constexpr FORCEINLINE FAsciiSet operator&(FAsciiSet OtherSet) const
+	{
+		return FAsciiSet(LoMask & OtherSet.LoMask, HiMask & OtherSet.HiMask);
+	}
+
 	/** Create new set containing inverse set of characters - likely including null-terminator */
 	constexpr FORCEINLINE FAsciiSet operator~() const
 	{
