Fixed deprecated FString constructors.

diff --git a/Engine/Source/Developer/DerivedDataCache/Private/HttpCacheStore.cpp b/Engine/Source/Developer/DerivedDataCache/Private/HttpCacheStore.cpp
--- a/Engine/Source/Developer/DerivedDataCache/Private/HttpCacheStore.cpp
+++ b/Engine/Source/Developer/DerivedDataCache/Private/HttpCacheStore.cpp
@@ -753,7 +753,7 @@ FString FHttpCacheStore::FHttpOperation::GetBodyAsString() const
 			return JsonStringBuilder.ToString();
 		}
 	}
-	return FString(Len, (const UTF8CHAR*)ResponseBody.GetData());
+	return FString::ConstructFromPtrSize((const UTF8CHAR*)ResponseBody.GetData(), Len);
 }
 
 TSharedPtr<FJsonObject> FHttpCacheStore::FHttpOperation::GetBodyAsJson() const
diff --git a/Engine/Source/Developer/DerivedDataCache/Private/S3CacheStore.cpp b/Engine/Source/Developer/DerivedDataCache/Private/S3CacheStore.cpp
--- a/Engine/Source/Developer/DerivedDataCache/Private/S3CacheStore.cpp
+++ b/Engine/Source/Developer/DerivedDataCache/Private/S3CacheStore.cpp
@@ -558,7 +558,7 @@ private:
 	static FString GetResponseAsString(const TArray<uint8>& Buffer)
 	{
 		FUTF8ToTCHAR TCHARData(reinterpret_cast<const ANSICHAR*>(Buffer.GetData()), Buffer.Num());
-		return FString(TCHARData.Length(), TCHARData.Get());
+		return FString::ConstructFromPtrSize(TCHARData.Get(), TCHARData.Length());
 	}
 
 	static int StaticStatusFn(void* Ptr, curl_off_t TotalDownloadSize, curl_off_t CurrentDownloadSize, curl_off_t TotalUploadSize, curl_off_t CurrentUploadSize)
diff --git a/Engine/Source/Developer/DesktopPlatform/Private/DesktopPlatformBase.cpp b/Engine/Source/Developer/DesktopPlatform/Private/DesktopPlatformBase.cpp
--- a/Engine/Source/Developer/DesktopPlatform/Private/DesktopPlatformBase.cpp
+++ b/Engine/Source/Developer/DesktopPlatform/Private/DesktopPlatformBase.cpp
@@ -281,14 +281,14 @@ bool FDesktopPlatformBase::TryGetEngineVersion(const FString& RootDir, FEngineVe
 					// Parse an identifier. Exact C rules for an identifier don't really matter; we just need alphanumeric sequences.
 					const TCHAR* TokenStart = TextPos++;
 					while(FChar::IsIdentifier(*TextPos)) TextPos++;
-					Tokens.Add(FString(UE_PTRDIFF_TO_INT32(TextPos - TokenStart), TokenStart));
+					Tokens.Add(FString::ConstructFromPtrSize(TokenStart, UE_PTRDIFF_TO_INT32(TextPos - TokenStart)));
 				}
 				else if(*TextPos == '\"')
 				{
 					// Parse a string
 					const TCHAR* TokenStart = TextPos++;
 					while(*TextPos != 0 && (TextPos == TokenStart + 1 || *(TextPos - 1) != '\"')) TextPos++;
-					Tokens.Add(FString(UE_PTRDIFF_TO_INT32(TextPos - TokenStart), TokenStart));
+					Tokens.Add(FString::ConstructFromPtrSize(TokenStart, UE_PTRDIFF_TO_INT32(TextPos - TokenStart)));
 				}
 				else if(*TextPos == '/' && *(TextPos + 1) == '/')
 				{
@@ -305,7 +305,7 @@ bool FDesktopPlatformBase::TryGetEngineVersion(const FString& RootDir, FEngineVe
 				else
 				{
 					// Take a single symbol character
-					Tokens.Add(FString(1, TextPos));
+					Tokens.Add(FString::ConstructFromPtrSize(TextPos, 1));
 					TextPos++;
 				}
 			}
diff --git a/Engine/Source/Developer/DesktopPlatform/Private/Windows/DesktopPlatformWindows.cpp b/Engine/Source/Developer/DesktopPlatform/Private/Windows/DesktopPlatformWindows.cpp
--- a/Engine/Source/Developer/DesktopPlatform/Private/Windows/DesktopPlatformWindows.cpp
+++ b/Engine/Source/Developer/DesktopPlatform/Private/Windows/DesktopPlatformWindows.cpp
@@ -361,7 +361,7 @@ void FDesktopPlatformWindows::EnumerateEngineInstallations(TMap<FString, FString
 				int32 ValueDataLength = ValueDataSize / sizeof(TCHAR);
 				if(ValueDataLength > 0 && ValueData[ValueDataLength - 1] == 0) ValueDataLength--;
 
-				FString NormalizedInstalledDirectory(ValueDataLength, ValueData);
+				FString NormalizedInstalledDirectory = FString::ConstructFromPtrSize(ValueData, ValueDataLength);
 				FPaths::NormalizeDirectoryName(NormalizedInstalledDirectory);
 				FPaths::CollapseRelativeDirectories(NormalizedInstalledDirectory);
 
diff --git a/Engine/Source/Developer/Zen/Private/ZenServerHttp.cpp b/Engine/Source/Developer/Zen/Private/ZenServerHttp.cpp
--- a/Engine/Source/Developer/Zen/Private/ZenServerHttp.cpp
+++ b/Engine/Source/Developer/Zen/Private/ZenServerHttp.cpp
@@ -382,7 +382,7 @@ namespace UE::Zen {
 			const ANSICHAR* ValueStart = Found + HeaderLen + 2; //colon and space
 			const size_t ValueSize = Linebreak - ValueStart;
 			FUTF8ToTCHAR TCHARData(ValueStart, ValueSize);
-			OutValue = FString(TCHARData.Length(), TCHARData.Get());
+			OutValue = FString::ConstructFromPtrSize(TCHARData.Get(), TCHARData.Length());
 			return true;
 		}
 		return false;
@@ -472,7 +472,7 @@ namespace UE::Zen {
 	{
 		// Content is NOT null-terminated; we need to specify lengths here
 		FUTF8ToTCHAR TCHARData(reinterpret_cast<const ANSICHAR*>(Buffer.GetData()), IntCastChecked<int32>(Buffer.Num()));
-		return FString(TCHARData.Length(), TCHARData.Get());
+		return FString::ConstructFromPtrSize(TCHARData.Get(), TCHARData.Length());
 	}
 
 	FCbObjectView FZenHttpRequest::GetResponseAsObject() const
@@ -497,7 +497,7 @@ namespace UE::Zen {
 					int CalculatedSize = FoundNulPtr != nullptr ? FoundNulPtr - DebugInfo : DebugInfoSize;
 
 					auto ConvertedString = StringCast<TCHAR>(static_cast<const ANSICHAR*>(DebugInfo), CalculatedSize);
-					FString DebugText(ConvertedString.Length(), ConvertedString.Get());
+					FString DebugText = FString::ConstructFromPtrSize(ConvertedString.Get(), ConvertedString.Length());
 					DebugText.ReplaceInline(TEXT("\n"), TEXT(""), ESearchCase::CaseSensitive);
 					DebugText.ReplaceInline(TEXT("\r"), TEXT(""), ESearchCase::CaseSensitive);
 					UE_LOG(LogZenHttp, VeryVerbose, TEXT("CURL %p: '%s'"), Request, *DebugText);
diff --git a/Engine/Source/Runtime/ApplicationCore/Private/Windows/TextStoreACP.cpp b/Engine/Source/Runtime/ApplicationCore/Private/Windows/TextStoreACP.cpp
--- a/Engine/Source/Runtime/ApplicationCore/Private/Windows/TextStoreACP.cpp
+++ b/Engine/Source/Runtime/ApplicationCore/Private/Windows/TextStoreACP.cpp
@@ -607,7 +607,7 @@ STDAPI FTextStoreACP::InsertTextAtSelection(DWORD dwFlags, __RPC__in_ecount_full
 		return E_INVALIDARG;
 	}
 
-	FString NewString(cch, pchText);
+	FString NewString = FString::ConstructFromPtrSize(pchText, cch);
 
 	uint32 BeginIndex;
 	uint32 Length;
diff --git a/Engine/Source/Runtime/ApplicationCore/Private/Windows/WindowsConsoleOutputDevice2.cpp b/Engine/Source/Runtime/ApplicationCore/Private/Windows/WindowsConsoleOutputDevice2.cpp
--- a/Engine/Source/Runtime/ApplicationCore/Private/Windows/WindowsConsoleOutputDevice2.cpp
+++ b/Engine/Source/Runtime/ApplicationCore/Private/Windows/WindowsConsoleOutputDevice2.cpp
@@ -1711,7 +1711,7 @@ public:
 					int32 Len = UE_PTRDIFF_TO_INT32(LineBreak - SearchStr);
 					if (LineBreak > SearchStr && *(LineBreak - 1) == '\r')
 						--Len;
-					Log.Add({ FString(Len, SearchStr), E.Verbosity, E.Category, E.Time, E.TextAttribute, 1 });
+					Log.Add({ FString::ConstructFromPtrSize(SearchStr, Len), E.Verbosity, E.Category, E.Time, E.TextAttribute, 1 });
 					SearchStr = LineBreak + 1;
 				}
 				else
diff --git a/Engine/Source/Runtime/AssetRegistry/Private/AssetRegistryState.cpp b/Engine/Source/Runtime/AssetRegistry/Private/AssetRegistryState.cpp
--- a/Engine/Source/Runtime/AssetRegistry/Private/AssetRegistryState.cpp
+++ b/Engine/Source/Runtime/AssetRegistry/Private/AssetRegistryState.cpp
@@ -2672,7 +2672,7 @@ void FAssetRegistryState::Dump(const TArray<FString>& Arguments, TArray<FString>
 			}
 		}
 
-		OutPages.Emplace(PageEndIndex, PageBuffer.GetData());
+		OutPages.Add(FString::ConstructFromPtrSize(PageBuffer.GetData(), PageEndIndex));
 		if (PageEndIndex != PageBuffer.Len())
 		{
 			PageEndIndex += LineTerminatorLen; // Skip the newline
@@ -2693,7 +2693,7 @@ void FAssetRegistryState::Dump(const TArray<FString>& Arguments, TArray<FString>
 	{
 		if (LinesPerPage == 1)
 		{
-			OutPages.Emplace(PageBuffer.Len(), PageBuffer.GetData());
+			OutPages.Add(FString::ConstructFromPtrSize(PageBuffer.GetData(), PageBuffer.Len()));
 			PageBuffer.Reset();
 		}
 		else
diff --git a/Engine/Source/Runtime/Core/Private/Internationalization/Text.cpp b/Engine/Source/Runtime/Core/Private/Internationalization/Text.cpp
--- a/Engine/Source/Runtime/Core/Private/Internationalization/Text.cpp
+++ b/Engine/Source/Runtime/Core/Private/Internationalization/Text.cpp
@@ -790,7 +790,7 @@ FText FText::AsMemory(uint64 NumBytes, const FNumberFormattingOptions* const Opt
 
 	const double MemorySizeAsDouble = (double)NumBytes / (double)Unit;
 	Args.Add( TEXT("Number"), FText::AsNumber( MemorySizeAsDouble, Options, TargetCulture) );
-	Args.Add( TEXT("Unit"), FText::FromString( FString( 1, &Prefixes[Prefix] ) + Suffix) );
+	Args.Add( TEXT("Unit"), FText::FromString( FString::ConstructFromPtrSize( &Prefixes[Prefix], 1 ) + Suffix) );
 	return FText::Format( NSLOCTEXT("Internationalization", "ComputerMemoryFormatting", "{Number} {Unit}"), Args);
 }
 
diff --git a/Engine/Source/Runtime/Core/Private/Internationalization/TextFormatArgumentModifier.cpp b/Engine/Source/Runtime/Core/Private/Internationalization/TextFormatArgumentModifier.cpp
--- a/Engine/Source/Runtime/Core/Private/Internationalization/TextFormatArgumentModifier.cpp
+++ b/Engine/Source/Runtime/Core/Private/Internationalization/TextFormatArgumentModifier.cpp
@@ -185,7 +185,7 @@ TSharedPtr<ITextFormatArgumentModifier> FTextFormatArgumentModifier_PluralForm::
 		PluralForms.Reserve(ArgKeyValues.Num());
 		for (const auto& Pair : ArgKeyValues)
 		{
-			FTextFormat PluralForm = FTextFormat::FromString(FString(Pair.Value.StringLen, Pair.Value.StringPtr), InPatternDef);
+			FTextFormat PluralForm = FTextFormat::FromString(FString::ConstructFromPtrSize(Pair.Value.StringPtr, Pair.Value.StringLen), InPatternDef);
 			if (!PluralForm.IsValid())
 			{
 				break;
@@ -374,12 +374,12 @@ TSharedPtr<ITextFormatArgumentModifier> FTextFormatArgumentModifier_GenderForm::
 	if (ParseValueArgs(InArgsString, ArgValues) && (ArgValues.Num() == 2 || ArgValues.Num() == 3))
 	{
 		// Gender forms may contain format markers, so pre-compile all the variants now so that Evaluate doesn't have to (this also lets us validate the gender form strings and fail if they're not correct)
-		FTextFormat MasculineForm = FTextFormat::FromString(FString(ArgValues[0].StringLen, ArgValues[0].StringPtr), InPatternDef);
-		FTextFormat FeminineForm  = FTextFormat::FromString(FString(ArgValues[1].StringLen, ArgValues[1].StringPtr), InPatternDef);
+		FTextFormat MasculineForm = FTextFormat::FromString(FString::ConstructFromPtrSize(ArgValues[0].StringPtr, ArgValues[0].StringLen), InPatternDef);
+		FTextFormat FeminineForm  = FTextFormat::FromString(FString::ConstructFromPtrSize(ArgValues[1].StringPtr, ArgValues[1].StringLen), InPatternDef);
 		FTextFormat NeuterForm;
 		if (ArgValues.Num() == 3)
 		{
-			NeuterForm = FTextFormat::FromString(FString(ArgValues[2].StringLen, ArgValues[2].StringPtr), InPatternDef);
+			NeuterForm = FTextFormat::FromString(FString::ConstructFromPtrSize(ArgValues[2].StringPtr, ArgValues[2].StringLen), InPatternDef);
 		}
 
 		// Did everything compile?
diff --git a/Engine/Source/Runtime/Core/Private/Internationalization/TextFormatter.cpp b/Engine/Source/Runtime/Core/Private/Internationalization/TextFormatter.cpp
--- a/Engine/Source/Runtime/Core/Private/Internationalization/TextFormatter.cpp
+++ b/Engine/Source/Runtime/Core/Private/Internationalization/TextFormatter.cpp
@@ -957,7 +957,7 @@ void FTextFormatData::GetFormatArgumentNames_NoLock(TArray<FString>& OutArgument
 
 			if (!bIsInArray)
 			{
-				OutArgumentNames.Add(FString(ArgumentToken->ArgumentNameLen, ArgumentToken->ArgumentNameStartPos));
+				OutArgumentNames.Add(FString::ConstructFromPtrSize(ArgumentToken->ArgumentNameStartPos, ArgumentToken->ArgumentNameLen));
 			}
 		}
 		else if (const auto* ArgumentModifierToken = Token.Node.Cast<TextFormatTokens::FArgumentModifierTokenSpecifier>())
@@ -1115,7 +1115,7 @@ FString FTextFormatter::FormatStr(const FTextFormat& InFmt, const FFormatOrdered
 			// We have existing code that is incorrectly using names in the format string when providing ordered arguments
 			// ICU used to fallback to treating the index of the argument within the string as if it were the index specified 
 			// by the argument name, so we need to emulate that behavior to avoid breaking some format operations
-			UE_LOG(LogTextFormatter, Warning, TEXT("Failed to parse argument \"%s\" as a number (using \"%d\" as a fallback). Please check your format string for errors: \"%s\"."), *FString(ArgumentToken.ArgumentNameLen, ArgumentToken.ArgumentNameStartPos), ArgumentNumber, *FmtPattern);
+			UE_LOG(LogTextFormatter, Warning, TEXT("Failed to parse argument \"%s\" as a number (using \"%d\" as a fallback). Please check your format string for errors: \"%s\"."), *FString::ConstructFromPtrSize(ArgumentToken.ArgumentNameStartPos, ArgumentToken.ArgumentNameLen), ArgumentNumber, *FmtPattern);
 			ArgumentIndex = ArgumentNumber;
 		}
 		return InArguments.IsValidIndex(ArgumentIndex) ? &(InArguments[ArgumentIndex]) : nullptr;
diff --git a/Engine/Source/Runtime/Core/Private/Misc/ExpressionParser.cpp b/Engine/Source/Runtime/Core/Private/Misc/ExpressionParser.cpp
--- a/Engine/Source/Runtime/Core/Private/Misc/ExpressionParser.cpp
+++ b/Engine/Source/Runtime/Core/Private/Misc/ExpressionParser.cpp
@@ -64,7 +64,7 @@ FString FTokenStream::GetErrorContext() const
 	}
 
 	static const int32 MaxChars = 32;
-	FString Context(FMath::Min(int32(EndPos - StartPos), MaxChars), StartPos);
+	FString Context = FString::ConstructFromPtrSize(StartPos, FMath::Min(int32(EndPos - StartPos), MaxChars));
 	if (EndPos - StartPos > MaxChars)
 	{
 		Context += TEXT("...");
diff --git a/Engine/Source/Runtime/Core/Private/Misc/FeedbackContextMarkup.cpp b/Engine/Source/Runtime/Core/Private/Misc/FeedbackContextMarkup.cpp
--- a/Engine/Source/Runtime/Core/Private/Misc/FeedbackContextMarkup.cpp
+++ b/Engine/Source/Runtime/Core/Private/Misc/FeedbackContextMarkup.cpp
@@ -168,7 +168,7 @@ bool FFeedbackContextMarkup::ReadString(const TCHAR*& Text, FString& OutString)
 		{
 			if(*End == *Text)
 			{
-				OutString = FString(UE_PTRDIFF_TO_INT32(End - (Text + 1)), Text + 1);
+				OutString = FString::ConstructFromPtrSize(Text + 1, UE_PTRDIFF_TO_INT32(End - (Text + 1)));
 				do { End++; } while(FChar::IsWhitespace(*End));
 				Text = End;
 				return true;
diff --git a/Engine/Source/Runtime/Core/Private/Misc/FileHelper.cpp b/Engine/Source/Runtime/Core/Private/Misc/FileHelper.cpp
--- a/Engine/Source/Runtime/Core/Private/Misc/FileHelper.cpp
+++ b/Engine/Source/Runtime/Core/Private/Misc/FileHelper.cpp
@@ -286,7 +286,7 @@ bool FFileHelper::LoadFileToStringArrayWithPredicate(TArray<FString>& Result, co
 				Pos++;
 			}
 
-			FString Line(UE_PTRDIFF_TO_INT32(Pos - LineStart), LineStart);
+			FString Line = FString::ConstructFromPtrSize(LineStart, UE_PTRDIFF_TO_INT32(Pos - LineStart));
 			if (Invoke(Predicate, Line))
 			{
 				Result.Add(MoveTemp(Line));
diff --git a/Engine/Source/Runtime/Core/Private/Misc/Parse.cpp b/Engine/Source/Runtime/Core/Private/Misc/Parse.cpp
--- a/Engine/Source/Runtime/Core/Private/Misc/Parse.cpp
+++ b/Engine/Source/Runtime/Core/Private/Misc/Parse.cpp
@@ -1293,7 +1293,7 @@ bool FParse::SchemeNameFromURI(const TCHAR* URI, FString& OutSchemeName)
 		{
 			if(Idx > 0 && URI[Idx] == TEXT(':'))
 			{
-				OutSchemeName = FString(Idx, URI);
+				OutSchemeName = FString::ConstructFromPtrSize(URI, Idx);
 				return true;
 			}
 			return false;
diff --git a/Engine/Source/Runtime/Core/Private/UObject/UnrealNames.cpp b/Engine/Source/Runtime/Core/Private/UObject/UnrealNames.cpp
--- a/Engine/Source/Runtime/Core/Private/UObject/UnrealNames.cpp
+++ b/Engine/Source/Runtime/Core/Private/UObject/UnrealNames.cpp
@@ -2260,11 +2260,11 @@ FString FNameEntry::GetPlainNameString() const
 	FNameBuffer Temp;
 	if (Header.bIsWide)
 	{
-		return FString(Header.Len, GetUnterminatedName(Temp.WideName));
+		return FString::ConstructFromPtrSize(GetUnterminatedName(Temp.WideName), Header.Len);
 	}
 	else
 	{
-		return FString(Header.Len, GetUnterminatedName(Temp.AnsiName));
+		return FString::ConstructFromPtrSize(GetUnterminatedName(Temp.AnsiName), Header.Len);
 	}
 }
 
diff --git a/Engine/Source/Runtime/Core/Private/Windows/WindowsPlatformFile.cpp b/Engine/Source/Runtime/Core/Private/Windows/WindowsPlatformFile.cpp
--- a/Engine/Source/Runtime/Core/Private/Windows/WindowsPlatformFile.cpp
+++ b/Engine/Source/Runtime/Core/Private/Windows/WindowsPlatformFile.cpp
@@ -1192,7 +1192,7 @@ bool FWindowsPlatformFile::HasMarkOfTheWeb(FStringView Filename, FString* OutSou
 		const int32 HostUrlLen = static_cast<int32>(GetPrivateProfileStringW(TEXT("ZoneTransfer"), TEXT("HostUrl"), nullptr, HostUrl, static_cast<DWORD>(GetNum(HostUrl)), *StreamPath));
 		if (0 < HostUrlLen && HostUrlLen <= GetNum(HostUrl))
 		{
-			*OutSourceURL = FString(HostUrlLen, HostUrl);
+			*OutSourceURL = FString::ConstructFromPtrSize(HostUrl, HostUrlLen);
 		}
 	}
 
diff --git a/Engine/Source/Runtime/Core/Private/Windows/WindowsPlatformMisc.cpp b/Engine/Source/Runtime/Core/Private/Windows/WindowsPlatformMisc.cpp
--- a/Engine/Source/Runtime/Core/Private/Windows/WindowsPlatformMisc.cpp
+++ b/Engine/Source/Runtime/Core/Private/Windows/WindowsPlatformMisc.cpp
@@ -3656,7 +3656,7 @@ bool FWindowsPlatformMisc::QueryRegKey( const Windows::HKEY InKey, const TCHAR*
 						if (RegQueryValueEx(Key, InValueName, NULL, NULL, (LPBYTE)Buffer, &Size) == ERROR_SUCCESS)
 						{
 							const uint32 Length = (Size / sizeof(TCHAR)) - 1;
-							OutData = FString(Length, (TCHAR*)Buffer);
+							OutData = FString::ConstructFromPtrSize((TCHAR*)Buffer, Length);
 							bSuccess = true;
 						}
 						delete[] Buffer;
diff --git a/Engine/Source/Runtime/Core/Public/Internationalization/IBreakIterator.h b/Engine/Source/Runtime/Core/Public/Internationalization/IBreakIterator.h
--- a/Engine/Source/Runtime/Core/Public/Internationalization/IBreakIterator.h
+++ b/Engine/Source/Runtime/Core/Public/Internationalization/IBreakIterator.h
@@ -20,7 +20,7 @@ public:
 	}
 	virtual void SetString(const TCHAR* const InString, const int32 InStringLength)
 	{
-		SetString(FString(InStringLength, InString));
+		SetString(FString::ConstructFromPtrSize(InString, (!InString || *InString == TEXT('\0')) ? 0 : InStringLength));
 	}
 	virtual void SetString(FString&& InString) = 0;
 	
diff --git a/Engine/Source/Runtime/Core/Public/Misc/ExpressionParserTypes.h b/Engine/Source/Runtime/Core/Public/Misc/ExpressionParserTypes.h
--- a/Engine/Source/Runtime/Core/Public/Misc/ExpressionParserTypes.h
+++ b/Engine/Source/Runtime/Core/Public/Misc/ExpressionParserTypes.h
@@ -41,7 +41,10 @@ public:
 	FStringToken() : TokenStart(nullptr), TokenEnd(nullptr), LineNumber(0), CharacterIndex(0) {}
 
 	/** Get the string representation of this token */
-	FString GetString() const { return FString((int32)(TokenEnd - TokenStart), TokenStart); }
+	FString GetString() const
+	{
+		return FString::ConstructFromPtrSize(TokenStart, (int32)(TokenEnd - TokenStart));
+	}
 
 	/** Check if this token is valid */
 	bool IsValid() const { return TokenEnd != TokenStart; }
diff --git a/Engine/Source/Runtime/Core/Public/Serialization/MemoryImage.h b/Engine/Source/Runtime/Core/Public/Serialization/MemoryImage.h
--- a/Engine/Source/Runtime/Core/Public/Serialization/MemoryImage.h
+++ b/Engine/Source/Runtime/Core/Public/Serialization/MemoryImage.h
@@ -772,7 +772,10 @@ public:
 		}
 	}
 
-	FORCEINLINE operator FString() const { return FString(Len(), Data.GetData()); }
+	FORCEINLINE operator FString() const
+	{
+		return FString::ConstructFromPtrSize(Data.GetData(), Len());
+	}
 
 	FORCEINLINE const TCHAR* operator*() const
 	{
diff --git a/Engine/Source/Runtime/CoreUObject/Private/UObject/FieldPathProperty.cpp b/Engine/Source/Runtime/CoreUObject/Private/UObject/FieldPathProperty.cpp
--- a/Engine/Source/Runtime/CoreUObject/Private/UObject/FieldPathProperty.cpp
+++ b/Engine/Source/Runtime/CoreUObject/Private/UObject/FieldPathProperty.cpp
@@ -138,7 +138,7 @@ const TCHAR* FFieldPathProperty::ImportText_Internal( const TCHAR* Buffer, void*
 			++SeparatorIndex;
 		}
 		// Copy the value string
-		PathName = FString(SeparatorIndex, Buffer);
+		PathName = FString::ConstructFromPtrSize(Buffer, SeparatorIndex);
 		// Advance the buffer to let the calling function know we succeeded
 		Buffer += SeparatorIndex;
 	}
diff --git a/Engine/Source/Runtime/CoreUObject/Public/AssetRegistry/AssetIdentifier.h b/Engine/Source/Runtime/CoreUObject/Public/AssetRegistry/AssetIdentifier.h
--- a/Engine/Source/Runtime/CoreUObject/Public/AssetRegistry/AssetIdentifier.h
+++ b/Engine/Source/Runtime/CoreUObject/Public/AssetRegistry/AssetIdentifier.h
@@ -72,7 +72,7 @@ struct FAssetIdentifier
 	{
 		TStringBuilder<256> Builder;
 		AppendString(Builder);
-		return FString(Builder.Len(), Builder.GetData());
+		return FString::ConstructFromPtrSize(Builder.GetData(), Builder.Len());
 	}
 
 	/** Appends to the given builder the string version of this identifier in Package.Object::Name format */
diff --git a/Engine/Source/Runtime/CoreUObject/Public/UObject/PrimaryAssetId.h b/Engine/Source/Runtime/CoreUObject/Public/UObject/PrimaryAssetId.h
--- a/Engine/Source/Runtime/CoreUObject/Public/UObject/PrimaryAssetId.h
+++ b/Engine/Source/Runtime/CoreUObject/Public/UObject/PrimaryAssetId.h
@@ -148,7 +148,7 @@ struct FPrimaryAssetId
 	{
 		TStringBuilder<256> Builder;
 		AppendString(Builder);
-		return FString(Builder.Len(), Builder.GetData());
+		return FString::ConstructFromPtrSize(Builder.GetData(), Builder.Len());
 	}
 
 	/** Appends to the given builder the string version of this identifier in Type:Name format */
