- No longer track missing CVars in OnSetCVarFromIniEntry() because there's no need to count misses here, it's expected that the editor will be missing a lot of the cvars that are defined in game modules, etc. This is expected that many are missing here.

diff --git a/Engine/Source/Runtime/Core/Private/Misc/ConfigUtilities.cpp b/Engine/Source/Runtime/Core/Private/Misc/ConfigUtilities.cpp
--- a/Engine/Source/Runtime/Core/Private/Misc/ConfigUtilities.cpp
+++ b/Engine/Source/Runtime/Core/Private/Misc/ConfigUtilities.cpp
@@ -147,7 +147,9 @@ void OnSetCVarFromIniEntry(const TCHAR *IniFile, const TCHAR *Key, const TCHAR*
 
 	Value = ConvertValueFromHumanFriendlyValue(Value);
 
-	IConsoleVariable* CVar = IConsoleManager::Get().FindConsoleVariable(Key); 
+	// we don't need to track cvar misses here (a lot will be not found early on in editor builds)
+	bool bTrackFrequentCalls = false;
+	IConsoleVariable* CVar = IConsoleManager::Get().FindConsoleVariable(Key, bTrackFrequentCalls);
 
 	if(CVar)
 	{
