[Backout] - CL29591063

(not backing out the fixes to structured logging itself, only its usage)

Original CL Desc
-----------------------------------------------------------------
Use structured logging for OutOfMemory messages

Structured logging allows to surface information contained in log lines more easily for automated processes. Out of memory errors are a common thing we'd like to detect and the current detection code has already been broken once.

This CL also contains a change to the structured log for fatal errors. Those were erroneously noted as assertion failures due to an unintended call to Debug::AssertFailed originating via copy-paste from FMsg::LogfImpl. Removing this call ensures that the new log matches the old behavior from unstructured logging (which I have manually verified).

diff --git a/Engine/Source/Runtime/Core/Private/GenericPlatform/GenericPlatformMemory.cpp b/Engine/Source/Runtime/Core/Private/GenericPlatform/GenericPlatformMemory.cpp
--- a/Engine/Source/Runtime/Core/Private/GenericPlatform/GenericPlatformMemory.cpp
+++ b/Engine/Source/Runtime/Core/Private/GenericPlatform/GenericPlatformMemory.cpp
@@ -8,7 +8,6 @@
 #include "Containers/StringConv.h"
 #include "UObject/NameTypes.h"
 #include "Logging/LogMacros.h"
-#include "Logging/StructuredLog.h"
 #include "Stats/Stats.h"
 #include "Containers/Ticker.h"
 #include "Misc/FeedbackContext.h"
@@ -249,11 +248,7 @@ void FGenericPlatformMemory::OnOutOfMemory(uint64 Size, uint32 Alignment)
 		FCoreDelegates::GetOutOfMemoryDelegate().Broadcast();
 
 		// ErrorMsg might be unrelated to OoM error in some cases as the code that calls OnOutOfMemory could have called other system functions that modified errno
-		UE_LOGFMT_NSLOC(LogMemory, Fatal, "Memory", "OutOfMemoryError", "Ran out of memory allocating {Bytes} bytes with alignment {Alignment}. Last error msg: {LastError}.",
-			("Bytes", Size),
-			("Alignment", Alignment),
-			("LastError", ErrorMsg)
-		);
+		UE_LOG(LogMemory, Fatal, TEXT("Ran out of memory allocating %llu bytes with alignment %u. Last error msg: %s."), Size, Alignment, ErrorMsg);
 	};
 	
 	UE_CALL_ONCE(HandleOOM);
diff --git a/Engine/Source/Runtime/Core/Private/IOS/IOSPlatformMemory.cpp b/Engine/Source/Runtime/Core/Private/IOS/IOSPlatformMemory.cpp
--- a/Engine/Source/Runtime/Core/Private/IOS/IOSPlatformMemory.cpp
+++ b/Engine/Source/Runtime/Core/Private/IOS/IOSPlatformMemory.cpp
@@ -10,7 +10,6 @@
 #include "HAL/MallocAnsi.h"
 #include "GenericPlatform/GenericPlatformMemoryPoolStats.h"
 #include "Misc/CoreDelegates.h"
-#include "Logging/StructuredLog.h"
 
 void FIOSPlatformMemory::OnOutOfMemory(uint64 Size, uint32 Alignment)
 {
@@ -53,11 +52,7 @@ void FIOSPlatformMemory::OnOutOfMemory(uint64 Size, uint32 Alignment)
 		FCoreDelegates::GetMemoryTrimDelegate().Broadcast();
     
 		// ErrorMsg might be unrelated to OoM error in some cases as the code that calls OnOutOfMemory could have called other system functions that modified errno
-		UE_LOGFMT_NSLOC(LogMemory, Warning, "Memory", "OutOfMemoryError", "Ran out of memory allocating {Bytes} bytes with alignment {Alignment}. Last error msg: {LastError}.",
-			("Bytes", Size),
-			("Alignment", Alignment),
-			("LastError", ErrorMsg)
-		);
+		UE_LOG(LogMemory, Warning, TEXT("Ran out of memory allocating %llu bytes with alignment %u. Last error msg: %s."), Size, Alignment, ErrorMsg);
     
 		// make this a fatal error that ends here not in the log
 		// changed to 3 from NULL because clang noticed writing to NULL and warned about it
diff --git a/Engine/Source/Runtime/Core/Private/Unix/UnixPlatformMemory.cpp b/Engine/Source/Runtime/Core/Private/Unix/UnixPlatformMemory.cpp
--- a/Engine/Source/Runtime/Core/Private/Unix/UnixPlatformMemory.cpp
+++ b/Engine/Source/Runtime/Core/Private/Unix/UnixPlatformMemory.cpp
@@ -14,7 +14,6 @@
 #include "Templates/UnrealTemplate.h"
 #include "Containers/UnrealString.h"
 #include "Logging/LogMacros.h"
-#include "Logging/StructuredLog.h"
 #include "HAL/MallocAnsi.h"
 #include "HAL/MallocMimalloc.h"
 #include "HAL/MallocJemalloc.h"
@@ -1197,11 +1196,7 @@ void FUnixPlatformMemory::OnOutOfMemory(uint64 Size, uint32 Alignment)
 		FCoreDelegates::GetOutOfMemoryDelegate().Broadcast();
 
 		// ErrorMsg might be unrelated to OoM error in some cases as the code that calls OnOutOfMemory could have called other system functions that modified errno
-		UE_LOGFMT_NSLOC(LogMemory, Fatal, "Memory", "OutOfMemoryError", "Ran out of memory allocating {Bytes} bytes with alignment {Alignment}. Last error msg: {LastError}.",
-			("Bytes", Size),
-			("Alignment", Alignment),
-			("LastError", ErrorMsg)
-		);
+		UE_LOG(LogMemory, Fatal, TEXT("Ran out of memory allocating %llu bytes with alignment %u. Last error msg: %s."), Size, Alignment, ErrorMsg);
 	};
 	
 	UE_CALL_ONCE(HandleOOM);
