Add csv profiling to AllocVirtual and BinnedAllocFromOS on windows/unix

diff --git a/Engine/Source/Runtime/Core/Private/Unix/UnixPlatformMemory.cpp b/Engine/Source/Runtime/Core/Private/Unix/UnixPlatformMemory.cpp
--- a/Engine/Source/Runtime/Core/Private/Unix/UnixPlatformMemory.cpp
+++ b/Engine/Source/Runtime/Core/Private/Unix/UnixPlatformMemory.cpp
@@ -22,6 +22,7 @@
 #include "HAL/MallocReplayProxy.h"
 #include "HAL/MallocStomp.h"
 #include "HAL/PlatformMallocCrash.h"
+#include "ProfilingDebugging/CsvProfiler.h"
 
 #include <sys/sysinfo.h>
 #include <sys/file.h>
@@ -94,6 +95,8 @@ namespace
 	bool GEnableProtectForkedPages = false;
 }
 
+CSV_DECLARE_CATEGORY_EXTERN(FMemory);
+
 /** Controls growth of pools - see PooledVirtualMemoryAllocator.cpp */
 extern float GVMAPoolScale;
 
@@ -428,6 +431,15 @@ struct FOSAllocationDescriptor
 
 void* FUnixPlatformMemory::BinnedAllocFromOS(SIZE_T Size)
 {
+#if UE_ENABLE_PLATFORM_MEMORY_CSV_STATS
+	const bool bIsInGameThread = IsInGameThread();
+	CSV_SCOPED_TIMING_STAT_RECURSIVE_CONDITIONAL(FMemory, BinnedAllocFromOSTime, bIsInGameThread);
+	if (bIsInGameThread)
+	{
+		CSV_CUSTOM_STAT(FMemory, BinnedAllocFromOSCount, 1, ECsvCustomStatOp::Accumulate);
+	}
+#endif // UE_ENABLE_PLATFORM_MEMORY_CSV_STATS
+
 #if UE_CHECK_LARGE_ALLOCATIONS
 	if (UE::Memory::Private::GEnableLargeAllocationChecks)
 	{
@@ -618,6 +630,15 @@ size_t FUnixPlatformMemory::FPlatformVirtualMemoryBlock::GetCommitAlignment()
 
 FUnixPlatformMemory::FPlatformVirtualMemoryBlock FUnixPlatformMemory::FPlatformVirtualMemoryBlock::AllocateVirtual(size_t InSize, size_t InAlignment)
 {
+#if UE_ENABLE_PLATFORM_MEMORY_CSV_STATS
+	const bool bIsInGameThread = IsInGameThread();
+	CSV_SCOPED_TIMING_STAT_RECURSIVE_CONDITIONAL(FMemory, AllocateVirtualTime, bIsInGameThread);
+	if (bIsInGameThread)
+	{
+		CSV_CUSTOM_STAT(FMemory, AllocateVirtualCount, 1, ECsvCustomStatOp::Accumulate);
+	}
+#endif // UE_ENABLE_PLATFORM_MEMORY_CSV_STATS
+
 	FPlatformVirtualMemoryBlock Result;
 	InSize = Align(InSize, GetVirtualSizeAlignment());
 	Result.VMSizeDivVirtualSizeAlignment = InSize / GetVirtualSizeAlignment();
diff --git a/Engine/Source/Runtime/Core/Private/Windows/WindowsPlatformMemory.cpp b/Engine/Source/Runtime/Core/Private/Windows/WindowsPlatformMemory.cpp
--- a/Engine/Source/Runtime/Core/Private/Windows/WindowsPlatformMemory.cpp
+++ b/Engine/Source/Runtime/Core/Private/Windows/WindowsPlatformMemory.cpp
@@ -26,6 +26,7 @@
 #include "Misc/OutputDeviceRedirector.h"
 #include "Stats/Stats.h"
 #include "Windows/AllowWindowsPlatformTypes.h"
+#include "ProfilingDebugging/CsvProfiler.h"
 
 #include <MemoryApi.h> // Include after AllowWindowsPlatformTypes.h
 
@@ -46,6 +47,7 @@
 
 DECLARE_MEMORY_STAT(TEXT("Windows Specific Memory Stat"),	STAT_WindowsSpecificMemoryStat, STATGROUP_MemoryPlatform);
 
+CSV_DECLARE_CATEGORY_EXTERN(FMemory);
 
 static int32 GWindowsPlatformMemoryGetStatsLimitTotalGB = 0;
 static FAutoConsoleVariableRef CVarLogPlatformMemoryStats(
@@ -440,6 +442,15 @@ bool FWindowsPlatformMemory::PageProtect(void* const Ptr, const SIZE_T Size, con
 }
 void* FWindowsPlatformMemory::BinnedAllocFromOS( SIZE_T Size )
 {
+#if UE_ENABLE_PLATFORM_MEMORY_CSV_STATS
+	const bool bIsInGameThread = IsInGameThread();
+	CSV_SCOPED_TIMING_STAT_RECURSIVE_CONDITIONAL(FMemory, BinnedAllocFromOSTime, bIsInGameThread);
+	if (bIsInGameThread)
+	{
+		CSV_CUSTOM_STAT(FMemory, BinnedAllocFromOSCount, 1, ECsvCustomStatOp::Accumulate);
+	}
+#endif // UE_ENABLE_PLATFORM_MEMORY_CSV_STATS
+
 #if UE_CHECK_LARGE_ALLOCATIONS
 	if (UE::Memory::Private::GEnableLargeAllocationChecks)
 	{
@@ -477,6 +488,15 @@ size_t FWindowsPlatformMemory::FPlatformVirtualMemoryBlock::GetCommitAlignment()
 
 FWindowsPlatformMemory::FPlatformVirtualMemoryBlock FWindowsPlatformMemory::FPlatformVirtualMemoryBlock::AllocateVirtual(size_t InSize, size_t InAlignment)
 {
+#if UE_ENABLE_PLATFORM_MEMORY_CSV_STATS
+	const bool bIsInGameThread = IsInGameThread();
+	CSV_SCOPED_TIMING_STAT_RECURSIVE_CONDITIONAL(FMemory, AllocateVirtualTime, bIsInGameThread);
+	if (bIsInGameThread)
+	{
+		CSV_CUSTOM_STAT(FMemory, AllocateVirtualCount, 1, ECsvCustomStatOp::Accumulate);
+	}
+#endif // UE_ENABLE_PLATFORM_MEMORY_CSV_STATS
+
 	FPlatformVirtualMemoryBlock Result;
 	InSize = Align(InSize, GetVirtualSizeAlignment());
 	Result.VMSizeDivVirtualSizeAlignment = InSize / GetVirtualSizeAlignment();
diff --git a/Engine/Source/Runtime/Core/Public/GenericPlatform/GenericPlatformMemory.h b/Engine/Source/Runtime/Core/Public/GenericPlatform/GenericPlatformMemory.h
--- a/Engine/Source/Runtime/Core/Public/GenericPlatform/GenericPlatformMemory.h
+++ b/Engine/Source/Runtime/Core/Public/GenericPlatform/GenericPlatformMemory.h
@@ -35,6 +35,10 @@ namespace UE::Memory::Private
 }
 #endif
 
+#ifndef UE_ENABLE_PLATFORM_MEMORY_CSV_STATS
+#define UE_ENABLE_PLATFORM_MEMORY_CSV_STATS 0
+#endif
+
 /**
 * Platform-dependent "bucket" for memory size, where Default is the normal, or possibly the largest.
 * This is generally used for texture LOD settings for how to fit in smaller memory devices
diff --git a/Engine/Source/Runtime/Core/Public/ProfilingDebugging/CsvProfiler.h b/Engine/Source/Runtime/Core/Public/ProfilingDebugging/CsvProfiler.h
--- a/Engine/Source/Runtime/Core/Public/ProfilingDebugging/CsvProfiler.h
+++ b/Engine/Source/Runtime/Core/Public/ProfilingDebugging/CsvProfiler.h
@@ -65,6 +65,14 @@ struct FCsvDeclaredStat;
 #define CSV_SCOPED_TIMING_STAT_EXCLUSIVE_CONDITIONAL(StatName,Condition) \
 	TRACE_CSV_PROFILER_INLINE_STAT_EXCLUSIVE(#StatName); \
 	FScopedCsvStatExclusiveConditional _ScopedCsvStatExclusive_ ## StatName (#StatName,Condition, "CSV_"#StatName);
+#define CSV_SCOPED_TIMING_STAT_RECURSIVE(Category,StatName) \
+	TRACE_CSV_PROFILER_INLINE_STAT(#StatName, CSV_CATEGORY_INDEX(Category)); \
+	static int32 _ScopedCsvStatRecursive_EntryCount_ ## StatName = 0; \
+	FScopedCsvStatRecursive _ScopedCsvStatRecursive_ ## StatName (_ScopedCsvStatRecursive_EntryCount_ ## StatName, #StatName, CSV_CATEGORY_INDEX(Category), "CSV_"#StatName);
+#define CSV_SCOPED_TIMING_STAT_RECURSIVE_CONDITIONAL(Category,StatName,Condition) \
+	TRACE_CSV_PROFILER_INLINE_STAT(#StatName, CSV_CATEGORY_INDEX(Category)); \
+	static int32 _ScopedCsvStatRecursive_EntryCount_ ## StatName = 0; \
+	FScopedCsvStatRecursiveConditional _ScopedCsvStatRecursive_ ## StatName (_ScopedCsvStatRecursive_EntryCount_ ## StatName, #StatName, CSV_CATEGORY_INDEX(Category), Condition, "CSV_"#StatName);
 
 #define CSV_SCOPED_WAIT(WaitTime)							FScopedCsvWaitConditional _ScopedCsvWait(WaitTime>0 && FCsvProfiler::IsWaitTrackingEnabledOnCurrentThread());
 #define CSV_SCOPED_WAIT_CONDITIONAL(Condition)				FScopedCsvWaitConditional _ScopedCsvWait(Condition);
@@ -116,6 +124,8 @@ struct FCsvDeclaredStat;
   #define CSV_SCOPED_TIMING_STAT_GLOBAL(StatName)					
   #define CSV_SCOPED_TIMING_STAT_EXCLUSIVE(StatName)
   #define CSV_SCOPED_TIMING_STAT_EXCLUSIVE_CONDITIONAL(StatName,Condition)
+  #define CSV_SCOPED_TIMING_STAT_RECURSIVE(Category,StatName)
+  #define CSV_SCOPED_TIMING_STAT_RECURSIVE_CONDITIONAL(Category,StatName,Condition)
   #define CSV_SCOPED_WAIT(WaitTime)
   #define CSV_SCOPED_WAIT_CONDITIONAL(Condition)
   #define CSV_SCOPED_SET_WAIT_STAT(StatName)
@@ -577,6 +587,68 @@ public:
 	bool bCondition;
 };
 
+class FScopedCsvStatRecursive
+{
+	const char* StatName;
+	uint32 CategoryIndex;
+	int32& EntryCounter;
+public:
+	FScopedCsvStatRecursive(int32& InEntryCounter, const char* InStatName, uint32 InCategoryIndex, const char* InNamedEventName = nullptr)
+		: StatName(InStatName)
+		, CategoryIndex(InCategoryIndex)
+		, EntryCounter(InEntryCounter)
+	{
+		++EntryCounter; // this needs to happen before BeginStat in case BeginStat causes reentry
+		if (EntryCounter == 1)
+		{
+			FCsvProfiler::BeginStat(StatName, CategoryIndex, InNamedEventName);
+		}
+	}
+	~FScopedCsvStatRecursive()
+	{
+		if (EntryCounter == 1)
+		{
+			FCsvProfiler::EndStat(StatName, CategoryIndex);
+		}
+		--EntryCounter;
+	}
+};
+
+class FScopedCsvStatRecursiveConditional
+{
+	const char* StatName;
+	uint32 CategoryIndex;
+	int32& EntryCounter;
+	bool bCondition;
+public:
+	FScopedCsvStatRecursiveConditional(int32& InEntryCounter, const char* InStatName, uint32 InCategoryIndex, bool bInCondition, const char* InNamedEventName = nullptr)
+		: StatName(InStatName)
+		, CategoryIndex(InCategoryIndex)
+		, EntryCounter(InEntryCounter)
+		, bCondition(bInCondition)
+	{
+		if (bCondition)
+		{
+			++EntryCounter; // this needs to happen before BeginStat in case BeginStat causes reentry
+			if (EntryCounter == 1)
+			{
+				FCsvProfiler::BeginStat(StatName, CategoryIndex, InNamedEventName);
+			}
+		}
+	}
+	~FScopedCsvStatRecursiveConditional()
+	{
+		if (bCondition)
+		{
+			if (EntryCounter == 1)
+			{
+				FCsvProfiler::EndStat(StatName, CategoryIndex);
+			}
+			--EntryCounter;
+		}
+	}
+};
+
 class FScopedCsvWaitConditional
 {
 public:
