Disable AutoRTFM when compiled with AutoRTFM support

diff --git a/Engine/Source/Runtime/Core/Private/AutoRTFM/API.cpp b/Engine/Source/Runtime/Core/Private/AutoRTFM/API.cpp
--- a/Engine/Source/Runtime/Core/Private/AutoRTFM/API.cpp
+++ b/Engine/Source/Runtime/Core/Private/AutoRTFM/API.cpp
@@ -5,7 +5,7 @@
 #include "HAL/IConsoleManager.h"
 
 #if UE_AUTORTFM
-bool GAutoRTFMRuntimeEnabled = true;
+bool GAutoRTFMRuntimeEnabled = false;
 static FAutoConsoleVariableRef CVarAutoRTFMRuntimeEnabled(
 	TEXT("AutoRTFMRuntimeEnabled"),
 	GAutoRTFMRuntimeEnabled,
@@ -16,6 +16,17 @@ static FAutoConsoleVariableRef CVarAutoRTFMRuntimeEnabled(
 static constexpr bool GAutoRTFMRuntimeEnabled = false;
 #endif
 
+namespace AutoRTFM
+{
+	void SetAutoRTFMRuntime(bool bEnabled)
+	{
+		// #noop if AutoRTFM is not compiled in, as GAutoRTFMRuntimeEnabled is a static constexpr when no AutoRTFM compiled in
+#if UE_AUTORTFM
+		GAutoRTFMRuntimeEnabled = bEnabled;
+#endif
+	}
+}
+
 #if (defined(__AUTORTFM) && __AUTORTFM)
 #include "AutoRTFM/AutoRTFMConstants.h"
 #include "CallNest.h"
diff --git a/Engine/Source/Runtime/Core/Public/AutoRTFM/AutoRTFM.h b/Engine/Source/Runtime/Core/Public/AutoRTFM/AutoRTFM.h
--- a/Engine/Source/Runtime/Core/Public/AutoRTFM/AutoRTFM.h
+++ b/Engine/Source/Runtime/Core/Public/AutoRTFM/AutoRTFM.h
@@ -403,6 +403,8 @@ enum class EContextStatus
 	AbortedByCascade = autortfm_status_aborted_by_cascade
 };
 
+UE_AUTORTFM_API void SetAutoRTFMRuntime(bool bEnabled);
+
 UE_AUTORTFM_FORCEINLINE bool IsTransactional() { return autortfm_is_transactional(); }
 UE_AUTORTFM_FORCEINLINE bool IsClosed() { return autortfm_is_closed(); }
 
