- Restored the conversion of "True" to "1" for other-platform CVar handling. This is needed for when we have an <int> type CVar, but the .ini has =True

diff --git a/Engine/Source/Runtime/Core/Private/HAL/ConsoleManager.cpp b/Engine/Source/Runtime/Core/Private/HAL/ConsoleManager.cpp
--- a/Engine/Source/Runtime/Core/Private/HAL/ConsoleManager.cpp
+++ b/Engine/Source/Runtime/Core/Private/HAL/ConsoleManager.cpp
@@ -650,7 +650,7 @@ public:
 	void Track(const TCHAR* InValue, EConsoleVariableFlags SetBy, FName Tag)
 	{
 		T LocalCopy;
-		TTypeFromString<T>::FromString(LocalCopy, InValue);
+		TTypeFromString<T>::FromString(LocalCopy, UE::ConfigUtilities::ConvertValueFromHumanFriendlyValue(InValue));
 		
 		int Priority = (int)SetBy;
 		TArray<FTaggedHistoryData>& ValueArray = History.FindOrAdd(Priority);
@@ -837,7 +837,7 @@ protected:
 		{
 			// update value
 			T ConvertedValue;
-			TTypeFromString<T>::FromString(ConvertedValue, InValue);
+			TTypeFromString<T>::FromString(ConvertedValue, UE::ConfigUtilities::ConvertValueFromHumanFriendlyValue(InValue));
 			SetInternal(ConvertedValue, SetBy);
 			
 			// update the setby
