Add a new CVar, Automation.LightweightStereoTestVariants, which stops variants from running if we've already encounted an error, and skips saving/attaching screenshots for successful variants to save space. Enabled by default. Also moves the "Supports Stereo Test Variants" checkbox to the "Screenshot" section.

diff --git a/Engine/Source/Developer/AutomationController/Private/AutomationControllerManager.cpp b/Engine/Source/Developer/AutomationController/Private/AutomationControllerManager.cpp
--- a/Engine/Source/Developer/AutomationController/Private/AutomationControllerManager.cpp
+++ b/Engine/Source/Developer/AutomationController/Private/AutomationControllerManager.cpp
@@ -501,7 +501,7 @@ void FAutomationControllerManager::ProcessComparisonQueue()
 
 			FImageComparisonResult Result = Entry->PendingComparison.Get();
 			FAutomationWorkerImageComparisonResults ResultMessage(
-				FGuid::NewGuid(),
+				Result.bSkipAttachingImages ? FGuid() : FGuid::NewGuid(),
 				Result.ScreenshotName,
 				Result.IsNew(),
 				Result.AreSimilar(),
@@ -526,7 +526,10 @@ void FAutomationControllerManager::ProcessComparisonQueue()
 				MessageEndpoint->Send(Message, Entry->Sender);
 			}
 
-			ReportImageComparisonResult(Entry->Sender, ResultMessage);
+			if (!Result.bSkipAttachingImages)
+			{
+				ReportImageComparisonResult(Entry->Sender, ResultMessage);
+			}
 		}
 	}
 }
diff --git a/Engine/Source/Developer/FunctionalTesting/Classes/ScreenshotFunctionalTest.h b/Engine/Source/Developer/FunctionalTesting/Classes/ScreenshotFunctionalTest.h
--- a/Engine/Source/Developer/FunctionalTesting/Classes/ScreenshotFunctionalTest.h
+++ b/Engine/Source/Developer/FunctionalTesting/Classes/ScreenshotFunctionalTest.h
@@ -29,7 +29,7 @@ public:
 	bool bCameraCutOnScreenshotPrep;
 
 	// If r.EnableStereoScreenshotTestVariants is enabled, also take and compare screenshots in stereo rendering emulation modes when running this test
-	UPROPERTY(EditAnywhere, BlueprintReadWrite, Category = "Functional Testing", SimpleDisplay)
+	UPROPERTY(EditAnywhere, BlueprintReadWrite, Category = "Screenshot", SimpleDisplay)
 	bool bSupportStereoTestVariants = true;
 
 protected:
diff --git a/Engine/Source/Developer/FunctionalTesting/Private/ScreenshotFunctionalTest.cpp b/Engine/Source/Developer/FunctionalTesting/Private/ScreenshotFunctionalTest.cpp
--- a/Engine/Source/Developer/FunctionalTesting/Private/ScreenshotFunctionalTest.cpp
+++ b/Engine/Source/Developer/FunctionalTesting/Private/ScreenshotFunctionalTest.cpp
@@ -154,7 +154,10 @@ void AScreenshotFunctionalTest::OnScreenShotCaptured(int32 InSizeX, int32 InSize
 
 void AScreenshotFunctionalTest::OnScreenshotTakenAndCompared()
 {
-	if (bShouldDoViewRectOffsetVariant)
+	FAutomationTestBase* CurrentTest = FAutomationTestFramework::Get().GetCurrentTest();
+	bool bSkipDueToError = FAutomationTestFramework::Get().NeedUseLightweightStereoTestVariants() && (!CurrentTest || CurrentTest->HasAnyErrors());
+
+	if (bShouldDoViewRectOffsetVariant && !bSkipDueToError)
 	{
 		bShouldDoViewRectOffsetVariant = false;
 
diff --git a/Engine/Source/Developer/ScreenShotComparisonTools/Private/ScreenShotManager.cpp b/Engine/Source/Developer/ScreenShotComparisonTools/Private/ScreenShotManager.cpp
--- a/Engine/Source/Developer/ScreenShotComparisonTools/Private/ScreenShotManager.cpp
+++ b/Engine/Source/Developer/ScreenShotComparisonTools/Private/ScreenShotManager.cpp
@@ -391,6 +391,14 @@ FImageComparisonResult FScreenShotManager::CompareScreenshot(const FString& InUn
 	// We use the subfolder path to the screenshot as name (before any environment specialization - platform, RHI - are appended).
 	ComparisonResult.ScreenshotName = ResultsSubFolder;
 
+	// Do not save passing variant test screenshots
+	const bool bIsVariant = !IncomingMetaData.VariantName.IsEmpty();
+	if (bIsVariant && FAutomationTestFramework::Get().NeedUseLightweightStereoTestVariants() && ComparisonResult.AreSimilar())
+	{
+		ComparisonResult.bSkipAttachingImages = true;
+		return ComparisonResult;
+	}
+
 	// Result paths should be relative to the project. Note this may be empty, and if it is MakePathRelative returns
 	// a non empty relative path... but we want it to stay empty as that's how we signal that no approved file exists
 	if (!ComparisonResult.ApprovedFilePath.IsEmpty())
diff --git a/Engine/Source/Developer/ScreenShotComparisonTools/Public/ImageComparer.h b/Engine/Source/Developer/ScreenShotComparisonTools/Public/ImageComparer.h
--- a/Engine/Source/Developer/ScreenShotComparisonTools/Public/ImageComparer.h
+++ b/Engine/Source/Developer/ScreenShotComparisonTools/Public/ImageComparer.h
@@ -303,24 +303,32 @@ public:
 	FText ErrorMessage;
 
 	/*
-		Name of the screenshot
+		Name of the screenshot (includes variant if applicable)
 	*/
 	UPROPERTY()
 	FString ScreenshotName;
 
+	/*
+		Whether to skip saving and attaching images to the report for this test
+	*/
+	UPROPERTY()
+	bool bSkipAttachingImages;
+
 	/*
 		Version of the image comparision result 
 	*/
 	UPROPERTY()
 	int32 Version;
 
-	static constexpr int32 CurrentVersion = 2;
+	static constexpr int32 CurrentVersion = 3;
+	static constexpr int32 OldestSupportedVersion = 2;
 
 	FImageComparisonResult()
 		: CreationTime(0)
 		, MaxLocalDifference(0.0f)
 		, GlobalDifference(0.0f)
 		, ErrorMessage()
+		, bSkipAttachingImages(false)
 		, Version(CurrentVersion)
 	{
 	}
@@ -330,6 +338,7 @@ public:
 		, MaxLocalDifference(0.0f)
 		, GlobalDifference(0.0f)
 		, ErrorMessage(Error)
+		, bSkipAttachingImages(false)
 		, Version(CurrentVersion)
 	{
 	}
@@ -339,7 +348,7 @@ public:
 	*/
 	bool IsValid() const
 	{
-		return Version == CurrentVersion;
+		return Version >= OldestSupportedVersion && Version <= CurrentVersion;
 	}
 	
 	/*
diff --git a/Engine/Source/Runtime/Core/Private/Misc/AutomationTest.cpp b/Engine/Source/Runtime/Core/Private/Misc/AutomationTest.cpp
--- a/Engine/Source/Runtime/Core/Private/Misc/AutomationTest.cpp
+++ b/Engine/Source/Runtime/Core/Private/Misc/AutomationTest.cpp
@@ -57,6 +57,12 @@ namespace AutomationTest
 		bEnableStereoTestVariants,
 		TEXT("Whether to enable stereo test variants for screenshot functional tests"));
 
+	static bool bLightweightStereoTestVariants = true;
+	static FAutoConsoleVariableRef CVarAutomationLightweightStereoTestVariants(
+		TEXT("Automation.LightweightStereoTestVariants"),
+		bLightweightStereoTestVariants,
+		TEXT("Whether to skip variants when the baseline test fails, and skip saving screenshots for successful variants"));
+
 	// The method prepares the filename and LineNumber to be placed in the form that could be extracted by SAutomationWindow widget if it is additionally eclosed into []
 	// The result format is filename(line)
 	static FString CreateFileLineDescription(const FString& Filename, const int32 LineNumber)
@@ -312,6 +318,11 @@ bool FAutomationTestFramework::NeedPerformStereoTestVariants()
 	return AutomationTest::bEnableStereoTestVariants;
 }
 
+bool FAutomationTestFramework::NeedUseLightweightStereoTestVariants()
+{
+	return AutomationTest::bLightweightStereoTestVariants;
+}
+
 bool FAutomationTestFramework::RegisterAutomationTest( const FString& InTestNameToRegister, FAutomationTestBase* InTestToRegister )
 {
 	const bool bAlreadyRegistered = AutomationTestClassNameToInstanceMap.Contains( InTestNameToRegister );
diff --git a/Engine/Source/Runtime/Core/Public/Misc/AutomationTest.h b/Engine/Source/Runtime/Core/Public/Misc/AutomationTest.h
--- a/Engine/Source/Runtime/Core/Public/Misc/AutomationTest.h
+++ b/Engine/Source/Runtime/Core/Public/Misc/AutomationTest.h
@@ -1160,10 +1160,15 @@ public:
 	static CORE_API bool NeedLogBPTestMetadata();
 
 	/**
-	 * Whether to also run stereo test variants for screenshot functional testts
+	 * Whether to also run stereo test variants for screenshot functional tests
 	 */
 	static CORE_API bool NeedPerformStereoTestVariants();
 
+	/**
+	 * Whether to skip variants when the baseline test fails, and skip saving screenshots for successful variants
+	 */
+	static CORE_API bool NeedUseLightweightStereoTestVariants();
+
 	/**
 	 * Notify that the screenshot comparison has completed
 	 */
