Adding an LLM console command to associate allocations with their underlying page allocation in Unix to classify them as shared/private/unreferenced.
This also fixes the old smaps parser to not incorrectly bail before EOF if fewer bytes were provided than asked for as that's an allowable condition.

diff --git a/Engine/Source/Runtime/Core/Private/HAL/LowLevelMemTracker.cpp b/Engine/Source/Runtime/Core/Private/HAL/LowLevelMemTracker.cpp
--- a/Engine/Source/Runtime/Core/Private/HAL/LowLevelMemTracker.cpp
+++ b/Engine/Source/Runtime/Core/Private/HAL/LowLevelMemTracker.cpp
@@ -7,9 +7,12 @@
 #include "HAL/LowLevelMemStats.h"
 #include "HAL/FileManager.h"
 #include "HAL/IConsoleManager.h"
+#include "HAL/PlatformMemory.h" // for page allocation association.
 #include "LowLevelMemTrackerPrivate.h"
 #include "MemPro/MemProProfiler.h"
 #include "Misc/CString.h"
+#include "Misc/FileHelper.h"
+#include "Misc/Fork.h"
 #include "Misc/OutputDeviceRedirector.h"
 #include "Misc/Paths.h"
 #include "Misc/ScopeLock.h"
@@ -461,6 +464,12 @@ public:
 
 	void Clear();
 
+	// Dump the allocation count and size for each tag, along with the porportion that is private/shared/unreferenced
+	// on linux machines. This will dump to a CSV in the project Saved/LLM directory with a separate file for each
+	// tracker and forked child. Integers will be appended to the filename to prevent overwrites. Returns false
+	// if the memory inforation can't be retrieved (likely because it's not on a supported platform)
+	bool DumpForkedAllocationInfo();
+
 	void PublishStats(UE::LLM::ESizeParams SizeParam);
 	void PublishCsv(UE::LLM::ESizeParams SizeParam);
 	void PublishTrace(UE::LLM::ESizeParams SizeParam);
@@ -1453,6 +1462,17 @@ bool FLowLevelMemTracker::Exec(const TCHAR* Cmd, FOutputDevice& Ar)
 
 			UpdateStatsPerFrame(TEXT("After cleanup"));
 		}
+		else if (FParse::Command(&Cmd, TEXT("DUMPPRIVATESHARED")))
+		{
+			for (int32 TrackerIndex = 0; TrackerIndex < static_cast<int32>(ELLMTracker::Max); TrackerIndex++)
+			{
+				if (GetTracker((ELLMTracker)TrackerIndex)->DumpForkedAllocationInfo() == false)
+				{
+					FPlatformMisc::LowLevelOutputDebugString(TEXT("Failed to dumping forked allocation info (check platform supports it?)\n"));
+					return true;
+				}
+			}
+		}
 		else if (FParse::Command(&Cmd, TEXT("REMEMBER")))
 		{
 			CSV_EVENT_GLOBAL(TEXT("LLM_REMEMBER"));
@@ -4069,6 +4089,228 @@ void FLLMTracker::PublishStats(UE::LLM::ESizeParams SizeParams)
 	}
 }
 
+static bool PageLessThan(const FForkedPageAllocation& LHS, const FForkedPageAllocation& RHS)
+{
+	return LHS.PageStart < RHS.PageStart;
+}
+static bool PageFinder(const FForkedPageAllocation& LHS, const uint64 AddressRHS)
+{
+	return LHS.PageStart < AddressRHS;
+}
+
+bool FLLMTracker::DumpForkedAllocationInfo()
+{
+	// Try to associate the allocations with a page range so we can determine whether it's
+	// in unique or shared memory. Then print out the set for each tag.
+	TArray<FForkedPageAllocation> Pages;
+	if (FGenericPlatformMemory::GetForkedPageAllocationInfo(Pages) == false)
+	{
+		return false;
+	}
+
+	Algo::Sort(Pages, PageLessThan);
+
+	enum EClassification
+	{
+		Class_Split,
+		Class_Unreferenced,
+		Class_Private,
+		Class_Shared,
+		Class_COUNT
+	};
+
+	struct FCounts
+	{
+		uint64 TotalAllocations = 0;
+		uint64 CrossPageAllocationCount = 0;
+
+		uint64 AllocCount[Class_COUNT] = {};
+
+		uint64 TotalBytes = 0;
+		uint64 ByteCount[Class_COUNT] = {};
+	};
+
+	
+	TMap< const FTagData*, FCounts> CountsPerTag;
+	AllocationMap.LockAll();
+	for (const FLLMAllocMap::FTuple& Tuple : AllocationMap)
+	{
+		void* Ptr = Tuple.Key.GetPointer();
+
+		int ContainingAllocationIndex = Algo::LowerBound(Pages, (uint64)Ptr, PageFinder);
+		if (ContainingAllocationIndex == Pages.Num())
+		{
+			UE_LOG(LogHAL, Error, TEXT("Can't find allocation 0x%llx in the pages list!"), (uint64)Ptr);
+			continue;
+		}
+
+		// We unfortunately don't know exactly which of the details apply because we don't know which pages
+		// in the allocation section are shared/unique etc, so we just keep some generalities.
+		uint64 SizeHigh = Tuple.Key.GetExtraData() << 32;
+		uint64 SizeLow = Tuple.Value1;
+		uint64 Size = SizeHigh | SizeLow;
+
+		// Number of page allocations for the ue allocation
+		uint32 PageAllocationCount = 0;
+
+		// Track the classification for each page allocation so we can try to classify the ue allocation
+		uint32 SplitCount = 0;
+		uint32 SharedCount = 0;
+		uint32 PrivateCount = 0;
+		uint32 UnreferencedCount = 0;
+		
+		// Walk the page allocations until we get them all if we aren't all in the same one.
+		uint64 Remaining = Size;
+		while (Remaining)
+		{
+			PageAllocationCount++;
+			FForkedPageAllocation& Page = Pages[ContainingAllocationIndex];
+
+			uint64 PageKiB = (Page.PageEnd - Page.PageStart) / 1024;
+
+			uint64 OffsetInPage = (uint64)Ptr - Page.PageStart;
+			uint64 RemainingInPage = Page.PageEnd - (uint64)Ptr;
+
+			uint64 AmountInThisPage = Remaining;
+			if (AmountInThisPage > RemainingInPage)
+			{
+				AmountInThisPage = RemainingInPage;
+			}
+
+			uint64 SharedKiB = Page.SharedCleanKiB + Page.SharedDirtyKiB;
+			uint64 PrivateKiB = Page.PrivateCleanKiB + Page.PrivateDirtyKiB;
+			uint64 UnreferencedKiB = PageKiB - SharedKiB - PrivateKiB;
+
+			// We can't classify the page if we have more than one of any.
+			uint64 TypeCount = !!SharedKiB + !!PrivateKiB + !!UnreferencedKiB;
+			bool bIsSplit = TypeCount > 1;
+
+			if (bIsSplit)
+			{
+				// The page is split and we have no way to know which is ours.
+				SplitCount++;
+			}
+			else if (SharedKiB)
+			{
+				SharedCount++;
+			}
+			else if (PrivateKiB)
+			{
+				PrivateCount++;
+			}
+			else
+			{
+				UnreferencedCount++;
+			}
+
+			Remaining -= AmountInThisPage;
+			ContainingAllocationIndex++;
+			if (ContainingAllocationIndex >= Pages.Num())
+			{
+				UE_LOG(LogHAL, Error, TEXT("Allocation 0x%llu extended beyond the pages!"), (uint64)Ptr);
+				break;
+			}
+		}
+
+		if (Remaining)
+		{
+			// We error'd out, ignore this allocation.
+			continue;
+		}
+
+		const FTagData* Tag = Tuple.Value2.GetTag(LLMRef);
+
+		// Classify the allocation - is it entirely private, entirely shared, or what.
+		uint64 ClassCount = !!SharedCount + !!PrivateCount + !!UnreferencedCount;
+		EClassification AllocClass = Class_Split;
+		if (ClassCount > 1)
+		{
+			AllocClass = Class_Split;
+		}
+		else if (SharedCount)
+		{
+			AllocClass = Class_Shared;
+		}
+		else if (PrivateCount)
+		{
+			AllocClass = Class_Private;
+		}
+		else
+		{
+			AllocClass = Class_Unreferenced;
+		}
+
+		// Associate the tag and the page stats.
+		FCounts& Counts = CountsPerTag.FindOrAdd(Tag);
+		Counts.AllocCount[AllocClass]++;
+		Counts.ByteCount[AllocClass] += Size;
+		Counts.TotalAllocations++;
+		Counts.TotalBytes += Size;
+		if (PageAllocationCount > 1)
+		{
+			Counts.CrossPageAllocationCount++;
+		}
+	}
+	AllocationMap.UnlockAll();
+
+	CountsPerTag.ValueSort([](const FCounts& A, const FCounts& B)
+	{
+		return A.TotalBytes > B.TotalBytes;
+	});
+
+	TArray<FString> Lines;
+	Lines.Reserve(CountsPerTag.Num() * 2 + 1);
+
+	Lines.Add(TEXT("Tag,SharedKib,PrivateKib,SplitKib,UnrefKib,TotalKib,SharedCount,PrivateCount,SplitCount,UnrefCount,TotalCount,CrossCount"));
+
+
+	for (TPair<const FTagData*, FCounts>& P : CountsPerTag)
+	{
+		Lines.Add(
+			FString::Printf(TEXT("%s,%llu,%llu,%llu,%llu,%llu,%llu,%llu,%llu,%llu,%llu,%llu"),
+				*P.Key->GetDisplayName().ToString(),
+				P.Value.ByteCount[Class_Shared] / (1024),
+				P.Value.ByteCount[Class_Private] / (1024),
+				P.Value.ByteCount[Class_Split] / (1024),
+				P.Value.ByteCount[Class_Unreferenced] / (1024),
+				P.Value.TotalBytes / (1024),
+				P.Value.AllocCount[Class_Shared],
+				P.Value.AllocCount[Class_Private],
+				P.Value.AllocCount[Class_Split],
+				P.Value.AllocCount[Class_Unreferenced],
+				P.Value.TotalAllocations,
+				P.Value.CrossPageAllocationCount)
+			);
+	}
+
+	FString SavedDir = FPaths::ProjectSavedDir() / TEXT("LLM");
+
+	uint16 ForkId = FForkProcessHelper::GetForkedChildProcessIndex();
+
+	FString ForkString = ForkId == 0 ? FString(TEXT("Parent")) : FString::Printf(TEXT("Child%d"), ForkId);
+	const TCHAR* TrackerString = TEXT("Default");
+	if (Tracker == ELLMTracker::Platform)
+	{
+		TrackerString = TEXT("Platform");
+	}
+	static_assert((int)ELLMTracker::Max == 2, "Add other tracker type strings here");
+
+	FString UniqueFilename = SavedDir / FString::Printf(TEXT("LLM_PrivateShared_Tracker%s_%s.csv"), TrackerString, *ForkString);
+
+	uint32 Counter = 1;
+	while (IFileManager::Get().FileSize(*UniqueFilename) >= 0)
+	{
+		UniqueFilename = SavedDir / *FString::Printf(TEXT("LLM_PrivateShared_Tracker%s_%s_%d.csv"), TrackerString, *ForkString, Counter);
+		Counter++;
+	}
+	
+	if (FFileHelper::SaveStringArrayToFile(Lines, *UniqueFilename) == false)
+	{
+		FPlatformMisc::LowLevelOutputDebugStringf(TEXT("Failed to write llm private/shared csv file: %s\n"), *UniqueFilename);
+	}
+	return true;
+}
+
 void FLLMTracker::PublishCsv(UE::LLM::ESizeParams SizeParams)
 {
 	CsvWriter.Publish(LLMRef, TagSizes, OverrideTrackedTotalTagData, OverrideUntaggedTagData, GetTrackedTotal(SizeParams), SizeParams);
diff --git a/Engine/Source/Runtime/Core/Private/Unix/UnixPlatformMemory.cpp b/Engine/Source/Runtime/Core/Private/Unix/UnixPlatformMemory.cpp
--- a/Engine/Source/Runtime/Core/Private/Unix/UnixPlatformMemory.cpp
+++ b/Engine/Source/Runtime/Core/Private/Unix/UnixPlatformMemory.cpp
@@ -940,6 +940,186 @@ static uint64 ParseSMapsFileChunk(ANSICHAR *Buffer, uint64 BufferSize, FProcFiel
 	return ParsePos ? ParsePos : BufferSize;
 }
 
+// For tracking what we've seen in the smaps output between addresses
+struct FPageParseState
+{
+	bool bHavePageAddress = false;
+	bool bHaveSharedClean = false;
+	bool bHaveSharedDirty = false;
+	bool bHavePrivateClean = false;
+	bool bHavePrivateDirty = false;
+	bool Complete()
+	{
+		return bHaveSharedClean && bHaveSharedDirty && bHavePrivateClean && bHavePrivateDirty;
+	}
+	void Reset()
+	{
+		*this = FPageParseState();
+	}
+};
+
+// info on where to put the smaps field and which ones we wants.
+struct FProcFieldWithOffset
+{
+	const ANSICHAR* Name = nullptr;
+	uint64 Offset = 0;
+	uint64 ConfirmOffset = 0;
+	uint32 NameLen = 0;
+
+	FProcFieldWithOffset(const ANSICHAR* NameIn, uint64 OffsetIn, uint64 ConfirmOffsetIn)
+		: Name(NameIn), Offset(OffsetIn), ConfirmOffset(ConfirmOffsetIn)
+	{
+		NameLen = FCStringAnsi::Strlen(Name);
+	}
+};
+
+// Look in Buffer for a line and advance it, parsing out data that we care about. Returns how far in to Buffer we made it before
+// we couldn't get a complete line.
+static uint64 ParseSMapsPage(ANSICHAR* Buffer, uint64 BufferSize, FPageParseState& State, TArray<FForkedPageAllocation>& Pages, FProcFieldWithOffset* Fields, uint32 FieldCount)
+{
+	uint64 ParsePos = 0;
+
+	for (uint64 Idx = 0; Idx < BufferSize; Idx++)
+	{
+		if (Buffer[Idx] == '\n')
+		{
+			ANSICHAR* Line = Buffer + ParsePos;
+
+			Buffer[Idx] = 0;
+
+			// If the line starts with a number, its an address.
+			if ((Line[0] >= '0' && Line[0] <= '9') ||
+				(Line[0] >= 'a' && Line[0] <= 'f')) // all the fields start with capitals so this works
+			{
+				check(!State.bHavePageAddress);
+
+				// Start a new page.
+				FForkedPageAllocation& Page = Pages.AddZeroed_GetRef();
+
+				// 'from' address is up to '-'
+				ANSICHAR* Dash = FCStringAnsi::Strchr(Line, '-');
+				ANSICHAR* Space = FCStringAnsi::Strchr(Dash, ' ');
+
+				Dash[0] = 0;
+				Space[0] = 0;
+
+				Page.PageStart = strtoll(Line, 0, 16);
+				Page.PageEnd = strtoll(Dash + 1, 0, 16);
+
+				State.bHavePageAddress = true;
+			}
+			else if (State.bHavePageAddress)
+			{
+				for (uint32 IdxField = 0; IdxField < FieldCount; IdxField++)
+				{
+					uint32 NameLen = Fields[IdxField].NameLen;
+
+					if (!FCStringAnsi::Strncmp(Line, Fields[IdxField].Name, NameLen))
+					{
+						// Make sure we don't already have it
+						check(((bool*)(&State))[Fields[IdxField].ConfirmOffset] == false);
+						((bool*)(&State))[Fields[IdxField].ConfirmOffset] = true;
+
+						uint64 ResultKb = atoll(Line + NameLen);
+
+						FForkedPageAllocation* Page = &Pages.Top();
+
+						((uint64*)Page)[Fields[IdxField].Offset] = ResultKb;
+
+						if (State.Complete())
+						{
+							// If we have all the fields we care about, go back to looking for an address.
+							State.Reset();
+						}
+						break;
+					}
+				}
+			}
+
+			ParsePos = Idx + 1;
+		}
+	}
+
+	// If we didn't find any linefeeds, skip the entire chunk
+	return ParsePos ? ParsePos : BufferSize;
+}
+
+bool FUnixPlatformMemory::GetForkedPageAllocationInfo(TArray<FForkedPageAllocation>& OutPages)
+{
+	OutPages.Reset();
+
+	const ANSICHAR Shared_CleanStr[] = "Shared_Clean:";
+	const ANSICHAR Shared_DirtyStr[] = "Shared_Dirty:";
+	const ANSICHAR Private_CleanStr[] = "Private_Clean:";
+	const ANSICHAR Private_DirtyStr[] = "Private_Dirty:";
+
+
+	FProcFieldWithOffset SMapsFields[] =
+	{
+		{ Shared_CleanStr, offsetof(FForkedPageAllocation, SharedCleanKiB) / 8 , offsetof(FPageParseState, bHaveSharedClean)    },
+		{ Shared_DirtyStr, offsetof(FForkedPageAllocation, SharedDirtyKiB) / 8 , offsetof(FPageParseState, bHaveSharedDirty)    },
+		{ Private_CleanStr,offsetof(FForkedPageAllocation, PrivateCleanKiB) / 8, offsetof(FPageParseState, bHavePrivateClean)  },
+		{ Private_DirtyStr,offsetof(FForkedPageAllocation, PrivateDirtyKiB) / 8, offsetof(FPageParseState, bHavePrivateDirty)  },
+	};
+
+	int Fd = open("/proc/self/smaps", O_RDONLY);
+	if (Fd < 0)
+	{
+		UE_LOG(LogHAL, Warning, TEXT("Failed to open /proc/self/smaps for pages"));
+		return false;
+	}
+
+	const uint32 ChunkSize = 16*1024;
+	ANSICHAR Buffer[ChunkSize];
+	uint64 BytesAvailableInChunk = 0;
+
+	FPageParseState ParseState;
+
+	for (;;)
+	{
+		ssize_t BytesRead;
+		uint64 BytesToRead = ChunkSize - BytesAvailableInChunk;
+
+		do
+		{
+			QUICK_SCOPE_CYCLE_COUNTER(STAT_GetPages_FileRead);
+			BytesRead = read(Fd, Buffer + BytesAvailableInChunk, BytesToRead);
+		} while (BytesRead < 0 && errno == EINTR);
+
+		if (BytesRead <= 0)
+		{
+			break;
+		}
+
+		BytesAvailableInChunk += BytesRead;
+
+		uint64 BytesParsed = ParseSMapsPage(Buffer, BytesAvailableInChunk, ParseState, OutPages, SMapsFields, UE_ARRAY_COUNT(SMapsFields));
+		checkf(BytesParsed <= BytesAvailableInChunk, TEXT("BytesParsed more than BytesAvailableInChunk %u %u"), BytesParsed, BytesAvailableInChunk);
+		if (BytesParsed > BytesAvailableInChunk)
+		{
+			FPlatformMisc::LowLevelOutputDebugString(TEXT("Critial parse fail in ParseSMapsPage"));
+			break;
+		}
+
+		BytesAvailableInChunk -= BytesParsed;
+		memmove(Buffer, Buffer + BytesParsed, BytesAvailableInChunk);
+
+		// If we made no parsing progress and we've reached end of file, stop trying.
+		// This is a sanity check because we expect to complete parsing prior to getting EOF, naturally exiting
+		// in the above <= 0 check.
+		if (BytesRead == 0 &&
+			BytesParsed == 0)
+		{
+			break;
+		}
+	}
+
+	close(Fd);
+	return true;
+}
+
+
+
 FExtendedPlatformMemoryStats FUnixPlatformMemory::GetExtendedStats()
 {
 	QUICK_SCOPE_CYCLE_COUNTER(STAT_GetExtendedStats);
@@ -994,9 +1174,9 @@ FExtendedPlatformMemoryStats FUnixPlatformMemory::GetExtendedStats()
 
 				uint64 BytesParsed = ParseSMapsFileChunk(Buffer, BytesAvailableInChunk, SMapsFields, UE_ARRAY_COUNT(SMapsFields));
 				checkf(BytesParsed <= BytesAvailableInChunk, TEXT("BytesParsed more than BytesAvailableInChunk %u %u"), BytesParsed, BytesAvailableInChunk);
-
-				if (BytesRead < BytesToRead)
+				if (BytesParsed > BytesAvailableInChunk)
 				{
+					FPlatformMisc::LowLevelOutputDebugString(TEXT("Critial parse fail in ParseSMapsPage"));
 					break;
 				}
 
diff --git a/Engine/Source/Runtime/Core/Public/GenericPlatform/GenericPlatformMemory.h b/Engine/Source/Runtime/Core/Public/GenericPlatform/GenericPlatformMemory.h
--- a/Engine/Source/Runtime/Core/Public/GenericPlatform/GenericPlatformMemory.h
+++ b/Engine/Source/Runtime/Core/Public/GenericPlatform/GenericPlatformMemory.h
@@ -178,6 +178,24 @@ struct FGenericPlatformMemoryStats : public FPlatformMemoryConstants
 	void SetEndFrameCsvStats() const {}
 };
 
+// Contains shared/private information for a single page allocation from the kernel. A page
+// allocation may contain many pages.
+struct FForkedPageAllocation
+{
+	// Start/End virtual address for the allocation.
+	uint64 PageStart;
+	uint64 PageEnd;
+
+	// The amount of memory in this allocation range that is shared across the forked
+	// child processes.
+	uint64 SharedCleanKiB;
+	uint64 SharedDirtyKiB;
+
+	// The amount of memory in this allocation range that has been written to by the child
+	// process, and as a result has been made unique to the process.
+	uint64 PrivateCleanKiB;
+	uint64 PrivateDirtyKiB;
+};
 
 
 
@@ -768,6 +786,16 @@ public:
 	*/
 	static bool HasForkPageProtectorEnabled() { return false; }
 
+	/**
+	* Return the page allocations from the operating system (/proc/self/smaps). This only means something on
+	* platforms that can fork and have Copy On Write behavior. 
+	*/
+	static CORE_API bool GetForkedPageAllocationInfo(TArray<FForkedPageAllocation>& OutPageAllocationInfos)
+	{
+		return false; // Most platform do not implement this.
+	}
+
+
 protected:
 	friend struct FGenericStatsUpdater;
 
diff --git a/Engine/Source/Runtime/Core/Public/Unix/UnixPlatformMemory.h b/Engine/Source/Runtime/Core/Public/Unix/UnixPlatformMemory.h
--- a/Engine/Source/Runtime/Core/Public/Unix/UnixPlatformMemory.h
+++ b/Engine/Source/Runtime/Core/Public/Unix/UnixPlatformMemory.h
@@ -83,6 +83,8 @@ struct FUnixPlatformMemory : public FGenericPlatformMemory
 	static CORE_API void* BinnedAllocFromOS(SIZE_T Size);
 	static CORE_API void BinnedFreeToOS(void* Ptr, SIZE_T Size);
 
+	static CORE_API bool GetForkedPageAllocationInfo(TArray<FForkedPageAllocation>& OutPageAllocationInfos);
+
 	class FPlatformVirtualMemoryBlock : public FBasicVirtualMemoryBlock
 	{
 	public:
