Missed some minor feedback

diff --git a/Engine/Source/Runtime/Core/Private/Unix/UnixPlatformMemory.cpp b/Engine/Source/Runtime/Core/Private/Unix/UnixPlatformMemory.cpp
--- a/Engine/Source/Runtime/Core/Private/Unix/UnixPlatformMemory.cpp
+++ b/Engine/Source/Runtime/Core/Private/Unix/UnixPlatformMemory.cpp
@@ -1097,8 +1097,9 @@ bool FUnixPlatformMemory::GetForkedPageAllocationInfo(TArray<FForkedPageAllocati
 		checkf(BytesParsed <= BytesAvailableInChunk, TEXT("BytesParsed more than BytesAvailableInChunk %u %u"), BytesParsed, BytesAvailableInChunk);
 		if (BytesParsed > BytesAvailableInChunk)
 		{
-			FPlatformMisc::LowLevelOutputDebugString(TEXT("Critial parse fail in ParseSMapsPage"));
-			break;
+			FPlatformMisc::LowLevelOutputDebugString(TEXT("Critical parse fail in ParseSMapsPage"));
+			close(Fd);
+			return false;
 		}
 
 		BytesAvailableInChunk -= BytesParsed;
@@ -1176,8 +1177,9 @@ FExtendedPlatformMemoryStats FUnixPlatformMemory::GetExtendedStats()
 				checkf(BytesParsed <= BytesAvailableInChunk, TEXT("BytesParsed more than BytesAvailableInChunk %u %u"), BytesParsed, BytesAvailableInChunk);
 				if (BytesParsed > BytesAvailableInChunk)
 				{
-					FPlatformMisc::LowLevelOutputDebugString(TEXT("Critial parse fail in ParseSMapsPage"));
-					break;
+					FPlatformMisc::LowLevelOutputDebugString(TEXT("Critical parse fail in ParseSMapsPage"));
+					close(Fd);
+					return FExtendedPlatformMemoryStats();
 				}
 
 				BytesAvailableInChunk -= BytesParsed;
