Add linear case to UnrealMath::FindBounds()
* Fixes a bug where SplineComponents compute the incorrect bounds in some cases

#rnx

diff --git a/Engine/Source/Runtime/Core/Private/Math/UnrealMath.cpp b/Engine/Source/Runtime/Core/Private/Math/UnrealMath.cpp
--- a/Engine/Source/Runtime/Core/Private/Math/UnrealMath.cpp
+++ b/Engine/Source/Runtime/Core/Private/Math/UnrealMath.cpp
@@ -1534,29 +1534,48 @@ static void FindBounds( T& OutMin, T& OutMax,  T Start, T StartLeaveTan, float S
 		const T b = -6.f*Start - 4.f*StartLeaveTan - 2.f*EndArriveTan + 6.f*End;
 		const T c = StartLeaveTan;
 
-		const T Discriminant = (b*b) - (4.f*a*c);
-		if(Discriminant > 0.f && !FMath::IsNearlyZero(a)) // Solving doesn't work if a is zero, which usually indicates co-incident start and end, and zero tangents anyway
+		if (FMath::IsNearlyZero(a))
 		{
-			const T SqrtDisc = FMath::Sqrt( Discriminant );
-
-			const T x0 = (-b + SqrtDisc)/(2.f*a); // x0 is the 'Alpha' ie between 0 and 1
-			const T t0 = StartT + x0*(EndT - StartT); // Then t0 is the actual 'time' on the curve
-			if(t0 > StartT && t0 < EndT)
+			// The derivative is linear, find the linear root.
+			if (!FMath::IsNearlyZero(b))
 			{
-				const T Val = FMath::CubicInterp( Start, StartLeaveTan, End, EndArriveTan, x0 );
-
-				OutMin = FMath::Min( OutMin, Val );
-				OutMax = FMath::Max( OutMax, Val );
+				const T x = -c / b;
+				const T t = StartT + x * (EndT - StartT);
+				if (t > StartT && t < EndT)
+				{
+					const T Val = FMath::CubicInterp(Start, StartLeaveTan, End, EndArriveTan, x);
+					OutMin = FMath::Min(OutMin, Val);
+					OutMax = FMath::Max(OutMax, Val);
+				}
 			}
-
-			const T x1 = (-b - SqrtDisc)/(2.f*a);
-			const T t1 = StartT + x1*(EndT - StartT);
-			if(t1 > StartT && t1 < EndT)
+		}
+		else
+		{
+			// The derivative is quadratic, find the quadratic roots.
+			const T Discriminant = (b * b) - (4.f * a * c);
+			if (Discriminant >= 0.f)
 			{
-				const T Val = FMath::CubicInterp( Start, StartLeaveTan, End, EndArriveTan, x1 );
+				const T SqrtDisc = FMath::Sqrt(Discriminant);
 
-				OutMin = FMath::Min( OutMin, Val );
-				OutMax = FMath::Max( OutMax, Val );
+				const T x0 = (-b + SqrtDisc) / (2.f * a); // x0 is the 'Alpha' ie between 0 and 1
+				const T t0 = StartT + x0 * (EndT - StartT); // Then t0 is the actual 'time' on the curve
+				if (t0 > StartT && t0 < EndT)
+				{
+					const T Val = FMath::CubicInterp(Start, StartLeaveTan, End, EndArriveTan, x0);
+
+					OutMin = FMath::Min(OutMin, Val);
+					OutMax = FMath::Max(OutMax, Val);
+				}
+
+				const T x1 = (-b - SqrtDisc) / (2.f * a);
+				const T t1 = StartT + x1 * (EndT - StartT);
+				if (t1 > StartT && t1 < EndT)
+				{
+					const T Val = FMath::CubicInterp(Start, StartLeaveTan, End, EndArriveTan, x1);
+
+					OutMin = FMath::Min(OutMin, Val);
+					OutMax = FMath::Max(OutMax, Val);
+				}
 			}
 		}
 	}
