- Added thread protection around the loading/clearing of cached other-platform cvars

diff --git a/Engine/Source/Runtime/Core/Private/HAL/ConsoleManager.cpp b/Engine/Source/Runtime/Core/Private/HAL/ConsoleManager.cpp
--- a/Engine/Source/Runtime/Core/Private/HAL/ConsoleManager.cpp
+++ b/Engine/Source/Runtime/Core/Private/HAL/ConsoleManager.cpp
@@ -2938,6 +2938,8 @@ void FConsoleManager::LoadAllPlatformCVars(FName PlatformName, const FString& De
 {
 	FName PlatformKey = MakePlatformKey(PlatformName, DeviceProfileName);
 	
+	// protect the cached CVar info from two threads trying to get a platform CVar at once, and both attempting to load all of the cvars at the same time
+	FScopeLock Lock(&CachedPlatformsAndDeviceProfilesLock);
 	if (CachedPlatformsAndDeviceProfiles.Contains(PlatformKey))
 	{
 		return;
@@ -2995,6 +2997,10 @@ void FConsoleManager::PreviewPlatformCVars(FName PlatformName, const FString& De
 void FConsoleManager::ClearAllPlatformCVars(FName PlatformName, const FString& DeviceProfileName)
 {
 	FName PlatformKey = MakePlatformKey(PlatformName, DeviceProfileName);
+
+	// protect the cached CVar info from two threads trying to get a platform CVar at once, and both attempting to load all of the cvars at the same time
+	FScopeLock Lock(&CachedPlatformsAndDeviceProfilesLock);
+	
 	if (!CachedPlatformsAndDeviceProfiles.Contains(PlatformKey))
 	{
 		return;
diff --git a/Engine/Source/Runtime/Core/Public/HAL/ConsoleManager.h b/Engine/Source/Runtime/Core/Public/HAL/ConsoleManager.h
--- a/Engine/Source/Runtime/Core/Public/HAL/ConsoleManager.h
+++ b/Engine/Source/Runtime/Core/Public/HAL/ConsoleManager.h
@@ -109,6 +109,7 @@ private: // ----------------------------------------------------
 
 	IConsoleThreadPropagation* ThreadPropagationCallback;
 	
+	FCriticalSection CachedPlatformsAndDeviceProfilesLock;
 	TSet<FName> CachedPlatformsAndDeviceProfiles;
 
 	// if true the next call to CallAllConsoleVariableSinks() we will call all registered sinks
