- Allow for arbitrary number of PreviewPlatforms differing only by DeviceProfile
- Can specify multiple DPs and FriendlyNames now in the PreviewPlaform sections of a DDPI.ini file (separated by comma)
- Use the FriendlyName from the DDPI.ini directly in the platform preview menus, instead of counting on the friendly name of a ShaderPlatform
- IOS DDPI.ini as an example, with 3 DPs, generic IOS (what we had before), iPhone6S and iPadPro6

diff --git a/Engine/Config/IOS/DataDrivenPlatformInfo.ini b/Engine/Config/IOS/DataDrivenPlatformInfo.ini
--- a/Engine/Config/IOS/DataDrivenPlatformInfo.ini
+++ b/Engine/Config/IOS/DataDrivenPlatformInfo.ini
@@ -32,7 +32,8 @@ ActiveIconName=LevelEditor.PreviewMode.iOS.Enabled
 InactiveIconName=LevelEditor.PreviewMode.iOS.Disabled
 MenuTooltip=NSLOCTEXT("PreviewPlatform", "PreviewMenuTooltip_iOS", "Mobile preview using iOS material quality settings.")
 IconText=NSLOCTEXT("PreviewPlatform", "PreviewIconText_iOS", "iOS")
-DeviceProfileName=IOS
+DeviceProfileName=IOS,iPhone6S,iPadPro6_129
+FriendlyName=Generic IOS,iPhone 6S, iPad Pro 6 12.9"
 
 [PreviewPlatform IOSMetalSM5]
 PlatformName=IOS
@@ -80,4 +81,4 @@ bSupports4ComponentUAVReadWrite=true
 bSupportsManualVertexFetch=false
 bSupportsVolumeTextureAtomics=false
 bSupportsShaderPipelines = false
-FriendlyName=LOCTEXT("FriendlyShaderPlatformName_IOSMobileSIM", "IOS Metal Simulator Mobile")
\ No newline at end of file
+FriendlyName=LOCTEXT("FriendlyShaderPlatformName_IOSMobileSIM", "IOS Metal Simulator Mobile")
diff --git a/Engine/Source/Editor/LevelEditor/Private/LevelEditorActions.cpp b/Engine/Source/Editor/LevelEditor/Private/LevelEditorActions.cpp
--- a/Engine/Source/Editor/LevelEditor/Private/LevelEditorActions.cpp
+++ b/Engine/Source/Editor/LevelEditor/Private/LevelEditorActions.cpp
@@ -3959,9 +3959,14 @@ void FLevelEditorCommands::RegisterCommands()
 				// if the shader platform isn't compiled in, we don't have a friendly name available, so use ugly name
 				FriendlyNameBuilder.AppendLine(FText::FromName(Item.PreviewShaderPlatformName));
 			}
+			else if (!Item.OptionalFriendlyNameOverride.IsEmpty())
+			{
+				FriendlyNameBuilder.AppendLine(Item.OptionalFriendlyNameOverride);
+			}
 			else
 			{
 				FriendlyNameBuilder.AppendLine(FDataDrivenShaderPlatformInfo::GetFriendlyName(ShaderPlatform));
+				FPlatformMisc::LowLevelOutputDebugStringf(TEXT("MENU friendly name %s\n"), *FriendlyNameBuilder.ToText().ToString());
 			}
 			if (FDataDrivenShaderPlatformInfo::GetShaderPlatformFromName(Item.ShaderPlatformToPreview) == GMaxRHIShaderPlatform)
 			{
@@ -3972,7 +3977,7 @@ void FLevelEditorCommands::RegisterCommands()
 		PreviewPlatformOverrides.Add(
 			FUICommandInfoDecl(
 				this->AsShared(),
-				FName(*FString::Printf(TEXT("PreviewPlatformOverrides_%s_%s"), *Item.PlatformName.ToString(), *Item.ShaderFormat.ToString())),
+				FName(*FString::Printf(TEXT("PreviewPlatformOverrides_%s_%s_%s"), *Item.PlatformName.ToString(), *Item.ShaderFormat.ToString(), *Item.DeviceProfileName.ToString())),
 				FriendlyNameBuilder.ToText(),
 				Item.MenuTooltip)
 			.UserInterfaceType(EUserInterfaceActionType::Check)
diff --git a/Engine/Source/Runtime/Core/Private/Misc/DataDrivenPlatformInfoRegistry.cpp b/Engine/Source/Runtime/Core/Private/Misc/DataDrivenPlatformInfoRegistry.cpp
--- a/Engine/Source/Runtime/Core/Private/Misc/DataDrivenPlatformInfoRegistry.cpp
+++ b/Engine/Source/Runtime/Core/Private/Misc/DataDrivenPlatformInfoRegistry.cpp
@@ -286,13 +286,34 @@ static void ParsePreviewPlatforms(const FConfigFile& IniFile)
 				Item.ActiveIconName = *GetSectionString(Section.Value, FName("ActiveIconName"));
 				Item.InactiveIconPath = GetSectionString(Section.Value, FName("InactiveIconPath"));
 				Item.InactiveIconName = *GetSectionString(Section.Value, FName("InactiveIconName"));
-				Item.DeviceProfileName = *GetSectionString(Section.Value, FName("DeviceProfileName"));
 				Item.ShaderPlatformToPreview = *GetSectionString(Section.Value, FName("ShaderPlatform"));
 				checkf(Item.ShaderPlatformToPreview != NAME_None, TEXT("DataDrivenPlatformInfo section [PreviewPlatform %s] must specify a ShaderPlatform"), *SectionName);
-				FTextStringHelper::ReadFromBuffer(*GetSectionString(Section.Value, FName("FriendlyName")), Item.OptionalFriendlyNameOverride);
 				FTextStringHelper::ReadFromBuffer(*GetSectionString(Section.Value, FName("MenuTooltip")), Item.MenuTooltip);
 				FTextStringHelper::ReadFromBuffer(*GetSectionString(Section.Value, FName("IconText")), Item.IconText);
-				PreviewPlatformMenuItems.Add(Item);
+
+
+				FString AllDeviceProfiles = GetSectionString(Section.Value, FName("DeviceProfileName"));
+				FString AllFriendlyName = GetSectionString(Section.Value, FName("FriendlyName"));
+				TArray<FString> DeviceProfileNames, FriendlyNames;
+				AllDeviceProfiles.ParseIntoArray(DeviceProfileNames, TEXT(","));
+				AllFriendlyName.ParseIntoArray(FriendlyNames, TEXT(","));
+
+				for (int DPIndex = 0; DPIndex < DeviceProfileNames.Num(); DPIndex++)
+				{
+					Item.DeviceProfileName = *DeviceProfileNames[DPIndex].TrimStartAndEnd();
+					if (DPIndex < FriendlyNames.Num())
+					{
+						Item.OptionalFriendlyNameOverride = FText::FromString(FriendlyNames[DPIndex].TrimStartAndEnd());
+					}
+					else if (DeviceProfileNames.Num() > 1)
+					{
+						Item.OptionalFriendlyNameOverride = FText::FromString(Item.DeviceProfileName.ToString());
+					}
+					
+					FPlatformMisc::LowLevelOutputDebugStringf(TEXT("Adding friendly name %s\n"), *Item.OptionalFriendlyNameOverride.ToString());
+					
+					PreviewPlatformMenuItems.Add(Item);
+				}
 			}
 		}
 	}
