- Added ability for a FConfigContext to have an override layers array in case something has a custom set of ini files

diff --git a/Engine/Source/Runtime/Core/Private/Misc/ConfigContext.cpp b/Engine/Source/Runtime/Core/Private/Misc/ConfigContext.cpp
--- a/Engine/Source/Runtime/Core/Private/Misc/ConfigContext.cpp
+++ b/Engine/Source/Runtime/Core/Private/Misc/ConfigContext.cpp
@@ -550,6 +550,12 @@ void FConfigContext::AddStaticLayersToHierarchy(TArray<FString>* GatheredLayerFi
 		Layers = GPluginLayers;
 		NumLayers = UE_ARRAY_COUNT(GPluginLayers);
 	}
+	// let the context override the layers if needed
+	if (OverrideLayers.Num() > 0)
+	{
+		Layers = OverrideLayers.GetData();
+		NumLayers = OverrideLayers.Num();
+	}
 
 	// go over all the config layers
 	for (int32 LayerIndex = 0; LayerIndex < NumLayers; LayerIndex++)
diff --git a/Engine/Source/Runtime/Core/Public/Misc/ConfigContext.h b/Engine/Source/Runtime/Core/Public/Misc/ConfigContext.h
--- a/Engine/Source/Runtime/Core/Public/Misc/ConfigContext.h
+++ b/Engine/Source/Runtime/Core/Public/Misc/ConfigContext.h
@@ -193,6 +193,9 @@ public:
 	FString ProjectNoRedistDir;
 	TMap<FString, FPerPlatformDirs> PerPlatformDirs;
 
+	// allow a custom set of layers
+	TArray<FConfigLayer> OverrideLayers;
+	
 	bool bUseHierarchyCache = false;
 	bool bAllowGeneratedIniWhenCooked = false;
 	bool bForceReload = false;
