- Removed log that wasn't meant to be checked in

diff --git a/Engine/Source/Runtime/Core/Private/Misc/DataDrivenPlatformInfoRegistry.cpp b/Engine/Source/Runtime/Core/Private/Misc/DataDrivenPlatformInfoRegistry.cpp
--- a/Engine/Source/Runtime/Core/Private/Misc/DataDrivenPlatformInfoRegistry.cpp
+++ b/Engine/Source/Runtime/Core/Private/Misc/DataDrivenPlatformInfoRegistry.cpp
@@ -309,9 +309,7 @@ static void ParsePreviewPlatforms(const FConfigFile& IniFile)
 					{
 						Item.OptionalFriendlyNameOverride = FText::FromString(Item.DeviceProfileName.ToString());
 					}
-					
-					FPlatformMisc::LowLevelOutputDebugStringf(TEXT("Adding friendly name %s\n"), *Item.OptionalFriendlyNameOverride.ToString());
-					
+									
 					PreviewPlatformMenuItems.Add(Item);
 				}
 			}
