Clarifying comment.

diff --git a/Engine/Source/Runtime/Core/Private/IO/IoStore.cpp b/Engine/Source/Runtime/Core/Private/IO/IoStore.cpp
--- a/Engine/Source/Runtime/Core/Private/IO/IoStore.cpp
+++ b/Engine/Source/Runtime/Core/Private/IO/IoStore.cpp
@@ -2226,9 +2226,9 @@ public:
 	// GenericPlatformFile isn't designed around a lot of jobs throwing accesses at it, so instead we 
 	// use IFileHandle directly and round robin between a number of file handles in order to saturate
 	// year 2022 ssd drives. For a file hot in the windows file cache, you can get 4+ GB/s with as few as 
-	// 4 file handles, however a cold file you need upwards of 32 in order to reach ~1.5 GB/s. This is
-	// low because IoStore reads are comparatively small - at most you're reading compression block sized
-	// chunks with uncompressed, however with Oodle those get cut by ~half, so with a default block size
+	// 4 file handles, however for a cold file you need upwards of 32 in order to reach ~1.5 GB/s. This is
+	// low because IoStoreReader (note: not IoDispatcher!) reads are comparatively small - at most you're reading compression block sized
+	// chunks when uncompressed, however with Oodle those get cut by ~half, so with a default block size
 	// of 64kb, reads are generally less than 32kb, which is tough to use and get full ssd bandwidth out of.
 	//
 	static constexpr uint32 NumHandlesPerFile = 12;
