Don't record logging CSV stats if we're crashing, otherwise we'll fail in CreateTls
#tests tested in a cooked build, preflighted to check it compiles for all plats

diff --git a/Engine/Source/Runtime/Core/Private/GenericPlatform/GenericPlatformMallocCrash.cpp b/Engine/Source/Runtime/Core/Private/GenericPlatform/GenericPlatformMallocCrash.cpp
--- a/Engine/Source/Runtime/Core/Private/GenericPlatform/GenericPlatformMallocCrash.cpp
+++ b/Engine/Source/Runtime/Core/Private/GenericPlatform/GenericPlatformMallocCrash.cpp
@@ -296,6 +296,11 @@ void FGenericPlatformMallocCrash::SetAsGMalloc()
 	CrashedThreadId = FPlatformTLS::GetCurrentThreadId();
 }
 
+bool FGenericPlatformMallocCrash::IsActive() const
+{
+	return GMalloc == this;
+}
+
 void* FGenericPlatformMallocCrash::Malloc( SIZE_T Size, uint32 Alignment )
 {
 	const uint32 Size32 = (uint32)Size;
@@ -641,6 +646,11 @@ void FGenericStackBasedMallocCrash::SetAsGMalloc()
 	GMalloc = this;
 }
 
+bool FGenericStackBasedMallocCrash::IsActive() const
+{
+	return GMalloc == this;
+}
+
 void* FGenericStackBasedMallocCrash::Malloc(SIZE_T Size, uint32 Alignment)
 {
 	Alignment = FMath::Max(Size >= 16 ? (uint32)16 : (uint32)8, Alignment);
diff --git a/Engine/Source/Runtime/Core/Private/Logging/LogMacros.cpp b/Engine/Source/Runtime/Core/Private/Logging/LogMacros.cpp
--- a/Engine/Source/Runtime/Core/Private/Logging/LogMacros.cpp
+++ b/Engine/Source/Runtime/Core/Private/Logging/LogMacros.cpp
@@ -3,6 +3,7 @@
 #include "Logging/LogMacros.h"
 #include "CoreGlobals.h"
 #include "HAL/Platform.h"
+#include "HAL/PlatformMallocCrash.h"
 #include "Misc/ScopeLock.h"
 #include "Misc/OutputDeviceRedirector.h"
 #include "Misc/FeedbackContext.h"
@@ -54,7 +55,6 @@ void FMsg::LogV(const ANSICHAR* File, int32 Line, const FLogCategoryName& Catego
 #if !NO_LOGGING
 	LLM_SCOPE_BYNAME("EngineMisc/FMsgLogf")
 	QUICK_SCOPE_CYCLE_COUNTER(STAT_FMsgLogf);
-	CSV_CUSTOM_STAT(FMsgLogf, FMsgLogfCount, 1, ECsvCustomStatOp::Accumulate);
 
 	if (LIKELY(Verbosity != ELogVerbosity::Fatal))
 	{
@@ -79,6 +79,15 @@ void FMsg::LogV(const ANSICHAR* File, int32 Line, const FLogCategoryName& Catego
 		{
 			(OutputDevice ? OutputDevice : GLog)->Serialize(Message, Verbosity, Category);
 		});
+
+#if CSV_PROFILER
+		// Only update the CSV stat if we're not crashing, otherwise things can get messy
+		if (LIKELY(!FPlatformMallocCrash::Get().IsActive()))
+		{
+			CSV_CUSTOM_STAT(FMsgLogf, FMsgLogfCount, 1, ECsvCustomStatOp::Accumulate);
+		}
+#endif
+
 	}
 	else
 	{
diff --git a/Engine/Source/Runtime/Core/Public/GenericPlatform/GenericPlatformMallocCrash.h b/Engine/Source/Runtime/Core/Public/GenericPlatform/GenericPlatformMallocCrash.h
--- a/Engine/Source/Runtime/Core/Public/GenericPlatform/GenericPlatformMallocCrash.h
+++ b/Engine/Source/Runtime/Core/Public/GenericPlatform/GenericPlatformMallocCrash.h
@@ -77,6 +77,8 @@ public:
 	 */
 	void SetAsGMalloc();
 
+	bool IsActive() const;
+
 	// FMalloc interface.
 	virtual void* Malloc( SIZE_T Size, uint32 Alignment ) override;
 
@@ -176,6 +178,8 @@ struct FGenericStackBasedMallocCrash : public FMalloc
 
 	void SetAsGMalloc();
 
+	bool IsActive() const;
+
 	virtual void* Malloc(SIZE_T Size, uint32 Alignment) override;
 
 	virtual void* Realloc(void* Ptr, SIZE_T NewSize, uint32 Alignment) override;
