- Added missing PC PlatformPreviews (and any others that didn't specify a DeviceProfile)
- Added Mac PlatformPreviews so it can end preview mode (which needs a maching PP to the host platform)

diff --git a/Engine/Config/Mac/DataDrivenPlatformInfo.ini b/Engine/Config/Mac/DataDrivenPlatformInfo.ini
--- a/Engine/Config/Mac/DataDrivenPlatformInfo.ini
+++ b/Engine/Config/Mac/DataDrivenPlatformInfo.ini
@@ -128,3 +128,15 @@ bSupportsVolumeTextureAtomics=false
 bSupportsShaderPipelines = false
 bSupportsDualSourceBlending = true
 FriendlyName=LOCTEXT("FriendlyShaderPlatformName_MacMobile", "Mac Metal Mobile")
+
+[PreviewPlatform METAL_SM5]
+PlatformName=Mac
+ShaderFormat=SF_METAL_SM5
+ShaderPlatform=METAL_SM5
+MenuTooltip=LOCTEXT("PreviewMenuTooltip_METAL_SM5", "Mac using SM5 profile")
+
+[PreviewPlatform METAL_SM6]
+PlatformName=Mac
+ShaderFormat=SF_METAL_SM6
+ShaderPlatform=METAL_SM6
+MenuTooltip=LOCTEXT("PreviewMenuTooltip_METAL_SM6", "Mac using SM6 profile")
diff --git a/Engine/Source/Runtime/Core/Private/Misc/DataDrivenPlatformInfoRegistry.cpp b/Engine/Source/Runtime/Core/Private/Misc/DataDrivenPlatformInfoRegistry.cpp
--- a/Engine/Source/Runtime/Core/Private/Misc/DataDrivenPlatformInfoRegistry.cpp
+++ b/Engine/Source/Runtime/Core/Private/Misc/DataDrivenPlatformInfoRegistry.cpp
@@ -297,6 +297,12 @@ static void ParsePreviewPlatforms(const FConfigFile& IniFile)
 				TArray<FString> DeviceProfileNames, FriendlyNames;
 				AllDeviceProfiles.ParseIntoArray(DeviceProfileNames, TEXT(","));
 				AllFriendlyName.ParseIntoArray(FriendlyNames, TEXT(","));
+				
+				if (DeviceProfileNames.Num() == 0)
+				{
+					DeviceProfileNames.Add(TEXT(""));
+					FriendlyNames.Add(TEXT(""));
+				}
 
 				for (int DPIndex = 0; DPIndex < DeviceProfileNames.Num(); DPIndex++)
 				{
