Change __builtin_trap to __builtin_debugtrap to allow continuing from asserts

diff --git a/Engine/Source/Runtime/Core/Public/IOS/IOSPlatform.h b/Engine/Source/Runtime/Core/Public/IOS/IOSPlatform.h
--- a/Engine/Source/Runtime/Core/Public/IOS/IOSPlatform.h
+++ b/Engine/Source/Runtime/Core/Public/IOS/IOSPlatform.h
@@ -77,7 +77,7 @@ typedef FIOSPlatformTypes FPlatformTypes;
 
 #define PLATFORM_GLOBAL_LOG_CATEGORY					LogIOS
 
-#define PLATFORM_BREAK()                                __builtin_trap()
+#define PLATFORM_BREAK()                                __builtin_debugtrap()
 
 #define PLATFORM_CODE_SECTION(Name)						__attribute__((section("__TEXT,__" Name ",regular,pure_instructions"))) \
 														__attribute__((aligned(4)))
diff --git a/Engine/Source/Runtime/Core/Public/Mac/MacPlatform.h b/Engine/Source/Runtime/Core/Public/Mac/MacPlatform.h
--- a/Engine/Source/Runtime/Core/Public/Mac/MacPlatform.h
+++ b/Engine/Source/Runtime/Core/Public/Mac/MacPlatform.h
@@ -89,7 +89,7 @@ typedef FMacPlatformTypes FPlatformTypes;
 #if PLATFORM_MAC_X86
 	#define PLATFORM_BREAK()							__asm__("int $3")
 #else
-    #define PLATFORM_BREAK()                            __builtin_trap()
+    #define PLATFORM_BREAK()                            __builtin_debugtrap()
 #endif
 
 #define PLATFORM_CODE_SECTION(Name)						__attribute__((section("__TEXT,__" Name ",regular,pure_instructions")))
