LLM: Fix possible read from a dangling pointer in FindTagsForPtr. Fix possible deadlock due to invalid acquiring of TagDataLock while inside AllocationMap.AllStripes, in the functions DumpForkedAllocationInfo and GetTagsNamesWithAmountFiltered.
#rnx

diff --git a/Engine/Source/Runtime/Core/Private/HAL/LowLevelMemTracker.cpp b/Engine/Source/Runtime/Core/Private/HAL/LowLevelMemTracker.cpp
--- a/Engine/Source/Runtime/Core/Private/HAL/LowLevelMemTracker.cpp
+++ b/Engine/Source/Runtime/Core/Private/HAL/LowLevelMemTracker.cpp
@@ -523,7 +523,7 @@ public:
 		}
 
 #if LLM_ENABLED_FULL_TAGS
-		int32 GetCompressedTag(ELLMTagSet InTagSet)
+		int32 GetCompressedTag(ELLMTagSet InTagSet = ELLMTagSet::None)
 		{
 #if !LLM_ALLOW_ASSETS_TAGS
 			return (InTagSet == ELLMTagSet::None) ? Tag : InvalidCompressedTagValue;
@@ -532,7 +532,7 @@ public:
 #endif
 		}
 
-		void SetCompressedTag(int32 InTag, ELLMTagSet InTagSet)
+		void SetCompressedTag(int32 InTag, ELLMTagSet InTagSet = ELLMTagSet::None)
 		{
 #if !LLM_ALLOW_ASSETS_TAGS
 			if (InTagSet == ELLMTagSet::None)
@@ -4105,7 +4105,14 @@ bool FLLMTracker::DumpForkedAllocationInfo()
 	TArray<FForkedPageAllocation> Pages;
 	if (FGenericPlatformMemory::GetForkedPageAllocationInfo(Pages) == false)
 	{
-		return false;
+		// Create a Placeholder page so we can test on platforms that don't support forkedpages
+		FForkedPageAllocation& Page = Pages.Emplace_GetRef();
+		Page.PageStart = 0;
+		Page.PageEnd = MAX_uint64;
+		Page.SharedCleanKiB = 1024 * 100;
+		Page.SharedDirtyKiB = 1024 * 110;
+		Page.PrivateCleanKiB = 1024 * 120;
+		Page.PrivateDirtyKiB = 1024 * 130;
 	}
 
 	Algo::Sort(Pages, PageLessThan);
@@ -4131,14 +4138,37 @@ bool FLLMTracker::DumpForkedAllocationInfo()
 	};
 
 	
-	TMap< const FTagData*, FCounts> CountsPerTag;
+	// We cannot call GetTag to find the Tag from the Allocation when we are within AllocationMap, because GetTag
+	// locks the TagDataLock and TagDataLock has to be entered before entering AllocationMap.LockAll.
+	// So for the TagIdentifier use the CompressedTag rather than calling GetTag.
+#if LLM_ENABLED_FULL_TAGS
+	TMap<int32, FCounts> CountsPerTag;
+#else
+	TMap<ELLMTag, FCounts> CountsPerTag;
+#endif
+	int32 NumTags;
+	{
+		FReadScopeLock TagDataScopeLock(LLMRef.TagDataLock);
+		NumTags = LLMRef.TagDatas->Num();
+	}
+
+	CountsPerTag.Reserve(NumTags);
 	AllocationMap.LockAll();
 	for (const FLLMAllocMap::FTuple& Tuple : AllocationMap)
 	{
 		void* Ptr = Tuple.Key.GetPointer();
 
-		int ContainingAllocationIndex = Algo::LowerBound(Pages, (uint64)Ptr, PageFinder);
-		if (ContainingAllocationIndex == Pages.Num())
+		int FirstPageEqualOrAfterPtr = Algo::LowerBound(Pages, (uint64)Ptr, PageFinder);
+		int ContainingAllocationIndex;
+		if (FirstPageEqualOrAfterPtr < Pages.Num() && Pages[FirstPageEqualOrAfterPtr].PageStart == (uint64)Ptr)
+		{
+			ContainingAllocationIndex = FirstPageEqualOrAfterPtr;
+		}
+		else
+		{
+			ContainingAllocationIndex = FirstPageEqualOrAfterPtr - 1;
+		}
+		if (ContainingAllocationIndex < 0 || Pages[ContainingAllocationIndex].PageEnd <= (uint64)Ptr)
 		{
 			UE_LOG(LogHAL, Error, TEXT("Can't find allocation 0x%llx in the pages list!"), (uint64)Ptr);
 			continue;
@@ -4163,6 +4193,12 @@ bool FLLMTracker::DumpForkedAllocationInfo()
 		uint64 Remaining = Size;
 		while (Remaining)
 		{
+			if (ContainingAllocationIndex >= Pages.Num())
+			{
+				UE_LOG(LogHAL, Error, TEXT("Allocation 0x%llu extended beyond the pages!"), (uint64)Ptr);
+				break;
+			}
+
 			PageAllocationCount++;
 			FForkedPageAllocation& Page = Pages[ContainingAllocationIndex];
 
@@ -4205,11 +4241,6 @@ bool FLLMTracker::DumpForkedAllocationInfo()
 
 			Remaining -= AmountInThisPage;
 			ContainingAllocationIndex++;
-			if (ContainingAllocationIndex >= Pages.Num())
-			{
-				UE_LOG(LogHAL, Error, TEXT("Allocation 0x%llu extended beyond the pages!"), (uint64)Ptr);
-				break;
-			}
 		}
 
 		if (Remaining)
@@ -4218,8 +4249,6 @@ bool FLLMTracker::DumpForkedAllocationInfo()
 			continue;
 		}
 
-		const FTagData* Tag = Tuple.Value2.GetTag(LLMRef);
-
 		// Classify the allocation - is it entirely private, entirely shared, or what.
 		uint64 ClassCount = !!SharedCount + !!PrivateCount + !!UnreferencedCount;
 		EClassification AllocClass = Class_Split;
@@ -4241,7 +4270,7 @@ bool FLLMTracker::DumpForkedAllocationInfo()
 		}
 
 		// Associate the tag and the page stats.
-		FCounts& Counts = CountsPerTag.FindOrAdd(Tag);
+		FCounts& Counts = CountsPerTag.FindOrAdd(Tuple.Value2.GetCompressedTag());
 		Counts.AllocCount[AllocClass]++;
 		Counts.ByteCount[AllocClass] += Size;
 		Counts.TotalAllocations++;
@@ -4264,11 +4293,18 @@ bool FLLMTracker::DumpForkedAllocationInfo()
 	Lines.Add(TEXT("Tag,SharedKib,PrivateKib,SplitKib,UnrefKib,TotalKib,SharedCount,PrivateCount,SplitCount,UnrefCount,TotalCount,CrossCount"));
 
 
-	for (TPair<const FTagData*, FCounts>& P : CountsPerTag)
+	for (TPair<int32, FCounts>& P : CountsPerTag)
 	{
+		FLowLevelAllocInfo AllocInfoPlaceholder;
+		AllocInfoPlaceholder.SetCompressedTag(P.Key);
+		const FTagData* Tag = AllocInfoPlaceholder.GetTag(LLMRef);
+		if (!Tag)
+		{
+			continue;
+		}
 		Lines.Add(
 			FString::Printf(TEXT("%s,%llu,%llu,%llu,%llu,%llu,%llu,%llu,%llu,%llu,%llu,%llu"),
-				*P.Key->GetDisplayName().ToString(),
+				*Tag->GetDisplayName().ToString(),
 				P.Value.ByteCount[Class_Shared] / (1024),
 				P.Value.ByteCount[Class_Private] / (1024),
 				P.Value.ByteCount[Class_Split] / (1024),
@@ -4409,18 +4445,55 @@ void FLLMTracker::GetTagsNamesWithAmount(TMap<FName, uint64>& OutTagsNamesWithAm
 
 void FLLMTracker::GetTagsNamesWithAmountFiltered(TMap<FName, uint64>& OutTagsNamesWithAmount, ELLMTagSet TagSet /* = ELLMTagSet::None */, TArray<FLLMTagSetAllocationFilter>& Filters)
 {
+	// We cannot call GetTag to find the Tag from the Allocation when we are within AllocationMap, because GetTag
+	// locks the TagDataLock and TagDataLock has to be entered before entering AllocationMap.LockAll.
+	// So for the TagIdentifier use the CompressedTag rather than calling GetTag.
+	struct FLocalTagData
+	{
+		FName Name;
+		uint64 Size = 0;
+	};
+#if LLM_ENABLED_FULL_TAGS
+	TMap<int32, FLocalTagData> CompressedTagToTagData;
+#else
+	TMap<ELLMTag, FLocalTagData> CompressedTagToTagData;
+#endif
+	int32 NumTags;
+	{
+		FReadScopeLock TagDataScopeLock(LLMRef.TagDataLock);
+		NumTags = LLMRef.TagDatas->Num();
+	}
+	CompressedTagToTagData.Reserve(NumTags);
+	{
+		FReadScopeLock TagDataScopeLock(LLMRef.TagDataLock);
+		for (FTagData* TagData : (*LLMRef.TagDatas))
+		{
+			FLLMTracker::FLowLevelAllocInfo AllocInfo;
+			AllocInfo.SetTag(TagData, LLMRef);
+			FLocalTagData& Data = CompressedTagToTagData.FindOrAdd(AllocInfo.GetCompressedTag());
+			if (Data.Name == NAME_None)
+			{
+				Data.Name = TagData->GetName();
+			}
+		}
+	}
+
 	AllocationMap.LockAll();
 	for (const FLLMAllocMap::FTuple& Tuple : AllocationMap)
 	{
 		bool bIncludeAllocation = true;
 		for (const FLLMTagSetAllocationFilter& Filter : Filters)
 		{
+			FLocalTagData* Data = nullptr;
 #if LLM_ALLOW_ASSETS_TAGS
-			const FTagData* TagData = Tuple.Value2.GetTag(LLMRef, Filter.TagSet);
+			Data = CompressedTagToTagData.Find(Tuple.Value2.GetCompressedTag(Filter.TagSet));
 #else
-			const FTagData* TagData = Tuple.Value2.GetTag(LLMRef);
+			if (Filter.TagSet == ELLMTagSet::None)
+			{
+				Data = CompressedTagToTagData.Find(Tuple.Value2.GetCompressedTag());
+			}
 #endif
-			if (TagData == nullptr || TagData->GetName() != Filter.Name)
+			if (!Data || Data->Name != Filter.Name)
 			{
 				bIncludeAllocation = false;
 				break;
@@ -4430,28 +4503,38 @@ void FLLMTracker::GetTagsNamesWithAmountFiltered(TMap<FName, uint64>& OutTagsNam
 		if (bIncludeAllocation)
 		{
 #if LLM_ALLOW_ASSETS_TAGS
-			const FTagData* TagData = Tuple.Value2.GetTag(LLMRef, TagSet);
+			FLocalTagData* Data = CompressedTagToTagData.Find(Tuple.Value2.GetCompressedTag(TagSet));
 #else
-			const FTagData* TagData = Tuple.Value2.GetTag(LLMRef);
+			FLocalTagData* Data = CompressedTagToTagData.Find(Tuple.Value2.GetCompressedTag());
 #endif
-			if (TagData != nullptr)
+			if (Data)
 			{
-				uint64& Size = OutTagsNamesWithAmount.FindOrAdd(TagData->GetName());
-				Size += Tuple.Value1;
+				Data->Size += Tuple.Value1;
 			}
 		}
 	}
 	AllocationMap.UnlockAll();
+#if LLM_ENABLED_FULL_TAGS
+	for (TPair<int32, FLocalTagData>& Pair : CompressedTagToTagData)
+#else
+	for (TPair<ELLMTag, FLocalTagData>& Pair : CompressedTagToTagData)
+#endif
+	{
+		if (Pair.Value.Size != 0)
+		{
+			OutTagsNamesWithAmount.FindOrAdd(Pair.Value.Name, 0) += Pair.Value.Size;
+		}
+	}
 }
 
 bool FLLMTracker::FindTagsForPtr(void* InPtr, TArray<const FTagData *, TInlineAllocator<static_cast<int32>(ELLMTagSet::Max)>>& OutTags)
 {
 	FLLMThreadState* State = GetOrCreateState();
 		
-	uint32* Size;
-	FLowLevelAllocInfo* AllocInfoPtr;
-	AllocationMap.Find(PointerKey(InPtr), Size, AllocInfoPtr);
-	if (!AllocInfoPtr)
+	uint32 Size;
+	FLowLevelAllocInfo AllocInfoPtr;
+	PointerKey FoundKey = AllocationMap.Find(PointerKey(InPtr), Size, AllocInfoPtr);
+	if (!FoundKey)
 	{
 		return false;
 	}
@@ -4460,10 +4543,10 @@ bool FLLMTracker::FindTagsForPtr(void* InPtr, TArray<const FTagData *, TInlineAl
 #if LLM_ALLOW_ASSETS_TAGS
 	for (int32 TagSetIndex = 0; TagSetIndex < static_cast<int32>(ELLMTagSet::Max); TagSetIndex++)
 	{
-		OutTags[TagSetIndex] = AllocInfoPtr->GetTag(LLMRef, static_cast<ELLMTagSet>(TagSetIndex));
+		OutTags[TagSetIndex] = AllocInfoPtr.GetTag(LLMRef, static_cast<ELLMTagSet>(TagSetIndex));
 	}
 #else
-	OutTags[static_cast<int32>(ELLMTagSet::None)] = AllocInfoPtr->GetTag(LLMRef);
+	OutTags[static_cast<int32>(ELLMTagSet::None)] = AllocInfoPtr.GetTag(LLMRef);
 #endif
 
 	return true;
diff --git a/Engine/Source/Runtime/Core/Private/HAL/LowLevelMemoryUtils.h b/Engine/Source/Runtime/Core/Private/HAL/LowLevelMemoryUtils.h
--- a/Engine/Source/Runtime/Core/Private/HAL/LowLevelMemoryUtils.h
+++ b/Engine/Source/Runtime/Core/Private/HAL/LowLevelMemoryUtils.h
@@ -25,7 +25,7 @@ typedef uint64 LLMNumAllocsType;
 typedef uint32 LLMNumAllocsType;
 #endif
 
-// POD types only
+// Trivially movable types without destructors only; memcpy is used and destructors are not called
 template<typename T, typename SizeType=int32>
 class FLLMArray
 {
@@ -323,18 +323,13 @@ public:
 
 	Values GetValue(const TKey& Key)
 	{
-		TValue1* Value1;
-		TValue2* Value2;
-		Key = Find(Key, Value1, Value2);
-		LLMEnsure(Key);
 		Values RetValues;
-		RetValues.Key = Key;
-		RetValues.Value1 = *Value1;
-		RetValues.Value2 = *Value2;
+		RetValues.Key = Find(Key, RetValues.Value1, RetValues.Value2);
+		LLMEnsure(RetValues.Key);
 		return RetValues;
 	}
 
-	TKey Find(const TKey& Key, TValue1*& OutValue1, TValue2*& OutValue2)
+	TKey Find(const TKey& Key, TValue1& OutValue1, TValue2& OutValue2)
 	{
 		SizeType KeyHash = Key.GetHashCode();
 		StripeData& Stripe = MapStripes[GetStripeIndex(KeyHash)];
@@ -346,13 +341,11 @@ public:
 		SizeType KeyIndex = Stripe.Map[MapIndex];
 		if (KeyIndex == InvalidIndex)
 		{
-			OutValue1 = nullptr;
-			OutValue2 = nullptr;
 			return TKey();
 		}
 
-		OutValue1 = &Stripe.Values1[KeyIndex];
-		OutValue2 = &Stripe.Values2[KeyIndex];
+		OutValue1 = Stripe.Values1[KeyIndex];
+		OutValue2 = Stripe.Values2[KeyIndex];
 
 		return Stripe.Keys[KeyIndex];
 	}
