LLM: Fix compile error when LLM_ALLOW_NAMES_TAGS=0.
#rnx

diff --git a/Engine/Source/Runtime/Core/Private/HAL/LowLevelMemTracker.cpp b/Engine/Source/Runtime/Core/Private/HAL/LowLevelMemTracker.cpp
--- a/Engine/Source/Runtime/Core/Private/HAL/LowLevelMemTracker.cpp
+++ b/Engine/Source/Runtime/Core/Private/HAL/LowLevelMemTracker.cpp
@@ -4293,7 +4293,11 @@ bool FLLMTracker::DumpForkedAllocationInfo()
 	Lines.Add(TEXT("Tag,SharedKib,PrivateKib,SplitKib,UnrefKib,TotalKib,SharedCount,PrivateCount,SplitCount,UnrefCount,TotalCount,CrossCount"));
 
 
+#if LLM_ENABLED_FULL_TAGS
 	for (TPair<int32, FCounts>& P : CountsPerTag)
+#else
+	for (TPair<ELLMTag, FCounts>& P : CountsPerTag)
+#endif
 	{
 		FLowLevelAllocInfo AllocInfoPlaceholder;
 		AllocInfoPlaceholder.SetCompressedTag(P.Key);
