Added enums for each of the supported loader types in the engine. GetLoaderType/LexToString now replace GetLoaderName.

Currently, loaders are identified via an FName returned by GetLoaderName, which is clumsy to use and a tad error prone. Instead, use an enum to identify the loader types via GetLoaderType.

This is also a prerequisite for UE-188423.

diff --git a/Engine/Source/Runtime/Core/Private/Misc/CoreGlobals.cpp b/Engine/Source/Runtime/Core/Private/Misc/CoreGlobals.cpp
--- a/Engine/Source/Runtime/Core/Private/Misc/CoreGlobals.cpp
+++ b/Engine/Source/Runtime/Core/Private/Misc/CoreGlobals.cpp
@@ -360,40 +360,58 @@ bool (*IsAsyncLoadingMultithreaded)() = &IsAsyncLoadingCoreInternal;
 void (*SuspendTextureStreamingRenderTasks)() = &appNoop;
 void (*ResumeTextureStreamingRenderTasks)() = &appNoop;
 
-static FName appNotAvailable()
+static ELoaderType LoaderNotInitialized()
 {
-	return TEXT("N/A");
+	return ELoaderType::NotInitialized;
+}
+ELoaderType (*GetLoaderType)() = &LoaderNotInitialized;
+
+const TCHAR* LexToString(ELoaderType Type)
+{
+	switch (Type)
+	{
+	case ELoaderType::NotInitialized:
+		return TEXT("NotInitialized");
+	case ELoaderType::LegacyLoader:
+		return TEXT("LegacyLoader");
+	case ELoaderType::EditorPackageLoader:
+		return TEXT("EditorPackageLoader");
+	case ELoaderType::ZenLoader:
+		return TEXT("ZenLoader");
+	default:
+		check(false);
+		return TEXT("");
+	}
 }
-FName (*GetLoaderName)() = &appNotAvailable;
 
 /** Whether the editor is currently loading a package or not												*/
-bool					GIsEditorLoadingPackage				= false;
+bool					GIsEditorLoadingPackage			= false;
 /** Whether the cooker is currently loading a package or not												*/
-bool					GIsCookerLoadingPackage = false;
+bool					GIsCookerLoadingPackage			= false;
 /** Whether GWorld points to the play in editor world														*/
 bool					GIsPlayInEditorWorld			= false;
-/** Unique ID for multiple PIE instances running in one process */
+/** Unique ID for multiple PIE instances running in one process												*/
 FPlayInEditorID			GPlayInEditorID;
-/** Whether or not PIE was attempting to play from PlayerStart							*/
+/** Whether or not PIE was attempting to play from PlayerStart												*/
 bool					GIsPIEUsingPlayerStart			= false;
 /** true if the runtime needs textures to be powers of two													*/
-bool					GPlatformNeedsPowerOfTwoTextures = false;
-/** Time at which FPlatformTime::Seconds() was first initialized (before main)											*/
+bool					GPlatformNeedsPowerOfTwoTextures= false;
+/** Time at which FPlatformTime::Seconds() was first initialized (before main)								*/
 double					GStartTime						= FPlatformTime::InitTiming();
 /** System time at engine init.																				*/
 FString					GSystemStartTime;
 /** Whether we are still in the initial loading proces.														*/
 bool					GIsInitialLoad					= true;
 /* Whether we are using the event driven loader */
-bool					GEventDrivenLoaderEnabled = false;
+bool					GEventDrivenLoaderEnabled		= false;
 
 /** Steadily increasing frame counter.																		*/
 uint64					GFrameCounter					= 0;
 uint64					GFrameCounterRenderThread		= 0;
 
 uint64					GLastGCFrame					= 0;
-/** The time input was sampled, in cycles. */
-uint64					GInputTime					= 0;
+/** The time input was sampled, in cycles.																	*/
+uint64					GInputTime						= 0;
 /** Incremented once per frame before the scene is being rendered. In split screen mode this is incremented once for all views (not for each view). */
 uint32					GFrameNumber					= 1;
 /** NEED TO RENAME, for RT version of GFrameTime use View.ViewFamily->FrameNumber or pass down from RT from GFrameTime). */
diff --git a/Engine/Source/Runtime/Core/Public/CoreGlobals.h b/Engine/Source/Runtime/Core/Public/CoreGlobals.h
--- a/Engine/Source/Runtime/Core/Public/CoreGlobals.h
+++ b/Engine/Source/Runtime/Core/Public/CoreGlobals.h
@@ -460,8 +460,18 @@ extern CORE_API bool (*IsAsyncLoadingSuspended)();
 /** Returns true if async loading is using the async loading thread */
 extern CORE_API bool(*IsAsyncLoadingMultithreaded)();
 
-/** Returns the name of the currently active loader, if any. */
-extern CORE_API FName(*GetLoaderName)();
+enum class ELoaderType : uint8
+{
+	NotInitialized,
+	LegacyLoader,
+	EditorPackageLoader,
+	ZenLoader
+};
+
+CORE_API const TCHAR* LexToString(ELoaderType Type);
+
+/** Returns the type of the currently active loader, if any. */
+extern CORE_API ELoaderType(*GetLoaderType)();
 
 /** Suspends texture updates caused by completed async IOs. */
 extern CORE_API void (*SuspendTextureStreamingRenderTasks)();
diff --git a/Engine/Source/Runtime/CoreUObject/Private/Serialization/AsyncLoading.cpp b/Engine/Source/Runtime/CoreUObject/Private/Serialization/AsyncLoading.cpp
--- a/Engine/Source/Runtime/CoreUObject/Private/Serialization/AsyncLoading.cpp
+++ b/Engine/Source/Runtime/CoreUObject/Private/Serialization/AsyncLoading.cpp
@@ -7524,7 +7524,7 @@ void FAsyncLoadingThread::FlushLoading(TConstArrayView<int32> RequestIDs)
 {
  	if (IsAsyncLoadingPackages())
 	{
-		checkf(IsInGameThread(), TEXT("The current loader %s is unable to FlushAsyncLoading from any thread other than the game thread."), *GetLoaderName().ToString());
+		checkf(IsInGameThread(), TEXT("The current loader %s is unable to FlushAsyncLoading from any thread other than the game thread."), LexToString(GetLoaderType()));
 
 		// Flushing async loading while loading is suspend will result in infinite stall
 		UE_CLOG(IsAsyncLoadingSuspendedInternal(), LogStreaming, Fatal, TEXT("Cannot Flush Async Loading while async loading is suspended (%d)"), GetAsyncLoadingSuspendedCount());
diff --git a/Engine/Source/Runtime/CoreUObject/Private/Serialization/AsyncLoading2.cpp b/Engine/Source/Runtime/CoreUObject/Private/Serialization/AsyncLoading2.cpp
--- a/Engine/Source/Runtime/CoreUObject/Private/Serialization/AsyncLoading2.cpp
+++ b/Engine/Source/Runtime/CoreUObject/Private/Serialization/AsyncLoading2.cpp
@@ -3080,9 +3080,9 @@ public:
 	FAsyncLoadingThread2(FIoDispatcher& IoDispatcher, IAsyncPackageLoader* InUncookedPackageLoader);
 	virtual ~FAsyncLoadingThread2();
 
-	virtual FName GetLoaderName() const override
+	virtual ELoaderType GetLoaderType() const override
 	{
-		return TEXT("ZenLoader");
+		return ELoaderType::ZenLoader;
 	}
 
 private:
@@ -9161,12 +9161,13 @@ void FAsyncLoadingThread2::FlushLoading(TConstArrayView<int32> RequestIDs)
 #else
 		const bool bIsFlushSupportedOnCurrentThread = IsInGameThread();
 #endif
+		ELoaderType LoaderType = GetLoaderType();
 
-		ensureMsgf(bIsFlushSupportedOnCurrentThread, TEXT("The current loader '%s' is unable to FlushAsyncLoading from the current thread."), *GetLoaderName().ToString());
+		ensureMsgf(bIsFlushSupportedOnCurrentThread, TEXT("The current loader '%s' is unable to FlushAsyncLoading from the current thread."), LexToString(LoaderType));
 
 		if (!bIsFlushSupportedOnCurrentThread)
 		{
-			UE_LOG(LogStreaming, Error, TEXT("The current loader '%s' is unable to FlushAsyncLoading from the current thread. Flush will be ignored."), *GetLoaderName().ToString());
+			UE_LOG(LogStreaming, Error, TEXT("The current loader '%s' is unable to FlushAsyncLoading from the current thread. Flush will be ignored."), LexToString(LoaderType));
 			return;
 		}
 
diff --git a/Engine/Source/Runtime/CoreUObject/Private/Serialization/AsyncLoadingThread.h b/Engine/Source/Runtime/CoreUObject/Private/Serialization/AsyncLoadingThread.h
--- a/Engine/Source/Runtime/CoreUObject/Private/Serialization/AsyncLoadingThread.h
+++ b/Engine/Source/Runtime/CoreUObject/Private/Serialization/AsyncLoadingThread.h
@@ -348,9 +348,9 @@ public:
 	FAsyncLoadingThread(int32 InThreadIndex);
 	virtual ~FAsyncLoadingThread();
 
-	virtual FName GetLoaderName() const override
+	virtual ELoaderType GetLoaderType() const override
 	{
-		return TEXT("LegacyLoader");
+		return ELoaderType::LegacyLoader;
 	}
 
 	IAsyncPackageLoader* GetIoStorePackageLoader() const
diff --git a/Engine/Source/Runtime/CoreUObject/Private/Serialization/AsyncPackageLoader.cpp b/Engine/Source/Runtime/CoreUObject/Private/Serialization/AsyncPackageLoader.cpp
--- a/Engine/Source/Runtime/CoreUObject/Private/Serialization/AsyncPackageLoader.cpp
+++ b/Engine/Source/Runtime/CoreUObject/Private/Serialization/AsyncPackageLoader.cpp
@@ -326,9 +326,9 @@ bool IsAsyncLoadingMultithreadedCoreUObjectInternal()
 	return GetAsyncPackageLoader().IsMultithreaded();
 }
 
-FName GetLoaderNameInternal()
+ELoaderType GetLoaderTypeInternal()
 {
-	return GetAsyncPackageLoader().GetLoaderName();
+	return GetAsyncPackageLoader().GetLoaderType();
 }
 
 void SuspendAsyncLoadingInternal()
diff --git a/Engine/Source/Runtime/CoreUObject/Private/Serialization/EditorPackageLoader.cpp b/Engine/Source/Runtime/CoreUObject/Private/Serialization/EditorPackageLoader.cpp
--- a/Engine/Source/Runtime/CoreUObject/Private/Serialization/EditorPackageLoader.cpp
+++ b/Engine/Source/Runtime/CoreUObject/Private/Serialization/EditorPackageLoader.cpp
@@ -25,9 +25,9 @@ public:
 
 	virtual ~FEditorPackageLoader () { }
 
-	virtual FName GetLoaderName() const override
+	virtual ELoaderType GetLoaderType() const override
 	{
-		return TEXT("EditorPackageLoader");
+		return ELoaderType::EditorPackageLoader;
 	}
 
 	virtual void InitializeLoading() override
diff --git a/Engine/Source/Runtime/CoreUObject/Private/UObject/CoreNative.cpp b/Engine/Source/Runtime/CoreUObject/Private/UObject/CoreNative.cpp
--- a/Engine/Source/Runtime/CoreUObject/Private/UObject/CoreNative.cpp
+++ b/Engine/Source/Runtime/CoreUObject/Private/UObject/CoreNative.cpp
@@ -17,7 +17,7 @@ void SuspendAsyncLoadingInternal();
 void ResumeAsyncLoadingInternal();
 bool IsAsyncLoadingSuspendedInternal();
 bool IsAsyncLoadingMultithreadedCoreUObjectInternal();
-FName GetLoaderNameInternal();
+ELoaderType GetLoaderTypeInternal();
 
 // CoreUObject module. Handles UObject system pre-init (registers init function with Core callbacks).
 class FCoreUObjectModule : public FDefaultModuleImpl
@@ -47,7 +47,7 @@ public:
 		ResumeAsyncLoading = &ResumeAsyncLoadingInternal;
 		IsAsyncLoadingSuspended = &IsAsyncLoadingSuspendedInternal;
 		IsAsyncLoadingMultithreaded = &IsAsyncLoadingMultithreadedCoreUObjectInternal;
-		GetLoaderName = &GetLoaderNameInternal;
+		GetLoaderType = &GetLoaderTypeInternal;
 
 		// Register the script callstack callback to the runtime error logging
 #if UE_RAISE_RUNTIME_ERRORS
diff --git a/Engine/Source/Runtime/CoreUObject/Private/UObject/PropertyObject.cpp b/Engine/Source/Runtime/CoreUObject/Private/UObject/PropertyObject.cpp
--- a/Engine/Source/Runtime/CoreUObject/Private/UObject/PropertyObject.cpp
+++ b/Engine/Source/Runtime/CoreUObject/Private/UObject/PropertyObject.cpp
@@ -60,13 +60,14 @@ EConvertFromTypeResult FObjectProperty::ConvertFromType(const FPropertyTag& Tag,
 
 		// If we're async loading, only zenloader supports synchronous loads from the async loading thread.
 		// Other loaders are not safe to do a sync load because they may crash or fail to set the variable, so throw an error if it's not already in memory
-		if (!IsInGameThread() && GetLoaderName() != TEXT("ZenLoader"))
+		ELoaderType LoaderType = GetLoaderType();
+		if (!IsInGameThread() && (LoaderType != ELoaderType::ZenLoader))
 		{
 			PreviousValueObj = PreviousValue.Get();
 
 			if (!PreviousValueObj && !PreviousValue.IsNull())
 			{
-				UE_LOG(LogClass, Error, TEXT("Failed to convert soft path %s to unloaded object as this is not supported during async loading with the currently active loader '%s'. Load and resave %s in the editor to fix!"), *PreviousValue.ToString(), *GetLoaderName().ToString(), *Slot.GetUnderlyingArchive().GetArchiveName());
+				UE_LOG(LogClass, Error, TEXT("Failed to convert soft path %s to unloaded object as this is not supported during async loading with the currently active loader '%s'. Load and resave %s in the editor to fix!"), *PreviousValue.ToString(), LexToString(LoaderType), *Slot.GetUnderlyingArchive().GetArchiveName());
 			}
 		}
 		else
diff --git a/Engine/Source/Runtime/CoreUObject/Private/UObject/UObjectGlobals.cpp b/Engine/Source/Runtime/CoreUObject/Private/UObject/UObjectGlobals.cpp
--- a/Engine/Source/Runtime/CoreUObject/Private/UObject/UObjectGlobals.cpp
+++ b/Engine/Source/Runtime/CoreUObject/Private/UObject/UObjectGlobals.cpp
@@ -1698,7 +1698,7 @@ UPackage* LoadPackageInternal(UPackage* InOuter, const FPackagePath& PackagePath
 		}
 	}
 
-	checkf(IsInGameThread(), TEXT("Unable to load %s. Objects and Packages can only be loaded from the game thread with the currently active loader '%s'."), *PackagePath.GetDebugName(), *GetLoaderName().ToString());
+	checkf(IsInGameThread(), TEXT("Unable to load %s. Objects and Packages can only be loaded from the game thread with the currently active loader '%s'."), *PackagePath.GetDebugName(), LexToString(GetLoaderType()));
 	UPackage* Result = nullptr;
 
 #if WITH_EDITOR
diff --git a/Engine/Source/Runtime/CoreUObject/Public/Serialization/AsyncPackageLoader.h b/Engine/Source/Runtime/CoreUObject/Public/Serialization/AsyncPackageLoader.h
--- a/Engine/Source/Runtime/CoreUObject/Public/Serialization/AsyncPackageLoader.h
+++ b/Engine/Source/Runtime/CoreUObject/Public/Serialization/AsyncPackageLoader.h
@@ -230,7 +230,7 @@ public:
 
 	virtual void NotifyRegistrationComplete() = 0;
 
-	virtual FName GetLoaderName() const = 0;
+	virtual ELoaderType GetLoaderType() const = 0;
 
 protected:
 	static int32 GetNextRequestId();
diff --git a/Engine/Source/Runtime/Engine/Private/Tests/Loading/AsyncLoadingTests_Shared.h b/Engine/Source/Runtime/Engine/Private/Tests/Loading/AsyncLoadingTests_Shared.h
--- a/Engine/Source/Runtime/Engine/Private/Tests/Loading/AsyncLoadingTests_Shared.h
+++ b/Engine/Source/Runtime/Engine/Private/Tests/Loading/AsyncLoadingTests_Shared.h
@@ -120,11 +120,12 @@ public:
 protected:
 	bool CanRunInEnvironment(const FString& TestParams, FString* OutReason, bool* OutWarn) const override
 	{
-		if (GetLoaderName() != TEXT("ZenLoader"))
+		ELoaderType LoaderType = GetLoaderType();
+		if (LoaderType != ELoaderType::ZenLoader)
 		{
 			if (OutReason)
 			{
-				*OutReason = FString::Printf(TEXT("Test %s is for ZenLoader only. Cannot run on non-compliant loader currently active: %s"), *GetTestName(), *GetLoaderName().ToString());
+				*OutReason = FString::Printf(TEXT("Test %s is for ZenLoader only. Cannot run on non-compliant loader currently active: %s"), *GetTestName(), LexToString(LoaderType));
 			}
 			return false;
 		}
