Some fixes for compile errors when building a modular target with the AutoRTFM compiler.

diff --git a/Engine/Source/Runtime/Core/Private/AutoRTFM/API.cpp b/Engine/Source/Runtime/Core/Private/AutoRTFM/API.cpp
--- a/Engine/Source/Runtime/Core/Private/AutoRTFM/API.cpp
+++ b/Engine/Source/Runtime/Core/Private/AutoRTFM/API.cpp
@@ -5,7 +5,7 @@
 #include "HAL/IConsoleManager.h"
 
 #if UE_AUTORTFM
-bool GAutoRTFMRuntimeEnabled = false;
+UE_AUTORTFM_API bool GAutoRTFMRuntimeEnabled = false;
 static FAutoConsoleVariableRef CVarAutoRTFMRuntimeEnabled(
 	TEXT("AutoRTFMRuntimeEnabled"),
 	GAutoRTFMRuntimeEnabled,
diff --git a/Engine/Source/Runtime/Core/Private/AutoRTFM/CompilerInterface.cpp b/Engine/Source/Runtime/Core/Private/AutoRTFM/CompilerInterface.cpp
--- a/Engine/Source/Runtime/Core/Private/AutoRTFM/CompilerInterface.cpp
+++ b/Engine/Source/Runtime/Core/Private/AutoRTFM/CompilerInterface.cpp
@@ -28,7 +28,7 @@ void CheckAlignment(FContext* Context, void* Ptr, size_t AlignmentMask)
     }
 }
 
-extern "C" AUTORTFM_NO_ASAN void autortfm_record_write(void* Ptr, size_t Size)
+extern "C" UE_AUTORTFM_API AUTORTFM_NO_ASAN void autortfm_record_write(void* Ptr, size_t Size)
 {
 	// check for writes to null here so we end up crashing in the user
 	// code rather than in the autortfm runtime.
@@ -41,7 +41,7 @@ extern "C" AUTORTFM_NO_ASAN void autortfm_record_write(void* Ptr, size_t Size)
 	Context->RecordWrite(Ptr, Size);
 }
 
-extern "C" AUTORTFM_NO_ASAN void autortfm_record_write_8(void* Ptr)
+extern "C" UE_AUTORTFM_API AUTORTFM_NO_ASAN void autortfm_record_write_8(void* Ptr)
 {
 	// check for writes to null here so we end up crashing in the user
 	// code rather than in the autortfm runtime.
@@ -54,31 +54,31 @@ extern "C" AUTORTFM_NO_ASAN void autortfm_record_write_8(void* Ptr)
 	Context->RecordWrite<8>(Ptr);
 }
 
-extern "C" void* autortfm_lookup_function(void* OriginalFunction, const char* Where)
+extern "C" UE_AUTORTFM_API void* autortfm_lookup_function(void* OriginalFunction, const char* Where)
 {
 	FContext* Context = FContext::Get();
     return FunctionMapLookup(OriginalFunction, Where);
 }
 
-extern "C" void autortfm_memcpy(void* Dst, const void* Src, size_t Size)
+extern "C" UE_AUTORTFM_API void autortfm_memcpy(void* Dst, const void* Src, size_t Size)
 {
 	FContext* Context = FContext::Get();
     Memcpy(Dst, Src, Size, Context);
 }
 
-extern "C" void autortfm_memmove(void* Dst, const void* Src, size_t Size)
+extern "C" UE_AUTORTFM_API void autortfm_memmove(void* Dst, const void* Src, size_t Size)
 {
 	FContext* Context = FContext::Get();
     Memmove(Dst, Src, Size, Context);
 }
 
-extern "C" void autortfm_memset(void* Dst, int Value, size_t Size)
+extern "C" UE_AUTORTFM_API void autortfm_memset(void* Dst, int Value, size_t Size)
 {
 	FContext* Context = FContext::Get();
     Memset(Dst, Value, Size, Context);
 }
 
-extern "C" void autortfm_llvm_fail(const char* Message)
+extern "C" UE_AUTORTFM_API void autortfm_llvm_fail(const char* Message)
 {
     if (Message)
     {
@@ -93,12 +93,12 @@ extern "C" void autortfm_llvm_fail(const char* Message)
     Context->AbortByLanguageAndThrow();
 }
 
-extern "C" void autortfm_llvm_alignment_error(FContext* Context, void* Ptr, size_t Alignment, const char* Message)
+extern "C" UE_AUTORTFM_API void autortfm_llvm_alignment_error(FContext* Context, void* Ptr, size_t Alignment, const char* Message)
 {
     AbortDueToBadAlignment(Context, Ptr, Alignment, Message);
 }
 
-extern "C" void autortfm_llvm_error(const char* Message)
+extern "C" UE_AUTORTFM_API void autortfm_llvm_error(const char* Message)
 {
 	if (Message)
 	{
diff --git a/Engine/Source/Runtime/Core/Public/AutoRTFM/AutoRTFM.h b/Engine/Source/Runtime/Core/Public/AutoRTFM/AutoRTFM.h
--- a/Engine/Source/Runtime/Core/Public/AutoRTFM/AutoRTFM.h
+++ b/Engine/Source/Runtime/Core/Public/AutoRTFM/AutoRTFM.h
@@ -238,7 +238,7 @@ UE_AUTORTFM_FORCEINLINE void autortfm_open(void (*work)(void* arg), void* arg)
 //
 // Guaranteed to crash if called outside a transaction.
 #if UE_AUTORTFM
-UE_AUTORTFM_API [[nodiscard]] autortfm_status autortfm_close(void (*work)(void* arg), void* arg);
+[[nodiscard]] UE_AUTORTFM_API autortfm_status autortfm_close(void (*work)(void* arg), void* arg);
 #else
 [[nodiscard]] UE_AUTORTFM_FORCEINLINE autortfm_status autortfm_close(void (*work)(void* arg), void* arg)
 {
