Add support for in place Base64 encoding/decoding

diff --git a/Engine/Source/Runtime/Core/Private/Misc/Base64.cpp b/Engine/Source/Runtime/Core/Private/Misc/Base64.cpp
--- a/Engine/Source/Runtime/Core/Private/Misc/Base64.cpp
+++ b/Engine/Source/Runtime/Core/Private/Misc/Base64.cpp
@@ -68,8 +68,7 @@ FString FBase64::Encode(const uint8* Source, uint32 Length, EBase64Mode Mode)
 
 	TArray<TCHAR, FString::AllocatorType>& OutCharArray = OutBuffer.GetCharArray();
 	OutCharArray.SetNum(ExpectedLength + 1);
-	int64 EncodedLength = Encode(Source, Length, OutCharArray.GetData(), Mode);
-	verify(EncodedLength == OutBuffer.Len());
+	Encode(Source, Length, OutCharArray.GetData(), Mode);
 
 	return OutBuffer;
 }
@@ -78,52 +77,45 @@ template<typename CharType> uint32 FBase64::Encode(const uint8* Source, uint32 L
 {
 	check(Mode == EBase64Mode::Standard || Mode == EBase64Mode::UrlSafe);
 
-	CharType* EncodedBytes = Dest;
 	const uint8* const EncodingAlphabet = (Mode == EBase64Mode::UrlSafe ? GBase64UrlSafeEncodingAlphabet : GBase64EncodingAlphabet);
 
-	// Loop through the buffer converting 3 bytes of binary data at a time
-	while (Length >= 3)
-	{
-		uint8 A = *Source++;
-		uint8 B = *Source++;
-		uint8 C = *Source++;
-		Length -= 3;
-
-		// The algorithm takes 24 bits of data (3 bytes) and breaks it into 4 6bit chunks represented as ascii
-		uint32 ByteTriplet = A << 16 | B << 8 | C;
-
-		// Use the 6bit block to find the representation ascii character for it
-		EncodedBytes[3] = EncodingAlphabet[ByteTriplet & 0x3F];
-		ByteTriplet >>= 6;
-		EncodedBytes[2] = EncodingAlphabet[ByteTriplet & 0x3F];
-		ByteTriplet >>= 6;
-		EncodedBytes[1] = EncodingAlphabet[ByteTriplet & 0x3F];
-		ByteTriplet >>= 6;
-		EncodedBytes[0] = EncodingAlphabet[ByteTriplet & 0x3F];
+	const uint32 PaddingLength = Length % 3;
+	const uint32 ExpectedLength = GetEncodedDataSize(Length);
+	CharType* EncodedBytes = Dest + ExpectedLength;
 
-		// Now we can append this buffer to our destination string
-		EncodedBytes += 4;
+	// Add a null terminator
+	*EncodedBytes = CHARTEXT(CharType, '\0');
+	
+	if (Length == 0)
+	{
+		return 0;
 	}
 
+	const uint8* ReversedSource = Source + Length;
+
 	// Since this algorithm operates on blocks, we may need to pad the last chunks
-	if (Length > 0)
+	if (PaddingLength > 0)
 	{
-		uint8 A = *Source++;
+		EncodedBytes -= 4;
+
+		uint8 A = 0;
 		uint8 B = 0;
 		uint8 C = 0;
 		// Grab the second character if it is a 2 uint8 finish
-		if (Length == 2)
+		if (PaddingLength == 2)
 		{
-			B = *Source;
+			B = *--ReversedSource;
 		}
+		A = *--ReversedSource;
+
 		uint32 ByteTriplet = A << 16 | B << 8 | C;
 		// Pad with = to make a 4 uint8 chunk
-		EncodedBytes[3] = '=';
+		EncodedBytes[3] = CHARTEXT(CharType, '=');
 		ByteTriplet >>= 6;
 		// If there's only one 1 uint8 left in the source, then you need 2 pad chars
-		if (Length == 1)
+		if (PaddingLength == 1)
 		{
-			EncodedBytes[2] = '=';
+			EncodedBytes[2] = CHARTEXT(CharType, '=');
 		}
 		else
 		{
@@ -134,14 +126,36 @@ template<typename CharType> uint32 FBase64::Encode(const uint8* Source, uint32 L
 		EncodedBytes[1] = EncodingAlphabet[ByteTriplet & 0x3F];
 		ByteTriplet >>= 6;
 		EncodedBytes[0] = EncodingAlphabet[ByteTriplet & 0x3F];
+	}
+	
+	Length -= PaddingLength;
+
+	// Loop through the buffer converting 3 bytes of binary data at a time
+	while (Length > 0)
+	{
+		EncodedBytes -= 4;
+
+		uint8 C = *--ReversedSource;
+		uint8 B = *--ReversedSource;
+		uint8 A = *--ReversedSource;
+		Length -= 3;
+
+		// The algorithm takes 24 bits of data (3 bytes) and breaks it into 4 6bit chunks represented as ascii
+		uint32 ByteTriplet = A << 16 | B << 8 | C;
 
-		EncodedBytes += 4;
+		// Use the 6bit block to find the representation ascii character for it
+		EncodedBytes[3] = EncodingAlphabet[ByteTriplet & 0x3F];
+		ByteTriplet >>= 6;
+		EncodedBytes[2] = EncodingAlphabet[ByteTriplet & 0x3F];
+		ByteTriplet >>= 6;
+		EncodedBytes[1] = EncodingAlphabet[ByteTriplet & 0x3F];
+		ByteTriplet >>= 6;
+		EncodedBytes[0] = EncodingAlphabet[ByteTriplet & 0x3F];
 	}
 
-	// Add a null terminator
-	*EncodedBytes = 0;
+	verify(EncodedBytes == Dest);
 
-	return UE_PTRDIFF_TO_UINT32(EncodedBytes - Dest);
+	return ExpectedLength;
 }
 
 template CORE_API uint32 FBase64::Encode<ANSICHAR>(const uint8* Source, uint32 Length, ANSICHAR* Dest, EBase64Mode Mode);
diff --git a/Engine/Source/Runtime/Core/Public/Misc/Base64.h b/Engine/Source/Runtime/Core/Public/Misc/Base64.h
--- a/Engine/Source/Runtime/Core/Public/Misc/Base64.h
+++ b/Engine/Source/Runtime/Core/Public/Misc/Base64.h
@@ -57,7 +57,7 @@ struct FBase64
 	 *
 	 * @param Source The binary data to encode
 	 * @param Length Length of the binary data to be encoded
-	 * @param Dest Buffer to receive the encoded data. Must be large enough to contain the entire output data (see GetEncodedDataSize()).
+	 * @param Dest Buffer to receive the encoded data. Must be large enough to contain the entire output data (see GetEncodedDataSize()). Can point to the same buffer as Source
 	 * @param Mode The mode to use for encoding. Default is EBase64Mode::Standard
 	 *
 	 * @return The length of the encoded data
@@ -103,7 +103,7 @@ struct FBase64
 	 *
 	 * @param Source The Base64 encoded string
 	 * @param Length Length of the Base64 encoded string
-	 * @param Dest Buffer to receive the decoded data
+	 * @param Dest Buffer to receive the decoded data.  Must be large enough to contain the entire output data (see GetDecodedDataSize()). Can point to the same buffer as Source
 	 * @param Mode The mode to use for decoding. Default is EBase64Mode::Standard
 	 *
 	 * @return true if the buffer was decoded, false if it was invalid.
