Check for Clang before Windows in the analysis include.

diff --git a/Engine/Source/Runtime/Core/Public/HAL/Platform.h b/Engine/Source/Runtime/Core/Public/HAL/Platform.h
--- a/Engine/Source/Runtime/Core/Public/HAL/Platform.h
+++ b/Engine/Source/Runtime/Core/Public/HAL/Platform.h
@@ -108,10 +108,10 @@
 #endif // defined(__clang__)
 #endif
 
-#if PLATFORM_WINDOWS
+#if PLATFORM_COMPILER_CLANG
+#include "Clang/ClangPlatformCodeAnalysis.h"
+#elif PLATFORM_WINDOWS
 	#include "Windows/WindowsPlatformCodeAnalysis.h"
-#elif PLATFORM_COMPILER_CLANG
-	#include "Clang/ClangPlatformCodeAnalysis.h"
 #endif
 
 #ifndef USING_ADDRESS_SANITISER
