Fix PS5 File handle handle registry not tracking ucas file handles opened by IOStore.
- IOStore ucas handles now get freed and reopened as needed when file handle pressure is too great.
- FRegisteredFileHandle now refcounts read requests so it can handle multiple read requests with the same handle at the same time.
- Failure to reopen files is now passed back up as file read failures
- Add some additional logging for FFileHelper errors
- Add error logging explicitly for runningout of PS5 file handles
- Make file handle casts use static_cast
- PS5 file handle now returns -1 for size if fstat fails, this matches windows behavior
- FAPRHandleCache now ignores files outside of /app0.  This prevents file churn in /download0 from bloating it.

diff --git a/Engine/Source/Runtime/Core/Private/IO/IoStore.cpp b/Engine/Source/Runtime/Core/Private/IO/IoStore.cpp
--- a/Engine/Source/Runtime/Core/Private/IO/IoStore.cpp
+++ b/Engine/Source/Runtime/Core/Private/IO/IoStore.cpp
@@ -2236,7 +2236,7 @@ public:
 	{
 		FCriticalSection HandleLock[NumHandlesPerFile];
 		IFileHandle* Handle[NumHandlesPerFile];
-		std::atomic_uint32_t NextHandleIndex;
+		std::atomic_uint32_t NextHandleIndex{0};
 		bool bValid = false;
 
 		FContainerFileAccess(IPlatformFile& Ipf, const TCHAR* ContainerFileName)
diff --git a/Engine/Source/Runtime/Core/Private/Misc/FileHelper.cpp b/Engine/Source/Runtime/Core/Private/Misc/FileHelper.cpp
--- a/Engine/Source/Runtime/Core/Private/Misc/FileHelper.cpp
+++ b/Engine/Source/Runtime/Core/Private/Misc/FileHelper.cpp
@@ -612,11 +612,15 @@ bool FFileHelper::SaveStringToFile( FStringView String, const TCHAR* Filename,
 {
 	// max size of the string is a UCS2CHAR for each character and some UNICODE magic 
 	TUniquePtr<FArchive> Ar = TUniquePtr<FArchive>( FileManager->CreateFileWriter( Filename, WriteFlags ) );
-	if( !Ar )
+	if (!Ar)
+	{
+		UE_LOG(LogStreaming, Warning, TEXT("SaveStringToFile: Failed to open writer. File:%s"), Filename);
 		return false;
+	}
 
-	if( String.IsEmpty() )
+	if (String.IsEmpty())
 	{
+		UE_LOG(LogStreaming, Warning, TEXT("SaveStringToFile: Empty String. File:%s"), Filename);
 		Ar->Close();
 		return true;
 	}
diff --git a/Engine/Source/Runtime/Core/Private/Unix/UnixPlatformFile.cpp b/Engine/Source/Runtime/Core/Private/Unix/UnixPlatformFile.cpp
--- a/Engine/Source/Runtime/Core/Private/Unix/UnixPlatformFile.cpp
+++ b/Engine/Source/Runtime/Core/Private/Unix/UnixPlatformFile.cpp
@@ -191,9 +191,19 @@ public:
 	{
 		struct FScopedReadTracker
 		{
-			FScopedReadTracker(FFileHandleUnix& InHandle) : Handle(InHandle) { GFileRegistry.TrackStartRead(&Handle); }
-			~FScopedReadTracker() { GFileRegistry.TrackEndRead(&Handle); }
+			FScopedReadTracker(FFileHandleUnix& InHandle) : Handle(InHandle) 
+			{ 
+				bSuccess = GFileRegistry.TrackStartRead(&Handle);
+			}
+			~FScopedReadTracker() 
+			{
+				if (bSuccess)
+				{
+					GFileRegistry.TrackEndRead(&Handle); 
+				}
+			}
 			FFileHandleUnix& Handle;
+			bool bSuccess = false;
 		};
 
 		check(IsValid());
@@ -201,6 +211,11 @@ public:
 		{
 			// Handle virtual file handles (only in read mode, write mode doesn't use the file handle registry)
 			FScopedReadTracker ScopedReadTracker(*this);
+			if (!ScopedReadTracker.bSuccess)
+			{
+				return false;
+			}
+
 			FScopedDiskUtilizationTracker Tracker(BytesToRead, FileOffset);
 
 			// seek to the offset on seek? this matches console behavior more closely
diff --git a/Engine/Source/Runtime/Core/Public/HAL/PlatformFileCommon.h b/Engine/Source/Runtime/Core/Public/HAL/PlatformFileCommon.h
--- a/Engine/Source/Runtime/Core/Public/HAL/PlatformFileCommon.h
+++ b/Engine/Source/Runtime/Core/Public/HAL/PlatformFileCommon.h
@@ -15,20 +15,13 @@
 
 class FRegisteredFileHandle : public IFileHandle
 {
-public:
-	FRegisteredFileHandle()
-		: NextLink(nullptr)
-		, PreviousLink(nullptr)
-		, bIsOpenAndAvailableForClosing(false)
-	{
-	}
-
 private:
 	friend class FFileHandleRegistry;
 
-	FRegisteredFileHandle* NextLink;
-	FRegisteredFileHandle* PreviousLink;
-	bool bIsOpenAndAvailableForClosing;
+	FRegisteredFileHandle* NextLink = nullptr;
+	FRegisteredFileHandle* PreviousLink = nullptr;
+	uint32 ReadRequestCount = 0;
+	bool bIsOpenAndAvailableForClosing = false;
 };
 
 
@@ -37,14 +30,11 @@ class FFileHandleRegistry
 public:
 	FFileHandleRegistry(int32 InMaxOpenHandles)
 		: MaxOpenHandles(InMaxOpenHandles)
-		, OpenAndAvailableForClosingHead(nullptr)
-		, OpenAndAvailableForClosingTail(nullptr)
-	{
-	}
+	{}
 
 	virtual ~FFileHandleRegistry() = default;
 
-	FRegisteredFileHandle* InitialOpenFile(const TCHAR* Filename)
+	UE_NODISCARD FRegisteredFileHandle* InitialOpenFile(const TCHAR* Filename)
 	{
 		if (HandlesCurrentlyInUse.Increment() > MaxOpenHandles)
 		{
@@ -70,6 +60,7 @@ public:
 		bool bWasOpen = false;
 		{
 			FScopeLock Lock(&LockSection);
+			check(Handle->ReadRequestCount == 0);
 			if (Handle->bIsOpenAndAvailableForClosing)
 			{
 				UnLink(Handle);
@@ -83,15 +74,22 @@ public:
 		}
 	}
 
-	void TrackStartRead(FRegisteredFileHandle* Handle)
+	// Only returns false if the platform handle cannot be reopened.
+	// TrackEndRead should only be called if TrackStartRead succeeded
+	UE_NODISCARD bool TrackStartRead(FRegisteredFileHandle* Handle)
 	{
 		{
 			FScopeLock Lock(&LockSection);
+			
+			if (Handle->ReadRequestCount++ != 0)
+			{
+				return true;
+			}
 
 			if (Handle->bIsOpenAndAvailableForClosing)
 			{
 				UnLink(Handle);
-				return;
+				return true;
 			}
 		}
 
@@ -99,14 +97,26 @@ public:
 		{
 			FreeHandles();
 		}
+
 		// can do this out of the lock, in case it's slow
-		PlatformReopenFile(Handle);
+		bool bSuccess = PlatformReopenFile(Handle);
+
+		if (!bSuccess)
+		{
+			FScopeLock Lock(&LockSection);
+			--Handle->ReadRequestCount;
+		}
+
+		return bSuccess;
 	}
 
+	// TrackEndRead should only be called if TrackStartRead succeeded
 	void TrackEndRead(FRegisteredFileHandle* Handle)
 	{
+		FScopeLock Lock(&LockSection);
+
+		if (--Handle->ReadRequestCount == 0)
 		{
-			FScopeLock Lock(&LockSection);
 			LinkToTail(Handle);
 		}
 	}
@@ -120,7 +130,7 @@ private:
 
 	void FreeHandles()
 	{
-		// do we need to make room for a file handle? this assumes that the critical section is already locked
+		// do we need to make room for a file handle?
 		while (HandlesCurrentlyInUse.GetValue() > MaxOpenHandles)
 		{
 			FRegisteredFileHandle* ToBeClosed = nullptr;
@@ -136,13 +146,12 @@ private:
 			}
 			else
 			{
-				FPlatformMisc::LowLevelOutputDebugString(TEXT("Spinning because we are actively reading from more file handles than we have possible handles.\r\n"));
+				FPlatformMisc::LowLevelOutputDebugString(TEXT("FFileHandleRegistry: Spinning because we are actively reading from more file handles than we have possible handles.\r\n"));
 				FPlatformProcess::SleepNoStats(.1f);
 			}
 		}
 	}
 
-
 	void LinkToTail(FRegisteredFileHandle* Handle)
 	{
 		check(!Handle->PreviousLink && !Handle->NextLink && !Handle->bIsOpenAndAvailableForClosing);
@@ -215,10 +224,10 @@ private:
 	// critical section to protect the below arrays
 	FCriticalSection LockSection;
 
-	int32 MaxOpenHandles;
+	int32 MaxOpenHandles = 0;
 
-	FRegisteredFileHandle* OpenAndAvailableForClosingHead;
-	FRegisteredFileHandle* OpenAndAvailableForClosingTail;
+	FRegisteredFileHandle* OpenAndAvailableForClosingHead = nullptr;
+	FRegisteredFileHandle* OpenAndAvailableForClosingTail = nullptr;
 
 	FThreadSafeCounter HandlesCurrentlyInUse;
 };
