CIS fix - FFileHelper::SaveStringToFile
- Fix unneeded warning for empty string
- Make sure archive error flags are still checked

diff --git a/Engine/Source/Runtime/Core/Private/Misc/FileHelper.cpp b/Engine/Source/Runtime/Core/Private/Misc/FileHelper.cpp
--- a/Engine/Source/Runtime/Core/Private/Misc/FileHelper.cpp
+++ b/Engine/Source/Runtime/Core/Private/Misc/FileHelper.cpp
@@ -618,40 +618,36 @@ bool FFileHelper::SaveStringToFile( FStringView String, const TCHAR* Filename,
 		return false;
 	}
 
-	if (String.IsEmpty())
+	if (!String.IsEmpty())
 	{
-		UE_LOG(LogStreaming, Warning, TEXT("SaveStringToFile: Empty String. File:%s"), Filename);
-		Ar->Close();
-		return true;
-	}
-
-	bool SaveAsUnicode = EncodingOptions == EEncodingOptions::ForceUnicode || ( EncodingOptions == EEncodingOptions::AutoDetect && !FCString::IsPureAnsi(String.GetData(), String.Len()) );
-	if( EncodingOptions == EEncodingOptions::ForceUTF8 )
-	{
-		UTF8CHAR UTF8BOM[] = { (UTF8CHAR)0xEF, (UTF8CHAR)0xBB, (UTF8CHAR)0xBF };
-		Ar->Serialize( &UTF8BOM, UE_ARRAY_COUNT(UTF8BOM) * sizeof(UTF8CHAR) );
+		bool SaveAsUnicode = EncodingOptions == EEncodingOptions::ForceUnicode || (EncodingOptions == EEncodingOptions::AutoDetect && !FCString::IsPureAnsi(String.GetData(), String.Len()));
+		if (EncodingOptions == EEncodingOptions::ForceUTF8)
+		{
+			UTF8CHAR UTF8BOM[] = { (UTF8CHAR)0xEF, (UTF8CHAR)0xBB, (UTF8CHAR)0xBF };
+			Ar->Serialize(&UTF8BOM, UE_ARRAY_COUNT(UTF8BOM) * sizeof(UTF8CHAR));
 
-		FTCHARToUTF8 UTF8String(String.GetData(), String.Len());
-		Ar->Serialize( (UTF8CHAR*)UTF8String.Get(), UTF8String.Length() * sizeof(UTF8CHAR) );
-	}
-	else if ( EncodingOptions == EEncodingOptions::ForceUTF8WithoutBOM )
-	{
-		FTCHARToUTF8 UTF8String(String.GetData(), String.Len());
-		Ar->Serialize((UTF8CHAR*)UTF8String.Get(), UTF8String.Length() * sizeof(UTF8CHAR));
-	}
-	else if (SaveAsUnicode)
-	{
-		UTF16CHAR BOM = UNICODE_BOM;
-		Ar->Serialize( &BOM, sizeof(UTF16CHAR) );
+			FTCHARToUTF8 UTF8String(String.GetData(), String.Len());
+			Ar->Serialize((UTF8CHAR*)UTF8String.Get(), UTF8String.Length() * sizeof(UTF8CHAR));
+		}
+		else if (EncodingOptions == EEncodingOptions::ForceUTF8WithoutBOM)
+		{
+			FTCHARToUTF8 UTF8String(String.GetData(), String.Len());
+			Ar->Serialize((UTF8CHAR*)UTF8String.Get(), UTF8String.Length() * sizeof(UTF8CHAR));
+		}
+		else if (SaveAsUnicode)
+		{
+			UTF16CHAR BOM = UNICODE_BOM;
+			Ar->Serialize(&BOM, sizeof(UTF16CHAR));
 
-		// Note: This is a no-op on platforms that are using a 16-bit TCHAR
-		FTCHARToUTF16 UTF16String(String.GetData(), String.Len());
-		Ar->Serialize((UTF16CHAR*)UTF16String.Get(), UTF16String.Length() * sizeof(UTF16CHAR));
-	}
-	else
-	{
-		auto Src = StringCast<ANSICHAR>(String.GetData(), String.Len());
-		Ar->Serialize( (ANSICHAR*)Src.Get(), Src.Length() * sizeof(ANSICHAR) );
+			// Note: This is a no-op on platforms that are using a 16-bit TCHAR
+			FTCHARToUTF16 UTF16String(String.GetData(), String.Len());
+			Ar->Serialize((UTF16CHAR*)UTF16String.Get(), UTF16String.Length() * sizeof(UTF16CHAR));
+		}
+		else
+		{
+			auto Src = StringCast<ANSICHAR>(String.GetData(), String.Len());
+			Ar->Serialize((ANSICHAR*)Src.Get(), Src.Length() * sizeof(ANSICHAR));
+		}
 	}
 
 	// Always explicitly close to catch errors from flush/close
