CIS fix.  remove log on write failure

#localization none
#tests none

diff --git a/Engine/Source/Runtime/Core/Private/Misc/FileHelper.cpp b/Engine/Source/Runtime/Core/Private/Misc/FileHelper.cpp
--- a/Engine/Source/Runtime/Core/Private/Misc/FileHelper.cpp
+++ b/Engine/Source/Runtime/Core/Private/Misc/FileHelper.cpp
@@ -614,7 +614,7 @@ bool FFileHelper::SaveStringToFile( FStringView String, const TCHAR* Filename,
 	TUniquePtr<FArchive> Ar = TUniquePtr<FArchive>( FileManager->CreateFileWriter( Filename, WriteFlags ) );
 	if (!Ar)
 	{
-		UE_LOG(LogStreaming, Warning, TEXT("SaveStringToFile: Failed to open writer. File:%s"), Filename);
+		//UE_LOG(LogStreaming, Warning, TEXT("SaveStringToFile: Failed to open writer. File:%s"), Filename);
 		return false;
 	}
 
