One more V591 adjustment after the C4702 fixes
#rnx

diff --git a/Engine/Source/Runtime/Core/Public/Templates/TypeHash.h b/Engine/Source/Runtime/Core/Public/Templates/TypeHash.h
--- a/Engine/Source/Runtime/Core/Public/Templates/TypeHash.h
+++ b/Engine/Source/Runtime/Core/Public/Templates/TypeHash.h
@@ -112,6 +112,7 @@ inline uint32 GetTypeHash(ScalarType Value)
 		else
 		{
 			static_assert(sizeof(ScalarType) == 0, "Unsupported integral type");
+			return 0;
 		}
 	}
 	else if constexpr (std::is_floating_point_v<ScalarType>)
@@ -127,6 +128,7 @@ inline uint32 GetTypeHash(ScalarType Value)
 		else
 		{
 			static_assert(sizeof(ScalarType) == 0, "Unsupported floating point type");
+			return 0;
 		}
 	}
 	else if constexpr (std::is_enum_v<ScalarType>)
