[Scheduler]
  - Remove costly insight trace that should not have been submitted

diff --git a/Engine/Source/Runtime/Core/Private/Async/Fundamental/Scheduler.cpp b/Engine/Source/Runtime/Core/Private/Async/Fundamental/Scheduler.cpp
--- a/Engine/Source/Runtime/Core/Private/Async/Fundamental/Scheduler.cpp
+++ b/Engine/Source/Runtime/Core/Private/Async/Fundamental/Scheduler.cpp
@@ -259,8 +259,6 @@ namespace LowLevelTasks
 
 			if (bWakeUpWorker)
 			{
-				TRACE_CPUPROFILER_EVENT_SCOPE(FScheduler::LaunchInternal_WakeUpWorker);
-
 				if (bWakeUpWorker && !WakeUpWorker(bIsBackgroundTask) && !bIsBackgroundTask)
 				{
 					WakeUpWorker(true);
