Fix/silence C4702 warnings
#rnx

diff --git a/Engine/Plugins/Animation/ControlRig/Source/ControlRig/Public/ControlRig.h b/Engine/Plugins/Animation/ControlRig/Source/ControlRig/Public/ControlRig.h
--- a/Engine/Plugins/Animation/ControlRig/Source/ControlRig/Public/ControlRig.h
+++ b/Engine/Plugins/Animation/ControlRig/Source/ControlRig/Public/ControlRig.h
@@ -783,8 +783,9 @@ public:
 	{
 #if WITH_EDITORONLY_DATA
 		return InteractionRig_DEPRECATED;
-#endif
+#else
 		return nullptr;
+#endif
 	}
 
 	UE_DEPRECATED(5.4, "InteractionRig is no longer used")
@@ -797,8 +798,9 @@ public:
 	{
 #if WITH_EDITORONLY_DATA
 		return InteractionRigClass_DEPRECATED;
-#endif
+#else
 		return nullptr;
+#endif
 	}
 
 	UE_DEPRECATED(5.4, "InteractionRig is no longer used")
diff --git a/Engine/Plugins/Editor/GuidedTutorials/Source/IntroTutorials/Private/TutorialText.h b/Engine/Plugins/Editor/GuidedTutorials/Source/IntroTutorials/Private/TutorialText.h
--- a/Engine/Plugins/Editor/GuidedTutorials/Source/IntroTutorials/Private/TutorialText.h
+++ b/Engine/Plugins/Editor/GuidedTutorials/Source/IntroTutorials/Private/TutorialText.h
@@ -148,8 +148,6 @@ struct FTextStyleAndName
 				}
 			}
 		}
-
-		return NAME_None;
 	}
 
 	FTextBlockStyle CreateTextBlockStyle() const
diff --git a/Engine/Plugins/Experimental/Mutable/Source/CustomizableObjectEditor/Private/MuCOE/SCustomizableObjectLayoutGrid.cpp b/Engine/Plugins/Experimental/Mutable/Source/CustomizableObjectEditor/Private/MuCOE/SCustomizableObjectLayoutGrid.cpp
--- a/Engine/Plugins/Experimental/Mutable/Source/CustomizableObjectEditor/Private/MuCOE/SCustomizableObjectLayoutGrid.cpp
+++ b/Engine/Plugins/Experimental/Mutable/Source/CustomizableObjectEditor/Private/MuCOE/SCustomizableObjectLayoutGrid.cpp
@@ -907,8 +907,6 @@ FCursorReply SCustomizableObjectLayoutGrid::OnCursorQuery(const FGeometry& MyGeo
 	{
 		return FCursorReply::Cursor(EMouseCursor::Default);
 	}
-
-	return FCursorReply::Unhandled();
 }
 
 
diff --git a/Engine/Plugins/FX/Niagara/Source/Niagara/Private/NiagaraEmitter.cpp b/Engine/Plugins/FX/Niagara/Source/Niagara/Private/NiagaraEmitter.cpp
--- a/Engine/Plugins/FX/Niagara/Source/Niagara/Private/NiagaraEmitter.cpp
+++ b/Engine/Plugins/FX/Niagara/Source/Niagara/Private/NiagaraEmitter.cpp
@@ -3358,8 +3358,9 @@ TStatId UNiagaraEmitter::GetStatID(bool bGameThread, bool bConcurrent)const
 			return StatID_RT;
 		}
 	}
-#endif
+#else
 	return TStatId();
+#endif
 }
 
 void FVersionedNiagaraEmitterData::ClearRuntimeAllocationEstimate(uint64 ReportHandle)
diff --git a/Engine/Plugins/FX/Niagara/Source/Niagara/Private/NiagaraSystem.cpp b/Engine/Plugins/FX/Niagara/Source/Niagara/Private/NiagaraSystem.cpp
--- a/Engine/Plugins/FX/Niagara/Source/Niagara/Private/NiagaraSystem.cpp
+++ b/Engine/Plugins/FX/Niagara/Source/Niagara/Private/NiagaraSystem.cpp
@@ -3356,8 +3356,9 @@ TStatId UNiagaraSystem::GetStatID(bool bGameThread, bool bConcurrent)const
 			return StatID_RT;
 		}
 	}
-#endif
+#else
 	return static_cast<const UObjectBaseUtility*>(this)->GetStatID();
+#endif
 }
 
 void UNiagaraSystem::AddToInstanceCountStat(int32 NumInstances, bool bSolo)const
diff --git a/Engine/Plugins/FX/Niagara/Source/NiagaraEditor/Private/NiagaraMessageManager.cpp b/Engine/Plugins/FX/Niagara/Source/NiagaraEditor/Private/NiagaraMessageManager.cpp
--- a/Engine/Plugins/FX/Niagara/Source/NiagaraEditor/Private/NiagaraMessageManager.cpp
+++ b/Engine/Plugins/FX/Niagara/Source/NiagaraEditor/Private/NiagaraMessageManager.cpp
@@ -139,7 +139,6 @@ const TOptional<const FString> FNiagaraMessageManager::GetStringForScriptUsageIn
 		ensureMsgf(false, TEXT("Tried to get script usage text for usage that is not handled!"));
 		return TOptional<const FString>();
 	}
-	return TOptional<const FString>();
 }
 
 void FNiagaraMessageManager::RegisterMessageTopic(FName TopicName)
diff --git a/Engine/Plugins/Messaging/UdpMessaging/Source/UdpMessaging/Private/Transport/UdpMessageSegmenter.cpp b/Engine/Plugins/Messaging/UdpMessaging/Source/UdpMessaging/Private/Transport/UdpMessageSegmenter.cpp
--- a/Engine/Plugins/Messaging/UdpMessaging/Source/UdpMessaging/Private/Transport/UdpMessageSegmenter.cpp
+++ b/Engine/Plugins/Messaging/UdpMessaging/Source/UdpMessaging/Private/Transport/UdpMessageSegmenter.cpp
@@ -39,7 +39,7 @@ bool FUdpMessageSegmenter::GetNextPendingSegment(TArray<uint8>& OutData, uint32&
 		return false;
 	}
 
-	for (TConstSetBitIterator<> It(PendingSendSegments); It; ++It)
+	if (TConstSetBitIterator<> It(PendingSendSegments); It)
 	{
 		OutSegment = It.GetIndex();
 
diff --git a/Engine/Plugins/Runtime/Metasound/Source/MetasoundStandardNodes/Private/MetasoundOscillatorNodes.cpp b/Engine/Plugins/Runtime/Metasound/Source/MetasoundStandardNodes/Private/MetasoundOscillatorNodes.cpp
--- a/Engine/Plugins/Runtime/Metasound/Source/MetasoundStandardNodes/Private/MetasoundOscillatorNodes.cpp
+++ b/Engine/Plugins/Runtime/Metasound/Source/MetasoundStandardNodes/Private/MetasoundOscillatorNodes.cpp
@@ -380,7 +380,6 @@ namespace Metasound
 				case ESineGenerationType::Wavetable: return MakeUnique<TOscillatorOperator<FSineWaveTable, FEnumSineGenerationTypeReadRef>>(OpParams, Type);
 				}
 			} // HasFM
-			return nullptr;
 		}
 	};
 
@@ -507,7 +506,6 @@ namespace Metasound
 			case ESawGenerationType::PolySmooth: return MakeUnique<TOscillatorOperator<FSawPolysmooth, FSawGenerationTypeReadRef>>(OpParams, Type);
 			}
 		}
-		return nullptr;
 	}
 
 	FSawOscilatorNode::FSawOscilatorNode(const FVertexName& InInstanceName, const FGuid& InInstanceID, float InDefaultFrequency, float InDefaultGlideFactor, bool bInDefaultEnablement)
@@ -691,7 +689,6 @@ namespace Metasound
 				case ESquareGenerationType::PolySmooth: return MakeUnique<FSquareOperator<FSquarePolysmooth>>(OpParams, PulseWidth, Type);
 				}
 			}
-			return nullptr;
 		}
 
 		static const FNodeClassMetadata& GetNodeInfo()
@@ -813,7 +810,6 @@ namespace Metasound
 				case ETriangleGenerationType::Trivial: return MakeUnique<TOscillatorOperator<FTriangle, FTriangleGenerationTypeReadRef>>(OpParams, Type);
 				}
 			}
-			return nullptr;
 		}
 
 		static const FNodeClassMetadata& GetNodeInfo()
diff --git a/Engine/Plugins/Runtime/RigVM/Source/RigVMDeveloper/Private/RigVMModel/RigVMNode.cpp b/Engine/Plugins/Runtime/RigVM/Source/RigVMDeveloper/Private/RigVMModel/RigVMNode.cpp
--- a/Engine/Plugins/Runtime/RigVM/Source/RigVMDeveloper/Private/RigVMModel/RigVMNode.cpp
+++ b/Engine/Plugins/Runtime/RigVM/Source/RigVMDeveloper/Private/RigVMModel/RigVMNode.cpp
@@ -637,10 +637,9 @@ TArray<int32> URigVMNode::GetInstructionsForVMImpl(const FRigVMExtendedExecuteCo
 	{
 		return InVM->GetByteCode().GetAllInstructionIndicesForSubject((URigVMNode*)this);
 	}
-	
-#endif
-
+#else
 	return Instructions;
+#endif
 }
 
 int32 URigVMNode::GetInstructionVisitedCount(const FRigVMExtendedExecuteContext& Context, URigVM* InVM, const FRigVMASTProxy& InProxy) const
diff --git a/Engine/Plugins/Tests/CQTest/Source/CQTestTests/Private/Assert/AssertionTests.cpp b/Engine/Plugins/Tests/CQTest/Source/CQTestTests/Private/Assert/AssertionTests.cpp
--- a/Engine/Plugins/Tests/CQTest/Source/CQTestTests/Private/Assert/AssertionTests.cpp
+++ b/Engine/Plugins/Tests/CQTest/Source/CQTestTests/Private/Assert/AssertionTests.cpp
@@ -19,12 +19,15 @@ namespace CQTests
 			Assert.Fail(AnyError);
 		}
 
+#pragma warning(push)
+#pragma warning(disable:4702)
 		TEST_METHOD(AssertFail_Macro_AddsErrorAndExits)
 		{
 			Assert.ExpectError(ExpectedError);
 			ASSERT_FAIL(ExpectedError);
 			ASSERT_FAIL(TEXT("It should not assert after the first expected assert since it exits"));
 		}
+#pragma warning(pop)
 
 		TEST_METHOD(Assertions_Accept_RawStrings)
 		{
diff --git a/Engine/Plugins/VirtualProduction/Rivermax/RivermaxCore/Source/RivermaxCore/Private/RivermaxOutputStream.cpp b/Engine/Plugins/VirtualProduction/Rivermax/RivermaxCore/Source/RivermaxCore/Private/RivermaxOutputStream.cpp
--- a/Engine/Plugins/VirtualProduction/Rivermax/RivermaxCore/Source/RivermaxCore/Private/RivermaxOutputStream.cpp
+++ b/Engine/Plugins/VirtualProduction/Rivermax/RivermaxCore/Source/RivermaxCore/Private/RivermaxOutputStream.cpp
@@ -1918,8 +1918,6 @@ namespace UE::RivermaxCore::Private
 			return false;
 		}
 		}
-
-		return false;
 	}
 
 	bool FRivermaxOutputStream::IsChunkOnTime_FrameCreation() const
diff --git a/Engine/Source/Developer/SlateReflector/Private/Widgets/SWidgetReflector.cpp b/Engine/Source/Developer/SlateReflector/Private/Widgets/SWidgetReflector.cpp
--- a/Engine/Source/Developer/SlateReflector/Private/Widgets/SWidgetReflector.cpp
+++ b/Engine/Source/Developer/SlateReflector/Private/Widgets/SWidgetReflector.cpp
@@ -1535,7 +1535,6 @@ void SWidgetReflector::HandlePickingModeStateChanged()
 	default:
 		return FSlateIcon(FWidgetReflectorStyle::GetStyleSetName(), "Icon.Empty");
 	}
-	return FSlateIcon(FWidgetReflectorStyle::GetStyleSetName(), "Icon.Empty");
 }
 
 FText SWidgetReflector::HandleGetPickingModeText() const
diff --git a/Engine/Source/Developer/Windows/LiveCoding/Private/External/LC_ClientUserCommandThread.cpp b/Engine/Source/Developer/Windows/LiveCoding/Private/External/LC_ClientUserCommandThread.cpp
--- a/Engine/Source/Developer/Windows/LiveCoding/Private/External/LC_ClientUserCommandThread.cpp
+++ b/Engine/Source/Developer/Windows/LiveCoding/Private/External/LC_ClientUserCommandThread.cpp
@@ -761,8 +761,4 @@ Thread::ReturnValue ClientUserCommandThread::ThreadFunction(Event* waitForStartE
 
 		delete command;
 	}
-
-	// BEGIN EPIC MOD
-	return Thread::ReturnValue(0u);
-	// END EPIC MOD
 }
diff --git a/Engine/Source/Editor/UnrealEd/Private/Fbx/FbxMaterialConflictWindow.cpp b/Engine/Source/Editor/UnrealEd/Private/Fbx/FbxMaterialConflictWindow.cpp
--- a/Engine/Source/Editor/UnrealEd/Private/Fbx/FbxMaterialConflictWindow.cpp
+++ b/Engine/Source/Editor/UnrealEd/Private/Fbx/FbxMaterialConflictWindow.cpp
@@ -443,8 +443,6 @@ FSlateColor FMaterialConflictData::GetCellColor(bool IsResultData) const
 		//Material is matched
 		return FSlateColor::UseForeground();
 	}
-	//Fallback
-	return FSlateColor::UseForeground();
 }
 
 void SFbxMaterialConflictWindow::FillMaterialListItem()
diff --git a/Engine/Source/Runtime/Analytics/TelemetryUtils/Public/TelemetryRouter.h b/Engine/Source/Runtime/Analytics/TelemetryUtils/Public/TelemetryRouter.h
--- a/Engine/Source/Runtime/Analytics/TelemetryUtils/Public/TelemetryRouter.h
+++ b/Engine/Source/Runtime/Analytics/TelemetryUtils/Public/TelemetryRouter.h
@@ -27,7 +27,10 @@ namespace UE::Telemetry::Private
         {
             return DATA_TYPE::TelemetryID;
         }
-        return FGuid{};
+		else
+		{
+			return FGuid{};
+		}
     }
 
     template <typename DATA_TYPE>
@@ -40,7 +43,10 @@ namespace UE::Telemetry::Private
         {
             return DATA_TYPE::TelemetryID;
         }
-        return FGuid{};
+        else
+        {
+            return FGuid{};
+        }
     }
 }
 
diff --git a/Engine/Source/Runtime/Core/Private/GenericPlatform/GenericPlatformMallocCrash.cpp b/Engine/Source/Runtime/Core/Private/GenericPlatform/GenericPlatformMallocCrash.cpp
--- a/Engine/Source/Runtime/Core/Private/GenericPlatform/GenericPlatformMallocCrash.cpp
+++ b/Engine/Source/Runtime/Core/Private/GenericPlatform/GenericPlatformMallocCrash.cpp
@@ -454,6 +454,8 @@ void FGenericPlatformMallocCrash::PrintPoolsUsage()
 #endif // _DEBUG
 }
 
+#pragma warning(push)
+#pragma warning(disable:4702)
 bool FGenericPlatformMallocCrash::IsOnCrashedThread() const
 {
 	// Suspend threads other than the crashed one to prevent serious memory errors.
@@ -468,6 +470,8 @@ bool FGenericPlatformMallocCrash::IsOnCrashedThread() const
 		return false;
 	}
 }
+#pragma warning(pop)
+
 
 bool FGenericPlatformMallocCrash::IsPtrInLargePool( void* Ptr ) const
 {
diff --git a/Engine/Source/Runtime/Core/Private/Modules/SimpleParse.cpp b/Engine/Source/Runtime/Core/Private/Modules/SimpleParse.cpp
--- a/Engine/Source/Runtime/Core/Private/Modules/SimpleParse.cpp
+++ b/Engine/Source/Runtime/Core/Private/Modules/SimpleParse.cpp
@@ -77,9 +77,6 @@ bool SimpleParseString(const TCHAR*& InOutPtr, T& OutStr)
 				break;
 		}
 	}
-
-	InOutPtr = Ptr + 1;
-	return true;
 }
 
 bool FSimpleParse::ParseString(const TCHAR*& InOutPtr, FString& OutStr)
diff --git a/Engine/Source/Runtime/Core/Public/Async/TaskGraphInterfaces.h b/Engine/Source/Runtime/Core/Public/Async/TaskGraphInterfaces.h
--- a/Engine/Source/Runtime/Core/Public/Async/TaskGraphInterfaces.h
+++ b/Engine/Source/Runtime/Core/Public/Async/TaskGraphInterfaces.h
@@ -1487,8 +1487,9 @@ public:
 	{
 #if STATS|| ENABLE_STATNAMEDEVENTS
 		return StatID;
-#endif
+#else
 		return TStatId();
+#endif
 	}
 
 private:
diff --git a/Engine/Source/Runtime/Core/Public/AutoRTFM/AutoRTFM.h b/Engine/Source/Runtime/Core/Public/AutoRTFM/AutoRTFM.h
--- a/Engine/Source/Runtime/Core/Public/AutoRTFM/AutoRTFM.h
+++ b/Engine/Source/Runtime/Core/Public/AutoRTFM/AutoRTFM.h
@@ -240,6 +240,8 @@ UE_AUTORTFM_FORCEINLINE void autortfm_open(void (*work)(void* arg), void* arg)
 #if UE_AUTORTFM
 [[nodiscard]] UE_AUTORTFM_API autortfm_status autortfm_close(void (*work)(void* arg), void* arg);
 #else
+#pragma warning(push)
+#pragma warning(disable:4702)
 [[nodiscard]] UE_AUTORTFM_FORCEINLINE autortfm_status autortfm_close(void (*work)(void* arg), void* arg)
 {
 	UE_AUTORTFM_UNUSED(work);
@@ -247,6 +249,7 @@ UE_AUTORTFM_FORCEINLINE void autortfm_open(void (*work)(void* arg), void* arg)
     abort();
 	return autortfm_status_aborted_by_language;
 }
+#pragma warning(pop)
 #endif
 
 // Records the pointer and size from the open into the current transaction
diff --git a/Engine/Source/Runtime/Core/Public/Stats/Stats.h b/Engine/Source/Runtime/Core/Public/Stats/Stats.h
--- a/Engine/Source/Runtime/Core/Public/Stats/Stats.h
+++ b/Engine/Source/Runtime/Core/Public/Stats/Stats.h
@@ -395,9 +395,9 @@ struct FDynamicStats
 	{
 #if	STATS
 		return CreateStatIdInternal<TStatGroup>( FName( *StatNameOrDescription ), EStatDataType::ST_int64, true);
-#endif // STATS
-
+#else
 		return TStatId();
+#endif // STATS
 	}
 
 	template< typename TStatGroup >
@@ -405,9 +405,9 @@ struct FDynamicStats
 	{
 #if	STATS
 		return CreateStatIdInternal<TStatGroup>(FName(*StatNameOrDescription), EStatDataType::ST_int64, false, !bIsAccumulator);
-#endif // STATS
-
+#else
 		return TStatId();
+#endif // STATS
 	}
 
 	template< typename TStatGroup >
@@ -415,8 +415,9 @@ struct FDynamicStats
 	{
 #if	STATS
 		return CreateStatIdInternal<TStatGroup>(FName(*StatNameOrDescription), EStatDataType::ST_double, false, !bIsAccumulator);
-#endif // STATS
+#else
 		return TStatId();
+#endif // STATS
 	}
 
 	template< typename TStatGroup >
@@ -424,8 +425,9 @@ struct FDynamicStats
 	{
 #if	STATS
 		return CreateStatIdInternal<TStatGroup>(StatNameOrDescription, EStatDataType::ST_int64, IsTimer);
-#endif // STATS
+#else
 		return TStatId();
+#endif // STATS
 	}
 
 	template< typename TStatGroup >
@@ -433,9 +435,9 @@ struct FDynamicStats
 	{
 #if	STATS
 		return CreateMemoryStatId<TStatGroup>(FName(*StatNameOrDescription), MemRegion);
-#endif // STATS
-
+#else
 		return TStatId();
+#endif // STATS
 	}
 
 	template< typename TStatGroup >
@@ -455,9 +457,9 @@ struct FDynamicStats
 			false, EStatDataType::ST_int64, *StatNameOrDescription.ToString(), false, false, MemRegion);
 
 		return StatID;
-#endif // STATS
-
+#else
 		return TStatId();
+#endif // STATS
 	}
 
 #if	STATS
diff --git a/Engine/Source/Runtime/Engine/Private/PerQualityLevelProperties.cpp b/Engine/Source/Runtime/Engine/Private/PerQualityLevelProperties.cpp
--- a/Engine/Source/Runtime/Engine/Private/PerQualityLevelProperties.cpp
+++ b/Engine/Source/Runtime/Engine/Private/PerQualityLevelProperties.cpp
@@ -53,7 +53,6 @@ namespace QualityLevelProperty
 		{
 			return NAME_None;
 		}
-		return QualityLevelNames[QL];
 	}
 
 	int32 FNameToQualityLevel(FName QL)
diff --git a/Engine/Source/Runtime/Engine/Private/PhysicsEngine/Experimental/ChaosCooking.cpp b/Engine/Source/Runtime/Engine/Private/PhysicsEngine/Experimental/ChaosCooking.cpp
--- a/Engine/Source/Runtime/Engine/Private/PhysicsEngine/Experimental/ChaosCooking.cpp
+++ b/Engine/Source/Runtime/Engine/Private/PhysicsEngine/Experimental/ChaosCooking.cpp
@@ -155,8 +155,6 @@ namespace Chaos
 				TArray<Chaos::TVector<int32, 3>> TrianglesLargeIdx;
 				return LambdaHelper(TrianglesLargeIdx);
 			}
-
-			return nullptr;
 		}
 
 		void BuildConvexMeshes(TArray<Chaos::FImplicitObjectPtr>& OutConvexMeshes, const FCookBodySetupInfo& InParams)
diff --git a/Engine/Source/Runtime/MovieScene/Private/MovieSceneSection.cpp b/Engine/Source/Runtime/MovieScene/Private/MovieSceneSection.cpp
--- a/Engine/Source/Runtime/MovieScene/Private/MovieSceneSection.cpp
+++ b/Engine/Source/Runtime/MovieScene/Private/MovieSceneSection.cpp
@@ -605,8 +605,9 @@ FColor UMovieSceneSection::GetColorTint() const
 {
 #if WITH_EDITORONLY_DATA
 	return ColorTint;
-#endif
+#else
 	return FColor(0, 0, 0, 0);
+#endif
 }
 
 UMovieSceneSection* UMovieSceneSection::SplitSection(FQualifiedFrameTime SplitTime, bool bDeleteKeys)
diff --git a/Engine/Source/Runtime/Online/BuildPatchServices/Private/Generation/ChunkDeltaOptimiser.cpp b/Engine/Source/Runtime/Online/BuildPatchServices/Private/Generation/ChunkDeltaOptimiser.cpp
--- a/Engine/Source/Runtime/Online/BuildPatchServices/Private/Generation/ChunkDeltaOptimiser.cpp
+++ b/Engine/Source/Runtime/Online/BuildPatchServices/Private/Generation/ChunkDeltaOptimiser.cpp
@@ -83,8 +83,9 @@ namespace DeltaOptimiseHelpers
 #if UE_BUILD_DEBUG
 		static const bool bSingleScannerThread = FParse::Param(FCommandLine::Get(), TEXT("singlescanneronly"));
 		return bSingleScannerThread ? false : bHasUnusedCpu;
-#endif
+#else
 		return bHasUnusedCpu;
+#endif
 	}
 
 	template <typename T>
@@ -101,8 +102,9 @@ namespace DeltaOptimiseHelpers
 #if UE_BUILD_DEBUG
 		static const bool bSingleScannerThread = FParse::Param(FCommandLine::Get(), TEXT("singlescanneronly"));
 		return bSingleScannerThread ? (FDataScannerCounter::GetNumIncompleteScanners() + FDataScannerCounter::GetNumRunningScanners()) > 0 : bScannerArrayFull;
-#endif
+#else
 		return bScannerArrayFull;
+#endif
 	}
 
 	FChunkPart SelectBytes(const FChunkPart& FullPart, uint32 LeftChop, uint32 Size)
