Fix csv profiling recursion in windows/linux platform memory

diff --git a/Engine/Source/Runtime/Core/Private/Unix/UnixPlatformMemory.cpp b/Engine/Source/Runtime/Core/Private/Unix/UnixPlatformMemory.cpp
--- a/Engine/Source/Runtime/Core/Private/Unix/UnixPlatformMemory.cpp
+++ b/Engine/Source/Runtime/Core/Private/Unix/UnixPlatformMemory.cpp
@@ -434,10 +434,6 @@ void* FUnixPlatformMemory::BinnedAllocFromOS(SIZE_T Size)
 #if UE_ENABLE_PLATFORM_MEMORY_CSV_STATS
 	const bool bIsInGameThread = IsInGameThread();
 	CSV_SCOPED_TIMING_STAT_RECURSIVE_CONDITIONAL(FMemory, BinnedAllocFromOSTime, bIsInGameThread);
-	if (bIsInGameThread)
-	{
-		CSV_CUSTOM_STAT(FMemory, BinnedAllocFromOSCount, 1, ECsvCustomStatOp::Accumulate);
-	}
 #endif // UE_ENABLE_PLATFORM_MEMORY_CSV_STATS
 
 #if UE_CHECK_LARGE_ALLOCATIONS
@@ -633,10 +629,6 @@ FUnixPlatformMemory::FPlatformVirtualMemoryBlock FUnixPlatformMemory::FPlatformV
 #if UE_ENABLE_PLATFORM_MEMORY_CSV_STATS
 	const bool bIsInGameThread = IsInGameThread();
 	CSV_SCOPED_TIMING_STAT_RECURSIVE_CONDITIONAL(FMemory, AllocateVirtualTime, bIsInGameThread);
-	if (bIsInGameThread)
-	{
-		CSV_CUSTOM_STAT(FMemory, AllocateVirtualCount, 1, ECsvCustomStatOp::Accumulate);
-	}
 #endif // UE_ENABLE_PLATFORM_MEMORY_CSV_STATS
 
 	FPlatformVirtualMemoryBlock Result;
diff --git a/Engine/Source/Runtime/Core/Private/Windows/WindowsPlatformMemory.cpp b/Engine/Source/Runtime/Core/Private/Windows/WindowsPlatformMemory.cpp
--- a/Engine/Source/Runtime/Core/Private/Windows/WindowsPlatformMemory.cpp
+++ b/Engine/Source/Runtime/Core/Private/Windows/WindowsPlatformMemory.cpp
@@ -445,10 +445,6 @@ void* FWindowsPlatformMemory::BinnedAllocFromOS( SIZE_T Size )
 #if UE_ENABLE_PLATFORM_MEMORY_CSV_STATS
 	const bool bIsInGameThread = IsInGameThread();
 	CSV_SCOPED_TIMING_STAT_RECURSIVE_CONDITIONAL(FMemory, BinnedAllocFromOSTime, bIsInGameThread);
-	if (bIsInGameThread)
-	{
-		CSV_CUSTOM_STAT(FMemory, BinnedAllocFromOSCount, 1, ECsvCustomStatOp::Accumulate);
-	}
 #endif // UE_ENABLE_PLATFORM_MEMORY_CSV_STATS
 
 #if UE_CHECK_LARGE_ALLOCATIONS
@@ -491,10 +487,6 @@ FWindowsPlatformMemory::FPlatformVirtualMemoryBlock FWindowsPlatformMemory::FPla
 #if UE_ENABLE_PLATFORM_MEMORY_CSV_STATS
 	const bool bIsInGameThread = IsInGameThread();
 	CSV_SCOPED_TIMING_STAT_RECURSIVE_CONDITIONAL(FMemory, AllocateVirtualTime, bIsInGameThread);
-	if (bIsInGameThread)
-	{
-		CSV_CUSTOM_STAT(FMemory, AllocateVirtualCount, 1, ECsvCustomStatOp::Accumulate);
-	}
 #endif // UE_ENABLE_PLATFORM_MEMORY_CSV_STATS
 
 	FPlatformVirtualMemoryBlock Result;
