[ParallelFor]
  - Reintroduce the incremental wake-up to avoid very small tasks to cost more in scheduling then they take to actually execute
  - Fix cancellation that was enqueuing task for other thread instead of cleaning them up in place
  - Restore previous fix by [REDACTED_PERSON] to avoid waking up threads during parallel executor's destruction if they were cancelled
  - A simple ParallelFor(32, [](int32) {}, EParallelForFlags::Unbalanced); statement is now 6x faster on many-core systems

#tests TSAN / UnitTests

diff --git a/Engine/Source/Runtime/Core/Public/Async/ParallelFor.h b/Engine/Source/Runtime/Core/Public/Async/ParallelFor.h
--- a/Engine/Source/Runtime/Core/Public/Async/ParallelFor.h
+++ b/Engine/Source/Runtime/Core/Public/Async/ParallelFor.h
@@ -224,8 +224,15 @@ namespace ParallelForImpl
 				}
 			}
 
+			int32 GetNextWorkerIndexToLaunch()
+			{
+				const int32 WorkerIndex = LaunchedWorkers.fetch_add(1, std::memory_order_relaxed);
+				return WorkerIndex >= Tasks.Num() ? -1 : WorkerIndex;
+			}
+
 			std::atomic_int BatchItem  { 0 };
 			std::atomic_int IncompleteBatches { 0 };
+			std::atomic_int LaunchedWorkers { 0 };
 			int32 Num;
 			int32 BatchSize;
 			int32 NumBatches;
@@ -237,15 +244,16 @@ namespace ParallelForImpl
 		};
 		using FDataHandle = TRefCountPtr<FParallelForData>;
 
-		//each task has an executor.
-		class FParallelExecutor 
+		// Each task has an executor.
+		class FParallelExecutor
 		{
-			FDataHandle Data;
-			mutable int32 WorkerIndex;
+			mutable FDataHandle Data;
+			int32 WorkerIndex;
 			LowLevelTasks::ETaskPriority Priority;
+			mutable bool bReschedule = false;
 
 		public:
-			inline FParallelExecutor(FDataHandle&& InData, int32 InWorkerIndex, LowLevelTasks::ETaskPriority InPriority) 
+			inline FParallelExecutor(FDataHandle&& InData, int32 InWorkerIndex, LowLevelTasks::ETaskPriority InPriority)
 				: Data(MoveTemp(InData))
 				, WorkerIndex(InWorkerIndex)
 				, Priority(InPriority)
@@ -253,16 +261,11 @@ namespace ParallelForImpl
 			}
 
 			FParallelExecutor(const FParallelExecutor&) = delete;
-			inline FParallelExecutor(FParallelExecutor&& Other) 
-				: Data(MoveTemp(Other.Data))
-				, WorkerIndex(Other.WorkerIndex)
-				, Priority(Other.Priority)
-			{
-			}
+			FParallelExecutor(FParallelExecutor&& Other) = default;
 
 			~FParallelExecutor()
 			{
-				if (Data.IsValid() && WorkerIndex >= 0)
+				if (Data.IsValid() && bReschedule)
 				{
 					FParallelExecutor::LaunchTask(nullptr, MoveTemp(Data), WorkerIndex, Priority);
 				}
@@ -298,6 +301,15 @@ namespace ParallelForImpl
 					}
 				};
 
+				const int32 NumBatches = Data->NumBatches;
+
+				// We're going to consume one ourself, so we need at least 2 left to consider launching a new worker
+				// We also do not launch a worker from the master as we already launched one before doing prework.
+				if (bIsMaster == false && Data->BatchItem.load(std::memory_order_relaxed) + 2 <= NumBatches)
+				{
+					LaunchAnotherWorkerIfNeeded(Data);
+				}
+
 				TRACE_CPUPROFILER_EVENT_SCOPE_TEXT(DebugName);
 
 				auto Now = [] { return FTimespan::FromSeconds(FPlatformTime::Seconds()); };
@@ -313,7 +325,6 @@ namespace ParallelForImpl
 
 				const int32 Num = Data->Num;
 				const int32 BatchSize = Data->BatchSize;
-				const int32 NumBatches = Data->NumBatches;
 				const TArrayView<ContextType>& Contexts = Data->Contexts;
 				const BodyType& Body = Data->Body;
 
@@ -322,12 +333,11 @@ namespace ParallelForImpl
 				{
 					int32 BatchIndex = Data->BatchItem.fetch_add(1, std::memory_order_relaxed);
 					
-					//save the last block for the master to safe an event
+					// Save the last block for the master to avoid an event
 					if (bSaveLastBlockForMaster && BatchIndex >= NumBatches - 1)
 					{
 						if (!bIsMaster)
 						{
-							WorkerIndex = -1;
 							return false;
 						}
 						BatchIndex = (NumBatches - 1);
@@ -352,12 +362,11 @@ namespace ParallelForImpl
 						{
 							Data->FinishedSignal->Trigger();
 						}
-						WorkerIndex = -1;
+
 						return true;
 					}
 					else if (EndIndex >= Num)
 					{
-						WorkerIndex = -1;
 						return false;
 					}
 					else if (!bIsBackgroundPriority)
@@ -368,7 +377,8 @@ namespace ParallelForImpl
 					auto PassedTime = [Start, &Now]() { return Now() - Start; };
 					if (PassedTime() > YieldingThreshold)
 					{
-						//abort and reschedule (in the destructor as WorkerIndex is larger_eq Zero) to give higher priority tasks a chance to run
+						// Abort and reschedule (in the destructor as WorkerIndex is larger_eq Zero) to give higher priority tasks a chance to run
+						bReschedule = true;
 						return false;
 					}
 				}
@@ -391,7 +401,39 @@ namespace ParallelForImpl
 
 				TracedTask.Task.Init(DebugName, InPriority, FParallelExecutor(MoveTemp(InData), InWorkerIndex, InPriority));
 				verify(LowLevelTasks::TryLaunch(TracedTask.Task, LowLevelTasks::EQueuePreference::GlobalQueuePreference));
+			}
 
+			static void PrepareTask(const TCHAR* DebugName, FDataHandle&& InData, int32 InWorkerIndex, LowLevelTasks::ETaskPriority InPriority)
+			{
+				FTracedTask& TracedTask = InData->Tasks[InWorkerIndex];
+				if (DebugName == nullptr)
+				{
+					DebugName = TracedTask.Task.GetDebugName();
+				}
+
+				TracedTask.Task.Init(DebugName, InPriority, FParallelExecutor(MoveTemp(InData), InWorkerIndex, InPriority));
+			}
+
+			static bool LaunchAnotherWorkerIfNeeded(FDataHandle& InData)
+			{
+				const int32 WorkerIndex = InData->GetNextWorkerIndexToLaunch();
+				if (WorkerIndex != -1)
+				{
+					FTracedTask& TracedTask = InData->Tasks[WorkerIndex];
+
+					if (TracedTask.TraceId != TaskTrace::InvalidId) // reused task
+					{
+						TaskTrace::Destroyed(TracedTask.TraceId);
+					}
+					TracedTask.TraceId = TaskTrace::GenerateTaskId();
+					if (LowLevelTasks::TryLaunch(TracedTask.Task, LowLevelTasks::EQueuePreference::GlobalQueuePreference))
+					{
+						TaskTrace::Launched(TracedTask.TraceId, TracedTask.Task.GetDebugName(), false, ENamedThreads::AnyThread, 0);
+						return true;
+					}
+				}
+
+				return false;
 			}
 		};
 
@@ -400,9 +442,12 @@ namespace ParallelForImpl
 		FDataHandle Data = new FParallelForData(Num, BatchSize, NumBatches, NumWorkers, Contexts, Body, FinishedSignal);
 		for (int32 Worker = 0; Worker < NumWorkers; Worker++)
 		{
-			FParallelExecutor::LaunchTask(DebugName, FDataHandle(Data), Worker, Priority);
+			FParallelExecutor::PrepareTask(DebugName, FDataHandle(Data), Worker, Priority);
 		}
 
+		// Launch the first worker before we start doing prework
+		FParallelExecutor::LaunchAnotherWorkerIfNeeded(Data);
+
 		// do the prework
 		CurrentThreadWorkToDoBeforeHelping();
 
@@ -415,7 +460,8 @@ namespace ParallelForImpl
 			TRACE_CPUPROFILER_EVENT_SCOPE(ParallelFor.Cancel);
 			for (FTracedTask& TracedTask : LocalExecutor.GetData()->Tasks)
 			{
-				TracedTask.Task.TryCancel(LowLevelTasks::ECancellationFlags::PrelaunchCancellation);
+				// Task is still required to run some cleanup when successfully cancelled, so do it here to avoid enqueuing in the global queue for cancelled tasks.
+				TracedTask.Task.TryCancel(LowLevelTasks::ECancellationFlags::PrelaunchCancellation | LowLevelTasks::ECancellationFlags::TryLaunchOnSuccess);
 			}
 		}
 
