[EditorPerf]
  - Increase mimalloc page reset delay from 500ms to 10s to reduce churn between application and the OS on the hottest memory pages (scratch space) during alloc intense workloads which could end up in a huge performance drop due to page faults and memory zeroing.

diff --git a/Engine/Source/Runtime/Core/Private/HAL/MallocMimalloc.cpp b/Engine/Source/Runtime/Core/Private/HAL/MallocMimalloc.cpp
--- a/Engine/Source/Runtime/Core/Private/HAL/MallocMimalloc.cpp
+++ b/Engine/Source/Runtime/Core/Private/HAL/MallocMimalloc.cpp
@@ -10,6 +10,7 @@
 #if MIMALLOC_ENABLED
 
 #include "HAL/UnrealMemory.h"
+#include "HAL/IConsoleManager.h"
 #include "Math/UnrealMathUtility.h"
 #include "Thirdparty/IncludeMiMalloc.h"
 
@@ -23,9 +24,31 @@
 /** Value we fill a new memory block with, in UE_BUILD_DEBUG **/
 #define DEBUG_FILL_NEW (0xcd)
 
+// Dramatically reduce memory zeroing and page faults during alloc intense workloads
+// by keeping freed pages for a little while instead of releasing them
+// right away to the OS, effectively acting like a scratch buffer 
+// until pages are both freed and inactive for the delay specified
+// in milliseconds.
+int32 GMiMallocMemoryResetDelay = 10000;
+
+void ApplyMiMallocMemoryResetDelayChange(IConsoleVariable* InConsolveVariable = nullptr)
+{
+	mi_option_set(mi_option_reset_delay, GMiMallocMemoryResetDelay);
+}
+
+static FAutoConsoleVariableRef CVarMiMallocMemoryResetDelay(
+	TEXT("mi.MemoryResetDelay"),
+	GMiMallocMemoryResetDelay,
+	TEXT("The time in milliseconds to keep recently freed memory pages inside the process for reuse. This can dramatically reduce OS overhead of memory zeroing and page faults during alloc intense workloads."),
+	FConsoleVariableDelegate::CreateStatic(&ApplyMiMallocMemoryResetDelayChange),
+	ECVF_Default
+);
+
 FMallocMimalloc::FMallocMimalloc()
 {
 	FPlatformMemory::MiMallocInit();
+
+	ApplyMiMallocMemoryResetDelayChange();
 }
 
 void* FMallocMimalloc::TryMalloc( SIZE_T Size, uint32 Alignment )
