Fixed parsing bugs in CultureInvariantDecimalToString

* It now handles a positive prefix.
* It now handles a decimal separator correctly when MaximumFractionalDigits is zero.
* It now reports any invalid digit characters and aborts parsing, rather than get stuck in an infinite loop.

#rnx

diff --git a/Engine/Source/Runtime/Core/Private/Internationalization/FastDecimalFormat.cpp b/Engine/Source/Runtime/Core/Private/Internationalization/FastDecimalFormat.cpp
--- a/Engine/Source/Runtime/Core/Private/Internationalization/FastDecimalFormat.cpp
+++ b/Engine/Source/Runtime/Core/Private/Internationalization/FastDecimalFormat.cpp
@@ -448,26 +448,36 @@ FString CultureInvariantDecimalToString(const double InVal, const TCHAR*& InBuff
 	}
 
 	// Scrape negative, apply at end
+	// Skip positive
 	bool bIsNegative = false;
-	static const TCHAR EuropeanNegativePrefix = TEXT('-');
+	constexpr TCHAR EuropeanNegativePrefix = TEXT('-');
+	constexpr TCHAR EuropeanPositivePrefix = TEXT('+');
 	if (*InBuffer == EuropeanNegativePrefix)
 	{
 		bIsNegative = true;
 		++InBuffer;
 	}
+	else if (*InBuffer == EuropeanPositivePrefix)
+	{
+		++InBuffer;
+	}
 
 	// Parse digits & decimal, no grouping on fractional
 	bool bParsedFractional = false;
 	uint8 FractionalDigitsPrinted = 0;
 	while (InBuffer < InBufferEnd)
 	{
-		static const TCHAR EuropeanDecimal = TEXT('.');
-		if (*InBuffer == EuropeanDecimal && InFormattingOptions.MaximumFractionalDigits > 0)
+		constexpr TCHAR EuropeanDecimal = TEXT('.');
+		if (*InBuffer == EuropeanDecimal)
 		{
-			bParsedFractional = true;
-			OutStr += InFormattingRules.DecimalSeparatorCharacter;
-			++InBuffer;
-			continue;
+			if (InFormattingOptions.MaximumFractionalDigits > 0)
+			{
+				bParsedFractional = true;
+				OutStr += InFormattingRules.DecimalSeparatorCharacter;
+				++InBuffer;
+				continue;
+			}
+			break; // Fractional digits not needed; stop parsing the source string now
 		}
 
 		if (!bParsedFractional && bUseGrouping && NumUntilNextGroup-- == 0)
@@ -478,12 +488,17 @@ FString CultureInvariantDecimalToString(const double InVal, const TCHAR*& InBuff
 
 		// 48 for raw ascii -> int
 		int32 CharacterIndex = (int32)InBuffer[0] - 48;
-		if (ensure(CharacterIndex >= 0 && CharacterIndex < UE_ARRAY_COUNT(InFormattingRules.DigitCharacters)))
+		if (ensureMsgf(CharacterIndex >= 0 && CharacterIndex < UE_ARRAY_COUNT(InFormattingRules.DigitCharacters), TEXT("Found invalid digit character '%c'!"), *InBuffer))
 		{
 			OutStr += InFormattingRules.DigitCharacters[CharacterIndex];
 			FractionalDigitsPrinted += bParsedFractional ? 1 : 0;
 			++InBuffer;
 		}
+		else
+		{
+			// Invalid digit; stop parsing the source string now
+			break;
+		}
 	}
 
 	// Apply back padding, if back isn't just zero
@@ -511,13 +526,10 @@ FString CultureInvariantDecimalToString(const double InVal, const TCHAR*& InBuff
 	const FString& FinalPrefixStr = (bIsNegative) ? SigningStrings.GetNegativePrefixString() : SigningStrings.GetPositivePrefixString();
 	const FString& FinalSuffixStr = (bIsNegative) ? SigningStrings.GetNegativeSuffixString() : SigningStrings.GetPositiveSuffixString();
 
-	OutStr += FinalSuffixStr.IsEmpty() ? "" : FinalSuffixStr;
-	if (OutStr.LastChar() != '\0')
-	{
-		OutStr += '\0';
-	}
-
-	return FinalPrefixStr.IsEmpty() ? OutStr.GetData() : FinalPrefixStr + OutStr.GetData();
+	OutStr.InsertAt(0, FinalPrefixStr);
+	OutStr += FinalSuffixStr;
+	
+	return OutStr.ToString();
 }
 
 void FractionalToString(const double InVal, const FDecimalNumberFormattingRules& InFormattingRules, FNumberFormattingOptions InFormattingOptions, FString& OutString)
