Fix the SharedPointer bug under AutoRTFM when compiled under debug.

The problem was an existing issue in the AutoRTFM runtime that was masked by the compiler on non-debug builds being smart enough to not record writes on objects that were newly allocated. The solution is that the AutoRTFM runtime must always record new memory (not just for larger allocations) at a performance cost.

diff --git a/Engine/Source/Runtime/Core/Private/AutoRTFM/IntervalTree.h b/Engine/Source/Runtime/Core/Private/AutoRTFM/IntervalTree.h
--- a/Engine/Source/Runtime/Core/Private/AutoRTFM/IntervalTree.h
+++ b/Engine/Source/Runtime/Core/Private/AutoRTFM/IntervalTree.h
@@ -9,6 +9,56 @@
 namespace AutoRTFM
 {
 
+template<typename T> struct TIntervalTreeArray final
+{
+	FORCEINLINE void Empty()
+	{
+		Payload.Empty();
+	}
+
+	FORCEINLINE uint32 Num() const
+	{
+		return Payload.Num();
+	}
+
+	FORCEINLINE bool IsEmpty() const
+	{
+		return Payload.IsEmpty();
+	}
+
+	FORCEINLINE void Push(T Thing)
+	{
+		Payload.Push(Thing);
+	}
+
+	FORCEINLINE uint32 Add(T Thing)
+	{
+		return Payload.Add(Thing);
+	}
+
+	FORCEINLINE T Pop()
+	{
+		return Payload.Pop();
+	}
+
+	FORCEINLINE T& operator[](uint32 Index)
+	{
+		// A bit nasty - but since we know we are in bounds, skip the check.
+		// This change alone accounted for a 13% uplift in performance in a
+		// tight loop of critical performance importance!
+		return Payload.GetData()[Index];
+	}
+
+	FORCEINLINE const T& operator[](uint32 Index) const
+	{
+		// Same as non-const version.
+		return Payload.GetData()[Index];
+	}
+
+private:
+	TArray<T> Payload;
+};
+
 struct FIntervalTree final
 {
     FIntervalTree() = default;
@@ -24,7 +74,7 @@ struct FIntervalTree final
 
     FORCENOINLINE bool Contains(const void* const Address, const size_t Size) const
     {
-        FRange NewRange(Address, Size);
+        const FRange NewRange(Address, Size);
 
         if (UNLIKELY(IntervalTreeNodeIndexNone == Root))
         {
@@ -46,14 +96,8 @@ struct FIntervalTree final
             {
                 return true;
             }
-            else if (NewRange.Start < Range.Start)
-            {
-                Current = Node.Left;
-            }
-            else
-            {
-                Current = Node.Right;
-            }
+
+			Current = (NewRange.Start < Range.Start) ? Node.Left : Node.Right;
         } while (IntervalTreeNodeIndexNone != Current);
 
         return false;
@@ -79,12 +123,11 @@ struct FIntervalTree final
 
 		ArrayType<FIntervalTreeNodeIndex> ToProcess;
 
-        ToProcess.Add(Other.Root);
+        ToProcess.Push(Other.Root);
 
         do
         {
-            const FIntervalTreeNodeIndex Current = ToProcess.Last();
-            ToProcess.Pop();
+            const FIntervalTreeNodeIndex Current = ToProcess.Pop();
 
             FRange Range = Other.Nodes[Current].Range;
 
@@ -106,7 +149,7 @@ struct FIntervalTree final
     }
 
 private:
-	template<typename T> using ArrayType = TPagedArray<T>;
+	template<typename T> using ArrayType = TIntervalTreeArray<T>;
 
     struct FRange final
     {
@@ -124,18 +167,16 @@ private:
     struct FIntervalTreeNode final
     {
         explicit FIntervalTreeNode(const FRange Range) :
-            Parent(IntervalTreeNodeIndexNone), bIsBlack(true), Range(Range) {}
+            Range(Range), Parent(IntervalTreeNodeIndexNone), bIsBlack(true) {}
 
 		FIntervalTreeNode(const FRange Range, FIntervalTreeNodeIndex Parent) :
-			Parent(Parent), bIsBlack(false), Range(Range) {}
+			Range(Range), Parent(Parent), bIsBlack(false) {}
 
+		FRange Range;
         FIntervalTreeNodeIndex Left = IntervalTreeNodeIndexNone;
         FIntervalTreeNodeIndex Right = IntervalTreeNodeIndexNone;
         FIntervalTreeNodeIndex Parent;
-
-        // TODO: optimize this waste of memory.
         bool bIsBlack;
-        FRange Range;
     };
 
     FIntervalTreeNodeIndex Root = IntervalTreeNodeIndexNone;
@@ -148,8 +189,7 @@ private:
         if (UNLIKELY(IntervalTreeNodeIndexNone == Root))
         {
             ASSERT(0 == Nodes.Num());
-            Nodes.Add(FIntervalTreeNode(NewRange));
-            Root = 0;
+            Root = Nodes.Add(FIntervalTreeNode(NewRange));
             return true;
         }
 
@@ -157,9 +197,7 @@ private:
 
         for(;;)
         {
-            FIntervalTreeNode& Node = Nodes[Current];
-
-            const FRange Range = Node.Range;
+            const FRange Range = Nodes[Current].Range;
 
             if (UNLIKELY((NewRange.Start < Range.End) && (Range.Start < NewRange.End)))
             {
@@ -168,41 +206,43 @@ private:
 
             if (NewRange.Start < Range.Start)
             {
-                if (NewRange.End == Range.Start)
+                if (UNLIKELY(NewRange.End == Range.Start))
                 {
-                    // We can just modify the existing node in place.
                     ASSERT(NewRange.Start < Range.Start);
-                    Node.Range.Start = NewRange.Start;
+
+					// We can just modify the existing node in place.
+					Nodes[Current].Range.Start = NewRange.Start;
                     return true;
                 }
-                else if (IntervalTreeNodeIndexNone == Node.Left)
+                else if (IntervalTreeNodeIndexNone == Nodes[Current].Left)
                 {
                     const FIntervalTreeNodeIndex Index = Nodes.Add(FIntervalTreeNode(NewRange, Current));
-                    Node.Left = Index;
+					Nodes[Current].Left = Index;
                     Current = Index;
                     break;
                 }
 
-                Current = Node.Left;
+                Current = Nodes[Current].Left;
             }
             else
             {
-                if (NewRange.Start == Range.End)
+                if (UNLIKELY(NewRange.Start == Range.End))
                 {
-                    // We can just modify the existing node in place.
                     ASSERT(NewRange.End > Range.End);
-                    Node.Range.End = NewRange.End;
+
+					// We can just modify the existing node in place.
+					Nodes[Current].Range.End = NewRange.End;
                     return true;
                 }
-                else if (IntervalTreeNodeIndexNone == Node.Right)
+                else if (IntervalTreeNodeIndexNone == Nodes[Current].Right)
                 {
                     const FIntervalTreeNodeIndex Index = Nodes.Add(FIntervalTreeNode(NewRange, Current));
-                    Node.Right = Index;
+					Nodes[Current].Right = Index;
                     Current = Index;
                     break;
                 }
 
-                Current = Node.Right;
+                Current = Nodes[Current].Right;
             }
 
             ASSERT(Root != Current);
diff --git a/Engine/Source/Runtime/Core/Private/AutoRTFM/TransactionInlines.h b/Engine/Source/Runtime/Core/Private/AutoRTFM/TransactionInlines.h
--- a/Engine/Source/Runtime/Core/Private/AutoRTFM/TransactionInlines.h
+++ b/Engine/Source/Runtime/Core/Private/AutoRTFM/TransactionInlines.h
@@ -34,7 +34,7 @@ AUTORTFM_NO_ASAN UE_AUTORTFM_FORCEINLINE void FTransaction::RecordWriteMaxPageSi
 
 AUTORTFM_NO_ASAN UE_AUTORTFM_FORCEINLINE void FTransaction::RecordWrite(void* LogicalAddress, size_t Size)
 {
-    if (Size == 0)
+    if (0 == Size)
     {
         return;
     }
@@ -62,17 +62,14 @@ AUTORTFM_NO_ASAN UE_AUTORTFM_FORCEINLINE void FTransaction::RecordWrite(void* Lo
     
         Stats.Collect<EStatsKind::HitSetMiss>();
     }
-    else
-    {
-        if (NewMemoryTracker.Contains(LogicalAddress, Size))
-        {
-            Stats.Collect<EStatsKind::NewMemoryTrackerHit>();
-            return;
-        }
 
-        Stats.Collect<EStatsKind::NewMemoryTrackerMiss>();
-    }
+	if (NewMemoryTracker.Contains(LogicalAddress, Size))
+	{
+		Stats.Collect<EStatsKind::NewMemoryTrackerHit>();
+		return;
+	}
 
+	Stats.Collect<EStatsKind::NewMemoryTrackerMiss>();
 
     uint8_t* const Address = reinterpret_cast<uint8_t*>(LogicalAddress);
 
@@ -112,6 +109,14 @@ template<unsigned SIZE> AUTORTFM_NO_ASAN UE_AUTORTFM_FORCEINLINE void FTransacti
 
     Stats.Collect<EStatsKind::HitSetMiss>();
 
+	if (NewMemoryTracker.Contains(LogicalAddress, SIZE))
+	{
+		Stats.Collect<EStatsKind::NewMemoryTrackerHit>();
+		return;
+	}
+
+	Stats.Collect<EStatsKind::NewMemoryTrackerMiss>();
+
     RecordWriteMaxPageSized(LogicalAddress, SIZE);
 }
 
@@ -122,22 +127,8 @@ UE_AUTORTFM_FORCEINLINE void FTransaction::DidAllocate(void* LogicalAddress, con
 		return;
 	}
 
-    const bool bUseHitSet = Size <= FWriteLogBumpAllocator::MaxSize;
-
-	if (bUseHitSet)
-    {
-        FMemoryLocation Key(LogicalAddress);
-        Key.SetTopTag(static_cast<uint16_t>(Size));
-        // Otherwise we need to record the write.
-        const bool DidInsert = HitSet.Insert(Key);
-        ASSERT(DidInsert);
-    }
-    
-    if (bTrackAllocationLocations || !bUseHitSet)
-    {
-        const bool DidInsert = NewMemoryTracker.Insert(LogicalAddress, Size);
-        ASSERT(DidInsert);
-    }
+    const bool DidInsert = NewMemoryTracker.Insert(LogicalAddress, Size);
+    ASSERT(DidInsert);
 }
 
 UE_AUTORTFM_FORCEINLINE void FTransaction::DidFree(void* LogicalAddress)
