[Scheduler]
  - Fix DoNotRunInsideBusyWait tasks that ended up being picked up during busy waiting anyway.
  - Change internal API of ExecuteTask to make it evident it can actually return a new task.
  - Make sure that busy wait will dequeue tasks from local/global queue in priority order.
  - Remove unused local queue functions.

diff --git a/Engine/Source/Runtime/Core/Private/Async/Fundamental/Scheduler.cpp b/Engine/Source/Runtime/Core/Private/Async/Fundamental/Scheduler.cpp
--- a/Engine/Source/Runtime/Core/Private/Async/Fundamental/Scheduler.cpp
+++ b/Engine/Source/Runtime/Core/Private/Async/Fundamental/Scheduler.cpp
@@ -143,20 +143,21 @@ namespace LowLevelTasks
 		}
 	}
 
-	inline void FScheduler::ExecuteTask(FTask*& InOutTask)
+	inline FTask* FScheduler::ExecuteTask(FTask* InTask)
 	{
 		FTask* ParentTask = FTask::ActiveTask;
-		FTask::ActiveTask = InOutTask;
+		FTask::ActiveTask = InTask;
+		FTask* OutTask;
 
-		if (!InOutTask->IsBackgroundTask())
+		if (!InTask->IsBackgroundTask())
 		{
 			TRACE_CPUPROFILER_EVENT_SCOPE(ExecuteForegroundTask);
-			InOutTask = InOutTask->ExecuteTask();
+			OutTask = InTask->ExecuteTask();
 		}
 		else
 		{
 			// Dynamic priority only enables for root task when we're not inside a named thread (i.e. GT, RT)
-			const bool bSkipPriorityChange = ParentTask || !GTaskGraphUseDynamicPrioritization || !FSchedulerTls::IsWorkerThread() || InOutTask->WasCanceledOrIsExpediting();
+			const bool bSkipPriorityChange = ParentTask || !GTaskGraphUseDynamicPrioritization || !FSchedulerTls::IsWorkerThread() || InTask->WasCanceledOrIsExpediting();
 
 			FRunnableThread* RunnableThread = nullptr;
 			if (!bSkipPriorityChange)
@@ -173,7 +174,7 @@ namespace LowLevelTasks
 
 			{
 				TRACE_CPUPROFILER_EVENT_SCOPE(ExecuteBackgroundTask);
-				InOutTask = InOutTask->ExecuteTask();
+				OutTask = InTask->ExecuteTask();
 			}
 
 			if (!bSkipPriorityChange)
@@ -184,6 +185,7 @@ namespace LowLevelTasks
 		}
 
 		FTask::ActiveTask = ParentTask;
+		return OutTask;
 	}
 
 	void FScheduler::StopWorkers(bool DrainGlobalQueue)
@@ -215,8 +217,7 @@ namespace LowLevelTasks
 				{
 					while(Task)
 					{
-						ExecuteTask(Task);
-						if(Task)
+						if ((Task = ExecuteTask(Task)) != nullptr)
 						{
 							verifySlow(Task->TryPrepareLaunch());
 						}
@@ -268,10 +269,9 @@ namespace LowLevelTasks
 		else
 		{
 			FTask* TaskPtr = &Task;
-			while(TaskPtr)
+			while (TaskPtr)
 			{
-				ExecuteTask(TaskPtr);
-				if(TaskPtr)
+				if ((TaskPtr = ExecuteTask(TaskPtr)) != nullptr)
 				{
 					verifySlow(TaskPtr->TryPrepareLaunch());
 				}
@@ -316,25 +316,20 @@ namespace LowLevelTasks
 			if constexpr (bIsBusyWaiting)
 			{
 				FTask::FInitData InitData = Task->GetInitData();
-				bool bAllowBusyWaiting = EnumHasAnyFlags(InitData.Flags, ETaskFlags::AllowBusyWaiting) || Task->WasCanceledOrIsExpediting();
-				
-				if (!bAllowBusyWaiting || AnyExecuted)
-				{
-					//either the task is not allowed during busy waiting or we have a  
-					//symetric switching task that we do not want to execute during BusyWaiting
-					//in either case we requeue and try again or we exit early with success.
-					QueueRegistry.Enqueue(Task, uint32(InitData.Priority));
+				const bool bAllowBusyWaiting = EnumHasAnyFlags(InitData.Flags, ETaskFlags::AllowBusyWaiting) || Task->WasCanceledOrIsExpediting();
 
-					if (AnyExecuted)
-					{
-						return AnyExecuted;
-					}
+				if (!bAllowBusyWaiting)
+				{
+					// The task is not allowed during busy waiting so requeue in a different
+					// queue to make sure we won't pick it up again.
+					constexpr static bool bAllowInsideBusyWaiting = false;
+					QueueRegistry.Enqueue<bAllowInsideBusyWaiting>(Task, uint32(InitData.Priority));
 
+					// Fetch another task we could potentially run.
 					Task = (Queue->*DequeueFunction)(bPermitBackgroundWork);
-					if(Task == nullptr)
-					{
-						return false;
-					}
+
+					// Make sure we run the filtering logic again for the new task.
+					continue;
 				}
 			}
 			else
@@ -361,12 +356,21 @@ namespace LowLevelTasks
 				}
 			}
 
-			ExecuteTask(Task);
-			if(Task)
+			AnyExecuted = true;
+
+			// Executing a task can return a continuation.
+			if ((Task = ExecuteTask(Task)) != nullptr)
 			{
 				verifySlow(Task->TryPrepareLaunch());
+
+				// When busy waiting, we exit every time a task is run
+				// so queue this new task for later and bail out.
+				if constexpr (bIsBusyWaiting)
+				{
+					QueueRegistry.Enqueue(Task, uint32(Task->GetPriority()));
+					return AnyExecuted;
+				}
 			}
-			AnyExecuted = true;
 		}
 		return AnyExecuted;
 	}
@@ -389,7 +393,7 @@ namespace LowLevelTasks
 		while (true)
 		{
 			bool bExecutedSomething = false;
-			while(TryExecuteTaskFrom<FSchedulerTls::FLocalQueueType, &FSchedulerTls::FLocalQueueType::Dequeue,  false>(WorkerEvent, WorkerLocalQueue, OutOfWork, bPermitBackgroundWork)
+			while(TryExecuteTaskFrom<FSchedulerTls::FLocalQueueType, &FSchedulerTls::FLocalQueueType::Dequeue, false>(WorkerEvent, WorkerLocalQueue, OutOfWork, bPermitBackgroundWork)
 			   || TryExecuteTaskFrom<FSchedulerTls::FLocalQueueType, &FSchedulerTls::FLocalQueueType::DequeueSteal, false>(WorkerEvent, WorkerLocalQueue, OutOfWork, bPermitBackgroundWork))
 			{
 				bPreparingWait = false;
@@ -438,6 +442,7 @@ namespace LowLevelTasks
 		++FSchedulerTls::BusyWaitingDepth;
 		ON_SCOPE_EXIT{ --FSchedulerTls::BusyWaitingDepth; };
 
+		constexpr static bool bIsInsideBusyWait = true;
 		check(ActiveWorkers.load(std::memory_order_relaxed));
 		FSchedulerTls::FLocalQueueType* WorkerLocalQueue = FSchedulerTls::LocalQueue;
 
@@ -448,24 +453,29 @@ namespace LowLevelTasks
 		Private::FOutOfWork OutOfWork;
 		while (true)
 		{
-			while((WorkerLocalQueue && TryExecuteTaskFrom<FSchedulerTls::FLocalQueueType, &FSchedulerTls::FLocalQueueType::DequeueLocal, true>(nullptr, WorkerLocalQueue, OutOfWork, bPermitBackgroundWork))
-									|| TryExecuteTaskFrom<FSchedulerTls::FQueueRegistry, &FSchedulerTls::FQueueRegistry::DequeueGlobal, true>(nullptr, &QueueRegistry, OutOfWork, bPermitBackgroundWork))
+			if (WorkerLocalQueue)
 			{
-				if (Conditional())
+				while(TryExecuteTaskFrom<FSchedulerTls::FLocalQueueType, &FSchedulerTls::FLocalQueueType::Dequeue<bIsInsideBusyWait>, bIsInsideBusyWait>(nullptr, WorkerLocalQueue, OutOfWork, bPermitBackgroundWork)
+ 									|| TryExecuteTaskFrom<FSchedulerTls::FLocalQueueType, &FSchedulerTls::FLocalQueueType::DequeueSteal, bIsInsideBusyWait>(nullptr, WorkerLocalQueue, OutOfWork, bPermitBackgroundWork))
 				{
-					return;
+					if (Conditional())
+					{
+						return;
+					}
+					WaitCount = 0;
 				}
-				WaitCount = 0;
 			}
-
-			while((WorkerLocalQueue && TryExecuteTaskFrom<FSchedulerTls::FLocalQueueType, &FSchedulerTls::FLocalQueueType::DequeueLocal, true>(nullptr, WorkerLocalQueue, OutOfWork, bPermitBackgroundWork))
-									|| TryExecuteTaskFrom<FSchedulerTls::FQueueRegistry, &FSchedulerTls::FQueueRegistry::DequeueSteal, true>(nullptr, &QueueRegistry, OutOfWork, bPermitBackgroundWork))
+			else
 			{
-				if (Conditional())
+				while(TryExecuteTaskFrom<FSchedulerTls::FQueueRegistry, &FSchedulerTls::FQueueRegistry::DequeueGlobal<bIsInsideBusyWait>, bIsInsideBusyWait>(nullptr, &QueueRegistry, OutOfWork, bPermitBackgroundWork)
+									|| TryExecuteTaskFrom<FSchedulerTls::FQueueRegistry, &FSchedulerTls::FQueueRegistry::DequeueSteal, bIsInsideBusyWait>(nullptr, &QueueRegistry, OutOfWork, bPermitBackgroundWork))
 				{
-					return;
+					if (Conditional())
+					{
+						return;
+					}
+					WaitCount = 0;
 				}
-				WaitCount = 0;
 			}
 
 			if (Conditional())
@@ -495,4 +505,4 @@ namespace LowLevelTasks
 			}
 		}
 	}
-}
+}
\ No newline at end of file
diff --git a/Engine/Source/Runtime/Core/Public/Async/Fundamental/LocalQueue.h b/Engine/Source/Runtime/Core/Public/Async/Fundamental/LocalQueue.h
--- a/Engine/Source/Runtime/Core/Public/Async/Fundamental/LocalQueue.h
+++ b/Engine/Source/Runtime/Core/Public/Async/Fundamental/LocalQueue.h
@@ -155,6 +155,16 @@ private:
 	using FOverflowQueueType = FAAArrayQueue<FTask>;
 	using DequeueHazard		 = typename FOverflowQueueType::DequeueHazard;
 
+	enum class EOverflowType
+	{
+		// Tasks that can't be run inside busy waiting will be moved in that queue during busy wait loops.
+		// This queue is always looked at first by non busy waiting workers to keep ordering as much as possible.
+		DoNotRunInsideBusyWait,
+		// All tasks including those that can't run inside busy wait are queued here by default to keep FIFO ordering as much as possible.
+		Default,
+		Count,
+	};
+
 public:
 	class TLocalQueue
 	{
@@ -165,9 +175,12 @@ public:
 		TLocalQueue(TLocalQueueRegistry& InRegistry, ELocalQueueType InQueueType) : Registry(&InRegistry), QueueType(InQueueType)
 		{
 			AffinityIndex = Registry->AddLocalQueue(this, QueueType);
-			for (int32 PriorityIndex = 0; PriorityIndex < int32(ETaskPriority::Count); PriorityIndex++)
+			for (int32 PriorityIndex = 0; PriorityIndex < int32(ETaskPriority::Count); ++PriorityIndex)
 			{
-				DequeueHazards[PriorityIndex] = Registry->OverflowQueues[PriorityIndex].getHeadHazard();
+				for (int32 OverflowIndex = 0; OverflowIndex < int32(EOverflowType::Count); ++OverflowIndex)
+				{
+					DequeueHazards[PriorityIndex][OverflowIndex] = Registry->OverflowQueues[PriorityIndex][OverflowIndex].getHeadHazard();
+				}
 			}
 		}
 
@@ -182,13 +195,13 @@ public:
 					{
 						break;
 					}
-					Registry->OverflowQueues[PriorityIndex].enqueue(Item);
+					Registry->OverflowQueues[PriorityIndex][(int32)EOverflowType::Default].enqueue(Item);
 				}
 			}
 			Registry->DeleteLocalQueue(this, QueueType);
 		}
 
-		//add an item to the local queue and overflow into the global queue if full
+		// add an item to the local queue and overflow into the global queue if full
 		// returns true if we should wake a worker
 		inline void Enqueue(FTask* Item, uint32 PriorityIndex)
 		{
@@ -198,53 +211,32 @@ public:
 
 			if (!LocalQueues[PriorityIndex].Put(Item))
 			{
-				Registry->OverflowQueues[PriorityIndex].enqueue(Item);
+				Registry->OverflowQueues[PriorityIndex][(int32)EOverflowType::Default].enqueue(Item);
 			}
 		}
 
 		// Check both the local and global queue in priority order
+		template <bool bIsInsideBusyWait = false>
 		inline FTask* Dequeue(bool GetBackGroundTasks)
 		{
-			int32 MaxPriority = GetBackGroundTasks ? int32(ETaskPriority::Count) : int32(ETaskPriority::ForegroundCount);
-			for (int32 PriorityIndex = 0; PriorityIndex < MaxPriority; PriorityIndex++)
-			{
-				FTask* Item;
-				if (LocalQueues[PriorityIndex].Get(Item))
-				{
-					return Item;
-				}
-				Item = Registry->OverflowQueues[PriorityIndex].dequeue(DequeueHazards[PriorityIndex]);
-				if (Item)
-				{
-					return Item;
-				}
-			}
-			return nullptr;
-		}
+			const int32 MaxPriority = GetBackGroundTasks ? int32(ETaskPriority::Count)   : int32(ETaskPriority::ForegroundCount);
+			constexpr int32 MinOverflow = bIsInsideBusyWait  ? int32(EOverflowType::Default) : int32(EOverflowType::DoNotRunInsideBusyWait);
 
-		inline FTask* DequeueLocal(bool GetBackGroundTasks)
-		{
-			int32 MaxPriority = GetBackGroundTasks ? int32(ETaskPriority::Count) : int32(ETaskPriority::ForegroundCount);
-			for (int32 PriorityIndex = 0; PriorityIndex < MaxPriority; PriorityIndex++)
+			for (int32 PriorityIndex = 0; PriorityIndex < MaxPriority; ++PriorityIndex)
 			{
 				FTask* Item;
 				if (LocalQueues[PriorityIndex].Get(Item))
 				{
 					return Item;
 				}
-			}
-			return nullptr;
-		}
 
-		inline FTask* DequeueGlobal(bool GetBackGroundTasks)
-		{
-			int32 MaxPriority = GetBackGroundTasks ? int32(ETaskPriority::Count) : int32(ETaskPriority::ForegroundCount);
-			for (int32 PriorityIndex = 0; PriorityIndex < MaxPriority; PriorityIndex++)
-			{
-				FTask* Item = Registry->OverflowQueues[PriorityIndex].dequeue(DequeueHazards[PriorityIndex]);
-				if (Item)
+				for (int32 OverflowIndex = MinOverflow; OverflowIndex < int32(EOverflowType::Count); ++OverflowIndex)
 				{
-					return Item;
+					Item = Registry->OverflowQueues[PriorityIndex][OverflowIndex].dequeue(DequeueHazards[PriorityIndex][OverflowIndex]);
+					if (Item)
+					{
+						return Item;
+					}
 				}
 			}
 			return nullptr;
@@ -273,7 +265,7 @@ public:
 	private:
 		static constexpr uint32    InvalidIndex = ~0u;
 		FLocalQueueType            LocalQueues[uint32(ETaskPriority::Count)];
-		DequeueHazard              DequeueHazards[uint32(ETaskPriority::Count)];
+		DequeueHazard              DequeueHazards[uint32(ETaskPriority::Count)][int32(EOverflowType::Count)];
 		TLocalQueueRegistry*       Registry;
 		uint32                     CachedRandomIndex = InvalidIndex;
 		uint32                     CachedPriorityIndex = 0;
@@ -310,7 +302,7 @@ private:
 		{
 			TLocalQueue* LocalQueue = LocalQueues[CachedRandomIndex];
 			for(uint32 PriorityIndex = 0; PriorityIndex < MaxPriority; PriorityIndex++)
-			{	
+			{
 				FTask* Item;
 				if (LocalQueue->LocalQueues[CachedPriorityIndex].Steal(Item))
 				{
@@ -327,30 +319,37 @@ private:
 
 public:
 	// enqueue an Item directy into the Global OverflowQueue
+	template <bool bAllowInsideBusyWaiting = true>
 	void Enqueue(FTask* Item, uint32 PriorityIndex)
 	{
 		check(PriorityIndex < int32(ETaskPriority::Count));
 		check(Item != nullptr);
 
-		OverflowQueues[PriorityIndex].enqueue(Item);
+		constexpr int32 OverflowIndex = bAllowInsideBusyWaiting ? (int32)EOverflowType::Default : (int32)EOverflowType::DoNotRunInsideBusyWait;
+		OverflowQueues[PriorityIndex][OverflowIndex].enqueue(Item);
 	}
 
 	// grab an Item directy from the Global OverflowQueue
+	template <bool bIsInsideBusyWait = false>
 	FTask* DequeueGlobal(bool GetBackGroundTasks = true)
 	{
-		int32 MaxPriority = GetBackGroundTasks ? int32(ETaskPriority::Count) : int32(ETaskPriority::ForegroundCount);
-		for (int32 PriorityIndex = 0; PriorityIndex < MaxPriority; PriorityIndex++)
+		const int32 MaxPriority = GetBackGroundTasks ? int32(ETaskPriority::Count) : int32(ETaskPriority::ForegroundCount);
+		constexpr int32 MinOverflow = bIsInsideBusyWait ? int32(EOverflowType::Default) : int32(EOverflowType::DoNotRunInsideBusyWait);
+
+		for (int32 PriorityIndex = 0; PriorityIndex < MaxPriority; ++PriorityIndex)
 		{
-			FTask* Item = OverflowQueues[PriorityIndex].dequeue();
-			if (Item)
+			for (int32 OverflowIndex = MinOverflow; OverflowIndex < int32(EOverflowType::Count); ++OverflowIndex)
 			{
-				return Item;
+				if (FTask* Item = OverflowQueues[PriorityIndex][OverflowIndex].dequeue())
+				{
+					return Item;
+				}
 			}
 		}
 		return nullptr;
 	}
 
-	inline FTask* DequeueSteal(bool GetBackGroundTasks = true)
+	inline FTask* DequeueSteal(bool GetBackGroundTasks)
 	{
 		uint32 CachedRandomIndex = Rand();
 		uint32 CachedPriorityIndex = 0;
@@ -363,7 +362,7 @@ public:
 	}
 
 private:
-	FOverflowQueueType     OverflowQueues[uint32(ETaskPriority::Count)];
+	FOverflowQueueType     OverflowQueues[uint32(ETaskPriority::Count)][uint32(EOverflowType::Count)];
 	TArray<TLocalQueue*>   LocalQueues;
 	int                    NumActiveWorkers[2] = { 0, 0 };
 };
diff --git a/Engine/Source/Runtime/Core/Public/Async/Fundamental/Scheduler.h b/Engine/Source/Runtime/Core/Public/Async/Fundamental/Scheduler.h
--- a/Engine/Source/Runtime/Core/Public/Async/Fundamental/Scheduler.h
+++ b/Engine/Source/Runtime/Core/Public/Async/Fundamental/Scheduler.h
@@ -194,7 +194,7 @@ namespace LowLevelTasks
 		~FScheduler();
 
 	private: 
-		void ExecuteTask(FTask*& InOutTask);
+		UE_NODISCARD FTask* ExecuteTask(FTask* InTask);
 		TUniquePtr<FThread> CreateWorker(bool bPermitBackgroundWork = false, FThread::EForkable IsForkable = FThread::NonForkable, Private::FWaitEvent* ExternalWorkerEvent = nullptr, FSchedulerTls::FLocalQueueType* ExternalWorkerLocalQueue = nullptr, EThreadPriority Priority = EThreadPriority::TPri_Normal, uint64 InAffinity = 0);
 		void WorkerMain(Private::FWaitEvent* WorkerEvent, FSchedulerTls::FLocalQueueType* ExternalWorkerLocalQueue, uint32 WaitCycles, bool bPermitBackgroundWork);
 		CORE_API void LaunchInternal(FTask& Task, EQueuePreference QueuePreference, bool bWakeUpWorker);
diff --git a/Engine/Source/Runtime/Core/Tests/Tasks/TasksTest.cpp b/Engine/Source/Runtime/Core/Tests/Tasks/TasksTest.cpp
--- a/Engine/Source/Runtime/Core/Tests/Tasks/TasksTest.cpp
+++ b/Engine/Source/Runtime/Core/Tests/Tasks/TasksTest.cpp
@@ -1665,6 +1665,57 @@ namespace UE { namespace TasksTests
 		}
 	}
 
+	TEST_CASE_NAMED(FTasksDoNotRunInsideBusyWait, "System::Core::Async::Tasks::DoNotRunInsideBusyWait", "[.][ApplicationContextMask][EngineFilter]")
+	{
+		using namespace LowLevelTasks;
+
+		FPlatformProcess::Sleep(0.1f); // give workers time to fall asleep, to avoid any reserve worker messing around
+
+		uint32 NumWorkers = LowLevelTasks::FScheduler::Get().GetNumWorkers();
+
+		// Block all workers to make sure the tasks we're going to queue are not executed before we enter our busy wait loop.
+		FTaskEvent ResumeEvent{ UE_SOURCE_LOCATION };
+		TArray<LowLevelTasks::FTask> WorkerBlockers = BlockWorkers(ResumeEvent, NumWorkers);
+
+		std::atomic<bool>     CanExecute { false };
+		std::atomic<int32>    NumExecuted { 0 };
+		UE::FManualResetEvent Done;
+
+		// Queue enough busy wait excluded tasks to verify that the scheduler
+		// will do the right thing and not end up actually executing one of them
+		// from inside busy wait.
+		for (int32 Index = 0; Index < 100; ++Index)
+		{
+			Launch(UE_SOURCE_LOCATION,
+				[&CanExecute, &NumExecuted, &Done]()
+				{
+					verify(CanExecute.load());
+					if (++NumExecuted == 100)
+					{
+						Done.Notify();
+					}
+				},
+				UE::Tasks::ETaskPriority::Default,
+				UE::Tasks::EExtendedTaskPriority::None,
+				UE::Tasks::ETaskFlags::DoNotRunInsideBusyWait // this should not be picked up by busy waiting
+			);
+		}
+
+		// Busy wait for a second, making sure no excluded tasks are executed.
+		double StartTime = FPlatformTime::Seconds() + 1;
+		BusyWaitUntil([&StartTime]() { return FPlatformTime::Seconds() > StartTime; });
+
+		// Allow execution now.
+		CanExecute = true;
+
+		ResumeEvent.Trigger();
+
+		// Now busy waiting will run since we just unblocked all workers.
+		verify(Done.WaitFor(UE::FMonotonicTimeSpan::FromSeconds(1)));
+
+		// Do not exit the test before all blocked workers' tasks are done.
+		LowLevelTasks::BusyWaitForTasks<LowLevelTasks::FTask>(WorkerBlockers);
+	}
 }}
 
 #endif // WITH_TESTS
