Significantly reduced PDB bloat caused by a lambda in the previous change in the assertion macros

diff --git a/Engine/Source/Runtime/Core/Private/Misc/AssertionMacros.cpp b/Engine/Source/Runtime/Core/Private/Misc/AssertionMacros.cpp
--- a/Engine/Source/Runtime/Core/Private/Misc/AssertionMacros.cpp
+++ b/Engine/Source/Runtime/Core/Private/Misc/AssertionMacros.cpp
@@ -611,6 +611,44 @@ bool FORCENOINLINE FDebug::CheckVerifyFailedImpl(
 #endif
 }
 
+bool FORCENOINLINE FDebug::CheckVerifyFailedImpl2(
+	const ANSICHAR* Expr,
+	const ANSICHAR* File,
+	int32 Line,
+	const TCHAR* Format,
+	...)
+{
+	va_list Args;
+
+	va_start(Args, Format);
+	FDebug::LogAssertFailedMessageImplV(Expr, File, Line, PLATFORM_RETURN_ADDRESS(), Format, Args);
+	va_end(Args);
+
+	if (GLog)
+	{
+		// Flushing the logs here increases the likelihood that recent messages will be written to the log file, stdout and the debugger console.
+		// Without this, some of the recent messages may not be reported when debugger stops due to an assertion failure.
+		GLog->Flush();
+	}
+
+	if (!FPlatformMisc::IsDebuggerPresent())
+	{
+		FPlatformMisc::PromptForRemoteDebugging(false);
+
+		va_start(Args, Format);
+		AssertFailedImplV(Expr, File, Line, PLATFORM_RETURN_ADDRESS(), Format, Args);
+		va_end(Args);
+
+		return false;
+	}
+
+#if UE_BUILD_SHIPPING
+	return true;
+#else
+	return !GIgnoreDebugger;
+#endif
+}
+
 #endif // DO_CHECK || DO_GUARD_SLOW || DO_ENSURE
 
 void VARARGS FDebug::AssertFailed(const ANSICHAR* Expr, const ANSICHAR* File, int32 Line, const TCHAR* Format/* = TEXT("")*/, ...)
diff --git a/Engine/Source/Runtime/Core/Public/Misc/AssertionMacros.h b/Engine/Source/Runtime/Core/Public/Misc/AssertionMacros.h
--- a/Engine/Source/Runtime/Core/Public/Misc/AssertionMacros.h
+++ b/Engine/Source/Runtime/Core/Public/Misc/AssertionMacros.h
@@ -88,6 +88,7 @@ struct FDebug
 #if DO_CHECK || DO_GUARD_SLOW || DO_ENSURE
 public:
 	static CORE_API bool VARARGS CheckVerifyFailedImpl(const ANSICHAR* Expr, const ANSICHAR* File, int32 Line, void* ProgramCounter, const TCHAR* Format, ...);
+	static CORE_API bool VARARGS CheckVerifyFailedImpl2(const ANSICHAR* Expr, const ANSICHAR* File, int32 Line, const TCHAR* Format, ...);
 private:
 	static CORE_API void VARARGS LogAssertFailedMessageImpl(const ANSICHAR* Expr, const ANSICHAR* File, int32 Line, void* ProgramCounter, const TCHAR* Fmt, ...);
 	static CORE_API void LogAssertFailedMessageImplV(const ANSICHAR* Expr, const ANSICHAR* File, int32 Line, void* ProgramCounter, const TCHAR* Fmt, va_list Args);
@@ -233,17 +234,10 @@ RetType FORCENOINLINE UE_DEBUG_SECTION DispatchCheckVerify(InnerType&& Inner, Ar
 		{ \
 			if(UNLIKELY(!(expr))) \
 			{ \
-				struct Impl \
+				if (FDebug::CheckVerifyFailedImpl2(#expr, __FILE__, __LINE__, TEXT(""))) \
 				{ \
-					static void FORCENOINLINE UE_DEBUG_SECTION ExecCheckImplInternal() \
-					{ \
-						if (FDebug::CheckVerifyFailedImpl(#expr, __FILE__, __LINE__, PLATFORM_RETURN_ADDRESS(), TEXT(""))) \
-						{ \
-							PLATFORM_BREAK(); \
-						} \
-					} \
-				}; \
-				Impl::ExecCheckImplInternal(); \
+					PLATFORM_BREAK(); \
+				} \
 				CA_ASSUME(false); \
 			} \
 		}
@@ -263,11 +257,11 @@ RetType FORCENOINLINE UE_DEBUG_SECTION DispatchCheckVerify(InnerType&& Inner, Ar
 		{ \
 			if(UNLIKELY(!(expr))) \
 			{ \
-				if (DispatchCheckVerify<bool>([&] () UE_DEBUG_SECTION \
+				UE_VALIDATE_FORMAT_STRING(format, ##__VA_ARGS__); \
+				if (FDebug::CheckVerifyFailedImpl2(#expr, __FILE__, __LINE__, format, ##__VA_ARGS__)) \
 				{ \
-					UE_VALIDATE_FORMAT_STRING(format, ##__VA_ARGS__); \
-					return FDebug::CheckVerifyFailedImpl(#expr, __FILE__, __LINE__, PLATFORM_RETURN_ADDRESS(), format, ##__VA_ARGS__); \
-				})) PLATFORM_BREAK(); \
+					PLATFORM_BREAK(); \
+				} \
 				CA_ASSUME(false); \
 			} \
 		}
