Added ability to disable FText display string support to save memory (eg, in server builds)

This is controlled by the "Localization.DisplayStringSupport" CVar, and defaults to "Auto" where servers no longer load or initialize display strings for FText.

diff --git a/Engine/Source/Runtime/Core/Private/Internationalization/StringTableCore.cpp b/Engine/Source/Runtime/Core/Private/Internationalization/StringTableCore.cpp
--- a/Engine/Source/Runtime/Core/Private/Internationalization/StringTableCore.cpp
+++ b/Engine/Source/Runtime/Core/Private/Internationalization/StringTableCore.cpp
@@ -69,7 +69,11 @@ const FString& FStringTableEntry::GetSourceString() const
 
 FTextConstDisplayStringPtr FStringTableEntry::GetDisplayString() const
 {
-	return FTextLocalizationManager::Get().GetDisplayString(DisplayStringId.GetNamespace(), DisplayStringId.GetKey(), &SourceString);
+	if (FTextLocalizationManager::IsDisplayStringSupportEnabled())
+	{
+		return FTextLocalizationManager::Get().GetDisplayString(DisplayStringId.GetNamespace(), DisplayStringId.GetKey(), &SourceString);
+	}
+	return nullptr;
 }
 
 FTextId FStringTableEntry::GetDisplayStringId() const
diff --git a/Engine/Source/Runtime/Core/Private/Internationalization/Text.cpp b/Engine/Source/Runtime/Core/Private/Internationalization/Text.cpp
--- a/Engine/Source/Runtime/Core/Private/Internationalization/Text.cpp
+++ b/Engine/Source/Runtime/Core/Private/Internationalization/Text.cpp
@@ -76,12 +76,6 @@ const FString& FTextInspector::GetDisplayString(const FText& Text)
 	return Text.TextData->GetDisplayString();
 }
 
-FTextConstDisplayStringPtr FTextInspector::GetSharedDisplayString(const FText& Text)
-{
-	Text.Rebuild();
-	return Text.TextData->GetLocalizedString();
-}
-
 bool FTextInspector::GetTableIdAndKey(const FText& Text, FName& OutTableId, FString& OutKey)
 {
 	FTextKey TmpKey;
diff --git a/Engine/Source/Runtime/Core/Private/Internationalization/TextHistory.cpp b/Engine/Source/Runtime/Core/Private/Internationalization/TextHistory.cpp
--- a/Engine/Source/Runtime/Core/Private/Internationalization/TextHistory.cpp
+++ b/Engine/Source/Runtime/Core/Private/Internationalization/TextHistory.cpp
@@ -900,7 +900,7 @@ void FTextHistory_Base::Serialize(FStructuredArchive::FRecord Record)
 
 bool FTextHistory_Base::CanUpdateDisplayString()
 {
-	return !TextId.IsEmpty();
+	return FTextLocalizationManager::IsDisplayStringSupportEnabled() && !TextId.IsEmpty();
 }
 
 void FTextHistory_Base::UpdateDisplayString()
@@ -2396,11 +2396,15 @@ const FString& FTextHistory_StringTableEntry::GetSourceString() const
 
 const FString& FTextHistory_StringTableEntry::GetDisplayString() const
 {
-	if (FTextConstDisplayStringPtr DisplayString = GetLocalizedString())
+	if (FTextLocalizationManager::IsDisplayStringSupportEnabled())
 	{
-		return *DisplayString;
+		if (FTextConstDisplayStringPtr DisplayString = GetLocalizedString())
+		{
+			return *DisplayString;
+		}
+		return FStringTableEntry::GetPlaceholderSourceString();
 	}
-	return FStringTableEntry::GetPlaceholderSourceString();
+	return GetSourceString();
 }
 
 FString FTextHistory_StringTableEntry::BuildInvariantDisplayString() const
diff --git a/Engine/Source/Runtime/Core/Private/Internationalization/TextLocalizationManager.cpp b/Engine/Source/Runtime/Core/Private/Internationalization/TextLocalizationManager.cpp
--- a/Engine/Source/Runtime/Core/Private/Internationalization/TextLocalizationManager.cpp
+++ b/Engine/Source/Runtime/Core/Private/Internationalization/TextLocalizationManager.cpp
@@ -33,6 +33,15 @@ DEFINE_LOG_CATEGORY_STATIC(LogTextLocalizationManager, Log, All);
 
 namespace TextLocalizationManager
 {
+enum class EDisplayStringSupport : int32
+{
+	Auto = 0,
+	Enabled = 1,
+	Disabled = 2,
+};
+static int32 DisplayStringSupport = static_cast<int32>(EDisplayStringSupport::Auto);
+static FAutoConsoleVariableRef CVarDisplayStringSupport(TEXT("Localization.DisplayStringSupport"), DisplayStringSupport, TEXT("Is display string support enabled? 0: Auto (default), 1: Enabled, 2: Disabled"));
+
 static bool AsyncLoadLocalizationData = true;
 static FAutoConsoleVariableRef CVarAsyncLoadLocalizationData(TEXT("Localization.AsyncLoadLocalizationData"), AsyncLoadLocalizationData, TEXT("True to load localization data asynchronously (non-blocking), or False to load it synchronously (blocking)"));
 
@@ -530,8 +539,29 @@ void FTextLocalizationManager::TearDown()
 	FTextKey::TearDown();
 }
 
+bool FTextLocalizationManager::IsDisplayStringSupportEnabled()
+{
+	switch (static_cast<TextLocalizationManager::EDisplayStringSupport>(TextLocalizationManager::DisplayStringSupport))
+	{
+	case TextLocalizationManager::EDisplayStringSupport::Auto:
+#if UE_EDITOR
+		return true; // IsRunningDedicatedServer asserts during static-init if called in the editor
+#else
+		return !IsRunningDedicatedServer();
+#endif
+	case TextLocalizationManager::EDisplayStringSupport::Enabled:
+		return true;
+	case TextLocalizationManager::EDisplayStringSupport::Disabled:
+		return false;
+	default:
+		checkf(false, TEXT("Unknown EDisplayStringSupport!"));
+		break;
+	}
+	return true;
+}
+
 FTextLocalizationManager::FTextLocalizationManager()
-	: TextRevisionCounter(0)
+	: TextRevisionCounter(1) // Default to 1 as 0 is considered unset
 	, LocResTextSource(MakeShared<FLocalizationResourceTextSource>())
 	, PolyglotTextSource(MakeShared<FPolyglotTextSource>())
 {
@@ -749,6 +779,11 @@ void FTextLocalizationManager::RegisterPolyglotTextData(TArrayView<const FPolygl
 
 FTextConstDisplayStringPtr FTextLocalizationManager::FindDisplayString(const FTextKey& Namespace, const FTextKey& Key, const FString* const SourceString) const
 {
+	if (!FTextLocalizationManager::IsDisplayStringSupportEnabled())
+	{
+		return nullptr;
+	}
+
 	FScopeLock ScopeLock(&DisplayStringLookupTableCS);
 
 	const FTextId TextId(Namespace, Key);
@@ -768,6 +803,18 @@ FTextConstDisplayStringRef FTextLocalizationManager::GetDisplayString(const FTex
 	TRACE_CPUPROFILER_EVENT_SCOPE(FTextLocalizationManager::GetDisplayString);
 	LLM_SCOPE(ELLMTag::Localization);
 
+	auto GetEmptyDisplayString = []()
+	{
+		static const FTextConstDisplayStringRef EmptyDisplayString = MakeTextDisplayString(FString());
+		return EmptyDisplayString;
+	};
+
+	if (!FTextLocalizationManager::IsDisplayStringSupportEnabled())
+	{
+		ensureAlwaysMsgf(false, TEXT("FTextLocalizationManager::GetDisplayString was called but display string support is disabled! Either update the calling code to respect FTextLocalizationManager::IsDisplayStringSupportEnabled, or re-enable display string support via 'Localization.DisplayStringSupport'."));
+		return GetEmptyDisplayString();
+	}
+
 	FScopeLock ScopeLock(&DisplayStringLookupTableCS);
 
 	// Hack fix for old assets that don't have namespace/key info.
@@ -870,12 +917,6 @@ FTextConstDisplayStringRef FTextLocalizationManager::GetDisplayString(const FTex
 		{
 			UE_LOG(LogTextLocalizationManager, Verbose, TEXT("An attempt was made to get a localized string (Namespace:%s, Key:%s, Source:%s), but it did not exist."), TextId.GetNamespace().GetChars(), TextId.GetKey().GetChars(), SourceString ? **SourceString : TEXT(""));
 		}
-
-		auto GetEmptyDisplayString = []()
-		{
-			static const FTextConstDisplayStringRef EmptyDisplayString = MakeTextDisplayString(FString());
-			return EmptyDisplayString;
-		};
 		
 		FTextConstDisplayStringRef UnlocalizedString = SourceString ? MakeTextDisplayString(CopyTemp(*SourceString)) : GetEmptyDisplayString();
 
@@ -1072,6 +1113,13 @@ void FTextLocalizationManager::OnCultureChanged()
 		return;
 	}
 
+	if (!FTextLocalizationManager::IsDisplayStringSupportEnabled())
+	{
+		// When display strings are disabled just bump the text revision (so that generated text updates correctly for the new locale) and bail
+		DirtyTextRevision();
+		return;
+	}
+
 	RefreshResources();
 
 	if (!TextLocalizationManager::AsyncLoadLocalizationDataOnLanguageChange)
@@ -1114,7 +1162,7 @@ void FTextLocalizationManager::LoadLocalizationResourcesForPrioritizedCultures_S
 	TRACE_CPUPROFILER_EVENT_SCOPE(FTextLocalizationManager::LoadLocalizationResourcesForPrioritizedCultures);
 
 	// Nothing to do?
-	if (PrioritizedCultureNames.Num() == 0)
+	if (!FTextLocalizationManager::IsDisplayStringSupportEnabled() || PrioritizedCultureNames.Num() == 0)
 	{
 		return;
 	}
@@ -1182,6 +1230,12 @@ void FTextLocalizationManager::LoadLocalizationTargetsForPrioritizedCultures_Syn
 	LLM_SCOPE(ELLMTag::Localization);
 	TRACE_CPUPROFILER_EVENT_SCOPE(FTextLocalizationManager::LoadLocalizationTargetsForPrioritizedCultures);
 
+	// Nothing to do?
+	if (!FTextLocalizationManager::IsDisplayStringSupportEnabled() || PrioritizedCultureNames.Num() == 0 || LocalizationTargetPaths.Num() == 0)
+	{
+		return;
+	}
+
 	// Load the resources from each localization target
 	FTextLocalizationResource UnusedNativeResource;
 	FTextLocalizationResource LocalizedResource;
@@ -1259,6 +1313,15 @@ void FTextLocalizationManager::LoadChunkedLocalizationResources_Sync(TArrayView<
 		return;
 	}
 
+	// Nothing to do?
+	if (!FTextLocalizationManager::IsDisplayStringSupportEnabled())
+	{
+		LocResTextSource->RegisterChunkId(ChunkId);
+
+		UE_LOG(LogTextLocalizationManager, Verbose, TEXT("Skipped loading localization data for chunk %d (from PAK '%s') as display strings are disabled"), ChunkId, *PakFilename);
+		return;
+	}
+
 	// If we're being notified so early that even InitEngineTextLocalization hasn't run, then we can't safely make the queries below as things like GConfig may not be available yet!
 	if (!IsInitialized())
 	{
@@ -1373,6 +1436,12 @@ void FTextLocalizationManager::UpdateFromNative(FTextLocalizationResource&& Text
 	TRACE_CPUPROFILER_EVENT_SCOPE(FTextLocalizationManager::UpdateFromNative);
 	LLM_SCOPE(ELLMTag::Localization);
 
+	// Nothing to do?
+	if (!FTextLocalizationManager::IsDisplayStringSupportEnabled())
+	{
+		return;
+	}
+
 	// Lock while updating the tables
 	{
 		FScopeLock ScopeLock(&DisplayStringLookupTableCS);
@@ -1480,6 +1549,12 @@ void FTextLocalizationManager::UpdateFromLocalizations(FTextLocalizationResource
 	TRACE_CPUPROFILER_EVENT_SCOPE(FTextLocalizationManager::UpdateFromLocalizations);
 	LLM_SCOPE(ELLMTag::Localization);
 
+	// Nothing to do?
+	if (!FTextLocalizationManager::IsDisplayStringSupportEnabled())
+	{
+		return;
+	}
+
 	static const bool bShouldLEETIFYUnlocalizedString = FParse::Param(FCommandLine::Get(), TEXT("LEETIFYUnlocalized"));
 
 	// Lock while updating the tables
diff --git a/Engine/Source/Runtime/Core/Public/Internationalization/Text.h b/Engine/Source/Runtime/Core/Public/Internationalization/Text.h
--- a/Engine/Source/Runtime/Core/Public/Internationalization/Text.h
+++ b/Engine/Source/Runtime/Core/Public/Internationalization/Text.h
@@ -1156,8 +1156,6 @@ public:
 	static CORE_API FTextId GetTextId(const FText& Text);
 	static CORE_API const FString* GetSourceString(const FText& Text);
 	static CORE_API const FString& GetDisplayString(const FText& Text);
-	UE_DEPRECATED(5.0, "GetSharedDisplayString is no longer guaranteed to return a valid result and should NOT be used! If you wanted to get the text ID, use FTextInspector::GetTextId instead. If you wanted a key for unique text instances, use FTextInspector::GetSharedDataId instead.")
-	static CORE_API FTextConstDisplayStringPtr GetSharedDisplayString(const FText& Text);
 	static CORE_API bool GetTableIdAndKey(const FText& Text, FName& OutTableId, FString& OutKey);
 	static CORE_API bool GetTableIdAndKey(const FText& Text, FName& OutTableId, FTextKey& OutKey);
 	static CORE_API uint32 GetFlags(const FText& Text);
diff --git a/Engine/Source/Runtime/Core/Public/Internationalization/TextLocalizationManager.h b/Engine/Source/Runtime/Core/Public/Internationalization/TextLocalizationManager.h
--- a/Engine/Source/Runtime/Core/Public/Internationalization/TextLocalizationManager.h
+++ b/Engine/Source/Runtime/Core/Public/Internationalization/TextLocalizationManager.h
@@ -113,6 +113,8 @@ public:
 	static CORE_API FTextLocalizationManager& Get();
 	static CORE_API void TearDown();
 
+	static CORE_API bool IsDisplayStringSupportEnabled();
+
 	CORE_API void DumpMemoryInfo() const;
 	CORE_API void CompactDataStructures();
 
@@ -172,12 +174,6 @@ public:
 	/** If an entry exists for the specified namespace and key, returns true and provides the localization resource identifier from which it was loaded. Otherwise, returns false. */
 	CORE_API bool GetLocResID(const FTextKey& Namespace, const FTextKey& Key, FString& OutLocResId) const;
 #endif
-
-	UE_DEPRECATED(5.0, "FindNamespaceAndKeyFromDisplayString no longer functions! Use FTextInspector::GetTextId instead.")
-	bool FindNamespaceAndKeyFromDisplayString(const FTextConstDisplayStringPtr& InDisplayString, FString& OutNamespace, FString& OutKey) const { return false; }
-
-	UE_DEPRECATED(5.0, "FindNamespaceAndKeyFromDisplayString no longer functions! Use FTextInspector::GetTextId instead.")
-	bool FindNamespaceAndKeyFromDisplayString(const FTextConstDisplayStringPtr& InDisplayString, FTextKey& OutNamespace, FTextKey& OutKey) const { return false; }
 	
 	/**	Attempts to register the specified display string, associating it with the specified namespace and key.
 	 *	Returns true if the display string has been or was already associated with the namespace and key.
diff --git a/Engine/Source/Runtime/Core/Tests/Internationalization/TextTest.cpp b/Engine/Source/Runtime/Core/Tests/Internationalization/TextTest.cpp
--- a/Engine/Source/Runtime/Core/Tests/Internationalization/TextTest.cpp
+++ b/Engine/Source/Runtime/Core/Tests/Internationalization/TextTest.cpp
@@ -546,6 +546,7 @@ public:
 					FormattedTestLayer2_OriginalLanguageSourceString = FormattedTestLayer2.BuildSourceString();
 
 #if ENABLE_LOC_TESTING
+					if (FTextLocalizationManager::IsDisplayStringSupportEnabled()) // Leetification requires display strings
 					{
 						// Swap to "LEET" culture to check if rebuilding works (verify the whole)
 						I18N.SetCurrentCulture(FLeetCulture::StaticGetName());
