UE4 -> UE namespace rename

#rnx

diff --git a/Engine/Source/Runtime/Core/Public/Internationalization/Internationalization.h b/Engine/Source/Runtime/Core/Public/Internationalization/Internationalization.h
--- a/Engine/Source/Runtime/Core/Public/Internationalization/Internationalization.h
+++ b/Engine/Source/Runtime/Core/Public/Internationalization/Internationalization.h
@@ -271,9 +271,9 @@ private:
 	TArray<FCultureRef> CustomCultures;
 };
 
-namespace UE4LocGen_Private
+namespace UE::Private::LocGen
 {
-	inline FCulturePtr GetCultureImpl(const TCHAR* InCulture)
+	inline FCulturePtr GetCultureByName(const TCHAR* InCulture)
 	{
 		return (InCulture && *InCulture)
 			? FInternationalization::Get().GetCulture(InCulture)
@@ -307,10 +307,10 @@ namespace UE4LocGen_Private
  * @param InCulture		The culture code to use, or an empty string to use the active locale.
  * @param InOpts		Custom formatting options specified as chained setter functions of FNumberFormattingOptions (eg, SetAlwaysSign(true).SetUseGrouping(false)).
  */
-#define LOCGEN_NUMBER(InNum, InCulture) FText::AsNumber(InNum, nullptr, UE4LocGen_Private::GetCultureImpl(TEXT(InCulture)))
-#define LOCGEN_NUMBER_GROUPED(InNum, InCulture) FText::AsNumber(InNum, &FNumberFormattingOptions::DefaultWithGrouping(), UE4LocGen_Private::GetCultureImpl(TEXT(InCulture)))
-#define LOCGEN_NUMBER_UNGROUPED(InNum, InCulture) FText::AsNumber(InNum, &FNumberFormattingOptions::DefaultNoGrouping(), UE4LocGen_Private::GetCultureImpl(TEXT(InCulture)))
-#define LOCGEN_NUMBER_CUSTOM(InNum, InOpts, InCulture) FText::AsNumber(InNum, &FNumberFormattingOptions().InOpts, UE4LocGen_Private::GetCultureImpl(TEXT(InCulture)))
+#define LOCGEN_NUMBER(InNum, InCulture) FText::AsNumber(InNum, nullptr, UE::Private::LocGen::GetCultureByName(TEXT(InCulture)))
+#define LOCGEN_NUMBER_GROUPED(InNum, InCulture) FText::AsNumber(InNum, &FNumberFormattingOptions::DefaultWithGrouping(), UE::Private::LocGen::GetCultureByName(TEXT(InCulture)))
+#define LOCGEN_NUMBER_UNGROUPED(InNum, InCulture) FText::AsNumber(InNum, &FNumberFormattingOptions::DefaultNoGrouping(), UE::Private::LocGen::GetCultureByName(TEXT(InCulture)))
+#define LOCGEN_NUMBER_CUSTOM(InNum, InOpts, InCulture) FText::AsNumber(InNum, &FNumberFormattingOptions().InOpts, UE::Private::LocGen::GetCultureByName(TEXT(InCulture)))
 
 /**
  * Generate an FText representation of the given number as a percentage (alias for FText::AsPercent).
@@ -319,10 +319,10 @@ namespace UE4LocGen_Private
  * @param InCulture		The culture code to use, or an empty string to use the active locale.
  * @param InOpts		Custom formatting options specified as chained setter functions of FNumberFormattingOptions (eg, SetAlwaysSign(true).SetUseGrouping(false)).
  */
-#define LOCGEN_PERCENT(InNum, InCulture) FText::AsPercent(InNum, nullptr, UE4LocGen_Private::GetCultureImpl(TEXT(InCulture)))
-#define LOCGEN_PERCENT_GROUPED(InNum, InCulture) FText::AsPercent(InNum, &FNumberFormattingOptions::DefaultWithGrouping(), UE4LocGen_Private::GetCultureImpl(TEXT(InCulture)))
-#define LOCGEN_PERCENT_UNGROUPED(InNum, InCulture) FText::AsPercent(InNum, &FNumberFormattingOptions::DefaultNoGrouping(), UE4LocGen_Private::GetCultureImpl(TEXT(InCulture)))
-#define LOCGEN_PERCENT_CUSTOM(InNum, InOpts, InCulture) FText::AsPercent(InNum, &FNumberFormattingOptions().InOpts, UE4LocGen_Private::GetCultureImpl(TEXT(InCulture)))
+#define LOCGEN_PERCENT(InNum, InCulture) FText::AsPercent(InNum, nullptr, UE::Private::LocGen::GetCultureByName(TEXT(InCulture)))
+#define LOCGEN_PERCENT_GROUPED(InNum, InCulture) FText::AsPercent(InNum, &FNumberFormattingOptions::DefaultWithGrouping(), UE::Private::LocGen::GetCultureByName(TEXT(InCulture)))
+#define LOCGEN_PERCENT_UNGROUPED(InNum, InCulture) FText::AsPercent(InNum, &FNumberFormattingOptions::DefaultNoGrouping(), UE::Private::LocGen::GetCultureByName(TEXT(InCulture)))
+#define LOCGEN_PERCENT_CUSTOM(InNum, InOpts, InCulture) FText::AsPercent(InNum, &FNumberFormattingOptions().InOpts, UE::Private::LocGen::GetCultureByName(TEXT(InCulture)))
 
 /**
  * Generate an FText representation of the given number as a currency (alias for FText::AsCurrencyBase).
@@ -331,7 +331,7 @@ namespace UE4LocGen_Private
  * @param InCurrency	The currency code (eg, USD, GBP, EUR).
  * @param InCulture		The culture code to use, or an empty string to use the active locale.
  */
-#define LOCGEN_CURRENCY(InNum, InCurrency, InCulture) FText::AsCurrencyBase(InNum, TEXT(InCurrency), UE4LocGen_Private::GetCultureImpl(TEXT(InCulture)))
+#define LOCGEN_CURRENCY(InNum, InCurrency, InCulture) FText::AsCurrencyBase(InNum, TEXT(InCurrency), UE::Private::LocGen::GetCultureByName(TEXT(InCulture)))
 
 /**
  * Generate an FText representation of the given timestamp as a date (alias for FText::AsDate).
@@ -341,8 +341,8 @@ namespace UE4LocGen_Private
  * @param InTimeZone	The timezone to display the timestamp in.
  * @param InCulture		The culture code to use, or an empty string to use the active locale.
  */
-#define LOCGEN_DATE_UTC(InUnixTime, InDateStyle, InTimeZone, InCulture) FText::AsDate(FDateTime::FromUnixTimestamp(InUnixTime), InDateStyle, TEXT(InTimeZone), UE4LocGen_Private::GetCultureImpl(TEXT(InCulture)))
-#define LOCGEN_DATE_LOCAL(InUnixTime, InDateStyle, InCulture) FText::AsDate(FDateTime::FromUnixTimestamp(InUnixTime), InDateStyle, FText::GetInvariantTimeZone(), UE4LocGen_Private::GetCultureImpl(TEXT(InCulture)))
+#define LOCGEN_DATE_UTC(InUnixTime, InDateStyle, InTimeZone, InCulture) FText::AsDate(FDateTime::FromUnixTimestamp(InUnixTime), InDateStyle, TEXT(InTimeZone), UE::Private::LocGen::GetCultureByName(TEXT(InCulture)))
+#define LOCGEN_DATE_LOCAL(InUnixTime, InDateStyle, InCulture) FText::AsDate(FDateTime::FromUnixTimestamp(InUnixTime), InDateStyle, FText::GetInvariantTimeZone(), UE::Private::LocGen::GetCultureByName(TEXT(InCulture)))
 
 /**
  * Generate an FText representation of the given timestamp as a time (alias for FText::AsTime).
@@ -352,8 +352,8 @@ namespace UE4LocGen_Private
  * @param InTimeZone	The timezone to display the timestamp in.
  * @param InCulture		The culture code to use, or an empty string to use the active locale.
  */
-#define LOCGEN_TIME_UTC(InUnixTime, InTimeStyle, InTimeZone, InCulture) FText::AsTime(FDateTime::FromUnixTimestamp(InUnixTime), InTimeStyle, TEXT(InTimeZone), UE4LocGen_Private::GetCultureImpl(TEXT(InCulture)))
-#define LOCGEN_TIME_LOCAL(InUnixTime, InTimeStyle, InCulture) FText::AsTime(FDateTime::FromUnixTimestamp(InUnixTime), InTimeStyle, FText::GetInvariantTimeZone(), UE4LocGen_Private::GetCultureImpl(TEXT(InCulture)))
+#define LOCGEN_TIME_UTC(InUnixTime, InTimeStyle, InTimeZone, InCulture) FText::AsTime(FDateTime::FromUnixTimestamp(InUnixTime), InTimeStyle, TEXT(InTimeZone), UE::Private::LocGen::GetCultureByName(TEXT(InCulture)))
+#define LOCGEN_TIME_LOCAL(InUnixTime, InTimeStyle, InCulture) FText::AsTime(FDateTime::FromUnixTimestamp(InUnixTime), InTimeStyle, FText::GetInvariantTimeZone(), UE::Private::LocGen::GetCultureByName(TEXT(InCulture)))
 
 /**
  * Generate an FText representation of the given timestamp as a date and time (alias for FText::AsDateTime).
@@ -364,8 +364,8 @@ namespace UE4LocGen_Private
  * @param InTimeZone	The timezone to display the timestamp in.
  * @param InCulture		The culture code to use, or an empty string to use the active locale.
  */
-#define LOCGEN_DATETIME_UTC(InUnixTime, InDateStyle, InTimeStyle, InTimeZone, InCulture) FText::AsDateTime(FDateTime::FromUnixTimestamp(InUnixTime), InDateStyle, InTimeStyle, TEXT(InTimeZone), UE4LocGen_Private::GetCultureImpl(TEXT(InCulture)))
-#define LOCGEN_DATETIME_LOCAL(InUnixTime, InDateStyle, InTimeStyle, InCulture) FText::AsDateTime(FDateTime::FromUnixTimestamp(InUnixTime), InDateStyle, InTimeStyle, FText::GetInvariantTimeZone(), UE4LocGen_Private::GetCultureImpl(TEXT(InCulture)))
+#define LOCGEN_DATETIME_UTC(InUnixTime, InDateStyle, InTimeStyle, InTimeZone, InCulture) FText::AsDateTime(FDateTime::FromUnixTimestamp(InUnixTime), InDateStyle, InTimeStyle, TEXT(InTimeZone), UE::Private::LocGen::GetCultureByName(TEXT(InCulture)))
+#define LOCGEN_DATETIME_LOCAL(InUnixTime, InDateStyle, InTimeStyle, InCulture) FText::AsDateTime(FDateTime::FromUnixTimestamp(InUnixTime), InDateStyle, InTimeStyle, FText::GetInvariantTimeZone(), UE::Private::LocGen::GetCultureByName(TEXT(InCulture)))
 
  /**
   * Generate an FText representation of the given timestamp as a date and time (alias for FText::AsDateTime).
@@ -375,8 +375,8 @@ namespace UE4LocGen_Private
   * @param InTimeZone	The timezone to display the timestamp in.
   * @param InCulture		The culture code to use, or an empty string to use the active locale.
   */
-#define LOCGEN_DATETIME_CUSTOM_UTC(InUnixTime, InPattern, InTimeZone, InCulture) FText::AsDateTime(FDateTime::FromUnixTimestamp(InUnixTime), TEXT(InPattern), TEXT(InTimeZone), UE4LocGen_Private::GetCultureImpl(TEXT(InCulture)))
-#define LOCGEN_DATETIME_CUSTOM_LOCAL(InUnixTime, InPattern, InCulture) FText::AsDateTime(FDateTime::FromUnixTimestamp(InUnixTime), TEXT(InPattern), FText::GetInvariantTimeZone(), UE4LocGen_Private::GetCultureImpl(TEXT(InCulture)))
+#define LOCGEN_DATETIME_CUSTOM_UTC(InUnixTime, InPattern, InTimeZone, InCulture) FText::AsDateTime(FDateTime::FromUnixTimestamp(InUnixTime), TEXT(InPattern), TEXT(InTimeZone), UE::Private::LocGen::GetCultureByName(TEXT(InCulture)))
+#define LOCGEN_DATETIME_CUSTOM_LOCAL(InUnixTime, InPattern, InCulture) FText::AsDateTime(FDateTime::FromUnixTimestamp(InUnixTime), TEXT(InPattern), FText::GetInvariantTimeZone(), UE::Private::LocGen::GetCultureByName(TEXT(InCulture)))
 
 /**
  * Generate an FText representation of the given FText when transformed into upper-case (alias for FText::ToUpper).
