Update SharedPointer.h documentation comment

Removing a comment that said we don't support implicit conversion to bool of TSharedPtr/TSharedRef, but we do for TSharedPtr and it is not very useful for TSharedRef, which makes the comment unnecessary.

#rnx

diff --git a/Engine/Source/Runtime/Core/Public/Templates/SharedPointer.h b/Engine/Source/Runtime/Core/Public/Templates/SharedPointer.h
--- a/Engine/Source/Runtime/Core/Public/Templates/SharedPointer.h
+++ b/Engine/Source/Runtime/Core/Public/Templates/SharedPointer.h
@@ -83,7 +83,6 @@
  *		- Shared pointers are not compatible with Unreal objects (UObject classes)!
  *		- Currently only types with that have regular destructors (no custom deleters)
  *		- Dynamically-allocated arrays are not supported yet (e.g. MakeShareable( new int32[20] ))
- *		- Implicit conversion of TSharedPtr/TSharedRef to bool is not supported yet
  *
  *
  *	Differences from other implementations (e.g. boost:shared_ptr, std::shared_ptr):
