[CIS] Fix FUnixSignalGameHitchHeartBeat added mising IsStartedSuspended().

#rnx
#tests builds

diff --git a/Engine/Source/Runtime/Core/Private/Unix/UnixSignalHeartBeat.cpp b/Engine/Source/Runtime/Core/Private/Unix/UnixSignalHeartBeat.cpp
--- a/Engine/Source/Runtime/Core/Private/Unix/UnixSignalHeartBeat.cpp
+++ b/Engine/Source/Runtime/Core/Private/Unix/UnixSignalHeartBeat.cpp
@@ -153,7 +153,7 @@ void FUnixSignalGameHitchHeartBeat::InitSettings()
 	{
 		bFirst = false;
 
-		bool bStartSuspended = false;
+		bStartSuspended = false;
 		GConfig->GetBool(TEXT("Core.System"), TEXT("GameThreadHeartBeatStartSuspended"), bStartSuspended, GEngineIni);
 
 		if (FParse::Param(FCommandLine::Get(), TEXT("hitchdetectionstartsuspended")))
@@ -274,6 +274,11 @@ void FUnixSignalGameHitchHeartBeat::ResumeHeartBeat()
 #endif
 }
 
+bool FUnixSignalGameHitchHeartBeat::IsStartedSuspended()
+{
+	return bStartSuspended;
+}
+
 void FUnixSignalGameHitchHeartBeat::Restart()
 {
 	UE_LOG(LogUnixHeartBeat, Verbose, TEXT("Restart"));
diff --git a/Engine/Source/Runtime/Core/Public/Unix/UnixSignalHeartBeat.h b/Engine/Source/Runtime/Core/Public/Unix/UnixSignalHeartBeat.h
--- a/Engine/Source/Runtime/Core/Public/Unix/UnixSignalHeartBeat.h
+++ b/Engine/Source/Runtime/Core/Public/Unix/UnixSignalHeartBeat.h
@@ -33,6 +33,11 @@ public:
 	*/
 	CORE_API void ResumeHeartBeat();
 
+	/**
+	* Check if started suspended.
+	*/
+	CORE_API bool IsStartedSuspended();
+
 	CORE_API void Restart();
 	CORE_API void Stop();
 
@@ -51,5 +56,6 @@ private:
 	bool bDisabled = false;
 	int32 SuspendCount = 0;
 	bool bTimerCreated = false;
+	bool bStartSuspended = false;
 	timer_t TimerId = 0;
 };
