Reduce code and debug symbols bloat caused by ensure macros

diff --git a/Engine/Source/Runtime/Core/Private/Misc/AssertionMacros.cpp b/Engine/Source/Runtime/Core/Private/Misc/AssertionMacros.cpp
--- a/Engine/Source/Runtime/Core/Private/Misc/AssertionMacros.cpp
+++ b/Engine/Source/Runtime/Core/Private/Misc/AssertionMacros.cpp
@@ -750,13 +750,9 @@ FORCENOINLINE void FDebug::DumpStackTraceToLog(const TCHAR* Heading, const ELogV
 }
 
 #if DO_ENSURE && !USING_CODE_ANALYSIS
-bool UE_DEBUG_SECTION VARARGS CheckVerifyImpl(bool& InOutExecuted, bool Always, const ANSICHAR* File, int32 Line, void* ProgramCounter, const ANSICHAR* Expr, const TCHAR* Format, ...)
+bool UE_DEBUG_SECTION VARARGS CheckVerifyImpl(std::atomic<bool>& bExecuted, bool bAlways, const ANSICHAR* File, int32 Line, void* ProgramCounter, const ANSICHAR* Expr, const TCHAR* Format, va_list Args)
 {
-	InOutExecuted = true;
-	va_list Args;
-	va_start(Args, Format);
 	FDebug::OptionallyLogFormattedEnsureMessageReturningFalse(true, Expr, File, Line, ProgramCounter, Format, Args);
-	va_end(Args);
 
 	if (!FPlatformMisc::IsDebuggerPresent())
 	{
@@ -770,6 +766,37 @@ bool UE_DEBUG_SECTION VARARGS CheckVerifyImpl(bool& InOutExecuted, bool Always,
 	return !GIgnoreDebugger;
 #endif
 }
+
+bool UE_DEBUG_SECTION UE::Assert::Private::ExecCheckImplInternal(std::atomic<bool>& bExecuted, bool bAlways, const ANSICHAR* File, int32 Line, const ANSICHAR* Expr)
+{
+	if ((bAlways || !bExecuted.load(std::memory_order_relaxed)) && FPlatformMisc::IsEnsureAllowed())
+	{
+		if (bExecuted.exchange(true, std::memory_order_release) && !bAlways)
+		{
+			return false;
+		}
+
+		va_list Args = {};
+		return CheckVerifyImpl(bExecuted, bAlways, File, Line, PLATFORM_RETURN_ADDRESS(), Expr, TEXT(""), Args);
+	}
+
+	return false;
+}
+
+bool UE_DEBUG_SECTION VARARGS UE::Assert::Private::EnsureFailed(const FStaticEnsureRecord* Ensure, ...)
+{
+	if (Ensure->bExecuted.exchange(true, std::memory_order_release) && !Ensure->bAlways)
+	{
+		return false;
+	}
+
+	va_list Args;
+	va_start(Args, Ensure);
+	const bool bResult = CheckVerifyImpl(Ensure->bExecuted, Ensure->bAlways, Ensure->File, Ensure->Line, PLATFORM_RETURN_ADDRESS(), Ensure->Expression, Ensure->Format, Args);
+	va_end(Args);
+
+	return bResult;
+}
 #endif
 
 #undef FILE_LINE_DESC_ANSI
diff --git a/Engine/Source/Runtime/Core/Public/Misc/AssertionMacros.h b/Engine/Source/Runtime/Core/Public/Misc/AssertionMacros.h
--- a/Engine/Source/Runtime/Core/Public/Misc/AssertionMacros.h
+++ b/Engine/Source/Runtime/Core/Public/Misc/AssertionMacros.h
@@ -13,6 +13,8 @@
 #include "Misc/VarArgs.h"
 #include "String/FormatStringSan.h"
 
+#include <atomic>
+
 #ifndef UE_DEBUG_SECTION
 #if (DO_CHECK || DO_GUARD_SLOW || DO_ENSURE) && !PLATFORM_CPU_ARM_FAMILY
 	// We'll put all assert implementation code into a separate section in the linked
@@ -356,7 +358,51 @@ RetType FORCENOINLINE UE_DEBUG_SECTION DispatchCheckVerify(InnerType&& Inner, Ar
  */
 
 #if DO_ENSURE && !USING_CODE_ANALYSIS // The Visual Studio 2013 analyzer doesn't understand these complex conditionals
-	struct FValidateArgsInternal
+
+	namespace UE::Assert::Private
+	{
+
+	/** Data about an ensure that is constant for every occurrence. */
+	struct FStaticEnsureRecord
+	{
+		const TCHAR* Format = nullptr;
+		const ANSICHAR* Expression = nullptr;
+		const ANSICHAR* File = nullptr;
+		int32 Line = 0;
+		bool bAlways = false;
+		std::atomic<bool>& bExecuted;
+
+		// Workaround for https://developercommunity.visualstudio.com/t/Incorrect-warning-C4700-with-unrelated-s/10285950
+		constexpr FStaticEnsureRecord(
+			const TCHAR* InFormat,
+			const ANSICHAR* InExpression,
+			const ANSICHAR* InFile,
+			int32 InLine,
+			bool bInAlways,
+			std::atomic<bool>& InExecuted)
+			: Format(InFormat)
+			, Expression(InExpression)
+			, File(InFile)
+			, Line(InLine)
+			, bAlways(bInAlways)
+			, bExecuted(InExecuted)
+		{
+		}
+	};
+
+	CORE_API bool UE_DEBUG_SECTION VARARGS EnsureFailed(const FStaticEnsureRecord* Ensure, ...);
+	
+	CORE_API bool UE_DEBUG_SECTION ExecCheckImplInternal(std::atomic<bool>& bExecuted, bool bAlways, const ANSICHAR* File, int32 Line, const ANSICHAR* Expr);
+
+	FORCEINLINE bool DebugBreak()
+	{
+		PLATFORM_BREAK();
+		return false;
+	}
+
+	} // UE::Assert::Private
+
+	struct UE_DEPRECATED(5.4, "Do not use directly. This internal type is being removed.") FValidateArgsInternal
 	{
 		template <typename... Types>
 		FValidateArgsInternal(Types... Args)
@@ -365,26 +411,23 @@ RetType FORCENOINLINE UE_DEBUG_SECTION DispatchCheckVerify(InnerType&& Inner, Ar
 		}
 	};
 
-	CORE_API bool UE_DEBUG_SECTION VARARGS CheckVerifyImpl(bool& InOutExecuted, bool Always, const ANSICHAR* File, int32 Line, void* ProgramCounter, const ANSICHAR* Expr, const TCHAR* Format, ...);
-
-	#define UE_ENSURE_IMPL(Capture, Always, InExpression, InFormat) \
-		(LIKELY(!!(InExpression)) || (DispatchCheckVerify<bool>([Capture] () UE_DEBUG_SECTION \
-		{ \
-			static bool bExecuted = false; \
-			return (!bExecuted || Always) && FPlatformMisc::IsEnsureAllowed() && CheckVerifyImpl(bExecuted, Always, __FILE__, __LINE__, PLATFORM_RETURN_ADDRESS(), #InExpression, InFormat); \
-		}) && [] () { PLATFORM_BREAK(); return false; } ()))
+	#define UE_ENSURE_IMPL(Always, InExpression) \
+		(LIKELY(!!(InExpression)) \
+			|| (::UE::Assert::Private::ExecCheckImplInternal([]() UE_DEBUG_SECTION -> std::atomic<bool>& { static std::atomic<bool> bExecuted = false; return bExecuted; } (), Always, __FILE__, __LINE__, #InExpression) \
+			&& ::UE::Assert::Private::DebugBreak()))
 
 	#define UE_ENSURE_IMPL2(Capture, Always, InExpression, InFormat, ...) \
-		(LIKELY(!!(InExpression)) || (DispatchCheckVerify<bool>([Capture] () UE_DEBUG_SECTION \
+		(LIKELY(!!(InExpression)) || ([Capture] () UE_DEBUG_SECTION \
 		{ \
-			static bool bExecuted = false; \
 			UE_VALIDATE_FORMAT_STRING(InFormat, ##__VA_ARGS__); \
-			return (!bExecuted || Always) && FPlatformMisc::IsEnsureAllowed() && CheckVerifyImpl(bExecuted, Always, __FILE__, __LINE__, PLATFORM_RETURN_ADDRESS(), #InExpression, InFormat, ##__VA_ARGS__); \
-		}) && [] () { PLATFORM_BREAK(); return false; } ()))
+			static std::atomic<bool> bExecuted = false; \
+			static constexpr ::UE::Assert::Private::FStaticEnsureRecord ENSURE_Static(InFormat, #InExpression, __builtin_FILE(), __builtin_LINE(), Always, bExecuted); \
+			return (Always || !bExecuted.load(std::memory_order_relaxed)) && FPlatformMisc::IsEnsureAllowed() && ::UE::Assert::Private::EnsureFailed(&ENSURE_Static, ##__VA_ARGS__); \
+		}() && ::UE::Assert::Private::DebugBreak()))
 
-	#define ensure(           InExpression                ) UE_ENSURE_IMPL ( , false, InExpression, TEXT(""))
+	#define ensure(           InExpression                ) UE_ENSURE_IMPL (   false, InExpression)
 	#define ensureMsgf(       InExpression, InFormat, ... ) UE_ENSURE_IMPL2(&, false, InExpression, InFormat, ##__VA_ARGS__)
-	#define ensureAlways(     InExpression                ) UE_ENSURE_IMPL ( , true,  InExpression, TEXT(""))
+	#define ensureAlways(     InExpression                ) UE_ENSURE_IMPL (   true,  InExpression)
 	#define ensureAlwaysMsgf( InExpression, InFormat, ... ) UE_ENSURE_IMPL2(&, true,  InExpression, InFormat, ##__VA_ARGS__)
 
 #else	// DO_ENSURE
