[Insights] Fixed noname timing regions (i.e. SlowTask without a default message). If the default message for FSlowTask is empty, the timing region will be named "<SlowTask>".

diff --git a/Engine/Source/Runtime/Core/Private/Misc/SlowTask.cpp b/Engine/Source/Runtime/Core/Private/Misc/SlowTask.cpp
--- a/Engine/Source/Runtime/Core/Private/Misc/SlowTask.cpp
+++ b/Engine/Source/Runtime/Core/Private/Misc/SlowTask.cpp
@@ -19,7 +19,7 @@ static int32 GSlowTaskMaxTraceRegionDepth = 2;
 static FAutoConsoleVariableRef CVarSlowTaskMaxTraceRegionDepth (
 	TEXT("Trace.SlowTaskMaxRegionDepth"),
 	GSlowTaskMaxTraceRegionDepth,
-	TEXT("Maximum depth of nested slow tasks to create as trace regions in insights"),
+	TEXT("Maximum depth of nested slow tasks to create as trace regions in Insights"),
 	ECVF_Default
 );
 
@@ -102,7 +102,14 @@ void FSlowTask::Initialize()
 		Context.ScopeStack.Push(this);
 		if (Context.ScopeStack.Num() <= GSlowTaskMaxTraceRegionDepth)
 		{
-			TRACE_BEGIN_REGION(*DefaultMessage.ToString());
+			if (!DefaultMessage.IsEmpty())
+			{
+				TRACE_BEGIN_REGION(*DefaultMessage.ToString());
+			}
+			else
+			{
+				TRACE_BEGIN_REGION(TEXT("<SlowTask>"));
+			}
 		}
 	}
 }
@@ -134,7 +141,14 @@ void FSlowTask::Destroy()
 		{
 			if (Context.ScopeStack.Num() <= GSlowTaskMaxTraceRegionDepth)
 			{
-				TRACE_END_REGION(*DefaultMessage.ToString());
+				if (!DefaultMessage.IsEmpty())
+				{
+					TRACE_END_REGION(*DefaultMessage.ToString());
+				}
+				else
+				{
+					TRACE_END_REGION(TEXT("<SlowTask>"));
+				}
 			}
 
 			FSlowTask* Task = Stack.Last();
