[ParallelFor]
  - Reintroduce the incremental wake-up to avoid very small tasks to cost more in scheduling then they take to actually execute.
  - Do not init tasks until we launch them to avoid the need of cancellation.
  - Remove cancellation as its not needed anymore.
  - A simple ParallelFor(32, [](int32) {}, EParallelForFlags::Unbalanced); statement is now 6x faster on many-core systems

#tests TSAN / UnitTests / StressTests

diff --git a/Engine/Source/Runtime/Core/Public/Async/ParallelFor.h b/Engine/Source/Runtime/Core/Public/Async/ParallelFor.h
--- a/Engine/Source/Runtime/Core/Public/Async/ParallelFor.h
+++ b/Engine/Source/Runtime/Core/Public/Async/ParallelFor.h
@@ -202,13 +202,15 @@ namespace ParallelForImpl
 		{
 			using UE::FInheritedContextBase::RestoreInheritedContext;
 
-			FParallelForData(int32 InNum, int32 InBatchSize, int32 InNumBatches, int32 InNumWorkers, const TArrayView<ContextType>& InContexts, const BodyType& InBody, FEventRef& InFinishedSignal)
-				: Num(InNum)
+			FParallelForData(const TCHAR* InDebugName, int32 InNum, int32 InBatchSize, int32 InNumBatches, int32 InNumWorkers, const TArrayView<ContextType>& InContexts, const BodyType& InBody, FEventRef& InFinishedSignal, LowLevelTasks::ETaskPriority InPriority)
+				: DebugName(InDebugName)
+				, Num(InNum)
 				, BatchSize(InBatchSize)
 				, NumBatches(InNumBatches)
 				, Contexts(InContexts)
 				, Body(InBody)
 				, FinishedSignal(InFinishedSignal)
+				, Priority(InPriority)
 			{
 				IncompleteBatches.store(NumBatches, std::memory_order_relaxed);
 				Tasks.AddDefaulted(InNumWorkers);
@@ -220,51 +222,57 @@ namespace ParallelForImpl
 			{
 				for (FTracedTask& Task : Tasks)
 				{
-					TaskTrace::Destroyed(Task.TraceId);
+					if (Task.TraceId != TaskTrace::InvalidId)
+					{
+						TaskTrace::Destroyed(Task.TraceId);
+					}
 				}
 			}
 
+			int32 GetNextWorkerIndexToLaunch()
+			{
+				const int32 WorkerIndex = LaunchedWorkers.fetch_add(1, std::memory_order_relaxed);
+				return WorkerIndex >= Tasks.Num() ? -1 : WorkerIndex;
+			}
+
+			const TCHAR* DebugName;
 			std::atomic_int BatchItem  { 0 };
 			std::atomic_int IncompleteBatches { 0 };
+			std::atomic_int LaunchedWorkers { 0 };
 			int32 Num;
 			int32 BatchSize;
 			int32 NumBatches;
 			const TArrayView<ContextType>& Contexts;
 			const BodyType& Body;
 			FEventRef& FinishedSignal;
+			LowLevelTasks::ETaskPriority Priority;
 
 			TArray<FTracedTask, TConcurrentLinearArrayAllocator<FTaskGraphBlockAllocationTag>> Tasks;
 		};
 		using FDataHandle = TRefCountPtr<FParallelForData>;
 
-		//each task has an executor.
-		class FParallelExecutor 
+		// Each task has an executor.
+		class FParallelExecutor
 		{
-			FDataHandle Data;
-			mutable int32 WorkerIndex;
-			LowLevelTasks::ETaskPriority Priority;
+			mutable FDataHandle Data;
+			int32 WorkerIndex;
+			mutable bool bReschedule = false;
 
 		public:
-			inline FParallelExecutor(FDataHandle&& InData, int32 InWorkerIndex, LowLevelTasks::ETaskPriority InPriority) 
+			inline FParallelExecutor(FDataHandle&& InData, int32 InWorkerIndex)
 				: Data(MoveTemp(InData))
 				, WorkerIndex(InWorkerIndex)
-				, Priority(InPriority)
 			{
 			}
 
 			FParallelExecutor(const FParallelExecutor&) = delete;
-			inline FParallelExecutor(FParallelExecutor&& Other) 
-				: Data(MoveTemp(Other.Data))
-				, WorkerIndex(Other.WorkerIndex)
-				, Priority(Other.Priority)
-			{
-			}
+			FParallelExecutor(FParallelExecutor&& Other) = default;
 
 			~FParallelExecutor()
 			{
-				if (Data.IsValid() && WorkerIndex >= 0)
+				if (Data.IsValid() && bReschedule)
 				{
-					FParallelExecutor::LaunchTask(nullptr, MoveTemp(Data), WorkerIndex, Priority);
+					FParallelExecutor::LaunchTask(MoveTemp(Data), WorkerIndex);
 				}
 			}
 
@@ -273,17 +281,11 @@ namespace ParallelForImpl
 				return Data;
 			}
 
-			inline bool operator()(const bool bIsMaster = false, const TCHAR* DebugName = nullptr) const noexcept
+			inline bool operator()(const bool bIsMaster = false) const noexcept
 			{
 				UE::FInheritedContextScope InheritedContextScope = Data->RestoreInheritedContext();
 				FMemMark Mark(FMemStack::Get());
 
-				if(DebugName == nullptr)
-				{
-					LowLevelTasks::FTask& Task = Data->Tasks[WorkerIndex].Task;
-					DebugName = Task.GetDebugName();
-				}
-				
 				TaskTrace::FId TraceId = TaskTrace::InvalidId;
 				if (!bIsMaster)
 				{
@@ -298,12 +300,21 @@ namespace ParallelForImpl
 					}
 				};
 
-				TRACE_CPUPROFILER_EVENT_SCOPE_TEXT(DebugName);
+				const int32 NumBatches = Data->NumBatches;
+
+				// We're going to consume one ourself, so we need at least 2 left to consider launching a new worker
+				// We also do not launch a worker from the master as we already launched one before doing prework.
+				if (bIsMaster == false && Data->BatchItem.load(std::memory_order_relaxed) + 2 <= NumBatches)
+				{
+					LaunchAnotherWorkerIfNeeded(Data);
+				}
+
+				TRACE_CPUPROFILER_EVENT_SCOPE_TEXT(Data->DebugName);
 
 				auto Now = [] { return FTimespan::FromSeconds(FPlatformTime::Seconds()); };
 				FTimespan Start = FTimespan::MinValue();
 				FTimespan YieldingThreshold;
-				const bool bIsBackgroundPriority = !bIsMaster && (Priority >= LowLevelTasks::ETaskPriority::BackgroundNormal);
+				const bool bIsBackgroundPriority = !bIsMaster && (Data->Priority >= LowLevelTasks::ETaskPriority::BackgroundNormal);
 
 				if (bIsBackgroundPriority)
 				{
@@ -313,7 +324,6 @@ namespace ParallelForImpl
 
 				const int32 Num = Data->Num;
 				const int32 BatchSize = Data->BatchSize;
-				const int32 NumBatches = Data->NumBatches;
 				const TArrayView<ContextType>& Contexts = Data->Contexts;
 				const BodyType& Body = Data->Body;
 
@@ -322,12 +332,11 @@ namespace ParallelForImpl
 				{
 					int32 BatchIndex = Data->BatchItem.fetch_add(1, std::memory_order_relaxed);
 					
-					//save the last block for the master to safe an event
+					// Save the last block for the master to avoid an event
 					if (bSaveLastBlockForMaster && BatchIndex >= NumBatches - 1)
 					{
 						if (!bIsMaster)
 						{
-							WorkerIndex = -1;
 							return false;
 						}
 						BatchIndex = (NumBatches - 1);
@@ -352,12 +361,11 @@ namespace ParallelForImpl
 						{
 							Data->FinishedSignal->Trigger();
 						}
-						WorkerIndex = -1;
+
 						return true;
 					}
 					else if (EndIndex >= Num)
 					{
-						WorkerIndex = -1;
 						return false;
 					}
 					else if (!bIsBackgroundPriority)
@@ -368,58 +376,57 @@ namespace ParallelForImpl
 					auto PassedTime = [Start, &Now]() { return Now() - Start; };
 					if (PassedTime() > YieldingThreshold)
 					{
-						//abort and reschedule (in the destructor as WorkerIndex is larger_eq Zero) to give higher priority tasks a chance to run
+						// Abort and reschedule to give higher priority tasks a chance to run
+						bReschedule = true;
 						return false;
 					}
 				}
 			}
 
-			static void LaunchTask(const TCHAR* DebugName, FDataHandle&& InData, int32 InWorkerIndex, LowLevelTasks::ETaskPriority InPriority)
+			static void LaunchTask(FDataHandle&& InData, int32 InWorkerIndex, bool bWakeUpWorker = true)
 			{
 				FTracedTask& TracedTask = InData->Tasks[InWorkerIndex];
-				if(DebugName == nullptr)
-				{
-					DebugName = TracedTask.Task.GetDebugName();
-				}
 
 				if (TracedTask.TraceId != TaskTrace::InvalidId) // reused task
 				{
 					TaskTrace::Destroyed(TracedTask.TraceId);
 				}
+
+				const TCHAR* DebugName = InData->DebugName;
+				LowLevelTasks::ETaskPriority Priority = InData->Priority;
+
 				TracedTask.TraceId = TaskTrace::GenerateTaskId();
 				TaskTrace::Launched(TracedTask.TraceId, DebugName, false, ENamedThreads::AnyThread, 0);
 
-				TracedTask.Task.Init(DebugName, InPriority, FParallelExecutor(MoveTemp(InData), InWorkerIndex, InPriority));
-				verify(LowLevelTasks::TryLaunch(TracedTask.Task, LowLevelTasks::EQueuePreference::GlobalQueuePreference));
+				TracedTask.Task.Init(DebugName, Priority, FParallelExecutor(MoveTemp(InData), InWorkerIndex));
+				verify(LowLevelTasks::TryLaunch(TracedTask.Task, LowLevelTasks::EQueuePreference::GlobalQueuePreference, bWakeUpWorker));
+			}
 
+			static void LaunchAnotherWorkerIfNeeded(FDataHandle& InData)
+			{
+				const int32 WorkerIndex = InData->GetNextWorkerIndexToLaunch();
+				if (WorkerIndex != -1)
+				{
+					LaunchTask(FDataHandle(InData), WorkerIndex);
+				}
 			}
 		};
 
 		//launch all the worker tasks
 		FEventRef FinishedSignal { EEventMode::ManualReset };
-		FDataHandle Data = new FParallelForData(Num, BatchSize, NumBatches, NumWorkers, Contexts, Body, FinishedSignal);
-		for (int32 Worker = 0; Worker < NumWorkers; Worker++)
-		{
-			FParallelExecutor::LaunchTask(DebugName, FDataHandle(Data), Worker, Priority);
-		}
+		FDataHandle Data = new FParallelForData(DebugName, Num, BatchSize, NumBatches, NumWorkers, Contexts, Body, FinishedSignal, Priority);
+
+		// Launch the first worker before we start doing prework
+		FParallelExecutor::LaunchAnotherWorkerIfNeeded(Data);
 
 		// do the prework
 		CurrentThreadWorkToDoBeforeHelping();
 
-		//help with the parallel-for to prevent deadlocks
-		FParallelExecutor LocalExecutor(MoveTemp(Data), NumWorkers, Priority);
-		const bool bFinishedLast = LocalExecutor(true, DebugName);
-
-		// Cancel tasks that have not yet launched since they will otherwise waste time on worker threads
-		{
-			TRACE_CPUPROFILER_EVENT_SCOPE(ParallelFor.Cancel);
-			for (FTracedTask& TracedTask : LocalExecutor.GetData()->Tasks)
-			{
-				TracedTask.Task.TryCancel(LowLevelTasks::ECancellationFlags::PrelaunchCancellation);
-			}
-		}
+		// help with the parallel-for to prevent deadlocks
+		FParallelExecutor LocalExecutor(MoveTemp(Data), NumWorkers);
+		const bool bFinishedLast = LocalExecutor(true);
 
-		if(!bFinishedLast)
+		if (!bFinishedLast)
 		{
 			const bool bPumpRenderingThread  = (Flags & EParallelForFlags::PumpRenderingThread) != EParallelForFlags::None;
 			if (bPumpRenderingThread && IsInActualRenderingThread())
