Enable C4702

diff --git a/Engine/Source/Runtime/Core/Public/Windows/WindowsPlatformCompilerSetup.h b/Engine/Source/Runtime/Core/Public/Windows/WindowsPlatformCompilerSetup.h
--- a/Engine/Source/Runtime/Core/Public/Windows/WindowsPlatformCompilerSetup.h
+++ b/Engine/Source/Runtime/Core/Public/Windows/WindowsPlatformCompilerSetup.h
@@ -176,7 +176,6 @@ static_assert(_MSC_VER >= 1929, "Visual Studio 2019 v16.11 or greater is require
 #pragma warning(disable: 4640) // 'instance' : construction of local static object is not thread-safe								https://docs.microsoft.com/en-us/cpp/error-messages/compiler-warnings/compiler-warning-level-3-c4640
 #pragma warning(disable: 4699) // creating precompiled header																		// No docs
 
-#pragma warning(disable: 4702) // unreachable code																					https://docs.microsoft.com/en-us/cpp/error-messages/compiler-warnings/compiler-warning-level-4-c4702
 #pragma warning(disable: 4710) // 'function' : function not inlined																	https://docs.microsoft.com/en-us/cpp/error-messages/compiler-warnings/compiler-warning-level-4-c4710
 #pragma warning(disable: 4711) // function selected for automatic inlining															https://docs.microsoft.com/en-us/cpp/error-messages/compiler-warnings/compiler-warning-level-1-c4711
 #pragma warning(disable: 4714) // function 'function' marked as __forceinline not inlined											https://docs.microsoft.com/en-us/cpp/error-messages/compiler-warnings/compiler-warning-level-4-c4714
