Fixed compilation in ensure macros on clang Win64 targets

diff --git a/Engine/Source/Runtime/Core/Private/Misc/AssertionMacros.cpp b/Engine/Source/Runtime/Core/Private/Misc/AssertionMacros.cpp
--- a/Engine/Source/Runtime/Core/Private/Misc/AssertionMacros.cpp
+++ b/Engine/Source/Runtime/Core/Private/Misc/AssertionMacros.cpp
@@ -783,16 +783,16 @@ bool UE_DEBUG_SECTION UE::Assert::Private::ExecCheckImplInternal(std::atomic<boo
 	return false;
 }
 
-bool UE_DEBUG_SECTION VARARGS UE::Assert::Private::EnsureFailed(const FStaticEnsureRecord* Ensure, ...)
+bool UE_DEBUG_SECTION VARARGS UE::Assert::Private::EnsureFailed(std::atomic<bool>& bExecuted, const FStaticEnsureRecord* Ensure, ...)
 {
-	if (Ensure->bExecuted.exchange(true, std::memory_order_release) && !Ensure->bAlways)
+	if (bExecuted.exchange(true, std::memory_order_release) && !Ensure->bAlways)
 	{
 		return false;
 	}
 
 	va_list Args;
 	va_start(Args, Ensure);
-	const bool bResult = CheckVerifyImpl(Ensure->bExecuted, Ensure->bAlways, Ensure->File, Ensure->Line, PLATFORM_RETURN_ADDRESS(), Ensure->Expression, Ensure->Format, Args);
+	const bool bResult = CheckVerifyImpl(bExecuted, Ensure->bAlways, Ensure->File, Ensure->Line, PLATFORM_RETURN_ADDRESS(), Ensure->Expression, Ensure->Format, Args);
 	va_end(Args);
 
 	return bResult;
diff --git a/Engine/Source/Runtime/Core/Public/Misc/AssertionMacros.h b/Engine/Source/Runtime/Core/Public/Misc/AssertionMacros.h
--- a/Engine/Source/Runtime/Core/Public/Misc/AssertionMacros.h
+++ b/Engine/Source/Runtime/Core/Public/Misc/AssertionMacros.h
@@ -370,7 +370,6 @@ RetType FORCENOINLINE UE_DEBUG_SECTION DispatchCheckVerify(InnerType&& Inner, Ar
 		const ANSICHAR* File = nullptr;
 		int32 Line = 0;
 		bool bAlways = false;
-		std::atomic<bool>& bExecuted;
 
 		// Workaround for https://developercommunity.visualstudio.com/t/Incorrect-warning-C4700-with-unrelated-s/10285950
 		constexpr FStaticEnsureRecord(
@@ -378,19 +377,17 @@ RetType FORCENOINLINE UE_DEBUG_SECTION DispatchCheckVerify(InnerType&& Inner, Ar
 			const ANSICHAR* InExpression,
 			const ANSICHAR* InFile,
 			int32 InLine,
-			bool bInAlways,
-			std::atomic<bool>& InExecuted)
+			bool bInAlways)
 			: Format(InFormat)
 			, Expression(InExpression)
 			, File(InFile)
 			, Line(InLine)
 			, bAlways(bInAlways)
-			, bExecuted(InExecuted)
 		{
 		}
 	};
 
-	CORE_API bool UE_DEBUG_SECTION VARARGS EnsureFailed(const FStaticEnsureRecord* Ensure, ...);
+	CORE_API bool UE_DEBUG_SECTION VARARGS EnsureFailed(std::atomic<bool>& bExecuted, const FStaticEnsureRecord* Ensure, ...);
 	
 	CORE_API bool UE_DEBUG_SECTION ExecCheckImplInternal(std::atomic<bool>& bExecuted, bool bAlways, const ANSICHAR* File, int32 Line, const ANSICHAR* Expr);
 
@@ -421,8 +418,8 @@ RetType FORCENOINLINE UE_DEBUG_SECTION DispatchCheckVerify(InnerType&& Inner, Ar
 		{ \
 			UE_VALIDATE_FORMAT_STRING(InFormat, ##__VA_ARGS__); \
 			static std::atomic<bool> bExecuted = false; \
-			static constexpr ::UE::Assert::Private::FStaticEnsureRecord ENSURE_Static(InFormat, #InExpression, __builtin_FILE(), __builtin_LINE(), Always, bExecuted); \
-			return (Always || !bExecuted.load(std::memory_order_relaxed)) && FPlatformMisc::IsEnsureAllowed() && ::UE::Assert::Private::EnsureFailed(&ENSURE_Static, ##__VA_ARGS__); \
+			static constexpr ::UE::Assert::Private::FStaticEnsureRecord ENSURE_Static(InFormat, #InExpression, __builtin_FILE(), __builtin_LINE(), Always); \
+			return (Always || !bExecuted.load(std::memory_order_relaxed)) && FPlatformMisc::IsEnsureAllowed() && ::UE::Assert::Private::EnsureFailed(bExecuted, &ENSURE_Static, ##__VA_ARGS__); \
 		}() && ::UE::Assert::Private::DebugBreak()))
 
 	#define ensure(           InExpression                ) UE_ENSURE_IMPL (   false, InExpression)
