Fixes IK Rig rotation engine test. Issues caused by floating point discrepency in Rotator/Quat conversion between Linux/Windows platforms.

#virtualized

diff --git a/Engine/Source/Developer/FunctionalTesting/Classes/FunctionalTest.h b/Engine/Source/Developer/FunctionalTesting/Classes/FunctionalTest.h
--- a/Engine/Source/Developer/FunctionalTesting/Classes/FunctionalTest.h
+++ b/Engine/Source/Developer/FunctionalTesting/Classes/FunctionalTest.h
@@ -501,6 +501,13 @@ public:
 	UFUNCTION(BlueprintCallable, Category = "Asserts", DisplayName = "Assert Equal (Rotator)", meta = ( HidePin = "ContextObject", DefaultToSelf = "ContextObject"))
 	virtual bool AssertEqual_Rotator(FRotator Actual, FRotator Expected, const FString& What, float Tolerance = 1.e-4f, const UObject* ContextObject = nullptr);
 
+	/**
+	 * Assert that the orientation of two rotators is the same within a small tolerance. Robust to quaternion singularities where angles can differ despite having an identical orientation. 
+	 * @param What	A name to use in the message if the assert fails ("Expected 'What' to be {Expected} but it was {Actual} for context ''")
+	 */
+	UFUNCTION(BlueprintCallable, Category = "Asserts", DisplayName = "Assert Equal (Rotator Orientation)", meta = ( HidePin = "ContextObject", DefaultToSelf = "ContextObject"))
+	virtual bool AssertEqual_RotatorOrientation(FRotator Actual, FRotator Expected, const FString& What, float Tolerance = 1.e-4f, const UObject* ContextObject = nullptr);
+
 	/**
 	 * Assert that the component angles of two rotators are all not equal within a small tolerance.
 	 * @param What	A name to use in the message if the assert fails ("Expected 'What' not to be {Expected} but it was {Actual} for context ''")
diff --git a/Engine/Source/Developer/FunctionalTesting/Private/FunctionalTest.cpp b/Engine/Source/Developer/FunctionalTesting/Private/FunctionalTest.cpp
--- a/Engine/Source/Developer/FunctionalTesting/Private/FunctionalTest.cpp
+++ b/Engine/Source/Developer/FunctionalTesting/Private/FunctionalTest.cpp
@@ -878,6 +878,20 @@ bool AFunctionalTest::AssertEqual_Rotator(FRotator Actual, FRotator Expected, co
 	}
 }
 
+bool AFunctionalTest::AssertEqual_RotatorOrientation(FRotator Actual, FRotator Expected, const FString& What, float Tolerance, const UObject* ContextObject)
+{
+	if ( !Expected.EqualsOrientation(Actual, Tolerance) )
+	{
+		LogStep(ELogVerbosity::Error, FString::Printf(TEXT("Expected '%s' to be {%s} but it was {%s} within tolerance {%f} for context '%s'"), *What, *Expected.ToString(), *Actual.ToString(), Tolerance, ContextObject ? *ContextObject->GetName() : TEXT("")));
+		return false;
+	}
+	else
+	{
+		LogStep(ELogVerbosity::Log, FString::Printf(TEXT("Rotator assertion passed (%s)"), *What));
+		return true;
+	}
+}
+
 bool AFunctionalTest::AssertNotEqual_Rotator(FRotator Actual, FRotator NotExpected, const FString& What, const UObject* ContextObject)
 {
 	if ( NotExpected.Equals(Actual) )
diff --git a/Engine/Source/Runtime/Core/Public/Math/Rotator.h b/Engine/Source/Runtime/Core/Public/Math/Rotator.h
--- a/Engine/Source/Runtime/Core/Public/Math/Rotator.h
+++ b/Engine/Source/Runtime/Core/Public/Math/Rotator.h
@@ -232,6 +232,20 @@ public:
 	 */
 	bool Equals( const TRotator<T>& R, T Tolerance = UE_KINDA_SMALL_NUMBER ) const;
 
+	/**
+	 * Checks whether two rotators have the same orientation within the specified tolerance, without requiring similar angles or wound multiples of 360.
+	 * Unlike Equals(), it can compare Rotators that represent the same final orientation, but get there via different intermediate rotations.
+	 * i.e. If we compare two rotators with different angles, but same orientation like so:
+	 * A = TRotator(0, 45, 0)
+	 * B = TRotator(180, 135, 180)
+	 * Then A.EqualsOrientation(B) would be TRUE while A.Equals(B) would be FALSE
+	 *
+	 * @param R The other rotator.
+	 * @param Tolerance Error Tolerance.
+	 * @return true if two rotators are equal, within specified tolerance, otherwise false.
+	 */
+	bool EqualsOrientation( const TRotator<T>& R, T Tolerance = UE_KINDA_SMALL_NUMBER ) const;
+
 	/**
 	 * Adds to each component of the rotator.
 	 *
@@ -626,6 +640,12 @@ FORCEINLINE bool TRotator<T>::Equals(const TRotator<T>& R, T Tolerance) const
 #endif
 }
 
+template <typename T>
+bool TRotator<T>::EqualsOrientation(const TRotator<T>& R, T Tolerance) const
+{
+	return Quaternion().AngularDistance(R.Quaternion()) <= Tolerance;
+}
+
 template<typename T>
 FORCEINLINE TRotator<T> TRotator<T>::Add( T DeltaPitch, T DeltaYaw, T DeltaRoll )
 {
