autortfm: redirect VerseStm OnCommit/OnRollback callbacks into AutoRTFM. Includes a fix for an edge case where we execute Verse during the commit handler of a transaction, this makes IsTransactional return false during that case

diff --git a/Engine/Source/Runtime/Core/Private/AutoRTFM/Context.cpp b/Engine/Source/Runtime/Core/Private/AutoRTFM/Context.cpp
--- a/Engine/Source/Runtime/Core/Private/AutoRTFM/Context.cpp
+++ b/Engine/Source/Runtime/Core/Private/AutoRTFM/Context.cpp
@@ -53,7 +53,7 @@ bool FContext::IsTransactional()
         return false;
     }
 
-    if (Context->GetStatus() != EContextStatus::Idle)
+    if ((Context->GetStatus() != EContextStatus::Idle) && (Context->GetStatus() != EContextStatus::Committing))
     {
         return true;
     }
