Add  TIsContiguousContainer<TBasicArray>

#rnx

diff --git a/Engine/Source/Runtime/Core/Public/Containers/BasicArray.h b/Engine/Source/Runtime/Core/Public/Containers/BasicArray.h
--- a/Engine/Source/Runtime/Core/Public/Containers/BasicArray.h
+++ b/Engine/Source/Runtime/Core/Public/Containers/BasicArray.h
@@ -5,6 +5,7 @@
 // HEADER_UNIT_SKIP - Not used?
 
 #include "Templates/MemoryOps.h"
+#include "Traits/IsContiguousContainer.h"
 
 #if _MSC_VER
 	#pragma warning(push)
@@ -176,6 +177,12 @@ private:
 	friend const T* end  (const TBasicArray& Arr) { return Arr.GetData() + Arr.Num(); }
 };
 
+template <typename T>
+struct TIsContiguousContainer<TBasicArray<T>>
+{
+	enum { Value = true };
+};
+
 #if _MSC_VER
 	#pragma warning(pop)
 #endif
