Remove FParse::Resolution declarations from Parse.h. The function definitions were removed from Parse.cpp back in 2017 in CL 3575017. The code lives on in UnrealEngine.cpp as ParseResolution.

diff --git a/Engine/Source/Runtime/Core/Public/Misc/Parse.h b/Engine/Source/Runtime/Core/Public/Misc/Parse.h
--- a/Engine/Source/Runtime/Core/Public/Misc/Parse.h
+++ b/Engine/Source/Runtime/Core/Public/Misc/Parse.h
@@ -136,12 +136,6 @@ struct FParse
 	static CORE_API uint32 HexNumber(const TCHAR* HexString);
 	static CORE_API uint64 HexNumber64(const TCHAR* HexString);
 
-	/** Parses a resolution in the form 1920x1080. */
-	static CORE_API bool Resolution( const TCHAR* InResolution, uint32& OutX, uint32& OutY );
-
-	/** Parses a resolution in the form 1920x1080<f|w|wf>. Same as above, but also attempts to process a fullscreen/windowed flag from the end */
-	static CORE_API bool Resolution(const TCHAR* InResolution, uint32& OutX, uint32& OutY, int32& OutWindowMode);
-
 	/** Parses the scheme name from a URI */
 	static CORE_API bool SchemeNameFromURI(const TCHAR* InURI, FString& OutSchemeName);
 
