Added floating point unit tests for FMath::Wrap.
Added a new FMath::Modulo which generically encapsulates integral/float differences.

diff --git a/Engine/Source/Runtime/Core/Public/Math/UnrealMathUtility.h b/Engine/Source/Runtime/Core/Public/Math/UnrealMathUtility.h
--- a/Engine/Source/Runtime/Core/Public/Math/UnrealMathUtility.h
+++ b/Engine/Source/Runtime/Core/Public/Math/UnrealMathUtility.h
@@ -609,6 +609,8 @@ public:
 	template< class T >
 	UE_NODISCARD static constexpr FORCEINLINE T Wrap(const T X, const T Min, const T Max)
 	{
+		checkSlow(Min <= Max);
+
 		T Size = Max - Min;
 		if (Size == 0)
 		{
@@ -629,6 +631,19 @@ public:
 		return EndVal;
 	}
 
+	template <typename T>
+	UE_NODISCARD static constexpr FORCEINLINE T Modulo(T Value, T Base)
+	{
+		if constexpr (std::is_floating_point_v<T>)
+		{
+			return FMath::Fmod(Value, Base);
+		}
+		else
+		{
+			return Value % Base;
+		}
+	}
+
 	/** Snaps a value to the nearest grid multiple */
 	template< class T >
 	UE_NODISCARD static constexpr FORCEINLINE T GridSnap(T Location, T Grid)
diff --git a/Engine/Source/Runtime/Core/Tests/Math/UnrealMathTest.cpp b/Engine/Source/Runtime/Core/Tests/Math/UnrealMathTest.cpp
--- a/Engine/Source/Runtime/Core/Tests/Math/UnrealMathTest.cpp
+++ b/Engine/Source/Runtime/Core/Tests/Math/UnrealMathTest.cpp
@@ -4491,36 +4491,131 @@ TEST_CASE_NAMED(FBitCastTest, "System::Core::Math::Bitcast", "[ApplicationContex
 
 TEST_CASE_NAMED(FMathWrapTest, "System::Core::Math::Wrap", "[ApplicationContextMask][SmokeFilter]")
 {
-	for (int Val = -5; Val != 5; ++Val)
+	// Tests wrapping of FMath::Wrap(Val, Min, Max) with a set of values in a set of ranges.
+	//
+	// The values to test are of the form ValFrom + X*ValStep, and do not exceed ValTo.
+	// The minimum values in the wrapping range are of the form: MinFrom + Y*MinStep, and do not exceed MinTo.
+	// The sizes of each range (where Max == Min + Size) are of the form: SizeFrom + Z*SizeStep, and do not exceed SizeTo.
+	auto WrapTest = []<typename T>(T ValFrom, T ValTo, T ValStep, T MinFrom, T MinTo, T MinStep, T SizeFrom, T SizeTo, T SizeStep)
 	{
-		for (int Min = -5; Min != 5; ++Min)
+		for (T Val = ValFrom; Val < ValTo; Val += ValStep)
 		{
-			// Size == 0
+			for (T Min = MinFrom; Min < MinTo; Min += MinStep)
 			{
-				int Wrap = FMath::Wrap(Val, Min, Min);
-
-				CHECK_MESSAGE(TEXT("Wrapped value should be in the empty range"), Wrap == Min);
-			}
+				for (T Size = SizeFrom; Size < SizeTo; Size += SizeStep)
+				{
+					if (Size == (T)0)
+					{
+						T Wrap = FMath::Wrap(Val, Min, Min);
 
-			for (int Size = 1; Size != 5; ++Size)
-			{
-				int Max = Min + Size;
-				int Wrap = FMath::Wrap(Val, Min, Max);
+						CHECK_MESSAGE(TEXT("Wrapped value should be in the empty range"), Wrap == Min);
+					}
+					else
+					{
+						T Max = Min + Size;
+						T Wrap = FMath::Wrap(Val, Min, Max);
 
-				CHECK_MESSAGE(TEXT("Wrapped value should be in the non-empty range"), Wrap >= Min && Wrap <= Max);
-				CHECK_MESSAGE(FString::Printf(TEXT("Wrapped value should be at a distance which is an exact multiple of the range size: (Val: %d, Min: %d, Max: %d, Wrap: %d, Mod: %d)"), Val, Min, Max, Wrap, (Wrap - Val) % Size), (Wrap - Val) % Size == 0);
+						CHECK_MESSAGE(TEXT("Wrapped value should be in the non-empty range"), Wrap >= Min && Wrap <= Max);
+						T Mod = FMath::Modulo((Wrap - Val), Size);
+						if constexpr (std::is_integral_v<T>)
+						{
+							CHECK_MESSAGE(FString::Printf(TEXT("Wrapped value should be at a distance which is an exact multiple of the range size: (Val: %d, Min: %d, Max: %d, Wrap: %d, Mod: %d)"), Val, Min, Max, Wrap, Mod), Mod == 0);
+						}
+						else
+						{
+							T Tolerance;
+							if constexpr (std::is_same_v<T, float>)
+							{
+								Tolerance = UE_KINDA_SMALL_NUMBER;
+							}
+							else
+							{
+								Tolerance = UE_DOUBLE_KINDA_SMALL_NUMBER;
+							}
+
+							// We need to check that we're in the range of zero *or* +/- size because of rounding
+							bool bIsExactMultipleOfSize = FMath::Square(Mod) < Tolerance || FMath::Square(Mod - Size) < Tolerance || FMath::Square(Mod + Size) < Tolerance;
+							CHECK_MESSAGE(FString::Printf(TEXT("Wrapped value should be at a distance which is an exact multiple of the range size: (Val: %f, Min: %f, Max: %f, Wrap: %f, Mod: %f)"), Val, Min, Max, Wrap, Mod), bIsExactMultipleOfSize);
+						}
 
-				if (Val < Min)
-				{
-					CHECK_FALSE_MESSAGE(TEXT("Wrapping a value from below a non-empty range should never give the max"), Wrap == Max);
-				}
-				else if (Val > Max)
-				{
-					CHECK_FALSE_MESSAGE(TEXT("Wrapping a value from above a non-empty range should never give the min"), Wrap == Min);
+						if (Val < Min)
+						{
+							CHECK_FALSE_MESSAGE(TEXT("Wrapping a value from below a non-empty range should never give the max"), Wrap == Max);
+						}
+						else if (Val > Max)
+						{
+							CHECK_FALSE_MESSAGE(TEXT("Wrapping a value from above a non-empty range should never give the min"), Wrap == Min);
+						}
+					}
 				}
 			}
 		}
-	}
+	};
+
+	// Integral
+	WrapTest(
+		/* ValFrom  */ -25,
+		/* ValTo    */  25,
+		/* ValStep  */   1,
+		/* MinFrom  */  -5,
+		/* MinTo    */   5,
+		/* MinStep  */   1,
+		/* SizeFrom */   0,
+		/* SizeTo   */   5,
+		/* SizeStep */   1
+	);
+
+	// Floats (with integral values)
+	WrapTest(
+		/* ValFrom  */ -25.0f,
+		/* ValTo    */  25.0f,
+		/* ValStep  */   1.0f,
+		/* MinFrom  */  -5.0f,
+		/* MinTo    */   5.0f,
+		/* MinStep  */   1.0f,
+		/* SizeFrom */   0.0f,
+		/* SizeTo   */   5.0f,
+		/* SizeStep */   1.0f
+	);
+
+	// Floats (with fractional values)
+	WrapTest(
+		/* ValFrom  */  -7.34f,
+		/* ValTo    */  12.19f,
+		/* ValStep  */   0.7361f,
+		/* MinFrom  */  -8.43f,
+		/* MinTo    */  11.84f,
+		/* MinStep  */   0.69f,
+		/* SizeFrom */   0.0f,
+		/* SizeTo   */   7.23f,
+		/* SizeStep */   0.59f
+	);
+
+	// Doubles (with integral values)
+	WrapTest(
+		/* ValFrom  */ -25.0,
+		/* ValTo    */  25.0,
+		/* ValStep  */   1.0,
+		/* MinFrom  */  -5.0,
+		/* MinTo    */   5.0,
+		/* MinStep  */   1.0,
+		/* SizeFrom */   0.0,
+		/* SizeTo   */   5.0,
+		/* SizeStep */   1.0
+	);
+
+	// Doubles (with fractional values)
+	WrapTest(
+		/* ValFrom  */  -7.34,
+		/* ValTo    */  12.19,
+		/* ValStep  */   0.7361,
+		/* MinFrom  */  -8.43,
+		/* MinTo    */  11.84,
+		/* MinStep  */   0.69,
+		/* SizeFrom */   0.0,
+		/* SizeTo   */   7.23,
+		/* SizeStep */   0.59
+	);
 }
 class FInitVectorTestClass {
 public:
