Fix for compile error in FMath::Wrap due to using a check() in a constexpr context.

diff --git a/Engine/Source/Runtime/Core/Public/Math/UnrealMathUtility.h b/Engine/Source/Runtime/Core/Public/Math/UnrealMathUtility.h
--- a/Engine/Source/Runtime/Core/Public/Math/UnrealMathUtility.h
+++ b/Engine/Source/Runtime/Core/Public/Math/UnrealMathUtility.h
@@ -609,7 +609,8 @@ public:
 	template< class T >
 	UE_NODISCARD static constexpr FORCEINLINE T Wrap(const T X, const T Min, const T Max)
 	{
-		checkSlow(Min <= Max);
+		// Our asserts are not constexpr-friendly yet
+		// checkSlow(Min <= Max);
 
 		T Size = Max - Min;
 		if (Size == 0)
