[CIS] Lost file for RedirectToFile engine command

#tests none
#rnx

diff --git a/Engine/Source/Runtime/Core/Private/HAL/ConsoleManager.cpp b/Engine/Source/Runtime/Core/Private/HAL/ConsoleManager.cpp
--- a/Engine/Source/Runtime/Core/Private/HAL/ConsoleManager.cpp
+++ b/Engine/Source/Runtime/Core/Private/HAL/ConsoleManager.cpp
@@ -3249,6 +3249,12 @@ void CreateConsoleVariables()
 
 #if	!UE_BUILD_SHIPPING
 	IConsoleManager::Get().RegisterConsoleCommand( TEXT( "DumpConsoleCommands" ), TEXT( "Dumps all console vaiables and commands and all exec that can be discovered to the log/console" ), ECVF_Default );
+	IConsoleManager::Get().RegisterConsoleCommand( TEXT( "RedirectToFile" ),
+		TEXT( "Creates a file inside Project's Saved folder and outputs command result into it as well as into the log.\n" )
+		TEXT( "Usage: RedirectToFile <filepath/filename> <command> [command arguments]\n" )
+		TEXT( "Example: RedirectToFile Profiling/CSV/objlist.csv obj list -csv -all\n" )
+		TEXT( "Directory structure under Project/Saved folder specified by <filepath> will be created for you if it doesn't exist." ),
+		ECVF_Default );
 
 	IConsoleManager::Get().RegisterConsoleCommand(TEXT("DumpCVars"),
 		TEXT("Lists all CVars (or a subset) and their values. Can also show help, and can save to .csv.\nUsage: DumpCVars [Prefix] [-showhelp] [-csv=[path]]\nIf -csv does not have a file specified, it will create a file in the Project Logs directory"),
