Properly initialize implicit tests with engine-only dependency

diff --git a/Engine/Source/Programs/UnrealBuildTool/Configuration/TestTargetRules.cs b/Engine/Source/Programs/UnrealBuildTool/Configuration/TestTargetRules.cs
--- a/Engine/Source/Programs/UnrealBuildTool/Configuration/TestTargetRules.cs
+++ b/Engine/Source/Programs/UnrealBuildTool/Configuration/TestTargetRules.cs
@@ -359,6 +359,10 @@ namespace UnrealBuildTool
 			bool IsEditorTestedTarget = (TestedTarget.Type == TargetType.Editor);
 			bBuildWithEditorOnlyData = bCompileAgainstEditor = IsEditorTestedTarget;
 			bBuildDeveloperTools = bCompileAgainstEngine && IsEditorTestedTarget;
+			
+			bUsePlatformFileStub = bCompileAgainstEngine;
+			bMockEngineDefaults = bCompileAgainstEngine;
+			bCompileWithPluginSupport = bCompileAgainstEngine;
 		}
 	}
 }
diff --git a/Engine/Source/Runtime/Core/Public/Misc/LowLevelTestAdapter.h b/Engine/Source/Runtime/Core/Public/Misc/LowLevelTestAdapter.h
--- a/Engine/Source/Runtime/Core/Public/Misc/LowLevelTestAdapter.h
+++ b/Engine/Source/Runtime/Core/Public/Misc/LowLevelTestAdapter.h
@@ -108,7 +108,6 @@ FString CaptureExpressionsAndValues(const FString& InExpressions, ArgTypes&&...
 #define TEST_CASE(PrettyName, TFlags) TEST_CASE_NAMED_STR(TEST_CASE_GENERATED_NAME_UNIQUE, TEST_CASE_GENERATED_NAME_UNIQUE_STR, PrettyName, TFlags)
 #define TEST_CASE_NAMED(ClassName, PrettyName, TFlags) TEST_CASE_NAMED_STR(ClassName, #ClassName, PrettyName, TFlags)
 
-#define INFO(What)
 //-V:CHECK:571,501,547
 #define CHECK(Expr) if (!(Expr)) { FAutomationTestFramework::Get().GetCurrentTest()->AddError(TEXT("Condition failed")); }
 //-V:CHECK_FALSE:571,501,547
@@ -134,6 +133,7 @@ FString CaptureExpressionsAndValues(const FString& InExpressions, ArgTypes&&...
 #define FAIL_CHECK(Message) FAutomationTestFramework::Get().GetCurrentTest()->AddError(Message);
 
 #define CAPTURE(...) FAutomationTestFramework::Get().GetCurrentTest()->AddInfo(CaptureExpressionsAndValues(#__VA_ARGS__, __VA_ARGS__));
+#define INFO(Message) FAutomationTestFramework::Get().GetCurrentTest()->AddInfo(Message)
 #define WARN(Message) FAutomationTestFramework::Get().GetCurrentTest()->AddWarning(Message); 
 #define FAIL_ON_MESSAGE(Message) FAutomationTestFramework::Get().GetCurrentTest()->AddExpectedError(Message); 
 
diff --git a/Engine/Source/Runtime/TypedElementFramework/Tests/Elements/Framework/TypedElementQueryBuilderTests.cpp b/Engine/Source/Runtime/TypedElementFramework/Tests/Elements/Framework/TypedElementQueryBuilderTests.cpp
--- a/Engine/Source/Runtime/TypedElementFramework/Tests/Elements/Framework/TypedElementQueryBuilderTests.cpp
+++ b/Engine/Source/Runtime/TypedElementFramework/Tests/Elements/Framework/TypedElementQueryBuilderTests.cpp
@@ -95,7 +95,7 @@ static bool TestMatching(const FQueryConditions& TestQuery, const TArray<FColumn
 		}
 	}
 
-	FAutomationTestFramework::Get().GetCurrentTest()->AddInfo(MoveTemp(Description));
+	INFO(MoveTemp(Description));
 	return (Result == Expected);
 }
 
