PlatformFileJournal: Fix fallback case for for FileJournalIterateDirectory to use the IterateDirectoryStat from the current PlatformFile rather than calling on the highest-level IPlatformFile from IPlatformFileManager.
#rnx

diff --git a/Engine/Source/Runtime/Core/Private/GenericPlatform/GenericPlatformFile.cpp b/Engine/Source/Runtime/Core/Private/GenericPlatform/GenericPlatformFile.cpp
--- a/Engine/Source/Runtime/Core/Private/GenericPlatform/GenericPlatformFile.cpp
+++ b/Engine/Source/Runtime/Core/Private/GenericPlatform/GenericPlatformFile.cpp
@@ -1034,7 +1034,7 @@ EFileJournalResult IPlatformFile::FileJournalGetLatestEntry(const TCHAR* VolumeO
 
 bool IPlatformFile::FileJournalIterateDirectory(const TCHAR* Directory, FDirectoryJournalVisitorFunc Visitor)
 {
-	return IFileManager::Get().IterateDirectoryStat(Directory,
+	return this->IterateDirectoryStat(Directory,
 		[&Visitor](const TCHAR* InPackageFilename, const FFileStatData& StatData)
 		{
 			return Visitor(InPackageFilename, UE::PlatformFileJournal::Private::ToJournalData(StatData));
@@ -1043,7 +1043,7 @@ bool IPlatformFile::FileJournalIterateDirectory(const TCHAR* Directory, FDirecto
 
 FFileJournalData IPlatformFile::FileJournalGetFileData(const TCHAR* FilenameOrDirectory)
 {
-	return UE::PlatformFileJournal::Private::ToJournalData(IFileManager::Get().GetStatData(FilenameOrDirectory));
+	return UE::PlatformFileJournal::Private::ToJournalData(this->GetStatData(FilenameOrDirectory));
 }
 
 EFileJournalResult IPlatformFile::FileJournalReadModified(const TCHAR* VolumeName,
