[Insights] CpuProfiler:
  - Flushed the thread buffer (if not already empty) when the thread ends.
  - Added timestamp ("Cycle" field) to the "CpuProfiler.EndThread" event.
  - Removed unused "CpuProfiler.EndCaptureV2" event.

diff --git a/Engine/Source/Runtime/Core/Private/ProfilingDebugging/CpuProfilerTrace.cpp b/Engine/Source/Runtime/Core/Private/ProfilingDebugging/CpuProfilerTrace.cpp
--- a/Engine/Source/Runtime/Core/Private/ProfilingDebugging/CpuProfilerTrace.cpp
+++ b/Engine/Source/Runtime/Core/Private/ProfilingDebugging/CpuProfilerTrace.cpp
@@ -33,11 +33,8 @@ UE_TRACE_EVENT_BEGIN(CpuProfiler, EventBatchV2, NoSync)
 	UE_TRACE_EVENT_FIELD(uint8[], Data)
 UE_TRACE_EVENT_END()
 
-UE_TRACE_EVENT_BEGIN(CpuProfiler, EndCaptureV2)
-	UE_TRACE_EVENT_FIELD(uint8[], Data)
-UE_TRACE_EVENT_END()
-
 UE_TRACE_EVENT_BEGIN(CpuProfiler, EndThread, NoSync)
+	UE_TRACE_EVENT_FIELD(uint64, Cycle) // added in UE 5.4
 UE_TRACE_EVENT_END()
 
 struct FCpuProfilerTraceInternal
@@ -88,7 +85,12 @@ struct FCpuProfilerTraceInternal
 
 		virtual ~FThreadBuffer()
 		{
-			UE_TRACE_LOG(CpuProfiler, EndThread, CpuChannel);
+			if (BufferSize > 0)
+			{
+				FCpuProfilerTraceInternal::FlushThreadBuffer(this);
+			}
+			UE_TRACE_LOG(CpuProfiler, EndThread, CpuChannel)
+				<< EndThread.Cycle(FPlatformTime::Cycles64());
 			// Clear the thread buffer pointer. In the rare event of there being scopes in the destructors of other
 			// FTLSAutoCleanup instances. In that case a new buffer is created for that event only. There is no way of
 			// controlling the order of destruction for FTLSAutoCleanup types.
@@ -107,7 +109,6 @@ struct FCpuProfilerTraceInternal
 	uint32 static GetNextSpecId();
 	FORCENOINLINE static FThreadBuffer* CreateThreadBuffer();
 	FORCENOINLINE static void FlushThreadBuffer(FThreadBuffer* ThreadBuffer);
-	FORCENOINLINE static void EndCapture(FThreadBuffer* ThreadBuffer);
 
 	struct FSuspendScopes
 	{
@@ -139,14 +140,6 @@ void FCpuProfilerTraceInternal::FlushThreadBuffer(FThreadBuffer* InThreadBuffer)
 	InThreadBuffer->LastCycle = 0;
 }
 
-void FCpuProfilerTraceInternal::EndCapture(FThreadBuffer* InThreadBuffer)
-{
-	UE_TRACE_LOG(CpuProfiler, EndCaptureV2, true)
-		<< EndCaptureV2.Data(InThreadBuffer->Buffer, InThreadBuffer->BufferSize);
-	InThreadBuffer->BufferSize = 0;
-	InThreadBuffer->LastCycle = 0;
-}
-
 #define CPUPROFILERTRACE_OUTPUTBEGINEVENT_PROLOGUE() \
 	++FCpuProfilerTraceInternal::ThreadDepth; \
 	FCpuProfilerTraceInternal::FThreadBuffer* ThreadBuffer = FCpuProfilerTraceInternal::ThreadBuffer; \
