IAS - move encryption key manager to core/misc

diff --git a/Engine/Source/Runtime/Core/Internal/Misc/EncryptionKeyManager.h b/Engine/Source/Runtime/Core/Internal/Misc/EncryptionKeyManager.h
new file mode 100644
--- /dev/null
+++ b/Engine/Source/Runtime/Core/Internal/Misc/EncryptionKeyManager.h
@@ -0,0 +1,49 @@
+// Copyright Epic Games, Inc. All Rights Reserved.
+
+#pragma once
+
+#include "Async/Mutex.h"
+#include "Async/UniqueLock.h"
+#include "Containers/Map.h"
+#include "Delegates/Delegate.h"
+#include "Misc/AES.h"
+#include "Misc/Guid.h"
+
+namespace UE
+{
+
+/** Manages a set of registered encryption key(s). */
+class CORE_API FEncryptionKeyManager
+{
+public:
+	FEncryptionKeyManager(const FEncryptionKeyManager&) = delete;
+	FEncryptionKeyManager(FEncryptionKeyManager&&) = delete;
+	FEncryptionKeyManager& operator=(const FEncryptionKeyManager&) = delete;
+	FEncryptionKeyManager& operator=(FEncryptionKeyManager&&) = delete;
+	~FEncryptionKeyManager();
+
+	/** Returns whether the specified encrypton key exist or not. */
+	bool ContainsKey(const FGuid& Id);
+	/** Add a new encryption key, ignored if the key already exist. */
+	void AddKey(const FGuid& Id, const FAES::FAESKey& Key);
+	/** Try retrieve the encryption key for the specified key ID. */
+	bool TryGetKey(const FGuid& Id, FAES::FAESKey& OutKey);
+	/** Returns a map of all available keys */
+	TMap<FGuid, FAES::FAESKey> GetAllKeys();
+
+	DECLARE_MULTICAST_DELEGATE_TwoParams(FEncryptionKeyAddedDelegate, const FGuid&, const FAES::FAESKey&);
+	/** Event triggered when a new key as been added. */
+	FEncryptionKeyAddedDelegate& OnKeyAdded() { return KeyAdded; }
+	/** Returns the single instance of the key manager. */
+	static FEncryptionKeyManager& Get();
+
+private:
+	FEncryptionKeyManager();
+	FAES::FAESKey* GetKey(const FGuid& Id);
+
+	FMutex Mutex;
+	TMap<FGuid, FAES::FAESKey> Keys;
+	FEncryptionKeyAddedDelegate KeyAdded;
+};
+
+} // namespace UE
diff --git a/Engine/Source/Runtime/Core/Private/Misc/EncryptionKeyManager.cpp b/Engine/Source/Runtime/Core/Private/Misc/EncryptionKeyManager.cpp
new file mode 100644
--- /dev/null
+++ b/Engine/Source/Runtime/Core/Private/Misc/EncryptionKeyManager.cpp
@@ -0,0 +1,84 @@
+// Copyright Epic Games, Inc. All Rights Reserved.
+
+#include "Misc/EncryptionKeyManager.h"
+#include "Async/UniqueLock.h"
+#include "Misc/CoreDelegates.h"
+
+namespace UE
+{
+
+FEncryptionKeyManager::FEncryptionKeyManager()
+{
+	if (FCoreDelegates::GetPakEncryptionKeyDelegate().IsBound())
+	{
+		FAES::FAESKey Key;
+		FCoreDelegates::GetPakEncryptionKeyDelegate().Execute(Key.Key);
+
+		if (Key.IsValid())
+		{
+			AddKey(FGuid(), Key);
+		}
+	}
+
+	FCoreDelegates::GetRegisterEncryptionKeyMulticastDelegate().AddRaw(this, &FEncryptionKeyManager::AddKey);
+}
+
+FEncryptionKeyManager::~FEncryptionKeyManager()
+{
+	FCoreDelegates::GetRegisterEncryptionKeyMulticastDelegate().RemoveAll(this);
+}
+
+bool FEncryptionKeyManager::ContainsKey(const FGuid& Id)
+{
+	return nullptr != GetKey(Id);
+}
+
+void FEncryptionKeyManager::AddKey(const FGuid& Id, const FAES::FAESKey& Key)
+{
+	bool bAdded = false;
+	{
+		TUniqueLock Lock(Mutex);
+
+		if (!Keys.Contains(Id))
+		{
+			Keys.Add(Id, Key);
+			bAdded = true;
+		}
+	}
+
+	if (bAdded && KeyAdded.IsBound())
+	{
+		KeyAdded.Broadcast(Id, Key);
+	}
+}
+
+bool FEncryptionKeyManager::TryGetKey(const FGuid& Id, FAES::FAESKey& OutKey)
+{
+	if (FAES::FAESKey* Key = GetKey(Id))
+	{
+		OutKey = *Key;
+		return true;
+	}
+
+	return false;
+}
+
+TMap<FGuid, FAES::FAESKey> FEncryptionKeyManager::GetAllKeys()
+{
+	TUniqueLock Lock(Mutex);
+	return Keys;
+}
+
+FAES::FAESKey* FEncryptionKeyManager::GetKey(const FGuid& Id)
+{
+	TUniqueLock Lock(Mutex);
+	return Keys.Find(Id);
+}
+
+FEncryptionKeyManager& FEncryptionKeyManager::Get()
+{
+	static FEncryptionKeyManager Mgr;
+	return Mgr;
+}
+
+} // namespace UE
diff --git a/Engine/Source/Runtime/Experimental/IoStoreOnDemand/Private/EncryptionKeyManager.cpp b/Engine/Source/Runtime/Experimental/IoStoreOnDemand/Private/EncryptionKeyManager.cpp
deleted file mode 100644
diff --git a/Engine/Source/Runtime/Experimental/IoStoreOnDemand/Private/EncryptionKeyManager.h b/Engine/Source/Runtime/Experimental/IoStoreOnDemand/Private/EncryptionKeyManager.h
deleted file mode 100644
diff --git a/Engine/Source/Runtime/Experimental/IoStoreOnDemand/Private/IoStoreOnDemand.cpp b/Engine/Source/Runtime/Experimental/IoStoreOnDemand/Private/IoStoreOnDemand.cpp
--- a/Engine/Source/Runtime/Experimental/IoStoreOnDemand/Private/IoStoreOnDemand.cpp
+++ b/Engine/Source/Runtime/Experimental/IoStoreOnDemand/Private/IoStoreOnDemand.cpp
@@ -2,7 +2,6 @@
 
 #include "IO/IoStoreOnDemand.h"
 
-#include "EncryptionKeyManager.h"
 #include "HAL/FileManager.h"
 #include "HAL/LowLevelMemTracker.h"
 #include "HAL/PlatformMisc.h"
@@ -15,6 +14,7 @@
 #include "Misc/ConfigCacheIni.h"
 #include "Misc/CoreDelegates.h"
 #include "Misc/DateTime.h"
+#include "Misc/EncryptionKeyManager.h"
 #include "Misc/FileHelper.h"
 #include "Misc/Parse.h"
 #include "Misc/PathViews.h"
diff --git a/Engine/Source/Runtime/Experimental/IoStoreOnDemand/Private/OnDemandIoDispatcherBackend.cpp b/Engine/Source/Runtime/Experimental/IoStoreOnDemand/Private/OnDemandIoDispatcherBackend.cpp
--- a/Engine/Source/Runtime/Experimental/IoStoreOnDemand/Private/OnDemandIoDispatcherBackend.cpp
+++ b/Engine/Source/Runtime/Experimental/IoStoreOnDemand/Private/OnDemandIoDispatcherBackend.cpp
@@ -7,7 +7,6 @@
 #include "Containers/StringView.h"
 #include "CoreHttp/LatencyTesting.h"
 #include "DistributionEndpoints.h"
-#include "EncryptionKeyManager.h"
 #include "GenericPlatform/GenericPlatformCrashContext.h"
 #include "HAL/Event.h"
 #include "HAL/FileManagerGeneric.h"
@@ -30,6 +29,7 @@
 #include "Math/NumericLimits.h"
 #include "Misc/CommandLine.h"
 #include "Misc/ConfigCacheIni.h"
+#include "Misc/EncryptionKeyManager.h"
 #include "Misc/EnumClassFlags.h"
 #include "Misc/PathViews.h"
 #include "Misc/Paths.h"
@@ -534,10 +534,9 @@ void FOnDemandIoStore::AddDeferredContainers()
 		{
 			FGuid KeyGuid;
 			ensure(FGuid::Parse(Container->EncryptionKeyGuid, KeyGuid));
-			if (const FAES::FAESKey* Key = FEncryptionKeyManager::Get().GetKey(KeyGuid))
+			if (FEncryptionKeyManager::Get().TryGetKey(KeyGuid, Container->EncryptionKey))
 			{
 				UE_LOG(LogIas, Log, TEXT("Mounting container '%s' (%d entries)"), *Container->Name, Container->TocEntries.Num());
-				Container->EncryptionKey = *Key;
 				RegisteredContainers.Add(Container);
 				It.RemoveCurrent();
 			}
