[UE] Add CVar to disable ensureAlways

#tests Validated I hit an ensure always and it only fired once

diff --git a/Engine/Source/Runtime/Core/Private/Misc/AssertionMacros.cpp b/Engine/Source/Runtime/Core/Private/Misc/AssertionMacros.cpp
--- a/Engine/Source/Runtime/Core/Private/Misc/AssertionMacros.cpp
+++ b/Engine/Source/Runtime/Core/Private/Misc/AssertionMacros.cpp
@@ -112,6 +112,13 @@ FAutoConsoleVariableRef CVarEnsuresAreErrors(
 	TEXT("True means failed ensures are logged as errors. False means they are logged as warnings."),
 	ECVF_Default);
 
+bool GEnsureAlwaysEnabled = true;
+FAutoConsoleVariableRef CVarEnsureAlwaysEnabled(
+	TEXT("core.EnsureAlwaysEnabled"),
+	GEnsureAlwaysEnabled,
+	TEXT("Set to false to turn ensureAlways into regular ensure"),
+	ECVF_Default);
+
 
 CORE_API void (*GPrintScriptCallStackFn)() = nullptr;
 
@@ -769,7 +776,7 @@ bool UE_DEBUG_SECTION VARARGS CheckVerifyImpl(std::atomic<bool>& bExecuted, bool
 
 bool UE_DEBUG_SECTION UE::Assert::Private::ExecCheckImplInternal(std::atomic<bool>& bExecuted, bool bAlways, const ANSICHAR* File, int32 Line, const ANSICHAR* Expr)
 {
-	if ((bAlways || !bExecuted.load(std::memory_order_relaxed)) && FPlatformMisc::IsEnsureAllowed())
+	if (((bAlways && GEnsureAlwaysEnabled) || !bExecuted.load(std::memory_order_relaxed)) && FPlatformMisc::IsEnsureAllowed())
 	{
 		if (bExecuted.exchange(true, std::memory_order_release) && !bAlways)
 		{
@@ -785,7 +792,7 @@ bool UE_DEBUG_SECTION UE::Assert::Private::ExecCheckImplInternal(std::atomic<boo
 
 bool UE_DEBUG_SECTION VARARGS UE::Assert::Private::EnsureFailed(std::atomic<bool>& bExecuted, const FStaticEnsureRecord* Ensure, ...)
 {
-	if (bExecuted.exchange(true, std::memory_order_release) && !Ensure->bAlways)
+	if (bExecuted.exchange(true, std::memory_order_release) && !(Ensure->bAlways && GEnsureAlwaysEnabled))
 	{
 		return false;
 	}
