In automated perf tests that capture csvs, include the device ip in the additional csv metadata

diff --git a/Engine/Source/Runtime/Core/Private/GenericPlatform/GenericPlatformMisc.cpp b/Engine/Source/Runtime/Core/Private/GenericPlatform/GenericPlatformMisc.cpp
--- a/Engine/Source/Runtime/Core/Private/GenericPlatform/GenericPlatformMisc.cpp
+++ b/Engine/Source/Runtime/Core/Private/GenericPlatform/GenericPlatformMisc.cpp
@@ -1816,6 +1816,14 @@ FGuid FGenericPlatformMisc::GetMachineId()
 	return MachineId;
 }
 
+FString FGenericPlatformMisc::GetDeviceTag()
+	{
+		FString DeviceTag = TEXT("");
+		FParse::Value( FCommandLine::Get(), TEXT("DeviceTag="), DeviceTag );
+
+		return DeviceTag;
+	}
+
 FString FGenericPlatformMisc::GetLoginId()
 {
 	PRAGMA_DISABLE_DEPRECATION_WARNINGS
diff --git a/Engine/Source/Runtime/Core/Private/ProfilingDebugging/CsvProfiler.cpp b/Engine/Source/Runtime/Core/Private/ProfilingDebugging/CsvProfiler.cpp
--- a/Engine/Source/Runtime/Core/Private/ProfilingDebugging/CsvProfiler.cpp
+++ b/Engine/Source/Runtime/Core/Private/ProfilingDebugging/CsvProfiler.cpp
@@ -2809,7 +2809,12 @@ FCsvProfiler::FCsvProfiler()
 	// for privacy, personal and free text fields are not allowed in shipping
 	SetMetadataInternal(TEXT("Commandline"), *CommandlineStr, false);
 	SetMetadataInternal(TEXT("LoginID"), *FPlatformMisc::GetLoginId());
-
+	FString DeviceTag = FPlatformMisc::GetDeviceTag();
+    if (!DeviceTag.IsEmpty())
+    {
+    	SetMetadataInternal(TEXT("DeviceTag"), *DeviceTag);
+    }
+	
 	// Set the device ID if the platform supports it
 	FString DeviceID = FPlatformMisc::GetDeviceId();
 	if (!DeviceID.IsEmpty())
diff --git a/Engine/Source/Runtime/Core/Public/GenericPlatform/GenericPlatformMisc.h b/Engine/Source/Runtime/Core/Public/GenericPlatform/GenericPlatformMisc.h
--- a/Engine/Source/Runtime/Core/Public/GenericPlatform/GenericPlatformMisc.h
+++ b/Engine/Source/Runtime/Core/Public/GenericPlatform/GenericPlatformMisc.h
@@ -1673,6 +1673,15 @@ public:
 	 */
 	static CORE_API FString GetEpicAccountId();
 
+	/**
+	 * Get Device Tag of the device running the session.
+	 *
+	 * Note: Gets from DeviceTag command line argument 
+	 *
+	 * @return empty string or the value from DeviceTag argument.
+	 */
+	static CORE_API FString GetDeviceTag();
+
 	/**
 	 * Gets a globally unique ID the represents a particular operating system install.
 	 * @returns an opaque string representing the ID, or an empty string if the platform doesn't support one.
