Added missing const specifier to FCsvDeclaredStat's name param

diff --git a/Engine/Source/Runtime/Core/Public/ProfilingDebugging/CsvProfiler.h b/Engine/Source/Runtime/Core/Public/ProfilingDebugging/CsvProfiler.h
--- a/Engine/Source/Runtime/Core/Public/ProfilingDebugging/CsvProfiler.h
+++ b/Engine/Source/Runtime/Core/Public/ProfilingDebugging/CsvProfiler.h
@@ -174,7 +174,7 @@ struct FCsvCategory;
 
 struct FCsvDeclaredStat
 {
-	FCsvDeclaredStat(TCHAR* InNameString, uint32 InCategoryIndex) 
+	FCsvDeclaredStat(const TCHAR* InNameString, uint32 InCategoryIndex) 
 		: Name(InNameString)
 		, CategoryIndex(InCategoryIndex) 
 	{
