Const correct and more thread safty encryption key manager

diff --git a/Engine/Source/Runtime/Core/Internal/Misc/EncryptionKeyManager.h b/Engine/Source/Runtime/Core/Internal/Misc/EncryptionKeyManager.h
--- a/Engine/Source/Runtime/Core/Internal/Misc/EncryptionKeyManager.h
+++ b/Engine/Source/Runtime/Core/Internal/Misc/EncryptionKeyManager.h
@@ -23,13 +23,13 @@ public:
 	~FEncryptionKeyManager();
 
 	/** Returns whether the specified encrypton key exist or not. */
-	bool ContainsKey(const FGuid& Id);
+	bool ContainsKey(const FGuid& Id) const;
 	/** Add a new encryption key, ignored if the key already exist. */
 	void AddKey(const FGuid& Id, const FAES::FAESKey& Key);
 	/** Try retrieve the encryption key for the specified key ID. */
-	bool TryGetKey(const FGuid& Id, FAES::FAESKey& OutKey);
+	bool TryGetKey(const FGuid& Id, FAES::FAESKey& OutKey) const;
 	/** Returns a map of all available keys */
-	TMap<FGuid, FAES::FAESKey> GetAllKeys();
+	TMap<FGuid, FAES::FAESKey> GetAllKeys() const;
 
 	DECLARE_MULTICAST_DELEGATE_TwoParams(FEncryptionKeyAddedDelegate, const FGuid&, const FAES::FAESKey&);
 	/** Event triggered when a new key as been added. */
@@ -39,9 +39,8 @@ public:
 
 private:
 	FEncryptionKeyManager();
-	FAES::FAESKey* GetKey(const FGuid& Id);
 
-	FMutex Mutex;
+	mutable FMutex Mutex;
 	TMap<FGuid, FAES::FAESKey> Keys;
 	FEncryptionKeyAddedDelegate KeyAdded;
 };
diff --git a/Engine/Source/Runtime/Core/Private/Misc/EncryptionKeyManager.cpp b/Engine/Source/Runtime/Core/Private/Misc/EncryptionKeyManager.cpp
--- a/Engine/Source/Runtime/Core/Private/Misc/EncryptionKeyManager.cpp
+++ b/Engine/Source/Runtime/Core/Private/Misc/EncryptionKeyManager.cpp
@@ -28,9 +28,10 @@ FEncryptionKeyManager::~FEncryptionKeyManager()
 	FCoreDelegates::GetRegisterEncryptionKeyMulticastDelegate().RemoveAll(this);
 }
 
-bool FEncryptionKeyManager::ContainsKey(const FGuid& Id)
+bool FEncryptionKeyManager::ContainsKey(const FGuid& Id) const
 {
-	return nullptr != GetKey(Id);
+	TUniqueLock Lock(Mutex);
+	return Keys.Contains(Id);
 }
 
 void FEncryptionKeyManager::AddKey(const FGuid& Id, const FAES::FAESKey& Key)
@@ -52,9 +53,10 @@ void FEncryptionKeyManager::AddKey(const FGuid& Id, const FAES::FAESKey& Key)
 	}
 }
 
-bool FEncryptionKeyManager::TryGetKey(const FGuid& Id, FAES::FAESKey& OutKey)
+bool FEncryptionKeyManager::TryGetKey(const FGuid& Id, FAES::FAESKey& OutKey) const
 {
-	if (FAES::FAESKey* Key = GetKey(Id))
+	TUniqueLock Lock(Mutex);
+	if (const FAES::FAESKey* Key = Keys.Find(Id))
 	{
 		OutKey = *Key;
 		return true;
@@ -63,18 +65,12 @@ bool FEncryptionKeyManager::TryGetKey(const FGuid& Id, FAES::FAESKey& OutKey)
 	return false;
 }
 
-TMap<FGuid, FAES::FAESKey> FEncryptionKeyManager::GetAllKeys()
+TMap<FGuid, FAES::FAESKey> FEncryptionKeyManager::GetAllKeys() const
 {
 	TUniqueLock Lock(Mutex);
 	return Keys;
 }
 
-FAES::FAESKey* FEncryptionKeyManager::GetKey(const FGuid& Id)
-{
-	TUniqueLock Lock(Mutex);
-	return Keys.Find(Id);
-}
-
 FEncryptionKeyManager& FEncryptionKeyManager::Get()
 {
 	static FEncryptionKeyManager Mgr;
