[TaskGraph]
  - Fix new frontend to send the first subsequent to the local queue like the old frontend
     - Add a benchmark for chained tasks
     - TestPerfChaining is 2.5x faster with this change and is now on par with the old frontend
  - Fix compilation errors when TASKGRAPH_NEW_FRONTEND=1
  - Add verbose Insight tracing in new frontend

diff --git a/Engine/Source/Runtime/Core/Private/Tasks/TaskPrivate.cpp b/Engine/Source/Runtime/Core/Private/Tasks/TaskPrivate.cpp
--- a/Engine/Source/Runtime/Core/Private/Tasks/TaskPrivate.cpp
+++ b/Engine/Source/Runtime/Core/Private/Tasks/TaskPrivate.cpp
@@ -12,7 +12,7 @@ namespace UE::Tasks
 		FExecutableTaskAllocator SmallTaskAllocator;
 		FTaskEventBaseAllocator TaskEventBaseAllocator;
 
-		void FTaskBase::Schedule()
+		void FTaskBase::Schedule(bool& bWakeUpWorker)
 		{
 			TaskTrace::Scheduled(GetTraceId());
 
@@ -42,10 +42,10 @@ namespace UE::Tasks
 			}
 #endif
 
-			LowLevelTasks::FScheduler::Get().TryLaunch(LowLevelTask, LowLevelTasks::EQueuePreference::GlobalQueuePreference, /*bWakeUpWorker=*/ true);
+			bWakeUpWorker |= LowLevelTasks::FSchedulerTls::IsBusyWaiting();
+			bWakeUpWorker |= LowLevelTasks::FScheduler::Get().TryLaunch(LowLevelTask, bWakeUpWorker ? LowLevelTasks::EQueuePreference::GlobalQueuePreference : LowLevelTasks::EQueuePreference::LocalQueuePreference, bWakeUpWorker);
 		}
 
-
 		thread_local uint32 TaskRetractionRecursion = 0;
 
 		bool IsThreadRetractingTask()
@@ -70,7 +70,7 @@ namespace UE::Tasks
 
 		bool FTaskBase::TryRetractAndExecute(FTimeout Timeout, uint32 RecursionDepth/* = 0*/)
 		{
-			TRACE_CPUPROFILER_EVENT_SCOPE(TaskRetraction);
+			TRACE_CPUPROFILER_EVENT_SCOPE(FTaskBase::TryRetractAndExecute);
 
 			if (!IsAwaitable())
 			{
@@ -279,7 +279,7 @@ namespace UE::Tasks
 
 				auto TaskBody = [CurrentThread, &TaskGraph] { TaskGraph.RequestReturn(CurrentThread); };
 				using FReturnFromNamedThreadTask = TExecutableTask<decltype(TaskBody)>;
-				FReturnFromNamedThreadTask ReturnTask{ TEXT("ReturnFromNamedThreadTask"), MoveTemp(TaskBody), ETaskPriority::High, ExtendedPriority };
+				FReturnFromNamedThreadTask ReturnTask { TEXT("ReturnFromNamedThreadTask"), MoveTemp(TaskBody), ETaskPriority::High, ExtendedPriority, ETaskFlags::None };
 				ReturnTask.AddPrerequisites(Task);
 				ReturnTask.TryLaunch(sizeof(ReturnTask)); // the result doesn't matter
 
diff --git a/Engine/Source/Runtime/Core/Public/Async/TaskGraphInterfaces.h b/Engine/Source/Runtime/Core/Public/Async/TaskGraphInterfaces.h
--- a/Engine/Source/Runtime/Core/Public/Async/TaskGraphInterfaces.h
+++ b/Engine/Source/Runtime/Core/Public/Async/TaskGraphInterfaces.h
@@ -477,9 +477,9 @@ public:
 		}
 	}
 
-	void Init(const TCHAR* InDebugName, UE::Tasks::ETaskPriority InPriority, UE::Tasks::EExtendedTaskPriority InExtendedPriority)
+	void Init(const TCHAR* InDebugName, UE::Tasks::ETaskPriority InPriority, UE::Tasks::EExtendedTaskPriority InExtendedPriority, UE::Tasks::ETaskFlags InTaskFlags = UE::Tasks::ETaskFlags::None)
 	{
-		FTaskBase::Init(InDebugName, InPriority, InExtendedPriority);
+		FTaskBase::Init(InDebugName, InPriority, InExtendedPriority, InTaskFlags);
 	}
 
 	void Unlock(ENamedThreads::Type CurrentThreadIfKnown = ENamedThreads::AnyThread)
diff --git a/Engine/Source/Runtime/Core/Public/Tasks/TaskPrivate.h b/Engine/Source/Runtime/Core/Public/Tasks/TaskPrivate.h
--- a/Engine/Source/Runtime/Core/Public/Tasks/TaskPrivate.h
+++ b/Engine/Source/Runtime/Core/Public/Tasks/TaskPrivate.h
@@ -36,6 +36,16 @@
 #include <atomic>
 #include <type_traits>
 
+#ifndef WITH_TASKGRAPH_VERBOSE_TRACE
+#define WITH_TASKGRAPH_VERBOSE_TRACE 0
+#endif
+
+#if WITH_TASKGRAPH_VERBOSE_TRACE
+#define TASKGRAPH_VERBOSE_EVENT_SCOPE(Name) TRACE_CPUPROFILER_EVENT_SCOPE(Name)
+#else
+#define TASKGRAPH_VERBOSE_EVENT_SCOPE(Name)
+#endif
+
 namespace UE::Tasks
 {
 	using LowLevelTasks::ETaskPriority;
@@ -212,6 +222,8 @@ namespace UE::Tasks
 			// Must not be called concurrently
 			bool AddPrerequisites(FTaskBase& Prerequisite)
 			{
+				TASKGRAPH_VERBOSE_EVENT_SCOPE(FTaskBase::AddPrerequisites_Single);
+
 				checkf(NumLocks.load(std::memory_order_relaxed) >= NumInitialLocks && NumLocks.load(std::memory_order_relaxed) < ExecutionFlag, TEXT("Prerequisites can be added only before the task is launched"));
 
 				// registering the task as a subsequent of the given prerequisite can cause its immediate launch by the prerequisite
@@ -248,6 +260,8 @@ namespace UE::Tasks
 			template<typename PrerequisiteCollectionType, decltype(std::declval<PrerequisiteCollectionType>().begin())* = nullptr>
 			void AddPrerequisites(const PrerequisiteCollectionType& InPrerequisites)
 			{
+				TASKGRAPH_VERBOSE_EVENT_SCOPE(FTaskBase::AddPrerequisites_Collection);
+
 				checkf(NumLocks.load(std::memory_order_relaxed) >= NumInitialLocks && NumLocks.load(std::memory_order_relaxed) < ExecutionFlag, TEXT("Prerequisites can be added only before the task is launched"));
 
 				// registering the task as a subsequent of the given prerequisite can cause its immediate launch by the prerequisite
@@ -326,7 +340,9 @@ namespace UE::Tasks
 			bool TryLaunch(uint64 TaskSize)
 			{
 				TaskTrace::Launched(GetTraceId(), LowLevelTask.GetDebugName(), true, TranslatePriority(LowLevelTask.GetPriority(), ExtendedPriority), TaskSize);
-				return TryUnlock();
+
+				bool bWakeUpWorker = true;
+				return TryUnlock(bWakeUpWorker);
 			}
 
 			// @return true if the task was executed and all its nested tasks are completed
@@ -352,6 +368,8 @@ namespace UE::Tasks
 			// adds a nested task that must be completed before the parent (this) is completed
 			void AddNested(FTaskBase& Nested)
 			{
+				TASKGRAPH_VERBOSE_EVENT_SCOPE(FTaskBase::AddNested);
+
 				uint32 PrevNumLocks = NumLocks.fetch_add(1, std::memory_order_relaxed); // in case we'll succeed in adding subsequent, 
 				// "happens before" registering this task as a subsequent
 				checkf(PrevNumLocks + 1 < TNumericLimits<uint32>::Max(), TEXT("Max number of nested tasks reached: %d"), TNumericLimits<uint32>::Max() - ExecutionFlag);
@@ -419,6 +437,8 @@ namespace UE::Tasks
 			// @returns true if the task was executed by the current thread
 			bool TryExecuteTask()
 			{
+				TASKGRAPH_VERBOSE_EVENT_SCOPE(FTaskBase::TryExecuteTask);
+
 				if (!TrySetExecutionFlag())
 				{
 					return false;
@@ -440,6 +460,7 @@ namespace UE::Tasks
 				{
 					UE::FInheritedContextScope InheritedContextScope = RestoreInheritedContext();
 					TaskTrace::FTaskTimingEventScope TaskEventScope(GetTraceId());
+					TASKGRAPH_VERBOSE_EVENT_SCOPE(FTaskBase::ExecuteTask);
 					ExecuteTask();
 				}
 
@@ -467,6 +488,7 @@ namespace UE::Tasks
 			// closes task by unlocking its subsequents and flagging it as completed
 			void Close()
 			{
+				TASKGRAPH_VERBOSE_EVENT_SCOPE(FTaskBase::Close);
 				checkSlow(!IsCompleted());
 
 				if (GetPipe() != nullptr)
@@ -478,9 +500,16 @@ namespace UE::Tasks
 				Subsequents.PopAllAndClose(Subs); // gets `Subs` in LIFO order
 				// try to maintain FIFO order where it's possible. e.g. if multiple piped tasks (subsequents) depend on the same (this) task, 
 				// preserve the piping order, which is also the order in which subsequents were added as dependencies of this task
-				for (int32 i = Subs.Num() - 1; i >= 0; --i)
+
+				// Push the first subsequent to the local queue so we pick it up directly as our next task.
+				// This saves us the cost of going to the global queue and performing a wake-up.
+				bool bWakeUpWorker = false;
+
+				for (int32 Index = Subs.Num() - 1; Index >= 0; --Index)
 				{
-					Subs[i]->TryUnlock();
+					// bWakeUpWorker is passed by reference and is automatically set to true if we successfully schedule a task on the local queue.
+					// so all the remaining ones are sent to the global queue.
+					Subs[Index]->TryUnlock(bWakeUpWorker);
 				}
 
 				// release nested tasks
@@ -494,8 +523,10 @@ namespace UE::Tasks
 		private:
 			// A task can be locked for execution (by prerequisites or if it's not launched yet) or for completion (by nested tasks).
 			// This method is called to unlock the task and so can result in its scheduling (and execution) or completion
-			bool TryUnlock()
+			bool TryUnlock(bool& bWakeUpWorker)
 			{
+				TASKGRAPH_VERBOSE_EVENT_SCOPE(FTaskBase::TryUnlock);
+
 				FPipe* LocalPipe = GetPipe(); // cache data locally so we won't need to touch the member (read below)
 
 				uint32 PrevNumLocks = NumLocks.fetch_sub(1, std::memory_order_acq_rel); // `acq_rel` to make it happen after task 
@@ -559,7 +590,7 @@ namespace UE::Tasks
 					}
 					else
 					{
-						Schedule();
+						Schedule(bWakeUpWorker);
 					}
 
 					return true;
@@ -578,7 +609,7 @@ namespace UE::Tasks
 				return true;
 			}
 
-			CORE_API void Schedule();
+			CORE_API void Schedule(bool& bWakeUpWorker);
 
 			// is called when the task has no pending prerequisites. Returns the previous piped task if any
 			CORE_API FTaskBase* TryPushIntoPipe();
@@ -595,8 +626,10 @@ namespace UE::Tasks
 
 			void ReleasePrerequisites()
 			{
+				TASKGRAPH_VERBOSE_EVENT_SCOPE(FTaskBase::ReleasePrerequisites);
 				while (FTaskBase* Prerequisite = Prerequisites.Pop())
 				{
+					TASKGRAPH_VERBOSE_EVENT_SCOPE(FTaskBase::ReleasePrerequisite);
 					Prerequisite->Release();
 				}
 			}
diff --git a/Engine/Source/Runtime/Core/Tests/Async/TaskGraphTest.cpp b/Engine/Source/Runtime/Core/Tests/Async/TaskGraphTest.cpp
--- a/Engine/Source/Runtime/Core/Tests/Async/TaskGraphTest.cpp
+++ b/Engine/Source/Runtime/Core/Tests/Async/TaskGraphTest.cpp
@@ -1016,6 +1016,25 @@ PRAGMA_ENABLE_DEPRECATION_WARNINGS
 		FTaskGraphInterface::Get().WaitUntilTasksComplete(Tasks);
 	}
 
+	template<int NumTasks>
+	void TestPerfChaining()
+	{
+		FGraphEventRef Ref;
+		for (int32 TaskIndex = 0; TaskIndex < NumTasks; ++TaskIndex)
+		{
+			if (Ref.IsValid())
+			{
+				Ref = FFunctionGraphTask::CreateAndDispatchWhenReady([&]() {}, TStatId{}, Ref);
+			}
+			else
+			{
+				Ref = FFunctionGraphTask::CreateAndDispatchWhenReady([&]() {});
+			}
+		}
+
+		FTaskGraphInterface::Get().WaitUntilTaskCompletes(Ref);
+	}
+
 	template<int32 NumTasks, int32 BatchSize>
 	void TestPerfBatch()
 	{
@@ -1274,6 +1293,7 @@ PRAGMA_ENABLE_DEPRECATION_WARNINGS
 		UE_BENCHMARK(5, TestPerfBatch<100000, 100>);
 		UE_BENCHMARK(5, TestPerfBatchOptimised<100000, 100>);
 		UE_BENCHMARK(5, TestLatency<10000>);
+		UE_BENCHMARK(5, TestPerfChaining<10000>);
 		UE_BENCHMARK(5, TestFGraphEventPerf<100000>);
 		UE_BENCHMARK(5, TestWorkStealing<100, 1000>);
 		UE_BENCHMARK(5, TestSpawning<100000>);
