Move the global GAutoRTFMRuntimeEnabled to only be exposed to functions
Move GAutoRTFMRuntimeEnabled to be an enum with 3 states:
  Disabled
  Enabled
  ForcedDisabled

ForcedDisabled will prevent this from ever being allowed to be turned on

diff --git a/Engine/Source/Programs/AutoRTFMTests/Private/AutoRTFMTests.cpp b/Engine/Source/Programs/AutoRTFMTests/Private/AutoRTFMTests.cpp
--- a/Engine/Source/Programs/AutoRTFMTests/Private/AutoRTFMTests.cpp
+++ b/Engine/Source/Programs/AutoRTFMTests/Private/AutoRTFMTests.cpp
@@ -16,7 +16,7 @@ public:
 		GEngineLoop.PreInit(0, nullptr);
 		FModuleManager::Get().StartProcessingNewlyLoadedObjects();
 
-		AutoRTFM::SetAutoRTFMRuntime(true);
+		AutoRTFM::SetAutoRTFMRuntime(AutoRTFM::EAutoRTFMEnabledState::AutoRTFM_Enabled);
 	}
 
 	void testRunEnded(const Catch::TestRunStats&) override
diff --git a/Engine/Source/Runtime/Core/Private/AutoRTFM/API.cpp b/Engine/Source/Runtime/Core/Private/AutoRTFM/API.cpp
--- a/Engine/Source/Runtime/Core/Private/AutoRTFM/API.cpp
+++ b/Engine/Source/Runtime/Core/Private/AutoRTFM/API.cpp
@@ -4,27 +4,46 @@
 #include "CoreGlobals.h"
 #include "HAL/IConsoleManager.h"
 
+namespace
+{
+	// Move this to a local only and use functions to access this
+	int GAutoRTFMRuntimeEnabled = AutoRTFM::EAutoRTFMEnabledState::AutoRTFM_Disabled;
+}
+
 #if UE_AUTORTFM
-UE_AUTORTFM_API bool GAutoRTFMRuntimeEnabled = false;
 static FAutoConsoleVariableRef CVarAutoRTFMRuntimeEnabled(
 	TEXT("AutoRTFMRuntimeEnabled"),
 	GAutoRTFMRuntimeEnabled,
 	TEXT("Enables the AutoRTFM runtime"),
 	ECVF_Default
 );
-#else
-static constexpr bool GAutoRTFMRuntimeEnabled = false;
 #endif
 
 namespace AutoRTFM
 {
-	void SetAutoRTFMRuntime(bool bEnabled)
+	bool SetAutoRTFMRuntime(EAutoRTFMEnabledState State)
 	{
 		// #noop if AutoRTFM is not compiled in, as GAutoRTFMRuntimeEnabled is a static constexpr when no AutoRTFM compiled in
 #if UE_AUTORTFM
-		GAutoRTFMRuntimeEnabled = bEnabled;
+		// If we have ForcedDisabled AutoRTFM from the CVar we will not change it from other means
+		if (GAutoRTFMRuntimeEnabled == EAutoRTFMEnabledState::AutoRTFM_ForcedDisabled)
+		{
+			UE_LOG(LogCore, Log, TEXT("Enabling AutoRTFM is disabled due to GAutoRTFMRuntimeEnabled set to forced disabled"));
+			return false;
+		}
+
+		GAutoRTFMRuntimeEnabled = State;
+
+		return true;
+#else
+		return false;
 #endif
 	}
+
+	bool IsAutoRTFMRuntimeEnabled()
+	{
+		return GAutoRTFMRuntimeEnabled == EAutoRTFMEnabledState::AutoRTFM_Enabled;
+	}
 }
 
 #if (defined(__AUTORTFM) && __AUTORTFM)
@@ -60,7 +79,7 @@ UE_AUTORTFM_FORCEINLINE autortfm_result TransactThenOpenImpl(void (*Work)(void*
 
 extern "C" UE_AUTORTFM_AUTORTFM("RTFM_autortfm_is_transactional") bool autortfm_is_transactional()
 {
-	if (GAutoRTFMRuntimeEnabled)
+	if (IsAutoRTFMRuntimeEnabled())
 	{
 		return FContext::Get()->IsTransactional();
 	}
@@ -76,7 +95,7 @@ extern "C" UE_AUTORTFM_AUTORTFM("RTFM_autortfm_is_closed") bool autortfm_is_clos
 // First Part - the API exposed outside transactions.
 extern "C" UE_AUTORTFM_AUTORTFM("RTFM_autortfm_transact") autortfm_result autortfm_transact(void (*Work)(void* Arg), void* Arg)
 {
-	if (GAutoRTFMRuntimeEnabled)
+	if (IsAutoRTFMRuntimeEnabled())
 	{
 	    return static_cast<autortfm_result>(FContext::Get()->Transact(Work, Arg));
 	}
@@ -136,7 +155,7 @@ extern "C" UE_AUTORTFM_AUTORTFM("RTFM_autortfm_clear_transaction_status") void a
 
 extern "C" UE_AUTORTFM_NOAUTORTFM bool autortfm_is_aborting()
 {
-	if (GAutoRTFMRuntimeEnabled)
+	if (IsAutoRTFMRuntimeEnabled())
 	{
 		return FContext::Get()->IsAborting();
 	}
@@ -168,7 +187,7 @@ extern "C" UE_AUTORTFM_AUTORTFM("RTFM_autortfm_close") autortfm_status autortfm_
 {
 	autortfm_status Result = autortfm_status_ontrack;
 
-	if (GAutoRTFMRuntimeEnabled)
+	if (IsAutoRTFMRuntimeEnabled())
 	{
 		UE_CLOG(!FContext::IsTransactional(), LogAutoRTFM, Fatal, TEXT("Close called from an outside a transaction."));
 
diff --git a/Engine/Source/Runtime/Core/Public/AutoRTFM/AutoRTFM.h b/Engine/Source/Runtime/Core/Public/AutoRTFM/AutoRTFM.h
--- a/Engine/Source/Runtime/Core/Public/AutoRTFM/AutoRTFM.h
+++ b/Engine/Source/Runtime/Core/Public/AutoRTFM/AutoRTFM.h
@@ -405,7 +405,15 @@ enum class EContextStatus
 	AbortedByCascade = autortfm_status_aborted_by_cascade
 };
 
-UE_AUTORTFM_API void SetAutoRTFMRuntime(bool bEnabled);
+enum EAutoRTFMEnabledState
+{
+	AutoRTFM_Disabled,
+	AutoRTFM_Enabled,
+	AutoRTFM_ForcedDisabled,
+};
+
+UE_AUTORTFM_API bool SetAutoRTFMRuntime(EAutoRTFMEnabledState bEnabled);
+UE_AUTORTFM_API bool IsAutoRTFMRuntimeEnabled();
 
 UE_AUTORTFM_FORCEINLINE bool IsTransactional() { return autortfm_is_transactional(); }
 UE_AUTORTFM_FORCEINLINE bool IsClosed() { return autortfm_is_closed(); }
