Fixed CIS error with GetTypeHash not visible when hashing StaticArray.

- Also improved StaticArray's hash.

diff --git a/Engine/Source/Runtime/Core/Public/Containers/StaticArray.h b/Engine/Source/Runtime/Core/Public/Containers/StaticArray.h
--- a/Engine/Source/Runtime/Core/Public/Containers/StaticArray.h
+++ b/Engine/Source/Runtime/Core/Public/Containers/StaticArray.h
@@ -8,6 +8,7 @@
 #include "Templates/UnrealTemplate.h"
 #include "Templates/UnrealTypeTraits.h"
 #include "Delegates/IntegerSequence.h"
+#include "Templates/TypeHash.h"
 
 /** An array with a static number of elements. */
 template <typename InElementType, uint32 NumElements, uint32 Alignment = alignof(InElementType)>
@@ -228,10 +229,10 @@ FArchive& operator<<(FArchive& Ar,TStaticArray<ElementType, NumElements, Alignme
 template <typename ElementType, uint32 NumElements, uint32 Alignment>
 uint32 GetTypeHash(const TStaticArray<ElementType, NumElements, Alignment>& Array)
 {
-	uint32 Result = 0;
-	for(uint32 ElementIndex = 0;ElementIndex < NumElements;++ElementIndex)
+	uint32 Hash = 0;
+	for (const ElementType& Element : Array)
 	{
-		Result ^= GetTypeHash(Array[ElementIndex]);
+		Hash = HashCombineFast(Hash, GetTypeHash(Element));
 	}
-	return Result;
+	return Hash;
 }
