Revert what look like local changes for testing

diff --git a/Engine/Source/Runtime/Core/Private/Windows/WindowsPlatformCrashContext.cpp b/Engine/Source/Runtime/Core/Private/Windows/WindowsPlatformCrashContext.cpp
--- a/Engine/Source/Runtime/Core/Private/Windows/WindowsPlatformCrashContext.cpp
+++ b/Engine/Source/Runtime/Core/Private/Windows/WindowsPlatformCrashContext.cpp
@@ -328,7 +328,11 @@ bool CreateCrashReportClientPath(TCHAR* OutClientPath, int32 MaxLength)
 
 		// Find the path to crash reporter binary. Avoid creating FStrings.
 		*OutClientPath = TCHAR('\0');
-		FCString::Strncat(OutClientPath, L"D:\\StagedBuilds\\B\\WindowsClient\\Engine\\Binaries\\Win64\\CrashReportClient.exe", MaxLength);
+		FCString::Strncat(OutClientPath, EngineDir, MaxLength);
+		FCString::Strncat(OutClientPath, TEXT("Binaries/"), MaxLength);
+		FCString::Strncat(OutClientPath, BinariesDir, MaxLength);
+		FCString::Strncat(OutClientPath, TEXT("/"), MaxLength);
+		FCString::Strncat(OutClientPath, CrashReportClientExeName, MaxLength);
 
 		const DWORD Results = GetFileAttributesW(OutClientPath);
 		return Results != INVALID_FILE_ATTRIBUTES;
