- Bump Niagara version to fix DDC pollution

diff --git a/Engine/Source/Runtime/Core/Private/UObject/FortniteReleaseBranchCustomObjectVersion.cpp b/Engine/Source/Runtime/Core/Private/UObject/FortniteReleaseBranchCustomObjectVersion.cpp
--- a/Engine/Source/Runtime/Core/Private/UObject/FortniteReleaseBranchCustomObjectVersion.cpp
+++ b/Engine/Source/Runtime/Core/Private/UObject/FortniteReleaseBranchCustomObjectVersion.cpp
@@ -6,7 +6,7 @@ TMap<FGuid, FGuid> FFortniteReleaseBranchCustomObjectVersion::GetSystemGuids()
 	TMap<FGuid, FGuid> SystemGuids;
 	const FDevSystemGuids& DevGuids = FDevSystemGuids::Get();
 
-	SystemGuids.Add(DevGuids.Niagara_LatestScriptCompileVersion, FGuid("C2B4733DD8CEFC43B07A1ADD16A8978B"));
+	SystemGuids.Add(DevGuids.Niagara_LatestScriptCompileVersion, FGuid("CB0D96CCD66AC742ABB0744654ED5FCE"));
 	SystemGuids.Add(DevGuids.SkeletalMeshDerivedDataVersion, FGuid("5705956EC7134274A5A1BD3FC74DB3A9"));
 	return SystemGuids;
 }
