Fixed incorrect use of non-templated vector type in TSphere

#rnx

diff --git a/Engine/Source/Runtime/Core/Public/Math/Sphere.h b/Engine/Source/Runtime/Core/Public/Math/Sphere.h
--- a/Engine/Source/Runtime/Core/Public/Math/Sphere.h
+++ b/Engine/Source/Runtime/Core/Public/Math/Sphere.h
@@ -279,7 +279,7 @@ TSphere<T> TSphere<T>::TransformBy(const TMatrix<T>& M) const
 {
 	TSphere<T>	Result;
 
-	FVector4 TransformedCenter = M.TransformPosition(this->Center);
+	TVector4<T> TransformedCenter = M.TransformPosition(this->Center);
 	Result.Center = TVector<T>(TransformedCenter.X, TransformedCenter.Y, TransformedCenter.Z);
 
 	const TVector<T> XAxis(M.M[0][0], M.M[0][1], M.M[0][2]);
