Mutex: Added constructors to acquire the lock upon construction

diff --git a/Engine/Source/Runtime/Core/Public/Async/ExternalMutex.h b/Engine/Source/Runtime/Core/Public/Async/ExternalMutex.h
--- a/Engine/Source/Runtime/Core/Public/Async/ExternalMutex.h
+++ b/Engine/Source/Runtime/Core/Public/Async/ExternalMutex.h
@@ -2,6 +2,7 @@
 
 #pragma once
 
+#include "Async/LockTags.h"
 #include "CoreTypes.h"
 #include <atomic>
 
@@ -21,7 +22,17 @@ namespace UE
 class FExternalMutex final
 {
 public:
-	constexpr FExternalMutex(std::atomic<uint8>& InState) : State(InState) {}
+	inline constexpr explicit FExternalMutex(std::atomic<uint8>& InState)
+		: State(InState)
+	{
+	}
+
+	/** Construct in a locked state. Avoids an expensive compare-and-swap at creation time. */
+	inline explicit FExternalMutex(std::atomic<uint8>& InState, FAcquireLock)
+		: State(InState)
+	{
+		State.fetch_or(IsLockedFlag, std::memory_order_acquire);
+	}
 
 	FExternalMutex(const FExternalMutex&) = delete;
 	FExternalMutex& operator=(const FExternalMutex&) = delete;
@@ -30,7 +41,7 @@ public:
 	{
 		return (State.load(std::memory_order_relaxed) & IsLockedFlag);
 	}
-	
+
 	inline bool TryLock()
 	{
 		uint8 Expected = State.load(std::memory_order_acquire);
@@ -65,10 +76,10 @@ private:
 	UE_API void LockSlow();
 	UE_API void UnlockSlow();
 
-	std::atomic<uint8>& State;
-
 	static constexpr uint8 IsLockedFlag = 1 << 0;
 	static constexpr uint8 HasWaitingThreadsFlag = 1 << 1;
+
+	std::atomic<uint8>& State;
 };
 
 } // UE
diff --git a/Engine/Source/Runtime/Core/Public/Async/LockTags.h b/Engine/Source/Runtime/Core/Public/Async/LockTags.h
new file mode 100644
--- /dev/null
+++ b/Engine/Source/Runtime/Core/Public/Async/LockTags.h
@@ -0,0 +1,14 @@
+// Copyright Epic Games, Inc. All Rights Reserved.
+
+#pragma once
+
+namespace UE
+{
+
+/** Use to acquire a lock on construction. */
+inline constexpr struct FAcquireLock final { explicit FAcquireLock() = default; } AcquireLock;
+
+/** Use with dynamic locks to defer locking on construction. */
+inline constexpr struct FDeferLock final { explicit FDeferLock() = default; } DeferLock;
+
+} // UE
diff --git a/Engine/Source/Runtime/Core/Public/Async/Mutex.h b/Engine/Source/Runtime/Core/Public/Async/Mutex.h
--- a/Engine/Source/Runtime/Core/Public/Async/Mutex.h
+++ b/Engine/Source/Runtime/Core/Public/Async/Mutex.h
@@ -2,6 +2,7 @@
 
 #pragma once
 
+#include "Async/LockTags.h"
 #include "CoreTypes.h"
 #include <atomic>
 
@@ -19,6 +20,12 @@ class FMutex final
 public:
 	constexpr FMutex() = default;
 
+	/** Construct in a locked state. Avoids an expensive compare-and-swap at creation time. */
+	inline explicit FMutex(FAcquireLock)
+		: State(IsLockedFlag)
+	{
+	}
+
 	FMutex(const FMutex&) = delete;
 	FMutex& operator=(const FMutex&) = delete;
 
diff --git a/Engine/Source/Runtime/Core/Public/Async/UniqueLock.h b/Engine/Source/Runtime/Core/Public/Async/UniqueLock.h
--- a/Engine/Source/Runtime/Core/Public/Async/UniqueLock.h
+++ b/Engine/Source/Runtime/Core/Public/Async/UniqueLock.h
@@ -2,17 +2,12 @@
 
 #pragma once
 
+#include "Async/LockTags.h"
 #include "Misc/AssertionMacros.h"
 
 namespace UE
 {
 
-/** Use with dynamic locks to defer locking on construction. */
-struct FDeferLock final
-{
-	explicit FDeferLock() = default;
-};
-
 /**
  * A basic mutex ownership wrapper that locks on construction and unlocks on destruction.
  *
