Set if AutoRTFM is enabled or disabled for our game crash data

diff --git a/Engine/Source/Runtime/Core/Private/AutoRTFM/API.cpp b/Engine/Source/Runtime/Core/Private/AutoRTFM/API.cpp
--- a/Engine/Source/Runtime/Core/Private/AutoRTFM/API.cpp
+++ b/Engine/Source/Runtime/Core/Private/AutoRTFM/API.cpp
@@ -3,11 +3,17 @@
 #include "AutoRTFM/AutoRTFM.h"
 #include "CoreGlobals.h"
 #include "HAL/IConsoleManager.h"
+#include "GenericPlatform/GenericPlatformCrashContext.h"
 
 namespace
 {
 	// Move this to a local only and use functions to access this
 	int GAutoRTFMRuntimeEnabled = AutoRTFM::EAutoRTFMEnabledState::AutoRTFM_Disabled;
+
+	void UpdateAutoRTFMRuntimeCrashData()
+	{
+		FGenericCrashContext::SetGameData(TEXT("IsAutoRTFMRuntimeEnabled"), GAutoRTFMRuntimeEnabled == AutoRTFM::EAutoRTFMEnabledState::AutoRTFM_Enabled ? TEXT("true") : TEXT("false"));
+	}
 }
 
 #if UE_AUTORTFM
@@ -15,8 +21,14 @@ static FAutoConsoleVariableRef CVarAutoRTFMRuntimeEnabled(
 	TEXT("AutoRTFMRuntimeEnabled"),
 	GAutoRTFMRuntimeEnabled,
 	TEXT("Enables the AutoRTFM runtime"),
+	FConsoleVariableDelegate::CreateLambda([] (IConsoleVariable*) { UpdateAutoRTFMRuntimeCrashData(); }),
 	ECVF_Default
 );
+
+static FDelayedAutoRegisterHelper DelayedAutoRegister(EDelayedRegisterRunPhase::EndOfEngineInit, []
+{
+	UpdateAutoRTFMRuntimeCrashData();
+});
 #endif
 
 namespace AutoRTFM
@@ -34,6 +46,8 @@ namespace AutoRTFM
 
 		GAutoRTFMRuntimeEnabled = State;
 
+		UpdateAutoRTFMRuntimeCrashData();
+
 		return true;
 #else
 		return false;
