Fix ClangPlatformCodeAnalysis.h include for windows platforms

#rnx

diff --git a/Engine/Source/Runtime/Core/Public/Windows/WindowsPlatform.h b/Engine/Source/Runtime/Core/Public/Windows/WindowsPlatform.h
--- a/Engine/Source/Runtime/Core/Public/Windows/WindowsPlatform.h
+++ b/Engine/Source/Runtime/Core/Public/Windows/WindowsPlatform.h
@@ -201,7 +201,11 @@ typedef FWindowsPlatformTypes FPlatformTypes;
 
 
 // Include code analysis features
-#include "Windows/WindowsPlatformCodeAnalysis.h"
+#if PLATFORM_COMPILER_CLANG
+	#include "Clang/ClangPlatformCodeAnalysis.h"
+#elif PLATFORM_WINDOWS
+	#include "Windows/WindowsPlatformCodeAnalysis.h"
+#endif
 
 // Other macros
 #ifndef ENABLE_WIN_ALLOC_TRACKING
