- Add additional macros for getting enum/member/function names as string views. Required for using in a compile time context.
    + GET_ENUMERATOR_NAME_STRING_VIEW_CHECKED
    + GET_MEMBER_NAME_STRING_VIEW_CHECKED
    + GET_FUNCTION_NAME_STRING_VIEW_CHECKED
- Add missing macro for getting enum name as a string, to match the existing macros for member and function names.
    + GET_ENUMERATOR_NAME_STRING_CHECKED

diff --git a/Engine/Source/Runtime/Core/Public/Misc/AssertionMacros.h b/Engine/Source/Runtime/Core/Public/Misc/AssertionMacros.h
--- a/Engine/Source/Runtime/Core/Public/Misc/AssertionMacros.h
+++ b/Engine/Source/Runtime/Core/Public/Misc/AssertionMacros.h
@@ -452,6 +452,12 @@ namespace UEAsserts_Private
 #define GET_ENUMERATOR_NAME_CHECKED(EnumName, EnumeratorName) \
 	((void)sizeof(UEAsserts_Private::GetMemberNameCheckedJunk(EnumName::EnumeratorName)), FName(TEXT(#EnumeratorName)))
 
+#define GET_ENUMERATOR_NAME_STRING_CHECKED(EnumName, EnumeratorName) \
+	((void)sizeof(UEAsserts_Private::GetMemberNameCheckedJunk(EnumName::EnumeratorName)), TEXT(#EnumeratorName))
+
+#define GET_ENUMERATOR_NAME_STRING_VIEW_CHECKED(EnumName, EnumeratorName) \
+	((void)sizeof(UEAsserts_Private::GetMemberNameCheckedJunk(EnumName::EnumeratorName)), TEXTVIEW(#EnumeratorName))
+
 // Returns FName(TEXT("MemberName")), while statically verifying that the member exists in ClassName
 #define GET_MEMBER_NAME_CHECKED(ClassName, MemberName) \
 	((void)sizeof(UEAsserts_Private::GetMemberNameCheckedJunk(((ClassName*)0)->MemberName)), FName(TEXT(#MemberName)))
@@ -459,6 +465,9 @@ namespace UEAsserts_Private
 #define GET_MEMBER_NAME_STRING_CHECKED(ClassName, MemberName) \
 	((void)sizeof(UEAsserts_Private::GetMemberNameCheckedJunk(((ClassName*)0)->MemberName)), TEXT(#MemberName))
 
+#define GET_MEMBER_NAME_STRING_VIEW_CHECKED(ClassName, MemberName) \
+	((void)sizeof(UEAsserts_Private::GetMemberNameCheckedJunk(((ClassName*)0)->MemberName)), TEXTVIEW(#MemberName))
+
 // Returns FName(TEXT("FunctionName")), while statically verifying that the function exists in ClassName
 #define GET_FUNCTION_NAME_CHECKED(ClassName, FunctionName) \
 	((void)sizeof(&ClassName::FunctionName), FName(TEXT(#FunctionName)))
@@ -466,6 +475,9 @@ namespace UEAsserts_Private
 #define GET_FUNCTION_NAME_STRING_CHECKED(ClassName, FunctionName) \
 	((void)sizeof(&ClassName::FunctionName), TEXT(#FunctionName))
 
+#define GET_FUNCTION_NAME_STRING_VIEW_CHECKED(ClassName, FunctionName) \
+	((void)sizeof(&ClassName::FunctionName), TEXTVIEW(#FunctionName))
+
 // Returns FName(TEXT("FunctionName")), while statically verifying that the function exists in ClassName
 // Handles overloaded functions by specifying the argument of the overload to use
 #define GET_FUNCTION_NAME_CHECKED_OneParam(ClassName, FunctionName, ArgType) \
