Fixed move semantics in TFrameValue.

diff --git a/Engine/Source/Runtime/Core/Public/Misc/FrameValue.h b/Engine/Source/Runtime/Core/Public/Misc/FrameValue.h
--- a/Engine/Source/Runtime/Core/Public/Misc/FrameValue.h
+++ b/Engine/Source/Runtime/Core/Public/Misc/FrameValue.h
@@ -18,17 +18,17 @@ private:
 	TOptional<ValueType> Value;
 
 public:
-	/** Construct an OptionaType with a valid value. */
+	/** Construct an OptionalType with a valid value. */
 	TFrameValue(const ValueType& InValue)
 		: FrameSet(GFrameCounter)
+		, Value(InValue)
 	{
-		Value = InValue;
 	}
 
 	TFrameValue(ValueType&& InValue)
 		: FrameSet(GFrameCounter)
+		, Value(MoveTemp(InValue))
 	{
-		Value = InValue;
 	}
 
 	/** Construct an OptionalType with no value; i.e. unset */
@@ -47,19 +47,19 @@ public:
 	TFrameValue(TFrameValue&& InValue)
 		: FrameSet(GFrameCounter)
 	{
-		Value = InValue;
+		Value = MoveTemp(InValue);
 	}
 
 	TFrameValue& operator=(const TFrameValue& InValue)
 	{
-		Value = InValue;
+		Value = InValue.Value;
 		FrameSet = GFrameCounter;
 		return *this;
 	}
 
 	TFrameValue& operator=(TFrameValue&& InValue)
 	{
-		Value = InValue;
+		Value = MoveTemp(InValue.Value);
 		FrameSet = GFrameCounter;
 		return *this;
 	}
@@ -73,7 +73,7 @@ public:
 
 	TFrameValue& operator=(ValueType&& InValue)
 	{
-		Value = InValue;
+		Value = MoveTemp(InValue);
 		FrameSet = GFrameCounter;
 		return *this;
 	}
@@ -91,8 +91,15 @@ public:
 		return Value.GetValue();
 	}
 
-	ValueType TryGetValue(ValueType UnsetValue) const
+	ValueType TryGetValue(ValueType UnsetValue) const&
+	{
+		return IsSet() ? Value.GetValue() : MoveTemp(UnsetValue);
+	}
+
+	ValueType TryGetValue(ValueType UnsetValue) &&
 	{
-		return IsSet() ? Value.GetValue() : UnsetValue;
+		ValueType Result = IsSet() ? MoveTemp(Value.GetValue()) : MoveTemp(UnsetValue);
+		Value.Reset();
+		return Result;
 	}
 };
\ No newline at end of file
