Remove not needed code in FAndroidMisc::ShouldDisablePluginAtRuntime

diff --git a/Engine/Source/Runtime/Core/Private/Android/AndroidPlatformMisc.cpp b/Engine/Source/Runtime/Core/Private/Android/AndroidPlatformMisc.cpp
--- a/Engine/Source/Runtime/Core/Private/Android/AndroidPlatformMisc.cpp
+++ b/Engine/Source/Runtime/Core/Private/Android/AndroidPlatformMisc.cpp
@@ -1839,18 +1839,6 @@ bool FAndroidMisc::IsSupportedAndroidDevice()
 }
 #endif
 
-bool FAndroidMisc::ShouldDisablePluginAtRuntime(const FString& PluginName)
-{
-#if PLATFORM_ANDROID_ARM64 || PLATFORM_ANDROID_X64
-	// disable OnlineSubsystemGooglePlay for unsupported Android architectures
-	if (PluginName.Equals(TEXT("OnlineSubsystemGooglePlay")))
-	{
-		return true;
-	}
-#endif
-	return false;
-}
-
 ///////////////////////////////////////////////////////////////////////////////
 //
 // Extracted from vk_platform.h and vulkan.h with modifications just to allow
diff --git a/Engine/Source/Runtime/Core/Public/Android/AndroidPlatformMisc.h b/Engine/Source/Runtime/Core/Public/Android/AndroidPlatformMisc.h
--- a/Engine/Source/Runtime/Core/Public/Android/AndroidPlatformMisc.h
+++ b/Engine/Source/Runtime/Core/Public/Android/AndroidPlatformMisc.h
@@ -58,7 +58,6 @@ struct FAndroidMisc : public FGenericPlatformMisc
 	static CORE_API EAppReturnType::Type MessageBoxExt( EAppMsgType::Type MsgType, const TCHAR* Text, const TCHAR* Caption );
 	static CORE_API bool UseRenderThread();
 	static CORE_API bool HasPlatformFeature(const TCHAR* FeatureName);
-	static CORE_API bool ShouldDisablePluginAtRuntime(const FString& PluginName);
 	static CORE_API bool SupportsES30();
 
 public:
