Adding some direction to the comments of EnumAsByte

diff --git a/Engine/Source/Runtime/Core/Public/Containers/EnumAsByte.h b/Engine/Source/Runtime/Core/Public/Containers/EnumAsByte.h
--- a/Engine/Source/Runtime/Core/Public/Containers/EnumAsByte.h
+++ b/Engine/Source/Runtime/Core/Public/Containers/EnumAsByte.h
@@ -8,6 +8,14 @@
 
 /**
  * Template to store enumeration values as bytes in a type-safe way.
+ * Blueprint enums should either be enum classes (preferred):
+ *		enum class EMyEnum : uint8 { One, Two }, which doesn't require wrapping in this template
+ * or a namespaced enum
+ *		namespace EMyEnum
+ *		{
+ *			enum Type // <- literally Type, not a placeholder
+ *			{ One, Two };
+ *		}
  */
 template <class InEnumType>
 class TEnumAsByte
