MemFree isn't including CachedAllocatorSlack in Shipping

#tests ReplayRun, XB1

diff --git a/Engine/Source/Runtime/Core/Private/HAL/MallocBinned2.cpp b/Engine/Source/Runtime/Core/Private/HAL/MallocBinned2.cpp
--- a/Engine/Source/Runtime/Core/Private/HAL/MallocBinned2.cpp
+++ b/Engine/Source/Runtime/Core/Private/HAL/MallocBinned2.cpp
@@ -1459,20 +1459,24 @@ int64 FMallocBinned2::GetTotalAllocatedSmallPoolMemory() const
 
 void FMallocBinned2::GetAllocatorStats( FGenericMemoryStats& OutStats )
 {
-#if BINNED2_ALLOCATOR_STATS
+#if PLATFORM_USE_CACHED_SLACK_MEMORY_IN_MEMORY_STATS || BINNED2_ALLOCATOR_STATS
 	NOALLOC_SCOPE_CYCLE_COUNTER(STAT_FMallocBinned2_GetAllocatorStats);
 
+	// Even if we have MB2 stats off, if we want to include cached slack, we need the PageAllocatorFreeCacheSize stat
+	uint64 OSPageAllocatorCachedFreeSize = CachedOSPageAllocator.GetCachedFreeTotal();
+	OutStats.Add(TEXT("PageAllocatorFreeCacheSize"), OSPageAllocatorCachedFreeSize);
+#endif 
+
+#if BINNED2_ALLOCATOR_STATS
 	int64  TotalAllocatedSmallPoolMemory           = GetTotalAllocatedSmallPoolMemory();
 	int64  LocalAllocatedOSSmallPoolMemory         = AllocatedOSSmallPoolMemory.Load(EMemoryOrder::Relaxed);
 	int64  LocalAllocatedLargePoolMemory           = AllocatedLargePoolMemory.Load(EMemoryOrder::Relaxed);
 	int64  LocalAllocatedLargePoolMemoryWAlignment = AllocatedLargePoolMemoryWAlignment.Load(EMemoryOrder::Relaxed);
-	uint64 OSPageAllocatorCachedFreeSize           = CachedOSPageAllocator.GetCachedFreeTotal();
 
 	OutStats.Add(TEXT("AllocatedSmallPoolMemory"), TotalAllocatedSmallPoolMemory);
 	OutStats.Add(TEXT("AllocatedOSSmallPoolMemory"), LocalAllocatedOSSmallPoolMemory);
 	OutStats.Add(TEXT("AllocatedLargePoolMemory"), LocalAllocatedLargePoolMemory);
 	OutStats.Add(TEXT("AllocatedLargePoolMemoryWAlignment"), LocalAllocatedLargePoolMemoryWAlignment);
-	OutStats.Add(TEXT("PageAllocatorFreeCacheSize"), OSPageAllocatorCachedFreeSize);
 
 	uint64 TotalAllocated = TotalAllocatedSmallPoolMemory + LocalAllocatedLargePoolMemory;
 	uint64 TotalOSAllocated = LocalAllocatedOSSmallPoolMemory + LocalAllocatedLargePoolMemoryWAlignment + OSPageAllocatorCachedFreeSize;
@@ -1480,6 +1484,7 @@ void FMallocBinned2::GetAllocatorStats( FGenericMemoryStats& OutStats )
 	OutStats.Add(TEXT("TotalAllocated"), TotalAllocated);
 	OutStats.Add(TEXT("TotalOSAllocated"), TotalOSAllocated);
 #endif
+
 	FMalloc::GetAllocatorStats(OutStats);
 }
 
diff --git a/Engine/Source/Runtime/Core/Private/ProfilingDebugging/CsvProfiler.cpp b/Engine/Source/Runtime/Core/Private/ProfilingDebugging/CsvProfiler.cpp
--- a/Engine/Source/Runtime/Core/Private/ProfilingDebugging/CsvProfiler.cpp
+++ b/Engine/Source/Runtime/Core/Private/ProfilingDebugging/CsvProfiler.cpp
@@ -16,7 +16,6 @@
 #include "HAL/PlatformProcess.h"
 #include "HAL/PlatformTime.h"
 #include "HAL/IConsoleManager.h"
-#include "HAL/MemoryMisc.h"
 #include "Misc/Paths.h"
 #include "Misc/CommandLine.h"
 #include "Misc/ScopeLock.h"
@@ -3044,26 +3043,6 @@ void FCsvProfiler::EndFrame()
 
 		FPlatformMemoryStats MemoryStats = FPlatformMemory::GetStats();
 
-#if PLATFORM_USE_CACHED_SLACK_MEMORY_IN_MEMORY_STATS
-		// Take into account cached slacked memory to reflect what is really available
-		if (GMalloc != nullptr)
-		{
-			uint64 CachedFreeSize = 0;
-
-			FGenericMemoryStats MallocStats;
-			GMalloc->GetAllocatorStats(MallocStats);
-
-			SIZE_T* PageAllocatorFreeCacheSize = MallocStats.Data.Find("PageAllocatorFreeCacheSize");
-			if (PageAllocatorFreeCacheSize != nullptr)
-			{
-				CachedFreeSize = *PageAllocatorFreeCacheSize;
-			}
-
-			MemoryStats.UsedPhysical -= CachedFreeSize;
-			MemoryStats.AvailablePhysical += CachedFreeSize;
-		}
-#endif
-
 		float PhysicalMBFree = float(MemoryStats.AvailablePhysical) / (1024.0f * 1024.0f);
 		float UsedExtendedMB = 0;
 		float PhysicalMBUsed = float(MemoryStats.UsedPhysical) / (1024.0f * 1024.0f);
