Simplified ChunkArray iterator following earlier review.

#rnx

diff --git a/Engine/Source/Runtime/Core/Public/Containers/ChunkedArray.h b/Engine/Source/Runtime/Core/Public/Containers/ChunkedArray.h
--- a/Engine/Source/Runtime/Core/Public/Containers/ChunkedArray.h
+++ b/Engine/Source/Runtime/Core/Public/Containers/ChunkedArray.h
@@ -13,37 +13,30 @@ namespace UE4ChunkedArray_Private
 	template <typename ChunkType, typename ElementType, uint32 NumElementsPerChunk>
 	struct TChunkedArrayIterator
 	{
-		TChunkedArrayIterator(ChunkType** InChunk, ChunkType** InLastChunk, ElementType* InElem)
-			: Elem (InElem)
-			, Chunk(InChunk)
-			, LastChunk(InLastChunk)
-		{
-		}
-
-		ElementType* Elem;
-		ChunkType**  Chunk;
-		ChunkType**  LastChunk;
-		uint32		 Count = 0;
+		ChunkType**	Chunk;
+		uint32		Count = 0;
+		uint32		ElementIndex = 0;
 
 		ElementType& operator*() const
 		{
-			return *Elem;
+			return (*Chunk)->Elements[ElementIndex];
 		}
 
 		void operator++()
 		{
-			++Count;
-			++Elem;
-			if (Chunk != LastChunk && Elem == (*Chunk)->Elements + NumElementsPerChunk)
+			++ElementIndex;
+			if (ElementIndex >= NumElementsPerChunk)
 			{
+				ElementIndex = 0;
 				++Chunk;
-				Elem = (*Chunk)->Elements;
 			}
+
+			++Count;
 		}
 
 		bool operator!=(const TChunkedArrayIterator& Rhs) const
 		{
-			return Count != Rhs.Count;
+			return Count < Rhs.Count;
 		}
 	};
 }
@@ -297,40 +290,22 @@ private:
 public:
 	FIterType begin()
 	{
-		int32 Num = NumElements;
-		FChunk** ChunkPtr = Chunks.GetData();
-		FChunk** LastChunkPtr = Chunks.GetData() + (Num ? Num - 1 : 0) / NumElementsPerChunk;
-		return FIterType(ChunkPtr, LastChunkPtr, ChunkPtr ? (*ChunkPtr)->Elements : nullptr);
+		return FIterType{Chunks.GetData()};
 	}
 
 	FConstIterType begin() const
 	{
-		int32 Num = NumElements;
-		const FChunk** ChunkPtr = Chunks.GetData();
-		const FChunk** LastChunkPtr = Chunks.GetData() + (Num ? Num - 1 : 0) / NumElementsPerChunk;
-		return FConstIterType(ChunkPtr, LastChunkPtr, ChunkPtr ? (*ChunkPtr)->Elements : nullptr);
+		return FConstIterType{Chunks.GetData()};
 	}
 
 	FIterType end()
 	{
-		int32 Num = NumElements;
-		bool bBeyondLastChunk = Num && (Num % NumElementsPerChunk) == 0;
-		FChunk** ChunkPtr = Chunks.GetData() + (Num / NumElementsPerChunk) + (bBeyondLastChunk ? -1 : 0); // do not read off the end of the chunk array!
-		FChunk** LastChunkPtr = Chunks.GetData() + (Num ? Num - 1 : 0) / NumElementsPerChunk;
-		FIterType Ret(ChunkPtr, LastChunkPtr, ChunkPtr ? (*ChunkPtr)->Elements + (bBeyondLastChunk ? NumElementsPerChunk : (Num % NumElementsPerChunk))  : nullptr);
-		Ret.Count = Num;
-		return Ret;
+		return FIterType{nullptr, uint32(NumElements)};
 	}
 
 	FConstIterType end() const
 	{
-		int32 Num = NumElements;
-		bool bBeyondLastChunk = Num && Num % NumElementsPerChunk == 0;
-		const FChunk** ChunkPtr = Chunks.GetData() + (Num / NumElementsPerChunk) + (bBeyondLastChunk ? -1 : 0); // do not read off the end of the chunk array!
-		const FChunk** LastChunkPtr = Chunks.GetData() + (Num ? Num - 1 : 0) / NumElementsPerChunk;
-		FConstIterType Ret(ChunkPtr, LastChunkPtr, ChunkPtr ? (*ChunkPtr)->Elements + (bBeyondLastChunk ? NumElementsPerChunk : (Num % NumElementsPerChunk)) : nullptr);
-		Ret.Count = Num;
-		return Ret;
+		return FConstIterType{nullptr, uint32(NumElements)};
 	}
 };
 
