Tasks System integration: prevents "busy waiting" by the new tasks on named threads

diff --git a/Engine/Source/Runtime/Core/Private/Tasks/TaskPrivate.cpp b/Engine/Source/Runtime/Core/Private/Tasks/TaskPrivate.cpp
--- a/Engine/Source/Runtime/Core/Private/Tasks/TaskPrivate.cpp
+++ b/Engine/Source/Runtime/Core/Private/Tasks/TaskPrivate.cpp
@@ -190,13 +190,27 @@ namespace UE::Tasks
 			TaskTrace::FWaitingScope WaitingScope(GetTraceId());
 			TRACE_CPUPROFILER_EVENT_SCOPE(Tasks::Wait);
 
-			// if we are on a named thread, handle waiting in TaskGraph-specific style
-			// (named threads don't support waiting with timeout)
-			if (Timeout == FTimeout::Never() && TryWaitOnNamedThread(*this))
+			return WaitImpl(Timeout);
+		}
+
+		void FTaskBase::WaitWithNamedThreadsSupport()
+		{
+			if (IsCompleted())
 			{
-				return true;
+				return;
 			}
 
+			TaskTrace::FWaitingScope WaitingScope(GetTraceId());
+			TRACE_CPUPROFILER_EVENT_SCOPE(Tasks::Wait);
+
+			if (!TryWaitOnNamedThread(*this))
+			{
+				WaitImpl(FTimeout::Never());
+			}
+		}
+
+		bool FTaskBase::WaitImpl(FTimeout Timeout)
+		{
 			// ignore the result as we still have to make sure the task is completed upon returning from this function call
 			TryRetractAndExecute(Timeout);
 
diff --git a/Engine/Source/Runtime/Core/Public/Async/TaskGraphInterfaces.h b/Engine/Source/Runtime/Core/Public/Async/TaskGraphInterfaces.h
--- a/Engine/Source/Runtime/Core/Public/Async/TaskGraphInterfaces.h
+++ b/Engine/Source/Runtime/Core/Public/Async/TaskGraphInterfaces.h
@@ -552,7 +552,7 @@ public:
 			return FTaskGraphInterface::Get().WaitUntilTaskCompletes(this, CurrentThreadIfKnown);
 		}
 
-		FTaskBase::Wait(UE::FTimeout::Never());
+		FTaskBase::WaitWithNamedThreadsSupport();
 	}
 
 	ENamedThreads::Type GetThreadToExecuteOn() const
diff --git a/Engine/Source/Runtime/Core/Public/Tasks/TaskPrivate.h b/Engine/Source/Runtime/Core/Public/Tasks/TaskPrivate.h
--- a/Engine/Source/Runtime/Core/Public/Tasks/TaskPrivate.h
+++ b/Engine/Source/Runtime/Core/Public/Tasks/TaskPrivate.h
@@ -393,6 +393,10 @@ namespace UE::Tasks
 			// @return true if the task is completed
 			CORE_API bool Wait(FTimeout Timeout);
 
+			// mimics the old tasks (TaskGraph) behaviour on named threads: waiting for a task on a named thread pulls other tasks from this
+			// named thread queue and executes them
+			CORE_API void WaitWithNamedThreadsSupport();
+
 			// waits until the task is completed or waiting timed out, while executing other tasks
 			bool BusyWait(FTimeout Timeout)
 			{
@@ -637,6 +641,8 @@ namespace UE::Tasks
 			CORE_API void StartPipeExecution();
 			CORE_API void FinishPipeExecution();
 
+			CORE_API bool WaitImpl(FTimeout Timeout);
+
 		private:
 			EExtendedTaskPriority ExtendedPriority; // internal priorities, if any
 
