Prevent duplicate DDPI entries if there is no project specified.

#rnx

diff --git a/Engine/Source/Runtime/Core/Private/Misc/DataDrivenPlatformInfoRegistry.cpp b/Engine/Source/Runtime/Core/Private/Misc/DataDrivenPlatformInfoRegistry.cpp
--- a/Engine/Source/Runtime/Core/Private/Misc/DataDrivenPlatformInfoRegistry.cpp
+++ b/Engine/Source/Runtime/Core/Private/Misc/DataDrivenPlatformInfoRegistry.cpp
@@ -57,7 +57,7 @@ static const TArray<FString>& GetDataDrivenIniFilenames()
 			FString IniPath = FPaths::Combine(FPaths::ProjectPlatformExtensionDir(*PlatformDir), TEXT("Config/DataDrivenPlatformInfo.ini"));
 			if (IFileManager::Get().FileExists(*IniPath))
 			{
-				DataDrivenIniFilenames.Add(IniPath);
+				DataDrivenIniFilenames.AddUnique(IniPath); // need AddUnique because if there's no project specified then ProjectConfigDir will be EngineConfigDir
 			}
 		}
 	}
