Fixed ensure macros breaking in line for Dev builds

diff --git a/Engine/Source/Runtime/Core/Public/Misc/AssertionMacros.h b/Engine/Source/Runtime/Core/Public/Misc/AssertionMacros.h
--- a/Engine/Source/Runtime/Core/Public/Misc/AssertionMacros.h
+++ b/Engine/Source/Runtime/Core/Public/Misc/AssertionMacros.h
@@ -391,12 +391,6 @@ RetType FORCENOINLINE UE_DEBUG_SECTION DispatchCheckVerify(InnerType&& Inner, Ar
 	
 	CORE_API bool UE_DEBUG_SECTION ExecCheckImplInternal(std::atomic<bool>& bExecuted, bool bAlways, const ANSICHAR* File, int32 Line, const ANSICHAR* Expr);
 
-	FORCEINLINE bool DebugBreak()
-	{
-		PLATFORM_BREAK();
-		return false;
-	}
-
 	} // UE::Assert::Private
 
 	struct UE_DEPRECATED(5.4, "Do not use directly. This internal type is being removed.") FValidateArgsInternal
@@ -411,7 +405,7 @@ RetType FORCENOINLINE UE_DEBUG_SECTION DispatchCheckVerify(InnerType&& Inner, Ar
 	#define UE_ENSURE_IMPL(Always, InExpression) \
 		(LIKELY(!!(InExpression)) \
 			|| (::UE::Assert::Private::ExecCheckImplInternal([]() UE_DEBUG_SECTION -> std::atomic<bool>& { static std::atomic<bool> bExecuted = false; return bExecuted; } (), Always, __FILE__, __LINE__, #InExpression) \
-			&& ::UE::Assert::Private::DebugBreak()))
+			&& [] () { PLATFORM_BREAK(); return false; } ()))
 
 	#define UE_ENSURE_IMPL2(Capture, Always, InExpression, InFormat, ...) \
 		(LIKELY(!!(InExpression)) || ([Capture] () UE_DEBUG_SECTION \
@@ -419,8 +413,12 @@ RetType FORCENOINLINE UE_DEBUG_SECTION DispatchCheckVerify(InnerType&& Inner, Ar
 			UE_VALIDATE_FORMAT_STRING(InFormat, ##__VA_ARGS__); \
 			static std::atomic<bool> bExecuted = false; \
 			static constexpr ::UE::Assert::Private::FStaticEnsureRecord ENSURE_Static(InFormat, #InExpression, __builtin_FILE(), __builtin_LINE(), Always); \
-			return (Always || !bExecuted.load(std::memory_order_relaxed)) && FPlatformMisc::IsEnsureAllowed() && ::UE::Assert::Private::EnsureFailed(bExecuted, &ENSURE_Static, ##__VA_ARGS__); \
-		}() && ::UE::Assert::Private::DebugBreak()))
+			if ((Always || !bExecuted.load(std::memory_order_relaxed)) && FPlatformMisc::IsEnsureAllowed() && ::UE::Assert::Private::EnsureFailed(bExecuted, &ENSURE_Static, ##__VA_ARGS__)) \
+			{ \
+				PLATFORM_BREAK(); \
+			} \
+			return false; \
+		} ()))
 
 	#define ensure(           InExpression                ) UE_ENSURE_IMPL (   false, InExpression)
 	#define ensureMsgf(       InExpression, InFormat, ... ) UE_ENSURE_IMPL2(&, false, InExpression, InFormat, ##__VA_ARGS__)
