- Now adding scalabiliy group (sg.xxx) to the preview mode, so code that queries sg. cvar values will get the target platform's version of the scalability setting
- Fixed the shadowed value for other platform CVars so that reading another platform value on non-game thread works
- Added new Cvar set flag (ECVF_Set_SetOnly_Unsafe) that is similar to the existing Unsafe flag, but helps with the two above fixes

diff --git a/Engine/Source/Runtime/Core/Private/HAL/ConsoleManager.cpp b/Engine/Source/Runtime/Core/Private/HAL/ConsoleManager.cpp
--- a/Engine/Source/Runtime/Core/Private/HAL/ConsoleManager.cpp
+++ b/Engine/Source/Runtime/Core/Private/HAL/ConsoleManager.cpp
@@ -212,8 +212,14 @@ public:
 		return bRet;
 	}
 
-	void OnChanged(EConsoleVariableFlags SetBy, bool bForce)
+	virtual void OnChanged(EConsoleVariableFlags SetBy, bool bForce)
 	{
+		// we don't want any of this if SetOnly is used
+		if (SetBy & ECVF_Set_SetOnly_Unsafe)
+		{
+			return;
+		}
+		
 		// SetBy can include set flags. Discard them here
 		SetBy = EConsoleVariableFlags(SetBy & ~ECVF_SetFlagMask);
 
@@ -357,6 +363,13 @@ private: // -----------------------------------------
 template <class T>
 void OnCVarChange(T& Dst, const T& Src, EConsoleVariableFlags Flags, EConsoleVariableFlags SetBy)
 {
+    // for the SetOnly case, just copy over the source to the dest
+    if (SetBy & ECVF_Set_SetOnly_Unsafe)
+    {
+        Dst = Src;
+        return;
+    }
+    
 	FConsoleManager& ConsoleManager = (FConsoleManager&)IConsoleManager::Get();
 
 #if WITH_RELOAD
@@ -386,7 +399,7 @@ void OnCVarChange(T& Dst, const T& Src, EConsoleVariableFlags Flags, EConsoleVar
 		check(0);
 	}
 
-	if ((SetBy & ECVF_Set_NoSinkCall_Unsafe)== 0)
+	if ((SetBy & ECVF_Set_NoSinkCall_Unsafe) == 0)
 	{
 		ConsoleManager.OnCVarChanged();
 	}
@@ -604,10 +617,9 @@ bool IConsoleManager::VisitPlatformCVarsForEmulation(FName PlatformName, const F
 						}
 					}
 				}
-				else
-				{
-					VisitIfAllowed(Key, Value, (EConsoleVariableFlags)(SectionPair.SetBy | PreviewFlag));
-				}
+
+				// run the callback with all cvars, even scalbility groups
+				VisitIfAllowed(Key, Value, (EConsoleVariableFlags)(SectionPair.SetBy | PreviewFlag));
 			}
 
 			// clean up the temp section we made
@@ -838,8 +850,9 @@ protected:
 			// update value
 			T ConvertedValue;
 			TTypeFromString<T>::FromString(ConvertedValue, UE::ConfigUtilities::ConvertValueFromHumanFriendlyValue(InValue));
-			SetInternal(ConvertedValue, SetBy);
-			
+			// set the value, and push to render thread value as well, but don't trigger callbacks and don't check priorties
+			SetInternalAndUpdateState(ConvertedValue, (EConsoleVariableFlags)(SetBy | ECVF_Set_SetOnly_Unsafe));
+
 			// update the setby
 			SetFlags((EConsoleVariableFlags)((GetFlags() & ~ECVF_SetByMask) | NewSetBy));
 		}
@@ -927,6 +940,15 @@ protected:
 			EConsoleVariableFlags NewSetBy;
 			auto MaxValue = PriorityHistory->GetMaxValue(NewSetBy);
 			
+			// when we preview SGs, we set their value, but dont run the callbacks, so here we are doing the same operation
+			if (GetFlags() & ECVF_ScalabilityGroup)
+			{
+				NewSetBy = (EConsoleVariableFlags)(NewSetBy | ECVF_Set_SetOnly_Unsafe);
+			}
+			
+			UE_LOG(LogConsoleManager, Display, TEXT(" |-> Unsetting %s = %s"), *IConsoleManager::Get().FindConsoleObjectName(this),
+				*TTypeToString<T>::ToString(MaxValue.GetValueOnGameThread()));
+
 			// and force it to the new value and call any set callbacks
 			SetInternalAndUpdateState(MaxValue.GetValueOnGameThread(), NewSetBy);
 		}
@@ -1338,12 +1360,9 @@ private: // ----------------------------------------------------
 
 	void OnChanged(EConsoleVariableFlags SetBy, bool bForce=false)
 	{
-		if(CanChange(SetBy))
-		{
-			// propagate from main thread to render thread or to reference
-			OnCVarChange(RefValue, MainValue, Flags, SetBy);
-			FConsoleVariableBase::OnChanged(SetBy, bForce);
-		}
+		// propagate from main thread to render thread or to reference
+		OnCVarChange(RefValue, MainValue, Flags, SetBy);
+		FConsoleVariableBase::OnChanged(SetBy, bForce);
 	}
 };
 
@@ -1442,12 +1461,9 @@ private: // ----------------------------------------------------
 	
 	void OnChanged(EConsoleVariableFlags SetBy, bool bForce=false)
 	{
-		if (CanChange(SetBy))
-		{
-			// propagate from main thread to render thread or to reference
-			OnCVarChange(RefValue, MainValue, Flags, SetBy);
-			FConsoleVariableBase::OnChanged(SetBy, bForce);
-		}
+		// propagate from main thread to render thread or to reference
+		OnCVarChange(RefValue, MainValue, Flags, SetBy);
+		FConsoleVariableBase::OnChanged(SetBy, bForce);
 	}
 };
 
@@ -2977,16 +2993,23 @@ void FConsoleManager::PreviewPlatformCVars(FName PlatformName, const FString& De
 	{
 		if (IConsoleVariable* CVar = Pair.Value->AsVariable())
 		{
-			// we want Preview but not Scalability or Cheat
-			if ((CVar->GetFlags() & (ECVF_Preview | ECVF_ScalabilityGroup | ECVF_Cheat)) == ECVF_Preview)
+			// we want Preview but not Cheat
+			if ((CVar->GetFlags() & (ECVF_Preview | ECVF_Cheat)) == ECVF_Preview)
 			{
+				EConsoleVariableFlags Flags = ECVF_SetByPreview;
+				if (CVar->GetFlags() & ECVF_ScalabilityGroup)
+				{
+					// we want to set SG cvars so they can be queried, but not send updates so that we don't use host platform's cvars
+					Flags = (EConsoleVariableFlags)(Flags | ECVF_Set_SetOnly_Unsafe);
+				}
+				
 				// if we have a value for the platform, then set it in the real CVar
-				if (CVar->HasPlatformValueVariable(PlatformKey, TEXT("/")))
+				if (CVar->HasPlatformValueVariable(PlatformKey, GSpecialDPNameForPremadePlatformKey))
 				{
-					TSharedPtr<IConsoleVariable> PlatformCVar = CVar->GetPlatformValueVariable(PlatformKey, TEXT("/"));
-					CVar->Set(*PlatformCVar->GetString(), ECVF_SetByPreview, PreviewModeTag);
+					TSharedPtr<IConsoleVariable> PlatformCVar = CVar->GetPlatformValueVariable(PlatformKey, GSpecialDPNameForPremadePlatformKey);
+					CVar->Set(*PlatformCVar->GetString(), Flags, PreviewModeTag);
 					
-					UE_LOG(LogConsoleManager, Display, TEXT("  |-> %s = %s"), *Pair.Key, *PlatformCVar->GetString());
+					UE_LOG(LogConsoleManager, Display, TEXT("  |-> Setting %s = %s"), *Pair.Key, *PlatformCVar->GetString());
 				}
 			}
 		}
diff --git a/Engine/Source/Runtime/Core/Public/HAL/IConsoleManager.h b/Engine/Source/Runtime/Core/Public/HAL/IConsoleManager.h
--- a/Engine/Source/Runtime/Core/Public/HAL/IConsoleManager.h
+++ b/Engine/Source/Runtime/Core/Public/HAL/IConsoleManager.h
@@ -124,6 +124,9 @@ enum EConsoleVariableFlags
 
 	// Use to set a cvar without calling all cvar sinks. Much faster, but potentially unsafe. Use only if you know the particular cvar/setting does not require a sink call
 	ECVF_Set_NoSinkCall_Unsafe =	0x00010000,
+	
+	// Similar to ECVF_Set_NoSinkCall_Unsafe, but this will do nothing but set the cvar, and push updates to render thread, no changing flags, no sinks, etc
+	ECVF_Set_SetOnly_Unsafe = 		0x00020000,
 
 	// ------------------------------------------------
 
diff --git a/Engine/Source/Runtime/Engine/Private/DeviceProfiles/DeviceProfile.cpp b/Engine/Source/Runtime/Engine/Private/DeviceProfiles/DeviceProfile.cpp
--- a/Engine/Source/Runtime/Engine/Private/DeviceProfiles/DeviceProfile.cpp
+++ b/Engine/Source/Runtime/Engine/Private/DeviceProfiles/DeviceProfile.cpp
@@ -390,12 +390,16 @@ void UDeviceProfile::ExpandDeviceProfileCVars()
 	// VisitPlatformCVarsForEmulation can't access Scalability.h, so make sure it doesn't change away from 3, or if it does, to fix up the hardcoded number
 	static_assert(Scalability::DefaultQualityLevel == 3, "If this trips, update this and IConsoleManager::VisitPlatformCVarsForEmulation with the new value!");
 
-	TMap<FString, FString> CVarsToAdd;
-	TMap<FString, FString> PreviewCVars;
 	IConsoleManager::VisitPlatformCVarsForEmulation(*DeviceType, GetName(),
 		[this](const FString& CVarName, const FString& CVarValue, EConsoleVariableFlags SetByAndPreview)
 		{
-			AllExpandedCVars.Add(CVarName, CVarValue);
+			// don't add scalabiliy groups to the expanded, but do add them to the preview set (this is to maintain same 
+			// functionality for GetAllExpandedCVars(), but allow to see what Preview will set the SGs to in the SetByPreview mode)
+			if (!CVarName.StartsWith(TEXT("sg.")))
+			{
+				AllExpandedCVars.Add(CVarName, CVarValue);
+			}
+
 			if (SetByAndPreview & EConsoleVariableFlags::ECVF_Preview)
 			{
 				AllPreviewCVars.Add(CVarName, CVarValue);
