Fixed some crash reporter issues:
- use BaseDir() instead of ProjectDir(), as ProjectDir could return a relative path (unsuitable for drive information)
- Added missing fields in the crash reporter context that is sent on abnormal termination. Some fields were added recently but were never synced/sent by the crash reporter client.
- Fix an issue that prevented listing of partitions. The Wbem Next() method can return a value that is not 0 but also not a failure.

#rnx

diff --git a/Engine/Source/Programs/CrashReportClient/Private/CrashReportClientApp.cpp b/Engine/Source/Programs/CrashReportClient/Private/CrashReportClientApp.cpp
--- a/Engine/Source/Programs/CrashReportClient/Private/CrashReportClientApp.cpp
+++ b/Engine/Source/Programs/CrashReportClient/Private/CrashReportClientApp.cpp
@@ -896,6 +896,9 @@ static bool LoadTempCrashContextFromFile(FSharedCrashContext& CrashContext, uint
 	FindAndParseValue(ContextProperties, TEXT("IsInternalBuild"), SessionContext.bIsInternalBuild);
 	FindAndParseValue(ContextProperties, TEXT("IsPerforceBuild"), SessionContext.bIsPerforceBuild);
 	FindAndParseValue(ContextProperties, TEXT("IsSourceDistribution"), SessionContext.bIsSourceDistribution);
+	FindAndCopyValue(ContextProperties, TEXT("EngineVersion"), SessionContext.EngineVersion);
+	FindAndCopyValue(ContextProperties, TEXT("EngineCompatibleVersion"), SessionContext.EngineCompatibleVersion);
+	FindAndCopyValue(ContextProperties, TEXT("BuildVersion"), SessionContext.BuildVersion);
 	FindAndCopyValue(ContextProperties, TEXT("GameName"), SessionContext.GameName);
 	FindAndCopyValue(ContextProperties, TEXT("ExecutableName"), SessionContext.ExecutableName);
 	FindAndCopyValue(ContextProperties, TEXT("GameSessionID"), SessionContext.GameSessionID);
diff --git a/Engine/Source/Runtime/Core/Private/Windows/WindowsPlatformCrashContext.cpp b/Engine/Source/Runtime/Core/Private/Windows/WindowsPlatformCrashContext.cpp
--- a/Engine/Source/Runtime/Core/Private/Windows/WindowsPlatformCrashContext.cpp
+++ b/Engine/Source/Runtime/Core/Private/Windows/WindowsPlatformCrashContext.cpp
@@ -30,6 +30,7 @@
 #include "Templates/UniquePtr.h"
 #include "Templates/UnrealTemplate.h"
 #include "Windows/WindowsPlatformMisc.h"
+#include "Windows/WindowsPlatformProcess.h"
 #include "Windows/WindowsPlatformStackWalk.h"
 #include <atomic>
 #include <signal.h>
@@ -214,14 +215,14 @@ void FWindowsPlatformCrashContext::AddPlatformSpecificProperties() const
 	AddCrashProperty(TEXT("PlatformIsRunningWindows"), 1);
 	AddCrashProperty(TEXT("IsRunningOnBattery"), FPlatformMisc::IsRunningOnBattery());
 	WIDECHAR DriveName = 0;
-	const TCHAR* ProjectDir = FGenericPlatformMisc::ProjectDir();
-	if (ProjectDir && *ProjectDir)
+	const TCHAR* BaseDir = FWindowsPlatformProcess::BaseDir();
+	if (BaseDir && *BaseDir)
 	{
-		FPlatformString::Convert(&DriveName, 1, ProjectDir, 1);
+		FPlatformString::Convert(&DriveName, 1, BaseDir, 1);
 		const FPlatformDriveStats* DriveStats = FWindowsPlatformMisc::GetDriveStats(DriveName);
 		if (DriveStats)
 		{
-			AddCrashProperty(TEXT("DriveStats.Project.Name"), ProjectDir);
+			AddCrashProperty(TEXT("DriveStats.Project.Name"), BaseDir);
 			AddCrashProperty(TEXT("DriveStats.Project.Type"), LexToString(DriveStats->DriveType));
 			AddCrashProperty(TEXT("DriveStats.Project.FreeSpaceKb"), DriveStats->FreeBytes/1024);
 		}
diff --git a/Engine/Source/Runtime/Core/Private/Windows/WindowsPlatformMisc.cpp b/Engine/Source/Runtime/Core/Private/Windows/WindowsPlatformMisc.cpp
--- a/Engine/Source/Runtime/Core/Private/Windows/WindowsPlatformMisc.cpp
+++ b/Engine/Source/Runtime/Core/Private/Windows/WindowsPlatformMisc.cpp
@@ -560,7 +560,7 @@ namespace
 			ULONG Returned;
 			IWbemClassObject* StorageWbemObject = nullptr;
 			hres = StorageEnumerator->Next(WBEM_INFINITE, 1, &StorageWbemObject, &Returned);
-			if (Returned == 0 || hres != S_OK)
+			if (Returned == 0 || FAILED(hres))
 			{
 				break;
 			}
