SlackReport debug command for TArray slack tracking.
* Requires ENABLE_ARRAY_SLACK_TRACKING define to be enabled at compile time (enabling the tracking all the time would be too significant a performance hit, as it requires a function call per array count change).
* Generates TSV (tab separated value) reports in Shared/Logs/SlackReport, which can be loaded into a spreadsheet program.  A filename can be specified to the command, or a default one with an incrementing index will be used.
* Switch -Stack=N controls the number of stack frames used for sorting.  Setting a value less than the maximum (9 levels) can coalesce instances of the same allocation from slightly different call stacks.  Switch -Verbose=0 enables a more condensed report that only shows the largest unique Num / Max bucket per call stack.

diff --git a/Engine/Plugins/Experimental/Mutable/Source/MutableRuntime/Public/MuR/MemoryTrackingAllocationPolicy.h b/Engine/Plugins/Experimental/Mutable/Source/MutableRuntime/Public/MuR/MemoryTrackingAllocationPolicy.h
--- a/Engine/Plugins/Experimental/Mutable/Source/MutableRuntime/Public/MuR/MemoryTrackingAllocationPolicy.h
+++ b/Engine/Plugins/Experimental/Mutable/Source/MutableRuntime/Public/MuR/MemoryTrackingAllocationPolicy.h
@@ -186,6 +186,17 @@ namespace mu
 			{
 				return Base.GetInitialCapacity();
 			}
+
+#if UE_ENABLE_ARRAY_SLACK_TRACKING
+			FORCEINLINE void SlackTrackerLogNum(SizeType NewNumUsed)
+			{
+				if constexpr (TAllocatorTraits<BaseAlloc>::SupportsSlackTracking)
+				{
+					Base.SlackTrackerLogNum(NewNumUsed);
+				}
+			}
+#endif
+
 		private:
 			typename BaseAlloc::ForAnyElementType Base;
 			SSIZE_T AllocSize = 0;
@@ -237,6 +248,7 @@ template<typename BaseAlloc, typename Counter>
 struct TAllocatorTraits<mu::FMemoryTrackingAllocatorWrapper<BaseAlloc, Counter>> : public TAllocatorTraitsBase<mu::FMemoryTrackingAllocatorWrapper<BaseAlloc, Counter>>
 {
 	enum { SupportsElementAlignment = TAllocatorTraits<BaseAlloc>::SupportsElementAlignment };
+	enum { SupportsSlackTracking = TAllocatorTraits<BaseAlloc>::SupportsSlackTracking };
 	enum { SupportsMoveFromOtherAllocator = false };
 };
 
diff --git a/Engine/Source/Runtime/Core/Private/HAL/ConsoleManager.cpp b/Engine/Source/Runtime/Core/Private/HAL/ConsoleManager.cpp
--- a/Engine/Source/Runtime/Core/Private/HAL/ConsoleManager.cpp
+++ b/Engine/Source/Runtime/Core/Private/HAL/ConsoleManager.cpp
@@ -3270,6 +3270,12 @@ void CreateConsoleVariables()
 	IConsoleManager::Get().RegisterConsoleCommand(TEXT("GPUDebugCrash"), TEXT("Crash GPU intentionally for debugging."), ECVF_Cheat);
 #endif
 
+#if UE_ENABLE_ARRAY_SLACK_TRACKING
+	IConsoleManager::Get().RegisterConsoleCommand(TEXT("SlackReport"),
+		TEXT("Generate an array slack memory report to Saved/Logs/SlackReport.  TSV format can be loaded as a spreadsheet.\nUage: SlackReport [Filename] [-Stack=N] [-Verbose=0,1]\nIf no filename, writes to a default filename which increments each report.\nStack setting specifies number of stack frames to consider when grouping allocations"),
+		ECVF_Default);
+#endif
+
 #if	!UE_BUILD_SHIPPING
 	IConsoleManager::Get().RegisterConsoleCommand( TEXT( "DumpConsoleCommands" ), TEXT( "Dumps all console vaiables and commands and all exec that can be discovered to the log/console" ), ECVF_Default );
 	IConsoleManager::Get().RegisterConsoleCommand( TEXT( "RedirectToFile" ),
diff --git a/Engine/Source/Runtime/Core/Private/HAL/LowLevelMemTracker.cpp b/Engine/Source/Runtime/Core/Private/HAL/LowLevelMemTracker.cpp
--- a/Engine/Source/Runtime/Core/Private/HAL/LowLevelMemTracker.cpp
+++ b/Engine/Source/Runtime/Core/Private/HAL/LowLevelMemTracker.cpp
@@ -23,6 +23,894 @@
 #include "Templates/Atomic.h"
 #include "Trace/Trace.inl"
 
+#if UE_ENABLE_ARRAY_SLACK_TRACKING
+
+// Specifies whether to generate the whole log file in memory before writing.  Switch uses an async thread for file writing,
+// which does array allocations, deadlocking on the critical section.  Writing to memory first avoids a dependency on the
+// async thread, as we can write the memory to a file at the end, after the lock is released.
+#define ARRAY_SLACK_LOG_TO_MEMORY !PLATFORM_WINDOWS
+
+// Useful values to set in the debugger, to set a breakpoint on a particular allocation tag, element size, and max count.  Sometimes this can
+// give more context regarding an allocation than what you get from a stack trace alone.  As an example, the constructor for UAudioCaptureComponent,
+// which allocated a 1.83 MB array, just showed up as "UClass::CreateDefaultObject" in the stack trace.  The constructor for the specific subclass
+// was optimized out, and there was no way to tell from the slack report what it was actually related to.  Stopping on the allocation in the debugger
+// makes it immediately obvious, because you can see that the UClass in question is UAudioCaptureComponent.  These debug values also let you stop
+// on places where the array count changes without triggering an allocation (the tracking only grabs call stacks on reallocation).
+int32 GArraySlackTagToWatch = -1;		// (int32)ELLMTag::UObject;
+int32 GArraySlackSizeToWatch = 0;
+int32 GArraySlackMaxToWatch = 0;
+
+uint32 GArraySlackDumpIndex = 0;				// Incremented for each auto-generated report file
+bool GArraySlackInit = false;					// Set when we start tracking slack -- startup constructor allocations add a lot of noise (start this as true if you want these)
+bool GArraySlackFirstStackOnly = true;			// Only do stack trace on first stack for a given allocation -- faster, but could be useful to know last allocation
+bool GArraySlackGroupByTag = false;				// Group slack by tag when running with -llm
+bool GArraySlackDefaultVerbose = true;			// Whether to default to verbose output, can be overridden with -Verbose=[0,1]
+
+// We require a minimum number of total bytes for a group of allocations with the same stack trace to be reported.  A setting of 64 discards 80% of
+// allocations representing less than 0.2% of the slack memory.  If you want to look at aggregations of smaller allocations, you can use -Stack=N to
+// trim off more of the call stack, which will make more allocations alias to the same stack trace, and show up in the report (this is probably what
+// you would do anyway when investigating that scenario).  Or you could locally set this to zero to get everything.
+//
+// Console platforms have far slower symbol lookup than PC, making slack reports take a couple orders of magnitude longer to generate than on PC,
+// so we use more conservative settings.  Without some sort of trimming of the generated results, a slack report can take over an hour, which just
+// isn't useful (it still can take 15 minutes with these settings, versus around 20 seconds on Windows for a much larger report).
+#if PLATFORM_WINDOWS
+static const int32 GArraySlackThreshold = 64;				// Typically covers 99.8% of slack memory
+static const int32 GArraySlackFullStackNum = MAX_int32;		// Show full call stack context for all allocations
+static const int32 GArraySlackDefaultStackDepth = 9;		// Sort by this deep in the call stack -- matches max call stack depth in FArraySlackTrackingHeader structure
+#else
+static const int32 GArraySlackThreshold = 8192;				// Typically covers 95% of slack memory
+static const int32 GArraySlackFullStackNum = 150;			// Show full call stack context for this many allocations
+static const int32 GArraySlackDefaultStackDepth = 5;		// Sort by this deep in the call stack
+#endif
+
+std::atomic<int64> GArrayMaxByTag[256];
+std::atomic<int64> GArrayUsedByTag[256];
+int64 GArraySlackByTag[256];
+int32 GArrayCountByTag[256];
+
+// Doubly linked list of all tracked allocations, and critical section to protect it.  Critical section is a pointer so we can detect if it's constructed,
+// otherwise you get crashes in startup constructors which run before the constructor of the critical section is called.  Initialized in the LLM tracker,
+// which gets initialized by the first startup code that uses an LLM scope (even when -llm is disabled), which tends to be pretty early -- allocations
+// before that just won't be tracked.  If this became an issue, we could atomically initialize the lock on first access from any thread, but we already
+// defer tracking until engine PreInit anyway to factor out noise from the myriad static global FString constructors, so it doesn't matter as it stands now.
+FArraySlackTrackingHeader* GTrackArrayDetailedList;
+FCriticalSection* GTrackArrayDetailedLock;
+
+// Dummy function to set a breakpoint on where it's called, if you want to investigate code related to a certain allocation
+FORCENOINLINE void LlmTrackSetBreakpointHere()
+{
+	// Need something non-empty that won't compile out
+	static int32 Dummy = 0;
+	Dummy++;
+}
+
+void FArraySlackTrackingHeader::AddAllocation()
+{
+	if (ArrayNum != INDEX_NONE)
+	{
+		GArrayMaxByTag[Tag].fetch_add(ArrayMax * (int64)ElemSize, std::memory_order_relaxed);
+		GArrayUsedByTag[Tag].fetch_add(ArrayNum * (int64)ElemSize, std::memory_order_relaxed);
+
+		// This code is only reached for reallocations, since during the initial allocation, ArrayNum won't have been set yet.
+		ReallocCount++;
+	}
+
+	if (GArraySlackFirstStackOnly == false || NumStackFrames == 0)
+	{
+		// Skip the first 3 stack frames, which are tracking code (CaptureStackBackTrace, LlmTrackArrayAddAllocation, FArraySlackTrackingHeader::Realloc)
+		constexpr int8 SkipStackFrames = 3;
+		uint64 StackFrameTemp[UE_ARRAY_COUNT(StackFrames) + SkipStackFrames];
+		NumStackFrames = (int8)FPlatformStackWalk::CaptureStackBackTrace(StackFrameTemp, UE_ARRAY_COUNT(StackFrameTemp)) - SkipStackFrames;
+		if (NumStackFrames < 0)
+		{
+			NumStackFrames = 0;
+		}
+		for (int32 StackIndex = 0; StackIndex < NumStackFrames; StackIndex++)
+		{
+			StackFrames[StackIndex] = StackFrameTemp[SkipStackFrames + StackIndex];
+		}
+	}
+
+	if (GTrackArrayDetailedLock && GArraySlackInit)
+	{
+		// For detailed tracking, we add the array header to a doubly linked list
+		FScopeLock Lock(GTrackArrayDetailedLock);
+
+		if (GTrackArrayDetailedList)
+		{
+			GTrackArrayDetailedList->Prev = &Next;
+		}
+		Next = GTrackArrayDetailedList;
+		Prev = &GTrackArrayDetailedList;
+		GTrackArrayDetailedList = this;
+
+		if ((Tag == GArraySlackTagToWatch) &&
+			(ElemSize == GArraySlackSizeToWatch) &&
+			(!GArraySlackMaxToWatch || (ArrayMax == GArraySlackMaxToWatch)))
+		{
+			LlmTrackSetBreakpointHere();
+		}
+
+		GArrayCountByTag[Tag]++;
+	}
+}
+
+void FArraySlackTrackingHeader::RemoveAllocation()
+{
+	if (ArrayNum != INDEX_NONE)
+	{
+		GArrayUsedByTag[Tag].fetch_sub(ArrayNum * (int64)ElemSize, std::memory_order_relaxed);
+		GArrayMaxByTag[Tag].fetch_sub(ArrayMax * (int64)ElemSize, std::memory_order_release);
+	}
+
+	if (Prev)
+	{
+		FScopeLock Lock(GTrackArrayDetailedLock);
+
+		GArrayCountByTag[Tag]--;
+
+		if (Next)
+		{
+			Next->Prev = Prev;
+		}
+		(*Prev) = Next;
+
+		Next = nullptr;
+		Prev = nullptr;
+	}
+}
+
+void FArraySlackTrackingHeader::UpdateNumUsed(int64 NewNumUsed)
+{
+	check(NewNumUsed <= ArrayMax);
+
+	if ((Tag == GArraySlackTagToWatch) &&
+		(ElemSize == GArraySlackSizeToWatch) &&
+		(!GArraySlackMaxToWatch || (ArrayMax == GArraySlackMaxToWatch)))
+	{
+		LlmTrackSetBreakpointHere();
+	}
+
+	// Track the allocation in our totals when ArrayNum is first set to something other than INDEX_NONE.  This allows us to
+	// factor out container allocations that aren't arrays (mainly hash tables), which won't ever call "UpdateNumUsed".
+	if (ArrayNum == INDEX_NONE)
+	{
+		GArrayMaxByTag[Tag].fetch_add(ArrayMax * (int64)ElemSize, std::memory_order_relaxed);
+		ArrayNum = 0;
+		FirstAllocFrame = (uint32)GFrameCounter;
+	}
+	GArrayUsedByTag[Tag].fetch_add((NewNumUsed - ArrayNum) * (int64)ElemSize, std::memory_order_relaxed);
+	ArrayNum = NewNumUsed;
+	ArrayPeak = FMath::Max(ArrayPeak, (uint32)FMath::Min(NewNumUsed, 0xffffffffll));
+}
+
+FORCENOINLINE void* FArraySlackTrackingHeader::Realloc(void* Ptr, int64 Count, uint64 ElemSize, int32 Alignment)
+{
+	// Figure out how much padding we need under the allocation
+	int32 HeaderAlign = FPlatformMath::RoundUpToPowerOfTwo(sizeof(FArraySlackTrackingHeader));
+	int32 PaddingRequired = HeaderAlign > Alignment ? HeaderAlign : Alignment;
+
+	// Get the base pointer of the original allocation, and remove tracking for it
+	if (Ptr)
+	{
+		FArraySlackTrackingHeader* TrackingHeader = (FArraySlackTrackingHeader*)((uint8*)Ptr - sizeof(FArraySlackTrackingHeader));
+		TrackingHeader->RemoveAllocation();
+
+		Ptr = (uint8*)TrackingHeader - TrackingHeader->AllocOffset;
+	}
+
+	uint8* ResultPtr = nullptr;
+	if (Count)
+	{
+		ResultPtr = (uint8*)FMemory::Realloc(Ptr, Count * ElemSize + PaddingRequired, Alignment);
+		ResultPtr += PaddingRequired;
+		FArraySlackTrackingHeader* TrackingHeader = (FArraySlackTrackingHeader*)(ResultPtr - sizeof(FArraySlackTrackingHeader));
+
+		// Set the tag and other default information in the allocation if it's newly created
+		if (!Ptr)
+		{
+			// Note that we initially set the slack tracking ArrayNum to INDEX_NONE.  The container allocator is used by both arrays and
+			// other containers (Set / Map / Hash), and we don't know it's actually an array until "UpdateNumUsed" is called on it.
+			check(PaddingRequired <= 65536);
+			TrackingHeader->Next = nullptr;
+			TrackingHeader->Prev = nullptr;
+			TrackingHeader->AllocOffset = (uint16)(PaddingRequired - sizeof(FArraySlackTrackingHeader));
+			TrackingHeader->Tag = LlmGetActiveTag();
+			TrackingHeader->NumStackFrames = 0;			// Filled in later...
+			TrackingHeader->FirstAllocFrame = 0;		// Filled in later...
+			TrackingHeader->ReallocCount = 0;
+			TrackingHeader->ArrayPeak = 0;
+			TrackingHeader->ElemSize = ElemSize;
+			TrackingHeader->ArrayNum = INDEX_NONE;		// Set in UpdateNumUsed
+		}
+
+		// Update ArrayMax and re-register the allocation
+		TrackingHeader->ArrayMax = Count;
+		TrackingHeader->AddAllocation();
+	}
+	else
+	{
+		FMemory::Free(Ptr);
+	}
+
+	return ResultPtr;
+}
+
+struct FArraySlackSortItem
+{
+	FArraySlackTrackingHeader* Header;
+	FName CustomName;
+	uint64 StackTraceTotalSlack;			// Sum of slack for elements with the same stack trace
+	uint32 StackTraceRunLength;				// Run of items with the same stack trace
+	uint32 RunLength;						// Run of identical elements (elemsize, num, max the same)
+	int32 StackTraceIgnore;					// Number of stack trace items to ignore
+
+	bool EqualsTagStackTrace(const FArraySlackSortItem& Other, uint32 StackTraceDepth, bool bLlmEnabled) const
+	{
+		if (bLlmEnabled)
+		{
+			if (Header->Tag != Other.Header->Tag)
+			{
+				return false;
+			}
+			if (CustomName != Other.CustomName)
+			{
+				return false;
+			}
+		}
+
+		// If the stack depth is set to zero, and LLM is disabled, basically everything in the capture will get lumped into
+		// one bucket.  Comparing by ElemSize is a last resort to force some differentiation in the report in that case (or
+		// perhaps if we have a platform that doesn't support stack traces, or someone wants to locally disable them for
+		// performance).  In cases where a stack frame exists, the leaf stack frame is always some sort of template type
+		// specific resize (i.e. TArray<float>::ResizeTo), so this would be redundant, but also harmless.
+		if (Header->ElemSize != Other.Header->ElemSize)
+		{
+			return false;
+		}
+
+		uint32 NumStackFramesThis = FMath::Min((uint32)(Header->NumStackFrames - StackTraceIgnore), StackTraceDepth);
+		uint32 NumStackFramesOther = FMath::Min((uint32)(Other.Header->NumStackFrames - Other.StackTraceIgnore), StackTraceDepth);
+
+		if (NumStackFramesThis != NumStackFramesOther)
+		{
+			return false;
+		}
+
+		return FMemory::Memcmp(&Header->StackFrames[StackTraceIgnore], &Other.Header->StackFrames[Other.StackTraceIgnore], NumStackFramesThis * sizeof(Header->StackFrames[0])) == 0;
+	}
+
+	bool Compare(const FArraySlackSortItem& Other, uint32 StackTraceDepth, bool bLlmEnabled) const
+	{
+		// Order by decreasing stack trace slack total bytes
+		if (StackTraceTotalSlack != Other.StackTraceTotalSlack)
+		{
+			return StackTraceTotalSlack > Other.StackTraceTotalSlack;
+		}
+
+		// Order by tag
+		if (bLlmEnabled)
+		{
+			if (Header->Tag != Other.Header->Tag)
+			{
+				return Header->Tag < Other.Header->Tag;
+			}
+			if (CustomName != Other.CustomName)
+			{
+				return CustomName.GetComparisonIndex().CompareFast(Other.CustomName.GetComparisonIndex()) < 0;
+			}
+		}
+
+		// Order by increasing element size
+		if (Header->ElemSize != Other.Header->ElemSize)
+		{
+			return Header->ElemSize < Other.Header->ElemSize;
+		}
+
+		// Order by stack trace
+		uint32 NumStackFramesThis = FMath::Min((uint32)(Header->NumStackFrames - StackTraceIgnore), StackTraceDepth);
+		uint32 NumStackFramesOther = FMath::Min((uint32)(Other.Header->NumStackFrames - Other.StackTraceIgnore), StackTraceDepth);
+
+		if (NumStackFramesThis != NumStackFramesOther)
+		{
+			return NumStackFramesThis < NumStackFramesOther;
+		}
+		int32 StackFrameOrdinalCompare = memcmp(&Header->StackFrames[StackTraceIgnore], &Other.Header->StackFrames[Other.StackTraceIgnore], NumStackFramesThis * sizeof(Header->StackFrames[0]));
+		if (StackFrameOrdinalCompare)
+		{
+			return StackFrameOrdinalCompare < 0;
+		}
+
+		int64 SlackBytesA = Header->SlackSizeInBytes();
+		int64 SlackBytesB = Other.Header->SlackSizeInBytes();
+		int64 RunSlackBytesA = SlackBytesA * RunLength;
+		int64 RunSlackBytesB = SlackBytesB * Other.RunLength;
+
+		// Order by decreasing run slack total bytes
+		if (RunSlackBytesA != RunSlackBytesB)
+		{
+			return RunSlackBytesA > RunSlackBytesB;
+		}
+
+		// Order by decreasing individual item slack bytes
+		if (SlackBytesA != SlackBytesB)
+		{
+			return SlackBytesA > SlackBytesB;
+		}
+
+		// Order by increasing Max
+		return Header->ArrayMax < Other.Header->ArrayMax;
+	}
+};
+
+static void LlmTrackArrayDumpTag(FArchive* LogFile, FOutputDevice& Output, TAnsiStringBuilder<4096>& Builder, int32 TagIndex, uint32 StackTraceDepth, bool bVerbose, double StartTime)
+{
+	FLowLevelMemTracker& Tracker = FLowLevelMemTracker::Get();
+	bool bLlmEnabled = Tracker.IsEnabled();
+
+	FScopeLock Lock(GTrackArrayDetailedLock);
+
+	TArray<FArraySlackSortItem> ArraySlackSortArray;
+	int32 ReserveAmount = 0;
+	if (TagIndex == INDEX_NONE)
+	{
+		for (int32 CountByTag : GArrayCountByTag)
+		{
+			ReserveAmount += CountByTag;
+		}
+	}
+	else
+	{
+		ReserveAmount = GArrayCountByTag[TagIndex];
+	}
+	ArraySlackSortArray.Reserve(ReserveAmount);
+	ArraySlackSortArray.GetAllocatorInstance().DisableSlackTracking();
+
+	for (FArraySlackTrackingHeader* Current = GTrackArrayDetailedList; Current; Current = Current->Next)
+	{
+		// Don't bother dumping allocations with zero waste (or untracked where ArrayNum == INDEX_NONE)
+		if ((TagIndex == INDEX_NONE || Current->Tag == TagIndex) && (Current->ArrayNum != INDEX_NONE) && (Current->ArrayNum != Current->ArrayMax))
+		{
+			FArraySlackSortItem& SortItem = ArraySlackSortArray.AddDefaulted_GetRef();
+
+			SortItem.Header = Current;
+			if (bLlmEnabled && (Current->Tag == (uint8)ELLMTag::CustomName))
+			{
+				SortItem.CustomName = Tracker.Get().FindPtrDisplayName((uint8*)Current - Current->AllocOffset);
+			}
+			else
+			{
+				SortItem.CustomName = NAME_None;
+			}
+
+			// Filled in later
+			SortItem.StackTraceTotalSlack = 0;
+			SortItem.StackTraceRunLength = 1;
+			SortItem.RunLength = 1;
+			SortItem.StackTraceIgnore = 0;
+		}
+	}
+
+	// We want to ignore ResizeAllocation() if it's the first stack frame, as it's not interesting.  This stack frame will sometimes be there,
+	// and sometimes not, because the most common ResizeAllocation() template variations are tagged FORCENOINLINE to reduce code size.  We set
+	// StackTraceIgnore=1 to indicate where this is the first stack frame, indicating we can ignore it downstream.  It has to be filled in before
+	// the sort, to properly handle the StackTraceDepth setting.
+	//
+	// Assuming symbol lookups are expensive, we sort by leaf stack frame first so we only need to do symbol lookups for unique leaf stack frames.
+	Algo::SortBy(ArraySlackSortArray, [](const FArraySlackSortItem& Item) { return Item.Header->StackFrames[0]; });
+
+	for (int32 ItemIndex = 0; ItemIndex < ArraySlackSortArray.Num(); ItemIndex++)
+	{
+		int32 StackTraceIgnore = 0;
+		if (ItemIndex == 0 || ArraySlackSortArray[ItemIndex].Header->StackFrames[0] != ArraySlackSortArray[ItemIndex - 1].Header->StackFrames[0])
+		{
+			StackTraceIgnore = 0;
+			if (ArraySlackSortArray[ItemIndex].Header->NumStackFrames)
+			{
+				FProgramCounterSymbolInfo SymbolInfo;
+				FPlatformStackWalk::ProgramCounterToSymbolInfo(ArraySlackSortArray[ItemIndex].Header->StackFrames[0], SymbolInfo);
+				if (FCStringAnsi::Strstr(SymbolInfo.FunctionName, "::ResizeAllocation("))
+				{
+					StackTraceIgnore = 1;
+				}
+			}
+		}
+
+		ArraySlackSortArray[ItemIndex].StackTraceIgnore = StackTraceIgnore;
+	}
+
+	// First pass sort -- we haven't yet filled in totals for StackTraceTotalSlack and RunTotalSlack
+	Algo::Sort(ArraySlackSortArray, [StackTraceDepth, bLlmEnabled](const FArraySlackSortItem& A, const FArraySlackSortItem& B) { return A.Compare(B, StackTraceDepth, bLlmEnabled); });
+
+	int64 IgnoredGroups = 0;
+	int64 IgnoredSlack = 0;
+	{
+		// Compute slack associated with each stack trace and runs of identical elements, and store it on the sort elements
+		int32 StackTraceRun = 0;
+		int64 StackTraceTotal = 0;
+		uint32 ElementRun = 0;
+		for (int32 ItemIndex = 0; ItemIndex < ArraySlackSortArray.Num(); ItemIndex++)
+		{
+			// Add current item
+			int64 ElementSlack = ArraySlackSortArray[ItemIndex].Header->SlackSizeInBytes();
+			StackTraceRun++;
+			StackTraceTotal += ElementSlack;
+			ElementRun++;
+
+			// If the next item has a different stack trace, or it's the end of the array, echo the stack trace total slack to all the items
+			if ((ItemIndex == ArraySlackSortArray.Num() - 1) || !ArraySlackSortArray[ItemIndex].EqualsTagStackTrace(ArraySlackSortArray[ItemIndex + 1], StackTraceDepth, bLlmEnabled))
+			{
+				for (int32 RunIndex = ItemIndex - (StackTraceRun - 1); RunIndex <= ItemIndex; RunIndex++)
+				{
+					ArraySlackSortArray[RunIndex].StackTraceTotalSlack = StackTraceTotal;
+					ArraySlackSortArray[RunIndex].StackTraceRunLength = StackTraceRun;
+				}
+				if (StackTraceTotal < GArraySlackThreshold)
+				{
+					IgnoredGroups++;
+					IgnoredSlack += StackTraceTotal;
+				}
+				StackTraceTotal = 0;
+				StackTraceRun = 0;
+			}
+
+			// Check if the item is different at all, and echo run length to all the items
+			if ((ItemIndex == ArraySlackSortArray.Num() - 1) || ArraySlackSortArray[ItemIndex].Compare(ArraySlackSortArray[ItemIndex + 1], StackTraceDepth, bLlmEnabled))
+			{
+				for (int32 RunIndex = ItemIndex - (ElementRun - 1); RunIndex <= ItemIndex; RunIndex++)
+				{
+					ArraySlackSortArray[RunIndex].RunLength = ElementRun;
+				}
+				ElementRun = 0;
+			}
+		}
+	}
+
+	// Second pass, final sort
+	Algo::Sort(ArraySlackSortArray, [StackTraceDepth, bLlmEnabled](const FArraySlackSortItem& A, const FArraySlackSortItem& B) { return A.Compare(B, StackTraceDepth, bLlmEnabled); });
+
+	// Only include tag column if LLM is enabled
+	const TCHAR* TagColumnSeparator = bLlmEnabled ? TEXT("\t") : TEXT("");
+	const ANSICHAR* TagColumnSeparatorANSI = bLlmEnabled ? "\t" : "";
+
+	if (LogFile)
+	{
+		Builder.Reset();
+		Builder.Appendf("Ignored:\t%lld\tGroups,\t%lld\tBytes total\n", IgnoredGroups, IgnoredSlack);
+		Builder.Appendf("Under:\t%d\tGroup size threshold\n\n", GArraySlackThreshold);
+
+		Builder.Appendf("NumArrays\tReallocs\tLifetime\tPeakAvg\tPeak\tElemSize\tNum\tMax\t%s\tStackSlack%s%s\tStackTrace\n",
+			bVerbose ? "ItemSlack" : "LargestItem",
+			TagColumnSeparatorANSI,
+			bLlmEnabled ? "Tag" : "");
+		LogFile->Serialize(Builder.GetData(), Builder.Len());
+	}
+	else
+	{
+		FPlatformMisc::LowLevelOutputDebugStringf(TEXT("Ignored:\t%lld\tGroups,\t%lld\tBytes total\n"), IgnoredGroups, IgnoredSlack);
+		FPlatformMisc::LowLevelOutputDebugStringf(TEXT("Under:\t%d\tGroup size threshold\n\n"), GArraySlackThreshold);
+
+		// We prepend "SlackReport" to every line when outputting to the debug window, as this can help filtering out random
+		// log lines after the output is cut and pasted.  You can sort by the first column to accomplish that.
+		FPlatformMisc::LowLevelOutputDebugStringf(TEXT("SlackReport\tNumArrays\tReallocs\tLifetime\tPeakAvg\tPeak\tElemSize\tNum\tMax\t%s\tStackSlack%s%s\tStackTrace\n"),
+			bVerbose ? "ItemSlack" : "LargestItem",
+			TagColumnSeparator,
+			bLlmEnabled ? TEXT("Tag") : TEXT(""));
+	}
+
+	{
+		uint32 CurrentFrame = (uint32)GFrameCounter;
+		int32 ItemRun = 0;
+		double ItemReallocs = 0.0;
+		double ItemLifetime = 0.0;
+		double ItemPeakSum = 0.0;
+		uint32 ItemPeak = 0;
+		int32 FullStacksPrinted = 0;
+
+		for (int32 ItemIndex = 0; ItemIndex < ArraySlackSortArray.Num(); ItemIndex++)
+		{
+#if PLATFORM_WINDOWS
+			// Windows is a lot faster, so we don't need as much progress logging
+			constexpr int32 ProgressInterval = 25000;
+#else
+			constexpr int32 ProgressInterval = 1000;
+#endif
+			if ((ItemIndex % ProgressInterval) == 0)
+			{
+				Output.Logf(TEXT("Array Slack %d / %d allocs...  (%.2lf minutes, batch %lld bytes -> threshold %lld)"),
+					ItemIndex, ArraySlackSortArray.Num(), (FPlatformTime::Seconds() - StartTime) / 60.0, ArraySlackSortArray[ItemIndex].StackTraceTotalSlack, GArraySlackThreshold);
+			}
+
+			// Count current item
+			ItemRun++;
+			ItemReallocs += ArraySlackSortArray[ItemIndex].Header->ReallocCount;
+			ItemLifetime += CurrentFrame - ArraySlackSortArray[ItemIndex].Header->FirstAllocFrame;
+			ItemPeakSum += ArraySlackSortArray[ItemIndex].Header->ArrayPeak;
+			ItemPeak = FMath::Max(ItemPeak, ArraySlackSortArray[ItemIndex].Header->ArrayPeak);
+
+			// If this item is different than the next, or the last item, echo the count and item
+			if ((ItemIndex == ArraySlackSortArray.Num() - 1) || ArraySlackSortArray[ItemIndex].Compare(ArraySlackSortArray[ItemIndex + 1], StackTraceDepth, bLlmEnabled))
+			{
+				// Determine if this run has the same stack trace as the previous.  In verbose mode, we only print stack trace specific totals for unique
+				// stack traces, and when not verbose, we only print lines at all for unique stack traces.
+				int32 RunStart = ItemIndex - (ItemRun - 1);
+				bool bUniqueStackTrace = RunStart == 0 || !ArraySlackSortArray[RunStart - 1].EqualsTagStackTrace(ArraySlackSortArray[RunStart], StackTraceDepth, bLlmEnabled);
+
+				if ((bVerbose || bUniqueStackTrace) && (ArraySlackSortArray[ItemIndex].StackTraceTotalSlack >= GArraySlackThreshold))
+				{
+					if (LogFile)
+					{
+						Builder.Reset();
+						Builder.Appendf(
+							"%llu\t%.1lf\t%.1lf\t%.1lf\t%u\t%lld\t%lld\t%lld\t%lld\t",
+							bVerbose ? ItemRun : ArraySlackSortArray[ItemIndex].StackTraceRunLength,
+							ItemReallocs / ItemRun,
+							ItemLifetime / ItemRun,
+							ItemPeakSum / ItemRun,
+							ItemPeak,
+							ArraySlackSortArray[ItemIndex].Header->ElemSize,
+							ArraySlackSortArray[ItemIndex].Header->ArrayNum,
+							ArraySlackSortArray[ItemIndex].Header->ArrayMax,
+							ArraySlackSortArray[ItemIndex].Header->SlackSizeInBytes() * ItemRun);
+
+						// Stack trace total slack if it's unique
+						if (bUniqueStackTrace)
+						{
+							Builder.Appendf("%lld", ArraySlackSortArray[ItemIndex].StackTraceTotalSlack);
+						}
+
+						// Tag name
+						if (bLlmEnabled)
+						{
+							Builder.AppendChar('\t');
+							if (ArraySlackSortArray[ItemIndex].CustomName != NAME_None)
+							{
+								Builder.Append(*ArraySlackSortArray[ItemIndex].CustomName.ToString());
+							}
+							else
+							{
+								Builder.Append(Tracker.FindTagDisplayName(ArraySlackSortArray[ItemIndex].Header->Tag).ToString());
+							}
+						}
+					}
+					else
+					{
+						TStringBuilder<32> ByStackSlack;
+						if (bUniqueStackTrace)
+						{
+							ByStackSlack.Appendf(TEXT("%lld"), ArraySlackSortArray[ItemIndex].StackTraceTotalSlack);
+						}
+
+						FPlatformMisc::LowLevelOutputDebugStringf(
+							TEXT("SlackReport\t%llu\t%.1lf\t%.1lf\t%.1lf\t%lld\t%lld\t%lld\t%lld\t%lld\t%s%s%s"),
+							bVerbose ? ItemRun : ArraySlackSortArray[ItemIndex].StackTraceRunLength,
+							ItemReallocs / ItemRun,
+							ItemLifetime / ItemRun,
+							ItemPeakSum / ItemRun,
+							ItemPeak,
+							ArraySlackSortArray[ItemIndex].Header->ElemSize,
+							ArraySlackSortArray[ItemIndex].Header->ArrayNum,
+							ArraySlackSortArray[ItemIndex].Header->ArrayMax,
+							ArraySlackSortArray[ItemIndex].Header->SlackSizeInBytes() * ItemRun,
+							ByStackSlack.ToString(),
+							TagColumnSeparator,
+							!bLlmEnabled ? TEXT("") :
+							(ArraySlackSortArray[ItemIndex].CustomName != NAME_None ?
+								*ArraySlackSortArray[ItemIndex].CustomName.ToString() :
+								*Tracker.FindTagDisplayName(ArraySlackSortArray[ItemIndex].Header->Tag).ToString()));
+					}
+
+					// Only print stack trace if this is a unique stack trace.
+					if (bUniqueStackTrace)
+					{
+						int32 StackFirst = ArraySlackSortArray[ItemIndex].StackTraceIgnore;
+						int32 StackCount = ArraySlackSortArray[ItemIndex].Header->NumStackFrames - ArraySlackSortArray[ItemIndex].StackTraceIgnore;
+						if (FullStacksPrinted >= GArraySlackFullStackNum)
+						{
+							StackCount = FMath::Min(StackCount, (int32)StackTraceDepth);
+						}
+						else
+						{
+							FullStacksPrinted++;
+						}
+
+						for (int32 StackIndex = StackFirst; StackIndex < StackFirst + StackCount; StackIndex++)
+						{
+							FProgramCounterSymbolInfo SymbolInfo;
+							FPlatformStackWalk::ProgramCounterToSymbolInfo(ArraySlackSortArray[ItemIndex].Header->StackFrames[StackIndex], SymbolInfo);
+
+							if (LogFile)
+							{
+								Builder.AppendChar('\t');
+								Builder.Append(SymbolInfo.FunctionName[0] ? SymbolInfo.FunctionName : "UnknownFunction");
+
+								if (SymbolInfo.Filename[0] && SymbolInfo.LineNumber)
+								{
+									// Format " [Filename:Line]"
+									Builder.Append(" [");
+									Builder.Append(SymbolInfo.Filename);
+									Builder.Appendf(":%i]", SymbolInfo.LineNumber);
+								}
+								else
+								{
+									Builder.Append(" []");
+								}
+							}
+							else
+							{
+								TStringBuilder<512> SymbolName;
+								SymbolName.AppendChar(TEXT('\t'));
+								SymbolName.Append(SymbolInfo.FunctionName[0] ? SymbolInfo.FunctionName : "UnknownFunction");
+
+								if (SymbolInfo.Filename[0] && SymbolInfo.LineNumber)
+								{
+									// Format " [Filename:Line]"
+									SymbolName.Append(TEXT(" ["));
+									SymbolName.Append(SymbolInfo.Filename);
+									SymbolName.Appendf(TEXT(":%i]"), SymbolInfo.LineNumber);
+								}
+								else
+								{
+									SymbolName.Append(TEXT(" []"));
+								}
+
+								FPlatformMisc::LowLevelOutputDebugString(SymbolName.ToString());
+							}
+						}
+					}
+
+					if (LogFile)
+					{
+						Builder.AppendChar('\n');
+						LogFile->Serialize(Builder.GetData(), Builder.Len());
+					}
+					else
+					{
+						FPlatformMisc::LowLevelOutputDebugString(TEXT("\n"));
+					}
+				}
+
+				ItemRun = 0;
+				ItemReallocs = 0.0;
+				ItemLifetime = 0.0;
+				ItemPeakSum = 0.0;
+				ItemPeak = 0;
+			}
+		}
+	}
+}
+
+void ArraySlackTrackInit()
+{
+	// Any array allocations before this is called won't have array slack tracking, although subsequent reallocations of existing arrays
+	// will gain tracking if that occurs.  The goal is to filter out startup constructors which run before Main, which introduce a
+	// ton of noise into slack reports.  Especially the roughly 30,000 static FString constructors in the code base, each with a
+	// unique call stack, and all having a little bit of slack due to malloc bucket size rounding.
+	GArraySlackInit = true;
+}
+
+static void LlmTrackArrayTick()
+{
+	// Updating these every frame is handy, so you can see them in a watch window in the debugger or debug print them without running a capture.
+	// We could consider including these as stats, so you can track them in Insights, but the report is giving enough information for now.
+	for (int32 TagIndex = 0; TagIndex < 256; TagIndex++)
+	{
+		GArraySlackByTag[TagIndex] = GArrayMaxByTag[TagIndex] - GArrayUsedByTag[TagIndex];
+	}
+}
+
+void ArraySlackTrackGenerateReport(const TCHAR* Cmd, FOutputDevice& Output)
+{
+	Output.Logf(TEXT("Generating Array Slack report."));
+
+	double StartTime = FPlatformTime::Seconds();
+
+	// Make sure the slack by tag totals are up to date
+	LlmTrackArrayTick();
+
+	// Parse command -- tokens
+	FString LogFilename;
+	uint32 StackTraceDepth = GArraySlackDefaultStackDepth;
+	bool bVerbose = GArraySlackDefaultVerbose;
+	for (FString Arg = FParse::Token(Cmd, false); !Arg.IsEmpty(); Arg = FParse::Token(Cmd, false))
+	{
+		if (Arg[0] == TEXT('-'))
+		{
+			FStringView StackDepthSwitch(TEXTVIEW("-Stack="));
+			FStringView VerboseSwitch(TEXTVIEW("-Verbose="));
+			if ((Arg.Len() > StackDepthSwitch.Len()) && !FCString::Strnicmp(&Arg[0], StackDepthSwitch.GetData(), StackDepthSwitch.Len()))
+			{
+				StackTraceDepth = (uint32)FCString::Strtoui64(&Arg[StackDepthSwitch.Len()], nullptr, 10);
+			}
+			else if ((Arg.Len() > VerboseSwitch.Len()) && !FCString::Strnicmp(&Arg[0], VerboseSwitch.GetData(), VerboseSwitch.Len()))
+			{
+				bVerbose = Arg[VerboseSwitch.Len()] != TEXT('0');
+			}
+			else
+			{
+				Output.Logf(TEXT("Array Slack unsupported switch: \"%s\".  Valid switches:  -Stack=N, -Verbose=[0,1]"), *Arg);
+			}
+		}
+		else
+		{
+			LogFilename = Arg;
+		}
+	}
+
+	FArchive* LogFile = nullptr;
+
+#if ARRAY_SLACK_LOG_TO_MEMORY
+	TArray<uint8> LogFileMemory;
+	LogFileMemory.Reserve(4 * 1024 * 1024);
+	LogFileMemory.GetAllocatorInstance().DisableSlackTracking();
+#endif
+	FString LogFilenameWithPath;
+
+	// Special name "NOFILE" indicates to write to debug log instead of file.  Useful for debugging the system, as you can see the
+	// lines of text generated while debugging, as opposed to needing to wait until the file gets written to look at the output.
+	if (!LogFilename.Equals(TEXT("NOFILE")))
+	{
+		FString AbsoluteProjectLogDir = FPaths::ConvertRelativePathToFull(FPaths::ProjectLogDir());
+		FString SlackReportLogDir = FPaths::Combine(AbsoluteProjectLogDir, TEXT("SlackReport"));
+		IPlatformFile::GetPlatformPhysical().CreateDirectoryTree(*SlackReportLogDir);
+
+		if (LogFilename.IsEmpty())
+		{
+			LogFilename = FString::Printf(TEXT("SlackDump_%03u.tsv"), GArraySlackDumpIndex++);
+		}
+		else
+		{
+			if (!LogFilename.EndsWith(TEXT(".tsv")))
+			{
+				LogFilename.Append(TEXT(".tsv"));
+			}
+		}
+		LogFilenameWithPath = SlackReportLogDir / LogFilename;
+
+#if ARRAY_SLACK_LOG_TO_MEMORY
+		LogFile = new FMemoryWriter(LogFileMemory);
+#else
+		IFileManager* FileManager = &IFileManager::Get();
+		LogFile = FileManager->CreateFileWriter(*LogFilenameWithPath, 0);
+#endif
+	}
+
+	const FLowLevelMemTracker& Tracker = FLowLevelMemTracker::Get();
+	bool bLlmEnabled = Tracker.IsEnabled();
+
+	TAnsiStringBuilder<4096> Builder;
+
+	TArray<uint64> SortedTags;
+	SortedTags.SetNumZeroed(256);
+	SortedTags.GetAllocatorInstance().DisableSlackTracking();
+
+	// Tag summary report is only useful when -llm is enabled
+	if (bLlmEnabled)
+	{
+		if (LogFile)
+		{
+			Builder.Append("Tag\tSlack\tTotalMem\n");
+			LogFile->Serialize(Builder.GetData(), Builder.Len());
+		}
+		else
+		{
+			FPlatformMisc::LowLevelOutputDebugString(TEXT("SlackByTag\tTag\tSlack\tTotalMem\n"));
+		}
+
+		// Sort tags by descending memory.  Store tag index in low 8 bits, memory in high 56 bits.
+		for (int32 TagIndex = 0; TagIndex < 256; TagIndex++)
+		{
+			SortedTags[TagIndex] = (GArraySlackByTag[TagIndex] << 8) | TagIndex;
+		}
+		SortedTags.Sort(TGreater<uint64>());
+
+		for (int32 SortedIndex = 0; SortedIndex < 256; SortedIndex++)
+		{
+			if (SortedTags[SortedIndex] >= 256)
+			{
+				uint32 TagIndex = (uint32)SortedTags[SortedIndex] & 0xff;
+
+				if (LogFile)
+				{
+					Builder.Reset();
+					Builder.Append(Tracker.FindTagDisplayName(TagIndex).ToString());
+					Builder.Appendf("\t%lld\t%lld\n",
+						GArraySlackByTag[TagIndex],
+						GArrayMaxByTag[TagIndex].load(std::memory_order_relaxed));
+					LogFile->Serialize(Builder.GetData(), Builder.Len());
+				}
+				else
+				{
+					FPlatformMisc::LowLevelOutputDebugStringf(TEXT("SlackByTag\t%s\t%lld\t%lld\n"),
+						*Tracker.FindTagDisplayName(TagIndex).ToString(),
+						GArraySlackByTag[TagIndex],
+						GArrayMaxByTag[TagIndex].load(std::memory_order_relaxed));
+				}
+			}
+		}
+
+		if (LogFile)
+		{
+			char NewLines[] = "\n\n";
+			LogFile->Serialize(NewLines, 2);
+		}
+	}
+
+	if (LogFile)
+	{
+		// Append information about options and timing of dump
+		FArraySlackTrackingHeader Dummy;
+		Builder.Reset();
+		Builder.Appendf("Ran with:\t-Stack=%u -Verbose=%d", FMath::Min((uint32)UE_ARRAY_COUNT(Dummy.StackFrames), StackTraceDepth), bVerbose ? 1 : 0);
+		if (FLowLevelMemTracker::Get().IsEnabled())
+		{
+			Builder.Append(" -llm");
+		}
+		if (!bVerbose)
+		{
+			Builder.Append("\t\t\tFields besides NumArrays / StackSlack are for the largest slack bucket (unique Num / Max combo), run with -Verbose=1 to see all buckets");
+		}
+		Builder.Appendf("\nOn frame:\t%u\n\n", (int32)GFrameCounter);
+		LogFile->Serialize(Builder.GetData(), Builder.Len());
+	}
+
+	if (bLlmEnabled && GArraySlackGroupByTag)
+	{
+		// Original behavior grouped by tag, but all in one batch is generally preferable.  Could expose this with a switch in the future.
+		for (int32 SortedIndex = 0; SortedIndex < 256; SortedIndex++)
+		{
+			if (SortedTags[SortedIndex] > 256)
+			{
+				int32 TagIndex = (int32)SortedTags[SortedIndex] & 0xff;
+				if (GArraySlackByTag[TagIndex])
+				{
+#if ARRAY_SLACK_LOG_TO_MEMORY
+					// Disable tracking each loop iteration, in case the memory grew to the point where it was reallocated in the previous iteration.
+					LogFileMemory.GetAllocatorInstance().DisableSlackTracking();
+#endif
+
+					LlmTrackArrayDumpTag(LogFile, Output, Builder, TagIndex, StackTraceDepth, bVerbose, StartTime);
+				}
+			}
+		}
+	}
+	else
+	{
+		// INDEX_NONE == dump all tags in one batch
+		LlmTrackArrayDumpTag(LogFile, Output, Builder, INDEX_NONE, StackTraceDepth, bVerbose, StartTime);
+	}
+
+	if (LogFile)
+	{
+		delete LogFile;
+
+#if ARRAY_SLACK_LOG_TO_MEMORY
+		// Now create the actual file
+		IFileManager* FileManager = &IFileManager::Get();
+		LogFile = FileManager->CreateFileWriter(*LogFilenameWithPath, 0);
+		LogFile->Serialize(LogFileMemory.GetData(), LogFileMemory.Num());
+		delete LogFile;
+#endif
+	}
+
+	Output.Logf(TEXT("Finished generating Array Slack report to %s."), LogFile ? *LogFilename : TEXT("[Debug Output]"));
+}
+
+uint8 LlmGetActiveTag()
+{
+	FLowLevelMemTracker& Tracker = FLowLevelMemTracker::Get();
+	if (!Tracker.IsInitialized())
+	{
+		return 0;
+	}
+
+	const UE::LLMPrivate::FTagData* TagData = Tracker.GetActiveTagData(ELLMTracker::Default);
+	return TagData ? (uint8)TagData->GetEnumTag() : 0;
+}
+
+#endif  // UE_ENABLE_ARRAY_SLACK_TRACKING
+
 UE_TRACE_CHANNEL(MemTagChannel, "Memory overview", true)
 
 UE_TRACE_EVENT_BEGIN(LLM, TagsSpec, NoSync|Important)
@@ -584,7 +1472,7 @@ public:
 	void GetTagsNamesWithAmount(TMap<FName, uint64>& OutTagsNamesWithAmount, ELLMTagSet TagSet = ELLMTagSet::None);
 	void GetTagsNamesWithAmountFiltered(TMap<FName, uint64>& OutTagsNamesWithAmount, ELLMTagSet TagSet, TArray<FLLMTagSetAllocationFilter>& Filters);
 
-	bool FindTagsForPtr(void* InPtr, TArray<const FTagData*, TInlineAllocator<static_cast<int32>(ELLMTagSet::Max)>>& OutTags);
+	bool FindTagsForPtr(void* InPtr, TArray<const FTagData*, TInlineAllocator<static_cast<int32>(ELLMTagSet::Max)>>& OutTags) const;
 
 	int64 GetAllocTypeAmount(ELLMAllocType AllocType);
 
@@ -912,6 +1800,11 @@ void FLowLevelMemTracker::OnPreFork()
 
 void FLowLevelMemTracker::UpdateStatsPerFrame(const TCHAR* LogName)
 {
+#if UE_ENABLE_ARRAY_SLACK_TRACKING
+	// Slack tracking, when compiled in, can run even when regular LLM tracking is disabled
+	LlmTrackArrayTick();
+#endif
+
 	if (bIsDisabled)
 	{
 		if (bFirstTimeUpdating)
@@ -1419,6 +2312,11 @@ UE::LLMPrivate::FLLMTracker* FLowLevelMemTracker::GetTracker(ELLMTracker Tracker
 	return Trackers[static_cast<int32>(Tracker)];
 }
 
+const UE::LLMPrivate::FLLMTracker* FLowLevelMemTracker::GetTracker(ELLMTracker Tracker) const
+{
+	return Trackers[static_cast<int32>(Tracker)];
+}
+
 bool FLowLevelMemTracker::Exec(const TCHAR* Cmd, FOutputDevice& Ar)
 {
 	if (bIsDisabled)
@@ -1637,6 +2535,9 @@ void FLowLevelMemTracker::FinishInitialise()
 	{
 		return;
 	}
+#if UE_ENABLE_ARRAY_SLACK_TRACKING
+	GTrackArrayDetailedLock = new FCriticalSection();
+#endif
 	bFullyInitialised = true;
 	// Make sure that FNames and Malloc have already been initialised, since we will use them during InitialiseTagDatas
 	// We force this by calling LLMGetTagUniqueName, which initializes FNames internally, and will therein trigger
@@ -2162,6 +3063,22 @@ FName FLowLevelMemTracker::FindTagDisplayName(uint64 Tag) const
 	return NAME_None;
 }
 
+FName FLowLevelMemTracker::FindPtrDisplayName(void* Ptr) const
+{
+	using namespace UE::LLMPrivate;
+
+	if (bIsDisabled || !bFullyInitialised)
+	{
+		return NAME_None;
+	}
+
+	const FLLMTracker* TrackerData = GetTracker(ELLMTracker::Default);
+	TArray<const FTagData*, TInlineAllocator<static_cast<int32>(ELLMTagSet::Max)>> Tags;
+	TrackerData->FindTagsForPtr(Ptr, Tags);
+
+	return Tags.Num() ? Tags[0]->GetDisplayName() : NAME_None;
+}
+
 FName FLowLevelMemTracker::GetTagDisplayName(const UE::LLMPrivate::FTagData* TagData) const
 {
 	return TagData->GetDisplayName();
@@ -4531,10 +5448,8 @@ void FLLMTracker::GetTagsNamesWithAmountFiltered(TMap<FName, uint64>& OutTagsNam
 	}
 }
 
-bool FLLMTracker::FindTagsForPtr(void* InPtr, TArray<const FTagData *, TInlineAllocator<static_cast<int32>(ELLMTagSet::Max)>>& OutTags)
+bool FLLMTracker::FindTagsForPtr(void* InPtr, TArray<const FTagData *, TInlineAllocator<static_cast<int32>(ELLMTagSet::Max)>>& OutTags) const
 {
-	FLLMThreadState* State = GetOrCreateState();
-		
 	uint32 Size;
 	FLowLevelAllocInfo AllocInfoPtr;
 	PointerKey FoundKey = AllocationMap.Find(PointerKey(InPtr), Size, AllocInfoPtr);
@@ -5306,4 +6221,25 @@ void FLLMCsvProfilerWriter::RecordTagToCsv(int32 CsvCategoryIndex, const FTagDat
 
 } // namespace UE::LLMPrivate
 
-#endif // #if ENABLE_LOW_LEVEL_MEM_TRACKER
+#else // #if ENABLE_LOW_LEVEL_MEM_TRACKER
+
+// We need to stub some functions so things link when the UE_ENABLE_ARRAY_SLACK_TRACKING debug feature is enabled in builds
+// where LLM is disabled.  Compiling out the slack tracking code completely is difficult due to include order issues.
+// Slack tracking is in a header that must be included before LLM, so it can't access the ENABLE_LOW_LEVEL_MEM_TRACKER
+// define.  And moving that define leads to a chain reaction of other include order issues.  It's just not worth it for
+// a rarely enabled debug feature to go to all that trouble, when stubbing functions works fine...
+#if UE_ENABLE_ARRAY_SLACK_TRACKING
+uint8 LlmGetActiveTag() { return 0; }
+void ArraySlackTrackInit() {}
+void ArraySlackTrackGenerateReport(const TCHAR* Cmd, FOutputDevice& Ar) {}
+void FArraySlackTrackingHeader::AddAllocation() {}
+void FArraySlackTrackingHeader::RemoveAllocation() {}
+void FArraySlackTrackingHeader::UpdateNumUsed(int64 NewNumUsed) {}
+
+FORCENOINLINE void* FArraySlackTrackingHeader::Realloc(void* Ptr, int64 Count, uint64 ElemSize, int32 Alignment)
+{
+	return FMemory::Realloc(Ptr, Count * ElemSize, Alignment);
+}
+#endif
+
+#endif  // #else .. #if ENABLE_LOW_LEVEL_MEM_TRACKER
diff --git a/Engine/Source/Runtime/Core/Private/HAL/LowLevelMemoryUtils.h b/Engine/Source/Runtime/Core/Private/HAL/LowLevelMemoryUtils.h
--- a/Engine/Source/Runtime/Core/Private/HAL/LowLevelMemoryUtils.h
+++ b/Engine/Source/Runtime/Core/Private/HAL/LowLevelMemoryUtils.h
@@ -329,10 +329,10 @@ public:
 		return RetValues;
 	}
 
-	TKey Find(const TKey& Key, TValue1& OutValue1, TValue2& OutValue2)
+	TKey Find(const TKey& Key, TValue1& OutValue1, TValue2& OutValue2) const
 	{
 		SizeType KeyHash = Key.GetHashCode();
-		StripeData& Stripe = MapStripes[GetStripeIndex(KeyHash)];
+		const StripeData& Stripe = MapStripes[GetStripeIndex(KeyHash)];
 
 		FScopeLock AllocationScopeLock(&Stripe.CriticalSection);
 
diff --git a/Engine/Source/Runtime/Core/Public/Containers/Array.h b/Engine/Source/Runtime/Core/Public/Containers/Array.h
--- a/Engine/Source/Runtime/Core/Public/Containers/Array.h
+++ b/Engine/Source/Runtime/Core/Public/Containers/Array.h
@@ -512,6 +512,17 @@ public:
 	TArray& operator=(const TArrayView<OtherElementType, OtherSizeType>& Other);
 
 private:
+
+	FORCEINLINE void SlackTrackerNumChanged()
+	{
+#if UE_ENABLE_ARRAY_SLACK_TRACKING
+		if constexpr (TAllocatorTraits<InAllocatorType>::SupportsSlackTracking)
+		{
+			AllocatorInstance.SlackTrackerLogNum(ArrayNum);
+		}
+#endif
+	}
+
 	/**
 	 * Moves or copies array. Depends on the array type traits.
 	 *
@@ -555,6 +566,9 @@ private:
 
 			FromArray.ArrayNum = 0;
 			FromArray.ArrayMax = FromArray.AllocatorInstance.GetInitialCapacity();
+
+			FromArray.SlackTrackerNumChanged();
+			ToArray.SlackTrackerNumChanged();
 		}
 		else
 		{
@@ -1312,6 +1326,11 @@ public:
 		{
 			ResizeGrow((SizeType)OldNum);
 		}
+		else
+		{
+			SlackTrackerNumChanged();
+		}
+
 		return OldNum;
 	}
 	FORCEINLINE SizeType AddUninitialized(SizeType Count)
@@ -1333,6 +1352,11 @@ public:
 		{
 			ResizeGrow((SizeType)OldNum);
 		}
+		else
+		{
+			SlackTrackerNumChanged();
+		}
+
 		return OldNum;
 	}
 
@@ -1350,6 +1374,10 @@ private:
 		{
 			ResizeGrow((SizeType)OldNum);
 		}
+		else
+		{
+			SlackTrackerNumChanged();
+		}
 		ElementType* Data = GetData() + Index;
 		RelocateConstructItems<ElementType>(Data + 1, Data, OldNum - Index);
 	}
@@ -1376,6 +1404,10 @@ private:
 		{
 			ResizeGrow((SizeType)OldNum);
 		}
+		else
+		{
+			SlackTrackerNumChanged();
+		}
 		ElementType* Data = GetData() + Index;
 		RelocateConstructItems<ElementType>(Data + Count, Data, OldNum - Index);
 	}
@@ -1534,6 +1566,8 @@ public:
 		RelocateConstructItems<ElementType>(GetData() + InIndex, Items.GetData(), NumNewElements);
 		Items.ArrayNum = 0;
 
+		Items.SlackTrackerNumChanged();
+
 		return InIndex;
 	}
 
@@ -1672,6 +1706,8 @@ private:
 			}
 			ArrayNum -= Count;
 
+			SlackTrackerNumChanged();
+
 			if (bAllowShrinking)
 			{
 				ResizeShrink();
@@ -1730,6 +1766,8 @@ private:
 			}
 			ArrayNum -= Count;
 
+			SlackTrackerNumChanged();
+
 			if (bAllowShrinking)
 			{
 				ResizeShrink();
@@ -1790,6 +1828,8 @@ public:
 		{
 			DestructItems(GetData(), ArrayNum);
 			ArrayNum = 0;
+
+			SlackTrackerNumChanged();
 		}
 		else
 		{
@@ -1815,6 +1855,8 @@ public:
 		checkSlow(Slack >= 0);
 		ArrayNum = 0;
 
+		SlackTrackerNumChanged();
+
 		if (ArrayMax != Slack)
 		{
 			ResizeTo(Slack);
@@ -1900,6 +1942,8 @@ public:
 	{
 		checkSlow(NewNum <= Num() && NewNum >= 0);
 		ArrayNum = NewNum;
+
+		SlackTrackerNumChanged();
 	}
 
 	/**
@@ -1951,6 +1995,8 @@ public:
 		SizeType Pos = AddUninitialized(SourceCount);
 		RelocateConstructItems<ElementType>(GetData() + Pos, Source.GetData(), SourceCount);
 		Source.ArrayNum = 0;
+
+		Source.SlackTrackerNumChanged();
 	}
 
 	/**
@@ -2381,6 +2427,8 @@ public:
 		RelocateConstructItems<ElementType>(GetData() + InIndex, Items.GetData(), NumNewElements);
 		Items.ArrayNum = 0;
 
+		Items.SlackTrackerNumChanged();
+
 		return InIndex;
 	}
 
@@ -2513,6 +2561,8 @@ public:
 		// Update the array count
 		--ArrayNum;
 
+		SlackTrackerNumChanged();
+
 		// Removed one item
 		return 1;
 	}
@@ -2580,6 +2630,9 @@ public:
 		} while (ReadIndex < OriginalNum);
 
 		ArrayNum = WriteIndex;
+
+		SlackTrackerNumChanged();
+
 		return OriginalNum - ArrayNum;
 	}
 
@@ -2894,6 +2947,7 @@ private:
 		{
 			AllocatorInstance.ResizeAllocation(CurrentArrayNum, NewArrayMax, sizeof(ElementType));
 		}
+		SlackTrackerNumChanged();
 	}
 
 	SizeType AllocatorCalculateSlackShrink(SizeType CurrentArrayNum, SizeType NewArrayMax)
@@ -3001,13 +3055,15 @@ private:
 		ArrayNum = NewNum;
 		if (OtherNum || PrevMax)
 		{
-		ResizeForCopy(NewNum, PrevMax);
-		ConstructItems<ElementType>(GetData(), OtherData, OtherNum);
+			ResizeForCopy(NewNum, PrevMax);
+			ConstructItems<ElementType>(GetData(), OtherData, OtherNum);
 		}
 		else
 		{
-		ArrayMax = AllocatorInstance.GetInitialCapacity();
+			ArrayMax = AllocatorInstance.GetInitialCapacity();
 		}
+
+		SlackTrackerNumChanged();
 	}
 
 	/**
@@ -3043,6 +3099,8 @@ private:
 		{
 			ArrayMax = AllocatorInstance.GetInitialCapacity();
 		}
+
+		SlackTrackerNumChanged();
 	}
 
 protected:
@@ -3561,6 +3619,8 @@ struct TArrayPrivateFriend
 			}
 		}
 
+		A.SlackTrackerNumChanged();
+
 		return Ar;
 	}
 };
diff --git a/Engine/Source/Runtime/Core/Public/Containers/ContainerAllocationPolicies.h b/Engine/Source/Runtime/Core/Public/Containers/ContainerAllocationPolicies.h
--- a/Engine/Source/Runtime/Core/Public/Containers/ContainerAllocationPolicies.h
+++ b/Engine/Source/Runtime/Core/Public/Containers/ContainerAllocationPolicies.h
@@ -13,6 +13,94 @@
 #include "Templates/TypeCompatibleBytes.h"
 #include <type_traits>
 
+
+// Array slack tracking is a debug feature to track unused space in heap allocated TArray (and TScriptArray) structures.  This feature increases heap
+// memory usage, and costs perf, so it is usually disabled by default.  Only works in builds where LLM is compiled in, but doesn't require -llm to be
+// active.  Note that cooks will run quite slow with tracking enabled due to involving significantly more allocations than the engine, so be careful
+// about leaving this enabled when kicking off a cook (including !WITH_EDITOR ensures it will only be enabled for client builds).  For more details,
+// see additional comments in LowLevelMemTracker.cpp.
+#ifndef UE_ENABLE_ARRAY_SLACK_TRACKING
+#define UE_ENABLE_ARRAY_SLACK_TRACKING (0 && !WITH_EDITOR)
+#endif
+
+
+#if UE_ENABLE_ARRAY_SLACK_TRACKING
+
+CORE_API uint8 LlmGetActiveTag();
+CORE_API void ArraySlackTrackInit();
+CORE_API void ArraySlackTrackGenerateReport(const TCHAR* Cmd, FOutputDevice& Ar);
+
+// For detailed tracking of array slack waste, we need to add a header to heap allocations.  It's impossible to keep track of the TArray structure itself,
+// because it can be inside other structures, and potentially moved around in unsafe ways, while the heap allocation the TArray points to is invariant.
+// The heap allocations also have an advantage in that the contents are copied on Move or Realloc, simplifying tracking.
+struct FArraySlackTrackingHeader
+{
+	FArraySlackTrackingHeader* Next;	// Linked list of tracked items
+	FArraySlackTrackingHeader** Prev;
+	uint16 AllocOffset;					// Offset below the header to the start of the actual allocation, to account for alignment padding
+	uint8 Tag;
+	int8 NumStackFrames;
+	uint32 FirstAllocFrame;				// Frame where array allocation first occurred
+	uint32 ReallocCount;				// Number of times realloc happened
+	uint32 ArrayPeak;					// Peak observed ArrayNum
+	uint64 ElemSize;
+
+	// Note that we initially set the slack tracking ArrayNum to INDEX_NONE.  The container allocator is used by both arrays and
+	// other containers (Set / Map / Hash), and we don't know it's actually an array until "UpdateNumUsed" is called on it.
+	int64 ArrayNum;
+	int64 ArrayMax;
+	uint64 StackFrames[9];
+
+	CORE_API void AddAllocation();
+	CORE_API void RemoveAllocation();
+	CORE_API void UpdateNumUsed(int64 NewNumUsed);
+	CORE_API FORCENOINLINE static void* Realloc(void* Ptr, int64 Count, uint64 ElemSize, int32 Alignment);
+
+	static void Free(void* Ptr)
+	{
+		if (Ptr)
+		{
+			FArraySlackTrackingHeader* TrackingHeader = (FArraySlackTrackingHeader*)((uint8*)Ptr - sizeof(FArraySlackTrackingHeader));
+			TrackingHeader->RemoveAllocation();
+
+			Ptr = (uint8*)TrackingHeader - TrackingHeader->AllocOffset;
+
+			FMemory::Free(Ptr);
+		}
+	}
+
+	static FORCEINLINE void UpdateNumUsed(void* Ptr, int64 NewNumUsed)
+	{
+		if (Ptr)
+		{
+			FArraySlackTrackingHeader* TrackingHeader = (FArraySlackTrackingHeader*)((uint8*)Ptr - sizeof(FArraySlackTrackingHeader));
+
+			TrackingHeader->UpdateNumUsed(NewNumUsed);
+		}
+	}
+
+	static FORCEINLINE void DisableTracking(void* Ptr)
+	{
+		if (Ptr)
+		{
+			FArraySlackTrackingHeader* TrackingHeader = (FArraySlackTrackingHeader*)((uint8*)Ptr - sizeof(FArraySlackTrackingHeader));
+
+			TrackingHeader->RemoveAllocation();
+
+			// When disabling tracking, we need to also reset ArrayNum, as it's used internally as a flag specifying whether
+			// the allocation is currently tracked.  We don't reset this inside RemoveAllocation, because ArrayNum needs to
+			// persist during realloc of tracked allocations, where RemoveAllocation is called, followed by AddAllocation.
+			TrackingHeader->ArrayNum = INDEX_NONE;
+		}
+	}
+
+	FORCEINLINE int64 SlackSizeInBytes() const
+	{
+		return (ArrayMax - ArrayNum) * ElemSize;
+	}
+};
+#endif  // UE_ENABLE_ARRAY_SLACK_TRACKING
+
 // This option disables array slack for initial allocations, e.g where TArray::SetNum 
 // is called. This tends to save a lot of memory with almost no measured performance cost.
 // NOTE: This can cause latent memory corruption issues to become more prominent
@@ -157,6 +245,7 @@ struct TAllocatorTraitsBase
 	enum { IsZeroConstruct           = false };
 	enum { SupportsFreezeMemoryImage = false };
 	enum { SupportsElementAlignment  = false };
+	enum { SupportsSlackTracking     = false };
 };
 
 template <typename AllocatorType>
@@ -331,6 +420,9 @@ public:
 
 		/** Returns number of pre-allocated elements the container can use before allocating more space */
 		SizeType GetInitialCapacity() const;
+
+		/** Function called when ArrayNum changes for a TArray or TScriptArray, if TAllocatorTraits<Allocator>::SupportsSlackTracking == true */
+		void SlackTrackerLogNum(SizeType NewNumUsed);
 	};
 
 	/**
@@ -376,7 +468,11 @@ public:
 
 			if (Data)
 			{
+#if UE_ENABLE_ARRAY_SLACK_TRACKING
+				FArraySlackTrackingHeader::Free(Data);
+#else
 				FMemory::Free(Data);
+#endif
 			}
 
 			Data       = Other.Data;
@@ -388,7 +484,11 @@ public:
 		{
 			if(Data)
 			{
+#if UE_ENABLE_ARRAY_SLACK_TRACKING
+				FArraySlackTrackingHeader::Free(Data);
+#else
 				FMemory::Free(Data);
+#endif
 			}
 		}
 
@@ -414,7 +514,11 @@ public:
 					UE::Core::Private::OnInvalidAlignedHeapAllocatorNum(NumElements, NumBytesPerElement);
 				}
 
+#if UE_ENABLE_ARRAY_SLACK_TRACKING
+				Data = (FScriptContainerElement*)FArraySlackTrackingHeader::Realloc(Data, NumElements, NumBytesPerElement, Alignment > alignof(FArraySlackTrackingHeader) ? Alignment : alignof(FArraySlackTrackingHeader));
+#else
 				Data = (FScriptContainerElement*)FMemory::Realloc( Data, NumElements*NumBytesPerElement, Alignment );
+#endif
 			}
 		}
 		FORCEINLINE SizeType CalculateSlackReserve(SizeType NumElements, SIZE_T NumBytesPerElement) const
@@ -445,6 +549,20 @@ public:
 			return 0;
 		}
 
+#if UE_ENABLE_ARRAY_SLACK_TRACKING
+		FORCEINLINE void SlackTrackerLogNum(SizeType NewNumUsed)
+		{
+			FArraySlackTrackingHeader::UpdateNumUsed(Data, (int64)NewNumUsed);
+		}
+
+		// Suppress slack tracking on an allocation -- should be called whenever the container may have been resized.
+		// Useful for debug allocations you don't want to show up in slack reports.
+		FORCEINLINE void DisableSlackTracking()
+		{
+			FArraySlackTrackingHeader::DisableTracking(Data);
+		}
+#endif
+
 	private:
 		ForAnyElementType(const ForAnyElementType&);
 		ForAnyElementType& operator=(const ForAnyElementType&);
@@ -476,6 +594,7 @@ template <uint32 Alignment>
 struct TAllocatorTraits<TAlignedHeapAllocator<Alignment>> : TAllocatorTraitsBase<TAlignedHeapAllocator<Alignment>>
 {
 	enum { IsZeroConstruct = true };
+	enum { SupportsSlackTracking = true };
 };
 
 template <int IndexSize>
@@ -527,7 +646,11 @@ public:
 
 			if (Data)
 			{
+#if UE_ENABLE_ARRAY_SLACK_TRACKING
+				FArraySlackTrackingHeader::Free(Data);
+#else
 				BaseMallocType::Free(Data);
+#endif
 			}
 
 			Data = Other.Data;
@@ -551,7 +674,11 @@ public:
 		{
 			if(Data)
 			{
+#if UE_ENABLE_ARRAY_SLACK_TRACKING
+				FArraySlackTrackingHeader::Free(Data);
+#else
 				BaseMallocType::Free(Data);
+#endif
 			}
 		}
 
@@ -578,7 +705,11 @@ public:
 					UE::Core::Private::OnInvalidSizedHeapAllocatorNum(IndexSize, NumElements, NumBytesPerElement);
 				}
 
+#if UE_ENABLE_ARRAY_SLACK_TRACKING
+				Data = (FScriptContainerElement*)FArraySlackTrackingHeader::Realloc(Data, NumElements, NumBytesPerElement, 0);
+#else
 				Data = (FScriptContainerElement*)BaseMallocType::Realloc( Data, NumElements*NumBytesPerElement );
+#endif
 			}
 		}
 		void ResizeAllocation(SizeType PreviousNumElements, SizeType NumElements, SIZE_T NumBytesPerElement, uint32 AlignmentOfElement)
@@ -599,7 +730,11 @@ public:
 					UE::Core::Private::OnInvalidSizedHeapAllocatorNum(IndexSize, NumElements, NumBytesPerElement);
 				}
 
+#if UE_ENABLE_ARRAY_SLACK_TRACKING
+				Data = (FScriptContainerElement*)FArraySlackTrackingHeader::Realloc(Data, NumElements, NumBytesPerElement, AlignmentOfElement > alignof(FArraySlackTrackingHeader) ? AlignmentOfElement : alignof(FArraySlackTrackingHeader));
+#else
 				Data = (FScriptContainerElement*)BaseMallocType::Realloc( Data, NumElements*NumBytesPerElement, AlignmentOfElement );
+#endif
 			}
 		}
 		FORCEINLINE SizeType CalculateSlackReserve(SizeType NumElements, SIZE_T NumBytesPerElement) const
@@ -642,6 +777,20 @@ public:
 			return 0;
 		}
 
+#if UE_ENABLE_ARRAY_SLACK_TRACKING
+		FORCEINLINE void SlackTrackerLogNum(SizeType NewNumUsed)
+		{
+			FArraySlackTrackingHeader::UpdateNumUsed(Data, (int64)NewNumUsed);
+		}
+
+		// Suppress slack tracking on an allocation -- should be called whenever the container may have been resized.
+		// Useful for debug allocations you don't want to show up in slack reports.
+		FORCEINLINE void DisableSlackTracking()
+		{
+			FArraySlackTrackingHeader::DisableTracking(Data);
+		}
+#endif
+
 	private:
 		ForAnyElementType(const ForAnyElementType&);
 		ForAnyElementType& operator=(const ForAnyElementType&);
@@ -677,6 +826,7 @@ struct TAllocatorTraits<TSizedHeapAllocator<IndexSize>> : TAllocatorTraitsBase<T
 {
 	enum { IsZeroConstruct          = true };
 	enum { SupportsElementAlignment = true };
+	enum { SupportsSlackTracking    = true };
 };
 
 using FHeapAllocator = TSizedHeapAllocator<32>;
@@ -816,6 +966,16 @@ public:
 			return NumInlineElements;
 		}
 
+#if UE_ENABLE_ARRAY_SLACK_TRACKING
+		FORCEINLINE void SlackTrackerLogNum(SizeType NewNumUsed)
+		{
+			if constexpr (TAllocatorTraits<SecondaryAllocator>::SupportsSlackTracking)
+			{
+				SecondaryData.SlackTrackerLogNum(NewNumUsed);
+			}
+		}
+#endif
+
 	private:
 		ForElementType(const ForElementType&);
 		ForElementType& operator=(const ForElementType&);
@@ -836,6 +996,12 @@ public:
 	typedef void ForAnyElementType;
 };
 
+template <uint32 NumInlineElements, int IndexSize, typename SecondaryAllocator>
+struct TAllocatorTraits<TSizedInlineAllocator<NumInlineElements, IndexSize, SecondaryAllocator>> : TAllocatorTraitsBase<TSizedInlineAllocator<NumInlineElements, IndexSize, SecondaryAllocator>>
+{
+	enum { SupportsSlackTracking = true };
+};
+
 template <uint32 NumInlineElements, typename SecondaryAllocator = FDefaultAllocator>
 using TInlineAllocator = TSizedInlineAllocator<NumInlineElements, 32, SecondaryAllocator>;
 
@@ -871,7 +1037,11 @@ public:
 		{
 			if (HasAllocation())
 			{
+#if UE_ENABLE_ARRAY_SLACK_TRACKING
+				FArraySlackTrackingHeader::Free(Data);
+#else
 				FMemory::Free(Data);
+#endif
 			}
 		}
 
@@ -886,7 +1056,11 @@ public:
 
 			if (HasAllocation())
 			{
+#if UE_ENABLE_ARRAY_SLACK_TRACKING
+				FArraySlackTrackingHeader::Free(Data);
+#else
 				FMemory::Free(Data);
+#endif
 			}
 
 			if (Other.HasAllocation())
@@ -916,7 +1090,11 @@ public:
 				if(HasAllocation())
 				{
 					RelocateConstructItems<ElementType>(GetInlineElements(), Data, PreviousNumElements);
+#if UE_ENABLE_ARRAY_SLACK_TRACKING
+					FArraySlackTrackingHeader::Free(Data);
+#else
 					FMemory::Free(Data);
+#endif
 					Data = GetInlineElements();
 				}
 			}
@@ -925,12 +1103,20 @@ public:
 				if (HasAllocation())
 				{
 					// Reallocate the indirect data for the new size.
+#if UE_ENABLE_ARRAY_SLACK_TRACKING
+					Data = (ElementType*)FArraySlackTrackingHeader::Realloc(Data, (int32)NumElements, (int32)NumBytesPerElement, 0);
+#else
 					Data = (ElementType*)FMemory::Realloc(Data, NumElements*NumBytesPerElement);
+#endif
 				}
 				else
 				{
 					// Allocate new indirect memory for the data.
+#if UE_ENABLE_ARRAY_SLACK_TRACKING
+					Data = (ElementType*)FArraySlackTrackingHeader::Realloc(nullptr, (int32)NumElements, (int32)NumBytesPerElement, 0);
+#else
 					Data = (ElementType*)FMemory::Realloc(nullptr, NumElements*NumBytesPerElement);
+#endif
 
 					// Move the data out of the inline data area into the new allocation.
 					RelocateConstructItems<ElementType>(Data, GetInlineElements(), PreviousNumElements);
@@ -971,6 +1157,18 @@ public:
 			return NumInlineElements;
 		}
 
+#if UE_ENABLE_ARRAY_SLACK_TRACKING
+		FORCEINLINE void SlackTrackerLogNum(SizeType NewNumUsed)
+		{
+			if (HasAllocation())
+			{
+				FArraySlackTrackingHeader* TrackingHeader = (FArraySlackTrackingHeader*)((uint8*)Data - sizeof(FArraySlackTrackingHeader));
+
+				TrackingHeader->UpdateNumUsed((int64)NewNumUsed);
+			}
+		}
+#endif
+
 	private:
 		ForElementType(const ForElementType&) = delete;
 		ForElementType& operator=(const ForElementType&) = delete;
@@ -991,6 +1189,12 @@ public:
 	typedef void ForAnyElementType;
 };
 
+template <uint32 NumInlineElements>
+struct TAllocatorTraits<TNonRelocatableInlineAllocator<NumInlineElements>> : TAllocatorTraitsBase<TNonRelocatableInlineAllocator<NumInlineElements>>
+{
+	enum { SupportsSlackTracking = true };
+};
+
 /**
  * The fixed allocation policy allocates up to a specified number of elements in the same allocation as the container.
  * It's like the inline allocator, except it doesn't provide secondary storage when the inline storage has been filled.
diff --git a/Engine/Source/Runtime/Core/Public/Containers/ScriptArray.h b/Engine/Source/Runtime/Core/Public/Containers/ScriptArray.h
--- a/Engine/Source/Runtime/Core/Public/Containers/ScriptArray.h
+++ b/Engine/Source/Runtime/Core/Public/Containers/ScriptArray.h
@@ -66,6 +66,8 @@ public:
 			(uint8*)this->GetAllocation() + (Index       )*NumBytesPerElement,
 			                                               (OldNum-Index)*NumBytesPerElement
 		);
+
+		SlackTrackerNumChanged();
 	}
 	int32 Add( int32 Count, int32 NumBytesPerElement, uint32 AlignmentOfElement )
 	{
@@ -79,6 +81,8 @@ public:
 			ResizeGrow(OldNum, NumBytesPerElement, AlignmentOfElement);
 		}
 
+		SlackTrackerNumChanged();
+
 		return OldNum;
 	}
 	int32 AddZeroed( int32 Count, int32 NumBytesPerElement, uint32 AlignmentOfElement )
@@ -116,11 +120,17 @@ public:
 		this->MoveToEmpty(Other);
 		ArrayNum = Other.ArrayNum; Other.ArrayNum = 0;
 		ArrayMax = Other.ArrayMax; Other.ArrayMax = 0;
+
+		this->SlackTrackerNumChanged();
+		Other.SlackTrackerNumChanged();
 	}
 	void Empty( int32 Slack, int32 NumBytesPerElement, uint32 AlignmentOfElement )
 	{
 		checkSlow(Slack>=0);
 		ArrayNum = 0;
+
+		SlackTrackerNumChanged();
+
 		if (Slack != ArrayMax)
 		{
 			ResizeTo(Slack, NumBytesPerElement, AlignmentOfElement);
@@ -131,6 +141,8 @@ public:
 		if (NewSize <= ArrayMax)
 		{
 			ArrayNum = 0;
+
+			SlackTrackerNumChanged();
 		}
 		else
 		{
@@ -188,6 +200,8 @@ public:
 			}
 			ArrayNum -= Count;
 
+			SlackTrackerNumChanged();
+
 			if (bAllowShrinking)
 			{
 				ResizeShrink(NumBytesPerElement, AlignmentOfElement);
@@ -213,6 +227,8 @@ protected:
 			ResizeInit(NumBytesPerElement, AlignmentOfElement);
 		}
 		ArrayNum = InNum;
+
+		SlackTrackerNumChanged();
 	}
 	int32	  ArrayNum;
 	int32	  ArrayMax;
@@ -248,6 +264,18 @@ protected:
 			this->ResizeAllocation(ArrayNum, ArrayMax, NumBytesPerElement, AlignmentOfElement);
 		}
 	}
+
+private:
+	FORCEINLINE void SlackTrackerNumChanged()
+	{
+#if UE_ENABLE_ARRAY_SLACK_TRACKING
+		if constexpr (TAllocatorTraits<AllocatorType>::SupportsSlackTracking)
+		{
+			((typename AllocatorType::ForAnyElementType*)this)->SlackTrackerLogNum(ArrayNum);
+		}
+#endif
+	}
+
 public:
 	// These should really be private, because they shouldn't be called, but there's a bunch of code
 	// that needs to be fixed first.
diff --git a/Engine/Source/Runtime/Core/Public/HAL/LowLevelMemTracker.h b/Engine/Source/Runtime/Core/Public/HAL/LowLevelMemTracker.h
--- a/Engine/Source/Runtime/Core/Public/HAL/LowLevelMemTracker.h
+++ b/Engine/Source/Runtime/Core/Public/HAL/LowLevelMemTracker.h
@@ -669,6 +669,7 @@ public:
 
 	/** Get the display name for the given ELLMTag. */
 	CORE_API FName FindTagDisplayName(uint64 Tag) const;
+	CORE_API FName FindPtrDisplayName(void* Ptr) const;
 
 	/** Get the display name for the given FTagData. */
 	CORE_API FName GetTagDisplayName(const UE::LLMPrivate::FTagData* TagData) const;
@@ -719,6 +720,11 @@ public:
 
 	CORE_API void OnPreFork();
 
+	FORCEINLINE bool IsInitialized() const
+	{
+		return bFullyInitialised;
+	}
+
 private:
 	CORE_API FLowLevelMemTracker();
 
@@ -739,6 +745,7 @@ private:
 	CORE_API void InitialiseProgramSize();
 
 	CORE_API class UE::LLMPrivate::FLLMTracker* GetTracker(ELLMTracker Tracker);
+	CORE_API const class UE::LLMPrivate::FLLMTracker* GetTracker(ELLMTracker Tracker) const;
 
 	CORE_API void TickInternal();
 	CORE_API void UpdateTags();
diff --git a/Engine/Source/Runtime/Engine/Private/UnrealEngine.cpp b/Engine/Source/Runtime/Engine/Private/UnrealEngine.cpp
--- a/Engine/Source/Runtime/Engine/Private/UnrealEngine.cpp
+++ b/Engine/Source/Runtime/Engine/Private/UnrealEngine.cpp
@@ -4943,6 +4943,14 @@ bool UEngine::Exec( UWorld* InWorld, const TCHAR* Cmd, FOutputDevice& Ar )
 	}
 #endif
 
+#if UE_ENABLE_ARRAY_SLACK_TRACKING
+	else if (FParse::Command(&Cmd, TEXT("SLACKREPORT")))
+	{
+		ArraySlackTrackGenerateReport(Cmd, Ar);
+		return true;
+	}
+#endif
+
 	else if ( FParse::Command(&Cmd,TEXT("SCALABILITY")) )
 	{
 		Scalability::ProcessCommand(Cmd, Ar);
diff --git a/Engine/Source/Runtime/Engine/Private/UserInterface/Console.cpp b/Engine/Source/Runtime/Engine/Private/UserInterface/Console.cpp
--- a/Engine/Source/Runtime/Engine/Private/UserInterface/Console.cpp
+++ b/Engine/Source/Runtime/Engine/Private/UserInterface/Console.cpp
@@ -441,6 +441,14 @@ void UConsole::BuildRuntimeAutoCompleteList(bool bForce)
 					FoundNodeIdx = NodeIdx;
 					Node = NodeList[FoundNodeIdx];
 					NodeList[FoundNodeIdx]->AutoCompleteListIndices.Add(ListIdx);
+#if UE_ENABLE_ARRAY_SLACK_TRACKING
+					// Disable array slack tracking for Console system related allocations.  This needs to be called any time a reallocation occurs,
+					// which in practical terms means any "Add" (although it's cheap if there wasn't a recent reallocation, a couple memory reads).
+					// The auto-complete code (triggered when bringing up a console window to issue the slack report command) generates over a million
+					// allocations, with over 10 MB cumulative slack memory, which ends up at the top of any slack report, and is debug related code
+					// that we don't care about tracking.  If anyone wants to look at it for some reason, just remove these lines.
+					NodeList[FoundNodeIdx]->AutoCompleteListIndices.GetAllocatorInstance().DisableSlackTracking();
+#endif
 					break;
 				}
 			}
@@ -449,6 +457,10 @@ void UConsole::BuildRuntimeAutoCompleteList(bool bForce)
 				FAutoCompleteNode* NewNode = new FAutoCompleteNode(Char);
 				NewNode->AutoCompleteListIndices.Add(ListIdx);
 				Node->ChildNodes.Add(NewNode);
+#if UE_ENABLE_ARRAY_SLACK_TRACKING
+				NewNode->AutoCompleteListIndices.GetAllocatorInstance().DisableSlackTracking();
+				Node->ChildNodes.GetAllocatorInstance().DisableSlackTracking();
+#endif
 				Node = NewNode;
 			}
 		}
diff --git a/Engine/Source/Runtime/Launch/Private/LaunchEngineLoop.cpp b/Engine/Source/Runtime/Launch/Private/LaunchEngineLoop.cpp
--- a/Engine/Source/Runtime/Launch/Private/LaunchEngineLoop.cpp
+++ b/Engine/Source/Runtime/Launch/Private/LaunchEngineLoop.cpp
@@ -4466,6 +4466,14 @@ int32 FEngineLoop::PreInitPostStartupScreen(const TCHAR* CmdLine)
 
 int32 FEngineLoop::PreInit(const TCHAR* CmdLine)
 {
+#if UE_ENABLE_ARRAY_SLACK_TRACKING
+	// Any array allocations before this point won't have array slack tracking, although subsequent reallocations of existing arrays
+	// will gain tracking if that occurs.  The goal is to filter out startup constructors which run before Main, which introduce a
+	// ton of noise into slack reports.  Especially the roughly 30,000 static FString constructors in the code base, each with a
+	// unique call stack, and all having a little bit of slack due to malloc bucket size rounding.
+	ArraySlackTrackInit();
+#endif
+
 	const int32 rv1 = PreInitPreStartupScreen(CmdLine);
 	if (rv1 != 0)
 	{
