Expose UTS controls and start with sponsor argument

* Expose Unreal Trace Server launching facilities in FTraceServerControl interface. Only enabled on platforms that can actually launch UTS.
* Renamed WITH_UNREAL_TRACE_LAUNCH to UE_TRACE_SERVER_LAUNCH_ENABLED
* Always start UnrealTraceServer with sponsor argument.

diff --git a/Engine/Source/Programs/UnrealInsights/UnrealInsights.Target.cs b/Engine/Source/Programs/UnrealInsights/UnrealInsights.Target.cs
--- a/Engine/Source/Programs/UnrealInsights/UnrealInsights.Target.cs
+++ b/Engine/Source/Programs/UnrealInsights/UnrealInsights.Target.cs
@@ -36,7 +36,9 @@ public class UnrealInsightsTarget : TargetRules
 
 		bHasExports = false;
 
+		// Enable server controls
+		GlobalDefinitions.Add("UE_TRACE_SERVER_CONTROLS_ENABLED=1");
 		// Have UnrealInsights implicitly launch the trace store.
-		GlobalDefinitions.Add("WITH_UNREAL_TRACE_LAUNCH=1");
+		GlobalDefinitions.Add("UE_TRACE_SERVER_LAUNCH_ENABLED=1");
 	}
 }
diff --git a/Engine/Source/Runtime/Core/Private/ProfilingDebugging/TraceAuxiliary.cpp b/Engine/Source/Runtime/Core/Private/ProfilingDebugging/TraceAuxiliary.cpp
--- a/Engine/Source/Runtime/Core/Private/ProfilingDebugging/TraceAuxiliary.cpp
+++ b/Engine/Source/Runtime/Core/Private/ProfilingDebugging/TraceAuxiliary.cpp
@@ -15,15 +15,19 @@
 # 	include <sys/wait.h>
 #endif
 
-#if !defined(WITH_UNREAL_TRACE_LAUNCH)
-#	define WITH_UNREAL_TRACE_LAUNCH (PLATFORM_DESKTOP && !UE_BUILD_SHIPPING && !IS_PROGRAM)
+#if defined(WITH_UNREAL_TRACE_LAUNCH)
+  #define UE_TRACE_SERVER_LAUNCH_ENABLED WITH_UNREAL_TRACE_LAUNCH
+  UE_DEPRECATED_MACRO(5.4, "The WITH_UNREAL_TRACE_LAUNCH macro has been deprecated in favor of UE_TRACE_SERVER_LAUNCH_ENABLED.")
+#elif !defined(UE_TRACE_SERVER_LAUNCH_ENABLED)
+  #define UE_TRACE_SERVER_LAUNCH_ENABLED (PLATFORM_DESKTOP && !UE_BUILD_SHIPPING && !IS_PROGRAM)
+  #define WITH_UNREAL_TRACE_LAUNCH UE_DEPRECATED_MACRO(5.4, "The WITH_UNREAL_TRACE_LAUNCH macro has been deprecated in favor of UE_TRACE_SERVER_LAUNCH_ENABLED.") UE_TRACE_SERVER_LAUNCH_ENABLED
 #endif
 
 #if !defined(UE_TRACE_AUTOSTART)
 	#define UE_TRACE_AUTOSTART 1
 #endif
 
-#if WITH_UNREAL_TRACE_LAUNCH
+#if UE_TRACE_SERVER_LAUNCH_ENABLED
 #include "Misc/Parse.h"
 #include "Misc/Paths.h"
 #endif
@@ -1157,145 +1161,6 @@ static FAutoConsoleCommand TraceBookmarkCmd(
 #endif // UE_TRACE_ENABLED
 
 
-
-#if WITH_UNREAL_TRACE_LAUNCH
-////////////////////////////////////////////////////////////////////////////////
-static std::atomic<int32> GUnrealTraceLaunched; // = 0;
-
-////////////////////////////////////////////////////////////////////////////////
-#if PLATFORM_WINDOWS
-static void LaunchUnrealTraceInternal(const TCHAR* CommandLine)
-{
-	if (GUnrealTraceLaunched.load(std::memory_order_relaxed))
-	{
-		UE_LOG(LogCore, Log, TEXT("UnrealTraceServer: Trace store already started"));
-		return;
-	}
-
-	TWideStringBuilder<MAX_PATH + 32> CreateProcArgs;
-	CreateProcArgs << "\"";
-	CreateProcArgs << FPaths::EngineDir();
-	CreateProcArgs << TEXT("/Binaries/Win64/UnrealTraceServer.exe\"");
-	CreateProcArgs << TEXT(" fork");
-
-	uint32 CreateProcFlags = CREATE_BREAKAWAY_FROM_JOB;
-	if (FParse::Param(CommandLine, TEXT("traceshowstore")))
-	{
-		CreateProcFlags |= CREATE_NEW_CONSOLE;
-	}
-	else
-	{
-		CreateProcFlags |= CREATE_NO_WINDOW;
-	}
-	STARTUPINFOW StartupInfo = { sizeof(STARTUPINFOW) };
-	PROCESS_INFORMATION ProcessInfo = {};
-	BOOL bOk = CreateProcessW(nullptr, LPWSTR(*CreateProcArgs), nullptr, nullptr,
-		false, CreateProcFlags, nullptr, nullptr, &StartupInfo, &ProcessInfo);
-
-	if (!bOk)
-	{
-		UE_LOG(LogCore, Display, TEXT("UnrealTraceServer: Unable to launch the trace store with '%s' (%08x)"), *CreateProcArgs, GetLastError());
-		return;
-	}
-
-	if (WaitForSingleObject(ProcessInfo.hProcess, 5000) == WAIT_TIMEOUT)
-	{
-		UE_LOG(LogCore, Warning, TEXT("UnrealTraceServer: Timed out waiting for the trace store to start"));
-	}
-	else
-	{
-		DWORD ExitCode = 0x0000'a9e0;
-		GetExitCodeProcess(ProcessInfo.hProcess, &ExitCode);
-		if (ExitCode)
-		{
-			UE_LOG(LogCore, Warning, TEXT("UnrealTraceServer: Trace store returned an error (0x%08x)"), ExitCode);
-		}
-		else
-		{
-			UE_LOG(LogCore, Log, TEXT("UnrealTraceServer: Trace store launch successful"));
-			GUnrealTraceLaunched.fetch_add(1, std::memory_order_relaxed);
-		}
-	}
-
-	CloseHandle(ProcessInfo.hProcess);
-	CloseHandle(ProcessInfo.hThread);
-}
-#endif // PLATFORM_WINDOWS
-
-////////////////////////////////////////////////////////////////////////////////
-#if PLATFORM_UNIX || PLATFORM_MAC
-static void LaunchUnrealTraceInternal(const TCHAR* CommandLine)
-{
-// TSAN doesn't like fork(), so disable this for now.
-#if !USING_THREAD_SANITISER 
-	if (GUnrealTraceLaunched.load(std::memory_order_relaxed))
-	{
-		UE_LOG(LogCore, Log, TEXT("UnrealTraceServer: Trace store already started"));
-		return;
-	}
-
-	TAnsiStringBuilder<320> BinPath;
-	BinPath << TCHAR_TO_UTF8(*FPaths::EngineDir());
-#if PLATFORM_UNIX
-	BinPath << "Binaries/Linux/UnrealTraceServer";
-#elif PLATFORM_MAC
-	BinPath << "Binaries/Mac/UnrealTraceServer";
-#endif
-	BinPath.ToString(); //Ensure zero termination
-
-	if (access(*BinPath, F_OK) < 0)
-	{
-		UE_LOG(LogCore, Display, TEXT("UnrealTraceServer: Binary not found (%s)"), ANSI_TO_TCHAR(*BinPath));
-		return;
-	}
-
-	TAnsiStringBuilder<64> ForkArg;
-	ForkArg << "fork";
-	ForkArg.ToString(); //Ensure zero termination
-
-	pid_t UtsPid = fork();
-	if (UtsPid < 0)
-	{
-		UE_LOG(LogCore, Display, TEXT("UnrealTraceServer: Unable to fork (errno: %d)"), errno);
-		return;
-	}
-	else if (UtsPid == 0)
-	{
-		char* Args[] = { BinPath.GetData(), ForkArg.GetData(), nullptr };
-		extern char** environ;
-		execve(*BinPath, Args, environ);
-		_exit(0x80 | (errno & 0x7f));
-	}
-
-	int32 WaitStatus = 0;
-	do
-	{
-		int32 WaitRet = waitpid(UtsPid, &WaitStatus, 0);
-		if (WaitRet < 0)
-		{
-			UE_LOG(LogCore, Display, TEXT("UnrealTraceServer: waitpid() error; (errno: %d)"), errno);
-			return;
-		}
-	}
-	while (!WIFEXITED(WaitStatus));
-
-	int32 UtsRet = WEXITSTATUS(WaitStatus);
-	if (UtsRet)
-	{
-		UE_LOG(LogCore, Display, TEXT("UnrealTraceServer: Trace store returned an error (0x%08x)"), UtsRet);
-	}
-	else
-	{
-		UE_LOG(LogCore, Log, TEXT("UnrealTraceServer: Trace store launch successful"));
-		GUnrealTraceLaunched.fetch_add(1, std::memory_order_relaxed);
-	}
-#endif // #if !USING_THREAD_SANITISER
-}
-#endif // PLATFORM_UNIX/MAC
-#endif // WITH_UNREAL_TRACE_LAUNCH
-
-
-
 ////////////////////////////////////////////////////////////////////////////////
 UE_TRACE_EVENT_BEGIN(Diagnostics, Session2, NoSync|Important)
 	UE_TRACE_EVENT_FIELD(UE::Trace::AnsiString, Platform)
@@ -1537,11 +1402,12 @@ void FTraceAuxiliary::Initialize(const TCHAR* CommandLine)
 {
 	TRACE_CPUPROFILER_EVENT_SCOPE(FTraceAux_Init);
 
-#if WITH_UNREAL_TRACE_LAUNCH
+#if UE_TRACE_SERVER_LAUNCH_ENABLED && UE_TRACE_SERVER_CONTROLS_ENABLED
+	// Auto launch Unreal Trace Server for certain configurations
 	if (!(FParse::Param(CommandLine, TEXT("notraceserver")) || FParse::Param(CommandLine, TEXT("buildmachine"))))
 	{
 		TRACE_CPUPROFILER_EVENT_SCOPE(FTraceAux_LaunchUnrealTrace);
-		LaunchUnrealTraceInternal(CommandLine);
+		FTraceServerControls::Start();
 	}
 #endif
 
@@ -1821,3 +1687,149 @@ void FTraceAuxiliary::TryAutoConnect()
 #endif // PLATFORM_WINDOWS
 #endif // UE_TRACE_ENABLED
 }
+
+////////////////////////////////////////////////////////////////////////////////
+#if UE_TRACE_SERVER_CONTROLS_ENABLED
+
+#if PLATFORM_WINDOWS
+bool LaunchTraceServerCommand(const TCHAR* Command, bool bAddSponsor)
+{
+	TWideStringBuilder<MAX_PATH + 32> CreateProcArgs;
+	CreateProcArgs << "\"" << FPaths::EngineDir() << TEXT("Binaries/Win64/UnrealTraceServer.exe\"");
+	CreateProcArgs << TEXT(" ") << Command;
+	if (bAddSponsor)
+	{
+		CreateProcArgs << TEXT(" --sponsor ") << FPlatformProcess::GetCurrentProcessId();
+	}
+
+	uint32 CreateProcFlags = CREATE_BREAKAWAY_FROM_JOB;
+	if (FParse::Param(FCommandLine::Get(), TEXT("traceshowstore")))
+	{
+		CreateProcFlags |= CREATE_NEW_CONSOLE;
+	}
+	else
+	{
+		CreateProcFlags |= CREATE_NO_WINDOW;
+	}
+	STARTUPINFOW StartupInfo = { sizeof(STARTUPINFOW) };
+	PROCESS_INFORMATION ProcessInfo = {};
+	const BOOL bOk = CreateProcessW(nullptr, LPWSTR(*CreateProcArgs), nullptr, nullptr,
+									false, CreateProcFlags, nullptr, nullptr, &StartupInfo, &ProcessInfo);
+
+	if (!bOk)
+	{
+		UE_LOG(LogCore, Display, TEXT("Unable to launch the Unreal Trace Server with '%s' (%08x)"), *CreateProcArgs, GetLastError());
+		return false;
+	}
+
+	bool bSuccess = false;
+	if (WaitForSingleObject(ProcessInfo.hProcess, 5000) == WAIT_TIMEOUT)
+	{
+		UE_LOG(LogCore, Warning, TEXT("Timed out waiting for the Unreal Trace Server to start"));
+	}
+	else
+	{
+		DWORD ExitCode = 0x0000'a9e0;
+		GetExitCodeProcess(ProcessInfo.hProcess, &ExitCode);
+		if (ExitCode)
+		{
+			UE_LOG(LogCore, Warning, TEXT("Unreal Trace Server returned an error (0x%08x)"), ExitCode);
+		}
+		else
+		{
+			UE_LOG(LogCore, Log, TEXT("Unreal Trace Server launch successful"));
+			bSuccess = true;
+		}
+	}
+
+	CloseHandle(ProcessInfo.hProcess);
+	CloseHandle(ProcessInfo.hThread);
+	
+	return bSuccess;
+}
+#endif
+
+#if PLATFORM_LINUX || PLATFORM_MAC
+static bool LaunchTraceServerCommand(const TCHAR* Command, bool bAddSponsor)
+{
+// TSAN doesn't like fork(), so disable this for now.
+#if !USING_THREAD_SANITISER 
+	TAnsiStringBuilder<320> BinPath;
+	BinPath << TCHAR_TO_UTF8(*FPaths::EngineDir());
+#if PLATFORM_UNIX
+	BinPath << "Binaries/Linux/UnrealTraceServer";
+#elif PLATFORM_MAC
+	BinPath << "Binaries/Mac/UnrealTraceServer";
+#endif
+	BinPath.ToString(); //Ensure zero termination
+
+	if (access(*BinPath, F_OK) < 0)
+	{
+		UE_LOG(LogCore, Display, TEXT("UnrealTraceServer: Binary not found (%s)"), ANSI_TO_TCHAR(*BinPath));
+		return false;
+	}
+
+	TAnsiStringBuilder<64> ForkArg;
+	ForkArg << " " << Command;
+	if (bAddSponsor)
+	{
+		ForkArg << TEXT(" --sponsor ") << FPlatformProcess::GetCurrentProcessId();
+	}
+	ForkArg.ToString(); //Ensure zero termination
+
+	pid_t UtsPid = fork();
+	if (UtsPid < 0)
+	{
+		UE_LOG(LogCore, Display, TEXT("UnrealTraceServer: Unable to fork (errno: %d)"), errno);
+		return false;
+	}
+	else if (UtsPid == 0)
+	{
+		char* Args[] = { BinPath.GetData(), ForkArg.GetData(), nullptr };
+		extern char** environ;
+		execve(*BinPath, Args, environ);
+		_exit(0x80 | (errno & 0x7f));
+	}
+
+	int32 WaitStatus = 0;
+	do
+	{
+		int32 WaitRet = waitpid(UtsPid, &WaitStatus, 0);
+		if (WaitRet < 0)
+		{
+			UE_LOG(LogCore, Display, TEXT("UnrealTraceServer: waitpid() error; (errno: %d)"), errno);
+			return false;
+		}
+	}
+	while (!WIFEXITED(WaitStatus));
+
+	int32 UtsRet = WEXITSTATUS(WaitStatus);
+	if (UtsRet)
+	{
+		UE_LOG(LogCore, Display, TEXT("UnrealTraceServer: Trace store returned an error (0x%08x)"), UtsRet);
+		return false;
+	}
+	else
+	{
+		UE_LOG(LogCore, Log, TEXT("UnrealTraceServer: Trace store launch successful"));
+		return true;
+	}
+#else // #if !USING_THREAD_SANITISER
+	return false;
+#endif
+}
+#endif
+
+////////////////////////////////////////////////////////////////////////////////
+bool FTraceServerControls::Start()
+{
+	return LaunchTraceServerCommand(TEXT("fork"), true);
+}
+
+////////////////////////////////////////////////////////////////////////////////
+bool FTraceServerControls::Stop()
+{
+	return LaunchTraceServerCommand(TEXT("kill"), false);
+}
+
+#endif // UE_TRACE_SERVER_CONTROLS_ENABLED
diff --git a/Engine/Source/Runtime/Core/Public/ProfilingDebugging/TraceAuxiliary.h b/Engine/Source/Runtime/Core/Public/ProfilingDebugging/TraceAuxiliary.h
--- a/Engine/Source/Runtime/Core/Public/ProfilingDebugging/TraceAuxiliary.h
+++ b/Engine/Source/Runtime/Core/Public/ProfilingDebugging/TraceAuxiliary.h
@@ -11,6 +11,10 @@
 #include "Logging/LogMacros.h"
 #include "Misc/Build.h"
 
+#if !defined(UE_TRACE_SERVER_CONTROLS_ENABLED)
+#	define UE_TRACE_SERVER_CONTROLS_ENABLED (PLATFORM_DESKTOP && !UE_BUILD_SHIPPING && !IS_PROGRAM)
+#endif
+
 ////////////////////////////////////////////////////////////////////////////////
 class FTraceAuxiliary
 {
@@ -222,3 +226,25 @@ public:
 	static CORE_API FOnSnapshotSaved OnSnapshotSaved;
 	
 };
+
+#if UE_TRACE_SERVER_CONTROLS_ENABLED
+
+/**
+ * Controls for Unreal Trace Server, the standalone server recording and storing traces.
+ */
+class FTraceServerControls
+{
+public:
+	/**
+	 * Launch the server using the "fork" command. This spins off an separate running process.	 * 
+	 * @return True if the server was successfully started or already running.
+	 */
+	static CORE_API bool Start();
+	/**
+	 * Stop any running instance of the server.
+	 * @return True if the stop command was successful. False otherwise.
+	 */
+	static CORE_API bool Stop();
+};
+
+#endif // UE_TRACE_SERVER_CONTROLS_ENABLED
\ No newline at end of file
