autortfm: first half of metrics gathering from the autortfm runtime

diff --git a/Engine/Source/Runtime/Core/Private/AutoRTFM/Context.cpp b/Engine/Source/Runtime/Core/Private/AutoRTFM/Context.cpp
--- a/Engine/Source/Runtime/Core/Private/AutoRTFM/Context.cpp
+++ b/Engine/Source/Runtime/Core/Private/AutoRTFM/Context.cpp
@@ -9,13 +9,30 @@
 #include "ScopedGuard.h"
 #include "Stats.h"
 #include "TransactionInlines.h"
+#include "AutoRTFM/AutoRTFMMetrics.h"
 
 #include "Templates/UniquePtr.h"
 #include "Containers/StringConv.h"
 
+namespace
+{
+AutoRTFM::FAutoRTFMMetrics GAutoRTFMMetrics;
+}
+
 namespace AutoRTFM
 {
 
+void ResetAutoRTFMMetrics()
+{
+	GAutoRTFMMetrics = FAutoRTFMMetrics{};
+}
+
+// get a snapshot of the current internal metrics
+FAutoRTFMMetrics GetAutoRTFMMetrics()
+{
+	return GAutoRTFMMetrics;
+}
+
 thread_local TUniquePtr<FContext> ContextTls;
 
 void FContext::InitializeGlobalData()
@@ -76,6 +93,8 @@ bool FContext::StartTransaction()
 	ASSERT(Status == EContextStatus::OnTrack);
 	PushTransaction(NewTransaction);
 
+	GAutoRTFMMetrics.NumTransactionsStarted++;
+
 	return true;
 }
 
@@ -115,11 +134,15 @@ ETransactionResult FContext::CommitTransaction()
 	// Parent transaction is now the current transaction
 	PopTransaction();
 
+	GAutoRTFMMetrics.NumTransactionsCommitted++;
+
 	return Result;
 }
 
 ETransactionResult FContext::AbortTransaction(bool bIsClosed, bool bIsCascading)
 {
+	GAutoRTFMMetrics.NumTransactionsAborted++;
+
 	ETransactionResult Result = ETransactionResult::AbortedByRequest;
 	ASSERT(Status == EContextStatus::OnTrack);
 	Status = bIsCascading ? EContextStatus::AbortedByCascade : EContextStatus::AbortedByRequest;
@@ -396,6 +419,7 @@ ETransactionResult FContext::Transact(void (*Function)(void* Arg), void* Arg)
 void FContext::AbortByRequestAndThrow()
 {
     ASSERT(Status == EContextStatus::OnTrack);
+	GAutoRTFMMetrics.NumTransactionsAbortedByRequest++;
     Status = EContextStatus::AbortedByRequest;
     CurrentTransaction->AbortAndThrow();
 }
@@ -403,6 +427,7 @@ void FContext::AbortByRequestAndThrow()
 void FContext::AbortByRequestWithoutThrowing()
 {
 	ASSERT(Status == EContextStatus::OnTrack);
+	GAutoRTFMMetrics.NumTransactionsAbortedByRequest++;
 	Status = EContextStatus::AbortedByRequest;
 	CurrentTransaction->AbortWithoutThrowing();
 }
@@ -411,6 +436,7 @@ void FContext::AbortByLanguageAndThrow()
 {
 	UE_DEBUG_BREAK();
     ASSERT(Status == EContextStatus::OnTrack);
+	GAutoRTFMMetrics.NumTransactionsAbortedByLanguage++;
     Status = EContextStatus::AbortedByLanguage;
     CurrentTransaction->AbortAndThrow();
 }
diff --git a/Engine/Source/Runtime/Core/Public/AutoRTFM/AutoRTFMMetrics.h b/Engine/Source/Runtime/Core/Public/AutoRTFM/AutoRTFMMetrics.h
new file mode 100644
--- /dev/null
+++ b/Engine/Source/Runtime/Core/Public/AutoRTFM/AutoRTFMMetrics.h
@@ -0,0 +1,26 @@
+// Copyright Epic Games, Inc. All Rights Reserved.
+
+#pragma once
+
+#include "AutoRTFM/AutoRTFM.h"
+
+namespace AutoRTFM
+{
+
+struct FAutoRTFMMetrics
+{
+	int64 NumTransactionsStarted;
+	int64 NumTransactionsCommitted;
+
+	int64 NumTransactionsAborted;
+	int64 NumTransactionsAbortedByRequest;
+	int64 NumTransactionsAbortedByLanguage;
+};
+
+// reset the internal metrics tracking back to zero
+void ResetAutoRTFMMetrics();
+
+// get a snapshot of the current internal metrics
+FAutoRTFMMetrics GetAutoRTFMMetrics();
+
+}
