Make it possible to enable CSV categories via config and also add support for hotfixing
#tests Tested with a replay in a cooked build

diff --git a/Engine/Source/Runtime/Core/Private/ProfilingDebugging/CsvProfiler.cpp b/Engine/Source/Runtime/Core/Private/ProfilingDebugging/CsvProfiler.cpp
--- a/Engine/Source/Runtime/Core/Private/ProfilingDebugging/CsvProfiler.cpp
+++ b/Engine/Source/Runtime/Core/Private/ProfilingDebugging/CsvProfiler.cpp
@@ -575,10 +575,20 @@ public:
 
 	void UpdateCategoryFromConfig(int32 CategoryIndex)
 	{
+		for (FString const& EnabledCategory : CategoriesEnabledInConfig)
+		{
+			if (FWildcardString::IsMatch(*EnabledCategory, *CategoryNames[CategoryIndex]))
+			{
+				UE_LOG(LogCsvProfiler, Log, TEXT("Config enabled category %s"), *CategoryNames[CategoryIndex]);
+				GCsvCategoriesEnabled[CategoryIndex] = true;
+			}
+		}
+
 		for (FString const& DisabledCategory : CategoriesDisabledInConfig)
 		{
 			if (FWildcardString::IsMatch(*DisabledCategory, *CategoryNames[CategoryIndex]))
 			{
+				UE_LOG(LogCsvProfiler, Log, TEXT("Config disabled category %s"), *CategoryNames[CategoryIndex]);
 				GCsvCategoriesEnabled[CategoryIndex] = false;
 			}
 		}
@@ -586,10 +596,21 @@ public:
 
 	void UpdateCategoriesFromConfig()
 	{
-		GConfig->GetArray(TEXT("CsvProfiler"), TEXT("DisabledCategories"), CategoriesDisabledInConfig, GEngineIni);
-		for (int i = 0; i < GetCategoryCount(); ++i)
+		TArray<FString> NewCategoriesDisabledInConfig;
+		TArray<FString> NewCategoriesEnabledInConfig;
+
+		GConfig->GetArray(TEXT("CsvProfiler"), TEXT("EnabledCategories"), NewCategoriesEnabledInConfig, GEngineIni);
+		GConfig->GetArray(TEXT("CsvProfiler"), TEXT("DisabledCategories"), NewCategoriesDisabledInConfig, GEngineIni);
+
+		// Check if the config changed. This prevents us from resetting categories unnecessarily at runtime (which could be confusing if debug category toggle commands are in play)
+		if (NewCategoriesEnabledInConfig != CategoriesEnabledInConfig || NewCategoriesDisabledInConfig != CategoriesDisabledInConfig )
 		{
-			UpdateCategoryFromConfig(i);
+			CategoriesEnabledInConfig = NewCategoriesEnabledInConfig;
+			CategoriesDisabledInConfig = NewCategoriesDisabledInConfig;
+			for (int i = 0; i < GetCategoryCount(); ++i)
+			{
+				UpdateCategoryFromConfig(i);
+			}
 		}
 	}
 
@@ -645,6 +666,7 @@ private:
 	TMap<FString, int32> CategoryNameToIndex;
 	TArray<FString> CategoryNames;
 	TArray<FString> CategoriesDisabledInConfig;
+	TArray<FString> CategoriesEnabledInConfig;
 
 	static FCsvCategoryData* Instance;
 };
@@ -2885,6 +2907,12 @@ void FCsvProfiler::BeginFrame()
 			{
 				UE_LOG(LogCsvProfiler, Display, TEXT("Capture Starting"));
 
+				if (GConfig)
+				{
+					// Update categories from the config. The config may have changed if there were hotfixes
+					FCsvCategoryData::Get()->UpdateCategoriesFromConfig();
+				}
+
 				// signal external profiler that we are capturing
 				OnCSVProfileStartDelegate.Broadcast();
 
