- Fixed previewing platforms without a named DeviceProfile

diff --git a/Engine/Source/Runtime/Core/Private/HAL/ConsoleManager.cpp b/Engine/Source/Runtime/Core/Private/HAL/ConsoleManager.cpp
--- a/Engine/Source/Runtime/Core/Private/HAL/ConsoleManager.cpp
+++ b/Engine/Source/Runtime/Core/Private/HAL/ConsoleManager.cpp
@@ -547,10 +547,12 @@ bool IConsoleManager::VisitPlatformCVarsForEmulation(FName PlatformName, const F
 			FConfigSection* NewSection = new FConfigSection();
 			bDeleteSection = true;
 
+			UE_LOG(LogConsoleManager, Verbose, TEXT("Gathering device profile cvars for %s, platform config %s"), *DeviceProfileName, *PlatformName.ToString());
 			// run the delegate (this code can't get into DP code directly, so we use a delegate), and walk over the results
 			for (TPair<FName, FString>& Pair : FCoreDelegates::GatherDeviceProfileCVars.Execute(DeviceProfileName))
 			{
 				NewSection->Add(Pair);
+				UE_LOG(LogConsoleManager, Verbose, TEXT("   %s = %s"), *Pair.Key.ToString(), *Pair.Value);
 			}
             
             Section = NewSection;
@@ -2962,8 +2964,8 @@ void FConsoleManager::LoadAllPlatformCVars(FName PlatformName, const FString& De
 	CachedPlatformsAndDeviceProfiles.Add(PlatformKey);
 	
 	// use the platform's base DeviceProfile for emulation
-	VisitPlatformCVarsForEmulation(PlatformName, DeviceProfileName,
-		[PlatformKey](const FString& CVarName, const FString& CVarValue, EConsoleVariableFlags SetByAndPreview)
+	VisitPlatformCVarsForEmulation(PlatformName, DeviceProfileName.IsEmpty() ? PlatformName.ToString() : DeviceProfileName,
+		[PlatformName, PlatformKey](const FString& CVarName, const FString& CVarValue, EConsoleVariableFlags SetByAndPreview)
 	{
 		// make sure the named cvar exists
 		IConsoleVariable* CVar = IConsoleManager::Get().FindConsoleVariable(*CVarName);
@@ -2978,6 +2980,10 @@ void FConsoleManager::LoadAllPlatformCVars(FName PlatformName, const FString& De
 		// now cache the passed in value
 		int32 SetBy = SetByAndPreview & ECVF_SetByMask;
 		PlatformCVar->SetOtherPlatformValue(*CVarValue, (EConsoleVariableFlags)SetBy, NAME_None);
+		
+		UE_LOG(LogConsoleManager, Verbose, TEXT("Loading %s@%s = %s [get = %s]"), *CVarName, *PlatformKey.ToString(),
+			   *CVarValue, *CVar->GetPlatformValueVariable(*PlatformName.ToString())->GetString());
+		
 	});
 }
 
