Replaced loop with modulo in FMath::Wrap to make it run in constant time.
Added some tests that would have failed with the previous implementation due to lack of floating point precision.
Added some constexpr tests, though only the integer tests can be supported before C++23.

diff --git a/Engine/Source/Runtime/Core/Public/Math/UnrealMathUtility.h b/Engine/Source/Runtime/Core/Public/Math/UnrealMathUtility.h
--- a/Engine/Source/Runtime/Core/Public/Math/UnrealMathUtility.h
+++ b/Engine/Source/Runtime/Core/Public/Math/UnrealMathUtility.h
@@ -6,9 +6,7 @@
 #include "Misc/AssertionMacros.h"
 #include "HAL/PlatformMath.h"
 #include "Math/MathFwd.h"
-#include "Templates/Decay.h"
-#include "Templates/IsFloatingPoint.h"
-#include "Templates/IsIntegral.h"
+#include "Templates/Identity.h"
 
 // Assert on non finite numbers. Used to track NaNs.
 #ifndef ENABLE_NAN_DIAGNOSTIC
@@ -609,25 +607,32 @@ public:
 	template< class T >
 	UE_NODISCARD static constexpr FORCEINLINE T Wrap(const T X, const T Min, const T Max)
 	{
+		// Use unsigned type for integers to allow for large ranges which don't overflow on subtraction
+		// We don't do that for floating point types because there are no unsigned versions of those.  We
+		// could bump up to the next size up (though there's no `long double` on some platforms), but such
+		// values near the extremes are unlikely.
+		using SizeType = typename std::conditional_t<std::is_integral_v<T>, std::make_unsigned<T>, TIdentity<T>>::type;
+
 		// Our asserts are not constexpr-friendly yet
 		// checkSlow(Min <= Max);
 
-		T Size = Max - Min;
+		SizeType Size = (SizeType)Max - (SizeType)Min;
 		if (Size == 0)
 		{
-			// Guard against zero-sized ranges causing an infinite loop.
+			// Guard against zero-sized ranges causing division by zero.
 			return Max;
 		}
 
 		T EndVal = X;
-		while (EndVal < Min)
+		if (EndVal < Min)
 		{
-			EndVal += Size;
+			SizeType Mod = FMath::Modulo((SizeType)((SizeType)Min - (SizeType)EndVal), Size);
+			EndVal = (Mod != (T)0) ? (T)((SizeType)Max - Mod) : Min;
 		}
-
-		while (EndVal > Max)
+		else if (EndVal > Max)
 		{
-			EndVal -= Size;
+			T Mod = FMath::Modulo((SizeType)((SizeType)EndVal - (SizeType)Max), Size);
+			EndVal = (Mod != (T)0) ? (T)((SizeType)Min + Mod) : Max;
 		}
 		return EndVal;
 	}
@@ -2639,3 +2644,9 @@ namespace LWC
 
 } // namespace LWC
 } // namespace UE
+
+#if UE_ENABLE_INCLUDE_ORDER_DEPRECATED_IN_5_4
+#include "Templates/Decay.h"
+#include "Templates/IsFloatingPoint.h"
+#include "Templates/IsIntegral.h"
+#endif
diff --git a/Engine/Source/Runtime/Core/Public/Templates/Identity.h b/Engine/Source/Runtime/Core/Public/Templates/Identity.h
--- a/Engine/Source/Runtime/Core/Public/Templates/Identity.h
+++ b/Engine/Source/Runtime/Core/Public/Templates/Identity.h
@@ -16,7 +16,8 @@
 template <typename T>
 struct TIdentity
 {
-	typedef T Type;
+	using Type = T;
+	using type = T;
 };
 
 template <typename T>
diff --git a/Engine/Source/Runtime/Core/Tests/Math/UnrealMathTest.cpp b/Engine/Source/Runtime/Core/Tests/Math/UnrealMathTest.cpp
--- a/Engine/Source/Runtime/Core/Tests/Math/UnrealMathTest.cpp
+++ b/Engine/Source/Runtime/Core/Tests/Math/UnrealMathTest.cpp
@@ -4616,6 +4616,55 @@ TEST_CASE_NAMED(FMathWrapTest, "System::Core::Math::Wrap", "[ApplicationContextM
 		/* SizeTo   */   7.23,
 		/* SizeStep */   0.59
 	);
+
+	// Test large values far away from the range
+	WrapTest(
+		/* ValFrom  */  1 << 30,
+		/* ValTo    */  (1 << 30) + 1,
+		/* ValStep  */  1,
+		/* MinFrom  */  0,
+		/* MinTo    */  1,
+		/* MinStep  */  1,
+		/* SizeFrom */  2,
+		/* SizeTo   */  3,
+		/* SizeStep */  1
+	);
+	WrapTest(
+		/* ValFrom  */  (float)(1 << 25),
+		/* ValTo    */  (float)(1 << 25) + 4.0f,
+		/* ValStep  */  4.0f,
+		/* MinFrom  */  0.0f,
+		/* MinTo    */  1.0f,
+		/* MinStep  */  1.0f,
+		/* SizeFrom */  1.0f,
+		/* SizeTo   */  2.0f,
+		/* SizeStep */  1.0f
+	);
+	WrapTest(
+		/* ValFrom  */  (double)(1ull << 54),
+		/* ValTo    */  (double)(1ull << 54) + 4.0,
+		/* ValStep  */  4.0,
+		/* MinFrom  */  0.0,
+		/* MinTo    */  1.0,
+		/* MinStep  */  1.0,
+		/* SizeFrom */  1.0,
+		/* SizeTo   */  2.0,
+		/* SizeStep */  1.0
+	);
+
+	// Test constexpr
+	static_assert(FMath::Wrap(-3, 0, 5) == 2);
+	//static_assert(FMath::Wrap(-3.0f, 0.0f, 5.0f) == 2); // needs constexpr fmod support in C++23
+	//static_assert(FMath::Wrap(-3.0, 0.0, 5.0) == 2); // needs constexpr fmod support in C++23
+
+	// These will fail to compile due to signed overload if we don't use unsigned diffs in the implementation
+	static_assert(FMath::Wrap(MAX_int32, MIN_int32 + 2, MAX_int32)); // range size will overflow
+	static_assert(FMath::Wrap(MIN_int32, MAX_int32 - 2, MAX_int32)); // distance from val to min will overflow
+	static_assert(FMath::Wrap(MAX_int32, MIN_int32, MIN_int32 + 2)); // distance from max to val will overflow
+
+	// Try with bytes too, where the value is more than 128 away from the range
+	static_assert(FMath::Wrap((int8)123, (int8)-20, (int8)-10) == (int8)-17);
+	static_assert(FMath::Wrap((int8)-123, (int8)10, (int8)20) == (int8)17);
 }
 class FInitVectorTestClass {
 public:
