Refactor FGenericPlatformOutputDevices::SetupOutputDevices code just a little so that we don't have to suffer code rot in a custom linux implementation

#PR 9643
#9643
#presubmit 657a2209be1cfc2f95b29cb7

diff --git a/Engine/Source/Programs/UnrealBuildTool/Platform/Linux/LinuxToolChain.cs b/Engine/Source/Programs/UnrealBuildTool/Platform/Linux/LinuxToolChain.cs
--- a/Engine/Source/Programs/UnrealBuildTool/Platform/Linux/LinuxToolChain.cs
+++ b/Engine/Source/Programs/UnrealBuildTool/Platform/Linux/LinuxToolChain.cs
@@ -539,6 +539,10 @@ namespace UnrealBuildTool
 
 			// build up the commandline common to C and C++
 
+			// These aren't supported on Linux at this time
+			Arguments.Add("-DUSE_DEBUG_LOGGING=0");
+			Arguments.Add("-DUSE_EVENT_LOGGING=0");
+
 			// always select the driver g++ in-case we are using a different binary for clang, such as clang/clang-cl
 			Arguments.Add("--driver-mode=g++");
 			if (ShouldUseLibcxx())
diff --git a/Engine/Source/Runtime/Core/Private/GenericPlatform/GenericPlatformOutputDevices.cpp b/Engine/Source/Runtime/Core/Private/GenericPlatform/GenericPlatformOutputDevices.cpp
--- a/Engine/Source/Runtime/Core/Private/GenericPlatform/GenericPlatformOutputDevices.cpp
+++ b/Engine/Source/Runtime/Core/Private/GenericPlatform/GenericPlatformOutputDevices.cpp
@@ -55,7 +55,9 @@ void FGenericPlatformOutputDevices::SetupOutputDevices()
 #endif // USE_DEBUG_LOGGING
 #endif
 
+#if USE_EVENT_LOGGING
 	GLog->AddOutputDevice(FPlatformOutputDevices::GetEventLog());
+#endif // USE_EVENT_LOGGING
 };
 
 void FGenericPlatformOutputDevices::ResetCachedAbsoluteFilename()
diff --git a/Engine/Source/Runtime/Core/Private/Unix/UnixPlatformOutputDevices.cpp b/Engine/Source/Runtime/Core/Private/Unix/UnixPlatformOutputDevices.cpp
--- a/Engine/Source/Runtime/Core/Private/Unix/UnixPlatformOutputDevices.cpp
+++ b/Engine/Source/Runtime/Core/Private/Unix/UnixPlatformOutputDevices.cpp
@@ -15,27 +15,6 @@
 
 class FOutputDeviceError;
 
-void FUnixOutputDevices::SetupOutputDevices()
-{
-	check(GLog);
-
-	ResetCachedAbsoluteFilename();
-
-	// add file log
-	GLog->AddOutputDevice(FPlatformOutputDevices::GetLog());
-
-	// @todo: set to false for minor utils?
-	if (GLogConsole != nullptr)
-	{
-		bool bLogToConsole = !NO_LOGGING && !FParse::Param(FCommandLine::Get(), TEXT("NOCONSOLE"));
-		if (bLogToConsole)
-		{
-			GLog->AddOutputDevice(GLogConsole);
-		}
-	}
-	// debug and event logging is not really supported on Unix. 
-}
-
 FString FUnixOutputDevices::GetAbsoluteLogFilename()
 {
 	// FIXME: this function should not exist once FGenericPlatformOutputDevices::GetAbsoluteLogFilename() returns absolute filename (see UE-25650)
diff --git a/Engine/Source/Runtime/Core/Public/Misc/OutputDevice.h b/Engine/Source/Runtime/Core/Public/Misc/OutputDevice.h
--- a/Engine/Source/Runtime/Core/Public/Misc/OutputDevice.h
+++ b/Engine/Source/Runtime/Core/Public/Misc/OutputDevice.h
@@ -19,6 +19,10 @@ namespace UE { class FLogRecord; }
 #define USE_DEBUG_LOGGING 1
 #endif
 
+#ifndef USE_EVENT_LOGGING
+#define USE_EVENT_LOGGING 1
+#endif
+
 #if !PLATFORM_SUPPORTS_COLORIZED_OUTPUT_DEVICE
 	// don't support colorized text on consoles
 	#define SET_WARN_COLOR(Color)
diff --git a/Engine/Source/Runtime/Core/Public/Unix/UnixPlatformOutputDevices.h b/Engine/Source/Runtime/Core/Public/Unix/UnixPlatformOutputDevices.h
--- a/Engine/Source/Runtime/Core/Public/Unix/UnixPlatformOutputDevices.h
+++ b/Engine/Source/Runtime/Core/Public/Unix/UnixPlatformOutputDevices.h
@@ -15,7 +15,6 @@ class FOutputDeviceError;
 
 struct FUnixOutputDevices : public FGenericPlatformOutputDevices
 {
-	static CORE_API void SetupOutputDevices();
 	static CORE_API FString GetAbsoluteLogFilename();
 	static CORE_API FOutputDevice* GetEventLog();
 	static CORE_API FOutputDeviceError* GetError();
