Fixing Android shipping builds after typo in CL 30376326

diff --git a/Engine/Source/Runtime/Core/Private/Android/AndroidScudoMemoryTrace.cpp b/Engine/Source/Runtime/Core/Private/Android/AndroidScudoMemoryTrace.cpp
--- a/Engine/Source/Runtime/Core/Private/Android/AndroidScudoMemoryTrace.cpp
+++ b/Engine/Source/Runtime/Core/Private/Android/AndroidScudoMemoryTrace.cpp
@@ -6,7 +6,7 @@
 #	define UE_MEMORY_TRACE_ANDROID_ENABLE_SCUDO_TRACING_SUPPORT 0
 #endif
 
-#if 1 || UE_MEMORY_TRACE_ANDROID_ENABLE_SCUDO_TRACING_SUPPORT
+#if UE_MEMORY_TRACE_ANDROID_ENABLE_SCUDO_TRACING_SUPPORT
 
 #include "AndroidCallstackTrace.h"
 #include "ProfilingDebugging/MemoryTrace.h"
