Add TBitArray range-based for loop support

#tests unit

diff --git a/Engine/Source/Runtime/Core/Public/Containers/BitArray.h b/Engine/Source/Runtime/Core/Public/Containers/BitArray.h
--- a/Engine/Source/Runtime/Core/Public/Containers/BitArray.h
+++ b/Engine/Source/Runtime/Core/Public/Containers/BitArray.h
@@ -303,6 +303,16 @@ public:
 
 	int32  WordIndex;
 	uint32 Mask;
+
+	FORCEINLINE bool operator==(FRelativeBitReference Other) const
+	{
+		return (WordIndex == Other.WordIndex) & (Mask == Other.Mask);
+	}
+
+	FORCEINLINE bool operator!=(FRelativeBitReference Other) const
+	{
+		return !(Other == *this);
+	}
 };
 
 class FBitArrayMemory
@@ -1463,11 +1473,18 @@ public:
 			}
 			return *this;
 		}
+
+		FORCEINLINE FBitReference operator*() const
+		{
+			return GetValue();
+		}
+
 		/** conversion to "bool" returning true if the iterator is valid. */
 		FORCEINLINE explicit operator bool() const
 		{ 
 			return Index < Array.Num(); 
 		}
+
 		/** inverse of the "bool" operator */
 		FORCEINLINE bool operator !() const 
 		{
@@ -1504,6 +1521,11 @@ public:
 			return *this;
 		}
 
+		FORCEINLINE FBitReference operator*() const
+		{
+			return GetValue();
+		}
+
 		/** conversion to "bool" returning true if the iterator is valid. */
 		FORCEINLINE explicit operator bool() const
 		{ 
@@ -1693,6 +1715,13 @@ public:
 		}
 	};
 
+	/** Enables range-based for loops, DO NOT USE DIRECTLY. */
+
+	FORCEINLINE FIterator		begin()			{ return FIterator(*this); }
+	FORCEINLINE FConstIterator	begin() const	{ return FConstIterator(*this); }
+	FORCEINLINE FIterator		end()			{ return FIterator(*this, NumBits); }
+	FORCEINLINE FConstIterator	end() const		{ return FConstIterator(*this, NumBits); }
+
 private:
 	AllocatorType AllocatorInstance;
 	int32         NumBits;
diff --git a/Engine/Source/Runtime/Core/Tests/Containers/BitArrayTests.cpp b/Engine/Source/Runtime/Core/Tests/Containers/BitArrayTests.cpp
--- a/Engine/Source/Runtime/Core/Tests/Containers/BitArrayTests.cpp
+++ b/Engine/Source/Runtime/Core/Tests/Containers/BitArrayTests.cpp
@@ -6,6 +6,7 @@
 #include "Misc/StringBuilder.h"
 #include "Serialization/MemoryReader.h"
 #include "Serialization/MemoryWriter.h"
+#include "Templates/Function.h"
 
 #include "Tests/TestHarnessAdapter.h"
 
@@ -1556,6 +1557,23 @@ TEST_CASE_NAMED(FBitArrayTestMisc, "System::Core::Containers::BitArray::Misc", "
 	Instance.TestInitAndSetNumUninitialized();
 }
 
+TEST_CASE_NAMED(FBitArrayTestFor, "System::Core::Containers::BitArray::For", "[ApplicationContextMask][EngineFilter]")
+{
+	std::initializer_list<bool> Expected = {true, false, true, true, false, false};
+
+	TBitArray<> Bits;
+	for (bool bBit : Expected)
+	{
+		Bits.Add(bBit);
+	}
+
+	const bool* ExpectedIt = Expected.begin();
+	for (bool bActual : Bits)
+	{
+		CHECK(bActual == *ExpectedIt++);
+	}
+}
+
 TEST_CASE_NAMED(FBitArrayInvariantsTest, "System::Core::Containers::BitArray::Invariants", "[ApplicationContextMask][SmokeFilter]")
 {
 	using namespace UE::BitArrayTest;
