Temporary fix, removing sponsor argument

diff --git a/Engine/Source/Runtime/Core/Private/ProfilingDebugging/TraceAuxiliary.cpp b/Engine/Source/Runtime/Core/Private/ProfilingDebugging/TraceAuxiliary.cpp
--- a/Engine/Source/Runtime/Core/Private/ProfilingDebugging/TraceAuxiliary.cpp
+++ b/Engine/Source/Runtime/Core/Private/ProfilingDebugging/TraceAuxiliary.cpp
@@ -1823,7 +1823,7 @@ static bool LaunchTraceServerCommand(const TCHAR* Command, bool bAddSponsor)
 ////////////////////////////////////////////////////////////////////////////////
 bool FTraceServerControls::Start()
 {
-	return LaunchTraceServerCommand(TEXT("fork"), true);
+	return LaunchTraceServerCommand(TEXT("fork"), false);
 }
 
 ////////////////////////////////////////////////////////////////////////////////
