CsvProfiler: Added named events (and trace cpu timing events) also for Csv stats identified by FName.

diff --git a/Engine/Source/Runtime/Core/Private/ProfilingDebugging/CsvProfiler.cpp b/Engine/Source/Runtime/Core/Private/ProfilingDebugging/CsvProfiler.cpp
--- a/Engine/Source/Runtime/Core/Private/ProfilingDebugging/CsvProfiler.cpp
+++ b/Engine/Source/Runtime/Core/Private/ProfilingDebugging/CsvProfiler.cpp
@@ -289,10 +289,10 @@ static FAutoConsoleVariableRef CVarNamedEventsTiming(TEXT("csv.NamedEventsTiming
 
 void CsvBeginNamedEvent(FColor Color, const char* NamedEventName)
 {
-#if CPUPROFILERTRACE_ENABLED 
+#if CPUPROFILERTRACE_ENABLED
 	if (UE_TRACE_CHANNELEXPR_IS_ENABLED(CpuChannel))
 	{
-		FCpuProfilerTrace::OutputBeginDynamicEvent(NamedEventName);
+		FCpuProfilerTrace::OutputBeginDynamicEvent(NamedEventName, __FILE__, __LINE__);
 	}
 	else
 #endif
@@ -305,6 +305,23 @@ void CsvBeginNamedEvent(FColor Color, const char* NamedEventName)
 	}
 }
 
+void CsvBeginNamedEvent(FColor Color, const FName& StatName)
+{
+#if CPUPROFILERTRACE_ENABLED
+	if (UE_TRACE_CHANNELEXPR_IS_ENABLED(CpuChannel))
+	{
+		FCpuProfilerTrace::OutputBeginDynamicEvent(StatName, __FILE__, __LINE__);
+	}
+	else
+#endif
+	{
+#if PLATFORM_IMPLEMENTS_BeginNamedEventStatic
+		FPlatformMisc::BeginNamedEventStatic(Color, *StatName.ToString());
+#else
+		FPlatformMisc::BeginNamedEvent(Color, *StatName.ToString());
+#endif
+	}
+}
 
 void CsvEndNamedEvent()
 {
@@ -3389,6 +3406,12 @@ void FCsvProfiler::BeginStat(const FName& StatName, uint32 CategoryIndex)
 #if RECORD_TIMESTAMPS
 	if (GCsvProfilerIsCapturing && GCsvCategoriesEnabled[CategoryIndex])
 	{
+#if CSV_PROFILER_SUPPORT_NAMED_EVENTS
+		if (UNLIKELY(GCsvProfilerNamedEventsTiming))
+		{
+			CsvBeginNamedEvent(FColor(255, 128, 255), StatName);
+		}
+#endif
 		FCsvProfilerThreadData::Get().AddTimestampBegin(StatName, CategoryIndex);
 	}
 #endif
@@ -3416,6 +3439,12 @@ void FCsvProfiler::EndStat(const FName& StatName, uint32 CategoryIndex)
 	if (GCsvProfilerIsCapturing && GCsvCategoriesEnabled[CategoryIndex])
 	{
 		FCsvProfilerThreadData::Get().AddTimestampEnd(StatName, CategoryIndex);
+#if CSV_PROFILER_SUPPORT_NAMED_EVENTS
+		if (UNLIKELY(GCsvProfilerNamedEventsTiming))
+		{
+			CsvEndNamedEvent();
+		}
+#endif
 	}
 #endif
 }
@@ -3442,7 +3471,6 @@ void FCsvProfiler::EndExclusiveStat(const char * StatName)
 	if (GCsvProfilerIsCapturing && GCsvCategoriesEnabled[CSV_CATEGORY_INDEX(Exclusive)])
 	{
 		FCsvProfilerThreadData::Get().AddTimestampExclusiveEnd(StatName);
-
 #if CSV_PROFILER_SUPPORT_NAMED_EVENTS
 		if (UNLIKELY(GCsvProfilerNamedEventsExclusive))
 		{
