FNames: Change FLazyName to keep track of the displayname (aka whether each letter is upper case or lower case) of its input when converted to an FName.
#rnx

diff --git a/Engine/Source/Runtime/Core/Private/UObject/UnrealNames.cpp b/Engine/Source/Runtime/Core/Private/UObject/UnrealNames.cpp
--- a/Engine/Source/Runtime/Core/Private/UObject/UnrealNames.cpp
+++ b/Engine/Source/Runtime/Core/Private/UObject/UnrealNames.cpp
@@ -3727,6 +3727,7 @@ void CheckLazyName(const CharType(&Literal)[N])
 	CharType Literal2[N];
 	FMemory::Memcpy(Literal2, Literal);
 	check(FLazyName(Literal) == FLazyName(Literal2));
+	check(WriteToString<64>(FLazyName(Literal).Resolve()).ToView().Equals(WriteToString<64>(Literal).ToView(), ESearchCase::CaseSensitive));
 }
 
 static void TestNameBatch();
@@ -3933,6 +3934,8 @@ void FName::AutoTest()
 
 	
 	CheckLazyName("Hej");
+	CheckLazyName("hej");
+	CheckLazyName("HEJ");
 	CheckLazyName(TEXT("Hej"));
 	CheckLazyName("Hej_0");
 	CheckLazyName("Hej_00");
@@ -4142,25 +4145,27 @@ FName FLazyName::Resolve() const
 
 	if (Copy.IsName())
 	{
-		FNameEntryId Id = Copy.AsName();
+		FNameEntryId ComparisonId = Copy.GetComparisionId();
+		FNameEntryId DisplayId = Copy.GetDisplayId();
 #if UE_FNAME_OUTLINE_NUMBER
-		return FNameHelper::MakeWithNumber(FNameEntryIds{ Id, Id}, FNAME_Add, Number);
+		return FNameHelper::MakeWithNumber(FNameEntryIds{ ComparisonId, DisplayId }, FNAME_Add, Number);
 #else // UE_FNAME_OUTLINE_NUMBER
-		return FName(Id, Id, Number);
+		return FName(ComparisonId, DisplayId, Number);
 #endif // UE_FNAME_OUTLINE_NUMBER
 	}
 
 	// Resolve to FName but throw away the number part
-	FNameEntryId Id = bLiteralIsWide ? FName(Copy.AsWideLiteral()).GetComparisonIndex()
-										: FName(Copy.AsAnsiLiteral()).GetComparisonIndex();
+	FName FullName = bLiteralIsWide ? FName(Copy.AsWideLiteral()) : FName(Copy.AsAnsiLiteral());
+	FNameEntryId ComparisonId = FullName.GetComparisonIndex();
+	FNameEntryId DisplayId = FullName.GetDisplayIndex();
 
 	// Deliberately unsynchronized write of word-sized int, ok if multiple threads resolve same lazy name
-	Either = FLiteralOrName(Id);
+	Either = FLiteralOrName(ComparisonId, DisplayId);
 
 #if UE_FNAME_OUTLINE_NUMBER
-	return FNameHelper::MakeWithNumber(FNameEntryIds{Id, Id}, FNAME_Add, Number);
+	return FNameHelper::MakeWithNumber(FNameEntryIds{ ComparisonId, DisplayId}, FNAME_Add, Number);
 #else // UE_FNAME_OUTLINE_NUMBER
-	return FName(Id, Id, Number);		
+	return FName(ComparisonId, DisplayId, Number);
 #endif // UE_FNAME_OUTLINE_NUMBER
 }
 
diff --git a/Engine/Source/Runtime/Core/Public/UObject/NameTypes.h b/Engine/Source/Runtime/Core/Public/UObject/NameTypes.h
--- a/Engine/Source/Runtime/Core/Public/UObject/NameTypes.h
+++ b/Engine/Source/Runtime/Core/Public/UObject/NameTypes.h
@@ -1575,7 +1575,7 @@ class FLazyName
 {
 public:
 	FLazyName()
-		: Either(FNameEntryId())
+		: Either(FNameEntryId(), FNameEntryId())
 	{}
 
 	/** @param Literal must be a string literal */
@@ -1595,7 +1595,7 @@ public:
 	{}
 
 	explicit FLazyName(FName Name)
-		: Either(Name.GetComparisonIndex())
+		: Either(Name.GetComparisonIndex(), Name.GetDisplayIndex())
 		, Number(Name.GetNumber())
 	{}
 	
@@ -1611,6 +1611,7 @@ private:
 	{
 		// NOTE: uses high bit of pointer for flag; this may be an issue in future when high byte of address may be used for features like hardware ASAN
 		static constexpr uint64 LiteralFlag = uint64(1) << (sizeof(uint64) * 8 - 1);
+		static constexpr uint32 DisplayNameShift = 32;
 
 		explicit FLiteralOrName(const ANSICHAR* Literal)
 			: Int(reinterpret_cast<uint64>(Literal) | LiteralFlag)
@@ -1620,9 +1621,13 @@ private:
 			: Int(reinterpret_cast<uint64>(Literal) | LiteralFlag)
 		{}
 
-		explicit FLiteralOrName(FNameEntryId Name)
-			: Int(Name.ToUnstableInt())
-		{}
+		explicit FLiteralOrName(FNameEntryId ComparisionId, FNameEntryId DisplayId)
+			: Int(ComparisionId.ToUnstableInt() |
+				(WITH_CASE_PRESERVING_NAME ? ((static_cast<uint64>(DisplayId.ToUnstableInt()) << DisplayNameShift)) : 0))
+		{
+			// FName indices fit within 31 bits -> (DisplayId & LiteralFlag) == 0 -> IsName() == true.
+			checkName(IsName());
+		}
 
 		bool IsName() const
 		{
@@ -1634,10 +1639,19 @@ private:
 			return (LiteralFlag & Int) != 0;
 		}
 
-		FNameEntryId AsName() const
+		FNameEntryId GetComparisionId() const
 		{
 			return FNameEntryId::FromUnstableInt(static_cast<uint32>(Int));
 		}
+
+		FNameEntryId GetDisplayId() const
+		{
+#if WITH_CASE_PRESERVING_NAME
+			return FNameEntryId::FromUnstableInt(static_cast<uint32>(Int >> DisplayNameShift)); // Can assume LiteralFlag == 0 
+#else
+			return GetComparisionId();
+#endif
+		}
 		
 		const ANSICHAR* AsAnsiLiteral() const
 		{
