I/O - log warnings when failing to uncompress chunk blocks

diff --git a/Engine/Source/Runtime/Core/Private/IO/IoChunkEncoding.cpp b/Engine/Source/Runtime/Core/Private/IO/IoChunkEncoding.cpp
--- a/Engine/Source/Runtime/Core/Private/IO/IoChunkEncoding.cpp
+++ b/Engine/Source/Runtime/Core/Private/IO/IoChunkEncoding.cpp
@@ -296,6 +296,8 @@ bool FIoChunkEncoding::Decode(const FIoChunkDecodingParams& Params, FMemoryView
 			FIoBlockHash Hash = HashBlock(BlockView);
 			if (Hash != BlockHash[BlockIndex])
 			{
+				UE_LOG(LogIoDispatcher, Warning, TEXT("Block hash mismatch, index '%d', hash '0x%uX', expected hash '0x%uX'"),
+					BlockIndex, Hash, BlockHash[BlockIndex]);
 				return false;
 			}
 		}
@@ -316,6 +318,8 @@ bool FIoChunkEncoding::Decode(const FIoChunkDecodingParams& Params, FMemoryView
 			{
 				if (!FCompression::UncompressMemory(Params.CompressionFormat, OutRawData.GetData(), int32(RawBlockReadSize), BlockView.GetData(), CompressedBlockSize))
 				{
+					UE_LOG(LogIoDispatcher, Warning, TEXT("Failed to uncompress block '%d', format '%s', raw size '" UINT64_FMT "' compressed size '%u'"),
+						BlockIndex, *Params.CompressionFormat.ToString(), RawBlockSize, AlignedBlockSize);
 					return false;
 				}
 			}
@@ -324,6 +328,8 @@ bool FIoChunkEncoding::Decode(const FIoChunkDecodingParams& Params, FMemoryView
 				FIoBuffer RawBlockTmp = FIoBuffer(RawBlockSize);
 				if (!FCompression::UncompressMemory(Params.CompressionFormat, RawBlockTmp.GetData(), int32(RawBlockSize), BlockView.GetData(), CompressedBlockSize))
 				{
+					UE_LOG(LogIoDispatcher, Warning, TEXT("Failed to uncompress block '%d', format '%s', raw size '" UINT64_FMT "' compressed size '%u'"),
+						BlockIndex, *Params.CompressionFormat.ToString(), RawBlockSize, AlignedBlockSize);
 					return false;
 				}
 				OutRawData.CopyFrom(RawBlockTmp.GetView().Mid(RawBlockOffset, RawBlockReadSize));
