Made UE_DEPRECATED_MACRO emit errors in MSVC when bWarningsAsErrors=true.

diff --git a/Engine/Source/Programs/UnrealBuildTool/Configuration/UEBuildTarget.cs b/Engine/Source/Programs/UnrealBuildTool/Configuration/UEBuildTarget.cs
--- a/Engine/Source/Programs/UnrealBuildTool/Configuration/UEBuildTarget.cs
+++ b/Engine/Source/Programs/UnrealBuildTool/Configuration/UEBuildTarget.cs
@@ -5193,6 +5193,16 @@ namespace UnrealBuildTool
 			// Set the global app name
 			GlobalCompileEnvironment.Definitions.Add(String.Format("UE_APP_NAME=\"{0}\"", AppName));
 
+			// Set the global
+			if (Rules.bWarningsAsErrors)
+			{
+				GlobalCompileEnvironment.Definitions.Add("UE_WARNINGS_AS_ERRORS=1");
+			}
+			else
+			{
+				GlobalCompileEnvironment.Definitions.Add("UE_WARNINGS_AS_ERRORS=0");
+			}
+
 			// Add global definitions for project-specific binaries. HACK: Also defining for monolithic builds in binary releases. Might be better to set this via command line instead?
 			if (!bUseSharedBuildEnvironment || bCompileMonolithic)
 			{
diff --git a/Engine/Source/Runtime/Core/Public/Clang/ClangPlatformCompilerPreSetup.h b/Engine/Source/Runtime/Core/Public/Clang/ClangPlatformCompilerPreSetup.h
--- a/Engine/Source/Runtime/Core/Public/Clang/ClangPlatformCompilerPreSetup.h
+++ b/Engine/Source/Runtime/Core/Public/Clang/ClangPlatformCompilerPreSetup.h
@@ -266,6 +266,11 @@
 		_Pragma(PREPROCESSOR_TO_STRING(message(Warning)))
 #endif // EMIT_CUSTOM_WARNING_AT_LINE
 
+#ifndef EMIT_CUSTOM_ERROR_AT_LINE
+	#define EMIT_CUSTOM_ERROR_AT_LINE(Line, Error) \
+		_Pragma(PREPROCESSOR_TO_STRING(message(Error)))
+#endif // EMIT_CUSTOM_ERROR_AT_LINE
+
 #ifndef UE_COMPILER_THIRD_PARTY_INCLUDES_START
 	#define UE_COMPILER_THIRD_PARTY_INCLUDES_START \
 		PRAGMA_DISABLE_REORDER_WARNINGS \
diff --git a/Engine/Source/Runtime/Core/Public/GenericPlatform/GenericPlatformCompilerPreSetup.h b/Engine/Source/Runtime/Core/Public/GenericPlatform/GenericPlatformCompilerPreSetup.h
--- a/Engine/Source/Runtime/Core/Public/GenericPlatform/GenericPlatformCompilerPreSetup.h
+++ b/Engine/Source/Runtime/Core/Public/GenericPlatform/GenericPlatformCompilerPreSetup.h
@@ -16,13 +16,26 @@
 	#define EMIT_CUSTOM_WARNING_AT_LINE(Line, Warning)
 #endif // EMIT_CUSTOM_WARNING_AT_LINE
 
+#ifndef EMIT_CUSTOM_ERROR_AT_LINE
+	#define EMIT_CUSTOM_ERROR_AT_LINE(Line, Error)
+#endif // EMIT_CUSTOM_ERROR_AT_LINE
+
 #ifndef EMIT_CUSTOM_WARNING
 	#define EMIT_CUSTOM_WARNING(Warning) \
 		EMIT_CUSTOM_WARNING_AT_LINE(__LINE__, Warning)
 #endif // EMIT_CUSTOM_WARNING
 
+#ifndef EMIT_CUSTOM_ERROR
+	#define EMIT_CUSTOM_ERROR(Error) \
+		EMIT_CUSTOM_ERROR_AT_LINE(__LINE__, Error)
+#endif // EMIT_CUSTOM_ERROR
+
 #ifndef UE_DEPRECATED_MACRO
-	#define UE_DEPRECATED_MACRO(Version, Message) EMIT_CUSTOM_WARNING(Message " Please update your code to the new API before upgrading to the next release, otherwise your project will no longer compile.")
+	#if UE_WARNINGS_AS_ERRORS
+		#define UE_DEPRECATED_MACRO(Version, Message) EMIT_CUSTOM_ERROR(Message " Please update your code to the new API before upgrading to the next release, otherwise your project will no longer compile.")
+	#else
+		#define UE_DEPRECATED_MACRO(Version, Message) EMIT_CUSTOM_WARNING(Message " Please update your code to the new API before upgrading to the next release, otherwise your project will no longer compile.")
+	#endif
 #endif
 
 #ifndef DEPRECATED_MACRO
diff --git a/Engine/Source/Runtime/Core/Public/MSVC/MSVCPlatformCompilerPreSetup.h b/Engine/Source/Runtime/Core/Public/MSVC/MSVCPlatformCompilerPreSetup.h
--- a/Engine/Source/Runtime/Core/Public/MSVC/MSVCPlatformCompilerPreSetup.h
+++ b/Engine/Source/Runtime/Core/Public/MSVC/MSVCPlatformCompilerPreSetup.h
@@ -246,4 +246,7 @@
 
 	#define EMIT_CUSTOM_WARNING_AT_LINE(Line, Warning) \
 		__pragma(message(WARNING_LOCATION(Line) ": warning C4996: " Warning))
+
+	#define EMIT_CUSTOM_ERROR_AT_LINE(Line, Error) \
+		__pragma(message(WARNING_LOCATION(Line) ": error C4996: " Error))
 #endif // defined(__clang__)
