Fix incorrect tangent computation in FInterpCurve::AutoSetTangents for linear InterpCurves

diff --git a/Engine/Source/Runtime/Core/Public/Math/InterpCurve.h b/Engine/Source/Runtime/Core/Public/Math/InterpCurve.h
--- a/Engine/Source/Runtime/Core/Public/Math/InterpCurve.h
+++ b/Engine/Source/Runtime/Core/Public/Math/InterpCurve.h
@@ -804,9 +804,10 @@ void FInterpCurve<T>::AutoSetTangents(float Tension, bool bStationaryEndpoints)
 		}
 		else if (ThisPoint.InterpMode == CIM_Linear)
 		{
-			T Tangent = NextPoint.OutVal - ThisPoint.OutVal;
-			ThisPoint.ArriveTangent = Tangent;
-			ThisPoint.LeaveTangent = Tangent;
+			ThisPoint.LeaveTangent = NextPoint.OutVal - ThisPoint.OutVal;
+
+			// Following from a curve, we should set the tangents equal so that there are no discontinuities
+			ThisPoint.ArriveTangent = PrevPoint.IsCurveKey() ? ThisPoint.LeaveTangent : ThisPoint.OutVal - PrevPoint.OutVal;
 		}
 		else if (ThisPoint.InterpMode == CIM_Constant)
 		{
