Core Tests: Relabel non-System.Core core tests as System.Core to match the other core tests.
#rnx

diff --git a/Engine/Source/Runtime/Core/Tests/Math/GuardedIntTest.cpp b/Engine/Source/Runtime/Core/Tests/Math/GuardedIntTest.cpp
--- a/Engine/Source/Runtime/Core/Tests/Math/GuardedIntTest.cpp
+++ b/Engine/Source/Runtime/Core/Tests/Math/GuardedIntTest.cpp
@@ -176,7 +176,7 @@ bool DoGuardedIntTestForType(FAutomationTestBase& Test)
 }
 }
 
-IMPLEMENT_SIMPLE_AUTOMATION_TEST(FGuardedIntTest, "Core.GuardedInt", EAutomationTestFlags::ApplicationContextMask | EAutomationTestFlags::EngineFilter)
+IMPLEMENT_SIMPLE_AUTOMATION_TEST(FGuardedIntTest, "System.Core.GuardedInt", EAutomationTestFlags::ApplicationContextMask | EAutomationTestFlags::EngineFilter)
 bool FGuardedIntTest::RunTest(const FString&)
 {
 	bool bResult = true;
diff --git a/Engine/Source/Runtime/Core/Tests/String/FormatStringSanTest.cpp b/Engine/Source/Runtime/Core/Tests/String/FormatStringSanTest.cpp
--- a/Engine/Source/Runtime/Core/Tests/String/FormatStringSanTest.cpp
+++ b/Engine/Source/Runtime/Core/Tests/String/FormatStringSanTest.cpp
@@ -8,7 +8,7 @@
 
 namespace UE
 {
-TEST_CASE_NAMED(FFormatStringValidatorTest, "Core::String::FormatStringSan", "[Core][String][FormatStringSan]")
+TEST_CASE_NAMED(FFormatStringValidatorTest, "System::Core::String::FormatStringSan", "[Core][String][FormatStringSan]")
 {
 	using namespace UE::Core::Private;
 	SECTION("Error Handling")
@@ -101,7 +101,7 @@ TEST_CASE_NAMED(FFormatStringValidatorTest, "Core::String::FormatStringSan", "[C
 	}
 }
 
-TEST_CASE_NAMED(FFormatStringConstStringValidationTest, "Core::String::FormatStringSan::ConstString", "[Core][String][FormatStringSan]")
+TEST_CASE_NAMED(FFormatStringConstStringValidationTest, "System::Core::String::FormatStringSan::ConstString", "[Core][String][FormatStringSan]")
 {
 	using namespace UE::Core::Private::FormatStringSan;
 
