PathViews and Paths: Change FPathViews::Split to handle special case "..". Change FPaths::Split to share implementation with FPathViews::Split.
#rnx

diff --git a/Engine/Source/Runtime/Core/Private/Misc/PathViews.cpp b/Engine/Source/Runtime/Core/Private/Misc/PathViews.cpp
--- a/Engine/Source/Runtime/Core/Private/Misc/PathViews.cpp
+++ b/Engine/Source/Runtime/Core/Private/Misc/PathViews.cpp
@@ -231,8 +231,18 @@ void FPathViews::Split(const FStringView& InPath, FStringView& OutPath, FStringV
 	const TCHAR* DotPos = Algo::FindLast(CleanName, TEXT('.'));
 	const int32 NameLen = DotPos ? UE_PTRDIFF_TO_INT32(DotPos - CleanName.GetData()) : CleanName.Len();
 	OutPath = InPath.LeftChop(CleanName.Len() + 1);
-	OutName = CleanName.Left(NameLen);
-	OutExt = CleanName.RightChop(NameLen + 1);
+	// If there is a ., check for CleanName == ..; that is a special case that is incorrectly handled if we
+	// always interpret the last . as a file extension marker
+	if (DotPos && CleanName == TEXTVIEW(".."))
+	{
+		OutName = CleanName;
+		OutExt.Reset();
+	}
+	else
+	{
+		OutName = CleanName.Left(NameLen);
+		OutExt = CleanName.RightChop(NameLen + 1);
+	}
 }
 
 FString FPathViews::ChangeExtension(const FStringView& InPath, const FStringView& InNewExtension)
diff --git a/Engine/Source/Runtime/Core/Private/Misc/Paths.cpp b/Engine/Source/Runtime/Core/Private/Misc/Paths.cpp
--- a/Engine/Source/Runtime/Core/Private/Misc/Paths.cpp
+++ b/Engine/Source/Runtime/Core/Private/Misc/Paths.cpp
@@ -1720,9 +1720,14 @@ bool FPaths::ValidatePath( const FString& InPath, FText* OutReason )
 
 void FPaths::Split( const FString& InPath, FString& PathPart, FString& FilenamePart, FString& ExtensionPart )
 {
-	PathPart = GetPath(InPath);
-	FilenamePart = GetBaseFilename(InPath);
-	ExtensionPart = GetExtension(InPath);
+	FStringView PathView;
+	FStringView FilenameView;
+	FStringView ExtensionView;
+
+	FPathViews::Split(InPath, PathView, FilenameView, ExtensionView);
+	PathPart = PathView;
+	FilenamePart = FilenameView;
+	ExtensionPart = ExtensionView;
 }
 
 const FString& FPaths::GetRelativePathToRoot()
diff --git a/Engine/Source/Runtime/Core/Tests/Misc/PathViewsTest.cpp b/Engine/Source/Runtime/Core/Tests/Misc/PathViewsTest.cpp
--- a/Engine/Source/Runtime/Core/Tests/Misc/PathViewsTest.cpp
+++ b/Engine/Source/Runtime/Core/Tests/Misc/PathViewsTest.cpp
@@ -497,6 +497,9 @@ TEST_CASE_NAMED(FPathViewsSplitTest, "System::Core::Misc::PathViews::Split", "[A
 	RunSplitTest(TEXT(".tar.gz"), TEXT(""), TEXT(".tar"), TEXT("gz"));
 	RunSplitTest(TEXT(".tar.gz/"), TEXT(".tar.gz"), TEXT(""), TEXT(""));
 	RunSplitTest(TEXT(".tar.gz\\"), TEXT(".tar.gz"), TEXT(""), TEXT(""));
+	// TEXT(".") is ambiguous; we currently treat it as an extension separator but we don't guarantee that in our contract
+	//RunSplitTest(TEXT("."), TEXT(""), TEXT(""), TEXT(""));
+	RunSplitTest(TEXT(".."), TEXT(""), TEXT(".."), TEXT(""));
 	RunSplitTest(TEXT("File"), TEXT(""), TEXT("File"), TEXT(""));
 	RunSplitTest(TEXT("File.txt"), TEXT(""), TEXT("File"), TEXT("txt"));
 	RunSplitTest(TEXT("File.tar.gz"), TEXT(""), TEXT("File.tar"), TEXT("gz"));
diff --git a/Engine/Source/Runtime/Core/Tests/Misc/PathsTest.cpp b/Engine/Source/Runtime/Core/Tests/Misc/PathsTest.cpp
--- a/Engine/Source/Runtime/Core/Tests/Misc/PathsTest.cpp
+++ b/Engine/Source/Runtime/Core/Tests/Misc/PathsTest.cpp
@@ -132,6 +132,46 @@ TEST_CASE_NAMED(FPathTests, "System::Core::Misc::Paths", "[ApplicationContextMas
 			CHECK_EQUALS(TEXT("ConvertRelativePathToFull"), FStringView(Actual), Pair.Expected);
 		}
 	}
+
+	// Split
+	auto RunSplitTest = [](const TCHAR* InPath, const TCHAR* InExpectedPath, const TCHAR* InExpectedName, const TCHAR* InExpectedExt)
+		{
+			FString SplitPath, SplitName, SplitExt;
+			FPaths::Split(InPath, SplitPath, SplitName, SplitExt);
+			if (SplitPath != InExpectedPath || SplitName != InExpectedName || SplitExt != InExpectedExt)
+			{
+				FAIL_CHECK(FString::Printf(TEXT("Failed to split path '%s' (got ('%s', '%s', '%s'), expected ('%s', '%s', '%s'))."), InPath,
+					*SplitPath, *SplitName, *SplitExt, InExpectedPath, InExpectedName, InExpectedExt));
+			}
+		};
+
+	RunSplitTest(TEXT(""), TEXT(""), TEXT(""), TEXT(""));
+	RunSplitTest(TEXT(".txt"), TEXT(""), TEXT(""), TEXT("txt"));
+	RunSplitTest(TEXT(".tar.gz"), TEXT(""), TEXT(".tar"), TEXT("gz"));
+	RunSplitTest(TEXT(".tar.gz/"), TEXT(".tar.gz"), TEXT(""), TEXT(""));
+	RunSplitTest(TEXT(".tar.gz\\"), TEXT(".tar.gz"), TEXT(""), TEXT(""));
+	// TEXT(".") is ambiguous; we currently treat it as an extension separator but we don't guarantee that in our contract
+	//RunSplitTest(TEXT("."), TEXT(""), TEXT(""), TEXT(""));
+	RunSplitTest(TEXT(".."), TEXT(""), TEXT(".."), TEXT(""));
+	RunSplitTest(TEXT("File"), TEXT(""), TEXT("File"), TEXT(""));
+	RunSplitTest(TEXT("File.txt"), TEXT(""), TEXT("File"), TEXT("txt"));
+	RunSplitTest(TEXT("File.tar.gz"), TEXT(""), TEXT("File.tar"), TEXT("gz"));
+	RunSplitTest(TEXT("File.tar.gz/"), TEXT("File.tar.gz"), TEXT(""), TEXT(""));
+	RunSplitTest(TEXT("File.tar.gz\\"), TEXT("File.tar.gz"), TEXT(""), TEXT(""));
+	RunSplitTest(TEXT("C:/Folder/"), TEXT("C:/Folder"), TEXT(""), TEXT(""));
+	RunSplitTest(TEXT("C:/Folder/File"), TEXT("C:/Folder"), TEXT("File"), TEXT(""));
+	RunSplitTest(TEXT("C:/Folder/File.txt"), TEXT("C:/Folder"), TEXT("File"), TEXT("txt"));
+	RunSplitTest(TEXT("C:/Folder/File.tar.gz"), TEXT("C:/Folder"), TEXT("File.tar"), TEXT("gz"));
+	RunSplitTest(TEXT("C:/Folder/First.Last/File"), TEXT("C:/Folder/First.Last"), TEXT("File"), TEXT(""));
+	RunSplitTest(TEXT("C:/Folder/First.Last/File.txt"), TEXT("C:/Folder/First.Last"), TEXT("File"), TEXT("txt"));
+	RunSplitTest(TEXT("C:/Folder/First.Last/File.tar.gz"), TEXT("C:/Folder/First.Last"), TEXT("File.tar"), TEXT("gz"));
+	RunSplitTest(TEXT("C:\\Folder\\"), TEXT("C:\\Folder"), TEXT(""), TEXT(""));
+	RunSplitTest(TEXT("C:\\Folder\\File"), TEXT("C:\\Folder"), TEXT("File"), TEXT(""));
+	RunSplitTest(TEXT("C:\\Folder\\First.Last\\"), TEXT("C:\\Folder\\First.Last"), TEXT(""), TEXT(""));
+	RunSplitTest(TEXT("C:\\Folder\\First.Last\\File"), TEXT("C:\\Folder\\First.Last"), TEXT("File"), TEXT(""));
+	RunSplitTest(TEXT("C:\\Folder\\First.Last\\File.txt"), TEXT("C:\\Folder\\First.Last"), TEXT("File"), TEXT("txt"));
+	RunSplitTest(TEXT("C:\\Folder\\First.Last\\File.tar.gz"), TEXT("C:\\Folder\\First.Last"), TEXT("File.tar"), TEXT("gz"));
+
 }
 
 #endif //WITH_TESTS
