Add cpp types for DoubleFloat vectors/matrices.

diff --git a/Engine/Source/Runtime/Core/Private/Math/DoubleFloat.cpp b/Engine/Source/Runtime/Core/Private/Math/DoubleFloat.cpp
new file mode 100644
--- /dev/null
+++ b/Engine/Source/Runtime/Core/Private/Math/DoubleFloat.cpp
@@ -0,0 +1,69 @@
+// Copyright Epic Games, Inc. All Rights Reserved.
+
+#include "Math/DoubleFloat.h"
+
+#define UE_DF_MIN_PRECISION (2e-4)
+// Max value of a float before it's precision is lower than UE_DF_MIN_PRECISION
+// (there may be 1 more implicit bit available in the significant, but this works as a safe upper bound)
+#define UE_DF_FLOAT_MAX_VALUE (2e23 * UE_DF_MIN_PRECISION - 1)
+
+FMatrix CheckMatrixPrecision(const FMatrix& Matrix)
+{
+	const double OriginMax = UE_DF_FLOAT_MAX_VALUE;
+
+	const FVector Origin = Matrix.GetOrigin();
+	const double OriginX = FMath::Abs(Origin.X);
+	const double OriginY = FMath::Abs(Origin.Y);
+	const double OriginZ = FMath::Abs(Origin.Z);
+	if (OriginX > OriginMax || OriginY > OriginMax || OriginZ > OriginMax)
+	{
+		ensure(false);
+	}
+	return Matrix;
+}
+
+FMatrix44f FDFMatrix::SafeCastMatrix(const FMatrix& Matrix)
+{
+	return FMatrix44f(CheckMatrixPrecision(Matrix));
+}
+
+FMatrix FDFMatrix::MakeToRelativeWorldMatrixDouble(const FVector Origin, const FMatrix& ToWorld)
+{
+	return CheckMatrixPrecision(ToWorld * FTranslationMatrix(-Origin));
+}
+
+FDFMatrix FDFMatrix::MakeToRelativeWorldMatrix(const FVector3f Origin, const FMatrix& ToWorld)
+{
+	return FDFMatrix(FMatrix44f(MakeToRelativeWorldMatrixDouble(FVector(Origin), ToWorld)), Origin);
+}
+
+FMatrix FDFMatrix::MakeClampedToRelativeWorldMatrixDouble(const FVector Origin, const FMatrix& ToWorld)
+{
+	const double OriginMax = UE_DF_FLOAT_MAX_VALUE;
+
+	// Clamp the relative matrix, avoid allowing the relative translation to get too far away from the origin
+	const FVector RelativeOrigin = ToWorld.GetOrigin() - Origin;
+	FVector ClampedRelativeOrigin = RelativeOrigin;
+	ClampedRelativeOrigin.X = FMath::Clamp(ClampedRelativeOrigin.X, -OriginMax, OriginMax);
+	ClampedRelativeOrigin.Y = FMath::Clamp(ClampedRelativeOrigin.Y, -OriginMax, OriginMax);
+	ClampedRelativeOrigin.Z = FMath::Clamp(ClampedRelativeOrigin.Z, -OriginMax, OriginMax);
+
+	FMatrix ClampedToRelativeWorld(ToWorld);
+	ClampedToRelativeWorld.SetOrigin(ClampedRelativeOrigin);
+	return ClampedToRelativeWorld;
+}
+
+FDFMatrix FDFMatrix::MakeClampedToRelativeWorldMatrix(const FVector3f Origin, const FMatrix& ToWorld)
+{
+	return FDFMatrix(FMatrix44f(MakeClampedToRelativeWorldMatrixDouble(FVector(Origin), ToWorld)), Origin);
+}
+
+FDFInverseMatrix FDFInverseMatrix::MakeFromRelativeWorldMatrix(const FVector3f Origin, const FMatrix& FromWorld)
+{
+	return FDFInverseMatrix(FMatrix44f(MakeFromRelativeWorldMatrixDouble(FVector(Origin), FromWorld)), Origin);
+}
+
+FMatrix FDFInverseMatrix::MakeFromRelativeWorldMatrixDouble(const FVector Origin, const FMatrix& FromWorld)
+{
+	return CheckMatrixPrecision(FTranslationMatrix(Origin) * FromWorld);
+}
diff --git a/Engine/Source/Runtime/Core/Public/Math/DoubleFloat.h b/Engine/Source/Runtime/Core/Public/Math/DoubleFloat.h
new file mode 100644
--- /dev/null
+++ b/Engine/Source/Runtime/Core/Public/Math/DoubleFloat.h
@@ -0,0 +1,213 @@
+// Copyright Epic Games, Inc. All Rights Reserved.
+
+#pragma once
+
+#include "CoreMinimal.h"
+#include "Math/MathFwd.h"
+#include "Math/Matrix.h"
+#include "Math/UnrealMathUtility.h"
+#include "Math/Vector.h"
+#include "Math/Vector4.h"
+
+// A high-precision floating point type, consisting of two 32-bit floats (High & Low).
+// Usable where doubles (fp64) are not an option.
+// See also the DoubleFloat.ush shader header.
+struct FDFScalar
+{
+	double GetDouble() const
+	{
+		return static_cast<double>(High) + static_cast<double>(Low);
+	}
+
+	FDFScalar(double In)
+	{
+		High = static_cast<float>(In);
+		Low = static_cast<float>(In - High);
+	}
+
+	FDFScalar(float High, float Low)
+		: High(High), Low(Low)
+	{ }
+
+	FDFScalar() = default;
+
+	float High = 0;
+	float Low = 0;
+};
+
+struct FDFVector2
+{
+	FVector2d GetVector2d() const
+	{
+		return FVector2d{ X().GetDouble(), Y().GetDouble() };
+	}
+
+	template<typename TInputScalar = double>
+	FDFVector2(const UE::Math::TVector2<TInputScalar>& In)
+	{
+		FDFScalar X(In.X);
+		FDFScalar Y(In.Y);
+		High = { X.High, Y.High };
+		Low = { X.Low, Y.Low };
+	}
+
+	FDFVector2(FVector2f High, FVector2f Low)
+		: High(High), Low(Low)
+	{}
+
+	FDFVector2() = default;
+
+	FDFScalar X() const
+	{
+		return { High.X, Low.X };
+	}
+
+	FDFScalar Y() const
+	{
+		return { High.Y, Low.Y };
+	}
+
+	FVector2f High{};
+	FVector2f Low{};
+};
+
+struct FDFVector3
+{
+	FVector3d GetVector3d() const
+	{
+		return FVector3d{ X().GetDouble(), Y().GetDouble(), Z().GetDouble() };
+	}
+
+	template<typename TInputScalar = double>
+	FDFVector3(const UE::Math::TVector<TInputScalar>& In)
+	{
+		FDFScalar X(In.X);
+		FDFScalar Y(In.Y);
+		FDFScalar Z(In.Z);
+		High = { X.High, Y.High, Z.High };
+		Low = { X.Low, Y.Low, Z.Low };
+	}
+
+	FDFVector3(FVector3f High, FVector3f Low)
+		: High(High), Low(Low)
+	{}
+
+	FDFVector3() = default;
+
+	FDFScalar X() const
+	{
+		return { High.X, Low.X };
+	}
+
+	FDFScalar Y() const
+	{
+		return { High.Y, Low.Y };
+	}
+
+	FDFScalar Z() const
+	{
+		return { High.Z, Low.Z };
+	}
+
+	FVector3f High{};
+	FVector3f Low{};
+};
+
+struct FDFVector4
+{
+	FVector4d GetVector4d() const
+	{
+		return FVector4d{ X().GetDouble(), Y().GetDouble(), Z().GetDouble(), W().GetDouble() };
+	}
+
+	template<typename TInputScalar = double>
+	explicit FDFVector4(const UE::Math::TVector4<TInputScalar>& In)
+	{
+		FDFScalar X(In.X);
+		FDFScalar Y(In.Y);
+		FDFScalar Z(In.Z);
+		FDFScalar W(In.W);
+		High = FVector4f{ X.High, Y.High, Z.High, W.High };
+		Low = FVector4f{ X.Low, Y.Low, Z.Low, W.Low };
+	}
+
+	template<typename TInputScalar = double>
+	explicit FDFVector4(const UE::Math::TVector<TInputScalar>& In, TInputScalar InW)
+	{
+		FDFScalar X(In.X);
+		FDFScalar Y(In.Y);
+		FDFScalar Z(In.Z);
+		FDFScalar W(InW);
+		High = FVector4f{ X.High, Y.High, Z.High, W.High };
+		Low = FVector4f{ X.Low, Y.Low, Z.Low, W.Low };
+	}
+
+	FDFVector4(FVector4f High, FVector4f Low)
+		: High(High), Low(Low)
+	{}
+
+	FDFVector4() = default;
+
+	FDFScalar X() const
+	{
+		return { High.X, Low.X };
+	}
+
+	FDFScalar Y() const
+	{
+		return { High.Y, Low.Y };
+	}
+
+	FDFScalar Z() const
+	{
+		return { High.Z, Low.Z };
+	}
+
+	FDFScalar W() const
+	{
+		return { High.W, Low.W };
+	}
+
+	FVector4f High{};
+	FVector4f Low{};
+};
+
+// Transforms *to* absolute world space
+struct FDFMatrix
+{
+	FDFMatrix(FMatrix44f M, FVector3f PostTranslation)
+		: M(M), PostTranslation(PostTranslation)
+	{}
+
+	FDFMatrix() = default;
+
+	FMatrix44f M{};
+	FVector3f PostTranslation{}; // Added to result, *after* multiplying 'M'
+
+	// Check if the origin of the matrix is small enough to ensure the precision >= UE_DF_MIN_PRECISION
+	CORE_API static FMatrix44f SafeCastMatrix(const FMatrix& Matrix);
+
+	// Apply post-translation to matrix
+	CORE_API static FDFMatrix MakeToRelativeWorldMatrix(const FVector3f Origin, const FMatrix& ToWorld);
+	CORE_API static FMatrix   MakeToRelativeWorldMatrixDouble(const FVector Origin, const FMatrix& ToWorld);
+
+	CORE_API static FDFMatrix MakeClampedToRelativeWorldMatrix(const FVector3f Origin, const FMatrix& ToWorld);
+	CORE_API static FMatrix   MakeClampedToRelativeWorldMatrixDouble(const FVector Origin, const FMatrix& ToWorld);
+};
+
+// Transforms *from* absolute world space
+struct FDFInverseMatrix
+{
+	FDFInverseMatrix(FMatrix44f M, FVector3f PreTranslation)
+		: M(M), PreTranslation(PreTranslation)
+	{}
+
+	FDFInverseMatrix() = default;
+
+	FMatrix44f M{};
+	FVector3f PreTranslation{}; // Subtracted from input position *before* multiplying 'M'
+
+	// Apply pre-translation to matrix
+	CORE_API static FDFInverseMatrix MakeFromRelativeWorldMatrix(const FVector3f Origin, const FMatrix& FromWorld);
+	CORE_API static FMatrix          MakeFromRelativeWorldMatrixDouble(const FVector Origin, const FMatrix& FromWorld);
+};
diff --git a/Engine/Source/Runtime/Engine/Private/SceneRelativeViewMatrices.cpp b/Engine/Source/Runtime/Engine/Private/SceneRelativeViewMatrices.cpp
--- a/Engine/Source/Runtime/Engine/Private/SceneRelativeViewMatrices.cpp
+++ b/Engine/Source/Runtime/Engine/Private/SceneRelativeViewMatrices.cpp
@@ -3,6 +3,7 @@
 #include "SceneRelativeViewMatrices.h"
 #include "Misc/LargeWorldRenderPosition.h"
 #include "SceneView.h"
+#include "Math/DoubleFloat.h"
 
 FRelativeViewMatrices FRelativeViewMatrices::Create(const FViewMatrices& Matrices, const FViewMatrices& PrevMatrices)
 {
@@ -42,3 +43,36 @@ FRelativeViewMatrices FRelativeViewMatrices::Create(const FInitializer& Initiali
 	Result.PrevClipToRelativeWorld = Result.PrevClipToView * Result.PrevViewToRelativeWorld;
 	return Result;
 }
+
+
+FDFRelativeViewMatrices FDFRelativeViewMatrices::Create(const FViewMatrices& Matrices, const FViewMatrices& PrevMatrices)
+{
+	FInitializer Initializer;
+	Initializer.ViewOrigin = Matrices.GetViewOrigin();
+	Initializer.ViewToWorld = Matrices.GetInvViewMatrix();
+	Initializer.WorldToView = Matrices.GetViewMatrix();
+	Initializer.ViewToClip = Matrices.GetProjectionMatrix();
+	Initializer.ClipToView = Matrices.GetInvProjectionMatrix();
+	Initializer.PrevViewToWorld = PrevMatrices.GetInvViewMatrix();
+	Initializer.PrevClipToView = PrevMatrices.GetInvProjectionMatrix();
+	return Create(Initializer);
+}
+
+FDFRelativeViewMatrices FDFRelativeViewMatrices::Create(const FInitializer& Initializer)
+{
+	FDFVector3 ViewOrigin(Initializer.ViewOrigin);
+
+	FDFRelativeViewMatrices Result;
+	Result.PositionHigh = ViewOrigin.High;
+	Result.RelativeWorldToView = FDFInverseMatrix::MakeFromRelativeWorldMatrix(ViewOrigin.High, Initializer.WorldToView).M;
+	Result.ViewToRelativeWorld = FDFMatrix::MakeToRelativeWorldMatrix(ViewOrigin.High, Initializer.ViewToWorld).M;
+	Result.ViewToClip = FMatrix44f(Initializer.ViewToClip);
+	Result.ClipToView = FMatrix44f(Initializer.ClipToView);
+	Result.RelativeWorldToClip = Result.RelativeWorldToView * Result.ViewToClip;
+	Result.ClipToRelativeWorld = Result.ClipToView * Result.ViewToRelativeWorld;
+
+	Result.PrevViewToRelativeWorld = FDFMatrix::MakeClampedToRelativeWorldMatrix(ViewOrigin.High, Initializer.PrevViewToWorld).M;
+	Result.PrevClipToView = FMatrix44f(Initializer.PrevClipToView);
+	Result.PrevClipToRelativeWorld = Result.PrevClipToView * Result.PrevViewToRelativeWorld;
+	return Result;
+}
diff --git a/Engine/Source/Runtime/Engine/Public/SceneRelativeViewMatrices.h b/Engine/Source/Runtime/Engine/Public/SceneRelativeViewMatrices.h
--- a/Engine/Source/Runtime/Engine/Public/SceneRelativeViewMatrices.h
+++ b/Engine/Source/Runtime/Engine/Public/SceneRelativeViewMatrices.h
@@ -39,3 +39,37 @@ struct FRelativeViewMatrices
 	FMatrix44f PrevClipToView;
 	FVector3f TilePosition;
 };
+
+struct FDFRelativeViewMatrices
+{
+	struct FInitializer
+	{
+		FVector ViewOrigin;
+		FMatrix ViewToWorld;
+		FMatrix WorldToView;
+		FMatrix ViewToClip;
+		FMatrix ClipToView;
+		FMatrix PrevViewToWorld;
+		FMatrix PrevClipToView;
+	};
+
+	ENGINE_API static FDFRelativeViewMatrices Create(const FViewMatrices& Matrices, const FViewMatrices& PrevMatrices);
+	ENGINE_API static FDFRelativeViewMatrices Create(const FInitializer& Initializer);
+
+	/** Use the same matrices for current and previous */
+	inline static FDFRelativeViewMatrices Create(const FViewMatrices& Matrices)
+	{
+		return Create(Matrices, Matrices);
+	}
+
+	FMatrix44f RelativeWorldToView;
+	FMatrix44f RelativeWorldToClip;
+	FMatrix44f ViewToRelativeWorld;
+	FMatrix44f ClipToRelativeWorld;
+	FMatrix44f ViewToClip;
+	FMatrix44f ClipToView;
+	FMatrix44f PrevViewToRelativeWorld;
+	FMatrix44f PrevClipToRelativeWorld;
+	FMatrix44f PrevClipToView;
+	FVector3f PositionHigh;
+};
\ No newline at end of file
