[TaskGraph]
  - Fix new frontend to use ConcurrentLinearAllocator for normal tasks like the old frontend.
  - Store subsequents and prerequisites in a more cache-friendly way to significantly improve performance.

diff --git a/Engine/Source/Runtime/Core/Private/Async/TaskGraph.cpp b/Engine/Source/Runtime/Core/Private/Async/TaskGraph.cpp
--- a/Engine/Source/Runtime/Core/Private/Async/TaskGraph.cpp
+++ b/Engine/Source/Runtime/Core/Private/Async/TaskGraph.cpp
@@ -2404,8 +2404,6 @@ FGraphEventImplAllocator& GetGraphEventImplAllocator()
 	return Singleton;
 }
 
-FGraphTaskAllocator SmallTaskAllocator;
-
 #else
 
 // Statics and some implementations from FBaseGraphTask and FGraphEvent
diff --git a/Engine/Source/Runtime/Core/Private/Tasks/TaskPrivate.cpp b/Engine/Source/Runtime/Core/Private/Tasks/TaskPrivate.cpp
--- a/Engine/Source/Runtime/Core/Private/Tasks/TaskPrivate.cpp
+++ b/Engine/Source/Runtime/Core/Private/Tasks/TaskPrivate.cpp
@@ -111,7 +111,7 @@ namespace UE::Tasks
 
 				// prerequisites are "consumed" here even if their retraction fails. this means that once prerequisite retraction failed, it won't be performed again. 
 				// this can be potentially improved by using a different container for prerequisites
-				while (FTaskBase* Prerequisite = Prerequisites.Pop())
+				for (FTaskBase* Prerequisite : Prerequisites.PopAll())
 				{
 					// ignore if retraction failed, as this thread still can try to help with other prerequisites instead of being blocked in waiting
 					Prerequisite->TryRetractAndExecute(Timeout, RecursionDepth);
@@ -160,7 +160,7 @@ namespace UE::Tasks
 				bool bSucceeded = true;
 				// prerequisites are "consumed" here even if their retraction fails. this means that once prerequisite retraction failed, it won't be performed again. 
 				// this can be potentially improved by using a different container for prerequisites
-				while (FTaskBase* Prerequisite = Prerequisites.Pop())
+				for (FTaskBase* Prerequisite : Prerequisites.PopAll())
 				{
 					if (!Prerequisite->TryRetractAndExecute(Timeout, RecursionDepth))
 					{
@@ -200,8 +200,8 @@ namespace UE::Tasks
 				return;
 			}
 
+			TRACE_CPUPROFILER_EVENT_SCOPE(FTaskBase::WaitWithNamedThreadsSupport);
 			TaskTrace::FWaitingScope WaitingScope(GetTraceId());
-			TRACE_CPUPROFILER_EVENT_SCOPE(Tasks::Wait);
 
 			if (!TryWaitOnNamedThread(*this))
 			{
diff --git a/Engine/Source/Runtime/Core/Public/Async/TaskGraphInterfaces.h b/Engine/Source/Runtime/Core/Public/Async/TaskGraphInterfaces.h
--- a/Engine/Source/Runtime/Core/Public/Async/TaskGraphInterfaces.h
+++ b/Engine/Source/Runtime/Core/Public/Async/TaskGraphInterfaces.h
@@ -463,18 +463,14 @@ class FBaseGraphTask : public UE::Tasks::Private::FTaskBase
 {
 public:
 	explicit FBaseGraphTask(const FGraphEventArray* InPrerequisites)
-		: FTaskBase(/*InitRefCount=*/ 1)
-	{		
+		: FTaskBase(/*InitRefCount=*/ 1, false /* bUnlockPrerequisites */)
+	{
 		if (InPrerequisites != nullptr)
 		{
-			for (const FGraphEventRef& Prereq : *InPrerequisites)
-			{
-				if (Prereq)
-				{
-					AddPrerequisites(*Prereq);
-				}
-			}
+			AddPrerequisites(*InPrerequisites, false /* bLockPrerequisite */);
 		}
+
+		UnlockPrerequisites();
 	}
 
 	void Init(const TCHAR* InDebugName, UE::Tasks::ETaskPriority InPriority, UE::Tasks::EExtendedTaskPriority InExtendedPriority, UE::Tasks::ETaskFlags InTaskFlags = UE::Tasks::ETaskFlags::None)
@@ -561,13 +557,9 @@ public:
 	}
 };
 
-static constexpr int32 SmallTaskSize = 256;
-using FGraphTaskAllocator = TLockFreeFixedSizeAllocator_TLSCache<SmallTaskSize, PLATFORM_CACHE_LINE_SIZE>;
-CORE_API extern FGraphTaskAllocator SmallTaskAllocator;
-
 // the new task implementation integrated into the old task API
 template<typename TTask>
-class TGraphTask : public FBaseGraphTask
+class TGraphTask final : public TConcurrentLinearObject<TGraphTask<TTask>, FTaskGraphBlockAllocationTag>, public FBaseGraphTask
 {
 public:
 	/**
@@ -633,9 +625,6 @@ public:
 		return FConstructor(Prerequisites);
 	}
 
-	static void* operator new(size_t Size);
-	static void operator delete(void* Ptr, size_t Size);
-
 private:
 	explicit TGraphTask(const FGraphEventArray* InPrerequisites)
 		: FBaseGraphTask(InPrerequisites)
@@ -661,18 +650,6 @@ private:
 	TTypeCompatibleBytes<TTask> TaskStorage;
 };
 
-template<typename TTask>
-void* TGraphTask<TTask>::operator new(size_t Size)
-{
-	return Size <= SmallTaskSize ? SmallTaskAllocator.Allocate() : GMalloc->Malloc(sizeof(TGraphTask), PLATFORM_CACHE_LINE_SIZE);
-}
-
-template<typename TTask>
-void TGraphTask<TTask>::operator delete(void* Ptr, size_t Size)
-{
-	Size <= SmallTaskSize ? SmallTaskAllocator.Free(Ptr) : GMalloc->Free(Ptr);
-}
-
 // an adaptation of FBaseGraphTask to be used as a standalone FGraphEvent
 class FGraphEventImpl : public FBaseGraphTask
 {
diff --git a/Engine/Source/Runtime/Core/Public/Tasks/TaskPrivate.h b/Engine/Source/Runtime/Core/Public/Tasks/TaskPrivate.h
--- a/Engine/Source/Runtime/Core/Public/Tasks/TaskPrivate.h
+++ b/Engine/Source/Runtime/Core/Public/Tasks/TaskPrivate.h
@@ -4,8 +4,10 @@
 
 #include "Async/Fundamental/Scheduler.h"
 #include "Async/Fundamental/Task.h"
+#include "Async/Mutex.h"
 #include "Async/TaskGraphFwd.h"
 #include "Async/TaskTrace.h"
+#include "Async/UniqueLock.h"
 #include "Containers/Array.h"
 #include "Containers/LockFreeFixedSizeAllocator.h"
 #include "Containers/LockFreeList.h"
@@ -152,9 +154,13 @@ namespace UE::Tasks
 			////////////////////////////////////////////////////////////////////////////
 
 		protected:
-			explicit FTaskBase(uint32 InitRefCount)
+			explicit FTaskBase(uint32 InitRefCount, bool bUnlockPrerequisites = true)
 				: RefCount(InitRefCount)
 			{
+				if (bUnlockPrerequisites)
+				{
+					Prerequisites.Unlock();
+				}
 			}
 
 			void Init(const TCHAR* InDebugName, ETaskPriority InPriority, EExtendedTaskPriority InExtendedPriority, ETaskFlags Flags)
@@ -254,11 +260,12 @@ namespace UE::Tasks
 				return Prerequisite.IsValid() ? AddPrerequisites(*Prerequisite.Pimpl) : false;
 			}
 
+protected:
 			// The task will be executed only when all prerequisites are completed.
 			// Must not be called concurrently.
 			// @param InPrerequisites - an iterable collection of tasks
 			template<typename PrerequisiteCollectionType, decltype(std::declval<PrerequisiteCollectionType>().begin())* = nullptr>
-			void AddPrerequisites(const PrerequisiteCollectionType& InPrerequisites)
+			void AddPrerequisites(const PrerequisiteCollectionType& InPrerequisites, bool bLockPrerequisite)
 			{
 				TASKGRAPH_VERBOSE_EVENT_SCOPE(FTaskBase::AddPrerequisites_Collection);
 
@@ -269,7 +276,6 @@ namespace UE::Tasks
 				// prerequisite can be added successfully, and release the lock if it wasn't
 				uint32 PrevNumLocks = NumLocks.fetch_add(GetNum(InPrerequisites), std::memory_order_relaxed); // relaxed because the following
 				// `AddSubsequent` provides required sync
-				checkf(PrevNumLocks + GetNum(InPrerequisites) < ExecutionFlag, TEXT("Max number of nested tasks reached: %d"), ExecutionFlag);
 
 				uint32 NumCompletedPrerequisites = 0;
 				for (auto& Prereq : InPrerequisites)
@@ -281,6 +287,10 @@ namespace UE::Tasks
 					{
 						Prerequisite = Prereq;
 					}
+					else if constexpr (std::is_same_v<FPrerequisiteType, FGraphEventRef>)
+					{
+						Prerequisite = Prereq.GetReference();
+					}
 					else if constexpr (std::is_pointer_v<FPrerequisiteType>)
 					{
 						Prerequisite = Prereq->Pimpl;
@@ -299,7 +309,14 @@ namespace UE::Tasks
 					if (Prerequisite->AddSubsequent(*this)) // acq_rel memory order
 					{
 						Prerequisite->AddRef(); // keep it alive until this task's execution
-						Prerequisites.Push(Prerequisite); // release memory order
+						if (bLockPrerequisite)
+						{
+							Prerequisites.Push(Prerequisite); // release memory order
+						}
+						else
+						{
+							Prerequisites.PushNoLock(Prerequisite); // relaxed memory order
+						}
 					}
 					else
 					{
@@ -307,9 +324,20 @@ namespace UE::Tasks
 					}
 				}
 
+				// This check is here to avoid the data dependency on PrevNumLocks.
+				checkf(PrevNumLocks + GetNum(InPrerequisites) < ExecutionFlag, TEXT("Max number of nested tasks reached: %d"), ExecutionFlag);
+
 				// unlock for prerequisites that weren't added
-				NumLocks.fetch_sub(NumCompletedPrerequisites, std::memory_order_relaxed);  // relaxed because the previous 
-				// `AddSubsequent` provides required sync
+				NumLocks.fetch_sub(NumCompletedPrerequisites, std::memory_order_release);
+			}
+public:
+			// The task will be executed only when all prerequisites are completed.
+			// Must not be called concurrently.
+			// @param InPrerequisites - an iterable collection of tasks
+			template<typename PrerequisiteCollectionType, decltype(std::declval<PrerequisiteCollectionType>().begin())* = nullptr>
+			void AddPrerequisites(const PrerequisiteCollectionType& InPrerequisites)
+			{
+				AddPrerequisites(InPrerequisites, true /* bLockPrerequisites */);
 			}
 
 			// the task unlocks all its subsequents on completion.
@@ -500,20 +528,15 @@ namespace UE::Tasks
 					ClearPipe();
 				}
 
-				TArray<FTaskBase*> Subs;
-				Subsequents.PopAllAndClose(Subs); // gets `Subs` in LIFO order
-				// try to maintain FIFO order where it's possible. e.g. if multiple piped tasks (subsequents) depend on the same (this) task, 
-				// preserve the piping order, which is also the order in which subsequents were added as dependencies of this task
-
 				// Push the first subsequent to the local queue so we pick it up directly as our next task.
 				// This saves us the cost of going to the global queue and performing a wake-up.
 				bool bWakeUpWorker = false;
 
-				for (int32 Index = Subs.Num() - 1; Index >= 0; --Index)
+				for (FTaskBase* Subsequent : Subsequents.Close())
 				{
 					// bWakeUpWorker is passed by reference and is automatically set to true if we successfully schedule a task on the local queue.
 					// so all the remaining ones are sent to the global queue.
-					Subs[Index]->TryUnlock(bWakeUpWorker);
+					Subsequent->TryUnlock(bWakeUpWorker);
 				}
 
 				// release nested tasks
@@ -631,7 +654,7 @@ namespace UE::Tasks
 			void ReleasePrerequisites()
 			{
 				TASKGRAPH_VERBOSE_EVENT_SCOPE(FTaskBase::ReleasePrerequisites);
-				while (FTaskBase* Prerequisite = Prerequisites.Pop())
+				for (FTaskBase* Prerequisite : Prerequisites.PopAll())
 				{
 					TASKGRAPH_VERBOSE_EVENT_SCOPE(FTaskBase::ReleasePrerequisite);
 					Prerequisite->Release();
@@ -648,20 +671,47 @@ namespace UE::Tasks
 
 			LowLevelTasks::FTask LowLevelTask;
 
-			// the number of times that the task should be unlocked before it can be scheduled or completed
-			// initial count is 1 for launching the task (it can't be scheduled before it's launched)
-			// reaches 0 the task is scheduled for execution.
-			// NumLocks's the most significant bit (see `ExecutionFlag`) is set on task execution start, and indicates that now 
-			// NumLocks is about how many times the task must be unlocked to be completed
-			static constexpr uint32 NumInitialLocks = 1;
-			std::atomic<uint32> NumLocks{ NumInitialLocks };
+			// the task is completed when its subsequents list is closed and no more can be added
+			template <typename AllocatorType = FDefaultAllocator>
+			class FSubsequents
+			{
+			public:
+				bool PushIfNotClosed(FTaskBase* NewItem)
+				{
+					TASKGRAPH_VERBOSE_EVENT_SCOPE(FSubsequents::PushIfNotClosed);
+					if (bIsClosed.load(std::memory_order_relaxed))
+					{
+						return false;
+					}
+					UE::TUniqueLock Lock(Mutex);
+					if (bIsClosed)
+					{
+						return false;
+					}
+					Subsequents.Emplace(NewItem);
+					return true;
+				}
 
-#if UE_TASK_TRACE_ENABLED
-			std::atomic<TaskTrace::FId> TraceId { TaskTrace::GenerateTaskId() };
-#endif
+				TArray<FTaskBase*, AllocatorType> Close()
+				{
+					TASKGRAPH_VERBOSE_EVENT_SCOPE(FSubsequents::Close);
+					UE::TUniqueLock Lock(Mutex);
+					bIsClosed = true;
+					return MoveTemp(Subsequents);
+				}
+
+				bool IsClosed() const
+				{
+					return bIsClosed;
+				}
+
+			private:
+				TArray<FTaskBase*, AllocatorType> Subsequents;
+				std::atomic<bool>  bIsClosed = false;
+				UE::FMutex Mutex;
+			};
 
-			// the task is completed when its subsequents list is closed
-			TClosableLockFreePointerListUnorderedSingleConsumer<FTaskBase, 0> Subsequents;
+			FSubsequents<TInlineAllocator<1>> Subsequents;
 
 			// stores backlinks to prerequsites, either execution prerequisites or nested tasks (completion prerequisites).
 			// It's populated in three stages:
@@ -669,11 +719,62 @@ namespace UE::Tasks
 			// 2) by piping, when the previous piped task (if any) is added as a prerequisite. can happen concurrently with other threads accessing prerequisites for
 			//		task retraction.
 			// 3) by adding nested tasks. after piping. during task execution.
-			TLockFreePointerListUnordered<FTaskBase, 0> Prerequisites;
+			template <typename AllocatorType = FDefaultAllocator>
+			class FPrerequisites
+			{
+			public:
+				void Push(FTaskBase* Prerequisite)
+				{
+					TASKGRAPH_VERBOSE_EVENT_SCOPE(FPrerequisites::Push);
+					UE::TUniqueLock Lock(Mutex);
+					Prerequisites.Emplace(Prerequisite);
+				}
+
+				void PushNoLock(FTaskBase* Prerequisite)
+				{
+					TASKGRAPH_VERBOSE_EVENT_SCOPE(FPrerequisites::PushNoLock);
+					Prerequisites.Emplace(Prerequisite);
+				}
+
+				TArray<FTaskBase*, AllocatorType> PopAll()
+				{
+					TASKGRAPH_VERBOSE_EVENT_SCOPE(FPrerequisites::PopAll);
+					UE::TUniqueLock Lock(Mutex);
+					return MoveTemp(Prerequisites);
+				}
+
+				void Unlock()
+				{
+					Mutex.Unlock();
+				}
+			private:
+				TArray<FTaskBase*, AllocatorType> Prerequisites;
+				UE::FMutex Mutex { UE::AcquireLock }; // Start locked by default to avoid compare exchange during construction.
+			};
+
+			FPrerequisites<TInlineAllocator<1>> Prerequisites;
 
 			FPipe* Pipe{ nullptr };
 
+#if UE_TASK_TRACE_ENABLED
+			std::atomic<TaskTrace::FId> TraceId{ TaskTrace::GenerateTaskId() };
+#endif
+
+			// the number of times that the task should be unlocked before it can be scheduled or completed
+			// initial count is 1 for launching the task (it can't be scheduled before it's launched)
+			// reaches 0 the task is scheduled for execution.
+			// NumLocks's the most significant bit (see `ExecutionFlag`) is set on task execution start, and indicates that now 
+			// NumLocks is about how many times the task must be unlocked to be completed
+			static constexpr uint32 NumInitialLocks = 1;
+			std::atomic<uint32> NumLocks{ NumInitialLocks };
+
 			std::atomic<uint32> ExecutingThreadId = FThread::InvalidThreadId;
+
+protected:
+			void UnlockPrerequisites()
+			{
+				Prerequisites.Unlock();
+			}
 		};
 
 		// an extension of FTaskBase for tasks that return a result.
@@ -786,7 +887,15 @@ namespace UE::Tasks
 
 			static void* operator new(size_t Size)
 			{
-				return Size <= SmallTaskSize ? SmallTaskAllocator.Allocate() : GMalloc->Malloc(sizeof(TExecutableTask), LargeTaskAlignment);
+				if (Size <= SmallTaskSize)
+				{
+					 return SmallTaskAllocator.Allocate();
+				}
+				else
+				{
+					TASKGRAPH_VERBOSE_EVENT_SCOPE(TExecutableTask::LargeAlloc);
+					return GMalloc->Malloc(sizeof(TExecutableTask), LargeTaskAlignment);
+				}
 			}
 
 			static void operator delete(void* Ptr, size_t Size)
