AssetRegistry: Fix failure to detect changes on disk when the Windows USN Journal is destroyed and recreated; the old NextJournalEntry was invalid but FileJournalReadModified silently ignored it.
#rnx

diff --git a/Engine/Source/Runtime/AssetRegistry/Private/AssetDataGathererDiscoveryCache.cpp b/Engine/Source/Runtime/AssetRegistry/Private/AssetDataGathererDiscoveryCache.cpp
--- a/Engine/Source/Runtime/AssetRegistry/Private/AssetDataGathererDiscoveryCache.cpp
+++ b/Engine/Source/Runtime/AssetRegistry/Private/AssetDataGathererDiscoveryCache.cpp
@@ -25,7 +25,7 @@ FArchive& operator<<(FArchive& Ar, FFileJournalFileHandle& H)
 
 namespace AssetDataGathererConstants
 {
-const FGuid DiscoveryCacheVersion(TEXT("62267A3E1D2144EF9C8734271DD28742"));
+const FGuid DiscoveryCacheVersion(TEXT("4F4C364CC08C47B9BF18278136E1CB6E"));
 }
 
 namespace UE::AssetDataGather::Private
@@ -54,12 +54,13 @@ void FAssetDataDiscoveryCache::LoadAndUpdateCache()
 	}
 
 	FString TestError;
+	FFileJournalId TestJournalId;
 	FFileJournalEntryHandle TestLatestJournalEntry;
 	FString ProjectDir = FPaths::ProjectDir();
 	IPlatformFile& PlatformFile = FPlatformFileManager::Get().GetPlatformFile();
 	FString TestVolumeName = PlatformFile.FileJournalGetVolumeName(ProjectDir);
-	EFileJournalResult TestResult = PlatformFile.FileJournalGetLatestEntry(*TestVolumeName,	TestLatestJournalEntry,
-		&TestError);
+	EFileJournalResult TestResult = PlatformFile.FileJournalGetLatestEntry(*TestVolumeName,	TestJournalId,
+		TestLatestJournalEntry, &TestError);
 	bool bPlatformSupported = TestResult == EFileJournalResult::Success;
 
 	bool bInvalidateEnabled = GPreloadSettings.IsDiscoveryCacheInvalidateEnabled();
@@ -126,7 +127,7 @@ void FAssetDataDiscoveryCache::LoadAndUpdateCache()
 		else
 		{
 			EFileJournalResult Result = PlatformFile.FileJournalReadModified(*VolumeInfo.VolumeName,
-				VolumeInfo.NextJournalEntryToScan, KnownDirectories, ModifiedDirectories,
+				VolumeInfo.JournalId, VolumeInfo.NextJournalEntryToScan, KnownDirectories, ModifiedDirectories,
 				VolumeInfo.NextJournalEntryToScan, &VolumeInfo.LastError);
 			switch (Result)
 			{
@@ -150,6 +151,7 @@ void FAssetDataDiscoveryCache::LoadAndUpdateCache()
 
 		if (!bReadModifiedSucceeded)
 		{
+			VolumeInfo.JournalId = VolumeInfo.JournalIdOnDisk;
 			VolumeInfo.NextJournalEntryToScan = VolumeInfo.NextJournalEntryOnDisk;
 			VolumeInfo.Dirs.Empty();
 		}
@@ -469,16 +471,19 @@ void FCachedVolumeInfo::InitializePlatformData()
 	if (VolumeName.IsEmpty() || VolumeName == GEmptyVolumeName)
 	{
 		bJournalAvailable = false;
+		JournalIdOnDisk = FileJournalIdInvalid;
 		NextJournalEntryOnDisk = FileJournalEntryHandleInvalid;
+		JournalId = FileJournalIdInvalid;
 		NextJournalEntryToScan = FileJournalEntryHandleInvalid;
 	}
 	else
 	{
 		EFileJournalResult Result = FPlatformFileManager::Get().GetPlatformFile().FileJournalGetLatestEntry(
-			*VolumeName, NextJournalEntryOnDisk, &LastError);
+			*VolumeName, JournalIdOnDisk, NextJournalEntryOnDisk, &LastError);
 		bJournalAvailable = Result == EFileJournalResult::Success;
 		if (NextJournalEntryToScan == FileJournalEntryHandleInvalid)
 		{
+			JournalId = JournalIdOnDisk;
 			NextJournalEntryToScan = NextJournalEntryOnDisk;
 		}
 	}
@@ -523,6 +528,7 @@ FArchive& operator<<(FArchive& Ar, FCachedVolumeInfo& Data)
 {
 	Ar << Data.Dirs;
 	Ar << Data.VolumeName;
+	Ar << Data.JournalId;
 	Ar << Data.NextJournalEntryToScan;
 	return Ar;
 }
diff --git a/Engine/Source/Runtime/AssetRegistry/Private/AssetDataGathererDiscoveryCache.h b/Engine/Source/Runtime/AssetRegistry/Private/AssetDataGathererDiscoveryCache.h
--- a/Engine/Source/Runtime/AssetRegistry/Private/AssetDataGathererDiscoveryCache.h
+++ b/Engine/Source/Runtime/AssetRegistry/Private/AssetDataGathererDiscoveryCache.h
@@ -84,12 +84,14 @@ public:
 public:
 	TMap<FString, FCachedDirScanDir> Dirs;
 	FString VolumeName;
+	FFileJournalId JournalId = FileJournalIdInvalid;
 	FFileJournalEntryHandle NextJournalEntryToScan = FileJournalEntryHandleInvalid;
 
 	// Transient
 	/** Directories that we marked for recursive removal from CachedVolumes; we consume this during WriteCacheFile. */
 	TArray<FString> DirsToRemove;
 	FString LastError;
+	FFileJournalId JournalIdOnDisk = FileJournalIdInvalid;
 	FFileJournalEntryHandle NextJournalEntryOnDisk = FileJournalEntryHandleInvalid;
 	bool bJournalAvailable = false;
 };
diff --git a/Engine/Source/Runtime/Core/Private/GenericPlatform/GenericPlatformFile.cpp b/Engine/Source/Runtime/Core/Private/GenericPlatform/GenericPlatformFile.cpp
--- a/Engine/Source/Runtime/Core/Private/GenericPlatform/GenericPlatformFile.cpp
+++ b/Engine/Source/Runtime/Core/Private/GenericPlatform/GenericPlatformFile.cpp
@@ -1022,13 +1022,14 @@ bool IPlatformFile::FileJournalIsAvailable(const TCHAR* VolumeOrPath, ELogVerbos
 }
 
 EFileJournalResult IPlatformFile::FileJournalGetLatestEntry(const TCHAR* VolumeOrPath,
-	FFileJournalEntryHandle& OutEntryHandle, FString* OutError)
+	FFileJournalId& OutJournalId, FFileJournalEntryHandle& OutEntryHandle, FString* OutError)
 {
 	if (OutError)
 	{
 		*OutError = UE::PlatformFileJournal::Private::PlatformNotAvailableMessage;
 	}
-	OutEntryHandle = 0;
+	OutJournalId = FileJournalIdInvalid;
+	OutEntryHandle = FileJournalEntryHandleInvalid;
 	return EFileJournalResult::InvalidPlatform;
 }
 
@@ -1047,10 +1048,11 @@ FFileJournalData IPlatformFile::FileJournalGetFileData(const TCHAR* FilenameOrDi
 }
 
 EFileJournalResult IPlatformFile::FileJournalReadModified(const TCHAR* VolumeName,
-	const FFileJournalEntryHandle& StartingJournalEntry, TMap<FFileJournalFileHandle, FString>& KnownDirectories,
-	TSet<FString>& OutModifiedDirectories, FFileJournalEntryHandle& OutNextJournalEntry, FString* OutError)
+	const FFileJournalId& JournalIdOfStartingEntry, const FFileJournalEntryHandle& StartingJournalEntry,
+	TMap<FFileJournalFileHandle, FString>& KnownDirectories, TSet<FString>& OutModifiedDirectories,
+	FFileJournalEntryHandle& OutNextJournalEntry, FString* OutError)
 {
-	OutNextJournalEntry = StartingJournalEntry;
+	OutNextJournalEntry = FileJournalEntryHandleInvalid;
 	if (OutError)
 	{
 		*OutError = UE::PlatformFileJournal::Private::PlatformNotAvailableMessage;
diff --git a/Engine/Source/Runtime/Core/Private/Windows/WindowsPlatformFileJournal.cpp b/Engine/Source/Runtime/Core/Private/Windows/WindowsPlatformFileJournal.cpp
--- a/Engine/Source/Runtime/Core/Private/Windows/WindowsPlatformFileJournal.cpp
+++ b/Engine/Source/Runtime/Core/Private/Windows/WindowsPlatformFileJournal.cpp
@@ -137,8 +137,9 @@ bool FWindowsPlatformFile::FileJournalIsAvailable(const TCHAR* VolumeOrPath, ELo
 		return true;
 	}
 
+	FFileJournalId JournalId;
 	FFileJournalEntryHandle Handle;
-	EFileJournalResult Result = FileJournalGetLatestEntry(VolumeOrPath, Handle, OutError);
+	EFileJournalResult Result = FileJournalGetLatestEntry(VolumeOrPath, JournalId, Handle, OutError);
 	if (Result == EFileJournalResult::Success)
 	{
 		if (OutErrorLevel)
@@ -159,11 +160,12 @@ bool FWindowsPlatformFile::FileJournalIsAvailable(const TCHAR* VolumeOrPath, ELo
 
 
 EFileJournalResult FWindowsPlatformFile::FileJournalGetLatestEntry(const TCHAR* VolumeName,
-	FFileJournalEntryHandle& OutEntryHandle, FString* OutError)
+	FFileJournalId& OutJournalId, FFileJournalEntryHandle& OutEntryHandle, FString* OutError)
 {
 	using namespace UE::WindowsPlatformFileJournal::Private;
 
-	OutEntryHandle = 0;
+	OutJournalId = FileJournalIdInvalid;
+	OutEntryHandle = FileJournalEntryHandleInvalid;
 
 	HANDLE JournalHandle;
 	TStringBuilder<16> NormalizedVolumeName;
@@ -181,7 +183,9 @@ EFileJournalResult FWindowsPlatformFile::FileJournalGetLatestEntry(const TCHAR*
 	{
 		return Result;
 	}
+	static_assert(sizeof(uint64) >= sizeof(JournalDescriptor.UsnJournalID), "We are storing JournalIDs as uint64");
 	static_assert(sizeof(uint64) >= sizeof(USN), "We are storing USNs as uint64");
+	OutJournalId = static_cast<FFileJournalId>(JournalDescriptor.UsnJournalID);
 	OutEntryHandle = static_cast<FFileJournalEntryHandle>(JournalDescriptor.NextUsn);
 	return EFileJournalResult::Success;
 }
@@ -352,15 +356,16 @@ FFileJournalData FWindowsPlatformFile::FileJournalGetFileData(const TCHAR* Filen
 }
 
 EFileJournalResult FWindowsPlatformFile::FileJournalReadModified(const TCHAR* VolumeName,
-	const FFileJournalEntryHandle& StartingJournalEntry, TMap<FFileJournalFileHandle, FString>& KnownDirectories,
-	TSet<FString>& OutModifiedDirectories, FFileJournalEntryHandle& OutNextJournalEntry, FString* OutError)
+	const FFileJournalId& JournalIdOfStartingEntry, const FFileJournalEntryHandle& StartingJournalEntry,
+	TMap<FFileJournalFileHandle, FString>& KnownDirectories, TSet<FString>& OutModifiedDirectories,
+	FFileJournalEntryHandle& OutNextJournalEntry, FString* OutError)
 {
 	using namespace UE::WindowsPlatformFileJournal::Private;
 
 	if (!FileJournalIsAvailableSystemWide(nullptr, OutError))
 	{
-		return IPlatformFile::FileJournalReadModified(VolumeName, StartingJournalEntry, KnownDirectories,
-			OutModifiedDirectories, OutNextJournalEntry, nullptr);
+		return IPlatformFile::FileJournalReadModified(VolumeName, JournalIdOfStartingEntry, StartingJournalEntry,
+			KnownDirectories, OutModifiedDirectories, OutNextJournalEntry, nullptr);
 	}
 
 	EFileJournalResult Result;
@@ -391,6 +396,35 @@ EFileJournalResult FWindowsPlatformFile::FileJournalReadModified(const TCHAR* Vo
 		return Result;
 	}
 
+	if (static_cast<DWORDLONG>(JournalIdOfStartingEntry) != JournalDescriptor.UsnJournalID)
+	{
+		if (OutError)
+		{
+			*OutError = FString::Printf(
+				TEXT("The JournalId StartingJournalEntry on volume '%s' came from a different journal. This can happen when the Journal has been destroyed and recreated. ")
+				TEXT("StoredJournalId = %" UINT64_FMT ", CurrentJournalId = %" UINT64_FMT "."),
+				VolumeName, static_cast<uint64>(JournalIdOfStartingEntry), static_cast<uint64>(JournalDescriptor.UsnJournalID));
+		}
+		return EFileJournalResult::JournalWrapped;
+	}
+
+	// FSCTL_READ_USN_JOURNAL does not return an error if the StartUsn is greater than the end of the current journal;
+	// it returns no results and reports the input StartUsn as the next USN to read. We want this to be an error
+	// condition similar to wrapping the journal. StartUSn == NextUSN is valid, but StartUSN > NextUSN is not.
+	// Return a wrap code for that case.
+	USN CurrentJournalUSN = static_cast<USN>(OutNextJournalEntry);
+	if (CurrentJournalUSN > JournalDescriptor.NextUsn)
+	{
+		if (OutError)
+		{
+			*OutError = FString::Printf(
+				TEXT("The stored value for StartingJournalEntry on volume '%s' is invalid; it is past the end of the current Journal. This can happen when the Journal has been destroyed and recreated. ")
+				TEXT("StoredValue = %" UINT64_FMT ", CurrentEnd = %" UINT64_FMT "."),
+				VolumeName, static_cast<uint64>(CurrentJournalUSN), static_cast<uint64>(JournalDescriptor.NextUsn));
+		}
+		return EFileJournalResult::JournalWrapped;
+	}
+
 	READ_USN_JOURNAL_DATA_V1 ReadJournalRequest;
 	// The USN at which to begin reading the change journal, or 0 to start from beginning
 	// ReadJournalRequest.StartUsn; // Set below in the loop
@@ -427,7 +461,7 @@ EFileJournalResult FWindowsPlatformFile::FileJournalReadModified(const TCHAR* Vo
 		DWORD dwIoControlCode = bShowFilenames ? FSCTL_READ_USN_JOURNAL : FSCTL_READ_UNPRIVILEGED_USN_JOURNAL;
 		bool bWasFirstCall = bFirstCall;
 		bFirstCall = false;
-		USN CurrentJournalUSN = static_cast<USN>(OutNextJournalEntry);
+		CurrentJournalUSN = static_cast<USN>(OutNextJournalEntry);
 		ReadJournalRequest.StartUsn = CurrentJournalUSN;
 
 		BOOL ReadJournalResult = ::DeviceIoControl(
@@ -858,7 +892,7 @@ bool FWindowsPlatformFile::FileJournalIsAvailable(const TCHAR* VolumeOrPath, ELo
 }
 
 EFileJournalResult  FWindowsPlatformFile::FileJournalGetLatestEntry(const TCHAR* VolumeName,
-	FFileJournalEntryHandle& OutEntryHandle, FString* OutError)
+	FFileJournalId& OutJournalId, FFileJournalEntryHandle& OutEntryHandle, FString* OutError)
 {
 	using namespace UE::WindowsPlatformFileJournal::Private;
 
@@ -869,7 +903,7 @@ EFileJournalResult  FWindowsPlatformFile::FileJournalGetLatestEntry(const TCHAR*
 		OutError = nullptr;
 	}
 #endif
-	return IPlatformFile::FileJournalGetLatestEntry(VolumeName, OutEntryHandle, OutError);
+	return IPlatformFile::FileJournalGetLatestEntry(VolumeName, OutJournalId, OutEntryHandle, OutError);
 }
 
 bool FWindowsPlatformFile::FileJournalIterateDirectory(const TCHAR* Directory,
@@ -884,8 +918,9 @@ FFileJournalData FWindowsPlatformFile::FileJournalGetFileData(const TCHAR* Filen
 }
 
 EFileJournalResult FWindowsPlatformFile::FileJournalReadModified(const TCHAR* VolumeName,
-	const FFileJournalEntryHandle& StartingJournalEntry, TMap<FFileJournalFileHandle, FString>& KnownDirectories,
-	TSet<FString>& OutModifiedDirectories, FFileJournalEntryHandle& OutNextJournalEntry, FString* OutError)
+	const FFileJournalId& JournalIdOfStartingEntry, const FFileJournalEntryHandle& StartingJournalEntry,
+	TMap<FFileJournalFileHandle, FString>& KnownDirectories, TSet<FString>& OutModifiedDirectories,
+	FFileJournalEntryHandle& OutNextJournalEntry, FString* OutError)
 {
 	using namespace UE::WindowsPlatformFileJournal::Private;
 
@@ -896,8 +931,8 @@ EFileJournalResult FWindowsPlatformFile::FileJournalReadModified(const TCHAR* Vo
 		OutError = nullptr;
 	}
 #endif
-	return IPlatformFile::FileJournalReadModified(VolumeName, StartingJournalEntry, KnownDirectories, OutModifiedDirectories,
-		OutNextJournalEntry, OutError);
+	return IPlatformFile::FileJournalReadModified(VolumeName, JournalIdOfStartingEntry, StartingJournalEntry,
+		KnownDirectories, OutModifiedDirectories, OutNextJournalEntry, OutError);
 }
 
 #endif // else !UE_WINDOWS_PLATFORM_FILEJOURNAL_ENABLED
\ No newline at end of file
diff --git a/Engine/Source/Runtime/Core/Private/Windows/WindowsPlatformFilePrivate.h b/Engine/Source/Runtime/Core/Private/Windows/WindowsPlatformFilePrivate.h
--- a/Engine/Source/Runtime/Core/Private/Windows/WindowsPlatformFilePrivate.h
+++ b/Engine/Source/Runtime/Core/Private/Windows/WindowsPlatformFilePrivate.h
@@ -81,14 +81,14 @@ public:
 
 	virtual bool FileJournalIsAvailable(const TCHAR* VolumeOrPath = nullptr, ELogVerbosity::Type* OutErrorLevel = nullptr,
 		FString* OutError = nullptr) override;
-	virtual EFileJournalResult FileJournalGetLatestEntry(const TCHAR* VolumeName, FFileJournalEntryHandle& OutEntryHandle,
-		FString* OutError = nullptr) override;
+	virtual EFileJournalResult FileJournalGetLatestEntry(const TCHAR* VolumeName, FFileJournalId& OutJournalId, 
+		FFileJournalEntryHandle& OutEntryHandle, FString* OutError = nullptr) override;
 	virtual bool FileJournalIterateDirectory(const TCHAR* Directory, FDirectoryJournalVisitorFunc Visitor) override;
 	virtual FFileJournalData FileJournalGetFileData(const TCHAR* FilenameOrDirectory) override;
 	virtual EFileJournalResult FileJournalReadModified(const TCHAR* VolumeName,
-		const FFileJournalEntryHandle& StartingJournalEntry, TMap<FFileJournalFileHandle, FString>& KnownDirectories,
-		TSet<FString>& OutModifiedDirectories, FFileJournalEntryHandle& OutNextJournalEntry,
-		FString* OutError = nullptr) override;
+		const FFileJournalId& JournalIdOfStartingEntry, const FFileJournalEntryHandle& StartingJournalEntry,
+		TMap<FFileJournalFileHandle, FString>& KnownDirectories, TSet<FString>& OutModifiedDirectories,
+		FFileJournalEntryHandle& OutNextJournalEntry, FString* OutError = nullptr) override;
 	virtual FString FileJournalGetVolumeName(FStringView InPath) override;
 };
 
diff --git a/Engine/Source/Runtime/Core/Public/GenericPlatform/GenericPlatformFile.h b/Engine/Source/Runtime/Core/Public/GenericPlatform/GenericPlatformFile.h
--- a/Engine/Source/Runtime/Core/Public/GenericPlatform/GenericPlatformFile.h
+++ b/Engine/Source/Runtime/Core/Public/GenericPlatform/GenericPlatformFile.h
@@ -230,6 +230,12 @@ struct FFileStatData
 };
 
 
+/**
+ * A handle used by the FileJournal API. Platform-specific identifier for which disk journal is being read.
+ */
+typedef uint64 FFileJournalId;
+constexpr FFileJournalId FileJournalIdInvalid = static_cast<FFileJournalId>(MAX_uint64);
+
 /**
  * A handle used by the FileJournal API. Represents an entry for an action on a file in the FileJournal.
  */
@@ -547,7 +553,7 @@ public:
 	 * and optionally a user-displayable explanation for the error code.
 	 */
 	CORE_API virtual EFileJournalResult FileJournalGetLatestEntry(const TCHAR* VolumeName,
-		FFileJournalEntryHandle& OutEntryHandle, FString* OutError = nullptr);
+		FFileJournalId& OutJournalId, FFileJournalEntryHandle& OutEntryHandle, FString* OutError = nullptr);
 
 	/** File and directory visitor function that takes FileJournal data. */
 	typedef TFunctionRef<bool(const TCHAR*, const FFileJournalData&)> FDirectoryJournalVisitorFunc;
@@ -580,9 +586,9 @@ public:
 	/**
 	 * Query the FileJournal to find a list of all directories on the given volume with files that have been added,
 	 * deleted, or modified in the specified time range. The beginning of the time range is specified by
-	 * StartingJournalEntry, which came from FileJournalGetLatestEntry or a previous call to FileJournalReadModified.
-	 * The end of the range is the latest modification on the volume. VolumeName can be the return value from
-	 * FileJournalGetVolumeName, or any path on the desired volume.
+	 * JournalIdOfStartingEntry and StartingJournalEntry, which came from FileJournalGetLatestEntry or a previous
+	 * call to FileJournalReadModified. The end of the range is the latest modification on the volume. VolumeName can
+	 * be the return value from FileJournalGetVolumeName, or any path on the desired volume.
 	 * 
 	 * The caller must provide the mapping from FFileJournalFileHandle to DirectoryName; the FFileJournalFileHandle 
 	 * for each Directory can be found from FileJournalGetFileData or FileJournalIterateDirectory.
@@ -594,9 +600,9 @@ public:
 	 * explanation for the error code. In an error case, partial results may still be written into the output.
 	 */
 	CORE_API virtual EFileJournalResult FileJournalReadModified(const TCHAR* VolumeName,
-		const FFileJournalEntryHandle& StartingJournalEntry, TMap<FFileJournalFileHandle, FString>& KnownDirectories,
-		TSet<FString>& OutModifiedDirectories, FFileJournalEntryHandle& OutNextJournalEntry,
-		FString* OutError = nullptr);
+		const FFileJournalId& JournalIdOfStartingEntry, const FFileJournalEntryHandle& StartingJournalEntry,
+		TMap<FFileJournalFileHandle, FString>& KnownDirectories, TSet<FString>& OutModifiedDirectories,
+		FFileJournalEntryHandle& OutNextJournalEntry, FString* OutError = nullptr);
 
 	/**
 	 * Return the VolumeSpecifier present in the given path. Returns empty string if path does not have a valid
