Added suppor for alignments higher than 16 b to MB2 and fixed quantization of allocations with higher alignment on MB3

diff --git a/Engine/Source/Runtime/Core/Private/HAL/MallocBinned2.cpp b/Engine/Source/Runtime/Core/Private/HAL/MallocBinned2.cpp
--- a/Engine/Source/Runtime/Core/Private/HAL/MallocBinned2.cpp
+++ b/Engine/Source/Runtime/Core/Private/HAL/MallocBinned2.cpp
@@ -876,6 +876,48 @@ bool FMallocBinned2::IsInternallyThreadSafe() const
 	return true;
 }
 
+static bool PromoteToLargerBin(SIZE_T& Size, uint32& Alignment, const FMallocBinned2& Malloc)
+{
+	// try to promote our allocation request to a larger bin with a matching natural alignment
+	// if requested alignment is larger than BINNED2_MINIMUM_ALIGNMENT but smaller than BINNED2_MAXIMUM_ALIGNMENT
+	// so we don't do a page allocation with a lot of memory waste
+	Alignment = FMath::Max<uint32>(Alignment, BINNED2_MINIMUM_ALIGNMENT);
+	const SIZE_T AlignedSize = Align(Size, Alignment);
+	if (UNLIKELY((AlignedSize <= BINNED2_MAX_SMALL_POOL_SIZE) && (Alignment <= BINNED2_MAXIMUM_ALIGNMENT)))
+	{
+		uint32 PoolIndex = Malloc.BoundSizeToPoolIndex(AlignedSize);
+		do
+		{
+			uint32 BlockSize = Malloc.PoolIndexToBlockSize(PoolIndex);
+			if (IsAligned(BlockSize, Alignment))
+			{
+				// we found a matching pool for our alignment and size requirements, so modify the size request to match
+				Size = SIZE_T(BlockSize);
+				Alignment = BINNED2_MINIMUM_ALIGNMENT;
+				return true;
+			}
+
+			PoolIndex++;
+		} while (PoolIndex < BINNED2_SMALL_POOL_COUNT);
+	}
+
+	return false;
+}
+
+void* FMallocBinned2::MallocSelect(SIZE_T Size, uint32 Alignment, bool bUseSmallPool)
+{
+	if (!bUseSmallPool)
+	{
+		bUseSmallPool = PromoteToLargerBin(Size, Alignment, *this);
+	}
+
+	if (bUseSmallPool)
+	{
+		return MallocExternalSmall(Size, Alignment);
+	}
+	return MallocExternalLarge(Size, Alignment);
+}
+
 void* FMallocBinned2::MallocExternalSmall(SIZE_T Size, uint32 Alignment)
 {
 	uint32 PoolIndex = BoundSizeToPoolIndex(Size);
@@ -1010,7 +1052,7 @@ void* FMallocBinned2::ReallocExternal(void* Ptr, SIZE_T NewSize, uint32 Alignmen
 #else
 		if( 
 #endif
-			((NewSize <= BlockSize) & (Alignment <= BINNED2_MINIMUM_ALIGNMENT)) && // one branch, not two
+			((NewSize <= BlockSize) & (IsAligned(BlockSize, Alignment))) && // one branch, not two
 			(PoolIndex == 0 || NewSize > PoolIndexToBlockSize(PoolIndex - 1)))
 		{
 			return Ptr;
@@ -1039,6 +1081,10 @@ void* FMallocBinned2::ReallocExternal(void* Ptr, SIZE_T NewSize, uint32 Alignmen
 	SIZE_T PoolOSRequestedBytes = Pool->GetOSRequestedBytes();
 	checkf(PoolOSRequestedBytes <= PoolOsBytes, TEXT("FMallocBinned2::ReallocExternal %d %d"), int32(PoolOSRequestedBytes), int32(PoolOsBytes));
 	bool bUseSmallMalloc = UseSmallAlloc(NewSize, Alignment);
+	if (!bUseSmallMalloc)
+	{
+		bUseSmallMalloc = PromoteToLargerBin(NewSize, Alignment, *this);
+	}
 	if (NewSize > PoolOsBytes || // can't fit in the old block
 		bUseSmallMalloc || // can switch to the small block allocator
 		Align(NewSize, OsAllocationGranularity) < PoolOsBytes) // we can get some pages back
diff --git a/Engine/Source/Runtime/Core/Private/HAL/MallocBinned3.cpp b/Engine/Source/Runtime/Core/Private/HAL/MallocBinned3.cpp
--- a/Engine/Source/Runtime/Core/Private/HAL/MallocBinned3.cpp
+++ b/Engine/Source/Runtime/Core/Private/HAL/MallocBinned3.cpp
@@ -989,6 +989,7 @@ void* FMallocBinned3::MallocExternal(SIZE_T Size, uint32 Alignment)
 					// we found a matching pool for our alignment and size requirements, so modify the size request to match
 					Size = SIZE_T(BlockSize);
 					UsePools = true;
+					Alignment = BINNED3_MINIMUM_ALIGNMENT;
 					break;
 				}
 
@@ -1155,7 +1156,7 @@ void* FMallocBinned3::ReallocExternal(void* Ptr, SIZE_T NewSize, uint32 Alignmen
 		check(Ptr); // null is an OS allocation because it will not fall in our VM block
 		uint32 BlockSize = PoolIndexToBlockSize(PoolIndex);
 		if (
-			((NewSize <= BlockSize) & (Alignment <= BINNED3_MINIMUM_ALIGNMENT)) && // one branch, not two
+			((NewSize <= BlockSize) & (IsAligned(BlockSize, Alignment))) && // one branch, not two
 			(PoolIndex == 0 || NewSize > PoolIndexToBlockSize(PoolIndex - 1)))
 		{
 #if BINNED3_ALLOCATOR_STATS
diff --git a/Engine/Source/Runtime/Core/Public/HAL/MallocBinned2.h b/Engine/Source/Runtime/Core/Public/HAL/MallocBinned2.h
--- a/Engine/Source/Runtime/Core/Public/HAL/MallocBinned2.h
+++ b/Engine/Source/Runtime/Core/Public/HAL/MallocBinned2.h
@@ -30,6 +30,7 @@ struct FGenericMemoryStats;
 #define BINNED2_LARGE_ALLOC					65536		// Alignment of OS-allocated pointer - pool-allocated pointers will have a non-aligned pointer
 #define BINNED2_MINIMUM_ALIGNMENT_SHIFT		4			// Alignment of blocks, expressed as a shift
 #define BINNED2_MINIMUM_ALIGNMENT			16			// Alignment of blocks
+#define BINNED2_MAXIMUM_ALIGNMENT			128
 #define BINNED2_MAX_SMALL_POOL_SIZE			(32768-16)	// Maximum block size in GMallocBinned2SmallBlockSizes
 #define BINNED2_SMALL_POOL_COUNT			45
 
@@ -501,45 +502,29 @@ public:
 		return MallocInline(Size, Alignment);
 #endif
 	}
-	FORCEINLINE void* MallocInline(SIZE_T Size, uint32 Alignment )
+	FORCEINLINE void* MallocInline(SIZE_T Size, uint32 Alignment)
 	{
-
-#if UE_USE_VERYLARGEPAGEALLOCATOR && BINNED2_BOOKKEEPING_AT_THE_END_OF_LARGEBLOCK
-
-		if (Alignment > BINNED2_MINIMUM_ALIGNMENT && (Size <= BINNED2_MAX_SMALL_POOL_SIZE))
-		{
-			Size = Align(Size, Alignment);
-		}
-#endif
-		void* Result = nullptr;
 		// Only allocate from the small pools if the size is small enough and the alignment isn't crazy large.
 		// With large alignments, we'll waste a lot of memory allocating an entire page, but such alignments are highly unlikely in practice.
-#if UE_USE_VERYLARGEPAGEALLOCATOR && BINNED2_BOOKKEEPING_AT_THE_END_OF_LARGEBLOCK
-		if ((Size <= BINNED2_MAX_SMALL_POOL_SIZE)) // one branch, not two
-#else
-		if ((Size <= BINNED2_MAX_SMALL_POOL_SIZE) & (Alignment <= BINNED2_MINIMUM_ALIGNMENT)) // one branch, not two
-#endif
+		const bool bUseSmallPool = UseSmallAlloc(Size, Alignment);
+		if (bUseSmallPool)
 		{
 			FPerThreadFreeBlockLists* Lists = GMallocBinned2PerThreadCaches ? FPerThreadFreeBlockLists::Get() : nullptr;
 			if (Lists)
 			{
-				uint32 PoolIndex = BoundSizeToPoolIndex(Size);
-				uint32 BlockSize = PoolIndexToBlockSize(PoolIndex);
-				Result = Lists->Malloc(PoolIndex);
-#if BINNED2_ALLOCATOR_STATS
-				if (Result)
+				const uint32 PoolIndex = BoundSizeToPoolIndex(Size);
+				if (void* Result = Lists->Malloc(PoolIndex))
 				{
+#if BINNED2_ALLOCATOR_STATS
+					const uint32 BlockSize = PoolIndexToBlockSize(PoolIndex);
 					Lists->AllocatedMemory += BlockSize;
-				}
 #endif
+					return Result;
+				}
 			}
 		}
-		if (Result == nullptr)
-		{
-			Result = MallocSelect(Size, Alignment);
-		}
 
-		return Result;
+		return MallocSelect(Size, Alignment, bUseSmallPool);
 	}
 	FORCEINLINE static bool UseSmallAlloc(SIZE_T Size, uint32 Alignment)
 	{
@@ -554,20 +539,10 @@ public:
 #endif
 		return bResult;
 	}
+	CORE_API void* MallocSelect(SIZE_T Size, uint32 Alignment, bool bUseSmallPool);
 	FORCEINLINE void* MallocSelect(SIZE_T Size, uint32 Alignment)
 	{
-		void* Result;
-
-		if (UseSmallAlloc(Size, Alignment))
-		{
-			Result = MallocExternalSmall(Size, Alignment);
-		}
-		else
-		{
-			Result = MallocExternalLarge(Size, Alignment);
-		}
-
-		return Result;
+		return MallocSelect(Size, Alignment, UseSmallAlloc(Size, Alignment));
 	}
 	FORCEINLINE virtual void* Realloc(void* Ptr, SIZE_T NewSize, uint32 Alignment) override
 	{
@@ -614,7 +589,7 @@ public:
 		}
 		if (NewSize <= BINNED2_MAX_SMALL_POOL_SIZE)
 #else
-if (NewSize <= BINNED2_MAX_SMALL_POOL_SIZE && Alignment <= BINNED2_MINIMUM_ALIGNMENT) // one branch, not two
+		if (NewSize <= BINNED2_MAX_SMALL_POOL_SIZE && Alignment <= BINNED2_MINIMUM_ALIGNMENT) // one branch, not two
 #endif
 		{
 			FPerThreadFreeBlockLists* Lists = GMallocBinned2PerThreadCaches ? FPerThreadFreeBlockLists::Get() : nullptr;
@@ -746,13 +721,31 @@ if (NewSize <= BINNED2_MAX_SMALL_POOL_SIZE && Alignment <= BINNED2_MINIMUM_ALIGN
 		if ((Count <= BINNED2_MAX_SMALL_POOL_SIZE) & (Alignment <= BINNED2_MINIMUM_ALIGNMENT)) // one branch, not two
 		{
 			SizeOut = PoolIndexToBlockSize(BoundSizeToPoolIndex(Count));
+			check(SizeOut >= Count);
+			return SizeOut;
 		}
-		else
+		Alignment = FMath::Max<uint32>(Alignment, BINNED2_MINIMUM_ALIGNMENT);
+		Count = Align(Count, Alignment);
+		if ((Count <= BINNED2_MAX_SMALL_POOL_SIZE) & (Alignment <= BINNED2_MAXIMUM_ALIGNMENT))
 		{
-			Alignment = FPlatformMath::Max<uint32>(Alignment, OsAllocationGranularity);
-			checkSlow(Alignment <= PageSize);
-			SizeOut = Align(Count, Alignment);
+			uint32 PoolIndex = BoundSizeToPoolIndex(Count);
+			do
+			{
+				uint32 BlockSize = PoolIndexToBlockSize(PoolIndex);
+				if (IsAligned(BlockSize, Alignment))
+				{
+					SizeOut = SIZE_T(BlockSize);
+					check(SizeOut >= Count);
+					return SizeOut;
+				}
+
+				PoolIndex++;
+			} while (PoolIndex < BINNED2_SMALL_POOL_COUNT);
 		}
+
+		Alignment = FPlatformMath::Max<uint32>(Alignment, OsAllocationGranularity);
+		checkSlow(Alignment <= PageSize);
+		SizeOut = Align(Count, Alignment);
 		check(SizeOut >= Count);
 		return SizeOut;
 	}
@@ -793,7 +786,7 @@ if (NewSize <= BINNED2_MAX_SMALL_POOL_SIZE && Alignment <= BINNED2_MINIMUM_ALIGN
 	// Mapping of sizes to small table indices
 	static CORE_API uint8 MemSizeToIndex[1 + (BINNED2_MAX_SMALL_POOL_SIZE >> BINNED2_MINIMUM_ALIGNMENT_SHIFT)];
 
-	FORCEINLINE uint32 BoundSizeToPoolIndex(SIZE_T Size) 
+	FORCEINLINE uint32 BoundSizeToPoolIndex(SIZE_T Size) const
 	{
 		auto Index = ((Size + BINNED2_MINIMUM_ALIGNMENT - 1) >> BINNED2_MINIMUM_ALIGNMENT_SHIFT);
 		checkSlow(Index >= 0 && Index <= (BINNED2_MAX_SMALL_POOL_SIZE >> BINNED2_MINIMUM_ALIGNMENT_SHIFT)); // and it should be in the table
@@ -801,7 +794,7 @@ if (NewSize <= BINNED2_MAX_SMALL_POOL_SIZE && Alignment <= BINNED2_MINIMUM_ALIGN
 		checkSlow(PoolIndex >= 0 && PoolIndex < BINNED2_SMALL_POOL_COUNT);
 		return PoolIndex;
 	}
-	FORCEINLINE uint32 PoolIndexToBlockSize(uint32 PoolIndex)
+	FORCEINLINE uint32 PoolIndexToBlockSize(uint32 PoolIndex) const
 	{
 		return SmallBlockSizesReversed[BINNED2_SMALL_POOL_COUNT - PoolIndex - 1];
 	}
diff --git a/Engine/Source/Runtime/Core/Public/HAL/MallocBinned3.h b/Engine/Source/Runtime/Core/Public/HAL/MallocBinned3.h
--- a/Engine/Source/Runtime/Core/Public/HAL/MallocBinned3.h
+++ b/Engine/Source/Runtime/Core/Public/HAL/MallocBinned3.h
@@ -667,12 +667,30 @@ public:
 		if ((Count <= BINNED3_MAX_SMALL_POOL_SIZE) & (Alignment <= BINNED3_MINIMUM_ALIGNMENT)) // one branch, not two
 		{
 			SizeOut = PoolIndexToBlockSize(BoundSizeToPoolIndex(Count));
+			check(SizeOut >= Count);
+			return SizeOut;
 		}
-		else
+		Alignment = FMath::Max<uint32>(Alignment, BINNED3_MINIMUM_ALIGNMENT);
+		Count = Align(Count, Alignment);
+		if ((Count <= BINNED3_MAX_SMALL_POOL_SIZE) & (Alignment <= BINNED3_MAX_SMALL_POOL_ALIGNMENT))
 		{
-			Alignment = FPlatformMath::Max<uint32>(Alignment, OsAllocationGranularity);
-			SizeOut = Align(Count, Alignment);
+			uint32 PoolIndex = BoundSizeToPoolIndex(Count);
+			do
+			{
+				uint32 BlockSize = PoolIndexToBlockSize(PoolIndex);
+				if (IsAligned(BlockSize, Alignment))
+				{
+					SizeOut = SIZE_T(BlockSize);
+					check(SizeOut >= Count);
+					return SizeOut;
+				}
+
+				PoolIndex++;
+			} while (PoolIndex < BINNED3_SMALL_POOL_COUNT);
 		}
+		
+		Alignment = FPlatformMath::Max<uint32>(Alignment, OsAllocationGranularity);
+		SizeOut = Align(Count, Alignment);
 		check(SizeOut >= Count);
 		return SizeOut;
 	}
