Bump material translator DDC key to fix shader cache poisoning.

diff --git a/Engine/Source/Runtime/Core/Private/UObject/DevObjectVersion.cpp b/Engine/Source/Runtime/Core/Private/UObject/DevObjectVersion.cpp
--- a/Engine/Source/Runtime/Core/Private/UObject/DevObjectVersion.cpp
+++ b/Engine/Source/Runtime/Core/Private/UObject/DevObjectVersion.cpp
@@ -129,7 +129,7 @@ FDevSystemGuids::FDevSystemGuids()
 	, POSESEARCHDB_DERIVEDDATA_VER(0x389117E4, 0x807A4CC0, 0x9F37C2E6, 0xD808A78D)
 	, SkeletalMeshDerivedDataVersion(0x9B5F4544, 0x76D7481C, 0x9AD3F614, 0xA6C07904)
 	, STATICMESH_DERIVEDDATA_VER(0x5B7A05A7, 0x7371440C, 0xA314621A, 0x9B1E08BC)
-	, MaterialTranslationDDCVersion(0x982B2D2E, 0x6E50436A, 0x95EF4D16, 0xB1D04920)
+	, MaterialTranslationDDCVersion(0x7CBF2AF7, 0xAD554586, 0x8E3A15AA, 0x2BA900D7)
 {
 }
 
diff --git a/Engine/Source/Runtime/Core/Private/UObject/FortniteMainBranchObjectVersion.cpp b/Engine/Source/Runtime/Core/Private/UObject/FortniteMainBranchObjectVersion.cpp
--- a/Engine/Source/Runtime/Core/Private/UObject/FortniteMainBranchObjectVersion.cpp
+++ b/Engine/Source/Runtime/Core/Private/UObject/FortniteMainBranchObjectVersion.cpp
@@ -14,6 +14,6 @@ TMap<FGuid, FGuid> FFortniteMainBranchObjectVersion::GetSystemGuids()
 	SystemGuids.Add(DevGuids.Niagara_LatestScriptCompileVersion, FGuid("CE33E29C6268469E8E18F1D96A55DFDA"));
 	SystemGuids.Add(DevGuids.SkeletalMeshDerivedDataVersion, FGuid("518D7C51D31D9BE4A9F4B81542AC8119"));
 	SystemGuids.Add(DevGuids.STATICMESH_DERIVEDDATA_VER, FGuid("543843B103794E0E9BA4BE80FB602F79"));
-	SystemGuids.Add(DevGuids.MaterialTranslationDDCVersion, FGuid("982B2D2E6E50436A95EF4D16B1D04920"));
+	SystemGuids.Add(DevGuids.MaterialTranslationDDCVersion, FGuid("7CBF2AF7AD5545868E3A15AA2BA900D7"));
 	return SystemGuids;
 }
