Treat an empty source string like a null source string in FindDisplayString/GetDisplayString on FTextLocalizationManager

Since an empty source string cannot be localized this should be a safe change to make, and generally addresses an issue where tooltips in a shipping editor would be empty, as shipping builds strip out the source string meta-data from the build

diff --git a/Engine/Source/Runtime/Core/Private/Internationalization/TextLocalizationManager.cpp b/Engine/Source/Runtime/Core/Private/Internationalization/TextLocalizationManager.cpp
--- a/Engine/Source/Runtime/Core/Private/Internationalization/TextLocalizationManager.cpp
+++ b/Engine/Source/Runtime/Core/Private/Internationalization/TextLocalizationManager.cpp
@@ -777,7 +777,7 @@ void FTextLocalizationManager::RegisterPolyglotTextData(TArrayView<const FPolygl
 	}
 }
 
-FTextConstDisplayStringPtr FTextLocalizationManager::FindDisplayString(const FTextKey& Namespace, const FTextKey& Key, const FString* const SourceString) const
+FTextConstDisplayStringPtr FTextLocalizationManager::FindDisplayString(const FTextKey& Namespace, const FTextKey& Key, const FString* const SourceStringPtr) const
 {
 	if (!FTextLocalizationManager::IsDisplayStringSupportEnabled())
 	{
@@ -788,17 +788,30 @@ FTextConstDisplayStringPtr FTextLocalizationManager::FindDisplayString(const FTe
 
 	const FTextId TextId(Namespace, Key);
 
-	const FDisplayStringEntry* LiveEntry = DisplayStringLookupTable.Find(TextId);
-
-	if ( LiveEntry != nullptr && ( !SourceString || LiveEntry->SourceStringHash == FTextLocalizationResource::HashString(*SourceString) ) )
+	if (const FDisplayStringEntry* LiveEntry = DisplayStringLookupTable.Find(TextId))
 	{
-		return LiveEntry->DisplayString;
+		auto GetSourceStringRef = [SourceStringPtr]() -> const FString&
+		{
+			if (SourceStringPtr)
+			{
+				return *SourceStringPtr;
+			}
+
+			static const FString EmptyString;
+			return EmptyString;
+		};
+
+		const FString& SourceString = GetSourceStringRef();
+		if (SourceString.IsEmpty() || LiveEntry->SourceStringHash == FTextLocalizationResource::HashString(SourceString))
+		{
+			return LiveEntry->DisplayString;
+		}
 	}
 
 	return nullptr;
 }
 
-FTextConstDisplayStringRef FTextLocalizationManager::GetDisplayString(const FTextKey& Namespace, const FTextKey& Key, const FString* const SourceString)
+FTextConstDisplayStringRef FTextLocalizationManager::GetDisplayString(const FTextKey& Namespace, const FTextKey& Key, const FString* const SourceStringPtr)
 {
 	TRACE_CPUPROFILER_EVENT_SCOPE(FTextLocalizationManager::GetDisplayString);
 	LLM_SCOPE(ELLMTag::Localization);
@@ -815,12 +828,25 @@ FTextConstDisplayStringRef FTextLocalizationManager::GetDisplayString(const FTex
 		return GetEmptyDisplayString();
 	}
 
+	auto GetSourceStringRef = [SourceStringPtr]() -> const FString&
+	{
+		if (SourceStringPtr)
+		{
+			return *SourceStringPtr;
+		}
+
+		static const FString EmptyString;
+		return EmptyString;
+	};
+
+	const FString& SourceString = GetSourceStringRef();
+
 	FScopeLock ScopeLock(&DisplayStringLookupTableCS);
 
 	// Hack fix for old assets that don't have namespace/key info.
 	if (Namespace.IsEmpty() && Key.IsEmpty())
 	{
-		return MakeShared<FString, ESPMode::ThreadSafe>(SourceString ? *SourceString : FString());
+		return MakeTextDisplayString(CopyTemp(SourceString));
 	}
 
 #if ENABLE_LOC_TESTING
@@ -840,7 +866,7 @@ FTextConstDisplayStringRef FTextLocalizationManager::GetDisplayString(const FTex
 
 	const FTextId TextId(Namespace, Key);
 
-	const uint32 SourceStringHash = SourceString ? FTextLocalizationResource::HashString(*SourceString) : 0;
+	const uint32 SourceStringHash = !SourceString.IsEmpty() ? FTextLocalizationResource::HashString(SourceString) : 0;
 
 	FDisplayStringEntry* LiveEntry = DisplayStringLookupTable.Find(TextId);
 
@@ -856,7 +882,7 @@ FTextConstDisplayStringRef FTextLocalizationManager::GetDisplayString(const FTex
 
 			if (SourceLiveEntry)
 			{
-				if (!SourceString || SourceLiveEntry->SourceStringHash == SourceStringHash)
+				if (SourceString.IsEmpty() || SourceLiveEntry->SourceStringHash == SourceStringHash)
 				{
 					SourceDisplayString = SourceLiveEntry->DisplayString;
 				}
@@ -873,10 +899,10 @@ FTextConstDisplayStringRef FTextLocalizationManager::GetDisplayString(const FTex
 	if (LiveEntry)
 	{
 		// If the source string (hash) is different, the local source has changed and should override
-		if (SourceString && SourceStringHash != LiveEntry->SourceStringHash)
+		if (!SourceString.IsEmpty() && SourceStringHash != LiveEntry->SourceStringHash)
 		{
 			LiveEntry->SourceStringHash = SourceStringHash;
-			LiveEntry->DisplayString = SourceDisplayString ? SourceDisplayString.ToSharedRef() : MakeTextDisplayString(CopyTemp(*SourceString));
+			LiveEntry->DisplayString = SourceDisplayString ? SourceDisplayString.ToSharedRef() : MakeTextDisplayString(CopyTemp(SourceString));
 			DirtyLocalRevisionForTextId(TextId);
 
 #if ENABLE_LOC_TESTING
@@ -898,7 +924,7 @@ FTextConstDisplayStringRef FTextLocalizationManager::GetDisplayString(const FTex
 	{
 		TRACE_CPUPROFILER_EVENT_SCOPE(FTextLocalizationManager::GetDisplayString_AddRelatedEntry);
 
-		check(!SourceString || SourceLiveEntry->SourceStringHash == SourceStringHash);
+		check(SourceString.IsEmpty() || SourceLiveEntry->SourceStringHash == SourceStringHash);
 		check(SourceDisplayString && SourceLiveEntry->DisplayString == SourceDisplayString);
 
 		// Clone the entry for the active ID
@@ -915,10 +941,10 @@ FTextConstDisplayStringRef FTextLocalizationManager::GetDisplayString(const FTex
 		// Don't log warnings about unlocalized strings if the system hasn't been initialized - we simply don't have localization data yet.
 		if (IsInitialized())
 		{
-			UE_LOG(LogTextLocalizationManager, Verbose, TEXT("An attempt was made to get a localized string (Namespace:%s, Key:%s, Source:%s), but it did not exist."), TextId.GetNamespace().GetChars(), TextId.GetKey().GetChars(), SourceString ? **SourceString : TEXT(""));
+			UE_LOG(LogTextLocalizationManager, Verbose, TEXT("An attempt was made to get a localized string (Namespace:%s, Key:%s, Source:%s), but it did not exist."), TextId.GetNamespace().GetChars(), TextId.GetKey().GetChars(), *SourceString);
 		}
 		
-		FTextConstDisplayStringRef UnlocalizedString = SourceString ? MakeTextDisplayString(CopyTemp(*SourceString)) : GetEmptyDisplayString();
+		FTextConstDisplayStringRef UnlocalizedString = !SourceString.IsEmpty() ? MakeTextDisplayString(CopyTemp(SourceString)) : GetEmptyDisplayString();
 
 #if ENABLE_LOC_TESTING
 		if ((bShouldLEETIFYAll || bShouldLEETIFYUnlocalizedString) && !UnlocalizedString->IsEmpty())
diff --git a/Engine/Source/Runtime/Core/Public/Internationalization/TextLocalizationManager.h b/Engine/Source/Runtime/Core/Public/Internationalization/TextLocalizationManager.h
--- a/Engine/Source/Runtime/Core/Public/Internationalization/TextLocalizationManager.h
+++ b/Engine/Source/Runtime/Core/Public/Internationalization/TextLocalizationManager.h
@@ -159,16 +159,16 @@ public:
 	CORE_API void RegisterPolyglotTextData(TArrayView<const FPolyglotTextData> InPolyglotTextDataArray, const bool InAddDisplayStrings = true);
 
 	/**	Finds and returns the display string with the given namespace and key, if it exists.
-	 *	Additionally, if a source string is specified and the found localized display string was not localized from that source string, null will be returned. */
-	CORE_API FTextConstDisplayStringPtr FindDisplayString(const FTextKey& Namespace, const FTextKey& Key, const FString* const SourceString = nullptr) const;
+	 *	Additionally, if a non-null and non-empty source string is specified and the found localized display string was not localized from that source string, null will be returned. */
+	CORE_API FTextConstDisplayStringPtr FindDisplayString(const FTextKey& Namespace, const FTextKey& Key, const FString* const SourceStringPtr = nullptr) const;
 
 	/**	Returns a display string with the given namespace and key.
 	 *	If no display string exists, it will be created using the source string or an empty string if no source string is provided.
 	 *	If a display string exists ...
 	 *		... but it was not localized from the specified source string, the display string will be set to the specified source and returned.
-	 *		... and it was localized from the specified source string (or none was provided), the display string will be returned.
+	 *		... and it was localized from the specified source string (or the source string was null or empty), the display string will be returned.
 	*/
-	CORE_API FTextConstDisplayStringRef GetDisplayString(const FTextKey& Namespace, const FTextKey& Key, const FString* const SourceString);
+	CORE_API FTextConstDisplayStringRef GetDisplayString(const FTextKey& Namespace, const FTextKey& Key, const FString* const SourceStringPtr);
 
 #if WITH_EDITORONLY_DATA
 	/** If an entry exists for the specified namespace and key, returns true and provides the localization resource identifier from which it was loaded. Otherwise, returns false. */
