Rename OpenCommit -> OnCommit and OpenAbort -> OnAbort to make it more obvious to users when the code executes.

diff --git a/Engine/Plugins/Experimental/UIFramework/Source/Private/Types/UIFWidgetTree.cpp b/Engine/Plugins/Experimental/UIFramework/Source/Private/Types/UIFWidgetTree.cpp
--- a/Engine/Plugins/Experimental/UIFramework/Source/Private/Types/UIFWidgetTree.cpp
+++ b/Engine/Plugins/Experimental/UIFramework/Source/Private/Types/UIFWidgetTree.cpp
@@ -222,11 +222,13 @@ void FUIFrameworkWidgetTree::AuthorityAddChildInternal(UUIFrameworkWidget* Paren
 
 			UObject* OldOuter = Child->GetOuter();
 
-			UE_AUTORTFM_OPEN({
+			UE_AUTORTFM_OPEN(
+			{
 				Child->Rename(nullptr, ReplicatedOwner);
 			});
 
-			UE_AUTORTFM_OPENABORT({
+			UE_AUTORTFM_ONABORT(
+			{
 				Child->Rename(nullptr, OldOuter);
 			});
 		}
diff --git a/Engine/Source/Programs/AutoRTFMTests/Private/API.cpp b/Engine/Source/Programs/AutoRTFMTests/Private/API.cpp
--- a/Engine/Source/Programs/AutoRTFMTests/Private/API.cpp
+++ b/Engine/Source/Programs/AutoRTFMTests/Private/API.cpp
@@ -276,7 +276,7 @@ TEST_CASE("API.autortfm_register_open_function")
     REQUIRE(0 == I);
 }
 
-TEST_CASE("API.autortfm_open_commit")
+TEST_CASE("API.autortfm_on_commit")
 {
     bool OuterTransaction = false;
     bool InnerTransaction = false;
@@ -286,7 +286,7 @@ TEST_CASE("API.autortfm_open_commit")
 
     AutoRTFM::Commit([&]
     {
-        autortfm_open_commit([](void* const Arg)
+        autortfm_on_commit([](void* const Arg)
         {
             *static_cast<bool* const>(Arg) = true;
         }, &OuterTransaction);
@@ -299,7 +299,7 @@ TEST_CASE("API.autortfm_open_commit")
 
         AutoRTFM::Commit([&]
         {
-            autortfm_open_commit([](void* const Arg)
+            autortfm_on_commit([](void* const Arg)
             {
                 *static_cast<bool* const>(Arg) = true;
             }, &InnerTransaction);
@@ -313,7 +313,7 @@ TEST_CASE("API.autortfm_open_commit")
 
         NestResult = AutoRTFM::Transact([&]
         {
-            autortfm_open_commit([](void* const Arg)
+            autortfm_on_commit([](void* const Arg)
             {
                 *static_cast<bool* const>(Arg) = true;
             }, &InnerTransactionWithAbort);
@@ -329,7 +329,7 @@ TEST_CASE("API.autortfm_open_commit")
 
         AutoRTFM::Open([&]
         {
-            autortfm_open_commit([](void* const Arg)
+            autortfm_on_commit([](void* const Arg)
             {
                 *static_cast<bool* const>(Arg) = true;
             }, &InnerOpenNest);
@@ -349,7 +349,7 @@ TEST_CASE("API.autortfm_open_commit")
     REQUIRE(AutoRTFM::ETransactionResult::AbortedByRequest == NestResult);
 }
 
-TEST_CASE("API.autortfm_open_abort")
+TEST_CASE("API.autortfm_on_abort")
 {
     bool OuterTransaction = false;
     bool InnerTransaction = false;
@@ -359,7 +359,7 @@ TEST_CASE("API.autortfm_open_abort")
 
     REQUIRE(AutoRTFM::ETransactionResult::Committed == AutoRTFM::Transact([&]
     {
-        autortfm_open_abort([](void* const Arg)
+        autortfm_on_abort([](void* const Arg)
         {
             *static_cast<bool* const>(Arg) = true;
         }, &OuterTransaction);
@@ -372,7 +372,7 @@ TEST_CASE("API.autortfm_open_abort")
 
         AutoRTFM::Commit([&]
         {
-            autortfm_open_abort([](void* const Arg)
+            autortfm_on_abort([](void* const Arg)
             {
                 *static_cast<bool* const>(Arg) = true;
             }, &InnerTransaction);
@@ -386,7 +386,7 @@ TEST_CASE("API.autortfm_open_abort")
 
         NestResult = AutoRTFM::Transact([&]
         {
-            autortfm_open_abort([](void* const Arg)
+            autortfm_on_abort([](void* const Arg)
             {
                 *static_cast<bool* const>(Arg) = true;
             }, &InnerTransactionWithAbort);
@@ -402,7 +402,7 @@ TEST_CASE("API.autortfm_open_abort")
 
         AutoRTFM::Open([&]
         {
-            autortfm_open_abort([](void* const Arg)
+            autortfm_on_abort([](void* const Arg)
             {
                 *static_cast<bool* const>(Arg) = true;
             }, &InnerOpenNest);
@@ -725,7 +725,7 @@ TEST_CASE("API.RegisterOpenFunction")
     REQUIRE(0 == I);
 }
 
-TEST_CASE("API.OpenCommit")
+TEST_CASE("API.OnCommit")
 {
     bool OuterTransaction = false;
     bool InnerTransaction = false;
@@ -735,7 +735,7 @@ TEST_CASE("API.OpenCommit")
 
     REQUIRE(AutoRTFM::ETransactionResult::Committed == AutoRTFM::Transact([&]
     {
-        AutoRTFM::OpenCommit([&]
+        AutoRTFM::OnCommit([&]
         {
             OuterTransaction = true;
         });
@@ -748,7 +748,7 @@ TEST_CASE("API.OpenCommit")
 
         AutoRTFM::Commit([&]
         {
-            AutoRTFM::OpenCommit([&]
+            AutoRTFM::OnCommit([&]
             {
                 InnerTransaction = true;
             });
@@ -762,7 +762,7 @@ TEST_CASE("API.OpenCommit")
 
         NestResult = AutoRTFM::Transact([&]
         {
-            AutoRTFM::OpenCommit([&]
+            AutoRTFM::OnCommit([&]
             {
                 InnerTransactionWithAbort = true;
             });
@@ -778,7 +778,7 @@ TEST_CASE("API.OpenCommit")
 
         AutoRTFM::Open([&]
         {
-            AutoRTFM::OpenCommit([&]
+            AutoRTFM::OnCommit([&]
             {
                 InnerOpenNest = true;
             });
@@ -798,7 +798,7 @@ TEST_CASE("API.OpenCommit")
     REQUIRE(AutoRTFM::ETransactionResult::AbortedByRequest == NestResult);
 }
 
-TEST_CASE("API.OpenAbort")
+TEST_CASE("API.OnAbort")
 {
     bool OuterTransaction = false;
     bool InnerTransaction = false;
@@ -808,7 +808,7 @@ TEST_CASE("API.OpenAbort")
 
     REQUIRE(AutoRTFM::ETransactionResult::Committed == AutoRTFM::Transact([&]
     {
-        AutoRTFM::OpenAbort([&]
+        AutoRTFM::OnAbort([&]
         {
             OuterTransaction = true;
         });
@@ -821,7 +821,7 @@ TEST_CASE("API.OpenAbort")
 
         AutoRTFM::Commit([&]
         {
-            AutoRTFM::OpenAbort([&]
+            AutoRTFM::OnAbort([&]
             {
                 InnerTransaction = true;
             });
@@ -835,7 +835,7 @@ TEST_CASE("API.OpenAbort")
 
         NestResult = AutoRTFM::Transact([&]
         {
-            AutoRTFM::OpenAbort([&]
+            AutoRTFM::OnAbort([&]
             {
                 InnerTransactionWithAbort = true;
             });
@@ -851,7 +851,7 @@ TEST_CASE("API.OpenAbort")
 
         AutoRTFM::Open([&]
         {
-            AutoRTFM::OpenAbort([&]
+            AutoRTFM::OnAbort([&]
             {
                 InnerOpenNest = true;
             });
diff --git a/Engine/Source/Programs/AutoRTFMTests/Private/Abort.cpp b/Engine/Source/Programs/AutoRTFMTests/Private/Abort.cpp
--- a/Engine/Source/Programs/AutoRTFMTests/Private/Abort.cpp
+++ b/Engine/Source/Programs/AutoRTFMTests/Private/Abort.cpp
@@ -59,13 +59,13 @@ TEST_CASE("Abort.NestedAbortOrder")
 	{
 		InnerResult = AutoRTFM::Transact([&]
 			{
-				AutoRTFM::OpenAbort([&]
+				AutoRTFM::OnAbort([&]
 					{
 						REQUIRE(1 == Orderer);
 						Orderer += 1;
 					});
 
-				AutoRTFM::OpenAbort([&]
+				AutoRTFM::OnAbort([&]
 					{
 						REQUIRE(0 == Orderer);
 						Orderer += 1;
@@ -85,7 +85,7 @@ TEST_CASE("Abort.TransactionInOpenCommit")
 
 	AutoRTFM::Commit([&]
 	{
-		AutoRTFM::OpenCommit([&]
+		AutoRTFM::OnCommit([&]
 		{
 			bool bDidSomething = false;
 
@@ -98,7 +98,7 @@ TEST_CASE("Abort.TransactionInOpenCommit")
 		});
 	});
 
-	REQUIRE(AutoRTFM::ETransactionResult::AbortedByTransactInOpenCommit == InnerResult);
+	REQUIRE(AutoRTFM::ETransactionResult::AbortedByTransactInOnCommit == InnerResult);
 }
 
 TEST_CASE("Abort.TransactionInOpenAbort")
@@ -108,7 +108,7 @@ TEST_CASE("Abort.TransactionInOpenAbort")
 
 	Result = AutoRTFM::Transact([&]
 	{
-		AutoRTFM::OpenAbort([&]
+		AutoRTFM::OnAbort([&]
 		{
 			bool bDidSomething = false;
 
@@ -123,7 +123,7 @@ TEST_CASE("Abort.TransactionInOpenAbort")
 		AutoRTFM::AbortTransaction();
 	});
 
-	REQUIRE(AutoRTFM::ETransactionResult::AbortedByTransactInOpenAbort == InnerResult);
+	REQUIRE(AutoRTFM::ETransactionResult::AbortedByTransactInOnAbort == InnerResult);
 	REQUIRE(AutoRTFM::ETransactionResult::AbortedByRequest == Result);
 }
 
diff --git a/Engine/Source/Runtime/Core/Private/AutoRTFM/API.cpp b/Engine/Source/Runtime/Core/Private/AutoRTFM/API.cpp
--- a/Engine/Source/Runtime/Core/Private/AutoRTFM/API.cpp
+++ b/Engine/Source/Runtime/Core/Private/AutoRTFM/API.cpp
@@ -234,21 +234,21 @@ extern "C" UE_AUTORTFM_NOAUTORTFM void autortfm_register_open_function(void* Ori
     FunctionMapAdd(OriginalFunction, NewFunction);
 }
 
-UE_AUTORTFM_AUTORTFM("RTFM_OpenCommit") void OpenCommit(TFunction<void()> && Work)
+UE_AUTORTFM_AUTORTFM("RTFM_OnCommit") void OnCommit(TFunction<void()> && Work)
 {
     Work();
 }
 
-UE_AUTORTFM_AUTORTFM("RTFM_OpenAbort") void OpenAbort(TFunction<void()> && Work)
+UE_AUTORTFM_AUTORTFM("RTFM_OnAbort") void OnAbort(TFunction<void()> && Work)
 {
 }
 
-extern "C" UE_AUTORTFM_AUTORTFM("RTFM_autortfm_open_commit") void autortfm_open_commit(void (*Work)(void* Arg), void* Arg)
+extern "C" UE_AUTORTFM_AUTORTFM("RTFM_autortfm_on_commit") void autortfm_on_commit(void (*Work)(void* Arg), void* Arg)
 {
     Work(Arg);
 }
 
-extern "C" UE_AUTORTFM_AUTORTFM("RTFM_autortfm_open_abort") void autortfm_open_abort(void (*Work)(void* arg), void* Arg)
+extern "C" UE_AUTORTFM_AUTORTFM("RTFM_autortfm_on_abort") void autortfm_on_abort(void (*Work)(void* arg), void* Arg)
 {
 }
 
@@ -442,28 +442,28 @@ extern "C" UE_AUTORTFM_NOAUTORTFM void RTFM_autortfm_record_open_write(void*, si
 	UE_LOG(LogAutoRTFM, Fatal, TEXT("The function `autortfm_record_open_write` was called from closed code."));
 }
 
-extern "C" UE_AUTORTFM_NOAUTORTFM void RTFM_OpenCommit(TFunction<void()> && Work)
+extern "C" UE_AUTORTFM_NOAUTORTFM void RTFM_OnCommit(TFunction<void()> && Work)
 {
 	FContext* Context = FContext::Get();
     ASSERT(Context->GetStatus() == EContextStatus::OnTrack);
     Context->GetCurrentTransaction()->DeferUntilCommit(MoveTemp(Work));
 }
 
-extern "C" UE_AUTORTFM_NOAUTORTFM void RTFM_OpenAbort(TFunction<void()>&& Work)
+extern "C" UE_AUTORTFM_NOAUTORTFM void RTFM_OnAbort(TFunction<void()>&& Work)
 {
 	FContext* Context = FContext::Get();
     ASSERT(Context->GetStatus() == EContextStatus::OnTrack);
     Context->GetCurrentTransaction()->DeferUntilAbort(MoveTemp(Work));
 }
 
-extern "C" UE_AUTORTFM_NOAUTORTFM void RTFM_autortfm_open_commit(void (*Work)(void* Arg), void* Arg)
+extern "C" UE_AUTORTFM_NOAUTORTFM void RTFM_autortfm_on_commit(void (*Work)(void* Arg), void* Arg)
 {
-    RTFM_OpenCommit([Work, Arg] { Work(Arg); });
+    RTFM_OnCommit([Work, Arg] { Work(Arg); });
 }
 
-extern "C" UE_AUTORTFM_NOAUTORTFM void RTFM_autortfm_open_abort(void (*Work)(void* arg), void* Arg)
+extern "C" UE_AUTORTFM_NOAUTORTFM void RTFM_autortfm_on_abort(void (*Work)(void* arg), void* Arg)
 {
-    RTFM_OpenAbort([Work, Arg] { Work(Arg); });
+    RTFM_OnAbort([Work, Arg] { Work(Arg); });
 }
 
 extern "C" UE_AUTORTFM_NOAUTORTFM void* RTFM_autortfm_did_allocate(void* Ptr, size_t Size)
diff --git a/Engine/Source/Runtime/Core/Private/AutoRTFM/Context.cpp b/Engine/Source/Runtime/Core/Private/AutoRTFM/Context.cpp
--- a/Engine/Source/Runtime/Core/Private/AutoRTFM/Context.cpp
+++ b/Engine/Source/Runtime/Core/Private/AutoRTFM/Context.cpp
@@ -289,12 +289,12 @@ ETransactionResult FContext::Transact(void (*Function)(void* Arg), void* Arg)
 
     if (UNLIKELY(EContextStatus::Committing == Status))
     {
-    	return ETransactionResult::AbortedByTransactInOpenCommit;
+    	return ETransactionResult::AbortedByTransactInOnCommit;
     }
 
     if (UNLIKELY(IsAborting()))
     {
-    	return ETransactionResult::AbortedByTransactInOpenAbort;
+    	return ETransactionResult::AbortedByTransactInOnAbort;
     }
     
     ASSERT(Status == EContextStatus::Idle || Status == EContextStatus::OnTrack);
diff --git a/Engine/Source/Runtime/Core/Private/AutoRTFM/StdLibSupport.cpp b/Engine/Source/Runtime/Core/Private/AutoRTFM/StdLibSupport.cpp
--- a/Engine/Source/Runtime/Core/Private/AutoRTFM/StdLibSupport.cpp
+++ b/Engine/Source/Runtime/Core/Private/AutoRTFM/StdLibSupport.cpp
@@ -308,7 +308,7 @@ BOOL RTFM_TlsSetValue(DWORD dwTlsIndex, LPVOID lpTlsValue)
 {
 	LPVOID CurrentValue = TlsGetValue(dwTlsIndex);
 
-	AutoRTFM::OpenAbort([dwTlsIndex, CurrentValue]
+	AutoRTFM::OnAbort([dwTlsIndex, CurrentValue]
 	{
 		TlsSetValue(dwTlsIndex, CurrentValue);
 	});
diff --git a/Engine/Source/Runtime/Core/Public/AutoRTFM/AutoRTFM.h b/Engine/Source/Runtime/Core/Public/AutoRTFM/AutoRTFM.h
--- a/Engine/Source/Runtime/Core/Public/AutoRTFM/AutoRTFM.h
+++ b/Engine/Source/Runtime/Core/Public/AutoRTFM/AutoRTFM.h
@@ -8,12 +8,6 @@
 #define UE_AUTORTFM 0
 #endif
 
-#if UE_AUTORTFM
-	#define UE_PRAGMA_AUTORTFM _Pragma("autortfm")
-#else
-	#define UE_PRAGMA_AUTORTFM
-#endif
-
 #if UE_AUTORTFM
 #define UE_AUTORTFM_AUTORTFM(F) [[clang::autortfm(F), clang::noinline]]
 #define UE_AUTORTFM_NOAUTORTFM [[clang::noautortfm, clang::noinline]]
@@ -37,9 +31,14 @@
 
 #if defined(UE_AUTORTFM_STANDALONE)
 #define UE_AUTORTFM_API
+
+#define UE_AUTORTFM_MOVETEMP(x) std::move(x)
 #else
 #include <HAL/Platform.h>
 #define UE_AUTORTFM_API CORE_API
+
+#include <Templates/UnrealTemplate.h>
+#define UE_AUTORTFM_MOVETEMP(x) MoveTemp(x)
 #endif
 
 #if UE_AUTORTFM
@@ -67,12 +66,14 @@ extern "C"
 // This must match AutoRTFM::ETransactionResult.
 typedef enum
 {
-    autortfm_aborted_by_request,
+    autortfm_aborted_by_request = 0,
     autortfm_aborted_by_language,
     autortfm_committed,
-	autortfm_aborted_by_transact_in_open_commit,
-	autortfm_aborted_by_transact_in_open_abort,
-	autortfm_aborted_by_cascade
+	autortfm_aborted_by_transact_in_on_commit,
+	autortfm_aborted_by_transact_in_on_abort,
+	autortfm_aborted_by_cascade,
+	autortfm_aborted_by_transact_in_open_commit [[deprecated("Use autortfm_aborted_by_transact_in_on_commit instead.")]] = autortfm_aborted_by_transact_in_on_commit,
+	autortfm_aborted_by_transact_in_open_abort [[deprecated("Use autortfm_aborted_by_transact_in_on_abort instead.")]] = autortfm_aborted_by_transact_in_on_abort,
 } autortfm_result;
 
 // This must match AutoRTFM::EContextStatus.
@@ -304,26 +305,38 @@ UE_AUTORTFM_FORCEINLINE void autortfm_register_open_function(void* original_func
 // If this is called outside a transaction or from an open nest then the work
 // happens immediately.
 #if UE_AUTORTFM
-UE_AUTORTFM_API void autortfm_open_commit(void (*work)(void* arg), void* arg);
+UE_AUTORTFM_API void autortfm_on_commit(void (*work)(void* arg), void* arg);
 #else
-UE_AUTORTFM_FORCEINLINE void autortfm_open_commit(void (*work)(void* arg), void* arg)
+UE_AUTORTFM_FORCEINLINE void autortfm_on_commit(void (*work)(void* arg), void* arg)
 {
     work(arg);
 }
 #endif
 
+[[deprecated("Use autortfm_on_commit instead.")]]
+UE_AUTORTFM_FORCEINLINE void autortfm_open_commit(void (*work)(void* arg), void* arg)
+{
+	autortfm_on_commit(work, arg);
+}
+
 // Have some work happen when this transaction aborts. If this is called
 // outside a transaction or from an open nest then the work is ignored.
 #if UE_AUTORTFM
-UE_AUTORTFM_API void autortfm_open_abort(void (*work)(void* arg), void* arg);
+UE_AUTORTFM_API void autortfm_on_abort(void (*work)(void* arg), void* arg);
 #else
-UE_AUTORTFM_FORCEINLINE void autortfm_open_abort(void (*work)(void* arg), void* arg)
+UE_AUTORTFM_FORCEINLINE void autortfm_on_abort(void (*work)(void* arg), void* arg)
 {
 	UE_AUTORTFM_UNUSED(work);
 	UE_AUTORTFM_UNUSED(arg);
 }
 #endif
 
+[[deprecated("Use autortfm_on_abort instead.")]]
+UE_AUTORTFM_FORCEINLINE void autortfm_open_abort(void (*work)(void* arg), void* arg)
+{
+	autortfm_on_abort(work, arg);
+}
+
 // Inform the runtime that we have performed a new object allocation. It's only
 // necessary to call this inside of custom malloc implementations. As an
 // optimization, you can choose to then only have your malloc return the pointer
@@ -389,8 +402,10 @@ enum class ETransactionResult
     AbortedByRequest = autortfm_aborted_by_request,
     AbortedByLanguage = autortfm_aborted_by_language,
     Committed = autortfm_committed,
-	AbortedByTransactInOpenCommit = autortfm_aborted_by_transact_in_open_commit,
-	AbortedByTransactInOpenAbort = autortfm_aborted_by_transact_in_open_abort,
+	AbortedByTransactInOnCommit = autortfm_aborted_by_transact_in_on_commit,
+	AbortedByTransactInOnAbort = autortfm_aborted_by_transact_in_on_abort,
+	AbortedByTransactInOpenCommit [[deprecated("Use AbortedByTransactInOnCommit instead.")]] = autortfm_aborted_by_transact_in_on_commit,
+	AbortedByTransactInOpenAbort [[deprecated("Use AbortedByTransactInOnAbort instead.")]] = autortfm_aborted_by_transact_in_on_abort,
 	AbortedByCascade = autortfm_aborted_by_cascade
 };
 
@@ -569,15 +584,27 @@ UE_AUTORTFM_FORCEINLINE void RegisterOpenFunction(void* OriginalFunction, void*
 }
 
 #if UE_AUTORTFM
-UE_AUTORTFM_API void OpenCommit(TFunction<void()>&& Work);
-UE_AUTORTFM_API void OpenAbort(TFunction<void()>&& Work);
+UE_AUTORTFM_API void OnCommit(TFunction<void()>&& Work);
+UE_AUTORTFM_API void OnAbort(TFunction<void()>&& Work);
 #else
 template<typename TFunctor>
-UE_AUTORTFM_FORCEINLINE void OpenCommit(const TFunctor& Work) { Work(); }
+UE_AUTORTFM_FORCEINLINE void OnCommit(const TFunctor& Work) { Work(); }
 template<typename TFunctor>
-UE_AUTORTFM_FORCEINLINE void OpenAbort(const TFunctor& Work) { }
+UE_AUTORTFM_FORCEINLINE void OnAbort(const TFunctor& Work) { }
 #endif
 
+[[deprecated("Use OnCommit instead.")]]
+UE_AUTORTFM_FORCEINLINE void OpenCommit(TFunction<void()>&& Work)
+{
+	OnCommit(UE_AUTORTFM_MOVETEMP(Work));
+}
+
+[[deprecated("Use OnAbort instead.")]]
+UE_AUTORTFM_FORCEINLINE void OpenAbort(TFunction<void()>&& Work)
+{
+	OnAbort(UE_AUTORTFM_MOVETEMP(Work));
+}
+
 UE_AUTORTFM_FORCEINLINE void* DidAllocate(void* Ptr, size_t Size)
 {
     return autortfm_did_allocate(Ptr, Size);
@@ -595,12 +622,21 @@ UE_AUTORTFM_FORCEINLINE void CheckConsistencyAssumingNoRaces()
 
 struct FRegisterOpenFunction
 {
-    FRegisterOpenFunction(void* OriginalFunction, void* NewFunction)
-    {
-        RegisterOpenFunction(OriginalFunction, NewFunction);
-    }
+	FRegisterOpenFunction(void* OriginalFunction, void* NewFunction)
+	{
+		RegisterOpenFunction(OriginalFunction, NewFunction);
+	}
 };
 
+// A collection of power-user functions that are reserved for use by the AutoRTFM runtime only.
+namespace ForTheRuntime
+{
+	[[deprecated("This macro is deprecated. Use UE_AUTORTFM_ONABORT instead!")]] UE_AUTORTFM_FORCEINLINE void DeprecatedUseOnAbortMacro() {}
+	[[deprecated("This macro is deprecated. Use UE_AUTORTFM_ONCOMMIT instead!")]] UE_AUTORTFM_FORCEINLINE void DeprecatedUseOnCommitMacro() {}
+} // namespace ForTheRuntime
+
+} // namespace AutoRTFM
+
 // Macro-based variants so we completely compile away when not in use, even in debug builds
 #if UE_AUTORTFM
 
@@ -612,17 +648,39 @@ struct FRegisterOpenFunction
 #define UE_AUTORTFM_END_DISABLE_WARNINGS
 #endif
 
-#define UE_AUTORTFM_OPEN(...) AutoRTFM::Open([&]() { __VA_ARGS__ })
-#define UE_AUTORTFM_OPENABORT(...) UE_AUTORTFM_BEGIN_DISABLE_WARNINGS AutoRTFM::OpenAbort([=]() { __VA_ARGS__ }) UE_AUTORTFM_END_DISABLE_WARNINGS
-#define UE_AUTORTFM_OPENCOMMIT(...) UE_AUTORTFM_BEGIN_DISABLE_WARNINGS AutoRTFM::OpenCommit([=]() { __VA_ARGS__ }) UE_AUTORTFM_END_DISABLE_WARNINGS
-#define UE_AUTORTFM_TRANSACT(...) AutoRTFM::Transact([&]() { __VA_ARGS__ })
+#define UE_AUTORTFM_OPEN_IMPL(...) AutoRTFM::Open([&]() { __VA_ARGS__ })
+#define UE_AUTORTFM_ONABORT_IMPL(...) UE_AUTORTFM_BEGIN_DISABLE_WARNINGS AutoRTFM::OnAbort([=]() { __VA_ARGS__ }) UE_AUTORTFM_END_DISABLE_WARNINGS
+#define UE_AUTORTFM_ONCOMMIT_IMPL(...) UE_AUTORTFM_BEGIN_DISABLE_WARNINGS AutoRTFM::OnCommit([=]() { __VA_ARGS__ }) UE_AUTORTFM_END_DISABLE_WARNINGS
+#define UE_AUTORTFM_TRANSACT_IMPL(...) AutoRTFM::Transact([&]() { __VA_ARGS__ })
 #else
-#define UE_AUTORTFM_OPEN(...) do { __VA_ARGS__ } while (false)
-#define UE_AUTORTFM_OPENABORT(...) do { /* do nothing */ } while (false)
-#define UE_AUTORTFM_OPENCOMMIT(...) do { __VA_ARGS__ } while (false)
-#define UE_AUTORTFM_TRANSACT(...) do { __VA_ARGS__ } while (false)
+#define UE_AUTORTFM_OPEN_IMPL(...) do { __VA_ARGS__ } while (false)
+#define UE_AUTORTFM_ONABORT_IMPL(...) do { /* do nothing */ } while (false)
+#define UE_AUTORTFM_ONCOMMIT_IMPL(...) do { __VA_ARGS__ } while (false)
+#define UE_AUTORTFM_TRANSACT_IMPL(...) do { __VA_ARGS__ } while (false)
 #endif
 
+// Runs a block of code in the open, non-transactionally. Anything performed in the open will not be undone if a transaction fails.
+#define UE_AUTORTFM_OPEN(...) UE_AUTORTFM_OPEN_IMPL(__VA_ARGS__)
+
+// Runs a block of code if a transaction aborts.
+// In non-transactional code paths the block of code will not be executed at all.
+// This captures any used variables from the parent function by-value.
+#define UE_AUTORTFM_ONABORT(...) UE_AUTORTFM_ONABORT_IMPL(__VA_ARGS__)
+
+// Runs a block of code if a transaction commits successfully.
+// In non-transactional code paths the block of code will be executed immediately.
+// This captures any used variables from the parent function by-value.
+#define UE_AUTORTFM_ONCOMMIT(...) UE_AUTORTFM_ONCOMMIT_IMPL(__VA_ARGS__)
+
+// Runs a block of code in the closed, transactionally, within a new transaction.
+#define UE_AUTORTFM_TRANSACT(...) UE_AUTORTFM_TRANSACT_IMPL(__VA_ARGS__)
+
+// Deprecated. Use UE_AUTORTFM_ONABORT instead.
+#define UE_AUTORTFM_OPENABORT(...) UE_AUTORTFM_ONABORT(AutoRTFM::ForTheRuntime::DeprecatedUseOnAbortMacro(); __VA_ARGS__)
+
+// Deprecated. Use UE_AUTORTFM_ONCOMMIT instead.
+#define UE_AUTORTFM_OPENCOMMIT(...) UE_AUTORTFM_ONCOMMIT(AutoRTFM::ForTheRuntime::DeprecatedUseOnCommitMacro(); __VA_ARGS__)
+
 #define UE_AUTORTFM_CONCAT_IMPL(A, B) A ## B
 #define UE_AUTORTFM_CONCAT(A, B) UE_AUTORTFM_CONCAT_IMPL(A, B)
 
@@ -636,5 +694,4 @@ struct FRegisterOpenFunction
 #define UE_AUTORTFM_REGISTER_SELF_FUNCTION(OriginalFunction)
 #endif
 
-} // namespace AutoRTFM
 #endif // __cplusplus
diff --git a/Engine/Source/Runtime/Core/Public/HAL/FMemory.inl b/Engine/Source/Runtime/Core/Public/HAL/FMemory.inl
--- a/Engine/Source/Runtime/Core/Public/HAL/FMemory.inl
+++ b/Engine/Source/Runtime/Core/Public/HAL/FMemory.inl
@@ -43,7 +43,7 @@ FMEMORY_INLINE_FUNCTION_DECORATOR void* FMemory::Malloc(SIZE_T Count, uint32 Ali
 	// AutoRTFM: This is a no-op for non-transactional code.
 	// For transactional code, this defers a call to Free if the transaction aborts,
 	// so that rolling back this allocation will end up freeing the memory.
-	AutoRTFM::OpenAbort([Ptr]
+	AutoRTFM::OnAbort([Ptr]
 	{
 		// Disable the code analysis warning that complains that Free is being passed
 		// a pointer that may be null. Free explicitly handles this case already.
@@ -131,7 +131,7 @@ FMEMORY_INLINE_FUNCTION_DECORATOR void FMemory::Free(void* Original)
 	// AutoRTFM: For transactional code, in order to support the transaction 
 	// aborting and needing to 'roll back' the Free, we defer the actual
 	// free until commit time.
-	UE_AUTORTFM_OPENCOMMIT(
+	UE_AUTORTFM_ONCOMMIT(
 	{
 		// optional tracking of every allocation
 		LLM_IF_ENABLED(FLowLevelMemTracker::Get().OnLowLevelFree(ELLMTracker::Default, Original, ELLMAllocType::FMalloc));
diff --git a/Engine/Source/Runtime/Core/Public/Templates/SharedPointerInternals.h b/Engine/Source/Runtime/Core/Public/Templates/SharedPointerInternals.h
--- a/Engine/Source/Runtime/Core/Public/Templates/SharedPointerInternals.h
+++ b/Engine/Source/Runtime/Core/Public/Templates/SharedPointerInternals.h
@@ -129,7 +129,7 @@ namespace SharedPointerInternals
 #endif
 
 				// If the transaction would abort, we need to undo adding the shared reference.
-				UE_AUTORTFM_OPENABORT(
+				UE_AUTORTFM_ONABORT(
 				{
 					ReleaseSharedReference();
 				});
@@ -191,7 +191,7 @@ namespace SharedPointerInternals
 				// If we succeedd in taking a shared reference count, we need to undo that on an abort.
 				if (bSucceeded)
 				{
-					UE_AUTORTFM_OPENABORT(
+					UE_AUTORTFM_ONABORT(
 					{
 						ReleaseSharedReference();
 					});
@@ -217,7 +217,7 @@ namespace SharedPointerInternals
 		{
 			if constexpr (Mode == ESPMode::ThreadSafe)
 			{
-				UE_AUTORTFM_OPENCOMMIT(
+				UE_AUTORTFM_ONCOMMIT(
 				{
 					// std::memory_order_acq_rel is used here so that, if we do end up executing the destructor, it's not possible
 					// for side effects from executing the destructor end up being visible before we've determined that the shared
diff --git a/Engine/Source/Runtime/CoreUObject/Private/Misc/GCObjectReferencer.cpp b/Engine/Source/Runtime/CoreUObject/Private/Misc/GCObjectReferencer.cpp
--- a/Engine/Source/Runtime/CoreUObject/Private/Misc/GCObjectReferencer.cpp
+++ b/Engine/Source/Runtime/CoreUObject/Private/Misc/GCObjectReferencer.cpp
@@ -263,10 +263,14 @@ void FGCObject::RegisterGCObject()
 		StaticInit();
 
 		// Add this instance to the referencer's list
-		UE_AUTORTFM_OPEN({
+		UE_AUTORTFM_OPEN(
+		{
 			GGCObjectReferencer->AddObject(this);
 		});
-		UE_AUTORTFM_OPENABORT({
+
+		// But if we abort, we don't want to leave a dangling reference!
+		UE_AUTORTFM_ONABORT(
+		{
 			GGCObjectReferencer->RemoveObject(this);
 		});
 
diff --git a/Engine/Source/Runtime/CoreUObject/Public/UObject/Stack.h b/Engine/Source/Runtime/CoreUObject/Public/UObject/Stack.h
--- a/Engine/Source/Runtime/CoreUObject/Public/UObject/Stack.h
+++ b/Engine/Source/Runtime/CoreUObject/Public/UObject/Stack.h
@@ -43,12 +43,12 @@ FORCEINLINE void* UeVstackAllocHelper(FVirtualStackAllocator* Allocator, size_t
 		// for these 'stack' allocations we call Free on both commit and abort, since
 		// the call frame that these would have otherwise existed in will always be
 		// returned from by the time we commit or abort
-		AutoRTFM::OpenCommit([Result]()
+		AutoRTFM::OnCommit([Result]
 		{
 			FMemory::Free(Result);
 		});
 
-		AutoRTFM::OpenAbort([Result]()
+		AutoRTFM::OnAbort([Result]
 		{
 			FMemory::Free(Result);
 		});
diff --git a/Engine/Source/Runtime/Engine/Private/DataReplication.cpp b/Engine/Source/Runtime/Engine/Private/DataReplication.cpp
--- a/Engine/Source/Runtime/Engine/Private/DataReplication.cpp
+++ b/Engine/Source/Runtime/Engine/Private/DataReplication.cpp
@@ -399,12 +399,14 @@ void FObjectReplicator::InitRecentProperties(uint8* Source)
 	// If acting as a server and are IsInternalAck, that means we're recording.
 	// In that case, we don't need to create any receiving state, as no one will be sending data to us.
 	ECreateRepStateFlags Flags = (Connection->IsInternalAck() && bIsServer) ? ECreateRepStateFlags::SkipCreateReceivingState : ECreateRepStateFlags::None;
-	UE_AUTORTFM_OPEN({
+	UE_AUTORTFM_OPEN(
+	{
 		RepState = LocalRepLayout.CreateRepState(Source, RepChangedPropertyTracker, Flags);
 	});
 
 	// RepState not valid at the start of this function so just go back to being a nullptr, and let the memory be cleaned up 
-	UE_AUTORTFM_OPENABORT({
+	UE_AUTORTFM_ONABORT(
+	{
 		RepState = nullptr;
 	});
 
diff --git a/Engine/Source/Runtime/Engine/Private/PrimitiveComponentPhysics.cpp b/Engine/Source/Runtime/Engine/Private/PrimitiveComponentPhysics.cpp
--- a/Engine/Source/Runtime/Engine/Private/PrimitiveComponentPhysics.cpp
+++ b/Engine/Source/Runtime/Engine/Private/PrimitiveComponentPhysics.cpp
@@ -1081,7 +1081,8 @@ void UPrimitiveComponent::SetCollisionEnabled(ECollisionEnabled::Type NewType)
 		});
 
 		// If we fail set the CollisionEnabled back to the CurrentType
-		UE_AUTORTFM_OPENABORT({
+		UE_AUTORTFM_ONABORT(
+		{
 			BodyInstance.SetCollisionEnabled(CurrentType);
 		});
 
