Fix the build fail caused by an accidental circular dependency I introduced.

diff --git a/Engine/Source/Runtime/Core/Public/AutoRTFM/AutoRTFM.h b/Engine/Source/Runtime/Core/Public/AutoRTFM/AutoRTFM.h
--- a/Engine/Source/Runtime/Core/Public/AutoRTFM/AutoRTFM.h
+++ b/Engine/Source/Runtime/Core/Public/AutoRTFM/AutoRTFM.h
@@ -31,14 +31,9 @@
 
 #if defined(UE_AUTORTFM_STANDALONE)
 #define UE_AUTORTFM_API
-
-#define UE_AUTORTFM_MOVETEMP(x) std::move(x)
 #else
 #include <HAL/Platform.h>
 #define UE_AUTORTFM_API CORE_API
-
-#include <Templates/UnrealTemplate.h>
-#define UE_AUTORTFM_MOVETEMP(x) MoveTemp(x)
 #endif
 
 #if UE_AUTORTFM
@@ -596,13 +591,13 @@ UE_AUTORTFM_FORCEINLINE void OnAbort(const TFunctor& Work) { }
 [[deprecated("Use OnCommit instead.")]]
 UE_AUTORTFM_FORCEINLINE void OpenCommit(TFunction<void()>&& Work)
 {
-	OnCommit(UE_AUTORTFM_MOVETEMP(Work));
+	OnCommit(std::move(Work));
 }
 
 [[deprecated("Use OnAbort instead.")]]
 UE_AUTORTFM_FORCEINLINE void OpenAbort(TFunction<void()>&& Work)
 {
-	OnAbort(UE_AUTORTFM_MOVETEMP(Work));
+	OnAbort(std::move(Work));
 }
 
 UE_AUTORTFM_FORCEINLINE void* DidAllocate(void* Ptr, size_t Size)
