Added FunctionFwd.h.

diff --git a/Engine/Source/Runtime/Core/Public/Templates/Function.h b/Engine/Source/Runtime/Core/Public/Templates/Function.h
--- a/Engine/Source/Runtime/Core/Public/Templates/Function.h
+++ b/Engine/Source/Runtime/Core/Public/Templates/Function.h
@@ -6,6 +6,7 @@
 #include "Misc/AssertionMacros.h"
 #include "Misc/IntrusiveUnsetOptionalState.h"
 #include "HAL/UnrealMemory.h"
+#include "Templates/FunctionFwd.h"
 #include "Templates/UnrealTypeTraits.h"
 #include "Templates/Invoke.h"
 #include "Templates/UnrealTemplate.h"
@@ -33,60 +34,6 @@
 	#define TFUNCTION_INLINE_ALIGNMENT    16
 #endif
 
-/**
- * TFunction<FuncType>
- *
- * See the class definition for intended usage.
- */
-template <typename FuncType>
-class TFunction;
-
-/**
- * TUniqueFunction<FuncType>
- *
- * See the class definition for intended usage.
- */
-template <typename FuncType>
-class TUniqueFunction;
-
-/**
- * TFunctionRef<FuncType>
- *
- * See the class definition for intended usage.
- */
-template <typename FuncType>
-class TFunctionRef;
-
-/**
- * Traits class which checks if T is a TFunction<> type.
- */
-template <typename T> struct TIsTFunction               { enum { Value = false }; };
-template <typename T> struct TIsTFunction<TFunction<T>> { enum { Value = true  }; };
-
-template <typename T> struct TIsTFunction<const          T> { enum { Value = TIsTFunction<T>::Value }; };
-template <typename T> struct TIsTFunction<      volatile T> { enum { Value = TIsTFunction<T>::Value }; };
-template <typename T> struct TIsTFunction<const volatile T> { enum { Value = TIsTFunction<T>::Value }; };
-
-/**
- * Traits class which checks if T is a TFunction<> type.
- */
-template <typename T> struct TIsTUniqueFunction                     { enum { Value = false }; };
-template <typename T> struct TIsTUniqueFunction<TUniqueFunction<T>> { enum { Value = true  }; };
-
-template <typename T> struct TIsTUniqueFunction<const          T> { enum { Value = TIsTUniqueFunction<T>::Value }; };
-template <typename T> struct TIsTUniqueFunction<      volatile T> { enum { Value = TIsTUniqueFunction<T>::Value }; };
-template <typename T> struct TIsTUniqueFunction<const volatile T> { enum { Value = TIsTUniqueFunction<T>::Value }; };
-
-/**
- * Traits class which checks if T is a TFunctionRef<> type.
- */
-template <typename T> struct TIsTFunctionRef                  { enum { Value = false }; };
-template <typename T> struct TIsTFunctionRef<TFunctionRef<T>> { enum { Value = true  }; };
-
-template <typename T> struct TIsTFunctionRef<const          T> { enum { Value = TIsTFunctionRef<T>::Value }; };
-template <typename T> struct TIsTFunctionRef<      volatile T> { enum { Value = TIsTFunctionRef<T>::Value }; };
-template <typename T> struct TIsTFunctionRef<const volatile T> { enum { Value = TIsTFunctionRef<T>::Value }; };
-
 /**
  * Private implementation details of TFunction and TFunctionRef.
  */
diff --git a/Engine/Source/Runtime/Core/Public/Templates/FunctionFwd.h b/Engine/Source/Runtime/Core/Public/Templates/FunctionFwd.h
new file mode 100644
--- /dev/null
+++ b/Engine/Source/Runtime/Core/Public/Templates/FunctionFwd.h
@@ -0,0 +1,57 @@
+// Copyright Epic Games, Inc. All Rights Reserved.
+
+#pragma once
+
+/**
+ * TFunction<FuncType>
+ *
+ * See the class definition for intended usage.
+ */
+template <typename FuncType>
+class TFunction;
+
+/**
+ * TUniqueFunction<FuncType>
+ *
+ * See the class definition for intended usage.
+ */
+template <typename FuncType>
+class TUniqueFunction;
+
+/**
+ * TFunctionRef<FuncType>
+ *
+ * See the class definition for intended usage.
+ */
+template <typename FuncType>
+class TFunctionRef;
+
+/**
+ * Traits class which checks if T is a TFunction<> type.
+ */
+template <typename T> struct TIsTFunction               { enum { Value = false }; };
+template <typename T> struct TIsTFunction<TFunction<T>> { enum { Value = true  }; };
+
+template <typename T> struct TIsTFunction<const          T> { enum { Value = TIsTFunction<T>::Value }; };
+template <typename T> struct TIsTFunction<      volatile T> { enum { Value = TIsTFunction<T>::Value }; };
+template <typename T> struct TIsTFunction<const volatile T> { enum { Value = TIsTFunction<T>::Value }; };
+
+/**
+ * Traits class which checks if T is a TFunction<> type.
+ */
+template <typename T> struct TIsTUniqueFunction                     { enum { Value = false }; };
+template <typename T> struct TIsTUniqueFunction<TUniqueFunction<T>> { enum { Value = true  }; };
+
+template <typename T> struct TIsTUniqueFunction<const          T> { enum { Value = TIsTUniqueFunction<T>::Value }; };
+template <typename T> struct TIsTUniqueFunction<      volatile T> { enum { Value = TIsTUniqueFunction<T>::Value }; };
+template <typename T> struct TIsTUniqueFunction<const volatile T> { enum { Value = TIsTUniqueFunction<T>::Value }; };
+
+/**
+ * Traits class which checks if T is a TFunctionRef<> type.
+ */
+template <typename T> struct TIsTFunctionRef                  { enum { Value = false }; };
+template <typename T> struct TIsTFunctionRef<TFunctionRef<T>> { enum { Value = true  }; };
+
+template <typename T> struct TIsTFunctionRef<const          T> { enum { Value = TIsTFunctionRef<T>::Value }; };
+template <typename T> struct TIsTFunctionRef<      volatile T> { enum { Value = TIsTFunctionRef<T>::Value }; };
+template <typename T> struct TIsTFunctionRef<const volatile T> { enum { Value = TIsTFunctionRef<T>::Value }; };
