Second attempt at a fix - don't use std::move or MoveTemp in the header at all.

diff --git a/Engine/Source/Runtime/Core/Private/AutoRTFM/API.cpp b/Engine/Source/Runtime/Core/Private/AutoRTFM/API.cpp
--- a/Engine/Source/Runtime/Core/Private/AutoRTFM/API.cpp
+++ b/Engine/Source/Runtime/Core/Private/AutoRTFM/API.cpp
@@ -243,6 +243,16 @@ UE_AUTORTFM_AUTORTFM("RTFM_OnAbort") void OnAbort(TFunction<void()> && Work)
 {
 }
 
+void OpenCommit(TFunction<void()>&& Work)
+{
+	OnCommit(MoveTemp(Work));
+}
+
+void OpenAbort(TFunction<void()>&& Work)
+{
+	OnAbort(MoveTemp(Work));
+}
+
 extern "C" UE_AUTORTFM_AUTORTFM("RTFM_autortfm_on_commit") void autortfm_on_commit(void (*Work)(void* Arg), void* Arg)
 {
     Work(Arg);
diff --git a/Engine/Source/Runtime/Core/Public/AutoRTFM/AutoRTFM.h b/Engine/Source/Runtime/Core/Public/AutoRTFM/AutoRTFM.h
--- a/Engine/Source/Runtime/Core/Public/AutoRTFM/AutoRTFM.h
+++ b/Engine/Source/Runtime/Core/Public/AutoRTFM/AutoRTFM.h
@@ -588,17 +588,20 @@ template<typename TFunctor>
 UE_AUTORTFM_FORCEINLINE void OnAbort(const TFunctor& Work) { }
 #endif
 
+#if UE_AUTORTFM
 [[deprecated("Use OnCommit instead.")]]
-UE_AUTORTFM_FORCEINLINE void OpenCommit(TFunction<void()>&& Work)
-{
-	OnCommit(std::move(Work));
-}
+UE_AUTORTFM_API void OpenCommit(TFunction<void()>&& Work);
 
 [[deprecated("Use OnAbort instead.")]]
-UE_AUTORTFM_FORCEINLINE void OpenAbort(TFunction<void()>&& Work)
-{
-	OnAbort(std::move(Work));
-}
+UE_AUTORTFM_API void OpenAbort(TFunction<void()>&& Work);
+#else
+template<typename TFunctor>
+UE_AUTORTFM_FORCEINLINE void OpenCommit(const TFunctor& Work) { Work(); }
+template<typename TFunctor>
+UE_AUTORTFM_FORCEINLINE void OpenAbort(const TFunctor& Work) { }
+#endif
+
+
 
 UE_AUTORTFM_FORCEINLINE void* DidAllocate(void* Ptr, size_t Size)
 {
