Build fix for CompositeBuffer tests.

Building foundation tests project errored with FFunctionRef being an undefied type in the CompsiteBuffer tests.

diff --git a/Engine/Source/Runtime/Core/Public/Memory/CompositeBuffer.h b/Engine/Source/Runtime/Core/Public/Memory/CompositeBuffer.h
--- a/Engine/Source/Runtime/Core/Public/Memory/CompositeBuffer.h
+++ b/Engine/Source/Runtime/Core/Public/Memory/CompositeBuffer.h
@@ -11,11 +11,10 @@
 #include "Memory/MemoryView.h"
 #include "Memory/SharedBuffer.h"
 #include "Templates/UnrealTemplate.h"
+#include "Templates/FunctionFwd.h"
 
 #include <type_traits>
 
-template <typename FuncType> class TFunctionRef;
-
 /**
  * FCompositeBuffer is a non-contiguous buffer composed of zero or more immutable shared buffers.
  *
diff --git a/Engine/Source/Runtime/Core/Tests/Memory/CompositeBufferTest.cpp b/Engine/Source/Runtime/Core/Tests/Memory/CompositeBufferTest.cpp
--- a/Engine/Source/Runtime/Core/Tests/Memory/CompositeBufferTest.cpp
+++ b/Engine/Source/Runtime/Core/Tests/Memory/CompositeBufferTest.cpp
@@ -3,6 +3,7 @@
 #if WITH_TESTS
 
 #include "Memory/CompositeBuffer.h"
+#include "Templates/Function.h"
 
 #include "Tests/TestHarnessAdapter.h"
 #include <type_traits>
