Update TFuture/TSharedFuture to work with non-copyable types.
- `ResultType TFuture::Consume()`
    + Moves and returns the result to the caller. Invalidates the future.
    + Equivalent to `std::future::get()`, but with a different name to avoid breaking existing usage of Get, and to make it clear that the result is being moved out.
- `const ResultType& TFuture::Get() const`
    + Returns the result by const reference, instead of trying to copy it.
    + Should provide perf benefits due to a copy of the result no longer being returned.
    + This method is still not equivalent to std::future::get. Changing it to return a const reference allows a non-copyable object to be peeked at.
    + Uses `UE_LIFETIMEBOUND` to track the use of the returned reference outlasting the future itself.
- `ResultType& TFuture::GetMutable()`
    + Returns the result by reference, instead of trying to copy it.
    + Has a different name to `Get` to avoid accidental mutation.
    + Not equivalent to any other method in std::future. Allows the result to be peeked at & mutated.
    + Uses `UE_LIFETIMEBOUND` to track the use of the returned reference outlasting the future itself.
- `InternalResultType& TFutureState::GetResult()`
    + A non-const version of the existing `const InternalResultType& TFutureState::GetResult() const` function, allowing the caller to move the result.
    + Uses `UE_LIFETIMEBOUND` to track the use of the returned reference outlasting the state itself.
- `const ResultType& TSharedFuture::Get() const`
    + Return the result by const ref instead of by value. This matches `std::shared_future::get()`

Some other misc improvements/fixes to TFuture/TSharedFuture
- Changed a few default constructors to be explicitly defaulted.
- Changed hidden copy constructor/assignment operators to be explicitly deleted. Should provide better errors.
- Improved comments around copy/move constructor/assignment operators for better documentation generation.
- Corrected some comments (typo, some incorrect variable names/descriptions for params)

Some fixups in LobbiesCommonTests, ConcertDataStoreTests and InterchangeSceneVariantSetsFactory.

diff --git a/Engine/Plugins/Developer/Concert/ConcertSync/ConcertSyncTest/Source/ConcertSyncTest/Private/ConcertDataStoreTests.cpp b/Engine/Plugins/Developer/Concert/ConcertSync/ConcertSyncTest/Source/ConcertSyncTest/Private/ConcertDataStoreTests.cpp
--- a/Engine/Plugins/Developer/Concert/ConcertSync/ConcertSyncTest/Source/ConcertSyncTest/Private/ConcertDataStoreTests.cpp
+++ b/Engine/Plugins/Developer/Concert/ConcertSync/ConcertSyncTest/Source/ConcertSyncTest/Private/ConcertDataStoreTests.cpp
@@ -413,7 +413,7 @@ public:
 
 	/** Ensures a key/value pair was added. */
 	template<typename T>
-	void EnsureValueAdded(TConcertDataStoreResult<T>&& Result, const T& ExpectedValue)
+	void EnsureValueAdded(const TConcertDataStoreResult<T>& Result, const T& ExpectedValue)
 	{
 		TestTrueExpr(Result.GetCode() == EConcertDataStoreResultCode::Added);
 		TestTrueExpr(Result.GetValue() == ExpectedValue);
@@ -422,7 +422,7 @@ public:
 
 	/** Ensures the expected value was fetched. */
 	template<typename T>
-	void EnsureValueFetched(TConcertDataStoreResult<T>&& Result, const T& ExpectedValue)
+	void EnsureValueFetched(const TConcertDataStoreResult<T>& Result, const T& ExpectedValue)
 	{
 		TestTrueExpr(Result.GetCode() == EConcertDataStoreResultCode::Fetched);
 		TestTrueExpr(Result.GetValue() == ExpectedValue);
@@ -431,7 +431,7 @@ public:
 
 	/** Ensures the desired value was exchanged. */
 	template<typename T>
-	void EnsureValueExchanged(TConcertDataStoreResult<T>&& Result, const T& ExpectedValue)
+	void EnsureValueExchanged(const TConcertDataStoreResult<T>& Result, const T& ExpectedValue)
 	{
 		TestTrueExpr(Result.GetCode() == EConcertDataStoreResultCode::Exchanged);
 		TestTrueExpr(Result.GetValue() == ExpectedValue);
@@ -439,7 +439,7 @@ public:
 	}
 
 	template<typename T>
-	void EnsureTypeMismatch(TConcertDataStoreResult<T>&& Result)
+	void EnsureTypeMismatch(const TConcertDataStoreResult<T>& Result)
 	{
 		TestTrueExpr(Result.GetCode() == EConcertDataStoreResultCode::TypeMismatch);
 		TestTrueExpr(!Result.IsValid()); // Server doesn't send back any value on error.
@@ -447,7 +447,7 @@ public:
 	}
 
 	template<typename T>
-	void EnsureNotFound(TConcertDataStoreResult<T>&& Result)
+	void EnsureNotFound(const TConcertDataStoreResult<T>& Result)
 	{
 		TestTrueExpr(Result.GetCode() == EConcertDataStoreResultCode::NotFound);
 		TestTrueExpr(!Result.IsValid()); // Server doesn't send back any value on error.
diff --git a/Engine/Plugins/Interchange/Runtime/Source/Import/Private/Scene/InterchangeSceneVariantSetsFactory.cpp b/Engine/Plugins/Interchange/Runtime/Source/Import/Private/Scene/InterchangeSceneVariantSetsFactory.cpp
--- a/Engine/Plugins/Interchange/Runtime/Source/Import/Private/Scene/InterchangeSceneVariantSetsFactory.cpp
+++ b/Engine/Plugins/Interchange/Runtime/Source/Import/Private/Scene/InterchangeSceneVariantSetsFactory.cpp
@@ -98,7 +98,7 @@ namespace UE::Interchange::Private
 		FString DisplayText;
 		VariantSetNode.GetCustomDisplayText(DisplayText);
 
-		const TOptional<UE::Interchange::FVariantSetPayloadData>& PayloadData = PayloadInterface.GetVariantSetPayloadData(PayloadKey).Get();
+		const TOptional<UE::Interchange::FVariantSetPayloadData> PayloadData = PayloadInterface.GetVariantSetPayloadData(PayloadKey).Consume();
 		if (!PayloadData.IsSet())
 		{
 			UE_LOG(LogInterchangeImport, Warning, TEXT("No payload for variant set %s"), *DisplayText);
diff --git a/Engine/Plugins/Online/OnlineServices/Source/OnlineServicesCommon/Private/Online/LobbiesCommonTests.cpp b/Engine/Plugins/Online/OnlineServices/Source/OnlineServicesCommon/Private/Online/LobbiesCommonTests.cpp
--- a/Engine/Plugins/Online/OnlineServices/Source/OnlineServicesCommon/Private/Online/LobbiesCommonTests.cpp
+++ b/Engine/Plugins/Online/OnlineServices/Source/OnlineServicesCommon/Private/Online/LobbiesCommonTests.cpp
@@ -1384,15 +1384,16 @@ TFunction<TFuture<void>(TOnlineAsyncOp<OpType>&)> CaptureStepResult(TOnlineAsync
 		Func(MoveTempIfPossible(InParams))
 		.Then([Promise = MoveTemp(Promise), Op = InAsyncOp.AsShared(), ResultKey](TFuture<TOnlineResult<SecondaryOpType>>&& Future) mutable
 		{
-			if (Future.Get().IsError())
+			TOnlineResult<SecondaryOpType> Result = Future.Consume();
+			if (Result.IsError())
 			{
-				UE_LOG(LogOnlineServices, Log, TEXT("[LobbiesFunctionalTest::CaptureStepResult] Failed secondary operation %s. Error: %s"), SecondaryOpType::Name, *Future.Get().GetErrorValue().GetLogString());
-				Op->SetError(Errors::RequestFailure(MoveTempIfPossible(Future.Get().GetErrorValue())));
+				UE_LOG(LogOnlineServices, Log, TEXT("[LobbiesFunctionalTest::CaptureStepResult] Failed secondary operation %s. Error: %s"), SecondaryOpType::Name, *Result.GetErrorValue().GetLogString());
+				Op->SetError(Errors::RequestFailure(MoveTempIfPossible(Result.GetErrorValue())));
 			}
 			else
 			{
 				UE_LOG(LogOnlineServices, Log, TEXT("[LobbiesFunctionalTest::CaptureStepResult] Captured secondary operation %s as %s"), SecondaryOpType::Name, *ResultKey);
-				Op->Data.Set(ResultKey, MoveTempIfPossible(Future.Get().GetOkValue()));
+				Op->Data.Set(ResultKey, MoveTempIfPossible(Result.GetOkValue()));
 			}
 			Promise.EmplaceValue();
 		});
@@ -1747,7 +1748,7 @@ TFuture<TOnlineResult<FFunctionalTestLogoutAllUsers>> FunctionalTestLogoutAllUse
 	.Then([Promise = MoveTemp(Promise)](TFuture<TArray<TOnlineResult<FFunctionalTestLogoutUser>>>&& Results)
 	{
 		bool HasAnyError = false;
-		for (TOnlineResult<FFunctionalTestLogoutUser>& Result : Results.Get())
+		for (const TOnlineResult<FFunctionalTestLogoutUser>& Result : Results.Get())
 		{
 			HasAnyError |= Result.IsError();
 		}
diff --git a/Engine/Source/Runtime/Core/Public/Async/Future.h b/Engine/Source/Runtime/Core/Public/Async/Future.h
--- a/Engine/Source/Runtime/Core/Public/Async/Future.h
+++ b/Engine/Source/Runtime/Core/Public/Async/Future.h
@@ -169,7 +169,7 @@ public:
 	 * @return The result value.
 	 * @see EmplaceResult
 	 */
-	const InternalResultType& GetResult() const
+	const InternalResultType& GetResult() const UE_LIFETIMEBOUND
 	{
 		while (!IsComplete())
 		{
@@ -179,10 +179,21 @@ public:
 		return *Result.GetTypedPtr();
 	}
 
+	/**
+	 * Gets the result (will block the calling thread until the result is available).
+	 *
+	 * @return The result value.
+	 * @see EmplaceResult
+	 */
+	InternalResultType& GetResult() UE_LIFETIMEBOUND
+	{
+		return const_cast<InternalResultType&>(AsConst(*this).GetResult());
+	}
+
 	/**
 	 * Sets the result and notifies any waiting threads.
 	 *
-	 * @param InResult The result to set.
+	 * @param Args The arguments to forward to the constructor of the result.
 	 * @see GetResult
 	 */
 	template<typename... ArgTypes>
@@ -277,7 +288,7 @@ protected:
 	typedef TSharedPtr<TFutureState<InternalResultType>, ESPMode::ThreadSafe> StateType;
 
 	/** Default constructor. */
-	TFutureBase() { }
+	TFutureBase() = default;
 
 	/**
 	 * Creates and initializes a new instance.
@@ -288,28 +299,20 @@ protected:
 		: State(InState)
 	{ }
 
-	/**
-	 * Protected move construction
-	 */
-	TFutureBase(TFutureBase&&) = default;
-
-	/**
-	 * Protected move assignment
-	 */
-	TFutureBase& operator=(TFutureBase&&) = default;
-
-	/**
-	 * Protected copy construction
-	 */
+	/** Protected copy constructor. */
 	TFutureBase(const TFutureBase&) = default;
 
-	/**
-	 * Protected copy assignment
-	 */
+	/** Protected copy assignment operator. */
 	TFutureBase& operator=(const TFutureBase&) = default;
 
+	/** Protected move constructor. */
+	TFutureBase(TFutureBase&&) = default;
+
+	/** Protected move assignment operator. */
+	TFutureBase& operator=(TFutureBase&&) = default;
+
 	/** Protected destructor. */
-	~TFutureBase() { }
+	~TFutureBase() = default;
 
 protected:
 
@@ -350,7 +353,7 @@ protected:
 
 	/**
 	 * Reset the future.
-	 *	Reseting a future removes any continuation from its shared state and invalidates it.
+	 *	Resetting a future removes any continuation from its shared state and invalidates it.
 	 *	Useful for discarding yet to be completed future cleanly.
 	 */
 	void Reset()
@@ -383,7 +386,7 @@ class TFuture
 public:
 
 	/** Default constructor. */
-	TFuture() { }
+	TFuture() = default;
 
 	/**
 	 * Creates and initializes a new instance.
@@ -394,31 +397,57 @@ public:
 		: BaseType(InState)
 	{ }
 
-	/**
-	 * Move constructor.
-	 */
+	/** Deleted copy constructor (futures cannot be copied). */
+	TFuture(const TFuture&) = delete;
+
+	/** Deleted copy assignment operator (futures cannot be copied). */
+	TFuture& operator=(const TFuture&) = delete;
+
+	/** Move constructor. */
 	TFuture(TFuture&&) = default;
 
-	/**
-	 * Move assignment operator.
-	 */
-	TFuture& operator=(TFuture&& Other) = default;
+	/** Move assignment operator. */
+	TFuture& operator=(TFuture&&) = default;
 
 	/** Destructor. */
-	~TFuture() { }
+	~TFuture() = default;
 
 public:
 
 	/**
-	 * Gets the future's result.
+	 * Gets the future's result by const reference.
+	 *
+	 * @return The result.
+	 * @note Not equivalent to std::future::get(). The future remains valid.
+	 */
+	const ResultType& Get() const UE_LIFETIMEBOUND
+	{
+		return this->GetState()->GetResult();
+	}
+
+	/**
+	 * Gets the future's result by non-const reference.
 	 *
 	 * @return The result.
+	 * @note Not equivalent to std::future::get(). The future remains valid.
 	 */
-	ResultType Get() const
+	ResultType& GetMutable() UE_LIFETIMEBOUND
 	{
 		return this->GetState()->GetResult();
 	}
 
+	/**
+	 * Consumes the future's result and invalidates the future.
+	 *
+	 * @return The result.
+	 * @note Equivalent to std::future::get(). Invalidates the future.
+	 */
+	ResultType Consume()
+	{
+		TFuture<ResultType> Local(MoveTemp(*this));
+		return MoveTemp(Local.GetState()->GetResult());
+	}
+
 	/**
 	 * Moves this future's state into a shared future.
 	 *
@@ -446,14 +475,6 @@ public:
 	 * @see TFutureBase
 	 */
 	using BaseType::Reset;
-
-private:
-
-	/** Hidden copy constructor (futures cannot be copied). */
-	TFuture(const TFuture&);
-
-	/** Hidden copy assignment (futures cannot be copied). */
-	TFuture& operator=(const TFuture&);
 };
 
 
@@ -469,7 +490,7 @@ class TFuture<ResultType&>
 public:
 
 	/** Default constructor. */
-	TFuture() { }
+	TFuture() = default;
 
 	/**
 	 * Creates and initializes a new instance.
@@ -480,18 +501,20 @@ public:
 		: BaseType(InState)
 	{ }
 
-	/**
-	 * Move constructor.
-	 */
+	/** Deleted copy constructor (futures cannot be copied). */
+	TFuture(const TFuture&) = delete;
+
+	/** Deleted copy assignment operator (futures cannot be copied). */
+	TFuture& operator=(const TFuture&) = delete;
+
+	/** Move constructor. */
 	TFuture(TFuture&&) = default;
 
-	/**
-	 * Move assignment operator.
-	 */
-	TFuture& operator=(TFuture&& Other) = default;
+	/** Move assignment operator. */
+	TFuture& operator=(TFuture&&) = default;
 
 	/** Destructor. */
-	~TFuture() { }
+	~TFuture() = default;
 
 public:
 
@@ -499,12 +522,36 @@ public:
 	 * Gets the future's result.
 	 *
 	 * @return The result.
+	 * @note Not equivalent to std::future::get(). The future remains valid.
 	 */
 	ResultType& Get() const
 	{
 		return *this->GetState()->GetResult();
 	}
 
+	/**
+	 * Gets the future's result.
+	 *
+	 * @return The result.
+	 * @note Not equivalent to std::future::get(). The future remains valid.
+	 */
+	ResultType& GetMutable()
+	{
+		return this->GetState()->GetResult();
+	}
+
+	/**
+	 * Consumes the future's result and invalidates the future.
+	 *
+	 * @return The result.
+	 * @note Equivalent to std::future::get(). Invalidates the future.
+	 */
+	ResultType& Consume()
+	{
+		TFuture<ResultType&> Local(MoveTemp(*this));
+		return Local.GetState()->GetResult();
+	}
+
 	/**
 	 * Moves this future's state into a shared future.
 	 *
@@ -532,14 +579,6 @@ public:
 	 * @see TFutureBase
 	 */
 	using BaseType::Reset;
-
-private:
-
-	/** Hidden copy constructor (futures cannot be copied). */
-	TFuture(const TFuture&);
-
-	/** Hidden copy assignment (futures cannot be copied). */
-	TFuture& operator=(const TFuture&);
 };
 
 
@@ -555,7 +594,7 @@ class TFuture<void>
 public:
 
 	/** Default constructor. */
-	TFuture() { }
+	TFuture() = default;
 
 	/**
 	 * Creates and initializes a new instance.
@@ -566,18 +605,20 @@ public:
 		: BaseType(InState)
 	{ }
 
-	/**
-	 * Move constructor.
-	 */
+	/** Deleted copy constructor (futures cannot be copied). */
+	TFuture(const TFuture&) = delete;
+
+	/** Deleted copy assignment operator (futures cannot be copied). */
+	TFuture& operator=(const TFuture&) = delete;
+
+	/** Move constructor. */
 	TFuture(TFuture&&) = default;
 
-	/**
-	 * Move assignment operator.
-	 */
-	TFuture& operator=(TFuture&& Other) = default;
+	/** Move assignment operator. */
+	TFuture& operator=(TFuture&&) = default;
 
 	/** Destructor. */
-	~TFuture() { }
+	~TFuture() = default;
 
 public:
 
@@ -585,12 +626,39 @@ public:
 	 * Gets the future's result.
 	 *
 	 * @return The result.
+	 * @note Not equivalent to std::future::get(). The future remains valid.
 	 */
 	void Get() const
 	{
+		// Wait for the dummy result to be available.
 		GetState()->GetResult();
 	}
 
+	/**
+	 * Gets the future's result.
+	 *
+	 * @return The result.
+	 * @note Not equivalent to std::future::get(). The future remains valid.
+	 */
+	void GetMutable()
+	{
+		// Wait for the dummy result to be available.
+		GetState()->GetResult();
+	}
+
+	/**
+	 * Consumes the future's result and invalidates the future.
+	 *
+	 * @return The result.
+	 * @note Equivalent to std::future::get(). Invalidates the future.
+	 */
+	void Consume()
+	{
+		TFuture<void> Local(MoveTemp(*this));
+		// Wait for the dummy result to be available.
+		Local.GetState()->GetResult();
+	}
+
 	/**
 	 * Moves this future's state into a shared future.
 	 *
@@ -615,14 +683,6 @@ public:
 	 * @see TFutureBase
 	 */
 	using BaseType::Reset;
-
-private:
-
-	/** Hidden copy constructor (futures cannot be copied). */
-	TFuture(const TFuture&);
-
-	/** Hidden copy assignment (futures cannot be copied). */
-	TFuture& operator=(const TFuture&);
 };
 
 /* TSharedFuture
@@ -640,7 +700,7 @@ class TSharedFuture
 public:
 
 	/** Default constructor. */
-	TSharedFuture() { }
+	TSharedFuture() = default;
 
 	/**
 	 * Creates and initializes a new instance.
@@ -660,28 +720,20 @@ public:
 		: BaseType(MoveTemp(Future))
 	{ }
 
-	/**
-	 * Copy constructor.
-	 */
+	/** Copy constructor. */
 	TSharedFuture(const TSharedFuture&) = default;
 
-	/**
-	 * Copy assignment operator.
-	 */
-	TSharedFuture& operator=(const TSharedFuture& Other) = default;
+	/** Copy assignment operator. */
+	TSharedFuture& operator=(const TSharedFuture&) = default;
 
-	/**
-	 * Move constructor.
-	 */
+	/** Move constructor. */
 	TSharedFuture(TSharedFuture&&) = default;
 
-	/**
-	 * Move assignment operator.
-	 */
-	TSharedFuture& operator=(TSharedFuture&& Other) = default;
+	/** Move assignment operator. */
+	TSharedFuture& operator=(TSharedFuture&&) = default;
 
 	/** Destructor. */
-	~TSharedFuture() { }
+	~TSharedFuture() = default;
 
 public:
 
@@ -690,7 +742,7 @@ public:
 	 *
 	 * @return The result.
 	 */
-	ResultType Get() const
+	const ResultType& Get() const
 	{
 		return this->GetState()->GetResult();
 	}
@@ -709,7 +761,7 @@ class TSharedFuture<ResultType&>
 public:
 
 	/** Default constructor. */
-	TSharedFuture() { }
+	TSharedFuture() = default;
 
 	/**
 	 * Creates and initializes a new instance.
@@ -729,26 +781,20 @@ public:
 		: BaseType(MoveTemp(Future))
 	{ }
 
-	/**
-	 * Copy constructor.
-	 */
+	/** Copy constructor. */
 	TSharedFuture(const TSharedFuture&) = default;
 
 	/** Copy assignment operator. */
-	TSharedFuture& operator=(const TSharedFuture& Other) = default;
+	TSharedFuture& operator=(const TSharedFuture&) = default;
 
-	/**
-	 * Move constructor.
-	 */
+	/** Move constructor. */
 	TSharedFuture(TSharedFuture&&) = default;
 
-	/**
-	 * Move assignment operator.
-	 */
-	TSharedFuture& operator=(TSharedFuture&& Other) = default;
+	/** Move assignment operator. */
+	TSharedFuture& operator=(TSharedFuture&&) = default;
 
 	/** Destructor. */
-	~TSharedFuture() { }
+	~TSharedFuture() = default;
 
 public:
 
@@ -776,7 +822,7 @@ class TSharedFuture<void>
 public:
 
 	/** Default constructor. */
-	TSharedFuture() { }
+	TSharedFuture() = default;
 
 	/**
 	 * Creates and initializes a new instance from shared state.
@@ -796,28 +842,20 @@ public:
 		: BaseType(MoveTemp(Future))
 	{ }
 
-	/**
-	 * Copy constructor.
-	 */
+	/** Copy constructor. */
 	TSharedFuture(const TSharedFuture&) = default;
 
-	/**
-	 * Copy assignment operator.
-	 */
-	TSharedFuture& operator=(const TSharedFuture& Other) = default;
+	/** Copy assignment operator. */
+	TSharedFuture& operator=(const TSharedFuture&) = default;
 
-	/**
-	 * Move constructor.
-	 */
+	/** Move constructor. */
 	TSharedFuture(TSharedFuture&&) = default;
 
-	/**
-	 * Move assignment operator.
-	 */
-	TSharedFuture& operator=(TSharedFuture&& Other) = default;
+	/** Move assignment operator. */
+	TSharedFuture& operator=(TSharedFuture&&) = default;
 
 	/** Destructor. */
-	~TSharedFuture() { }
+	~TSharedFuture() = default;
 
 public:
 
@@ -1018,7 +1056,7 @@ public:
 	 * The result must be set only once. An assertion will
 	 * be triggered if this method is called a second time.
 	 *
-	 * @param Result The result value to set.
+	 * @param Args The arguments to forward to the constructor of the result.
 	 */
 	template<typename... ArgTypes>
 	void EmplaceValue(ArgTypes&&... Args)
@@ -1216,8 +1254,6 @@ public:
 	 *
 	 * The result must be set only once. An assertion will
 	 * be triggered if this method is called a second time.
-	 *
-	 * @param Result The result value to set.
 	 */
 	void EmplaceValue()
 	{
@@ -1275,11 +1311,11 @@ auto TFutureBase<InternalResultType>::Then(Func Continuation) //-> TFuture<declt
 // Next implementation
 template<typename InternalResultType>
 template<typename Func>
-auto TFutureBase<InternalResultType>::Next(Func Continuation) //-> TFuture<decltype(Continuation(Get()))>
+auto TFutureBase<InternalResultType>::Next(Func Continuation) //-> TFuture<decltype(Continuation(Consume()))>
 {
 	return this->Then([Continuation = MoveTemp(Continuation)](TFuture<InternalResultType> Self) mutable
 	{
-		return Continuation(Self.Get());
+		return Continuation(Self.Consume());
 	});
 }
 
