MemPro platform extension cleanup

diff --git a/Engine/Source/Runtime/Core/Public/MemPro/MemPro.h b/Engine/Source/Runtime/Core/Public/MemPro/MemPro.h
--- a/Engine/Source/Runtime/Core/Public/MemPro/MemPro.h
+++ b/Engine/Source/Runtime/Core/Public/MemPro/MemPro.h
@@ -70,7 +70,7 @@ OR IN CONNECTION WITH THE USE OR PERFORMANCE OF THIS SOFTWARE.
 //@EPIC BEGIN external definition of platforms
 #if defined(MEMPRO_PLATFORM_XBOXONE)
 #elif defined(MEMPRO_PLATFORM_PS4)
-#elif defined(MEMPRO_PLATFORM_SWITCH)
+#elif defined(MEMPRO_PLATFORM_EXTENSION)
 #elif defined(_WIN32) || defined(_WIN64) || defined(WIN32) || defined(WIN64) || defined(__WIN32__)
 //@EPIC END
 	#if defined(_XBOX_ONE)
@@ -297,6 +297,13 @@ namespace
 		#include "Windows/HideWindowsPlatformTypes.h"
 	#endif
 
+//@EPIC BEGIN: other platforms
+#elif defined(MEMPRO_PLATFORM_EXTENSION) && defined(__UNREAL__)
+
+	#include COMPILED_PLATFORM_HEADER_WITH_PREFIX(MemPro,MemProExt.h)
+
+//@EPIC END
+
 #elif defined(MEMPRO_PLATFORM_XBOXONE)
 
 	#ifdef __UNREAL__
@@ -317,16 +324,6 @@ namespace
 		#include "MemProPS4.hpp"			// contact [REDACTED_EMAIL] for this platform
 	#endif
 
-//@EPIC BEGIN: other platforms
-#elif defined(MEMPRO_PLATFORM_SWITCH)
-
-	#include "MemPro/MemProSwitch.h"
-
-#elif defined(MEMPRO_PLATFORM_PS5)
-
-	#include "MemPro/MemProPS5.h"
-//@EPIC END
-
 #elif defined(MEMPRO_PLATFORM_UNIX)
 
 	#define MEMPRO_UNIX_BASED_PLATFORM
