Added StrConv as a long-term replacement for StringCast, because of the latter's behavior with regards to mid-string zeros.
Fixed FString serialization when it contains a trailing zero.

diff --git a/Engine/Source/Runtime/Core/Private/Containers/String.cpp.inl b/Engine/Source/Runtime/Core/Private/Containers/String.cpp.inl
--- a/Engine/Source/Runtime/Core/Private/Containers/String.cpp.inl
+++ b/Engine/Source/Runtime/Core/Private/Containers/String.cpp.inl
@@ -1877,7 +1877,7 @@ FArchive& operator<<( FArchive& Ar, UE_STRING_CLASS& A )
 			const bool bSaveUnicodeChar = Ar.IsForcingUnicode() || !TCString<ElementType>::IsPureAnsi(*A);
 			if (bSaveUnicodeChar)
 			{
-				// This preprocessor block should not be necessary when the StringCast above understands UTF16CHAR.
+				// This preprocessor block should not be necessary when the StrCast below understands UTF16CHAR.
 				FTCHARToUTF16 UTF16String(*A, A.Len() + 1); // include the null terminator
 				int32 Num = UTF16String.Length() + 1; // include the null terminator
 
@@ -1908,7 +1908,7 @@ FArchive& operator<<( FArchive& Ar, UE_STRING_CLASS& A )
 
 				if (Num)
 				{
-					Ar.Serialize((void*)StringCast<ANSICHAR>(A.Data.GetData(), Num).Get(), sizeof(ANSICHAR) * Num);
+					Ar.Serialize((void*)StrCast<ANSICHAR>(A.Data.GetData(), Num).Get(), sizeof(ANSICHAR) * Num);
 				}
 			}
 		#else
@@ -1918,7 +1918,7 @@ FArchive& operator<<( FArchive& Ar, UE_STRING_CLASS& A )
 
 			if (SaveNum)
 			{
-				auto CompactString = StringCast<UTF8CHAR>(A.Data.GetData(), SaveNum);
+				auto CompactString = StrCast<UTF8CHAR>(A.Data.GetData(), SaveNum);
 				Ar.Serialize((void*)CompactString.Get(), sizeof(UTF8CHAR) * CompactString.Length());
 			}
 		#endif
diff --git a/Engine/Source/Runtime/Core/Public/Containers/StringConv.h b/Engine/Source/Runtime/Core/Public/Containers/StringConv.h
--- a/Engine/Source/Runtime/Core/Public/Containers/StringConv.h
+++ b/Engine/Source/Runtime/Core/Public/Containers/StringConv.h
@@ -331,6 +331,13 @@ namespace UE::Core::Private
 		No  = 0,
 		Yes = 1
 	};
+
+	// This tag type is only intended to be used by StrCast to construct existing classes in direct ways.
+	// It is defined is in this Private namespace to indicate that the constructors that use it are not to be used by end-users.
+	struct FFromStrCast
+	{
+		explicit FFromStrCast() = default;
+	};
 }
 
 using FTCHARToUTF8_Convert /*UE_DEPRECATED(5.1, "FTCHARToUTF8_Convert has been deprecated in favor of FPlatformString::Convert and StringCast")*/ = UE::Core::Private::FTCHARToUTF8_Convert;
@@ -724,6 +731,11 @@ public:
 	{
 	}
 
+	explicit TStringConversion(UE::Core::Private::FFromStrCast, const FromType* Source, int32 SourceLen)
+	{
+		Init(Source, SourceLen, UE::Core::Private::ENullTerminatedString::No);
+	}
+
 	/**
 	 * Construct from a compatible character range such as TStringView or TStringBuilder.
 	 */
@@ -828,7 +840,7 @@ public:
 		typename SrcBufferType
 		UE_REQUIRES(TIsCharEncodingCompatibleWith_V<SrcBufferType, FromType>)
 	>
-	TStringPointer(const SrcBufferType* Source, int32 SourceLen)
+	explicit TStringPointer(const SrcBufferType* Source, int32 SourceLen)
 	{
 		if (Source)
 		{
@@ -848,6 +860,13 @@ public:
 		}
 	}
 
+	// This constructor is only intended to be used by StrCast, not end users
+	explicit TStringPointer(UE::Core::Private::FFromStrCast, const FromType* Source, int32 SourceLen)
+		: Ptr((const ToType*)Source)
+		, StringLength(SourceLen)
+	{
+	}
+
 	/**
 	 * Construct from a compatible character range such as TStringView or TStringBuilder.
 	 */
@@ -1141,6 +1160,43 @@ FORCEINLINE auto StringCast(const From* Str, int32 Len)
 	}
 }
 
+/**
+ * Creates an object which acts as a source of a given string type.  See example above.
+ *
+ * This is intended as the long-term replacement for StringCast, which doesn't cope well
+ * with zeros mid-string, which it can interpret as a null-terminators and give surprising behavior.
+ *
+ * StrCast expects correctly-typed strings.  If a cast is attempted with a char* or ANSICHAR*,
+ * and the string contains characters that are non-ASCII, including UTF-8 code units outside of
+ * the 7-bit ASCII range, then those values will fail to convert and a bogus char will be written
+ * in their place.
+ *
+ * If a conversion from UTF-8 is desired, the pointer should be cast to UTF8CHAR* before being
+ * passed to StrCast.
+ *
+ * Similarly, doing a StrCast<char> or StrCast<ANSICHAR> on a Unicode string will only
+ * successfully convert Unicode characters which already lie in the ASCII range.  For converting
+ * to UTF-8, StrCast<UTF8CHAR> should be used.
+ *
+ * The source string must not be modified or destroyed until after the result of this function
+ * has been destroyed.
+ *
+ * @param Str A pointer to the start of the string to convert.  Must be non-null if Len is non-zero.
+ * @param Len The number of From elements in Str.  Must be non-negative.
+ */
+template <typename To, int32 DefaultConversionSize = DEFAULT_STRING_CONVERSION_SIZE, typename From>
+FORCEINLINE auto StrCast(const From* Str, int32 Len)
+{
+	if constexpr (TIsCharEncodingCompatibleWith_V<From, To>)
+	{
+		return TStringPointer<To>(UE::Core::Private::FFromStrCast{}, (const To*)Str, Len);
+	}
+	else
+	{
+		return TStringConversion<TStringConvert<From, To>, DefaultConversionSize>(UE::Core::Private::FFromStrCast{}, Str, Len);
+	}
+}
+
 
 /**
  * Casts one fixed-width char type into another.
