Fix for stomp when optimizing a div + mad combo with the EVVM.  We don't support the result of the div being the last operand of the mad, but we would end up allocating space for the registers, which would then be iterated over.

diff --git a/Engine/Source/Runtime/Core/Private/UObject/FortniteReleaseBranchCustomObjectVersion.cpp b/Engine/Source/Runtime/Core/Private/UObject/FortniteReleaseBranchCustomObjectVersion.cpp
--- a/Engine/Source/Runtime/Core/Private/UObject/FortniteReleaseBranchCustomObjectVersion.cpp
+++ b/Engine/Source/Runtime/Core/Private/UObject/FortniteReleaseBranchCustomObjectVersion.cpp
@@ -6,7 +6,7 @@ TMap<FGuid, FGuid> FFortniteReleaseBranchCustomObjectVersion::GetSystemGuids()
 	TMap<FGuid, FGuid> SystemGuids;
 	const FDevSystemGuids& DevGuids = FDevSystemGuids::Get();
 
-	SystemGuids.Add(DevGuids.Niagara_LatestScriptCompileVersion, FGuid("C7E3EAD0DBD441B2911434B5CDA377D2"));
+	SystemGuids.Add(DevGuids.Niagara_LatestScriptCompileVersion, FGuid("FB079EFC568B49FBA6ECB6C4DFDD7EBB"));
 	SystemGuids.Add(DevGuids.SkeletalMeshDerivedDataVersion, FGuid("5705956EC7134274A5A1BD3FC74DB3A9"));
 	return SystemGuids;
 }
