Setting 28.10 net cl

#rnx
#release-28.20

diff --git a/Engine/Source/Runtime/Core/Public/Misc/NetworkVersion.h b/Engine/Source/Runtime/Core/Public/Misc/NetworkVersion.h
--- a/Engine/Source/Runtime/Core/Public/Misc/NetworkVersion.h
+++ b/Engine/Source/Runtime/Core/Public/Misc/NetworkVersion.h
@@ -9,7 +9,7 @@
 #include "Misc/EngineNetworkCustomVersion.h"
 
 // The version number used for determining network compatibility. If zero, uses the engine compatible version.
-#define ENGINE_NET_VERSION  0
+#define ENGINE_NET_VERSION  30535556
 
 // The version number used for determining replay compatibility
 #define ENGINE_REPLAY_VERSION  ENGINE_NET_VERSION
