Alllow applications to limit the total number of worker threads in TaskGraph

diff --git a/Engine/Source/Runtime/Core/Private/Async/TaskGraph.cpp b/Engine/Source/Runtime/Core/Private/Async/TaskGraph.cpp
--- a/Engine/Source/Runtime/Core/Private/Async/TaskGraph.cpp
+++ b/Engine/Source/Runtime/Core/Private/Async/TaskGraph.cpp
@@ -2352,6 +2352,11 @@ void FTaskGraphInterface::Startup(int32 NumThreads)
 		GUseNewTaskBackend = 1;
 	}
 
+	// Limit the total number of threads used
+#if defined(UE_TASKGRAPH_THREAD_LIMIT)
+	NumThreads = FMath::Min(NumThreads, int32(UE_TASKGRAPH_THREAD_LIMIT));
+#endif
+
 	if (GUseNewTaskBackend)
 	{
 		//We want to reduce the number of overall threads that UE uses so that there is are some 
