Ensured that the primary TIsContiguousContainer trait isn't instantiated with an incomplete type and so memoized into a bad state.

diff --git a/Engine/Source/Runtime/Core/Public/Traits/IsContiguousContainer.h b/Engine/Source/Runtime/Core/Public/Traits/IsContiguousContainer.h
--- a/Engine/Source/Runtime/Core/Public/Traits/IsContiguousContainer.h
+++ b/Engine/Source/Runtime/Core/Public/Traits/IsContiguousContainer.h
@@ -3,6 +3,7 @@
 #pragma once
 
 #include "CoreTypes.h"
+#include "Misc/StaticAssertCompleteType.h"
 #include <initializer_list>
 
 /**
@@ -13,6 +14,7 @@
 template <typename T>
 struct TIsContiguousContainer
 {
+	UE_STATIC_ASSERT_COMPLETE_TYPE(T, "TIsContiguousContainer instantiated with an incomplete type");
 	enum { Value = false };
 };
 
