Ensured that the primary TIsContiguousContainer trait isn't instantiated with an incomplete type and so memoized into a bad state.
Added a specialization for unbounded array types, which are not regarded as contiguous containers because they don't have an determinable end.

diff --git a/Engine/Source/Runtime/Core/Public/Traits/IsContiguousContainer.h b/Engine/Source/Runtime/Core/Public/Traits/IsContiguousContainer.h
--- a/Engine/Source/Runtime/Core/Public/Traits/IsContiguousContainer.h
+++ b/Engine/Source/Runtime/Core/Public/Traits/IsContiguousContainer.h
@@ -3,6 +3,7 @@
 #pragma once
 
 #include "CoreTypes.h"
+#include "Misc/StaticAssertCompleteType.h"
 #include <initializer_list>
 
 /**
@@ -13,6 +14,7 @@
 template <typename T>
 struct TIsContiguousContainer
 {
+	UE_STATIC_ASSERT_COMPLETE_TYPE(T, "TIsContiguousContainer instantiated with an incomplete type");
 	enum { Value = false };
 };
 
@@ -30,6 +32,14 @@ template <typename T, size_t N> struct TIsContiguousContainer<const          T[N
 template <typename T, size_t N> struct TIsContiguousContainer<      volatile T[N]> { enum { Value = true }; };
 template <typename T, size_t N> struct TIsContiguousContainer<const volatile T[N]> { enum { Value = true }; };
 
+/**
+ * Specialization for unbounded C arrays (never contiguous - should be treated as pointers which are not regarded as contiguous containers)
+ */
+template <typename T> struct TIsContiguousContainer<               T[]> { enum { Value = false }; };
+template <typename T> struct TIsContiguousContainer<const          T[]> { enum { Value = false }; };
+template <typename T> struct TIsContiguousContainer<      volatile T[]> { enum { Value = false }; };
+template <typename T> struct TIsContiguousContainer<const volatile T[]> { enum { Value = false }; };
+
 /**
  * Specialization for initializer lists (also always contiguous)
  */
