[Insights] Added the Windows error message to the log warning when the "UnrealTraceServer.exe" process fails to start. Also updated log messages re UTS launch.

diff --git a/Engine/Source/Runtime/Core/Private/ProfilingDebugging/TraceAuxiliary.cpp b/Engine/Source/Runtime/Core/Private/ProfilingDebugging/TraceAuxiliary.cpp
--- a/Engine/Source/Runtime/Core/Private/ProfilingDebugging/TraceAuxiliary.cpp
+++ b/Engine/Source/Runtime/Core/Private/ProfilingDebugging/TraceAuxiliary.cpp
@@ -120,7 +120,7 @@ public:
 	bool SendSnapshot(const TCHAR* InHost, uint32 InPort, const FTraceAuxiliary::FLogCategoryAlias& LogCategory);
 
 	// True if this is parent process with forking requested before forking.
-	bool					IsParentProcessAndPreFork();
+	bool IsParentProcessAndPreFork();
 
 private:
 	enum class EState : uint8
@@ -150,7 +150,7 @@ private:
 	bool bWorkerThreadStarted = false;
 	bool bTruncateFile = false;
 	FString PausedPreset;
-	
+
 	struct FCurrentTraceTarget
 	{
 		FString TraceDest;
@@ -334,13 +334,13 @@ bool FTraceAuxiliaryImpl::Connect(FTraceAuxiliary::EConnectionType Type, const T
 		{
 			FString StartedDest;
 			FTraceAuxiliary::EConnectionType StartedType = FTraceAuxiliary::EConnectionType::None;
-			
+
 			{
 				FReadScopeLock _(CurrentTargetLock);
 				StartedDest = CurrentTraceTarget.TraceDest;
 				StartedType = CurrentTraceTarget.TraceType;
 			}
-			
+
 			FTraceAuxiliary::OnTraceStarted.Broadcast(StartedType, StartedDest);
 		}
 	}
@@ -367,7 +367,7 @@ bool FTraceAuxiliaryImpl::Stop()
 
 	FString StopedDest;
 	FTraceAuxiliary::EConnectionType StopedType = FTraceAuxiliary::EConnectionType::None;
-	
+
 	{
 		FWriteScopeLock _(CurrentTargetLock);
 		StopedDest = CurrentTraceTarget.TraceDest;
@@ -377,7 +377,7 @@ bool FTraceAuxiliaryImpl::Stop()
 	}
 
 	FTraceAuxiliary::OnTraceStopped.Broadcast(StopedType, StopedDest);
-	
+
 	return true;
 }
 
@@ -1491,7 +1491,7 @@ void FTraceAuxiliary::Initialize(const TCHAR* CommandLine)
 		SessionGuid = FApp::GetInstanceId();
 	}
 	FMemory::Memcpy((FGuid&)Desc.SessionGuid, SessionGuid);
-	
+
 	if (FParse::Value(CommandLine, TEXT("-tracetailmb="), Desc.TailSizeBytes))
 	{
 		Desc.TailSizeBytes <<= 20;
@@ -1691,12 +1691,32 @@ void FTraceAuxiliary::TryAutoConnect()
 ////////////////////////////////////////////////////////////////////////////////
 #if UE_TRACE_SERVER_CONTROLS_ENABLED
 
+enum class ELaunchTraceServerCommand
+{
+	Fork,
+	Kill
+};
+
 #if PLATFORM_WINDOWS
-bool LaunchTraceServerCommand(const TCHAR* Command, bool bAddSponsor)
+bool LaunchTraceServerCommand(ELaunchTraceServerCommand Command, bool bAddSponsor)
 {
+	FString FilePath = FPaths::Combine(FPaths::EngineDir(), TEXT("Binaries/Win64/UnrealTraceServer.exe"));
+	if (!FPaths::FileExists(FilePath))
+	{
+		UE_LOG(LogCore, Display, TEXT("UTS: The Unreal Trace Server binary is not available ('%s')"), *FilePath);
+		return false;
+	}
+
 	TWideStringBuilder<MAX_PATH + 32> CreateProcArgs;
-	CreateProcArgs << "\"" << FPaths::EngineDir() << TEXT("Binaries/Win64/UnrealTraceServer.exe\"");
-	CreateProcArgs << TEXT(" ") << Command;
+	CreateProcArgs << TEXT("\"") << FilePath << TEXT("\"");
+	if (Command == ELaunchTraceServerCommand::Fork)
+	{
+		CreateProcArgs << TEXT(" fork");
+	}
+	else if (Command == ELaunchTraceServerCommand::Kill)
+	{
+		CreateProcArgs << TEXT(" kill");
+	}
 	if (bAddSponsor)
 	{
 		CreateProcArgs << TEXT(" --sponsor ") << FPlatformProcess::GetCurrentProcessId();
@@ -1711,21 +1731,26 @@ bool LaunchTraceServerCommand(const TCHAR* Command, bool bAddSponsor)
 	{
 		CreateProcFlags |= CREATE_NO_WINDOW;
 	}
+
 	STARTUPINFOW StartupInfo = { sizeof(STARTUPINFOW) };
 	PROCESS_INFORMATION ProcessInfo = {};
+
 	const BOOL bOk = CreateProcessW(nullptr, LPWSTR(*CreateProcArgs), nullptr, nullptr,
 									false, CreateProcFlags, nullptr, nullptr, &StartupInfo, &ProcessInfo);
 
 	if (!bOk)
 	{
-		UE_LOG(LogCore, Display, TEXT("Unable to launch the Unreal Trace Server with '%s' (%08x)"), *CreateProcArgs, GetLastError());
+		DWORD LastError = GetLastError();
+		TCHAR ErrorBuffer[1024];
+		FWindowsPlatformMisc::GetSystemErrorMessage(ErrorBuffer, UE_ARRAY_COUNT(ErrorBuffer), LastError);
+		UE_LOG(LogCore, Warning, TEXT("UTS: Unable to launch the Unreal Trace Server with '%s'. %s Error: 0x%X (%u)"), *CreateProcArgs, ErrorBuffer, LastError, LastError);
 		return false;
 	}
 
 	bool bSuccess = false;
 	if (WaitForSingleObject(ProcessInfo.hProcess, 5000) == WAIT_TIMEOUT)
 	{
-		UE_LOG(LogCore, Warning, TEXT("Timed out waiting for the Unreal Trace Server to start"));
+		UE_LOG(LogCore, Warning, TEXT("UTS: Timed out waiting for the Unreal Trace Server process to start"));
 	}
 	else
 	{
@@ -1733,27 +1758,37 @@ bool LaunchTraceServerCommand(const TCHAR* Command, bool bAddSponsor)
 		GetExitCodeProcess(ProcessInfo.hProcess, &ExitCode);
 		if (ExitCode)
 		{
-			UE_LOG(LogCore, Warning, TEXT("Unreal Trace Server returned an error (0x%08x)"), ExitCode);
+			UE_LOG(LogCore, Warning, TEXT("UTS: Unreal Trace Server process returned an error (0x%08x)"), ExitCode);
 		}
 		else
 		{
-			UE_LOG(LogCore, Log, TEXT("Unreal Trace Server launch successful"));
+			if (Command == ELaunchTraceServerCommand::Kill)
+			{
+				UE_LOG(LogCore, Log, TEXT("UTS: Unreal Trace Server was stopped"));
+			}
+			else
+			{
+				UE_LOG(LogCore, Log, TEXT("UTS: Unreal Trace Server launched successfully"));
+			}
 			bSuccess = true;
 		}
 	}
 
 	CloseHandle(ProcessInfo.hProcess);
 	CloseHandle(ProcessInfo.hThread);
-	
+
 	return bSuccess;
 }
-#endif
+#endif // PLATFORM_WINDOWS
 
 #if PLATFORM_LINUX || PLATFORM_MAC
-static bool LaunchTraceServerCommand(const TCHAR* Command, bool bAddSponsor)
+static bool LaunchTraceServerCommand(ELaunchTraceServerCommand Command, bool bAddSponsor)
 {
-// TSAN doesn't like fork(), so disable this for now.
-#if !USING_THREAD_SANITISER 
+#if USING_THREAD_SANITISER
+	// TSAN doesn't like fork(), so disable this for now.
+	return false;
+#else // !USING_THREAD_SANITISER
+
 	TAnsiStringBuilder<320> BinPath;
 	BinPath << TCHAR_TO_UTF8(*FPaths::EngineDir());
 #if PLATFORM_UNIX
@@ -1765,39 +1800,48 @@ static bool LaunchTraceServerCommand(const TCHAR* Command, bool bAddSponsor)
 
 	if (access(*BinPath, F_OK) < 0)
 	{
-		UE_LOG(LogCore, Display, TEXT("UnrealTraceServer: Binary not found (%s)"), ANSI_TO_TCHAR(*BinPath));
+		UE_LOG(LogCore, Display, TEXT("UTS: The Unreal Trace Server binary is not available ('%s')"), ANSI_TO_TCHAR(*BinPath));
 		return false;
 	}
 
 	TAnsiStringBuilder<64> ForkArg;
-	ForkArg << " " << Command;
+	if (Command == ELaunchTraceServerCommand::Fork)
+	{
+		ForkArg << " fork";
+	}
+	else if (Command == ELaunchTraceServerCommand::Kill)
+	{
+		ForkArg << " kill";
+	}
 	if (bAddSponsor)
 	{
-		ForkArg << TEXT(" --sponsor ") << FPlatformProcess::GetCurrentProcessId();
+		ForkArg << " --sponsor " << FPlatformProcess::GetCurrentProcessId();
 	}
 	ForkArg.ToString(); //Ensure zero termination
 
 	pid_t UtsPid = fork();
 	if (UtsPid < 0)
 	{
-		UE_LOG(LogCore, Display, TEXT("UnrealTraceServer: Unable to fork (errno: %d)"), errno);
+		UE_LOG(LogCore, Warning, TEXT("UTS: Unable to fork (errno: %d)"), errno);
 		return false;
 	}
 	else if (UtsPid == 0)
 	{
+		// Launch UTS from the child process.
 		char* Args[] = { BinPath.GetData(), ForkArg.GetData(), nullptr };
 		extern char** environ;
 		execve(*BinPath, Args, environ);
 		_exit(0x80 | (errno & 0x7f));
 	}
 
+	// Wait until the child process finishes.
 	int32 WaitStatus = 0;
 	do
 	{
 		int32 WaitRet = waitpid(UtsPid, &WaitStatus, 0);
 		if (WaitRet < 0)
 		{
-			UE_LOG(LogCore, Display, TEXT("UnrealTraceServer: waitpid() error; (errno: %d)"), errno);
+			UE_LOG(LogCore, Warning, TEXT("UTS: waitpid() error (errno: %d)"), errno);
 			return false;
 		}
 	}
@@ -1806,30 +1850,35 @@ static bool LaunchTraceServerCommand(const TCHAR* Command, bool bAddSponsor)
 	int32 UtsRet = WEXITSTATUS(WaitStatus);
 	if (UtsRet)
 	{
-		UE_LOG(LogCore, Display, TEXT("UnrealTraceServer: Trace store returned an error (0x%08x)"), UtsRet);
+		UE_LOG(LogCore, Warning, TEXT("UTS: Unreal Trace Server process returned an error (0x%08x)"), UtsRet);
 		return false;
 	}
 	else
 	{
-		UE_LOG(LogCore, Log, TEXT("UnrealTraceServer: Trace store launch successful"));
+		if (Command == ELaunchTraceServerCommand::Kill)
+		{
+			UE_LOG(LogCore, Log, TEXT("UTS: Unreal Trace Server was stopped"));
+		}
+		else
+		{
+			UE_LOG(LogCore, Log, TEXT("UTS: Unreal Trace Server launched successfully"));
+		}
 		return true;
 	}
-#else // #if !USING_THREAD_SANITISER
-	return false;
-#endif
+#endif // !USING_THREAD_SANITISER
 }
-#endif
+#endif // PLATFORM_LINUX || PLATFORM_MAC
 
 ////////////////////////////////////////////////////////////////////////////////
 bool FTraceServerControls::Start()
 {
-	return LaunchTraceServerCommand(TEXT("fork"), false);
+	return LaunchTraceServerCommand(ELaunchTraceServerCommand::Fork, false);
 }
 
 ////////////////////////////////////////////////////////////////////////////////
 bool FTraceServerControls::Stop()
 {
-	return LaunchTraceServerCommand(TEXT("kill"), false);
+	return LaunchTraceServerCommand(ELaunchTraceServerCommand::Kill, false);
 }
 
 #endif // UE_TRACE_SERVER_CONTROLS_ENABLED
