Replaced UE_NODISCARD with [[nodiscard]].

diff --git a/Engine/Source/Runtime/Core/Private/Containers/String.cpp.inl b/Engine/Source/Runtime/Core/Private/Containers/String.cpp.inl
--- a/Engine/Source/Runtime/Core/Private/Containers/String.cpp.inl
+++ b/Engine/Source/Runtime/Core/Private/Containers/String.cpp.inl
@@ -686,7 +686,7 @@ namespace UE::String::Private
 {
 
 template <typename LhsType, typename RhsType>
-UE_NODISCARD FORCEINLINE UE_STRING_CLASS PREPROCESSOR_JOIN(ConcatStrings_, UE_STRING_CLASS)(LhsType&& Lhs, RhsType&& Rhs)
+[[nodiscard]] FORCEINLINE UE_STRING_CLASS PREPROCESSOR_JOIN(ConcatStrings_, UE_STRING_CLASS)(LhsType&& Lhs, RhsType&& Rhs)
 {
 	Lhs.CheckInvariants();
 	Rhs.CheckInvariants();
@@ -705,7 +705,7 @@ UE_NODISCARD FORCEINLINE UE_STRING_CLASS PREPROCESSOR_JOIN(ConcatStrings_, UE_ST
 }
 
 template <typename LhsCharType, typename RhsType>
-UE_NODISCARD FORCEINLINE UE_STRING_CLASS PREPROCESSOR_JOIN(ConcatRangeString_, UE_STRING_CLASS)(const LhsCharType* Lhs, int32 LhsLen, RhsType&& Rhs)
+[[nodiscard]] FORCEINLINE UE_STRING_CLASS PREPROCESSOR_JOIN(ConcatRangeString_, UE_STRING_CLASS)(const LhsCharType* Lhs, int32 LhsLen, RhsType&& Rhs)
 {
 	using ElementType = UE_STRING_CLASS::ElementType;
 
@@ -734,7 +734,7 @@ UE_NODISCARD FORCEINLINE UE_STRING_CLASS PREPROCESSOR_JOIN(ConcatRangeString_, U
 }
 
 template <typename LhsType, typename RhsCharType>
-UE_NODISCARD FORCEINLINE UE_STRING_CLASS PREPROCESSOR_JOIN(ConcatStringRange_, UE_STRING_CLASS)(LhsType&& Lhs, const RhsCharType* Rhs, int32 RhsLen)
+[[nodiscard]] FORCEINLINE UE_STRING_CLASS PREPROCESSOR_JOIN(ConcatStringRange_, UE_STRING_CLASS)(LhsType&& Lhs, const RhsCharType* Rhs, int32 RhsLen)
 {
 	Lhs.CheckInvariants();
 	checkSlow(RhsLen >= 0);
@@ -750,7 +750,7 @@ UE_NODISCARD FORCEINLINE UE_STRING_CLASS PREPROCESSOR_JOIN(ConcatStringRange_, U
 }
 
 template <typename LhsCharType, typename RhsType>
-UE_NODISCARD FORCEINLINE UE_STRING_CLASS PREPROCESSOR_JOIN(ConcatPtrString_, UE_STRING_CLASS)(const LhsCharType* Lhs, RhsType&& Rhs)
+[[nodiscard]] FORCEINLINE UE_STRING_CLASS PREPROCESSOR_JOIN(ConcatPtrString_, UE_STRING_CLASS)(const LhsCharType* Lhs, RhsType&& Rhs)
 {
 	checkSlow(Lhs);
 	if (!Lhs)
@@ -762,7 +762,7 @@ UE_NODISCARD FORCEINLINE UE_STRING_CLASS PREPROCESSOR_JOIN(ConcatPtrString_, UE_
 }
 
 template <typename LhsType, typename RhsCharType>
-UE_NODISCARD FORCEINLINE UE_STRING_CLASS PREPROCESSOR_JOIN(ConcatStringPtr_, UE_STRING_CLASS)(LhsType&& Lhs, const RhsCharType* Rhs)
+[[nodiscard]] FORCEINLINE UE_STRING_CLASS PREPROCESSOR_JOIN(ConcatStringPtr_, UE_STRING_CLASS)(LhsType&& Lhs, const RhsCharType* Rhs)
 {
 	checkSlow(Rhs);
 	if (!Rhs)
diff --git a/Engine/Source/Runtime/Core/Private/IO/IoStore.cpp b/Engine/Source/Runtime/Core/Private/IO/IoStore.cpp
--- a/Engine/Source/Runtime/Core/Private/IO/IoStore.cpp
+++ b/Engine/Source/Runtime/Core/Private/IO/IoStore.cpp
@@ -328,7 +328,7 @@ public:
 		}
 	}
 
-	UE_NODISCARD FIoStatus Initialize(const FIoStoreWriterSettings& InWriterSettings)
+	[[nodiscard]] FIoStatus Initialize(const FIoStoreWriterSettings& InWriterSettings)
 	{
 		TRACE_CPUPROFILER_EVENT_SCOPE(FIoStoreWriterContext::Initialize);
 		WriterSettings = InWriterSettings;
@@ -524,7 +524,7 @@ FIoStoreWriterContext::~FIoStoreWriterContext()
 	delete Impl;
 }
 
-UE_NODISCARD FIoStatus FIoStoreWriterContext::Initialize(const FIoStoreWriterSettings& InWriterSettings)
+[[nodiscard]] FIoStatus FIoStoreWriterContext::Initialize(const FIoStoreWriterSettings& InWriterSettings)
 {
 	return Impl->Initialize(InWriterSettings);
 }
@@ -754,7 +754,7 @@ public:
 		}
 	}
 
-	UE_NODISCARD FIoStatus Initialize(FIoStoreWriterContextImpl& InContext, const FIoContainerSettings& InContainerSettings)
+	[[nodiscard]] FIoStatus Initialize(FIoStoreWriterContextImpl& InContext, const FIoContainerSettings& InContainerSettings)
 	{
 		WriterContext = &InContext;
 		ContainerSettings = InContainerSettings;
@@ -2302,7 +2302,7 @@ public:
 		});
 	}
 
-	UE_NODISCARD FIoStatus Initialize(FStringView InContainerPath, const TMap<FGuid, FAES::FAESKey>& InDecryptionKeys)
+	[[nodiscard]] FIoStatus Initialize(FStringView InContainerPath, const TMap<FGuid, FAES::FAESKey>& InDecryptionKeys)
 	{
 		TRACE_CPUPROFILER_EVENT_SCOPE(FIoStoreReader::Initialize);
 		ContainerPath = InContainerPath;
diff --git a/Engine/Source/Runtime/Core/Public/Async/Fundamental/Scheduler.h b/Engine/Source/Runtime/Core/Public/Async/Fundamental/Scheduler.h
--- a/Engine/Source/Runtime/Core/Public/Async/Fundamental/Scheduler.h
+++ b/Engine/Source/Runtime/Core/Public/Async/Fundamental/Scheduler.h
@@ -194,7 +194,7 @@ namespace LowLevelTasks
 		~FScheduler();
 
 	private: 
-		UE_NODISCARD FTask* ExecuteTask(FTask* InTask);
+		[[nodiscard]] FTask* ExecuteTask(FTask* InTask);
 		TUniquePtr<FThread> CreateWorker(bool bPermitBackgroundWork = false, FThread::EForkable IsForkable = FThread::NonForkable, Private::FWaitEvent* ExternalWorkerEvent = nullptr, FSchedulerTls::FLocalQueueType* ExternalWorkerLocalQueue = nullptr, EThreadPriority Priority = EThreadPriority::TPri_Normal, uint64 InAffinity = 0);
 		void WorkerMain(Private::FWaitEvent* WorkerEvent, FSchedulerTls::FLocalQueueType* ExternalWorkerLocalQueue, uint32 WaitCycles, bool bPermitBackgroundWork);
 		CORE_API void LaunchInternal(FTask& Task, EQueuePreference QueuePreference, bool bWakeUpWorker);
diff --git a/Engine/Source/Runtime/Core/Public/Async/InheritedContext.h b/Engine/Source/Runtime/Core/Public/Async/InheritedContext.h
--- a/Engine/Source/Runtime/Core/Public/Async/InheritedContext.h
+++ b/Engine/Source/Runtime/Core/Public/Async/InheritedContext.h
@@ -132,7 +132,7 @@ namespace UE
 		}
 
 		// must be called where the inherited context should be restored, e.g. at the start of an async task execution
-		UE_NODISCARD CORE_API FInheritedContextScope RestoreInheritedContext();
+		[[nodiscard]] CORE_API FInheritedContextScope RestoreInheritedContext();
 
 	private:
 #if ENABLE_LOW_LEVEL_MEM_TRACKER
diff --git a/Engine/Source/Runtime/Core/Public/Containers/StringView.h b/Engine/Source/Runtime/Core/Public/Containers/StringView.h
--- a/Engine/Source/Runtime/Core/Public/Containers/StringView.h
+++ b/Engine/Source/Runtime/Core/Public/Containers/StringView.h
@@ -349,7 +349,7 @@ public:
 	 *
 	 * @returns True if index is valid. False otherwise.
 	 */
-	UE_NODISCARD FORCEINLINE bool IsValidIndex(int32 Index) const
+	[[nodiscard]] FORCEINLINE bool IsValidIndex(int32 Index) const
 	{
 		return Index >= 0 && Index < Len();
 	}
diff --git a/Engine/Source/Runtime/Core/Public/Containers/UnrealString.h b/Engine/Source/Runtime/Core/Public/Containers/UnrealString.h
--- a/Engine/Source/Runtime/Core/Public/Containers/UnrealString.h
+++ b/Engine/Source/Runtime/Core/Public/Containers/UnrealString.h
@@ -24,7 +24,7 @@
  * @param Count number of bytes to convert
  * @return Valid string representing bytes.
  */
-UE_NODISCARD inline FString BytesToString(const uint8* In, int32 Count)
+[[nodiscard]] inline FString BytesToString(const uint8* In, int32 Count)
 {
 	FString Result;
 	Result.Empty(Count);
@@ -65,7 +65,7 @@ inline int32 StringToBytes( const FString& String, uint8* OutBytes, int32 MaxBuf
 }
 
 /** Returns Char value of Nibble */
-UE_NODISCARD inline TCHAR NibbleToTChar(uint8 Num)
+[[nodiscard]] inline TCHAR NibbleToTChar(uint8 Num)
 {
 	if (Num > 9)
 	{
@@ -86,7 +86,7 @@ inline void ByteToHex(uint8 In, FString& Result)
 }
 
 /** Convert bytes to uppercase hex string */
-UE_NODISCARD inline FString BytesToHex(const uint8* Bytes, int32 NumBytes)
+[[nodiscard]] inline FString BytesToHex(const uint8* Bytes, int32 NumBytes)
 {
 	FString Out;
 	BytesToHex(Bytes, NumBytes, Out);
@@ -94,7 +94,7 @@ UE_NODISCARD inline FString BytesToHex(const uint8* Bytes, int32 NumBytes)
 }
 
 /** Convert bytes to lowercase hex string */
-UE_NODISCARD inline FString BytesToHexLower(const uint8* Bytes, int32 NumBytes)
+[[nodiscard]] inline FString BytesToHexLower(const uint8* Bytes, int32 NumBytes)
 {
 	FString Out;
 	BytesToHexLower(Bytes, NumBytes, Out);
@@ -106,7 +106,7 @@ UE_NODISCARD inline FString BytesToHexLower(const uint8* Bytes, int32 NumBytes)
  * @param Char		The character
  * @return	True if in 0-9 and A-F ranges
  */
-UE_NODISCARD inline const bool CheckTCharIsHex(const TCHAR Char)
+[[nodiscard]] inline const bool CheckTCharIsHex(const TCHAR Char)
 {
 	return (Char >= TEXT('0') && Char <= TEXT('9')) || (Char >= TEXT('A') && Char <= TEXT('F')) || (Char >= TEXT('a') && Char <= TEXT('f'));
 }
@@ -116,7 +116,7 @@ UE_NODISCARD inline const bool CheckTCharIsHex(const TCHAR Char)
  * @param Hex		The character
  * @return	The uint8 value of a hex character
  */
-UE_NODISCARD inline const uint8 TCharToNibble(const TCHAR Hex)
+[[nodiscard]] inline const uint8 TCharToNibble(const TCHAR Hex)
 {
 	if (Hex >= '0' && Hex <= '9')
 	{
@@ -140,7 +140,7 @@ template <
 	typename T
 	UE_REQUIRES(std::is_arithmetic_v<T>)
 >
-UE_NODISCARD StringType LexToString(const T& Value)
+[[nodiscard]] StringType LexToString(const T& Value)
 {
 	// std::remove_cv_t to remove potential volatile decorations. Removing const is pointless, but harmless because it's specified in the param declaration.
 	return StringType::Printf(TFormatSpecifier<std::remove_cv_t<T>>::GetFormatSpecifier(), Value);
@@ -151,13 +151,13 @@ template <
 	typename CharType
 	UE_REQUIRES(TIsCharType_V<CharType>)
 >
-UE_NODISCARD StringType LexToString(const CharType* Ptr)
+[[nodiscard]] StringType LexToString(const CharType* Ptr)
 {
 	return StringType(Ptr);
 }
 
 template <typename StringType = FString>
-UE_NODISCARD inline StringType LexToString(bool Value)
+[[nodiscard]] inline StringType LexToString(bool Value)
 {
 	using ElementType = typename StringType::ElementType;
 	return Value ? CHARTEXT(ElementType, "true") : CHARTEXT(ElementType, "false");
@@ -165,21 +165,21 @@ UE_NODISCARD inline StringType LexToString(bool Value)
 
 /** Helper template to convert to sanitized strings */
 template <typename StringType = FString, typename T>
-UE_NODISCARD StringType LexToSanitizedString(const T& Value)
+[[nodiscard]] StringType LexToSanitizedString(const T& Value)
 {
 	return LexToString<StringType>(Value);
 }
 
 /** Overloaded for floats */
 template <typename StringType = FString>
-UE_NODISCARD inline StringType LexToSanitizedString(float Value)
+[[nodiscard]] inline StringType LexToSanitizedString(float Value)
 {
 	return StringType::SanitizeFloat(Value);
 }
 
 /** Overloaded for doubles */
 template <typename StringType = FString>
-UE_NODISCARD inline StringType LexToSanitizedString(double Value)
+[[nodiscard]] inline StringType LexToSanitizedString(double Value)
 {
 	return StringType::SanitizeFloat(Value);
 }
@@ -189,13 +189,13 @@ template<typename T>
 struct TTypeToString
 {
 	template <typename StringType = FString>
-	UE_NODISCARD static StringType ToString(const T& Value)
+	[[nodiscard]] static StringType ToString(const T& Value)
 	{
 		return LexToString<StringType>(Value);
 	}
 
 	template <typename StringType = FString>
-	UE_NODISCARD static StringType ToSanitizedString(const T& Value)
+	[[nodiscard]] static StringType ToSanitizedString(const T& Value)
 	{
 		return LexToSanitizedString<StringType>(Value);
 	}
diff --git a/Engine/Source/Runtime/Core/Public/Containers/UnrealString.h.inl b/Engine/Source/Runtime/Core/Public/Containers/UnrealString.h.inl
--- a/Engine/Source/Runtime/Core/Public/Containers/UnrealString.h.inl
+++ b/Engine/Source/Runtime/Core/Public/Containers/UnrealString.h.inl
@@ -302,7 +302,7 @@ public:
 	 *
 	 * @return true if this string is empty, otherwise return false.
 	 */
-	UE_NODISCARD FORCEINLINE bool IsEmpty() const
+	[[nodiscard]] FORCEINLINE bool IsEmpty() const
 	{
 		return Data.Num() <= 1;
 	}
@@ -326,7 +326,7 @@ public:
 	 *
 	 * @returns True if index is valid. False otherwise.
 	 */
-	UE_NODISCARD FORCEINLINE bool IsValidIndex(int32 Index) const
+	[[nodiscard]] FORCEINLINE bool IsValidIndex(int32 Index) const
 	{
 		return Index >= 0 && Index < Len();
 	}
@@ -336,7 +336,7 @@ public:
 	 *
 	 * @Return Pointer to Array of ElementType if Num, otherwise the empty string
 	 */
-	UE_NODISCARD FORCEINLINE const ElementType* operator*() const UE_LIFETIMEBOUND
+	[[nodiscard]] FORCEINLINE const ElementType* operator*() const UE_LIFETIMEBOUND
 	{
 		return Data.Num() ? Data.GetData() : CHARTEXT(ElementType, "");
 	}
@@ -347,13 +347,13 @@ public:
 	 * @warning: Operations on the TArray<*CHAR> can be unsafe, such as adding
 	 *		non-terminating 0's or removing the terminating zero.
 	 */
-	UE_NODISCARD FORCEINLINE DataType& GetCharArray() UE_LIFETIMEBOUND
+	[[nodiscard]] FORCEINLINE DataType& GetCharArray() UE_LIFETIMEBOUND
 	{
 		return Data;
 	}
 
 	/** Get string as const array of TCHARS */
-	UE_NODISCARD FORCEINLINE const DataType& GetCharArray() const UE_LIFETIMEBOUND
+	[[nodiscard]] FORCEINLINE const DataType& GetCharArray() const UE_LIFETIMEBOUND
 	{
 		return Data;
 	}
@@ -578,7 +578,7 @@ public:
 		typename CharType
 		UE_REQUIRES(TIsCharType_V<CharType>)
 	>
-	UE_NODISCARD FORCEINLINE friend UE_STRING_CLASS operator+(const UE_STRING_CLASS& Lhs, CharType Rhs)
+	[[nodiscard]] FORCEINLINE friend UE_STRING_CLASS operator+(const UE_STRING_CLASS& Lhs, CharType Rhs)
 	{
 		Lhs.CheckInvariants();
 
@@ -600,7 +600,7 @@ public:
 		typename CharType
 		UE_REQUIRES(TIsCharType_V<CharType>)
 	>
-	UE_NODISCARD FORCEINLINE friend UE_STRING_CLASS operator+(UE_STRING_CLASS&& Lhs, CharType Rhs)
+	[[nodiscard]] FORCEINLINE friend UE_STRING_CLASS operator+(UE_STRING_CLASS&& Lhs, CharType Rhs)
 	{
 		Lhs.CheckInvariants();
 
@@ -611,28 +611,28 @@ public:
 	}
 
 private:
-	UE_NODISCARD static CORE_API UE_STRING_CLASS ConcatFF(const UE_STRING_CLASS& Lhs, const UE_STRING_CLASS& Rhs);
-	UE_NODISCARD static CORE_API UE_STRING_CLASS ConcatFF(UE_STRING_CLASS&& Lhs, const UE_STRING_CLASS& Rhs);
-	UE_NODISCARD static CORE_API UE_STRING_CLASS ConcatFF(const UE_STRING_CLASS& Lhs, UE_STRING_CLASS&& Rhs);
-	UE_NODISCARD static CORE_API UE_STRING_CLASS ConcatFF(UE_STRING_CLASS&& Lhs, UE_STRING_CLASS&& Rhs);
-	UE_NODISCARD static CORE_API UE_STRING_CLASS ConcatFC(const UE_STRING_CLASS& Lhs, const ElementType* Rhs);
-	UE_NODISCARD static CORE_API UE_STRING_CLASS ConcatFC(UE_STRING_CLASS&& Lhs,	const ElementType* Rhs);
-	UE_NODISCARD static CORE_API UE_STRING_CLASS ConcatCF(const ElementType* Lhs, const UE_STRING_CLASS& Rhs);
-	UE_NODISCARD static CORE_API UE_STRING_CLASS ConcatCF(const ElementType* Lhs, UE_STRING_CLASS&& Rhs);
-	UE_NODISCARD static CORE_API UE_STRING_CLASS ConcatFR(const UE_STRING_CLASS& Lhs, const ElementType* Rhs, int32 RhsLen);
-	UE_NODISCARD static CORE_API UE_STRING_CLASS ConcatFR(UE_STRING_CLASS&& Lhs,	const ElementType* Rhs, int32 RhsLen);
-	UE_NODISCARD static CORE_API UE_STRING_CLASS ConcatRF(const ElementType* Lhs, int32 LhsLen, const UE_STRING_CLASS& Rhs);
-	UE_NODISCARD static CORE_API UE_STRING_CLASS ConcatRF(const ElementType* Lhs, int32 LhsLen, UE_STRING_CLASS&& Rhs);
+	[[nodiscard]] static CORE_API UE_STRING_CLASS ConcatFF(const UE_STRING_CLASS& Lhs, const UE_STRING_CLASS& Rhs);
+	[[nodiscard]] static CORE_API UE_STRING_CLASS ConcatFF(UE_STRING_CLASS&& Lhs, const UE_STRING_CLASS& Rhs);
+	[[nodiscard]] static CORE_API UE_STRING_CLASS ConcatFF(const UE_STRING_CLASS& Lhs, UE_STRING_CLASS&& Rhs);
+	[[nodiscard]] static CORE_API UE_STRING_CLASS ConcatFF(UE_STRING_CLASS&& Lhs, UE_STRING_CLASS&& Rhs);
+	[[nodiscard]] static CORE_API UE_STRING_CLASS ConcatFC(const UE_STRING_CLASS& Lhs, const ElementType* Rhs);
+	[[nodiscard]] static CORE_API UE_STRING_CLASS ConcatFC(UE_STRING_CLASS&& Lhs,	const ElementType* Rhs);
+	[[nodiscard]] static CORE_API UE_STRING_CLASS ConcatCF(const ElementType* Lhs, const UE_STRING_CLASS& Rhs);
+	[[nodiscard]] static CORE_API UE_STRING_CLASS ConcatCF(const ElementType* Lhs, UE_STRING_CLASS&& Rhs);
+	[[nodiscard]] static CORE_API UE_STRING_CLASS ConcatFR(const UE_STRING_CLASS& Lhs, const ElementType* Rhs, int32 RhsLen);
+	[[nodiscard]] static CORE_API UE_STRING_CLASS ConcatFR(UE_STRING_CLASS&& Lhs,	const ElementType* Rhs, int32 RhsLen);
+	[[nodiscard]] static CORE_API UE_STRING_CLASS ConcatRF(const ElementType* Lhs, int32 LhsLen, const UE_STRING_CLASS& Rhs);
+	[[nodiscard]] static CORE_API UE_STRING_CLASS ConcatRF(const ElementType* Lhs, int32 LhsLen, UE_STRING_CLASS&& Rhs);
 
 public:
-	UE_NODISCARD FORCEINLINE friend UE_STRING_CLASS operator+(const UE_STRING_CLASS& Lhs, const UE_STRING_CLASS& Rhs)	{ return ConcatFF(Lhs, Rhs); }
-	UE_NODISCARD FORCEINLINE friend UE_STRING_CLASS operator+(UE_STRING_CLASS&& Lhs, const UE_STRING_CLASS& Rhs)		{ return ConcatFF(MoveTemp(Lhs), Rhs); }
-	UE_NODISCARD FORCEINLINE friend UE_STRING_CLASS operator+(const UE_STRING_CLASS& Lhs, UE_STRING_CLASS&& Rhs)		{ return ConcatFF(Lhs,MoveTemp(Rhs)); }
-	UE_NODISCARD FORCEINLINE friend UE_STRING_CLASS operator+(UE_STRING_CLASS&& Lhs, UE_STRING_CLASS&& Rhs)				{ return ConcatFF(MoveTemp(Lhs), MoveTemp(Rhs)); }
-	UE_NODISCARD FORCEINLINE friend UE_STRING_CLASS operator+(const ElementType* Lhs, const UE_STRING_CLASS& Rhs)		{ return ConcatCF(Lhs, Rhs); }
-	UE_NODISCARD FORCEINLINE friend UE_STRING_CLASS operator+(const ElementType* Lhs, UE_STRING_CLASS&& Rhs)			{ return ConcatCF(Lhs, MoveTemp(Rhs)); }
-	UE_NODISCARD FORCEINLINE friend UE_STRING_CLASS operator+(const UE_STRING_CLASS& Lhs, const ElementType* Rhs)		{ return ConcatFC(Lhs, Rhs); }
-	UE_NODISCARD FORCEINLINE friend UE_STRING_CLASS operator+(UE_STRING_CLASS&& Lhs, const ElementType* Rhs)			{ return ConcatFC(MoveTemp(Lhs), Rhs); }
+	[[nodiscard]] FORCEINLINE friend UE_STRING_CLASS operator+(const UE_STRING_CLASS& Lhs, const UE_STRING_CLASS& Rhs)	{ return ConcatFF(Lhs, Rhs); }
+	[[nodiscard]] FORCEINLINE friend UE_STRING_CLASS operator+(UE_STRING_CLASS&& Lhs, const UE_STRING_CLASS& Rhs)		{ return ConcatFF(MoveTemp(Lhs), Rhs); }
+	[[nodiscard]] FORCEINLINE friend UE_STRING_CLASS operator+(const UE_STRING_CLASS& Lhs, UE_STRING_CLASS&& Rhs)		{ return ConcatFF(Lhs,MoveTemp(Rhs)); }
+	[[nodiscard]] FORCEINLINE friend UE_STRING_CLASS operator+(UE_STRING_CLASS&& Lhs, UE_STRING_CLASS&& Rhs)				{ return ConcatFF(MoveTemp(Lhs), MoveTemp(Rhs)); }
+	[[nodiscard]] FORCEINLINE friend UE_STRING_CLASS operator+(const ElementType* Lhs, const UE_STRING_CLASS& Rhs)		{ return ConcatCF(Lhs, Rhs); }
+	[[nodiscard]] FORCEINLINE friend UE_STRING_CLASS operator+(const ElementType* Lhs, UE_STRING_CLASS&& Rhs)			{ return ConcatCF(Lhs, MoveTemp(Rhs)); }
+	[[nodiscard]] FORCEINLINE friend UE_STRING_CLASS operator+(const UE_STRING_CLASS& Lhs, const ElementType* Rhs)		{ return ConcatFC(Lhs, Rhs); }
+	[[nodiscard]] FORCEINLINE friend UE_STRING_CLASS operator+(UE_STRING_CLASS&& Lhs, const ElementType* Rhs)			{ return ConcatFC(MoveTemp(Lhs), Rhs); }
 
 	template <
 		typename CharRangeType,
@@ -644,7 +644,7 @@ public:
 			!std::is_base_of_v<UE_STRING_CLASS, std::decay_t<CharRangeType>>
 		)
 	>
-	UE_NODISCARD FORCEINLINE friend UE_STRING_CLASS operator+(CharRangeType&& Lhs, const UE_STRING_CLASS& Rhs)
+	[[nodiscard]] FORCEINLINE friend UE_STRING_CLASS operator+(CharRangeType&& Lhs, const UE_STRING_CLASS& Rhs)
 	{
 		return ConcatRF(GetData(Lhs), GetNum(Lhs), Rhs);
 	}
@@ -658,7 +658,7 @@ public:
 			!std::is_base_of_v<UE_STRING_CLASS, std::decay_t<CharRangeType>>
 		)
 	>
-	UE_NODISCARD FORCEINLINE friend UE_STRING_CLASS operator+(CharRangeType&& Lhs, UE_STRING_CLASS&& Rhs)
+	[[nodiscard]] FORCEINLINE friend UE_STRING_CLASS operator+(CharRangeType&& Lhs, UE_STRING_CLASS&& Rhs)
 	{
 		return ConcatRF(GetData(Lhs), GetNum(Lhs), MoveTemp(Rhs));
 	}
@@ -672,7 +672,7 @@ public:
 			!std::is_base_of_v<UE_STRING_CLASS, std::decay_t<CharRangeType>>
 		)
 	>
-	UE_NODISCARD FORCEINLINE friend UE_STRING_CLASS operator+(const UE_STRING_CLASS& Lhs, CharRangeType&& Rhs)
+	[[nodiscard]] FORCEINLINE friend UE_STRING_CLASS operator+(const UE_STRING_CLASS& Lhs, CharRangeType&& Rhs)
 	{
 		return ConcatFR(Lhs, GetData(Rhs), GetNum(Rhs));
 	}
@@ -686,7 +686,7 @@ public:
 			!std::is_base_of_v<UE_STRING_CLASS, std::decay_t<CharRangeType>>
 		)
 	>
-	UE_NODISCARD FORCEINLINE friend UE_STRING_CLASS operator+(UE_STRING_CLASS&& Lhs, CharRangeType&& Rhs)
+	[[nodiscard]] FORCEINLINE friend UE_STRING_CLASS operator+(UE_STRING_CLASS&& Lhs, CharRangeType&& Rhs)
 	{
 		return ConcatFR(MoveTemp(Lhs), GetData(Rhs), GetNum(Rhs));
 	}
@@ -750,7 +750,7 @@ public:
 	 * @param Rhs Path to concatenate.
 	 * @return The new concatenated path
 	 */
-	UE_NODISCARD FORCEINLINE friend UE_STRING_CLASS operator/(const UE_STRING_CLASS& Lhs, const ElementType* Rhs)
+	[[nodiscard]] FORCEINLINE friend UE_STRING_CLASS operator/(const UE_STRING_CLASS& Lhs, const ElementType* Rhs)
 	{
 		checkSlow(Rhs);
 
@@ -768,7 +768,7 @@ public:
 	 * @param Rhs Path to concatenate.
 	 * @return The new concatenated path
 	 */
-	UE_NODISCARD FORCEINLINE friend UE_STRING_CLASS operator/(UE_STRING_CLASS&& Lhs, const ElementType* Rhs)
+	[[nodiscard]] FORCEINLINE friend UE_STRING_CLASS operator/(UE_STRING_CLASS&& Lhs, const ElementType* Rhs)
 	{
 		checkSlow(Rhs);
 
@@ -786,7 +786,7 @@ public:
 	 * @param Rhs Path to concatenate.
 	 * @return The new concatenated path
 	 */
-	UE_NODISCARD FORCEINLINE friend UE_STRING_CLASS operator/(const UE_STRING_CLASS& Lhs, const UE_STRING_CLASS& Rhs)
+	[[nodiscard]] FORCEINLINE friend UE_STRING_CLASS operator/(const UE_STRING_CLASS& Lhs, const UE_STRING_CLASS& Rhs)
 	{
 		int32 StrLength = Rhs.Len();
 
@@ -802,7 +802,7 @@ public:
 	 * @param Rhs Path to concatenate.
 	 * @return The new concatenated path
 	 */
-	UE_NODISCARD FORCEINLINE friend UE_STRING_CLASS operator/(UE_STRING_CLASS&& Lhs, const UE_STRING_CLASS& Rhs)
+	[[nodiscard]] FORCEINLINE friend UE_STRING_CLASS operator/(UE_STRING_CLASS&& Lhs, const UE_STRING_CLASS& Rhs)
 	{
 		int32 StrLength = Rhs.Len();
 
@@ -818,7 +818,7 @@ public:
 	 * @param Rhs Path to concatenate.
 	 * @return new string of the path
 	 */
-	UE_NODISCARD FORCEINLINE friend UE_STRING_CLASS operator/(const ElementType* Lhs, const UE_STRING_CLASS& Rhs)
+	[[nodiscard]] FORCEINLINE friend UE_STRING_CLASS operator/(const ElementType* Lhs, const UE_STRING_CLASS& Rhs)
 	{
 		int32 StrLength = Rhs.Len();
 
@@ -835,7 +835,7 @@ public:
 	 * @return true if the left string is lexicographically <= the right string, otherwise false
 	 * @note case insensitive
 	 */
-	UE_NODISCARD FORCEINLINE friend bool operator<=(const UE_STRING_CLASS& Lhs, const UE_STRING_CLASS& Rhs)
+	[[nodiscard]] FORCEINLINE friend bool operator<=(const UE_STRING_CLASS& Lhs, const UE_STRING_CLASS& Rhs)
 	{
 		return FPlatformString::Stricmp(*Lhs, *Rhs) <= 0;
 	}
@@ -849,7 +849,7 @@ public:
 	 * @note case insensitive
 	 */
 	template <typename CharType>
-	UE_NODISCARD FORCEINLINE friend bool operator<=(const UE_STRING_CLASS& Lhs, const CharType* Rhs)
+	[[nodiscard]] FORCEINLINE friend bool operator<=(const UE_STRING_CLASS& Lhs, const CharType* Rhs)
 	{
 		return FPlatformString::Stricmp(*Lhs, Rhs) <= 0;
 	}
@@ -863,7 +863,7 @@ public:
 	 * @note case insensitive
 	 */
 	template <typename CharType>
-	UE_NODISCARD FORCEINLINE friend bool operator<=(const CharType* Lhs, const UE_STRING_CLASS& Rhs)
+	[[nodiscard]] FORCEINLINE friend bool operator<=(const CharType* Lhs, const UE_STRING_CLASS& Rhs)
 	{
 		return FPlatformString::Stricmp(Lhs, *Rhs) <= 0;
 	}
@@ -876,7 +876,7 @@ public:
 	 * @return true if the left string is lexicographically < the right string, otherwise false
 	 * @note case insensitive
 	 */
-	UE_NODISCARD FORCEINLINE friend bool operator<(const UE_STRING_CLASS& Lhs, const UE_STRING_CLASS& Rhs)
+	[[nodiscard]] FORCEINLINE friend bool operator<(const UE_STRING_CLASS& Lhs, const UE_STRING_CLASS& Rhs)
 	{
 		return FPlatformString::Stricmp(*Lhs, *Rhs) < 0;
 	}
@@ -890,7 +890,7 @@ public:
 	 * @note case insensitive
 	 */
 	template <typename CharType>
-	UE_NODISCARD FORCEINLINE friend bool operator<(const UE_STRING_CLASS& Lhs, const CharType* Rhs)
+	[[nodiscard]] FORCEINLINE friend bool operator<(const UE_STRING_CLASS& Lhs, const CharType* Rhs)
 	{
 		return FPlatformString::Stricmp(*Lhs, Rhs) < 0;
 	}
@@ -904,7 +904,7 @@ public:
 	 * @note case insensitive
 	 */
 	template <typename CharType>
-	UE_NODISCARD FORCEINLINE friend bool operator<(const CharType* Lhs, const UE_STRING_CLASS& Rhs)
+	[[nodiscard]] FORCEINLINE friend bool operator<(const CharType* Lhs, const UE_STRING_CLASS& Rhs)
 	{
 		return FPlatformString::Stricmp(Lhs, *Rhs) < 0;
 	}
@@ -917,7 +917,7 @@ public:
 	 * @return true if the left string is lexicographically >= the right string, otherwise false
 	 * @note case insensitive
 	 */
-	UE_NODISCARD FORCEINLINE friend bool operator>=(const UE_STRING_CLASS& Lhs, const UE_STRING_CLASS& Rhs)
+	[[nodiscard]] FORCEINLINE friend bool operator>=(const UE_STRING_CLASS& Lhs, const UE_STRING_CLASS& Rhs)
 	{
 		return FPlatformString::Stricmp(*Lhs, *Rhs) >= 0;
 	}
@@ -931,7 +931,7 @@ public:
 	 * @note case insensitive
 	 */
 	template <typename CharType>
-	UE_NODISCARD FORCEINLINE friend bool operator>=(const UE_STRING_CLASS& Lhs, const CharType* Rhs)
+	[[nodiscard]] FORCEINLINE friend bool operator>=(const UE_STRING_CLASS& Lhs, const CharType* Rhs)
 	{
 		return FPlatformString::Stricmp(*Lhs, Rhs) >= 0;
 	}
@@ -945,7 +945,7 @@ public:
 	 * @note case insensitive
 	 */
 	template <typename CharType>
-	UE_NODISCARD FORCEINLINE friend bool operator>=(const CharType* Lhs, const UE_STRING_CLASS& Rhs)
+	[[nodiscard]] FORCEINLINE friend bool operator>=(const CharType* Lhs, const UE_STRING_CLASS& Rhs)
 	{
 		return FPlatformString::Stricmp(Lhs, *Rhs) >= 0;
 	}
@@ -958,7 +958,7 @@ public:
 	 * @return true if the left string is lexicographically > the right string, otherwise false
 	 * @note case insensitive
 	 */
-	UE_NODISCARD FORCEINLINE friend bool operator>(const UE_STRING_CLASS& Lhs, const UE_STRING_CLASS& Rhs)
+	[[nodiscard]] FORCEINLINE friend bool operator>(const UE_STRING_CLASS& Lhs, const UE_STRING_CLASS& Rhs)
 	{
 		return FPlatformString::Stricmp(*Lhs, *Rhs) > 0;
 	}
@@ -972,7 +972,7 @@ public:
 	 * @note case insensitive
 	 */
 	template <typename CharType>
-	UE_NODISCARD FORCEINLINE friend bool operator>(const UE_STRING_CLASS& Lhs, const CharType* Rhs)
+	[[nodiscard]] FORCEINLINE friend bool operator>(const UE_STRING_CLASS& Lhs, const CharType* Rhs)
 	{
 		return FPlatformString::Stricmp(*Lhs, Rhs) > 0;
 	}
@@ -986,7 +986,7 @@ public:
 	 * @note case insensitive
 	 */
 	template <typename CharType>
-	UE_NODISCARD FORCEINLINE friend bool operator>(const CharType* Lhs, const UE_STRING_CLASS& Rhs)
+	[[nodiscard]] FORCEINLINE friend bool operator>(const CharType* Lhs, const UE_STRING_CLASS& Rhs)
 	{
 		return FPlatformString::Stricmp(Lhs, *Rhs) > 0;
 	}
@@ -999,7 +999,7 @@ public:
 	 * @return true if the left string is lexicographically == the right string, otherwise false
 	 * @note case insensitive
 	 */
-	UE_NODISCARD FORCEINLINE bool operator==(const UE_STRING_CLASS& Rhs) const
+	[[nodiscard]] FORCEINLINE bool operator==(const UE_STRING_CLASS& Rhs) const
 	{
 		return Equals(Rhs, ESearchCase::IgnoreCase);
 	}
@@ -1013,7 +1013,7 @@ public:
 	 * @note case insensitive
 	 */
 	template <typename CharType>
-	UE_NODISCARD FORCEINLINE bool operator==(const CharType* Rhs) const
+	[[nodiscard]] FORCEINLINE bool operator==(const CharType* Rhs) const
 	{
 		return FPlatformString::Stricmp(**this, Rhs) == 0;
 	}
@@ -1027,7 +1027,7 @@ public:
 	 * @note case insensitive
 	 */
 	template <typename CharType>
-	UE_NODISCARD FORCEINLINE friend bool operator==(const CharType* Lhs, const UE_STRING_CLASS& Rhs)
+	[[nodiscard]] FORCEINLINE friend bool operator==(const CharType* Lhs, const UE_STRING_CLASS& Rhs)
 	{
 		return FPlatformString::Stricmp(Lhs, *Rhs) == 0;
 	}
@@ -1040,7 +1040,7 @@ public:
 	 * @return true if the left string is lexicographically != the right string, otherwise false
 	 * @note case insensitive
 	 */
-	UE_NODISCARD FORCEINLINE bool operator!=(const UE_STRING_CLASS& Rhs) const
+	[[nodiscard]] FORCEINLINE bool operator!=(const UE_STRING_CLASS& Rhs) const
 	{
 		return !Equals(Rhs, ESearchCase::IgnoreCase);
 	}
@@ -1054,7 +1054,7 @@ public:
 	 * @note case insensitive
 	 */
 	template <typename CharType>
-	UE_NODISCARD FORCEINLINE bool operator!=(const CharType* Rhs) const
+	[[nodiscard]] FORCEINLINE bool operator!=(const CharType* Rhs) const
 	{
 		return FPlatformString::Stricmp(**this, Rhs) != 0;
 	}
@@ -1068,24 +1068,24 @@ public:
 	 * @note case insensitive
 	 */
 	template <typename CharType>
-	UE_NODISCARD FORCEINLINE friend bool operator!=(const CharType* Lhs, const UE_STRING_CLASS& Rhs)
+	[[nodiscard]] FORCEINLINE friend bool operator!=(const CharType* Lhs, const UE_STRING_CLASS& Rhs)
 	{
 		return FPlatformString::Stricmp(Lhs, *Rhs) != 0;
 	}
 
 	/** Get the length of the string, excluding terminating character */
-	UE_NODISCARD FORCEINLINE int32 Len() const
+	[[nodiscard]] FORCEINLINE int32 Len() const
 	{
 		return Data.Num() ? Data.Num() - 1 : 0;
 	}
 
 	/** Returns the left most given number of characters */
-	UE_NODISCARD FORCEINLINE UE_STRING_CLASS Left( int32 Count ) const &
+	[[nodiscard]] FORCEINLINE UE_STRING_CLASS Left( int32 Count ) const &
 	{
 		return UE_STRING_CLASS::ConstructFromPtrSize(**this, FMath::Clamp(Count,0,Len()) );
 	}
 
-	UE_NODISCARD FORCEINLINE UE_STRING_CLASS Left(int32 Count) &&
+	[[nodiscard]] FORCEINLINE UE_STRING_CLASS Left(int32 Count) &&
 	{
 		LeftInline(Count, false);
 		return MoveTemp(*this);
@@ -1100,13 +1100,13 @@ public:
 	}
 
 	/** Returns the left most characters from the string chopping the given number of characters from the end */
-	UE_NODISCARD FORCEINLINE UE_STRING_CLASS LeftChop( int32 Count ) const &
+	[[nodiscard]] FORCEINLINE UE_STRING_CLASS LeftChop( int32 Count ) const &
 	{
 		const int32 Length = Len();
 		return UE_STRING_CLASS::ConstructFromPtrSize( **this, FMath::Clamp(Length-Count,0, Length) );
 	}
 
-	UE_NODISCARD FORCEINLINE UE_STRING_CLASS LeftChop(int32 Count)&&
+	[[nodiscard]] FORCEINLINE UE_STRING_CLASS LeftChop(int32 Count)&&
 	{
 		LeftChopInline(Count, false);
 		return MoveTemp(*this);
@@ -1120,13 +1120,13 @@ public:
 	}
 
 	/** Returns the string to the right of the specified location, counting back from the right (end of the word). */
-	UE_NODISCARD FORCEINLINE UE_STRING_CLASS Right( int32 Count ) const &
+	[[nodiscard]] FORCEINLINE UE_STRING_CLASS Right( int32 Count ) const &
 	{
 		const int32 Length = Len();
 		return UE_STRING_CLASS( **this + Length-FMath::Clamp(Count,0,Length) );
 	}
 
-	UE_NODISCARD FORCEINLINE UE_STRING_CLASS Right(int32 Count) &&
+	[[nodiscard]] FORCEINLINE UE_STRING_CLASS Right(int32 Count) &&
 	{
 		RightInline(Count, false);
 		return MoveTemp(*this);
@@ -1140,9 +1140,9 @@ public:
 	}
 
 	/** Returns the string to the right of the specified location, counting forward from the left (from the beginning of the word). */
-	UE_NODISCARD CORE_API UE_STRING_CLASS RightChop( int32 Count ) const &;
+	[[nodiscard]] CORE_API UE_STRING_CLASS RightChop( int32 Count ) const &;
 
-	UE_NODISCARD FORCEINLINE UE_STRING_CLASS RightChop(int32 Count) &&
+	[[nodiscard]] FORCEINLINE UE_STRING_CLASS RightChop(int32 Count) &&
 	{
 		RightChopInline(Count, false);
 		return MoveTemp(*this);
@@ -1155,12 +1155,12 @@ public:
 	}
 
 	/** Returns the substring from Start position for Count characters. */
-	UE_NODISCARD CORE_API UE_STRING_CLASS Mid(int32 Start, int32 Count) const &;
-	UE_NODISCARD CORE_API UE_STRING_CLASS Mid(int32 Start, int32 Count) &&;
+	[[nodiscard]] CORE_API UE_STRING_CLASS Mid(int32 Start, int32 Count) const &;
+	[[nodiscard]] CORE_API UE_STRING_CLASS Mid(int32 Start, int32 Count) &&;
 
 	/** Returns the substring from Start position to the end */
-	UE_NODISCARD FORCEINLINE UE_STRING_CLASS Mid(int32 Start) const & { return RightChop(Start); }
-	UE_NODISCARD FORCEINLINE UE_STRING_CLASS Mid(int32 Start) && { return ((UE_STRING_CLASS&&)*this).RightChop(Start); }
+	[[nodiscard]] FORCEINLINE UE_STRING_CLASS Mid(int32 Start) const & { return RightChop(Start); }
+	[[nodiscard]] FORCEINLINE UE_STRING_CLASS Mid(int32 Start) && { return ((UE_STRING_CLASS&&)*this).RightChop(Start); }
 
 	/** Modifies the string such that it is now the substring from Start position for Count characters. */
 	FORCEINLINE void MidInline(int32 Start, int32 Count = MAX_int32, bool bAllowShrinking = true)
@@ -1198,7 +1198,7 @@ public:
 			!std::is_base_of_v<UE_STRING_CLASS, std::decay_t<CharRangeType>>
 		)
 	>
-	UE_NODISCARD int32 Find(CharRangeType&& SubStr, ESearchCase::Type SearchCase = ESearchCase::IgnoreCase,
+	[[nodiscard]] int32 Find(CharRangeType&& SubStr, ESearchCase::Type SearchCase = ESearchCase::IgnoreCase,
 		ESearchDir::Type SearchDir = ESearchDir::FromStart, int32 StartPosition = INDEX_NONE) const
 	{
 		static_assert(std::is_same_v<CharRangeElementType, ElementType>, "Expected a range of ElementType");
@@ -1214,7 +1214,7 @@ public:
 	 * @param SearchCase		Indicates whether the search is case sensitive or not
 	 * @param SearchDir			Indicates whether the search starts at the beginning or at the end.
 	 */
-	UE_NODISCARD int32 Find(const ElementType* SubStr, ESearchCase::Type SearchCase = ESearchCase::IgnoreCase,
+	[[nodiscard]] int32 Find(const ElementType* SubStr, ESearchCase::Type SearchCase = ESearchCase::IgnoreCase,
 		ESearchDir::Type SearchDir = ESearchDir::FromStart, int32 StartPosition = INDEX_NONE) const
 	{
 		return SubStr ? Find(SubStr, TCString<ElementType>::Strlen(SubStr), SearchCase, SearchDir, StartPosition) : INDEX_NONE;
@@ -1236,7 +1236,7 @@ public:
 	 *
 	 *        Consider using UE::String::FindLast() as an alternative.
 	 */
-	UE_NODISCARD int32 Find(const UE_STRING_CLASS& SubStr, ESearchCase::Type SearchCase = ESearchCase::IgnoreCase,
+	[[nodiscard]] int32 Find(const UE_STRING_CLASS& SubStr, ESearchCase::Type SearchCase = ESearchCase::IgnoreCase,
 		ESearchDir::Type SearchDir = ESearchDir::FromStart, int32 StartPosition = INDEX_NONE) const
 	{
 		return Find(*SubStr, SubStr.Len(), SearchCase, SearchDir, StartPosition);
@@ -1259,7 +1259,7 @@ public:
 	 *
 	 *        Consider using UE::String::FindLast() as an alternative.
 	 */
-	UE_NODISCARD CORE_API int32 Find(const ElementType* SubStr, int32 InSubStrLen, ESearchCase::Type SearchCase = ESearchCase::IgnoreCase,
+	[[nodiscard]] CORE_API int32 Find(const ElementType* SubStr, int32 InSubStrLen, ESearchCase::Type SearchCase = ESearchCase::IgnoreCase,
 		ESearchDir::Type SearchDir = ESearchDir::FromStart, int32 StartPosition = INDEX_NONE) const;
 
 
@@ -1281,7 +1281,7 @@ public:
 			!std::is_base_of_v<UE_STRING_CLASS, std::decay_t<CharRangeType>>
 		)
 	>
-	UE_NODISCARD FORCEINLINE bool Contains(CharRangeType&& SubStr, ESearchCase::Type SearchCase = ESearchCase::IgnoreCase,
+	[[nodiscard]] FORCEINLINE bool Contains(CharRangeType&& SubStr, ESearchCase::Type SearchCase = ESearchCase::IgnoreCase,
 		ESearchDir::Type SearchDir = ESearchDir::FromStart) const
 	{
 		static_assert(std::is_same_v<CharRangeElementType, ElementType>, "Expected a range of characters");
@@ -1296,7 +1296,7 @@ public:
 	 * @param SearchDir			Indicates whether the search starts at the beginning or at the end ( defaults to ESearchDir::FromStart )
 	 * @return					Returns whether the string contains the substring. If the substring is empty, returns true.
 	 **/
-	UE_NODISCARD FORCEINLINE bool Contains(const ElementType* SubStr, ESearchCase::Type SearchCase = ESearchCase::IgnoreCase,
+	[[nodiscard]] FORCEINLINE bool Contains(const ElementType* SubStr, ESearchCase::Type SearchCase = ESearchCase::IgnoreCase,
 		ESearchDir::Type SearchDir = ESearchDir::FromStart) const
 	{
 		return Find(SubStr, SearchCase, SearchDir) != INDEX_NONE;
@@ -1310,7 +1310,7 @@ public:
 	 * @param SearchDir			Indicates whether the search starts at the beginning or at the end ( defaults to ESearchDir::FromStart )
 	 * @return					Returns whether the string contains the substring. If the substring is empty, returns true.
 	 **/
-	UE_NODISCARD FORCEINLINE bool Contains(const UE_STRING_CLASS& SubStr, ESearchCase::Type SearchCase = ESearchCase::IgnoreCase,
+	[[nodiscard]] FORCEINLINE bool Contains(const UE_STRING_CLASS& SubStr, ESearchCase::Type SearchCase = ESearchCase::IgnoreCase,
 							  ESearchDir::Type SearchDir = ESearchDir::FromStart ) const
 	{
 		return Find(*SubStr, SearchCase, SearchDir) != INDEX_NONE;
@@ -1325,7 +1325,7 @@ public:
 	 * @param SearchDir			Indicates whether the search starts at the beginning or at the end ( defaults to ESearchDir::FromStart )
 	 * @return					Returns whether the string contains the substring. If the substring is empty, returns true.
 	 **/
-	UE_NODISCARD FORCEINLINE bool Contains(const ElementType* SubStr, int32 SubStrLen,
+	[[nodiscard]] FORCEINLINE bool Contains(const ElementType* SubStr, int32 SubStrLen,
 		ESearchCase::Type SearchCase = ESearchCase::IgnoreCase, ESearchDir::Type SearchDir = ESearchDir::FromStart) const
 	{
 		return Find(SubStr, SubStrLen, SearchCase, SearchDir) != INDEX_NONE;
@@ -1391,7 +1391,7 @@ public:
 	 * @param SearchCase 	Whether or not the comparison should ignore case
 	 * @return true if this string is lexicographically equivalent to the other, otherwise false
 	 */
-	UE_NODISCARD FORCEINLINE bool Equals(const UE_STRING_CLASS& Other, ESearchCase::Type SearchCase = ESearchCase::CaseSensitive) const
+	[[nodiscard]] FORCEINLINE bool Equals(const UE_STRING_CLASS& Other, ESearchCase::Type SearchCase = ESearchCase::CaseSensitive) const
 	{
 		int32 Num = Data.Num();
 		int32 OtherNum = Other.Data.Num();
@@ -1423,7 +1423,7 @@ public:
 	 * @param SearchCase 	Whether or not the comparison should ignore case
 	 * @return 0 if equal, negative if less than, positive if greater than
 	 */
-	UE_NODISCARD FORCEINLINE int32 Compare( const UE_STRING_CLASS& Other, ESearchCase::Type SearchCase = ESearchCase::CaseSensitive ) const
+	[[nodiscard]] FORCEINLINE int32 Compare( const UE_STRING_CLASS& Other, ESearchCase::Type SearchCase = ESearchCase::CaseSensitive ) const
 	{
 		if( SearchCase == ESearchCase::CaseSensitive )
 		{
@@ -1452,37 +1452,37 @@ public:
 	CORE_API bool Split(const UE_STRING_CLASS& InS, UE_STRING_CLASS* LeftS, UE_STRING_CLASS* RightS) const;
 
 	/** Returns a new string with the characters of this converted to uppercase */
-	UE_NODISCARD CORE_API UE_STRING_CLASS ToUpper() const &;
+	[[nodiscard]] CORE_API UE_STRING_CLASS ToUpper() const &;
 
 	/**
 	 * Converts all characters in this rvalue string to uppercase and moves it into the returned string.
 	 * @return a new string with the characters of this converted to uppercase
 	 */
-	UE_NODISCARD CORE_API UE_STRING_CLASS ToUpper() &&;
+	[[nodiscard]] CORE_API UE_STRING_CLASS ToUpper() &&;
 
 	/** Converts all characters in this string to uppercase */
 	CORE_API void ToUpperInline();
 
 	/** Returns a new string with the characters of this converted to lowercase */
-	UE_NODISCARD CORE_API UE_STRING_CLASS ToLower() const &;
+	[[nodiscard]] CORE_API UE_STRING_CLASS ToLower() const &;
 
 	/**
 	 * Converts all characters in this rvalue string to lowercase and moves it into the returned string.
 	 * @return a new string with the characters of this converted to lowercase
 	 */
-	UE_NODISCARD CORE_API UE_STRING_CLASS ToLower() &&;
+	[[nodiscard]] CORE_API UE_STRING_CLASS ToLower() &&;
 
 	/** Converts all characters in this string to lowercase */
 	CORE_API void ToLowerInline();
 
 	/** Pad the left of this string for ChCount characters */
-	UE_NODISCARD CORE_API UE_STRING_CLASS LeftPad( int32 ChCount ) const;
+	[[nodiscard]] CORE_API UE_STRING_CLASS LeftPad( int32 ChCount ) const;
 
 	/** Pad the right of this string for ChCount characters */
-	UE_NODISCARD CORE_API UE_STRING_CLASS RightPad( int32 ChCount ) const;
+	[[nodiscard]] CORE_API UE_STRING_CLASS RightPad( int32 ChCount ) const;
 	
 	/** Returns true if the string only contains numeric characters */
-	UE_NODISCARD CORE_API bool IsNumeric() const;
+	[[nodiscard]] CORE_API bool IsNumeric() const;
 	
 	/** Removes spaces from the string.  I.E. "Spaces Are Cool" --> "SpacesAreCool". */
 	CORE_API void RemoveSpacesInline();
@@ -1496,7 +1496,7 @@ public:
 	 * @returns A string that was constructed using format and additional parameters.
 	 */
 	template <typename FmtType, typename... Types>
-	UE_NODISCARD static UE_STRING_CLASS Printf(const FmtType& Fmt, Types... Args)
+	[[nodiscard]] static UE_STRING_CLASS Printf(const FmtType& Fmt, Types... Args)
 	{
 		static_assert(TIsArrayOrRefOfTypeByPredicate<FmtType, TIsCharEncodingCompatibleWithElementType>::Value, "Formatting string must be a character array.");
 		static_assert((TIsValidVariadicFunctionArg<Types>::Value && ...), "Invalid argument(s) passed to Printf");
@@ -1529,7 +1529,7 @@ public:
 	 * @param InNamedArguments		A map of named arguments that match the tokens specified in InExpression
 	 * @return A string containing the formatted text
 	 */
-	UE_NODISCARD static CORE_API UE_STRING_CLASS Format(const ElementType* InFormatString, const PREPROCESSOR_JOIN(UE_STRING_CLASS, FormatNamedArguments)& InNamedArguments);
+	[[nodiscard]] static CORE_API UE_STRING_CLASS Format(const ElementType* InFormatString, const PREPROCESSOR_JOIN(UE_STRING_CLASS, FormatNamedArguments)& InNamedArguments);
 
 	/**
 	 * Format the specified string using the specified arguments. Replaces instances of {0} with indices from the given array matching the index specified in the token
@@ -1537,10 +1537,10 @@ public:
 	 * @param InOrderedArguments	An array of ordered arguments that match the tokens specified in InExpression
 	 * @return A string containing the formatted text
 	 */
-	UE_NODISCARD static CORE_API UE_STRING_CLASS Format(const ElementType* InFormatString, const PREPROCESSOR_JOIN(UE_STRING_CLASS, FormatOrderedArguments)& InOrderedArguments);
+	[[nodiscard]] static CORE_API UE_STRING_CLASS Format(const ElementType* InFormatString, const PREPROCESSOR_JOIN(UE_STRING_CLASS, FormatOrderedArguments)& InOrderedArguments);
 
 	/** Returns a string containing only the Ch character */
-	UE_NODISCARD static CORE_API UE_STRING_CLASS Chr( ElementType Ch );
+	[[nodiscard]] static CORE_API UE_STRING_CLASS Chr( ElementType Ch );
 
 	/**
 	 * Returns a string that is full of a variable number of characters
@@ -1550,7 +1550,7 @@ public:
 	 * 
 	 * @return The string of NumCharacters characters.
 	 */
-	UE_NODISCARD static CORE_API UE_STRING_CLASS ChrN( int32 NumCharacters, ElementType Char );
+	[[nodiscard]] static CORE_API UE_STRING_CLASS ChrN( int32 NumCharacters, ElementType Char );
 
 	/**
 	 * Serializes the string.
@@ -1579,7 +1579,7 @@ public:
 			!std::is_base_of_v<UE_STRING_CLASS, std::decay_t<CharRangeType>>
 		)
 	>
-	UE_NODISCARD bool StartsWith(CharRangeType&& InPrefix, ESearchCase::Type SearchCase = ESearchCase::IgnoreCase) const
+	[[nodiscard]] bool StartsWith(CharRangeType&& InPrefix, ESearchCase::Type SearchCase = ESearchCase::IgnoreCase) const
 	{
 		static_assert(std::is_same_v<CharRangeElementType, ElementType>, "Expected a range of characters");
 		return StartsWith(GetData(InPrefix), GetNum(InPrefix), SearchCase);
@@ -1591,7 +1591,7 @@ public:
 	 * @param SearchCase		Indicates whether the search is case sensitive or not ( defaults to ESearchCase::IgnoreCase )
 	 * @return true if this string begins with specified text, false otherwise
 	 */
-	UE_NODISCARD bool StartsWith(const ElementType* InPrefix, ESearchCase::Type SearchCase = ESearchCase::IgnoreCase) const
+	[[nodiscard]] bool StartsWith(const ElementType* InPrefix, ESearchCase::Type SearchCase = ESearchCase::IgnoreCase) const
 	{
 		return StartsWith(InPrefix, InPrefix ? TCString<ElementType>::Strlen(InPrefix) : 0, SearchCase);
 	}
@@ -1602,7 +1602,7 @@ public:
 	 * @param SearchCase		Indicates whether the search is case sensitive or not ( defaults to ESearchCase::IgnoreCase )
 	 * @return true if this string begins with specified text, false otherwise
 	 */
-	UE_NODISCARD bool StartsWith(const UE_STRING_CLASS& InPrefix, ESearchCase::Type SearchCase = ESearchCase::IgnoreCase) const
+	[[nodiscard]] bool StartsWith(const UE_STRING_CLASS& InPrefix, ESearchCase::Type SearchCase = ESearchCase::IgnoreCase) const
 	{
 		return StartsWith(*InPrefix, InPrefix.Len(), SearchCase);
 	}
@@ -1613,7 +1613,7 @@ public:
 	 * @param SearchCase		Indicates whether the search is case sensitive or not ( defaults to ESearchCase::IgnoreCase )
 	 * @return true if this string begins with specified text, false otherwise
 	 */
-	UE_NODISCARD CORE_API bool StartsWith(const ElementType* InPrefix, int32 InPrefixLen, ESearchCase::Type SearchCase = ESearchCase::IgnoreCase) const;
+	[[nodiscard]] CORE_API bool StartsWith(const ElementType* InPrefix, int32 InPrefixLen, ESearchCase::Type SearchCase = ESearchCase::IgnoreCase) const;
 
 	/**
 	 * Test whether this string ends with given suffix.
@@ -1631,7 +1631,7 @@ public:
 			!std::is_base_of_v<UE_STRING_CLASS, std::decay_t<CharRangeType>>
 		)
 	>
-	UE_NODISCARD bool EndsWith(CharRangeType&& InSuffix, ESearchCase::Type SearchCase = ESearchCase::IgnoreCase) const
+	[[nodiscard]] bool EndsWith(CharRangeType&& InSuffix, ESearchCase::Type SearchCase = ESearchCase::IgnoreCase) const
 	{
 		static_assert(std::is_same_v<CharRangeElementType, ElementType>, "Expected a range of characters");
 		return EndsWith(GetData(InSuffix), GetNum(InSuffix), SearchCase);
@@ -1643,7 +1643,7 @@ public:
 	 * @param SearchCase		Indicates whether the search is case sensitive or not ( defaults to ESearchCase::IgnoreCase )
 	 * @return true if this string ends with specified text, false otherwise
 	 */
-	UE_NODISCARD bool EndsWith(const ElementType* InSuffix, ESearchCase::Type SearchCase = ESearchCase::IgnoreCase) const
+	[[nodiscard]] bool EndsWith(const ElementType* InSuffix, ESearchCase::Type SearchCase = ESearchCase::IgnoreCase) const
 	{
 		return EndsWith(InSuffix, InSuffix ? TCString<ElementType>::Strlen(InSuffix) : 0, SearchCase);
 	}
@@ -1654,7 +1654,7 @@ public:
 	 * @param SearchCase		Indicates whether the search is case sensitive or not ( defaults to ESearchCase::IgnoreCase )
 	 * @return true if this string ends with specified text, false otherwise
 	 */
-	UE_NODISCARD bool EndsWith(const UE_STRING_CLASS& InSuffix, ESearchCase::Type SearchCase = ESearchCase::IgnoreCase) const
+	[[nodiscard]] bool EndsWith(const UE_STRING_CLASS& InSuffix, ESearchCase::Type SearchCase = ESearchCase::IgnoreCase) const
 	{
 		return EndsWith(*InSuffix, InSuffix.Len(), SearchCase);
 	}
@@ -1665,7 +1665,7 @@ public:
 	 * @param SearchCase		Indicates whether the search is case sensitive or not ( defaults to ESearchCase::IgnoreCase )
 	 * @return true if this string ends with specified text, false otherwise
 	 */
-	UE_NODISCARD CORE_API bool EndsWith(const ElementType* InSuffix, int32 InSuffixLen, ESearchCase::Type SearchCase = ESearchCase::IgnoreCase ) const;
+	[[nodiscard]] CORE_API bool EndsWith(const ElementType* InSuffix, int32 InSuffixLen, ESearchCase::Type SearchCase = ESearchCase::IgnoreCase ) const;
 
 	/**
 	 * Searches this string for a given wild card
@@ -1685,7 +1685,7 @@ public:
 			!std::is_base_of_v<UE_STRING_CLASS, std::decay_t<CharRangeType>>
 		)
 	>
-	UE_NODISCARD bool MatchesWildcard(CharRangeType&& Wildcard, ESearchCase::Type SearchCase = ESearchCase::IgnoreCase) const
+	[[nodiscard]] bool MatchesWildcard(CharRangeType&& Wildcard, ESearchCase::Type SearchCase = ESearchCase::IgnoreCase) const
 	{
 		static_assert(std::is_same_v<CharRangeElementType, ElementType>, "Expected a range of characters");
 		return MatchesWildcard(GetData(Wildcard), GetNum(Wildcard), SearchCase);
@@ -1699,7 +1699,7 @@ public:
 	 * @return true if this string matches the *?-type wildcard given.
 	 * @warning This is a simple, SLOW routine. Use with caution
 	 */
-	UE_NODISCARD bool MatchesWildcard(const ElementType* Wildcard, ESearchCase::Type SearchCase = ESearchCase::IgnoreCase) const
+	[[nodiscard]] bool MatchesWildcard(const ElementType* Wildcard, ESearchCase::Type SearchCase = ESearchCase::IgnoreCase) const
 	{
 		return MatchesWildcard(Wildcard, Wildcard ? TCString<ElementType>::Strlen(Wildcard) : 0, SearchCase);
 	}
@@ -1712,7 +1712,7 @@ public:
 	 * @return true if this string matches the *?-type wildcard given.
 	 * @warning This is a simple, SLOW routine. Use with caution
 	 */
-	UE_NODISCARD bool MatchesWildcard(const UE_STRING_CLASS& Wildcard, ESearchCase::Type SearchCase = ESearchCase::IgnoreCase) const
+	[[nodiscard]] bool MatchesWildcard(const UE_STRING_CLASS& Wildcard, ESearchCase::Type SearchCase = ESearchCase::IgnoreCase) const
 	{
 		return MatchesWildcard(*Wildcard, Wildcard.Len(), SearchCase);
 	}
@@ -1725,7 +1725,7 @@ public:
 	 * @return true if this string matches the *?-type wildcard given.
 	 * @warning This is a simple, SLOW routine. Use with caution
 	 */
-	UE_NODISCARD CORE_API bool MatchesWildcard(const ElementType* Wildcard, int32 WildcardLen, ESearchCase::Type SearchCase = ESearchCase::IgnoreCase) const;
+	[[nodiscard]] CORE_API bool MatchesWildcard(const ElementType* Wildcard, int32 WildcardLen, ESearchCase::Type SearchCase = ESearchCase::IgnoreCase) const;
 
 	/**
 	 * Removes whitespace characters from the start and end of this string. Modifies the string in-place.
@@ -1736,13 +1736,13 @@ public:
 	 * Removes whitespace characters from the start and end of this string.
 	 * @note Unlike Trim() this function returns a copy, and does not mutate the string.
 	 */
-	UE_NODISCARD CORE_API UE_STRING_CLASS TrimStartAndEnd() const &;
+	[[nodiscard]] CORE_API UE_STRING_CLASS TrimStartAndEnd() const &;
 
 	/**
 	 * Removes whitespace characters from the start and end of this string.
 	 * @note Unlike Trim() this function returns a copy, and does not mutate the string.
 	 */
-	UE_NODISCARD CORE_API UE_STRING_CLASS TrimStartAndEnd() &&;
+	[[nodiscard]] CORE_API UE_STRING_CLASS TrimStartAndEnd() &&;
 
 	/**
 	 * Removes whitespace characters from the start of this string. Modifies the string in-place.
@@ -1753,13 +1753,13 @@ public:
 	 * Removes whitespace characters from the start of this string.
 	 * @note Unlike Trim() this function returns a copy, and does not mutate the string.
 	 */
-	UE_NODISCARD CORE_API UE_STRING_CLASS TrimStart() const &;
+	[[nodiscard]] CORE_API UE_STRING_CLASS TrimStart() const &;
 
 	/**
 	 * Removes whitespace characters from the start of this string.
 	 * @note Unlike Trim() this function returns a copy, and does not mutate the string.
 	 */
-	UE_NODISCARD CORE_API UE_STRING_CLASS TrimStart() &&;
+	[[nodiscard]] CORE_API UE_STRING_CLASS TrimStart() &&;
 
 	/**
 	 * Removes whitespace characters from the end of this string. Modifies the string in-place.
@@ -1770,13 +1770,13 @@ public:
 	 * Removes whitespace characters from the end of this string.
 	 * @note Unlike TrimTrailing() this function returns a copy, and does not mutate the string.
 	 */
-	UE_NODISCARD CORE_API UE_STRING_CLASS TrimEnd() const &;
+	[[nodiscard]] CORE_API UE_STRING_CLASS TrimEnd() const &;
 
 	/**
 	 * Removes whitespace characters from the end of this string.
 	 * @note Unlike TrimTrailing() this function returns a copy, and does not mutate the string.
 	 */
-	UE_NODISCARD CORE_API UE_STRING_CLASS TrimEnd() &&;
+	[[nodiscard]] CORE_API UE_STRING_CLASS TrimEnd() &&;
 
 	/** 
 	 * Trims the inner array after the null terminator.
@@ -1798,23 +1798,23 @@ public:
 	/**
 	 * Returns a copy of this string with wrapping quotation marks removed.
 	 */
-	UE_NODISCARD CORE_API UE_STRING_CLASS TrimQuotes(bool* bQuotesRemoved = nullptr ) const &;
+	[[nodiscard]] CORE_API UE_STRING_CLASS TrimQuotes(bool* bQuotesRemoved = nullptr ) const &;
 
 	/**
 	 * Returns this string with wrapping quotation marks removed.
 	 */
-	UE_NODISCARD CORE_API UE_STRING_CLASS TrimQuotes(bool* bQuotesRemoved = nullptr) &&;
+	[[nodiscard]] CORE_API UE_STRING_CLASS TrimQuotes(bool* bQuotesRemoved = nullptr) &&;
 	
 	/**
 	* Returns a copy of this string with wrapping CharacterToTrim removed (removes at max one instance in the beginning and end of the string).
 	* @see TrimCharInline for an inline variant
 	*/
-	UE_NODISCARD CORE_API UE_STRING_CLASS TrimChar(ElementType CharacterToTrim, bool* bCharRemoved = nullptr ) const &;
+	[[nodiscard]] CORE_API UE_STRING_CLASS TrimChar(ElementType CharacterToTrim, bool* bCharRemoved = nullptr ) const &;
 
 	/**
 	* Returns a copy of this string with wrapping CharacterToTrim removed (removes at max one instance in the beginning and end of the string).
 	*/
-	UE_NODISCARD CORE_API UE_STRING_CLASS TrimChar(ElementType CharacterToTrim, bool* bCharRemoved = nullptr) &&;
+	[[nodiscard]] CORE_API UE_STRING_CLASS TrimChar(ElementType CharacterToTrim, bool* bCharRemoved = nullptr) &&;
 	
 	/**
 	 * Breaks up a delimited string into elements of a string array.
@@ -1873,12 +1873,12 @@ public:
 	/**
 	 * Returns a copy of this string, with the characters in reverse order
 	 */
-	UE_NODISCARD CORE_API UE_STRING_CLASS Reverse() const &;
+	[[nodiscard]] CORE_API UE_STRING_CLASS Reverse() const &;
 
 	/**
 	 * Returns this string, with the characters in reverse order
 	 */
-	UE_NODISCARD CORE_API UE_STRING_CLASS Reverse() &&;
+	[[nodiscard]] CORE_API UE_STRING_CLASS Reverse() &&;
 
 	/**
 	 * Reverses the order of characters in this string
@@ -1893,7 +1893,7 @@ public:
 	 * @param SearchCase	Indicates whether the search is case sensitive or not ( defaults to ESearchCase::IgnoreCase )
 	 * @return a copy of this string with the replacement made
 	 */
-	UE_NODISCARD CORE_API UE_STRING_CLASS Replace(const ElementType* From, const ElementType* To, ESearchCase::Type SearchCase = ESearchCase::IgnoreCase) const &;
+	[[nodiscard]] CORE_API UE_STRING_CLASS Replace(const ElementType* From, const ElementType* To, ESearchCase::Type SearchCase = ESearchCase::IgnoreCase) const &;
 
 	/**
 	 * Replace all occurrences of a substring in this string
@@ -1903,7 +1903,7 @@ public:
 	 * @param SearchCase	Indicates whether the search is case sensitive or not ( defaults to ESearchCase::IgnoreCase )
 	 * @return a copy of this string with the replacement made
 	 */
-	UE_NODISCARD CORE_API UE_STRING_CLASS Replace(const ElementType* From, const ElementType* To, ESearchCase::Type SearchCase = ESearchCase::IgnoreCase) &&;
+	[[nodiscard]] CORE_API UE_STRING_CLASS Replace(const ElementType* From, const ElementType* To, ESearchCase::Type SearchCase = ESearchCase::IgnoreCase) &&;
 
 	/**
 	 * Replace all occurrences of SearchText with ReplacementText in this string.
@@ -1946,7 +1946,7 @@ public:
 	/**
 	 * Returns a copy of this string with all quote marks escaped (unless the quote is already escaped)
 	 */
-	UE_NODISCARD UE_STRING_CLASS ReplaceQuotesWithEscapedQuotes() const &
+	[[nodiscard]] UE_STRING_CLASS ReplaceQuotesWithEscapedQuotes() const &
 	{
 		UE_STRING_CLASS Result(*this);
 		return MoveTemp(Result).ReplaceQuotesWithEscapedQuotes();
@@ -1955,7 +1955,7 @@ public:
 	/**
 	 * Returns a copy of this string with all quote marks escaped (unless the quote is already escaped)
 	 */
-	UE_NODISCARD CORE_API UE_STRING_CLASS ReplaceQuotesWithEscapedQuotes() &&;
+	[[nodiscard]] CORE_API UE_STRING_CLASS ReplaceQuotesWithEscapedQuotes() &&;
 
 	/**
 	 * Replaces certain characters with the "escaped" version of that character (i.e. replaces "\n" with "\\n").
@@ -1973,7 +1973,7 @@ public:
 	 *
 	 * @return	a string with all control characters replaced by the escaped version.
 	 */
-	UE_NODISCARD UE_STRING_CLASS ReplaceCharWithEscapedChar( const TArray<ElementType>* Chars = nullptr ) const &
+	[[nodiscard]] UE_STRING_CLASS ReplaceCharWithEscapedChar( const TArray<ElementType>* Chars = nullptr ) const &
 	{
 		UE_STRING_CLASS Result(*this);
 		Result.ReplaceCharWithEscapedCharInline(Chars);
@@ -1988,7 +1988,7 @@ public:
 	 *
 	 * @return	a string with all control characters replaced by the escaped version.
 	 */
-	UE_NODISCARD UE_STRING_CLASS ReplaceCharWithEscapedChar( const TArray<ElementType>* Chars = nullptr ) &&
+	[[nodiscard]] UE_STRING_CLASS ReplaceCharWithEscapedChar( const TArray<ElementType>* Chars = nullptr ) &&
 	{
 		ReplaceCharWithEscapedCharInline(Chars);
 		return MoveTemp(*this);
@@ -2005,7 +2005,7 @@ public:
 	 * replaces "\\n" with "\n".  Counterpart to ReplaceCharWithEscapedChar().
 	 * @return copy of this string with replacement made
 	 */
-	UE_NODISCARD UE_STRING_CLASS ReplaceEscapedCharWithChar( const TArray<ElementType>* Chars = nullptr ) const &
+	[[nodiscard]] UE_STRING_CLASS ReplaceEscapedCharWithChar( const TArray<ElementType>* Chars = nullptr ) const &
 	{
 		UE_STRING_CLASS Result(*this);
 		Result.ReplaceEscapedCharWithCharInline(Chars);
@@ -2017,7 +2017,7 @@ public:
 	 * replaces "\\n" with "\n".  Counterpart to ReplaceCharWithEscapedChar().
 	 * @return copy of this string with replacement made
 	 */
-	UE_NODISCARD UE_STRING_CLASS ReplaceEscapedCharWithChar( const TArray<ElementType>* Chars = nullptr ) &&
+	[[nodiscard]] UE_STRING_CLASS ReplaceEscapedCharWithChar( const TArray<ElementType>* Chars = nullptr ) &&
 	{
 		ReplaceEscapedCharWithCharInline(Chars);
 		return MoveTemp(*this);
@@ -2034,7 +2034,7 @@ public:
 	 * @param InSpacesPerTab - Number of spaces that a tab represents
 	 * @return copy of this string with replacement made
 	 */
-	UE_NODISCARD UE_STRING_CLASS ConvertTabsToSpaces(const int32 InSpacesPerTab) const &
+	[[nodiscard]] UE_STRING_CLASS ConvertTabsToSpaces(const int32 InSpacesPerTab) const &
 	{
 		UE_STRING_CLASS FinalString(*this);
 		FinalString.ConvertTabsToSpacesInline(InSpacesPerTab);
@@ -2046,14 +2046,14 @@ public:
 	 * @param InSpacesPerTab - Number of spaces that a tab represents
 	 * @return copy of this string with replacement made
 	 */
-	UE_NODISCARD UE_STRING_CLASS ConvertTabsToSpaces(const int32 InSpacesPerTab) &&
+	[[nodiscard]] UE_STRING_CLASS ConvertTabsToSpaces(const int32 InSpacesPerTab) &&
 	{
 		ConvertTabsToSpacesInline(InSpacesPerTab);
 		return MoveTemp(*this);
 	}
 
 	// Takes the number passed in and formats the string in comma format ( 12345 becomes "12,345")
-	UE_NODISCARD static CORE_API UE_STRING_CLASS FormatAsNumber( int32 InNumber );
+	[[nodiscard]] static CORE_API UE_STRING_CLASS FormatAsNumber( int32 InNumber );
 
 	// To allow more efficient memory handling, automatically adds one for the string termination.
 	CORE_API void Reserve(int32 CharacterCount);
@@ -2069,7 +2069,7 @@ public:
 
 
 	/** Converts an integer to a string. */
-	UE_NODISCARD static FORCEINLINE UE_STRING_CLASS FromInt( int32 Num )
+	[[nodiscard]] static FORCEINLINE UE_STRING_CLASS FromInt( int32 Num )
 	{
 		UE_STRING_CLASS Ret;
 		Ret.AppendInt(Num); 
@@ -2086,7 +2086,7 @@ public:
 	 *
 	 * @return The boolean value
 	 */
-	UE_NODISCARD CORE_API bool ToBool() const;
+	[[nodiscard]] CORE_API bool ToBool() const;
 
 	/**
 	 * Converts a buffer to a string
@@ -2096,7 +2096,7 @@ public:
 	 *
 	 * @return the blob in string form
 	 */
-	UE_NODISCARD static CORE_API UE_STRING_CLASS FromBlob(const uint8* SrcBuffer,const uint32 SrcSize);
+	[[nodiscard]] static CORE_API UE_STRING_CLASS FromBlob(const uint8* SrcBuffer,const uint32 SrcSize);
 
 	/**
 	 * Converts a string into a buffer
@@ -2116,7 +2116,7 @@ public:
 	 *
 	 * @return the blob in string form
 	 */
-	UE_NODISCARD static CORE_API UE_STRING_CLASS FromHexBlob(const uint8* SrcBuffer,const uint32 SrcSize);
+	[[nodiscard]] static CORE_API UE_STRING_CLASS FromHexBlob(const uint8* SrcBuffer,const uint32 SrcSize);
 
 	/**
 	 * Converts a string into a buffer
@@ -2137,7 +2137,7 @@ public:
 	 *
 	 * @return sanitized string version of float
 	 */
-	UE_NODISCARD static CORE_API UE_STRING_CLASS SanitizeFloat( double InFloat, const int32 InMinFractionalDigits = 1 );
+	[[nodiscard]] static CORE_API UE_STRING_CLASS SanitizeFloat( double InFloat, const int32 InMinFractionalDigits = 1 );
 
 	/**
 	 * Joins a range of 'something that can be concatentated to strings with +=' together into a single string with separators.
@@ -2148,7 +2148,7 @@ public:
 	 * @return	The final, joined, separated string.
 	 */
 	template <typename RangeType>
-	UE_NODISCARD static UE_STRING_CLASS Join(const RangeType& Range, const ElementType* Separator)
+	[[nodiscard]] static UE_STRING_CLASS Join(const RangeType& Range, const ElementType* Separator)
 	{
 		UE_STRING_CLASS Result;
 		bool            First = true;
@@ -2179,7 +2179,7 @@ public:
 	 * @return	The final, joined, separated string.
 	 */
 	template <typename RangeType, typename ProjectionType>
-	UE_NODISCARD static UE_STRING_CLASS JoinBy(const RangeType& Range, const ElementType* Separator, ProjectionType Proj)
+	[[nodiscard]] static UE_STRING_CLASS JoinBy(const RangeType& Range, const ElementType* Separator, ProjectionType Proj)
 	{
 		UE_STRING_CLASS Result;
 		bool            First = true;
@@ -2285,13 +2285,13 @@ inline void LexFromString(bool& OutValue, 				const UE_STRING_CHARTYPE* Buffer)
 inline void LexFromString(UE_STRING_CLASS& OutValue,	const UE_STRING_CHARTYPE* Buffer)	{	OutValue = Buffer;						}
 
 template <typename StringType = FString>
-UE_NODISCARD FORCEINLINE StringType LexToString(UE_STRING_CLASS&& Str)
+[[nodiscard]] FORCEINLINE StringType LexToString(UE_STRING_CLASS&& Str)
 {
 	return MoveTemp(Str);
 }
 
 template <typename StringType = FString>
-UE_NODISCARD FORCEINLINE StringType LexToString(const UE_STRING_CLASS& Str)
+[[nodiscard]] FORCEINLINE StringType LexToString(const UE_STRING_CLASS& Str)
 {
 	return Str;
 }
diff --git a/Engine/Source/Runtime/Core/Public/Delegates/DelegateSignatureImpl.inl b/Engine/Source/Runtime/Core/Public/Delegates/DelegateSignatureImpl.inl
--- a/Engine/Source/Runtime/Core/Public/Delegates/DelegateSignatureImpl.inl
+++ b/Engine/Source/Runtime/Core/Public/Delegates/DelegateSignatureImpl.inl
@@ -86,7 +86,7 @@ public:
 	 * Static: Creates a raw C++ pointer global function delegate
 	 */
 	template <typename... VarTypes>
-	UE_NODISCARD inline static TDelegate<RetValType(ParamTypes...), UserPolicy> CreateStatic(typename TIdentity<RetValType (*)(ParamTypes..., std::decay_t<VarTypes>...)>::Type InFunc, VarTypes&&... Vars)
+	[[nodiscard]] inline static TDelegate<RetValType(ParamTypes...), UserPolicy> CreateStatic(typename TIdentity<RetValType (*)(ParamTypes..., std::decay_t<VarTypes>...)>::Type InFunc, VarTypes&&... Vars)
 	{
 		TDelegate<RetValType(ParamTypes...), UserPolicy> Result;
 		Result.template CreateDelegateInstance<TBaseStaticDelegateInstance<FuncType, UserPolicy, std::decay_t<VarTypes>...>>(InFunc, Forward<VarTypes>(Vars)...);
@@ -98,7 +98,7 @@ public:
 	 * technically this works for any functor types, but lambdas are the primary use case
 	 */
 	template<typename FunctorType, typename... VarTypes>
-	UE_NODISCARD inline static TDelegate<RetValType(ParamTypes...), UserPolicy> CreateLambda(FunctorType&& InFunctor, VarTypes&&... Vars)
+	[[nodiscard]] inline static TDelegate<RetValType(ParamTypes...), UserPolicy> CreateLambda(FunctorType&& InFunctor, VarTypes&&... Vars)
 	{
 		TDelegate<RetValType(ParamTypes...), UserPolicy> Result;
 		Result.template CreateDelegateInstance<TBaseFunctorDelegateInstance<FuncType, UserPolicy, typename TRemoveReference<FunctorType>::Type, std::decay_t<VarTypes>...>>(Forward<FunctorType>(InFunctor), Forward<VarTypes>(Vars)...);
@@ -110,14 +110,14 @@ public:
 	 * technically this works for any functor types, but lambdas are the primary use case
 	 */
 	template<typename UserClass, ESPMode Mode, typename FunctorType, typename... VarTypes>
-	UE_NODISCARD inline static TDelegate<RetValType(ParamTypes...), UserPolicy> CreateSPLambda(const TSharedRef<UserClass, Mode>& InUserObjectRef, FunctorType&& InFunctor, VarTypes&&... Vars)
+	[[nodiscard]] inline static TDelegate<RetValType(ParamTypes...), UserPolicy> CreateSPLambda(const TSharedRef<UserClass, Mode>& InUserObjectRef, FunctorType&& InFunctor, VarTypes&&... Vars)
 	{
 		TDelegate<RetValType(ParamTypes...), UserPolicy> Result;
 		Result.template CreateDelegateInstance<TBaseSPLambdaDelegateInstance<UserClass, Mode, FuncType, UserPolicy, typename TRemoveReference<FunctorType>::Type, std::decay_t<VarTypes>...>>(InUserObjectRef, Forward<FunctorType>(InFunctor), Forward<VarTypes>(Vars)...);
 		return Result;
 	}
 	template <typename UserClass, typename FunctorType, typename... VarTypes>
-	UE_NODISCARD inline static TDelegate<RetValType(ParamTypes...), UserPolicy> CreateSPLambda(UserClass* InUserObject, FunctorType&& InFunctor, VarTypes&&... Vars)
+	[[nodiscard]] inline static TDelegate<RetValType(ParamTypes...), UserPolicy> CreateSPLambda(UserClass* InUserObject, FunctorType&& InFunctor, VarTypes&&... Vars)
 	{
 		TDelegate<RetValType(ParamTypes...), UserPolicy> Result;
 		Result.template CreateDelegateInstance<TBaseSPLambdaDelegateInstance<UserClass, decltype(InUserObject->AsShared())::Mode, FuncType, UserPolicy, typename TRemoveReference<FunctorType>::Type, std::decay_t<VarTypes>...>>(StaticCastSharedRef<UserClass>(InUserObject->AsShared()), Forward<FunctorType>(InFunctor), Forward<VarTypes>(Vars)...);
@@ -129,7 +129,7 @@ public:
 	 * technically this works for any functor types, but lambdas are the primary use case
 	 */
 	template<typename UserClass, typename FunctorType, typename... VarTypes>
-	UE_NODISCARD inline static TDelegate<RetValType(ParamTypes...), UserPolicy> CreateWeakLambda(UserClass* InUserObject, FunctorType&& InFunctor, VarTypes&&... Vars)
+	[[nodiscard]] inline static TDelegate<RetValType(ParamTypes...), UserPolicy> CreateWeakLambda(UserClass* InUserObject, FunctorType&& InFunctor, VarTypes&&... Vars)
 	{
 		TDelegate<RetValType(ParamTypes...), UserPolicy> Result;
 		Result.template CreateDelegateInstance<TWeakBaseFunctorDelegateInstance<UserClass, FuncType, UserPolicy, typename TRemoveReference<FunctorType>::Type, std::decay_t<VarTypes>...>>(InUserObject, Forward<FunctorType>(InFunctor), Forward<VarTypes>(Vars)...);
@@ -143,7 +143,7 @@ public:
 	 * deleted out from underneath your delegate. Be careful when calling Execute()!
 	 */
 	template <typename UserClass, typename... VarTypes>
-	UE_NODISCARD inline static TDelegate<RetValType(ParamTypes...), UserPolicy> CreateRaw(UserClass* InUserObject, typename TMemFunPtrType<false, UserClass, RetValType (ParamTypes..., std::decay_t<VarTypes>...)>::Type InFunc, VarTypes&&... Vars)
+	[[nodiscard]] inline static TDelegate<RetValType(ParamTypes...), UserPolicy> CreateRaw(UserClass* InUserObject, typename TMemFunPtrType<false, UserClass, RetValType (ParamTypes..., std::decay_t<VarTypes>...)>::Type InFunc, VarTypes&&... Vars)
 	{
 		static_assert(!TIsConst<UserClass>::Value, "Attempting to bind a delegate with a const object pointer and non-const member function.");
 
@@ -152,7 +152,7 @@ public:
 		return Result;
 	}
 	template <typename UserClass, typename... VarTypes>
-	UE_NODISCARD inline static TDelegate<RetValType(ParamTypes...), UserPolicy> CreateRaw(const UserClass* InUserObject, typename TMemFunPtrType<true, UserClass, RetValType (ParamTypes..., std::decay_t<VarTypes>...)>::Type InFunc, VarTypes&&... Vars)
+	[[nodiscard]] inline static TDelegate<RetValType(ParamTypes...), UserPolicy> CreateRaw(const UserClass* InUserObject, typename TMemFunPtrType<true, UserClass, RetValType (ParamTypes..., std::decay_t<VarTypes>...)>::Type InFunc, VarTypes&&... Vars)
 	{
 		TDelegate<RetValType(ParamTypes...), UserPolicy> Result;
 		Result.template CreateDelegateInstance<TBaseRawMethodDelegateInstance<true, const UserClass, FuncType, UserPolicy, std::decay_t<VarTypes>...>>(InUserObject, InFunc, Forward<VarTypes>(Vars)...);
@@ -166,7 +166,7 @@ public:
 	 * You can use ExecuteIfBound() to call them.
 	 */
 	template <typename UserClass, ESPMode Mode, typename... VarTypes>
-	UE_NODISCARD inline static TDelegate<RetValType(ParamTypes...), UserPolicy> CreateSP(const TSharedRef<UserClass, Mode>& InUserObjectRef, typename TMemFunPtrType<false, UserClass, RetValType (ParamTypes..., std::decay_t<VarTypes>...)>::Type InFunc, VarTypes&&... Vars)
+	[[nodiscard]] inline static TDelegate<RetValType(ParamTypes...), UserPolicy> CreateSP(const TSharedRef<UserClass, Mode>& InUserObjectRef, typename TMemFunPtrType<false, UserClass, RetValType (ParamTypes..., std::decay_t<VarTypes>...)>::Type InFunc, VarTypes&&... Vars)
 	{
 		static_assert(!TIsConst<UserClass>::Value, "Attempting to bind a delegate with a const object pointer and non-const member function.");
 
@@ -175,7 +175,7 @@ public:
 		return Result;
 	}
 	template <typename UserClass, ESPMode Mode, typename... VarTypes>
-	UE_NODISCARD inline static TDelegate<RetValType(ParamTypes...), UserPolicy> CreateSP(const TSharedRef<UserClass, Mode>& InUserObjectRef, typename TMemFunPtrType<true, UserClass, RetValType (ParamTypes..., std::decay_t<VarTypes>...)>::Type InFunc, VarTypes&&... Vars)
+	[[nodiscard]] inline static TDelegate<RetValType(ParamTypes...), UserPolicy> CreateSP(const TSharedRef<UserClass, Mode>& InUserObjectRef, typename TMemFunPtrType<true, UserClass, RetValType (ParamTypes..., std::decay_t<VarTypes>...)>::Type InFunc, VarTypes&&... Vars)
 	{
 		TDelegate<RetValType(ParamTypes...), UserPolicy> Result;
 		Result.template CreateDelegateInstance<TBaseSPMethodDelegateInstance<true, const UserClass, Mode, FuncType, UserPolicy, std::decay_t<VarTypes>...>>(InUserObjectRef, InFunc, Forward<VarTypes>(Vars)...);
@@ -189,7 +189,7 @@ public:
 	 * You can use ExecuteIfBound() to call them.
 	 */
 	template <typename UserClass, typename... VarTypes>
-	UE_NODISCARD inline static TDelegate<RetValType(ParamTypes...), UserPolicy> CreateSP(UserClass* InUserObject, typename TMemFunPtrType<false, UserClass, RetValType (ParamTypes..., std::decay_t<VarTypes>...)>::Type InFunc, VarTypes&&... Vars)
+	[[nodiscard]] inline static TDelegate<RetValType(ParamTypes...), UserPolicy> CreateSP(UserClass* InUserObject, typename TMemFunPtrType<false, UserClass, RetValType (ParamTypes..., std::decay_t<VarTypes>...)>::Type InFunc, VarTypes&&... Vars)
 	{
 		static_assert(!TIsConst<UserClass>::Value, "Attempting to bind a delegate with a const object pointer and non-const member function.");
 
@@ -198,7 +198,7 @@ public:
 		return Result;
 	}
 	template <typename UserClass, typename... VarTypes>
-	UE_NODISCARD inline static TDelegate<RetValType(ParamTypes...), UserPolicy> CreateSP(const UserClass* InUserObject, typename TMemFunPtrType<true, UserClass, RetValType (ParamTypes..., std::decay_t<VarTypes>...)>::Type InFunc, VarTypes&&... Vars)
+	[[nodiscard]] inline static TDelegate<RetValType(ParamTypes...), UserPolicy> CreateSP(const UserClass* InUserObject, typename TMemFunPtrType<true, UserClass, RetValType (ParamTypes..., std::decay_t<VarTypes>...)>::Type InFunc, VarTypes&&... Vars)
 	{
 		TDelegate<RetValType(ParamTypes...), UserPolicy> Result;
 		Result.template CreateDelegateInstance<TBaseSPMethodDelegateInstance<true, const UserClass, decltype(InUserObject->AsShared())::Mode, FuncType, UserPolicy, std::decay_t<VarTypes>...>>(StaticCastSharedRef<const UserClass>(InUserObject->AsShared()), InFunc, Forward<VarTypes>(Vars)...);
@@ -214,7 +214,7 @@ public:
 	 * Note: This function is redundant, but is retained for backwards compatibility.  CreateSP() works in both thread-safe and not-thread-safe modes and should be preferred.
 	 */
 	template <typename UserClass, typename... VarTypes>
-	UE_NODISCARD inline static TDelegate<RetValType(ParamTypes...), UserPolicy> CreateThreadSafeSP(const TSharedRef<UserClass, ESPMode::ThreadSafe>& InUserObjectRef, typename TMemFunPtrType<false, UserClass, RetValType (ParamTypes..., std::decay_t<VarTypes>...)>::Type InFunc, VarTypes&&... Vars)
+	[[nodiscard]] inline static TDelegate<RetValType(ParamTypes...), UserPolicy> CreateThreadSafeSP(const TSharedRef<UserClass, ESPMode::ThreadSafe>& InUserObjectRef, typename TMemFunPtrType<false, UserClass, RetValType (ParamTypes..., std::decay_t<VarTypes>...)>::Type InFunc, VarTypes&&... Vars)
 	{
 		static_assert(!TIsConst<UserClass>::Value, "Attempting to bind a delegate with a const object pointer and non-const member function.");
 
@@ -223,7 +223,7 @@ public:
 		return Result;
 	}
 	template <typename UserClass, typename... VarTypes>
-	UE_NODISCARD inline static TDelegate<RetValType(ParamTypes...), UserPolicy> CreateThreadSafeSP(const TSharedRef<UserClass, ESPMode::ThreadSafe>& InUserObjectRef, typename TMemFunPtrType<true, UserClass, RetValType (ParamTypes..., std::decay_t<VarTypes>...)>::Type InFunc, VarTypes&&... Vars)
+	[[nodiscard]] inline static TDelegate<RetValType(ParamTypes...), UserPolicy> CreateThreadSafeSP(const TSharedRef<UserClass, ESPMode::ThreadSafe>& InUserObjectRef, typename TMemFunPtrType<true, UserClass, RetValType (ParamTypes..., std::decay_t<VarTypes>...)>::Type InFunc, VarTypes&&... Vars)
 	{
 		TDelegate<RetValType(ParamTypes...), UserPolicy> Result;
 		Result.template CreateDelegateInstance<TBaseSPMethodDelegateInstance<true, const UserClass, ESPMode::ThreadSafe, FuncType, UserPolicy, std::decay_t<VarTypes>...>>(InUserObjectRef, InFunc, Forward<VarTypes>(Vars)...);
@@ -239,7 +239,7 @@ public:
 	 * Note: This function is redundant, but is retained for backwards compatibility.  CreateSP() works in both thread-safe and not-thread-safe modes and should be preferred.
 	 */
 	template <typename UserClass, typename... VarTypes>
-	UE_NODISCARD inline static TDelegate<RetValType(ParamTypes...), UserPolicy> CreateThreadSafeSP(UserClass* InUserObject, typename TMemFunPtrType<false, UserClass, RetValType (ParamTypes..., std::decay_t<VarTypes>...)>::Type InFunc, VarTypes&&... Vars)
+	[[nodiscard]] inline static TDelegate<RetValType(ParamTypes...), UserPolicy> CreateThreadSafeSP(UserClass* InUserObject, typename TMemFunPtrType<false, UserClass, RetValType (ParamTypes..., std::decay_t<VarTypes>...)>::Type InFunc, VarTypes&&... Vars)
 	{
 		static_assert(!TIsConst<UserClass>::Value, "Attempting to bind a delegate with a const object pointer and non-const member function.");
 
@@ -248,7 +248,7 @@ public:
 		return Result;
 	}
 	template <typename UserClass, typename... VarTypes>
-	UE_NODISCARD inline static TDelegate<RetValType(ParamTypes...), UserPolicy> CreateThreadSafeSP(const UserClass* InUserObject, typename TMemFunPtrType<true, UserClass, RetValType (ParamTypes..., std::decay_t<VarTypes>...)>::Type InFunc, VarTypes&&... Vars)
+	[[nodiscard]] inline static TDelegate<RetValType(ParamTypes...), UserPolicy> CreateThreadSafeSP(const UserClass* InUserObject, typename TMemFunPtrType<true, UserClass, RetValType (ParamTypes..., std::decay_t<VarTypes>...)>::Type InFunc, VarTypes&&... Vars)
 	{
 		TDelegate<RetValType(ParamTypes...), UserPolicy> Result;
 		Result.template CreateDelegateInstance<TBaseSPMethodDelegateInstance<true, const UserClass, ESPMode::ThreadSafe, FuncType, UserPolicy, std::decay_t<VarTypes>...>>(StaticCastSharedRef<const UserClass>(InUserObject->AsShared()), InFunc, Forward<VarTypes>(Vars)...);
@@ -262,14 +262,14 @@ public:
 	 * You can use ExecuteIfBound() to call them.
 	 */
 	template <typename UObjectTemplate, typename... VarTypes>
-	UE_NODISCARD inline static TDelegate<RetValType(ParamTypes...), UserPolicy> CreateUFunction(UObjectTemplate* InUserObject, const FName& InFunctionName, VarTypes&&... Vars)
+	[[nodiscard]] inline static TDelegate<RetValType(ParamTypes...), UserPolicy> CreateUFunction(UObjectTemplate* InUserObject, const FName& InFunctionName, VarTypes&&... Vars)
 	{
 		TDelegate<RetValType(ParamTypes...), UserPolicy> Result;
 		Result.template CreateDelegateInstance<TBaseUFunctionDelegateInstance<UObjectTemplate, FuncType, UserPolicy, std::decay_t<VarTypes>...>>(InUserObject, InFunctionName, Forward<VarTypes>(Vars)...);
 		return Result;
 	}
 	template <typename UObjectTemplate, typename... VarTypes>
-	UE_NODISCARD inline static TDelegate<RetValType(ParamTypes...), UserPolicy> CreateUFunction(TObjectPtr<UObjectTemplate> InUserObject, const FName& InFunctionName, VarTypes&&... Vars)
+	[[nodiscard]] inline static TDelegate<RetValType(ParamTypes...), UserPolicy> CreateUFunction(TObjectPtr<UObjectTemplate> InUserObject, const FName& InFunctionName, VarTypes&&... Vars)
 	{
 		TDelegate<RetValType(ParamTypes...), UserPolicy> Result;
 		Result.template CreateDelegateInstance<TBaseUFunctionDelegateInstance<UObjectTemplate, FuncType, UserPolicy, std::decay_t<VarTypes>...>>(ToRawPtr(InUserObject), InFunctionName, Forward<VarTypes>(Vars)...);
@@ -283,7 +283,7 @@ public:
 	 * You can use ExecuteIfBound() to call them.
 	 */
 	template <typename UserClass, typename... VarTypes>
-	UE_NODISCARD inline static TDelegate<RetValType(ParamTypes...), UserPolicy> CreateUObject(UserClass* InUserObject, typename TMemFunPtrType<false, UserClass, RetValType (ParamTypes..., std::decay_t<VarTypes>...)>::Type InFunc, VarTypes&&... Vars)
+	[[nodiscard]] inline static TDelegate<RetValType(ParamTypes...), UserPolicy> CreateUObject(UserClass* InUserObject, typename TMemFunPtrType<false, UserClass, RetValType (ParamTypes..., std::decay_t<VarTypes>...)>::Type InFunc, VarTypes&&... Vars)
 	{
 		static_assert(!TIsConst<UserClass>::Value, "Attempting to bind a delegate with a const object pointer and non-const member function.");
 
@@ -292,14 +292,14 @@ public:
 		return Result;
 	}
 	template <typename UserClass, typename... VarTypes>
-	UE_NODISCARD inline static TDelegate<RetValType(ParamTypes...), UserPolicy> CreateUObject(const UserClass* InUserObject, typename TMemFunPtrType<true, UserClass, RetValType (ParamTypes..., std::decay_t<VarTypes>...)>::Type InFunc, VarTypes&&... Vars)
+	[[nodiscard]] inline static TDelegate<RetValType(ParamTypes...), UserPolicy> CreateUObject(const UserClass* InUserObject, typename TMemFunPtrType<true, UserClass, RetValType (ParamTypes..., std::decay_t<VarTypes>...)>::Type InFunc, VarTypes&&... Vars)
 	{
 		TDelegate<RetValType(ParamTypes...), UserPolicy> Result;
 		Result.template CreateDelegateInstance<TBaseUObjectMethodDelegateInstance<true, const UserClass, FuncType, UserPolicy, std::decay_t<VarTypes>...>>(InUserObject, InFunc, Forward<VarTypes>(Vars)...);
 		return Result;
 	}
 	template <typename UserClass, typename... VarTypes>
-	UE_NODISCARD inline static TDelegate<RetValType(ParamTypes...), UserPolicy> CreateUObject(TObjectPtr<UserClass> InUserObject, typename TMemFunPtrType<false, UserClass, RetValType (ParamTypes..., std::decay_t<VarTypes>...)>::Type InFunc, VarTypes&&... Vars)
+	[[nodiscard]] inline static TDelegate<RetValType(ParamTypes...), UserPolicy> CreateUObject(TObjectPtr<UserClass> InUserObject, typename TMemFunPtrType<false, UserClass, RetValType (ParamTypes..., std::decay_t<VarTypes>...)>::Type InFunc, VarTypes&&... Vars)
 	{
 		static_assert(!TIsConst<UserClass>::Value, "Attempting to bind a delegate with a const object pointer and non-const member function.");
 
@@ -308,7 +308,7 @@ public:
 		return Result;
 	}
 	template <typename UserClass, typename... VarTypes>
-	UE_NODISCARD inline static TDelegate<RetValType(ParamTypes...), UserPolicy> CreateUObject(TObjectPtr<UserClass> InUserObject, typename TMemFunPtrType<true, UserClass, RetValType (ParamTypes..., std::decay_t<VarTypes>...)>::Type InFunc, VarTypes&&... Vars)
+	[[nodiscard]] inline static TDelegate<RetValType(ParamTypes...), UserPolicy> CreateUObject(TObjectPtr<UserClass> InUserObject, typename TMemFunPtrType<true, UserClass, RetValType (ParamTypes..., std::decay_t<VarTypes>...)>::Type InFunc, VarTypes&&... Vars)
 	{
 		TDelegate<RetValType(ParamTypes...), UserPolicy> Result;
 		Result.template CreateDelegateInstance<TBaseUObjectMethodDelegateInstance<true, const UserClass, FuncType, UserPolicy, std::decay_t<VarTypes>...>>(ToRawPtr(InUserObject), InFunc, Forward<VarTypes>(Vars)...);
diff --git a/Engine/Source/Runtime/Core/Public/Experimental/Containers/HazardPointer.h b/Engine/Source/Runtime/Core/Public/Experimental/Containers/HazardPointer.h
--- a/Engine/Source/Runtime/Core/Public/Experimental/Containers/HazardPointer.h
+++ b/Engine/Source/Runtime/Core/Public/Experimental/Containers/HazardPointer.h
@@ -91,7 +91,7 @@ class FHazardPointerCollection
 		}
 
 		//assign hazard pointer once acquired
-		UE_NODISCARD inline void* SetHazard(void* InHazard)
+		[[nodiscard]] inline void* SetHazard(void* InHazard)
 		{
 			Hazard.store(reinterpret_cast<uintptr_t>(InHazard), std::memory_order_release);
 			std::atomic_thread_fence(std::memory_order_seq_cst);
diff --git a/Engine/Source/Runtime/Core/Public/Experimental/Coroutine/CoroSpinLock.h b/Engine/Source/Runtime/Core/Public/Experimental/Coroutine/CoroSpinLock.h
--- a/Engine/Source/Runtime/Core/Public/Experimental/Coroutine/CoroSpinLock.h
+++ b/Engine/Source/Runtime/Core/Public/Experimental/Coroutine/CoroSpinLock.h
@@ -23,7 +23,7 @@ public:
 
 namespace CoroSpinLock_Detail
 {
-	class UE_NODISCARD FScope
+	class [[nodiscard]] FScope
 	{
 		FCoroSpinLock* CoroSpinLock = nullptr;
 		bool WriteLock = true;
@@ -84,7 +84,7 @@ namespace CoroSpinLock_Detail
 		}
 	};
 
-	class UE_NODISCARD FAwaitable
+	class [[nodiscard]] FAwaitable
 	{
 		FCoroSpinLock* CoroSpinLock = nullptr;
 		TLaunchedCoroTask<void> Task;
@@ -176,7 +176,7 @@ class FCoroSpinLock
 {
 	FRWLock CS;
 
-	class UE_NODISCARD FScope
+	class [[nodiscard]] FScope
 	{
 		FRWLock* LockObject = nullptr;
 		bool WriteLock = true;
diff --git a/Engine/Source/Runtime/Core/Public/Experimental/Coroutine/Coroutine.h b/Engine/Source/Runtime/Core/Public/Experimental/Coroutine/Coroutine.h
--- a/Engine/Source/Runtime/Core/Public/Experimental/Coroutine/Coroutine.h
+++ b/Engine/Source/Runtime/Core/Public/Experimental/Coroutine/Coroutine.h
@@ -613,7 +613,7 @@ namespace CoroTask_Detail
 * TCoroTask serves as the ReturnType for Coroutine Tasks
 */
 template<typename ReturnType>
-class UE_NODISCARD TCoroTask
+class [[nodiscard]] TCoroTask
 {
 	friend class CoroTask_Detail::IPromise;
 	friend class CoroTask_Detail::FPromise;
@@ -716,7 +716,7 @@ public:
 /*
 * FLockedTask is a special use case Task that cannot be expedited
 */
-class UE_NODISCARD FLockedTask : private TCoroTask<void>
+class [[nodiscard]] FLockedTask : private TCoroTask<void>
 {
 	FLockedTask(promise_type* InPromise) : TCoroTask<void>(InPromise)
 	{
@@ -773,7 +773,7 @@ public:
 * TLaunchedCoroTask is the moved result of a Launched TCoroTask this way uniqueness is guranteed
 */
 template<typename ReturnType>
-class UE_NODISCARD TLaunchedCoroTask : protected TCoroTask<ReturnType>
+class [[nodiscard]] TLaunchedCoroTask : protected TCoroTask<ReturnType>
 {
 	friend class CoroTask_Detail::IPromise;
 	friend class CoroTask_Detail::FPromise;
@@ -870,7 +870,7 @@ inline TLaunchedCoroTask<ReturnType> TCoroTask<ReturnType>::Launch(const TCHAR*
 * TCoroFrame just serves as storage for a Stackframe and therefore allowing the stackless Coroutines to have a Stack
 */
 template<typename ReturnType>
-class UE_NODISCARD TCoroFrame
+class [[nodiscard]] TCoroFrame
 {
 public:
 	using promise_type = CoroTask_Detail::TFramePromise<ReturnType>;
diff --git a/Engine/Source/Runtime/Core/Public/HAL/PlatformFileCommon.h b/Engine/Source/Runtime/Core/Public/HAL/PlatformFileCommon.h
--- a/Engine/Source/Runtime/Core/Public/HAL/PlatformFileCommon.h
+++ b/Engine/Source/Runtime/Core/Public/HAL/PlatformFileCommon.h
@@ -34,7 +34,7 @@ public:
 
 	virtual ~FFileHandleRegistry() = default;
 
-	UE_NODISCARD FRegisteredFileHandle* InitialOpenFile(const TCHAR* Filename)
+	[[nodiscard]] FRegisteredFileHandle* InitialOpenFile(const TCHAR* Filename)
 	{
 		if (HandlesCurrentlyInUse.Increment() > MaxOpenHandles)
 		{
@@ -76,7 +76,7 @@ public:
 
 	// Only returns false if the platform handle cannot be reopened.
 	// TrackEndRead should only be called if TrackStartRead succeeded
-	UE_NODISCARD bool TrackStartRead(FRegisteredFileHandle* Handle)
+	[[nodiscard]] bool TrackStartRead(FRegisteredFileHandle* Handle)
 	{
 		{
 			FScopeLock Lock(&LockSection);
diff --git a/Engine/Source/Runtime/Core/Public/IO/IoBuffer.h b/Engine/Source/Runtime/Core/Public/IO/IoBuffer.h
--- a/Engine/Source/Runtime/Core/Public/IO/IoBuffer.h
+++ b/Engine/Source/Runtime/Core/Public/IO/IoBuffer.h
@@ -64,7 +64,7 @@ public:
 	 *
 	 * @return A status wrapper around the memory pointer. Even if the status is valid the pointer might still be null.
 	 */
-	UE_NODISCARD CORE_API TIoStatusOr<uint8*> Release();
+	[[nodiscard]] CORE_API TIoStatusOr<uint8*> Release();
 
 private:
 	/** Core buffer object. For internal use only, used by FIoBuffer
diff --git a/Engine/Source/Runtime/Core/Public/IO/IoDispatcher.h b/Engine/Source/Runtime/Core/Public/IO/IoDispatcher.h
--- a/Engine/Source/Runtime/Core/Public/IO/IoDispatcher.h
+++ b/Engine/Source/Runtime/Core/Public/IO/IoDispatcher.h
@@ -562,7 +562,7 @@ public:
 	CORE_API FIoStoreWriterContext();
 	CORE_API ~FIoStoreWriterContext();
 
-	UE_NODISCARD CORE_API FIoStatus Initialize(const FIoStoreWriterSettings& InWriterSettings);
+	[[nodiscard]] CORE_API FIoStatus Initialize(const FIoStoreWriterSettings& InWriterSettings);
 	CORE_API TSharedPtr<class IIoStoreWriter> CreateContainer(const TCHAR* InContainerPath, const FIoContainerSettings& InContainerSettings);
 	CORE_API void Flush();
 	CORE_API FProgress GetProgress() const;
@@ -749,7 +749,7 @@ public:
 	CORE_API FIoStoreReader();
 	CORE_API ~FIoStoreReader();
 
-	UE_NODISCARD CORE_API FIoStatus Initialize(FStringView ContainerPath, const TMap<FGuid, FAES::FAESKey>& InDecryptionKeys);
+	[[nodiscard]] CORE_API FIoStatus Initialize(FStringView ContainerPath, const TMap<FGuid, FAES::FAESKey>& InDecryptionKeys);
 	CORE_API FIoContainerId GetContainerId() const;
 	CORE_API uint32 GetVersion() const;
 	CORE_API EIoContainerFlags GetContainerFlags() const;
diff --git a/Engine/Source/Runtime/Core/Public/Math/Box.h b/Engine/Source/Runtime/Core/Public/Math/Box.h
--- a/Engine/Source/Runtime/Core/Public/Math/Box.h
+++ b/Engine/Source/Runtime/Core/Public/Math/Box.h
@@ -279,7 +279,7 @@ public:
 	 * @param W The size to increase the volume by.
 	 * @return A new bounding box.
 	 */
-	UE_NODISCARD FORCEINLINE TBox<T> ExpandBy(T W) const
+	[[nodiscard]] FORCEINLINE TBox<T> ExpandBy(T W) const
 	{
 		return TBox<T>(Min - TVector<T>(W, W, W), Max + TVector<T>(W, W, W));
 	}
@@ -290,7 +290,7 @@ public:
 	* @param V The size to increase the volume by.
 	* @return A new bounding box.
 	*/
-	UE_NODISCARD FORCEINLINE TBox<T> ExpandBy(const TVector<T>& V) const
+	[[nodiscard]] FORCEINLINE TBox<T> ExpandBy(const TVector<T>& V) const
 	{
 		return TBox<T>(Min - V, Max + V);
 	}
@@ -302,7 +302,7 @@ public:
 	* @param Pos The size to increase the volume by in the positive direction (positive values move the bounds outwards)
 	* @return A new bounding box.
 	*/
-	UE_NODISCARD TBox<T> ExpandBy(const TVector<T>& Neg, const TVector<T>& Pos) const
+	[[nodiscard]] TBox<T> ExpandBy(const TVector<T>& Neg, const TVector<T>& Pos) const
 	{
 		return TBox<T>(Min - Neg, Max + Pos);
 	}
@@ -313,7 +313,7 @@ public:
 	 * @param Offset The vector to shift the box by.
 	 * @return A new bounding box.
 	 */
-	UE_NODISCARD FORCEINLINE TBox<T> ShiftBy( const TVector<T>& Offset ) const
+	[[nodiscard]] FORCEINLINE TBox<T> ShiftBy( const TVector<T>& Offset ) const
 	{
 		return TBox<T>(Min + Offset, Max + Offset);
 	}
@@ -324,7 +324,7 @@ public:
 	 * @param Destination The destination point to move center of box to.
 	 * @return A new bounding box.
 	 */
-	UE_NODISCARD FORCEINLINE TBox<T> MoveTo( const TVector<T>& Destination ) const
+	[[nodiscard]] FORCEINLINE TBox<T> MoveTo( const TVector<T>& Destination ) const
 	{
 		const TVector<T> Offset = Destination - GetCenter();
 		return TBox<T>(Min + Offset, Max + Offset);
@@ -436,7 +436,7 @@ public:
 	 * @param Other The bounding box to test overlap
 	 * @return the overlap box. It can be 0 if they don't overlap
 	 */
-	UE_NODISCARD TBox<T> Overlap( const TBox<T>& Other ) const;
+	[[nodiscard]] TBox<T> Overlap( const TBox<T>& Other ) const;
 
 	/**
 	  * Gets a bounding volume transformed by an inverted TTransform<T> object.
@@ -444,7 +444,7 @@ public:
 	  * @param M The transformation object to perform the inversely transform this box with.
 	  * @return	The transformed box.
 	  */
-	UE_NODISCARD TBox<T> InverseTransformBy( const TTransform<T>& M ) const;
+	[[nodiscard]] TBox<T> InverseTransformBy( const TTransform<T>& M ) const;
 
 	/** 
 	 * Checks whether the given location is inside this box.
@@ -562,7 +562,7 @@ public:
 	 * @return The transformed box.
 	 * @see TransformProjectBy
 	 */
-	UE_NODISCARD TBox<T> TransformBy( const TMatrix<T>& M ) const;
+	[[nodiscard]] TBox<T> TransformBy( const TMatrix<T>& M ) const;
 
 	/**
 	 * Gets a bounding volume transformed by a TTransform<T> object.
@@ -571,7 +571,7 @@ public:
 	 * @return The transformed box.
 	 * @see TransformProjectBy
 	 */
-	UE_NODISCARD TBox<T> TransformBy( const TTransform<T>& M ) const;
+	[[nodiscard]] TBox<T> TransformBy( const TTransform<T>& M ) const;
 
 	/** 
 	 * Returns the current world bounding box transformed and projected to screen space
@@ -580,7 +580,7 @@ public:
 	 * @return The transformed box.
 	 * @see TransformBy
 	 */
-	UE_NODISCARD TBox<T> TransformProjectBy( const TMatrix<T>& ProjM ) const;
+	[[nodiscard]] TBox<T> TransformProjectBy( const TMatrix<T>& ProjM ) const;
 
 	/**
 	 * Get a textual representation of this box.
