Replaced UE_NODISCARD with [[nodiscard]].

diff --git a/Engine/Source/Runtime/Core/Public/Math/UnrealMathUtility.h b/Engine/Source/Runtime/Core/Public/Math/UnrealMathUtility.h
--- a/Engine/Source/Runtime/Core/Public/Math/UnrealMathUtility.h
+++ b/Engine/Source/Runtime/Core/Public/Math/UnrealMathUtility.h
@@ -267,14 +267,14 @@ struct FMath : public FPlatformMath
 	// Random Number Functions
 
 	/** Helper function for rand implementations. Returns a random number in [0..A) */
-	UE_NODISCARD static FORCEINLINE int32 RandHelper(int32 A)
+	[[nodiscard]] static FORCEINLINE int32 RandHelper(int32 A)
 	{
 		// Note that on some platforms RAND_MAX is a large number so we cannot do ((rand()/(RAND_MAX+1)) * A)
 		// or else we may include the upper bound results, which should be excluded.
 		return A > 0 ? Min(TruncToInt(FRand() * (float)A), A - 1) : 0;
 	}
 
-	UE_NODISCARD static FORCEINLINE int64 RandHelper64(int64 A)
+	[[nodiscard]] static FORCEINLINE int64 RandHelper64(int64 A)
 	{
 		// Note that on some platforms RAND_MAX is a large number so we cannot do ((rand()/(RAND_MAX+1)) * A)
 		// or else we may include the upper bound results, which should be excluded.
@@ -282,37 +282,37 @@ struct FMath : public FPlatformMath
 	}
 
 	/** Helper function for rand implementations. Returns a random number >= Min and <= Max */
-	UE_NODISCARD static FORCEINLINE int32 RandRange(int32 Min, int32 Max)
+	[[nodiscard]] static FORCEINLINE int32 RandRange(int32 Min, int32 Max)
 	{
 		const int32 Range = (Max - Min) + 1;
 		return Min + RandHelper(Range);
 	}
 
-	UE_NODISCARD static FORCEINLINE int64 RandRange(int64 Min, int64 Max)
+	[[nodiscard]] static FORCEINLINE int64 RandRange(int64 Min, int64 Max)
 	{
 		const int64 Range = (Max - Min) + 1;
 		return Min + RandHelper64(Range);
 	}
 
 	/** Util to generate a random number in a range. Overloaded to distinguish from int32 version, where passing a float is typically a mistake. */
-	UE_NODISCARD static FORCEINLINE float RandRange(float InMin, float InMax)
+	[[nodiscard]] static FORCEINLINE float RandRange(float InMin, float InMax)
 	{
 		return FRandRange(InMin, InMax);
 	}
 
-	UE_NODISCARD static FORCEINLINE double RandRange(double InMin, double InMax)
+	[[nodiscard]] static FORCEINLINE double RandRange(double InMin, double InMax)
 	{
 		return FRandRange(InMin, InMax);
 	}
 
 	/** Util to generate a random number in a range. */
-	UE_NODISCARD static FORCEINLINE float FRandRange(float InMin, float InMax)
+	[[nodiscard]] static FORCEINLINE float FRandRange(float InMin, float InMax)
 	{
 		return InMin + (InMax - InMin) * FRand();
 	}
 
 	/** Util to generate a random number in a range. */
-	UE_NODISCARD static FORCEINLINE double FRandRange(double InMin, double InMax)
+	[[nodiscard]] static FORCEINLINE double FRandRange(double InMin, double InMax)
 	{
 		return InMin + (InMax - InMin) * FRand();	// LWC_TODO: Implement FRandDbl() for increased precision
 	}
@@ -320,31 +320,31 @@ struct FMath : public FPlatformMath
 	RESOLVE_FLOAT_AMBIGUITY_2_ARGS(FRandRange);
 
 	/** Util to generate a random boolean. */
-	UE_NODISCARD static FORCEINLINE bool RandBool()
+	[[nodiscard]] static FORCEINLINE bool RandBool()
 	{
 		return (RandRange(0,1) == 1) ? true : false;
 	}
 
 	/** Return a uniformly distributed random unit length vector = point on the unit sphere surface. */
-	UE_NODISCARD static FVector VRand();
+	[[nodiscard]] static FVector VRand();
 	
 	/**
 	 * Returns a random unit vector, uniformly distributed, within the specified cone
 	 * ConeHalfAngleRad is the half-angle of cone, in radians.  Returns a normalized vector. 
 	 */
-	UE_NODISCARD static CORE_API FVector VRandCone(FVector const& Dir, float ConeHalfAngleRad);
+	[[nodiscard]] static CORE_API FVector VRandCone(FVector const& Dir, float ConeHalfAngleRad);
 
 	/** 
 	 * This is a version of VRandCone that handles "squished" cones, i.e. with different angle limits in the Y and Z axes.
 	 * Assumes world Y and Z, although this could be extended to handle arbitrary rotations.
 	 */
-	UE_NODISCARD static CORE_API FVector VRandCone(FVector const& Dir, float HorizontalConeHalfAngleRad, float VerticalConeHalfAngleRad);
+	[[nodiscard]] static CORE_API FVector VRandCone(FVector const& Dir, float HorizontalConeHalfAngleRad, float VerticalConeHalfAngleRad);
 
 	/** Returns a random point, uniformly distributed, within the specified radius */
-	UE_NODISCARD static CORE_API FVector2D RandPointInCircle(float CircleRadius);
+	[[nodiscard]] static CORE_API FVector2D RandPointInCircle(float CircleRadius);
 
 	/** Returns a random point within the passed in bounding box */
-	UE_NODISCARD static CORE_API FVector RandPointInBox(const FBox& Box);
+	[[nodiscard]] static CORE_API FVector RandPointInBox(const FBox& Box);
 
 	/** 
 	 * Given a direction vector and a surface normal, returns the vector reflected across the surface normal.
@@ -355,13 +355,13 @@ struct FMath : public FPlatformMath
 	 *
 	 * @returns Reflected vector.
 	 */
-	UE_NODISCARD static CORE_API FVector GetReflectionVector(const FVector& Direction, const FVector& SurfaceNormal);
+	[[nodiscard]] static CORE_API FVector GetReflectionVector(const FVector& Direction, const FVector& SurfaceNormal);
 	
 	// Predicates
 
 	/** Checks if value is within a range, exclusive on MaxValue) */
 	template< class T, class U> 
-	UE_NODISCARD static constexpr FORCEINLINE bool IsWithin(const T& TestValue, const U& MinValue, const U& MaxValue)
+	[[nodiscard]] static constexpr FORCEINLINE bool IsWithin(const T& TestValue, const U& MinValue, const U& MaxValue)
 	{
 		return ((TestValue >= MinValue) && (TestValue < MaxValue));
 	}
@@ -369,7 +369,7 @@ struct FMath : public FPlatformMath
 
 	/** Checks if value is within a range, inclusive on MaxValue) */
 	template< class T, class U> 
-	UE_NODISCARD static constexpr FORCEINLINE bool IsWithinInclusive(const T& TestValue, const U& MinValue, const U& MaxValue)
+	[[nodiscard]] static constexpr FORCEINLINE bool IsWithinInclusive(const T& TestValue, const U& MinValue, const U& MaxValue)
 	{
 		return ((TestValue>=MinValue) && (TestValue <= MaxValue));
 	}
@@ -381,12 +381,12 @@ struct FMath : public FPlatformMath
 	 *	@param ErrorTolerance	Maximum allowed difference for considering them as 'nearly equal'
 	 *	@return					true if A and B are nearly equal
 	 */
-	UE_NODISCARD static FORCEINLINE bool IsNearlyEqual(float A, float B, float ErrorTolerance = UE_SMALL_NUMBER)
+	[[nodiscard]] static FORCEINLINE bool IsNearlyEqual(float A, float B, float ErrorTolerance = UE_SMALL_NUMBER)
 	{
 		return Abs<float>( A - B ) <= ErrorTolerance;
 	}
 
-	UE_NODISCARD static FORCEINLINE bool IsNearlyEqual(double A, double B, double ErrorTolerance = UE_DOUBLE_SMALL_NUMBER)
+	[[nodiscard]] static FORCEINLINE bool IsNearlyEqual(double A, double B, double ErrorTolerance = UE_DOUBLE_SMALL_NUMBER)
 	{
 		return Abs<double>(A - B) <= ErrorTolerance;
 	}
@@ -400,7 +400,7 @@ struct FMath : public FPlatformMath
 	 *	@param ErrorTolerance	Maximum allowed difference for considering Value as 'nearly zero'
 	 *	@return					true if Value is nearly zero
 	 */
-	UE_NODISCARD static FORCEINLINE bool IsNearlyZero(float Value, float ErrorTolerance = UE_SMALL_NUMBER)
+	[[nodiscard]] static FORCEINLINE bool IsNearlyZero(float Value, float ErrorTolerance = UE_SMALL_NUMBER)
 	{
 		return Abs<float>( Value ) <= ErrorTolerance;
 	}
@@ -411,7 +411,7 @@ struct FMath : public FPlatformMath
 	 *	@param ErrorTolerance	Maximum allowed difference for considering Value as 'nearly zero'
 	 *	@return					true if Value is nearly zero
 	 */
-	UE_NODISCARD static FORCEINLINE bool IsNearlyZero(double Value, double ErrorTolerance = UE_DOUBLE_SMALL_NUMBER)
+	[[nodiscard]] static FORCEINLINE bool IsNearlyZero(double Value, double ErrorTolerance = UE_DOUBLE_SMALL_NUMBER)
 	{
 		return Abs<double>( Value ) <= ErrorTolerance;
 	}
@@ -479,7 +479,7 @@ public:
 	 *	                        numbers are allowed to differ.
 	 *	@return					true if the two values are nearly equal.
 	 */
-	UE_NODISCARD static FORCEINLINE bool IsNearlyEqualByULP(float A, float B, int32 MaxUlps = 4)
+	[[nodiscard]] static FORCEINLINE bool IsNearlyEqualByULP(float A, float B, int32 MaxUlps = 4)
 	{
 		return TIsNearlyEqualByULP<float, uint32, uint32(1U << 31)>(A, B, MaxUlps);
 	}
@@ -501,7 +501,7 @@ public:
 	 *	                        numbers are allowed to differ.
 	 *	@return					true if the two values are nearly equal.
 	 */
-	UE_NODISCARD static FORCEINLINE bool IsNearlyEqualByULP(double A, double B, int32 MaxUlps = 4)
+	[[nodiscard]] static FORCEINLINE bool IsNearlyEqualByULP(double A, double B, int32 MaxUlps = 4)
 	{
 		return TIsNearlyEqualByULP<double, uint64, uint64(1ULL << 63)>(A, B, MaxUlps);
 	}
@@ -512,19 +512,19 @@ public:
 	*	@return			true if Value is a power of two
 	*/
 	template <typename T>
-	UE_NODISCARD static constexpr FORCEINLINE bool IsPowerOfTwo( T Value )
+	[[nodiscard]] static constexpr FORCEINLINE bool IsPowerOfTwo( T Value )
 	{
 		return ((Value & (Value - 1)) == (T)0);
 	}
 
 	/** Converts a float to a nearest less or equal integer. */
-	UE_NODISCARD static FORCEINLINE float Floor(float F)
+	[[nodiscard]] static FORCEINLINE float Floor(float F)
 	{
 		return FloorToFloat(F);
 	}
 
 	/** Converts a double to a nearest less or equal integer. */
-	UE_NODISCARD static FORCEINLINE double Floor(double F)
+	[[nodiscard]] static FORCEINLINE double Floor(double F)
 	{
 		return FloorToDouble(F);
 	}
@@ -537,7 +537,7 @@ public:
 		typename IntegralType
 		UE_REQUIRES(std::is_integral_v<IntegralType>)
 	>
-	UE_NODISCARD static constexpr FORCEINLINE IntegralType Floor(IntegralType I)
+	[[nodiscard]] static constexpr FORCEINLINE IntegralType Floor(IntegralType I)
 	{
 	    return I;
 	}
@@ -547,48 +547,48 @@ public:
 
 	/** Returns highest of 3 values */
 	template< class T > 
-	UE_NODISCARD static constexpr FORCEINLINE T Max3( const T A, const T B, const T C )
+	[[nodiscard]] static constexpr FORCEINLINE T Max3( const T A, const T B, const T C )
 	{
 		return Max ( Max( A, B ), C );
 	}
 
 	/** Returns lowest of 3 values */
 	template< class T > 
-	UE_NODISCARD static constexpr FORCEINLINE T Min3( const T A, const T B, const T C )
+	[[nodiscard]] static constexpr FORCEINLINE T Min3( const T A, const T B, const T C )
 	{
 		return Min ( Min( A, B ), C );
 	}
 
 	template< class T > 
-	UE_NODISCARD static constexpr FORCEINLINE int32 Max3Index( const T A, const T B, const T C )
+	[[nodiscard]] static constexpr FORCEINLINE int32 Max3Index( const T A, const T B, const T C )
 	{
 		return ( A > B ) ? ( ( A > C ) ? 0 : 2 ) : ( ( B > C ) ? 1 : 2 );
 	}
 
 	/** Returns index of the lowest value */
 	template< class T > 
-	UE_NODISCARD static constexpr FORCEINLINE int32 Min3Index( const T A, const T B, const T C )
+	[[nodiscard]] static constexpr FORCEINLINE int32 Min3Index( const T A, const T B, const T C )
 	{
 		return ( A < B ) ? ( ( A < C ) ? 0 : 2 ) : ( ( B < C ) ? 1 : 2 );
 	}
 
 	/** Multiples value by itself */
 	template< class T > 
-	UE_NODISCARD static constexpr FORCEINLINE T Square( const T A )
+	[[nodiscard]] static constexpr FORCEINLINE T Square( const T A )
 	{
 		return A*A;
 	}
 
 	/** Cubes the value */
 	template< class T > 
-    UE_NODISCARD static constexpr FORCEINLINE T Cube( const T A )
+    [[nodiscard]] static constexpr FORCEINLINE T Cube( const T A )
 	{
 		return A*A*A;
 	}
 
 	/** Clamps X to be between Min and Max, inclusive */
 	template< class T >
-	UE_NODISCARD static constexpr FORCEINLINE T Clamp(const T X, const T MinValue, const T MaxValue)
+	[[nodiscard]] static constexpr FORCEINLINE T Clamp(const T X, const T MinValue, const T MaxValue)
 	{
 		return Max(Min(X, MaxValue), MinValue);
 	}
@@ -596,16 +596,16 @@ public:
 	MIX_FLOATS_3_ARGS(Clamp);
 	
 	/** Clamps X to be between Min and Max, inclusive. Explicitly defined here for floats/doubles because static analysis gets confused between template and int versions. */
-	UE_NODISCARD static constexpr FORCEINLINE float Clamp(const float X, const float Min, const float Max) { return Clamp<float>(X, Min, Max); }
-	UE_NODISCARD static constexpr FORCEINLINE double Clamp(const double X, const double Min, const double Max) { return Clamp<double>(X, Min, Max); }
+	[[nodiscard]] static constexpr FORCEINLINE float Clamp(const float X, const float Min, const float Max) { return Clamp<float>(X, Min, Max); }
+	[[nodiscard]] static constexpr FORCEINLINE double Clamp(const double X, const double Min, const double Max) { return Clamp<double>(X, Min, Max); }
 
 	/** Clamps X to be between Min and Max, inclusive. Overload to support mixed int64/int32 types. */
-	UE_NODISCARD static constexpr FORCEINLINE int64 Clamp(const int64 X, const int32 Min, const int32 Max) { return Clamp<int64>(X, Min, Max); }
+	[[nodiscard]] static constexpr FORCEINLINE int64 Clamp(const int64 X, const int32 Min, const int32 Max) { return Clamp<int64>(X, Min, Max); }
 
 	/** Wraps X to be between Min and Max, inclusive. */
 	/** When X can wrap to both Min and Max, it will wrap to Min if it lies below the range and wrap to Max if it is above the range. */
 	template< class T >
-	UE_NODISCARD static constexpr FORCEINLINE T Wrap(const T X, const T Min, const T Max)
+	[[nodiscard]] static constexpr FORCEINLINE T Wrap(const T X, const T Min, const T Max)
 	{
 		// Use unsigned type for integers to allow for large ranges which don't overflow on subtraction
 		// We don't do that for floating point types because there are no unsigned versions of those.  We
@@ -638,7 +638,7 @@ public:
 	}
 
 	template <typename T>
-	UE_NODISCARD static constexpr FORCEINLINE T Modulo(T Value, T Base)
+	[[nodiscard]] static constexpr FORCEINLINE T Modulo(T Value, T Base)
 	{
 		if constexpr (std::is_floating_point_v<T>)
 		{
@@ -652,7 +652,7 @@ public:
 
 	/** Snaps a value to the nearest grid multiple */
 	template< class T >
-	UE_NODISCARD static constexpr FORCEINLINE T GridSnap(T Location, T Grid)
+	[[nodiscard]] static constexpr FORCEINLINE T GridSnap(T Location, T Grid)
 	{
 		return (Grid == T{}) ? Location : (Floor((Location + (Grid/(T)2)) / Grid) * Grid);
 	}
@@ -661,21 +661,21 @@ public:
 
 	/** Divides two integers and rounds up */
 	template <class T>
-	UE_NODISCARD static constexpr FORCEINLINE T DivideAndRoundUp(T Dividend, T Divisor)
+	[[nodiscard]] static constexpr FORCEINLINE T DivideAndRoundUp(T Dividend, T Divisor)
 	{
 		return (Dividend + Divisor - 1) / Divisor;
 	}
 
 	/** Divides two integers and rounds down */
 	template <class T>
-	UE_NODISCARD static constexpr FORCEINLINE T DivideAndRoundDown(T Dividend, T Divisor)
+	[[nodiscard]] static constexpr FORCEINLINE T DivideAndRoundDown(T Dividend, T Divisor)
 	{
 		return Dividend / Divisor;
 	}
 
 	/** Divides two integers and rounds to nearest */
 	template <class T>
-	UE_NODISCARD static constexpr FORCEINLINE T DivideAndRoundNearest(T Dividend, T Divisor)
+	[[nodiscard]] static constexpr FORCEINLINE T DivideAndRoundNearest(T Dividend, T Divisor)
 	{
 		return (Dividend >= 0)
 			? (Dividend + Divisor / 2) / Divisor
@@ -689,7 +689,7 @@ public:
 	 *
 	 * @return the base 2 log of the value
 	 */
-	UE_NODISCARD static FORCEINLINE float Log2(float Value)
+	[[nodiscard]] static FORCEINLINE float Log2(float Value)
 	{
 		// Cached value for fast conversions
 		constexpr float LogToLog2 = 1.44269502f; // 1.f / Loge(2.f)
@@ -704,7 +704,7 @@ public:
 	 *
 	 * @return the base 2 log of the value
 	 */
-	UE_NODISCARD static FORCEINLINE double Log2(double Value)
+	[[nodiscard]] static FORCEINLINE double Log2(double Value)
 	{
 		// Cached value for fast conversions
 		constexpr double LogToLog2 = 1.4426950408889634; // 1.0 / Loge(2.0);
@@ -794,7 +794,7 @@ public:
 	* @param Value  input angle
 	* @return ASin of Value
 	*/
-	UE_NODISCARD static FORCEINLINE float FastAsin(float Value)
+	[[nodiscard]] static FORCEINLINE float FastAsin(float Value)
 	{
 		// Clamp input to [-1,1].
 		const bool nonnegative = (Value >= 0.0f);
@@ -813,7 +813,7 @@ public:
 	}
 #undef FASTASIN_HALF_PI
 
-	UE_NODISCARD static FORCEINLINE double FastAsin(double Value)
+	[[nodiscard]] static FORCEINLINE double FastAsin(double Value)
 	{
 		// TODO: add fast approximation
 		return FMath::Asin(Value);
@@ -827,7 +827,7 @@ public:
 	 * @return					Value in degrees.
 	 */
 	template<class T>
-	UE_NODISCARD static constexpr FORCEINLINE auto RadiansToDegrees(T const& RadVal) -> decltype(RadVal * (180.f / UE_PI))
+	[[nodiscard]] static constexpr FORCEINLINE auto RadiansToDegrees(T const& RadVal) -> decltype(RadVal * (180.f / UE_PI))
 	{
 		return RadVal * (180.f / UE_PI);
 	}
@@ -841,7 +841,7 @@ public:
 	 * @return					Value in radians.
 	 */
 	template<class T>
-	UE_NODISCARD static constexpr FORCEINLINE auto DegreesToRadians(T const& DegVal) -> decltype(DegVal * (UE_PI / 180.f))
+	[[nodiscard]] static constexpr FORCEINLINE auto DegreesToRadians(T const& DegVal) -> decltype(DegVal * (UE_PI / 180.f))
 	{
 		return DegVal * (UE_PI / 180.f);
 	}
@@ -857,7 +857,7 @@ public:
 	 * @return Returns clamped angle in the range -180..180.
 	 */
 	template<typename T>
-	UE_NODISCARD static T ClampAngle(T AngleDegrees, T MinAngleDegrees, T MaxAngleDegrees);
+	[[nodiscard]] static T ClampAngle(T AngleDegrees, T MinAngleDegrees, T MaxAngleDegrees);
 
 	RESOLVE_FLOAT_AMBIGUITY_3_ARGS(ClampAngle);
 
@@ -867,7 +867,7 @@ public:
 		typename T2
 		UE_REQUIRES(std::is_floating_point_v<T> || std::is_floating_point_v<T2>)
 	>
-	UE_NODISCARD static constexpr auto FindDeltaAngleDegrees(T A1, T2 A2) -> decltype(A1 * A2)
+	[[nodiscard]] static constexpr auto FindDeltaAngleDegrees(T A1, T2 A2) -> decltype(A1 * A2)
 	{
 		// Find the difference
 		auto Delta = A2 - A1;
@@ -895,7 +895,7 @@ public:
 		typename T2
 		UE_REQUIRES(std::is_floating_point_v<T> || std::is_floating_point_v<T2>)
 	>
-	UE_NODISCARD static constexpr auto FindDeltaAngleRadians(T A1, T2 A2) -> decltype(A1 * A2)
+	[[nodiscard]] static constexpr auto FindDeltaAngleRadians(T A1, T2 A2) -> decltype(A1 * A2)
 	{
 		// Find the difference
 		auto Delta = A2 - A1;
@@ -922,7 +922,7 @@ public:
 		typename T
 		UE_REQUIRES(std::is_floating_point_v<T>)
 	>
-	UE_NODISCARD static constexpr T UnwindRadians(T A)
+	[[nodiscard]] static constexpr T UnwindRadians(T A)
 	{
 		while(A > UE_PI)
 		{
@@ -942,7 +942,7 @@ public:
 		typename T
 		UE_REQUIRES(std::is_floating_point_v<T>)
 	>
-	UE_NODISCARD static constexpr T UnwindDegrees(T A)
+	[[nodiscard]] static constexpr T UnwindDegrees(T A)
 	{
 		while(A > 180.f)
 		{
@@ -974,7 +974,7 @@ public:
 	 *
 	 * @return a new rotation component value
 	 */
-	UE_NODISCARD static CORE_API float FixedTurn(float InCurrent, float InDesired, float InDeltaRate);
+	[[nodiscard]] static CORE_API float FixedTurn(float InCurrent, float InDesired, float InDeltaRate);
 
 	/** Converts given Cartesian coordinate pair to Polar coordinate system. */
 	template<typename T>
@@ -1039,7 +1039,7 @@ public:
 	 * @return	FVector2D	X = Azimuth angle (in radians) (-PI, +PI)
 	 *						Y = Elevation angle (in radians) (-PI/2, +PI/2)
 	 */
-	UE_NODISCARD static CORE_API FVector2D GetAzimuthAndElevation(const FVector &Direction, const FVector &AxisX, const FVector &AxisY, const FVector &AxisZ);
+	[[nodiscard]] static CORE_API FVector2D GetAzimuthAndElevation(const FVector &Direction, const FVector &AxisX, const FVector &AxisY, const FVector &AxisZ);
 
 	// Interpolation Functions
 
@@ -1049,7 +1049,7 @@ public:
 		typename T2
 		UE_REQUIRES(std::is_floating_point_v<T>)
 	>
-	UE_NODISCARD static constexpr FORCEINLINE auto GetRangePct(T MinValue, T MaxValue, T2 Value)
+	[[nodiscard]] static constexpr FORCEINLINE auto GetRangePct(T MinValue, T MaxValue, T2 Value)
 	{
 		// Avoid Divide by Zero.
 		// But also if our range is a point, output whether Value is before or after.
@@ -1069,7 +1069,7 @@ public:
 		typename T2
 		UE_REQUIRES(std::is_floating_point_v<T>)
 	>
-	UE_NODISCARD static auto GetRangePct(UE::Math::TVector2<T> const& Range, T2 Value)
+	[[nodiscard]] static auto GetRangePct(UE::Math::TVector2<T> const& Range, T2 Value)
 	{
 		return GetRangePct(Range.X, Range.Y, Value);
 	}
@@ -1080,14 +1080,14 @@ public:
 		typename T2
 		UE_REQUIRES(std::is_floating_point_v<T>)
 	>
-	UE_NODISCARD static auto GetRangeValue(UE::Math::TVector2<T> const& Range, T2 Pct)
+	[[nodiscard]] static auto GetRangeValue(UE::Math::TVector2<T> const& Range, T2 Pct)
 	{
 		return Lerp(Range.X, Range.Y, Pct);
 	}
 
 	/** For the given Value clamped to the [Input:Range] inclusive, returns the corresponding percentage in [Output:Range] Inclusive. */
 	template<typename T, typename T2>
-	UE_NODISCARD static FORCEINLINE auto GetMappedRangeValueClamped(const UE::Math::TVector2<T>& InputRange, const UE::Math::TVector2<T>& OutputRange, const T2 Value)
+	[[nodiscard]] static FORCEINLINE auto GetMappedRangeValueClamped(const UE::Math::TVector2<T>& InputRange, const UE::Math::TVector2<T>& OutputRange, const T2 Value)
 	{
 		using RangePctType = decltype(T() * T2());
 		const RangePctType ClampedPct = Clamp<RangePctType>(GetRangePct(InputRange, Value), 0.f, 1.f);
@@ -1096,25 +1096,25 @@ public:
 
 	/** Transform the given Value relative to the input range to the Output Range. */
 	template<typename T, typename T2>
-	UE_NODISCARD static FORCEINLINE auto GetMappedRangeValueUnclamped(const UE::Math::TVector2<T>& InputRange, const UE::Math::TVector2<T>& OutputRange, const T2 Value)
+	[[nodiscard]] static FORCEINLINE auto GetMappedRangeValueUnclamped(const UE::Math::TVector2<T>& InputRange, const UE::Math::TVector2<T>& OutputRange, const T2 Value)
 	{
 		return GetRangeValue(OutputRange, GetRangePct(InputRange, Value));
 	}
 
 	template<class T>
-	UE_NODISCARD static FORCEINLINE double GetRangePct(TRange<T> const& Range, T Value)
+	[[nodiscard]] static FORCEINLINE double GetRangePct(TRange<T> const& Range, T Value)
 	{
 		return GetRangePct(Range.GetLowerBoundValue(), Range.GetUpperBoundValue(), Value);
 	}
 
 	template<class T>
-	UE_NODISCARD static FORCEINLINE T GetRangeValue(TRange<T> const& Range, T Pct)
+	[[nodiscard]] static FORCEINLINE T GetRangeValue(TRange<T> const& Range, T Pct)
 	{
 		return FMath::Lerp<T>(Range.GetLowerBoundValue(), Range.GetUpperBoundValue(), Pct);
 	}
 
 	template<class T>
-	UE_NODISCARD static FORCEINLINE T GetMappedRangeValueClamped(const TRange<T>& InputRange, const TRange<T>& OutputRange, const T Value)
+	[[nodiscard]] static FORCEINLINE T GetMappedRangeValueClamped(const TRange<T>& InputRange, const TRange<T>& OutputRange, const T Value)
 	{
 		const T ClampedPct = FMath::Clamp<T>(GetRangePct(InputRange, Value), 0, 1);
 		return GetRangeValue(OutputRange, ClampedPct);
@@ -1126,7 +1126,7 @@ public:
 		typename U
 		UE_REQUIRES(!TCustomLerp<T>::Value && (std::is_floating_point_v<U> || std::is_same_v<T, U>))
 	>
-	UE_NODISCARD static constexpr FORCEINLINE_DEBUGGABLE T Lerp( const T& A, const T& B, const U& Alpha )
+	[[nodiscard]] static constexpr FORCEINLINE_DEBUGGABLE T Lerp( const T& A, const T& B, const U& Alpha )
 	{
 		return (T)(A + Alpha * (B-A));
 	}
@@ -1137,7 +1137,7 @@ public:
 		typename U
 		UE_REQUIRES(TCustomLerp<T>::Value)
 	>
-	UE_NODISCARD static FORCEINLINE_DEBUGGABLE T Lerp(const T& A, const T& B, const U& Alpha)
+	[[nodiscard]] static FORCEINLINE_DEBUGGABLE T Lerp(const T& A, const T& B, const U& Alpha)
 	{
 		return TCustomLerp<T>::Lerp(A, B, Alpha);
 	}
@@ -1149,7 +1149,7 @@ public:
 		typename T3
 		UE_REQUIRES(!std::is_same_v<T1, T2> && !TCustomLerp<T1>::Value && !TCustomLerp<T2>::Value)
 	>
-	UE_NODISCARD static auto Lerp( const T1& A, const T2& B, const T3& Alpha ) -> decltype(A * B)
+	[[nodiscard]] static auto Lerp( const T1& A, const T2& B, const T3& Alpha ) -> decltype(A * B)
 	{
 		using ABType = decltype(A * B);
 		return Lerp(ABType(A), ABType(B), Alpha);
@@ -1157,14 +1157,14 @@ public:
 
 	/** Performs a linear interpolation between two values, Alpha ranges from 0-1. Handles full numeric range of T */
 	template< class T > 
-	UE_NODISCARD static constexpr FORCEINLINE_DEBUGGABLE T LerpStable( const T& A, const T& B, double Alpha )
+	[[nodiscard]] static constexpr FORCEINLINE_DEBUGGABLE T LerpStable( const T& A, const T& B, double Alpha )
 	{
 		return (T)((A * (1.0 - Alpha)) + (B * Alpha));
 	}
 
 	/** Performs a linear interpolation between two values, Alpha ranges from 0-1. Handles full numeric range of T */
 	template< class T >
-	UE_NODISCARD static constexpr FORCEINLINE_DEBUGGABLE T LerpStable(const T& A, const T& B, float Alpha)
+	[[nodiscard]] static constexpr FORCEINLINE_DEBUGGABLE T LerpStable(const T& A, const T& B, float Alpha)
 	{
 		return (T)((A * (1.0f - Alpha)) + (B * Alpha));
 	}
@@ -1176,7 +1176,7 @@ public:
 		typename T3
 		UE_REQUIRES(!std::is_same_v<T1, T2>)
 	>
-	UE_NODISCARD static auto LerpStable( const T1& A, const T2& B, const T3& Alpha ) -> decltype(A * B)
+	[[nodiscard]] static auto LerpStable( const T1& A, const T2& B, const T3& Alpha ) -> decltype(A * B)
 	{
 		using ABType = decltype(A * B);
 		return LerpStable(ABType(A), ABType(B), Alpha);
@@ -1188,7 +1188,7 @@ public:
 		typename U
 		UE_REQUIRES(!TCustomLerp<T>::Value && (std::is_floating_point_v<U> || std::is_same_v<T, U>))
 	>
-	UE_NODISCARD static constexpr FORCEINLINE_DEBUGGABLE T BiLerp(const T& P00,const T& P10,const T& P01,const T& P11, const U& FracX, const U& FracY)
+	[[nodiscard]] static constexpr FORCEINLINE_DEBUGGABLE T BiLerp(const T& P00,const T& P10,const T& P01,const T& P11, const U& FracX, const U& FracY)
 	{
 		return Lerp(
 			Lerp(P00,P10,FracX),
@@ -1203,7 +1203,7 @@ public:
 		typename U
 		UE_REQUIRES(TCustomLerp<T>::Value)
 	>
-	UE_NODISCARD static FORCEINLINE_DEBUGGABLE T BiLerp(const T& P00, const T& P10, const T& P01, const T& P11, const U& FracX, const U& FracY)
+	[[nodiscard]] static FORCEINLINE_DEBUGGABLE T BiLerp(const T& P00, const T& P10, const T& P01, const T& P11, const U& FracX, const U& FracY)
 	{
 		return TCustomLerp<T>::BiLerp(P00, P10, P01, P11, FracX, FracY);
 	}
@@ -1222,7 +1222,7 @@ public:
 		typename U
 		UE_REQUIRES(!TCustomLerp<T>::Value && (std::is_floating_point_v<U> || std::is_same_v<T, U>))
 	>
-	UE_NODISCARD static constexpr FORCEINLINE_DEBUGGABLE T CubicInterp( const T& P0, const T& T0, const T& P1, const T& T1, const U& A )
+	[[nodiscard]] static constexpr FORCEINLINE_DEBUGGABLE T CubicInterp( const T& P0, const T& T0, const T& P1, const T& T1, const U& A )
 	{
 		const U A2 = A * A;
 		const U A3 = A2 * A;
@@ -1236,7 +1236,7 @@ public:
 		typename U
 		UE_REQUIRES(TCustomLerp<T>::Value)
 	>
-	UE_NODISCARD static FORCEINLINE_DEBUGGABLE T CubicInterp(const T& P0, const T& T0, const T& P1, const T& T1, const U& A)
+	[[nodiscard]] static FORCEINLINE_DEBUGGABLE T CubicInterp(const T& P0, const T& T0, const T& P1, const T& T1, const U& A)
 	{
 		return TCustomLerp<T>::CubicInterp(P0, T0, P1, T1, A);
 	}
@@ -1255,7 +1255,7 @@ public:
 		typename U
 		UE_REQUIRES(std::is_floating_point_v<U>)
 	>
-	UE_NODISCARD static constexpr FORCEINLINE_DEBUGGABLE T CubicInterpDerivative( const T& P0, const T& T0, const T& P1, const T& T1, const U& A )
+	[[nodiscard]] static constexpr FORCEINLINE_DEBUGGABLE T CubicInterpDerivative( const T& P0, const T& T0, const T& P1, const T& T1, const U& A )
 	{
 		T a = 6.f*P0 + 3.f*T0 + 3.f*T1 - 6.f*P1;
 		T b = -6.f*P0 - 4.f*T0 - 2.f*T1 + 6.f*P1;
@@ -1280,7 +1280,7 @@ public:
 		typename U
 		UE_REQUIRES(std::is_floating_point_v<U>)
 	>
-	UE_NODISCARD static constexpr FORCEINLINE_DEBUGGABLE T CubicInterpSecondDerivative( const T& P0, const T& T0, const T& P1, const T& T1, const U& A )
+	[[nodiscard]] static constexpr FORCEINLINE_DEBUGGABLE T CubicInterpSecondDerivative( const T& P0, const T& T0, const T& P1, const T& T1, const U& A )
 	{
 		T a = 12.f*P0 + 6.f*T0 + 6.f*T1 - 12.f*P1;
 		T b = -6.f*P0 - 4.f*T0 - 2.f*T1 + 6.f*P1;
@@ -1290,7 +1290,7 @@ public:
 
 	/** Interpolate between A and B, applying an ease in function.  Exp controls the degree of the curve. */
 	template< class T >
-	UE_NODISCARD static FORCEINLINE_DEBUGGABLE T InterpEaseIn(const T& A, const T& B, float Alpha, float Exp)
+	[[nodiscard]] static FORCEINLINE_DEBUGGABLE T InterpEaseIn(const T& A, const T& B, float Alpha, float Exp)
 	{
 		float const ModifiedAlpha = Pow(Alpha, Exp);
 		return Lerp<T>(A, B, ModifiedAlpha);
@@ -1298,7 +1298,7 @@ public:
 
 	/** Interpolate between A and B, applying an ease out function.  Exp controls the degree of the curve. */
 	template< class T >
-	UE_NODISCARD static FORCEINLINE_DEBUGGABLE T InterpEaseOut(const T& A, const T& B, float Alpha, float Exp)
+	[[nodiscard]] static FORCEINLINE_DEBUGGABLE T InterpEaseOut(const T& A, const T& B, float Alpha, float Exp)
 	{
 		float const ModifiedAlpha = 1.f - Pow(1.f - Alpha, Exp);
 		return Lerp<T>(A, B, ModifiedAlpha);
@@ -1306,7 +1306,7 @@ public:
 
 	/** Interpolate between A and B, applying an ease in/out function.  Exp controls the degree of the curve. */
 	template< class T > 
-	UE_NODISCARD static FORCEINLINE_DEBUGGABLE T InterpEaseInOut( const T& A, const T& B, float Alpha, float Exp )
+	[[nodiscard]] static FORCEINLINE_DEBUGGABLE T InterpEaseInOut( const T& A, const T& B, float Alpha, float Exp )
 	{
 		return Lerp<T>(A, B, (Alpha < 0.5f) ?
 			InterpEaseIn(0.f, 1.f, Alpha * 2.f, Exp) * 0.5f :
@@ -1315,7 +1315,7 @@ public:
 
 	/** Interpolation between A and B, applying a step function. */
 	template< class T >
-	UE_NODISCARD static constexpr FORCEINLINE_DEBUGGABLE T InterpStep(const T& A, const T& B, float Alpha, int32 Steps)
+	[[nodiscard]] static constexpr FORCEINLINE_DEBUGGABLE T InterpStep(const T& A, const T& B, float Alpha, int32 Steps)
 	{
 		if (Steps <= 1 || Alpha <= 0)
 		{
@@ -1334,7 +1334,7 @@ public:
 
 	/** Interpolation between A and B, applying a sinusoidal in function. */
 	template< class T >
-	UE_NODISCARD static FORCEINLINE_DEBUGGABLE T InterpSinIn(const T& A, const T& B, float Alpha)
+	[[nodiscard]] static FORCEINLINE_DEBUGGABLE T InterpSinIn(const T& A, const T& B, float Alpha)
 	{
 		float const ModifiedAlpha = -1.f * Cos(Alpha * UE_HALF_PI) + 1.f;
 		return Lerp<T>(A, B, ModifiedAlpha);
@@ -1342,7 +1342,7 @@ public:
 	
 	/** Interpolation between A and B, applying a sinusoidal out function. */
 	template< class T >
-	UE_NODISCARD static FORCEINLINE_DEBUGGABLE T InterpSinOut(const T& A, const T& B, float Alpha)
+	[[nodiscard]] static FORCEINLINE_DEBUGGABLE T InterpSinOut(const T& A, const T& B, float Alpha)
 	{
 		float const ModifiedAlpha = Sin(Alpha * UE_HALF_PI);
 		return Lerp<T>(A, B, ModifiedAlpha);
@@ -1350,7 +1350,7 @@ public:
 
 	/** Interpolation between A and B, applying a sinusoidal in/out function. */
 	template< class T >
-	UE_NODISCARD static FORCEINLINE_DEBUGGABLE T InterpSinInOut(const T& A, const T& B, float Alpha)
+	[[nodiscard]] static FORCEINLINE_DEBUGGABLE T InterpSinInOut(const T& A, const T& B, float Alpha)
 	{
 		return Lerp<T>(A, B, (Alpha < 0.5f) ?
 			InterpSinIn(0.f, 1.f, Alpha * 2.f) * 0.5f :
@@ -1359,7 +1359,7 @@ public:
 
 	/** Interpolation between A and B, applying an exponential in function. */
 	template< class T >
-	UE_NODISCARD static FORCEINLINE_DEBUGGABLE T InterpExpoIn(const T& A, const T& B, float Alpha)
+	[[nodiscard]] static FORCEINLINE_DEBUGGABLE T InterpExpoIn(const T& A, const T& B, float Alpha)
 	{
 		float const ModifiedAlpha = (Alpha == 0.f) ? 0.f : Pow(2.f, 10.f * (Alpha - 1.f));
 		return Lerp<T>(A, B, ModifiedAlpha);
@@ -1367,7 +1367,7 @@ public:
 
 	/** Interpolation between A and B, applying an exponential out function. */
 	template< class T >
-	UE_NODISCARD static FORCEINLINE_DEBUGGABLE T InterpExpoOut(const T& A, const T& B, float Alpha)
+	[[nodiscard]] static FORCEINLINE_DEBUGGABLE T InterpExpoOut(const T& A, const T& B, float Alpha)
 	{
 		float const ModifiedAlpha = (Alpha == 1.f) ? 1.f : -Pow(2.f, -10.f * Alpha) + 1.f;
 		return Lerp<T>(A, B, ModifiedAlpha);
@@ -1375,7 +1375,7 @@ public:
 
 	/** Interpolation between A and B, applying an exponential in/out function. */
 	template< class T >
-	UE_NODISCARD static FORCEINLINE_DEBUGGABLE T InterpExpoInOut(const T& A, const T& B, float Alpha)
+	[[nodiscard]] static FORCEINLINE_DEBUGGABLE T InterpExpoInOut(const T& A, const T& B, float Alpha)
 	{
 		return Lerp<T>(A, B, (Alpha < 0.5f) ?
 			InterpExpoIn(0.f, 1.f, Alpha * 2.f) * 0.5f :
@@ -1384,7 +1384,7 @@ public:
 
 	/** Interpolation between A and B, applying a circular in function. */
 	template< class T >
-	UE_NODISCARD static FORCEINLINE_DEBUGGABLE T InterpCircularIn(const T& A, const T& B, float Alpha)
+	[[nodiscard]] static FORCEINLINE_DEBUGGABLE T InterpCircularIn(const T& A, const T& B, float Alpha)
 	{
 		float const ModifiedAlpha = -1.f * (Sqrt(1.f - Alpha * Alpha) - 1.f);
 		return Lerp<T>(A, B, ModifiedAlpha);
@@ -1392,7 +1392,7 @@ public:
 
 	/** Interpolation between A and B, applying a circular out function. */
 	template< class T >
-	UE_NODISCARD static FORCEINLINE_DEBUGGABLE T InterpCircularOut(const T& A, const T& B, float Alpha)
+	[[nodiscard]] static FORCEINLINE_DEBUGGABLE T InterpCircularOut(const T& A, const T& B, float Alpha)
 	{
 		Alpha -= 1.f;
 		float const ModifiedAlpha = Sqrt(1.f - Alpha  * Alpha);
@@ -1401,7 +1401,7 @@ public:
 
 	/** Interpolation between A and B, applying a circular in/out function. */
 	template< class T >
-	UE_NODISCARD static FORCEINLINE_DEBUGGABLE T InterpCircularInOut(const T& A, const T& B, float Alpha)
+	[[nodiscard]] static FORCEINLINE_DEBUGGABLE T InterpCircularInOut(const T& A, const T& B, float Alpha)
 	{
 		return Lerp<T>(A, B, (Alpha < 0.5f) ?
 			InterpCircularIn(0.f, 1.f, Alpha * 2.f) * 0.5f :
@@ -1411,7 +1411,7 @@ public:
 	// Rotator specific interpolation
 	// Similar to Lerp, but does not take the shortest path. Allows interpolation over more than 180 degrees.
 	template< typename T, typename U >
-	UE_NODISCARD static UE::Math::TRotator<T> LerpRange(const UE::Math::TRotator<T>& A, const UE::Math::TRotator<T>& B, U Alpha);
+	[[nodiscard]] static UE::Math::TRotator<T> LerpRange(const UE::Math::TRotator<T>& A, const UE::Math::TRotator<T>& B, U Alpha);
 
 	/*
 	 *	Cubic Catmull-Rom Spline interpolation. Based on http://www.cemyuksel.com/research/catmullrom_param/catmullrom.pdf 
@@ -1425,7 +1425,7 @@ public:
 	 *	T - The interpolation factor in the range 0 to 1. 0 returns P1. 1 returns P2.
 	 */
 	template< class U > 
-	UE_NODISCARD static constexpr FORCEINLINE_DEBUGGABLE U CubicCRSplineInterp(const U& P0, const U& P1, const U& P2, const U& P3, const float T0, const float T1, const float T2, const float T3, const float T)
+	[[nodiscard]] static constexpr FORCEINLINE_DEBUGGABLE U CubicCRSplineInterp(const U& P0, const U& P1, const U& P2, const U& P3, const float T0, const float T1, const float T2, const float T3, const float T)
 	{
 		//Based on http://www.cemyuksel.com/research/catmullrom_param/catmullrom.pdf 
 		float InvT1MinusT0 = 1.0f / (T1 - T0);
@@ -1445,7 +1445,7 @@ public:
 
 	/* Same as CubicCRSplineInterp but with additional saftey checks. If the checks fail P1 is returned. **/
 	template< class U >
-	UE_NODISCARD static constexpr FORCEINLINE_DEBUGGABLE U CubicCRSplineInterpSafe(const U& P0, const U& P1, const U& P2, const U& P3, const float T0, const float T1, const float T2, const float T3, const float T)
+	[[nodiscard]] static constexpr FORCEINLINE_DEBUGGABLE U CubicCRSplineInterpSafe(const U& P0, const U& P1, const U& P2, const U& P3, const float T0, const float T1, const float T2, const float T3, const float T)
 	{
 		//Based on http://www.cemyuksel.com/research/catmullrom_param/catmullrom.pdf 
 		float T1MinusT0 = (T1 - T0);
@@ -1478,29 +1478,29 @@ public:
 	// Special-case interpolation
 
 	/** Interpolate a normal vector Current to Target, by interpolating the angle between those vectors with constant step. */
-	UE_NODISCARD static CORE_API FVector VInterpNormalRotationTo(const FVector& Current, const FVector& Target, float DeltaTime, float RotationSpeedDegrees);
+	[[nodiscard]] static CORE_API FVector VInterpNormalRotationTo(const FVector& Current, const FVector& Target, float DeltaTime, float RotationSpeedDegrees);
 
 	/** Interpolate vector from Current to Target with constant step */
-	UE_NODISCARD static CORE_API FVector VInterpConstantTo(const FVector& Current, const FVector& Target, float DeltaTime, float InterpSpeed);
+	[[nodiscard]] static CORE_API FVector VInterpConstantTo(const FVector& Current, const FVector& Target, float DeltaTime, float InterpSpeed);
 
 	/** Interpolate vector from Current to Target. Scaled by distance to Target, so it has a strong start speed and ease out. */
-	UE_NODISCARD static CORE_API FVector VInterpTo( const FVector& Current, const FVector& Target, float DeltaTime, float InterpSpeed );
+	[[nodiscard]] static CORE_API FVector VInterpTo( const FVector& Current, const FVector& Target, float DeltaTime, float InterpSpeed );
 	
 	/** Interpolate vector2D from Current to Target with constant step */
-	UE_NODISCARD static CORE_API FVector2D Vector2DInterpConstantTo( const FVector2D& Current, const FVector2D& Target, float DeltaTime, float InterpSpeed );
+	[[nodiscard]] static CORE_API FVector2D Vector2DInterpConstantTo( const FVector2D& Current, const FVector2D& Target, float DeltaTime, float InterpSpeed );
 
 	/** Interpolate vector2D from Current to Target. Scaled by distance to Target, so it has a strong start speed and ease out. */
-	UE_NODISCARD static CORE_API FVector2D Vector2DInterpTo( const FVector2D& Current, const FVector2D& Target, float DeltaTime, float InterpSpeed );
+	[[nodiscard]] static CORE_API FVector2D Vector2DInterpTo( const FVector2D& Current, const FVector2D& Target, float DeltaTime, float InterpSpeed );
 
 	/** Interpolate rotator from Current to Target with constant step */
-	UE_NODISCARD static CORE_API FRotator RInterpConstantTo( const FRotator& Current, const FRotator& Target, float DeltaTime, float InterpSpeed);
+	[[nodiscard]] static CORE_API FRotator RInterpConstantTo( const FRotator& Current, const FRotator& Target, float DeltaTime, float InterpSpeed);
 
 	/** Interpolate rotator from Current to Target. Scaled by distance to Target, so it has a strong start speed and ease out. */
-	UE_NODISCARD static CORE_API FRotator RInterpTo( const FRotator& Current, const FRotator& Target, float DeltaTime, float InterpSpeed);
+	[[nodiscard]] static CORE_API FRotator RInterpTo( const FRotator& Current, const FRotator& Target, float DeltaTime, float InterpSpeed);
 
 	/** Interpolate float from Current to Target with constant step */
 	template<typename T1, typename T2 = T1, typename T3 = T2, typename T4 = T3>
-	UE_NODISCARD static auto FInterpConstantTo( T1 Current, T2 Target, T3 DeltaTime, T4 InterpSpeed )
+	[[nodiscard]] static auto FInterpConstantTo( T1 Current, T2 Target, T3 DeltaTime, T4 InterpSpeed )
 	{
 		using RetType = decltype(T1() * T2() * T3() * T4());
 	
@@ -1518,7 +1518,7 @@ public:
 
 	/** Interpolate float from Current to Target. Scaled by distance to Target, so it has a strong start speed and ease out. */
 	template<typename T1, typename T2 = T1, typename T3 = T2, typename T4 = T3>
-	UE_NODISCARD static auto FInterpTo( T1  Current, T2 Target, T3 DeltaTime, T4 InterpSpeed )
+	[[nodiscard]] static auto FInterpTo( T1  Current, T2 Target, T3 DeltaTime, T4 InterpSpeed )
 	{
 		using RetType = decltype(T1() * T2() * T3() * T4());
 	
@@ -1544,15 +1544,15 @@ public:
 	}
 
 	/** Interpolate Linear Color from Current to Target. Scaled by distance to Target, so it has a strong start speed and ease out. */
-	UE_NODISCARD static CORE_API FLinearColor CInterpTo(const FLinearColor& Current, const FLinearColor& Target, float DeltaTime, float InterpSpeed);
+	[[nodiscard]] static CORE_API FLinearColor CInterpTo(const FLinearColor& Current, const FLinearColor& Target, float DeltaTime, float InterpSpeed);
 
 	/** Interpolate quaternion from Current to Target with constant step (in radians) */
 	template< class T >
-	UE_NODISCARD static CORE_API UE::Math::TQuat<T> QInterpConstantTo(const UE::Math::TQuat<T>& Current, const UE::Math::TQuat<T>& Target, float DeltaTime, float InterpSpeed);
+	[[nodiscard]] static CORE_API UE::Math::TQuat<T> QInterpConstantTo(const UE::Math::TQuat<T>& Current, const UE::Math::TQuat<T>& Target, float DeltaTime, float InterpSpeed);
 
 	/** Interpolate quaternion from Current to Target. Scaled by angle to Target, so it has a strong start speed and ease out. */
 	template< class T >
-	UE_NODISCARD static CORE_API UE::Math::TQuat<T> QInterpTo(const UE::Math::TQuat<T>& Current, const UE::Math::TQuat<T>& Target, float DeltaTime, float InterpSpeed);
+	[[nodiscard]] static CORE_API UE::Math::TQuat<T> QInterpTo(const UE::Math::TQuat<T>& Current, const UE::Math::TQuat<T>& Target, float DeltaTime, float InterpSpeed);
 
 	/**
 	* Returns an approximation of Exp(-X) based on a Taylor expansion that has had the coefficients adjusted (using
@@ -1561,7 +1561,7 @@ public:
 	* to zero for large X). 
 	*/
 	template<class T>
-	UE_NODISCARD static constexpr T InvExpApprox(T X)
+	[[nodiscard]] static constexpr T InvExpApprox(T X)
 	{
 		constexpr T A(1.00746054f); // 1 / 1! in Taylor series
 		constexpr T B(0.45053901f); // 1 / 2! in Taylor series
@@ -1813,7 +1813,7 @@ public:
 	 *
 	 * @return  Pulsating value (0.0-1.0)
 	 */
-	UE_NODISCARD static float MakePulsatingValue( const double InCurrentTime, const float InPulsesPerSecond, const float InPhase = 0.0f )
+	[[nodiscard]] static float MakePulsatingValue( const double InCurrentTime, const float InPulsesPerSecond, const float InPhase = 0.0f )
 	{
 		return 0.5f + 0.5f * FMath::Sin( ( ( 0.25f + InPhase ) * UE_TWO_PI ) + ( (float)InCurrentTime * UE_TWO_PI ) * InPulsesPerSecond );
 	}
@@ -1831,7 +1831,7 @@ public:
 	 * @return The point of intersection between the ray and the plane.
 	 */
 	template<typename FReal>
-	UE_NODISCARD static UE::Math::TVector<FReal> RayPlaneIntersection(const UE::Math::TVector<FReal>& RayOrigin, const UE::Math::TVector<FReal>& RayDirection, const UE::Math::TPlane<FReal>& Plane);
+	[[nodiscard]] static UE::Math::TVector<FReal> RayPlaneIntersection(const UE::Math::TVector<FReal>& RayOrigin, const UE::Math::TVector<FReal>& RayDirection, const UE::Math::TPlane<FReal>& Plane);
 
 	/**
 	 * Find the intersection of a ray and a plane.  The ray has a start point with an infinite length.  Assumes that the
@@ -1844,7 +1844,7 @@ public:
 	 * @return The distance parameter along ray of the point of intersection between the ray and the plane.
 	 */
 	template<typename FReal>
-	UE_NODISCARD static FReal RayPlaneIntersectionParam(const UE::Math::TVector<FReal>& RayOrigin, const UE::Math::TVector<FReal>& RayDirection, const UE::Math::TPlane<FReal>& Plane);
+	[[nodiscard]] static FReal RayPlaneIntersectionParam(const UE::Math::TVector<FReal>& RayOrigin, const UE::Math::TVector<FReal>& RayDirection, const UE::Math::TPlane<FReal>& Plane);
 
 	/**
 	 * Find the intersection of a line and an offset plane. Assumes that the
@@ -1859,7 +1859,7 @@ public:
 	 * @return The point of intersection between the line and the plane.
 	 */
 	template<typename FReal>
-	UE_NODISCARD static UE::Math::TVector<FReal> LinePlaneIntersection(const UE::Math::TVector<FReal>& Point1, const UE::Math::TVector<FReal>& Point2, const UE::Math::TVector<FReal>& PlaneOrigin, const UE::Math::TVector<FReal>& PlaneNormal);
+	[[nodiscard]] static UE::Math::TVector<FReal> LinePlaneIntersection(const UE::Math::TVector<FReal>& Point1, const UE::Math::TVector<FReal>& Point2, const UE::Math::TVector<FReal>& PlaneOrigin, const UE::Math::TVector<FReal>& PlaneNormal);
 
 	/**
 	 * Find the intersection of a line and a plane. Assumes that the line and
@@ -1873,12 +1873,12 @@ public:
 	 * @return The point of intersection between the line and the plane.
 	 */
 	template<typename FReal>
-	UE_NODISCARD static UE::Math::TVector<FReal> LinePlaneIntersection(const UE::Math::TVector<FReal>& Point1, const UE::Math::TVector<FReal>& Point2, const UE::Math::TPlane<FReal>& Plane);
+	[[nodiscard]] static UE::Math::TVector<FReal> LinePlaneIntersection(const UE::Math::TVector<FReal>& Point1, const UE::Math::TVector<FReal>& Point2, const UE::Math::TPlane<FReal>& Plane);
 
 
 	// @parma InOutScissorRect should be set to View.ViewRect before the call
 	// @return 0: light is not visible, 1:use scissor rect, 2: no scissor rect needed
-	UE_NODISCARD static CORE_API uint32 ComputeProjectedSphereScissorRect(FIntRect& InOutScissorRect, FVector SphereOrigin, float Radius, FVector ViewOrigin, const FMatrix& ViewMatrix, const FMatrix& ProjMatrix);
+	[[nodiscard]] static CORE_API uint32 ComputeProjectedSphereScissorRect(FIntRect& InOutScissorRect, FVector SphereOrigin, float Radius, FVector ViewOrigin, const FMatrix& ViewMatrix, const FMatrix& ProjMatrix);
 
 	// @param ConeOrigin Cone origin
 	// @param ConeDirection Cone direction
@@ -1887,7 +1887,7 @@ public:
 	// @param SinConeAngle Sin of the cone angle
 	// @return Minimal bounding sphere encompassing given cone
 	template<typename FReal>
-	UE_NODISCARD static UE::Math::TSphere<FReal> ComputeBoundingSphereForCone(UE::Math::TVector<FReal> const& ConeOrigin, UE::Math::TVector<FReal> const& ConeDirection, FReal ConeRadius, FReal CosConeAngle, FReal SinConeAngle);
+	[[nodiscard]] static UE::Math::TSphere<FReal> ComputeBoundingSphereForCone(UE::Math::TVector<FReal> const& ConeOrigin, UE::Math::TVector<FReal> const& ConeDirection, FReal ConeRadius, FReal CosConeAngle, FReal SinConeAngle);
 
 	/** 
 	 * Determine if a plane and an AABB intersect
@@ -1895,7 +1895,7 @@ public:
 	 * @param AABB - the axis aligned bounding box to test
 	 * @return if collision occurs
 	 */
-	UE_NODISCARD static CORE_API bool PlaneAABBIntersection(const FPlane& P, const FBox& AABB);
+	[[nodiscard]] static CORE_API bool PlaneAABBIntersection(const FPlane& P, const FBox& AABB);
 
 	/**
 	 * Determine the position of an AABB relative to a plane:
@@ -1904,7 +1904,7 @@ public:
 	 * @param AABB - the axis aligned bounding box to test
 	 * @return -1 if below, 1 if above, 0 if intersects
 	 */
-	UE_NODISCARD static CORE_API int32 PlaneAABBRelativePosition(const FPlane& P, const FBox& AABB);
+	[[nodiscard]] static CORE_API int32 PlaneAABBRelativePosition(const FPlane& P, const FBox& AABB);
 
 	/**
 	 * Performs a sphere vs box intersection test using Arvo's algorithm:
@@ -1920,55 +1920,55 @@ public:
 	 * @return Whether the sphere/box intersect or not.
 	 */
 	template<typename FReal>
-	UE_NODISCARD static bool SphereAABBIntersection(const UE::Math::TVector<FReal>& SphereCenter,const FReal RadiusSquared,const UE::Math::TBox<FReal>& AABB);
+	[[nodiscard]] static bool SphereAABBIntersection(const UE::Math::TVector<FReal>& SphereCenter,const FReal RadiusSquared,const UE::Math::TBox<FReal>& AABB);
 
 	/**
 	 * Converts a sphere into a point plus radius squared for the test above
 	 */
 	template<typename FReal>
-	UE_NODISCARD static bool SphereAABBIntersection(const UE::Math::TSphere<FReal>& Sphere, const UE::Math::TBox<FReal>& AABB);
+	[[nodiscard]] static bool SphereAABBIntersection(const UE::Math::TSphere<FReal>& Sphere, const UE::Math::TBox<FReal>& AABB);
 
 	/** Determines whether a point is inside a box. */
 	template<typename FReal>
-	UE_NODISCARD static bool PointBoxIntersection( const UE::Math::TVector<FReal>& Point, const UE::Math::TBox<FReal>& Box );
+	[[nodiscard]] static bool PointBoxIntersection( const UE::Math::TVector<FReal>& Point, const UE::Math::TBox<FReal>& Box );
 
 	/** Determines whether a line intersects a box. */
 	template<typename FReal>
-	UE_NODISCARD static bool LineBoxIntersection( const UE::Math::TBox<FReal>& Box, const UE::Math::TVector<FReal>& Start, const UE::Math::TVector<FReal>& End, const UE::Math::TVector<FReal>& Direction );
+	[[nodiscard]] static bool LineBoxIntersection( const UE::Math::TBox<FReal>& Box, const UE::Math::TVector<FReal>& Start, const UE::Math::TVector<FReal>& End, const UE::Math::TVector<FReal>& Direction );
 
 	/** Determines whether a line intersects a box. This overload avoids the need to do the reciprocal every time. */
 	template<typename FReal>
-	UE_NODISCARD static bool LineBoxIntersection( const UE::Math::TBox<FReal>& Box, const UE::Math::TVector<FReal>& Start, const UE::Math::TVector<FReal>& End, const UE::Math::TVector<FReal>& Direction, const UE::Math::TVector<FReal>& OneOverDirection );
+	[[nodiscard]] static bool LineBoxIntersection( const UE::Math::TBox<FReal>& Box, const UE::Math::TVector<FReal>& Start, const UE::Math::TVector<FReal>& End, const UE::Math::TVector<FReal>& Direction, const UE::Math::TVector<FReal>& OneOverDirection );
 
 	/* Swept-Box vs Box test */
 	static CORE_API bool LineExtentBoxIntersection(const FBox& inBox, const FVector& Start, const FVector& End, const FVector& Extent, FVector& HitLocation, FVector& HitNormal, float& HitTime);
 
 	/** Determines whether a line intersects a sphere. */
 	template<typename FReal>
-	UE_NODISCARD static bool LineSphereIntersection(const UE::Math::TVector<FReal>& Start,const UE::Math::TVector<FReal>& Dir, FReal Length,const UE::Math::TVector<FReal>& Origin, FReal Radius);
+	[[nodiscard]] static bool LineSphereIntersection(const UE::Math::TVector<FReal>& Start,const UE::Math::TVector<FReal>& Dir, FReal Length,const UE::Math::TVector<FReal>& Origin, FReal Radius);
 
 	/**
 	 * Assumes the cone tip is at 0,0,0 (means the SphereCenter is relative to the cone tip)
 	 * @return true: cone and sphere do intersect, false otherwise
 	 */
-	UE_NODISCARD static CORE_API bool SphereConeIntersection(const FVector& SphereCenter, float SphereRadius, const FVector& ConeAxis, float ConeAngleSin, float ConeAngleCos);
+	[[nodiscard]] static CORE_API bool SphereConeIntersection(const FVector& SphereCenter, float SphereRadius, const FVector& ConeAxis, float ConeAngleSin, float ConeAngleCos);
 
 	/** Find the point on the line segment from LineStart to LineEnd which is closest to Point */
-	UE_NODISCARD static CORE_API FVector ClosestPointOnLine(const FVector& LineStart, const FVector& LineEnd, const FVector& Point);
+	[[nodiscard]] static CORE_API FVector ClosestPointOnLine(const FVector& LineStart, const FVector& LineEnd, const FVector& Point);
 
 	/** Find the point on the infinite line between two points (LineStart, LineEnd) which is closest to Point */
-	UE_NODISCARD static CORE_API FVector ClosestPointOnInfiniteLine(const FVector& LineStart, const FVector& LineEnd, const FVector& Point);
+	[[nodiscard]] static CORE_API FVector ClosestPointOnInfiniteLine(const FVector& LineStart, const FVector& LineEnd, const FVector& Point);
 
 	/** Compute intersection point of three planes. Return 1 if valid, 0 if infinite. */
 	template<typename FReal>
-	UE_NODISCARD static bool IntersectPlanes3(UE::Math::TVector<FReal>& I, const UE::Math::TPlane<FReal>& P1, const UE::Math::TPlane<FReal>& P2, const UE::Math::TPlane<FReal>& P3 );
+	[[nodiscard]] static bool IntersectPlanes3(UE::Math::TVector<FReal>& I, const UE::Math::TPlane<FReal>& P1, const UE::Math::TPlane<FReal>& P2, const UE::Math::TPlane<FReal>& P3 );
 
 	/**
 	 * Compute intersection point and direction of line joining two planes.
 	 * Return 1 if valid, 0 if infinite.
 	 */
 	template<typename FReal>
-	UE_NODISCARD static bool IntersectPlanes2(UE::Math::TVector<FReal>& I, UE::Math::TVector<FReal>& D, const UE::Math::TPlane<FReal>& P1, const UE::Math::TPlane<FReal>& P2);
+	[[nodiscard]] static bool IntersectPlanes2(UE::Math::TVector<FReal>& I, UE::Math::TVector<FReal>& D, const UE::Math::TPlane<FReal>& P1, const UE::Math::TPlane<FReal>& P2);
 
 	/**
 	 * Calculates the distance of a given Point in world space to a given line,
@@ -1982,7 +1982,7 @@ public:
 	 * @return	distance of Point from line defined by (Origin, Direction)
 	 */
 	static CORE_API float PointDistToLine(const FVector &Point, const FVector &Direction, const FVector &Origin, FVector &OutClosestPoint);
-	UE_NODISCARD static CORE_API float PointDistToLine(const FVector &Point, const FVector &Direction, const FVector &Origin);
+	[[nodiscard]] static CORE_API float PointDistToLine(const FVector &Point, const FVector &Direction, const FVector &Origin);
 
 	/**
 	 * Returns closest point on a segment to a given point.
@@ -1995,7 +1995,7 @@ public:
 	 *
 	 * @return	point on the segment defined by (StartPoint, EndPoint) that is closest to Point.
 	 */
-	UE_NODISCARD static CORE_API FVector ClosestPointOnSegment(const FVector &Point, const FVector &StartPoint, const FVector &EndPoint);
+	[[nodiscard]] static CORE_API FVector ClosestPointOnSegment(const FVector &Point, const FVector &StartPoint, const FVector &EndPoint);
 
 	/**
 	* FVector2D version of ClosestPointOnSegment.
@@ -2009,7 +2009,7 @@ public:
 	*
 	* @return	point on the segment defined by (StartPoint, EndPoint) that is closest to Point.
 	*/
-	UE_NODISCARD static CORE_API FVector2D ClosestPointOnSegment2D(const FVector2D &Point, const FVector2D &StartPoint, const FVector2D &EndPoint);
+	[[nodiscard]] static CORE_API FVector2D ClosestPointOnSegment2D(const FVector2D &Point, const FVector2D &StartPoint, const FVector2D &EndPoint);
 
 	/**
 	 * Returns distance from a point to the closest point on a segment.
@@ -2020,7 +2020,7 @@ public:
 	 *
 	 * @return	closest distance from Point to segment defined by (StartPoint, EndPoint).
 	 */
-	UE_NODISCARD static CORE_API float PointDistToSegment(const FVector &Point, const FVector &StartPoint, const FVector &EndPoint);
+	[[nodiscard]] static CORE_API float PointDistToSegment(const FVector &Point, const FVector &StartPoint, const FVector &EndPoint);
 
 	/**
 	 * Returns square of the distance from a point to the closest point on a segment.
@@ -2031,7 +2031,7 @@ public:
 	 *
 	 * @return	square of the closest distance from Point to segment defined by (StartPoint, EndPoint).
 	 */
-	UE_NODISCARD static CORE_API float PointDistToSegmentSquared(const FVector &Point, const FVector &StartPoint, const FVector &EndPoint);
+	[[nodiscard]] static CORE_API float PointDistToSegmentSquared(const FVector &Point, const FVector &StartPoint, const FVector &EndPoint);
 
 	/** 
 	 * Find closest points between 2 segments.
@@ -2065,7 +2065,7 @@ public:
 	 * @param Plane		- plane to intersect with
 	 * @return time(T) of intersection
 	 */
-	UE_NODISCARD static CORE_API float GetTForSegmentPlaneIntersect(const FVector& StartPoint, const FVector& EndPoint, const FPlane& Plane);
+	[[nodiscard]] static CORE_API float GetTForSegmentPlaneIntersect(const FVector& StartPoint, const FVector& EndPoint, const FPlane& Plane);
 
 	/**
 	 * Returns true if there is an intersection between the segment specified by StartPoint and Endpoint, and
@@ -2114,7 +2114,7 @@ public:
 	 *
 	 * @return	Point on triangle ABC closest to given point
 	 */
-	UE_NODISCARD static CORE_API FVector ClosestPointOnTriangleToPoint(const FVector& Point, const FVector& A, const FVector& B, const FVector& C);
+	[[nodiscard]] static CORE_API FVector ClosestPointOnTriangleToPoint(const FVector& Point, const FVector& A, const FVector& B, const FVector& C);
 
 	/**
 	 * Returns closest point on a tetrahedron to a point.
@@ -2126,7 +2126,7 @@ public:
 	 *
 	 * @return	Point on tetrahedron ABCD closest to given point
 	 */
-	UE_NODISCARD static CORE_API FVector ClosestPointOnTetrahedronToPoint(const FVector& Point, const FVector& A, const FVector& B, const FVector& C, const FVector& D);
+	[[nodiscard]] static CORE_API FVector ClosestPointOnTetrahedronToPoint(const FVector& Point, const FVector& A, const FVector& B, const FVector& C, const FVector& D);
 
 	/** 
 	 * Find closest point on a Sphere to a Line.
@@ -2151,7 +2151,7 @@ public:
 	 *
 	 * @return true if the point is within the cone, false otherwise.
 	 */
-	UE_NODISCARD static CORE_API bool GetDistanceWithinConeSegment(FVector Point, FVector ConeStartPoint, FVector ConeLine, float RadiusAtStart, float RadiusAtEnd, float &PercentageOut);
+	[[nodiscard]] static CORE_API bool GetDistanceWithinConeSegment(FVector Point, FVector ConeStartPoint, FVector ConeLine, float RadiusAtStart, float RadiusAtEnd, float &PercentageOut);
 
 	/**
 	 * Determines whether a given set of points are coplanar, with a tolerance. Any three points or less are always coplanar.
@@ -2161,7 +2161,7 @@ public:
 	 *
 	 * @return Whether the points are relatively coplanar, based on the tolerance
 	 */
-	UE_NODISCARD static CORE_API bool PointsAreCoplanar(const TArray<FVector>& Points, const float Tolerance = 0.1f);
+	[[nodiscard]] static CORE_API bool PointsAreCoplanar(const TArray<FVector>& Points, const float Tolerance = 0.1f);
 
 	/**
 	 * Truncates a floating point number to half if closer than the given tolerance.
@@ -2169,44 +2169,44 @@ public:
 	 * @param Tolerance			Maximum allowed difference to 0.5 in order to truncate
 	 * @return					The truncated value
 	 */
-	UE_NODISCARD static CORE_API float TruncateToHalfIfClose(float F, float Tolerance = UE_SMALL_NUMBER);
-	UE_NODISCARD static CORE_API double TruncateToHalfIfClose(double F, double Tolerance = UE_SMALL_NUMBER);
+	[[nodiscard]] static CORE_API float TruncateToHalfIfClose(float F, float Tolerance = UE_SMALL_NUMBER);
+	[[nodiscard]] static CORE_API double TruncateToHalfIfClose(double F, double Tolerance = UE_SMALL_NUMBER);
 
 	/**
 	* Converts a floating point number to the nearest integer, equidistant ties go to the value which is closest to an even value: 1.5 becomes 2, 0.5 becomes 0
 	* @param F		Floating point value to convert
 	* @return		The rounded integer
 	*/
-	UE_NODISCARD static CORE_API float RoundHalfToEven(float F);
-	UE_NODISCARD static CORE_API double RoundHalfToEven(double F);
+	[[nodiscard]] static CORE_API float RoundHalfToEven(float F);
+	[[nodiscard]] static CORE_API double RoundHalfToEven(double F);
 
 	/**
 	* Converts a floating point number to the nearest integer, equidistant ties go to the value which is further from zero: -0.5 becomes -1.0, 0.5 becomes 1.0
 	* @param F		Floating point value to convert
 	* @return		The rounded integer
 	*/
-	UE_NODISCARD static CORE_API float RoundHalfFromZero(float F);
-	UE_NODISCARD static CORE_API double RoundHalfFromZero(double F);
+	[[nodiscard]] static CORE_API float RoundHalfFromZero(float F);
+	[[nodiscard]] static CORE_API double RoundHalfFromZero(double F);
 
 	/**
 	* Converts a floating point number to the nearest integer, equidistant ties go to the value which is closer to zero: -0.5 becomes 0, 0.5 becomes 0
 	* @param F		Floating point value to convert
 	* @return		The rounded integer
 	*/
-	UE_NODISCARD static CORE_API float RoundHalfToZero(float F);
-	UE_NODISCARD static CORE_API double RoundHalfToZero(double F);
+	[[nodiscard]] static CORE_API float RoundHalfToZero(float F);
+	[[nodiscard]] static CORE_API double RoundHalfToZero(double F);
 
 	/**
 	* Converts a floating point number to an integer which is further from zero, "larger" in absolute value: 0.1 becomes 1, -0.1 becomes -1
 	* @param F		Floating point value to convert
 	* @return		The rounded integer
 	*/
-	UE_NODISCARD static FORCEINLINE float RoundFromZero(float F)
+	[[nodiscard]] static FORCEINLINE float RoundFromZero(float F)
 	{
 		return (F < 0.0f) ? FloorToFloat(F) : CeilToFloat(F);
 	}
 
-	UE_NODISCARD static FORCEINLINE double RoundFromZero(double F)
+	[[nodiscard]] static FORCEINLINE double RoundFromZero(double F)
 	{
 		return (F < 0.0) ? FloorToDouble(F) : CeilToDouble(F);
 	}
@@ -2216,12 +2216,12 @@ public:
 	* @param F		Floating point value to convert
 	* @return		The rounded integer
 	*/
-	UE_NODISCARD static FORCEINLINE float RoundToZero(float F)
+	[[nodiscard]] static FORCEINLINE float RoundToZero(float F)
 	{
 		return (F < 0.0f) ? CeilToFloat(F) : FloorToFloat(F);
 	}
 
-	UE_NODISCARD static FORCEINLINE double RoundToZero(double F)
+	[[nodiscard]] static FORCEINLINE double RoundToZero(double F)
 	{
 		return (F < 0.0) ? CeilToDouble(F) : FloorToDouble(F);
 	}
@@ -2231,12 +2231,12 @@ public:
 	* @param F		Floating point value to convert
 	* @return		The rounded integer
 	*/
-	UE_NODISCARD static FORCEINLINE float RoundToNegativeInfinity(float F)
+	[[nodiscard]] static FORCEINLINE float RoundToNegativeInfinity(float F)
 	{
 		return FloorToFloat(F);
 	}
 
-	UE_NODISCARD static FORCEINLINE double RoundToNegativeInfinity(double F)
+	[[nodiscard]] static FORCEINLINE double RoundToNegativeInfinity(double F)
 	{
 		return FloorToDouble(F);
 	}
@@ -2246,12 +2246,12 @@ public:
 	* @param F		Floating point value to convert
 	* @return		The rounded integer
 	*/
-	UE_NODISCARD static FORCEINLINE float RoundToPositiveInfinity(float F)
+	[[nodiscard]] static FORCEINLINE float RoundToPositiveInfinity(float F)
 	{
 		return CeilToFloat(F);
 	}
 
-	UE_NODISCARD static FORCEINLINE double RoundToPositiveInfinity(double F)
+	[[nodiscard]] static FORCEINLINE double RoundToPositiveInfinity(double F)
 	{
 		return CeilToDouble(F);
 	}
@@ -2265,7 +2265,7 @@ public:
 	 * @param	Val		The value to use
 	 * @return	FString	The human readable string
 	 */
-	UE_NODISCARD static CORE_API FString FormatIntToHumanReadable(int32 Val);
+	[[nodiscard]] static CORE_API FString FormatIntToHumanReadable(int32 Val);
 
 
 	// Utilities
@@ -2277,7 +2277,7 @@ public:
 	 * @param NumBytes		Number of bytes to test (will be rounded down to a multiple of 4)
 	 * @return				true if the memory region passed the test
 	 */
-	UE_NODISCARD static CORE_API bool MemoryTest( void* BaseAddress, uint32 NumBytes );
+	[[nodiscard]] static CORE_API bool MemoryTest( void* BaseAddress, uint32 NumBytes );
 
 	/**
 	 * Evaluates a numerical equation.
@@ -2302,7 +2302,7 @@ public:
 	 * @return Vector containing the three weights a,b,c such that Point = a*A + b*B + c*C
 	 *							                                or Point = A + b*(B-A) + c*(C-A) = (1-b-c)*A + b*B + c*C
 	 */
-	UE_NODISCARD static CORE_API FVector GetBaryCentric2D(const FVector& Point, const FVector& A, const FVector& B, const FVector& C);
+	[[nodiscard]] static CORE_API FVector GetBaryCentric2D(const FVector& Point, const FVector& A, const FVector& B, const FVector& C);
 
 	/**
 	 * Computes the barycentric coordinates for a given point in a triangle - simpler than the Compute versions
@@ -2313,7 +2313,7 @@ public:
 	 * @return Vector containing the three weights a,b,c such that Point = a*A + b*B + c*C
 	 *							                                or Point = A + b*(B-A) + c*(C-A) = (1-b-c)*A + b*B + c*C
 	 */
-	UE_NODISCARD static CORE_API FVector GetBaryCentric2D(const FVector2D& Point, const FVector2D& A, const FVector2D& B, const FVector2D& C);
+	[[nodiscard]] static CORE_API FVector GetBaryCentric2D(const FVector2D& Point, const FVector2D& A, const FVector2D& B, const FVector2D& C);
 
 	/**
 	 * Computes the barycentric coordinates for a given point in a 3D triangle.
@@ -2325,7 +2325,7 @@ public:
 	 * @return Vector containing the three weights a,b,c such that Point = a*A + b*B + c*C
 	 *							                               or Point = A + b*(B-A) + c*(C-A) = (1-b-c)*A + b*B + c*C
 	 */
-	UE_NODISCARD static CORE_API FVector ComputeBaryCentric2D(const FVector& Point, const FVector& A, const FVector& B, const FVector& C);
+	[[nodiscard]] static CORE_API FVector ComputeBaryCentric2D(const FVector& Point, const FVector& A, const FVector& B, const FVector& C);
 
 	/**
 	 * Computes the barycentric coordinates for a given point in a triangle
@@ -2338,7 +2338,7 @@ public:
 	 * 
 	 * @return false if the result could not be computed (occurs when the triangle area is too small)
 	 */
-	UE_NODISCARD static CORE_API bool ComputeBarycentricTri(const FVector& Point, const FVector& A, const FVector& B, const FVector& C, FVector& OutBarycentric, double Tolerance = UE_DOUBLE_SMALL_NUMBER);
+	[[nodiscard]] static CORE_API bool ComputeBarycentricTri(const FVector& Point, const FVector& A, const FVector& B, const FVector& C, FVector& OutBarycentric, double Tolerance = UE_DOUBLE_SMALL_NUMBER);
 
 	/**
 	 * Computes the barycentric coordinates for a given point on a tetrahedron (3D)
@@ -2348,7 +2348,7 @@ public:
 	 *
 	 * @return Vector containing the four weights a,b,c,d such that Point = a*A + b*B + c*C + d*D
 	 */
-	UE_NODISCARD static CORE_API FVector4 ComputeBaryCentric3D(const FVector& Point, const FVector& A, const FVector& B, const FVector& C, const FVector& D);
+	[[nodiscard]] static CORE_API FVector4 ComputeBaryCentric3D(const FVector& Point, const FVector& A, const FVector& B, const FVector& C, const FVector& D);
 
 	/** 32 bit values where BitFlag[x] == (1<<x) */
 	static CORE_API const uint32 BitFlag[32];
@@ -2364,7 +2364,7 @@ public:
 	 * @return Smoothed value between 0 and 1
 	 */
 	template<typename T>
-	UE_NODISCARD static constexpr T SmoothStep(T A, T B, T X)
+	[[nodiscard]] static constexpr T SmoothStep(T A, T B, T X)
 	{
 		if (X < A)
 		{
@@ -2382,7 +2382,7 @@ public:
 	 * Get a bit in memory created from bitflags (uint32 Value:1), used for EngineShowFlags,
 	 * TestBitFieldFunctions() tests the implementation
 	 */
-	UE_NODISCARD static constexpr bool ExtractBoolFromBitfield(const uint8* Ptr, uint32 Index)
+	[[nodiscard]] static constexpr bool ExtractBoolFromBitfield(const uint8* Ptr, uint32 Index)
 	{
 		const uint8* BytePtr = Ptr + Index / 8;
 		uint8 Mask = (uint8)(1 << (Index & 0x7));
@@ -2417,7 +2417,7 @@ public:
 
 	// @param x assumed to be in this range: 0..1
 	// @return 0..255
-	UE_NODISCARD static uint8 Quantize8UnsignedByte(float x)
+	[[nodiscard]] static uint8 Quantize8UnsignedByte(float x)
 	{
 		// 0..1 -> 0..255
 		int32 Ret = (int32)(x * 255.f + 0.5f);
@@ -2430,7 +2430,7 @@ public:
 	
 	// @param x assumed to be in this range: -1..1
 	// @return 0..255
-	UE_NODISCARD static uint8 Quantize8SignedByte(float x)
+	[[nodiscard]] static uint8 Quantize8SignedByte(float x)
 	{
 		// -1..1 -> 0..1
 		float y = x * 0.5f + 0.5f;
@@ -2439,7 +2439,7 @@ public:
 	}
 
 	// Use the Euclidean method to find the GCD
-	UE_NODISCARD static constexpr int32 GreatestCommonDivisor(int32 a, int32 b)
+	[[nodiscard]] static constexpr int32 GreatestCommonDivisor(int32 a, int32 b)
 	{
 		while (b != 0)
 		{
@@ -2452,7 +2452,7 @@ public:
 
 	// LCM = a/gcd * b
 	// a and b are the number we want to find the lcm
-	UE_NODISCARD static constexpr int32 LeastCommonMultiplier(int32 a, int32 b)
+	[[nodiscard]] static constexpr int32 LeastCommonMultiplier(int32 a, int32 b)
 	{
 		int32 CurrentGcd = GreatestCommonDivisor(a, b);
 		return CurrentGcd == 0 ? 0 : (a / CurrentGcd) * b;
@@ -2465,7 +2465,7 @@ public:
 	 *
 	 * @return	Perlin noise in the range of -1.0 to 1.0
 	 */
-	UE_NODISCARD static CORE_API float PerlinNoise1D(float Value);
+	[[nodiscard]] static CORE_API float PerlinNoise1D(float Value);
 
 	/**
 	* Generates a 2D Perlin noise sample at the given location.  Returns a continuous random value between -1.0 and 1.0.
@@ -2474,7 +2474,7 @@ public:
 	*
 	* @return	Perlin noise in the range of -1.0 to 1.0
 	*/
-	UE_NODISCARD static CORE_API float PerlinNoise2D(const FVector2D& Location);
+	[[nodiscard]] static CORE_API float PerlinNoise2D(const FVector2D& Location);
 	 
 
 	/**
@@ -2484,7 +2484,7 @@ public:
 	*
 	* @return	Perlin noise in the range of -1.0 to 1.0
 	*/
-	UE_NODISCARD static CORE_API float PerlinNoise3D(const FVector& Location);
+	[[nodiscard]] static CORE_API float PerlinNoise3D(const FVector& Location);
 
 	/**
 	 * Calculates the new value in a weighted moving average series using the previous value and the weight
@@ -2496,7 +2496,7 @@ public:
 	 * @return the next value in the series
 	 */
 	template<typename T>
-	UE_NODISCARD static inline T WeightedMovingAverage(T CurrentSample, T PreviousSample, T Weight)
+	[[nodiscard]] static inline T WeightedMovingAverage(T CurrentSample, T PreviousSample, T Weight)
 	{
 		Weight = Clamp<T>(Weight, 0.f, 1.f);
 		T WAvg = (CurrentSample * Weight) + (PreviousSample * (1.f - Weight));
@@ -2517,7 +2517,7 @@ public:
 	 * @return the next value in the series
 	 */
 	template<typename T>
-	UE_NODISCARD static inline T DynamicWeightedMovingAverage(T CurrentSample, T PreviousSample, T MaxDistance, T MinWeight, T MaxWeight)
+	[[nodiscard]] static inline T DynamicWeightedMovingAverage(T CurrentSample, T PreviousSample, T MaxDistance, T MinWeight, T MaxWeight)
 	{
 		// We need the distance between samples to determine how much of each weight to use
 		const T Distance = Abs<T>(CurrentSample - PreviousSample);
@@ -2541,8 +2541,8 @@ public:
 	 * Do not call this directly, use VectorMatrixInverse or Matrix::Inverse
 	 * this is the fallback scalar implementation used by VectorMatrixInverse
 	 */
-	UE_NODISCARD static CORE_API bool MatrixInverse(FMatrix44f* DstMatrix, const FMatrix44f* SrcMatrix);
-	UE_NODISCARD static CORE_API bool MatrixInverse(FMatrix44d* DstMatrix, const FMatrix44d* SrcMatrix);
+	[[nodiscard]] static CORE_API bool MatrixInverse(FMatrix44f* DstMatrix, const FMatrix44f* SrcMatrix);
+	[[nodiscard]] static CORE_API bool MatrixInverse(FMatrix44d* DstMatrix, const FMatrix44d* SrcMatrix);
 
 };
 
diff --git a/Engine/Source/Runtime/Core/Public/Memory/MemoryView.h b/Engine/Source/Runtime/Core/Public/Memory/MemoryView.h
--- a/Engine/Source/Runtime/Core/Public/Memory/MemoryView.h
+++ b/Engine/Source/Runtime/Core/Public/Memory/MemoryView.h
@@ -61,22 +61,22 @@ public:
 	}
 
 	/** Returns a pointer to the start of the view. */
-	UE_NODISCARD constexpr inline DataType* GetData() const { return Data; }
+	[[nodiscard]] constexpr inline DataType* GetData() const { return Data; }
 
 	/** Returns a pointer to the end of the view. */
-	UE_NODISCARD inline DataType* GetDataEnd() const { return GetDataAtOffsetNoCheck(Size); }
+	[[nodiscard]] inline DataType* GetDataEnd() const { return GetDataAtOffsetNoCheck(Size); }
 
 	/** Returns the number of bytes in the view. */
-	UE_NODISCARD constexpr inline uint64 GetSize() const { return Size; }
+	[[nodiscard]] constexpr inline uint64 GetSize() const { return Size; }
 
 	/** Returns whether the view has a size of 0 regardless of its data pointer. */
-	UE_NODISCARD constexpr inline bool IsEmpty() const { return Size == 0; }
+	[[nodiscard]] constexpr inline bool IsEmpty() const { return Size == 0; }
 
 	/** Resets to an empty view. */
 	constexpr inline void Reset() { *this = TMemoryView(); }
 
 	/** Returns the left-most part of the view by taking the given number of bytes from the left. */
-	UE_NODISCARD constexpr inline TMemoryView Left(uint64 InSize) const
+	[[nodiscard]] constexpr inline TMemoryView Left(uint64 InSize) const
 	{
 		TMemoryView View(*this);
 		View.LeftInline(InSize);
@@ -84,7 +84,7 @@ public:
 	}
 
 	/** Returns the left-most part of the view by chopping the given number of bytes from the right. */
-	UE_NODISCARD constexpr inline TMemoryView LeftChop(uint64 InSize) const
+	[[nodiscard]] constexpr inline TMemoryView LeftChop(uint64 InSize) const
 	{
 		TMemoryView View(*this);
 		View.LeftChopInline(InSize);
@@ -92,7 +92,7 @@ public:
 	}
 
 	/** Returns the right-most part of the view by taking the given number of bytes from the right. */
-	UE_NODISCARD inline TMemoryView Right(uint64 InSize) const
+	[[nodiscard]] inline TMemoryView Right(uint64 InSize) const
 	{
 		TMemoryView View(*this);
 		View.RightInline(InSize);
@@ -100,7 +100,7 @@ public:
 	}
 
 	/** Returns the right-most part of the view by chopping the given number of bytes from the left. */
-	UE_NODISCARD inline TMemoryView RightChop(uint64 InSize) const
+	[[nodiscard]] inline TMemoryView RightChop(uint64 InSize) const
 	{
 		TMemoryView View(*this);
 		View.RightChopInline(InSize);
@@ -108,7 +108,7 @@ public:
 	}
 
 	/** Returns the middle part of the view by taking up to the given number of bytes from the given position. */
-	UE_NODISCARD inline TMemoryView Mid(uint64 InOffset, uint64 InSize = TNumericLimits<uint64>::Max()) const
+	[[nodiscard]] inline TMemoryView Mid(uint64 InOffset, uint64 InSize = TNumericLimits<uint64>::Max()) const
 	{
 		TMemoryView View(*this);
 		View.MidInline(InOffset, InSize);
@@ -153,21 +153,21 @@ public:
 
 	/** Returns whether this view fully contains the other view. */
 	template <typename OtherDataType>
-	UE_NODISCARD inline bool Contains(const TMemoryView<OtherDataType>& InView) const
+	[[nodiscard]] inline bool Contains(const TMemoryView<OtherDataType>& InView) const
 	{
 		return Data <= InView.Data && GetDataAtOffsetNoCheck(Size) >= InView.GetDataAtOffsetNoCheck(InView.Size);
 	}
 
 	/** Returns whether this view intersects the other view. */
 	template <typename OtherDataType>
-	UE_NODISCARD inline bool Intersects(const TMemoryView<OtherDataType>& InView) const
+	[[nodiscard]] inline bool Intersects(const TMemoryView<OtherDataType>& InView) const
 	{
 		return Data < InView.GetDataAtOffsetNoCheck(InView.Size) && InView.Data < GetDataAtOffsetNoCheck(Size);
 	}
 
 	/** Returns whether the bytes of this view are equal or less/greater than the bytes of the other view. */
 	template <typename OtherDataType>
-	UE_NODISCARD inline int32 CompareBytes(const TMemoryView<OtherDataType>& InView) const
+	[[nodiscard]] inline int32 CompareBytes(const TMemoryView<OtherDataType>& InView) const
 	{
 		const int32 Compare = Data == InView.Data ? 0 : FMemory::Memcmp(Data, InView.Data, FMath::Min(Size, InView.Size));
 		return Compare || Size == InView.Size ? Compare : Size < InView.Size ? -1 : 1;
@@ -175,28 +175,28 @@ public:
 
 	/** Returns whether the bytes of this views are equal to the bytes of the other view. */
 	template <typename OtherDataType>
-	UE_NODISCARD inline bool EqualBytes(const TMemoryView<OtherDataType>& InView) const
+	[[nodiscard]] inline bool EqualBytes(const TMemoryView<OtherDataType>& InView) const
 	{
 		return Size == InView.Size && (Data == InView.Data || FMemory::Memcmp(Data, InView.Data, Size) == 0);
 	}
 
 	/** Returns whether the data pointers and sizes of this view and the other view are equal. */
 	template <typename OtherDataType>
-	UE_NODISCARD constexpr inline bool Equals(const TMemoryView<OtherDataType>& InView) const
+	[[nodiscard]] constexpr inline bool Equals(const TMemoryView<OtherDataType>& InView) const
 	{
 		return Size == InView.Size && (Size == 0 || Data == InView.Data);
 	}
 
 	/** Returns whether the data pointers and sizes of this view and the other view are equal. */
 	template <typename OtherDataType>
-	UE_NODISCARD constexpr inline bool operator==(const TMemoryView<OtherDataType>& InView) const
+	[[nodiscard]] constexpr inline bool operator==(const TMemoryView<OtherDataType>& InView) const
 	{
 		return Equals(InView);
 	}
 
 	/** Returns whether the data pointers and sizes of this view and the other view are not equal. */
 	template <typename OtherDataType>
-	UE_NODISCARD constexpr inline bool operator!=(const TMemoryView<OtherDataType>& InView) const
+	[[nodiscard]] constexpr inline bool operator!=(const TMemoryView<OtherDataType>& InView) const
 	{
 		return !Equals(InView);
 	}
@@ -222,7 +222,7 @@ public:
 
 private:
 	/** Returns the data pointer advanced by an offset in bytes. */
-	UE_NODISCARD inline DataType* GetDataAtOffsetNoCheck(uint64 InOffset) const
+	[[nodiscard]] inline DataType* GetDataAtOffsetNoCheck(uint64 InOffset) const
 	{
 		return reinterpret_cast<ByteType*>(Data) + InOffset;
 	}
@@ -237,26 +237,26 @@ private:
 
 /** Advances the start of the view by an offset, which is clamped to stay within the view. */
 template <typename DataType>
-UE_NODISCARD constexpr inline TMemoryView<DataType> operator+(const TMemoryView<DataType>& View, uint64 Offset)
+[[nodiscard]] constexpr inline TMemoryView<DataType> operator+(const TMemoryView<DataType>& View, uint64 Offset)
 {
 	return TMemoryView<DataType>(View) += Offset;
 }
 
 /** Advances the start of the view by an offset, which is clamped to stay within the view. */
 template <typename DataType>
-UE_NODISCARD constexpr inline TMemoryView<DataType> operator+(uint64 Offset, const TMemoryView<DataType>& View)
+[[nodiscard]] constexpr inline TMemoryView<DataType> operator+(uint64 Offset, const TMemoryView<DataType>& View)
 {
 	return TMemoryView<DataType>(View) += Offset;
 }
 
 /** Make a non-owning mutable view of Size bytes starting at Data. */
-UE_NODISCARD constexpr inline TMemoryView<void> MakeMemoryView(void* Data, uint64 Size)
+[[nodiscard]] constexpr inline TMemoryView<void> MakeMemoryView(void* Data, uint64 Size)
 {
 	return TMemoryView<void>(Data, Size);
 }
 
 /** Make a non-owning const view of Size bytes starting at Data. */
-UE_NODISCARD constexpr inline TMemoryView<const void> MakeMemoryView(const void* Data, uint64 Size)
+[[nodiscard]] constexpr inline TMemoryView<const void> MakeMemoryView(const void* Data, uint64 Size)
 {
 	return TMemoryView<const void>(Data, Size);
 }
@@ -264,7 +264,7 @@ UE_NODISCARD constexpr inline TMemoryView<const void> MakeMemoryView(const void*
 /** Make a non-owning mutable view starting at Data and ending at DataEnd. */
 template <typename DataEndType,
 	decltype(ImplicitConv<void*>(DeclVal<DataEndType*>()))* = nullptr>
-UE_NODISCARD inline TMemoryView<void> MakeMemoryView(void* Data, DataEndType* DataEnd)
+[[nodiscard]] inline TMemoryView<void> MakeMemoryView(void* Data, DataEndType* DataEnd)
 {
 	return TMemoryView<void>(Data, DataEnd);
 }
@@ -272,7 +272,7 @@ UE_NODISCARD inline TMemoryView<void> MakeMemoryView(void* Data, DataEndType* Da
 /** Make a non-owning const view starting at Data and ending at DataEnd. */
 template <typename DataEndType,
 	decltype(ImplicitConv<const void*>(DeclVal<DataEndType*>()))* = nullptr>
-UE_NODISCARD inline TMemoryView<const void> MakeMemoryView(const void* Data, DataEndType* DataEnd)
+[[nodiscard]] inline TMemoryView<const void> MakeMemoryView(const void* Data, DataEndType* DataEnd)
 {
 	return TMemoryView<const void>(Data, DataEnd);
 }
@@ -283,7 +283,7 @@ UE_NODISCARD inline TMemoryView<const void> MakeMemoryView(const void* Data, Dat
  * This overload is only available when the element type does not need to be deduced.
  */
 template <typename T>
-UE_NODISCARD constexpr inline TMemoryView<const void> MakeMemoryView(std::initializer_list<typename TIdentity<T>::Type> List)
+[[nodiscard]] constexpr inline TMemoryView<const void> MakeMemoryView(std::initializer_list<typename TIdentity<T>::Type> List)
 {
 	return TMemoryView<const void>(GetData(List), GetNum(List) * sizeof(T));
 }
@@ -291,7 +291,7 @@ UE_NODISCARD constexpr inline TMemoryView<const void> MakeMemoryView(std::initia
 /** Make a non-owning view of the memory of the contiguous container. */
 template <typename ContainerType,
 	typename TEnableIf<TIsContiguousContainer<ContainerType>::Value>::Type* = nullptr>
-UE_NODISCARD constexpr inline auto MakeMemoryView(ContainerType&& Container)
+[[nodiscard]] constexpr inline auto MakeMemoryView(ContainerType&& Container)
 {
 	using ElementType = typename TRemovePointer<decltype(GetData(DeclVal<ContainerType>()))>::Type;
 	constexpr bool bIsConst = TIsConst<ElementType>::Value;
diff --git a/Engine/Source/Runtime/Core/Public/Misc/Attribute.h b/Engine/Source/Runtime/Core/Public/Misc/Attribute.h
--- a/Engine/Source/Runtime/Core/Public/Misc/Attribute.h
+++ b/Engine/Source/Runtime/Core/Public/Misc/Attribute.h
@@ -98,7 +98,7 @@ public:
 	 *
 	 * @param  InGetter		Delegate to bind
 	 */
-	UE_NODISCARD static TAttribute Create( const FGetter& InGetter )
+	[[nodiscard]] static TAttribute Create( const FGetter& InGetter )
 	{
 		const bool bExplicitConstructor = true;
 		return TAttribute( InGetter, bExplicitConstructor );
@@ -112,7 +112,7 @@ public:
 	 * @param  InFuncPtr Member function to bind.  The function's structure (return value, arguments, etc) must match IBoundAttributeDelegate's definition.
 	 */
 	template <typename... VarTypes >
-	UE_NODISCARD static TAttribute CreateStatic( TIdentity_T< typename FGetter::template TFuncPtr< VarTypes... > > InFuncPtr, VarTypes... Vars )
+	[[nodiscard]] static TAttribute CreateStatic( TIdentity_T< typename FGetter::template TFuncPtr< VarTypes... > > InFuncPtr, VarTypes... Vars )
 	{
 		const bool bExplicitConstructor = true;
 		return TAttribute( FGetter::CreateStatic( InFuncPtr, Vars... ), bExplicitConstructor );
@@ -126,7 +126,7 @@ public:
 	 * @param  InFuncPtr Member function to bind.  The function's structure (return value, arguments, etc) must match IBoundAttributeDelegate's definition.
 	 */
 	template <typename FuncPtrType, typename... VarTypes >
-	UE_NODISCARD static TAttribute CreateStatic(FuncPtrType&& InFuncPtr, VarTypes... Vars)
+	[[nodiscard]] static TAttribute CreateStatic(FuncPtrType&& InFuncPtr, VarTypes... Vars)
 	{
 		const bool bExplicitConstructor = true;
 		return TAttribute(FGetter::CreateStatic(InFuncPtr, MoveTemp(Vars)...), bExplicitConstructor);
@@ -137,7 +137,7 @@ public:
 	 */
 	template<typename SourceType, typename SourceTypeOrBase, typename... PayloadTypes>
 	UE_DEPRECATED(5.0, "Attribute's Getter should be const.")
-	UE_NODISCARD FORCEINLINE static TAttribute CreateRaw(SourceType* InObject, ObjectType (SourceTypeOrBase::*InMethod)(PayloadTypes...), typename TDecay<PayloadTypes>::Type... InputPayload)
+	[[nodiscard]] FORCEINLINE static TAttribute CreateRaw(SourceType* InObject, ObjectType (SourceTypeOrBase::*InMethod)(PayloadTypes...), typename TDecay<PayloadTypes>::Type... InputPayload)
 	{
 		return Create(FGetter::CreateRaw(InObject, InMethod, MoveTemp(InputPayload)...));
 	}
@@ -146,7 +146,7 @@ public:
 	 * Helper function for creating TAttributes from a const member function pointer, accessed through a raw pointer
 	 */
 	template<typename SourceType, typename SourceTypeOrBase, typename... PayloadTypes>
-	UE_NODISCARD FORCEINLINE static TAttribute CreateRaw(const SourceType* InObject, ObjectType (SourceTypeOrBase::*InMethod)(PayloadTypes...) const, typename TDecay<PayloadTypes>::Type... InputPayload)
+	[[nodiscard]] FORCEINLINE static TAttribute CreateRaw(const SourceType* InObject, ObjectType (SourceTypeOrBase::*InMethod)(PayloadTypes...) const, typename TDecay<PayloadTypes>::Type... InputPayload)
 	{
 		return Create(FGetter::CreateRaw(InObject, InMethod, MoveTemp(InputPayload)...));
 	}
@@ -156,7 +156,7 @@ public:
 	 */
 	template<typename SourceType, typename SourceTypeOrBase, typename... PayloadTypes>
 	UE_DEPRECATED(5.0, "Attribute's Getter should be const.")
-	UE_NODISCARD FORCEINLINE static TAttribute CreateSP(SourceType* InObject, ObjectType (SourceTypeOrBase::*InMethod)(PayloadTypes...), typename TDecay<PayloadTypes>::Type... InputPayload)
+	[[nodiscard]] FORCEINLINE static TAttribute CreateSP(SourceType* InObject, ObjectType (SourceTypeOrBase::*InMethod)(PayloadTypes...), typename TDecay<PayloadTypes>::Type... InputPayload)
 	{
 		return Create(FGetter::CreateSP(InObject, InMethod, MoveTemp(InputPayload)...));
 	}
@@ -165,7 +165,7 @@ public:
 	 * Helper function for creating TAttributes from a const member function pointer, accessed through a weak pointer to the shared object
 	 */
 	template<typename SourceType, typename SourceTypeOrBase, typename... PayloadTypes>
-	UE_NODISCARD FORCEINLINE static TAttribute CreateSP(const SourceType* InObject, ObjectType (SourceTypeOrBase::*InMethod)(PayloadTypes...) const, typename TDecay<PayloadTypes>::Type... InputPayload)
+	[[nodiscard]] FORCEINLINE static TAttribute CreateSP(const SourceType* InObject, ObjectType (SourceTypeOrBase::*InMethod)(PayloadTypes...) const, typename TDecay<PayloadTypes>::Type... InputPayload)
 	{
 		return Create(FGetter::CreateSP(InObject, InMethod, MoveTemp(InputPayload)...));
 	}
@@ -175,7 +175,7 @@ public:
 	 * TAttribute<float> FloatAttribute = TAttribute<float>::CreateLambda([]{ return 10.f; });
 	 */
 	template<typename LambdaType, typename... PayloadTypes>
-	UE_NODISCARD FORCEINLINE static TAttribute CreateLambda(LambdaType&& InCallable, PayloadTypes&&... InputPayload)
+	[[nodiscard]] FORCEINLINE static TAttribute CreateLambda(LambdaType&& InCallable, PayloadTypes&&... InputPayload)
 	{
 		return Create(FGetter::CreateLambda(InCallable, Forward<PayloadTypes>(InputPayload)...));
 	}
@@ -189,14 +189,14 @@ public:
 	 * @param  InFunctionName Member function name to bind.
 	 */
 	template< class SourceType >
-	UE_NODISCARD static TAttribute< ObjectType > Create(SourceType* InUserObject, const FName& InFunctionName)
+	[[nodiscard]] static TAttribute< ObjectType > Create(SourceType* InUserObject, const FName& InFunctionName)
 	{
 		TAttribute< ObjectType > Attrib;
 		Attrib.BindUFunction<SourceType>(InUserObject, InFunctionName);
 		return Attrib;
 	}
 
-	UE_NODISCARD FORCEINLINE static TAttribute< ObjectType > Create(TFunction<ObjectType(void)>&& InLambda)
+	[[nodiscard]] FORCEINLINE static TAttribute< ObjectType > Create(TFunction<ObjectType(void)>&& InLambda)
 	{
 		return Create(TAttribute< ObjectType >::FGetter::CreateLambda(MoveTemp(InLambda)));
 	}
@@ -395,7 +395,7 @@ public:
 	 *
 	 * @return  The attribute's FGetter.
 	 */
-	UE_NODISCARD const FGetter& GetBinding() const
+	[[nodiscard]] const FGetter& GetBinding() const
 	{
 		return Getter;
 	}
@@ -405,7 +405,7 @@ public:
 	 *
 	 * @return  The attribute's FGetter or `Value`
 	 */
-	UE_NODISCARD TVariant<ObjectType, FGetter> Steal()
+	[[nodiscard]] TVariant<ObjectType, FGetter> Steal()
 	{
 		checkf(IsSet(), TEXT("It is an error to call Steal() on an unset TAttribute. Check IsSet() before calling Steal()."));
 		bIsSet = false;
@@ -466,7 +466,7 @@ private:
  */
 template<typename T, typename SourceType, typename SourceTypeOrBase, typename... PayloadTypes>
 UE_DEPRECATED(5.0, "Attribute's Getter should be const.")
-UE_NODISCARD FORCEINLINE TAttribute<T> MakeAttributeRaw(SourceType* InObject, T (SourceTypeOrBase::*InMethod)(PayloadTypes...), typename TDecay<PayloadTypes>::Type... InputPayload)
+[[nodiscard]] FORCEINLINE TAttribute<T> MakeAttributeRaw(SourceType* InObject, T (SourceTypeOrBase::*InMethod)(PayloadTypes...), typename TDecay<PayloadTypes>::Type... InputPayload)
 {
 	return TAttribute<T>::Create(TAttribute<T>::FGetter::CreateRaw(InObject, InMethod, MoveTemp(InputPayload)...));
 }
@@ -475,7 +475,7 @@ UE_NODISCARD FORCEINLINE TAttribute<T> MakeAttributeRaw(SourceType* InObject, T
  * Helper function for creating TAttributes from a const member function pointer, accessed through a raw pointer
  */
 template<typename T, typename SourceType, typename SourceTypeOrBase, typename... PayloadTypes>
-UE_NODISCARD FORCEINLINE TAttribute<T> MakeAttributeRaw(const SourceType* InObject, T (SourceTypeOrBase::*InMethod)(PayloadTypes...) const, typename TDecay<PayloadTypes>::Type... InputPayload)
+[[nodiscard]] FORCEINLINE TAttribute<T> MakeAttributeRaw(const SourceType* InObject, T (SourceTypeOrBase::*InMethod)(PayloadTypes...) const, typename TDecay<PayloadTypes>::Type... InputPayload)
 {
 	return TAttribute<T>::Create(TAttribute<T>::FGetter::CreateRaw(InObject, InMethod, MoveTemp(InputPayload)...));
 }
@@ -485,7 +485,7 @@ UE_NODISCARD FORCEINLINE TAttribute<T> MakeAttributeRaw(const SourceType* InObje
  */
 template<typename T, typename SourceType, typename SourceTypeOrBase, typename... PayloadTypes>
 UE_DEPRECATED(5.0, "Attribute's Getter should be const.")
-UE_NODISCARD FORCEINLINE TAttribute<T> MakeAttributeSP(SourceType* InObject, T (SourceTypeOrBase::*InMethod)(PayloadTypes...), typename TDecay<PayloadTypes>::Type...  InputPayload)
+[[nodiscard]] FORCEINLINE TAttribute<T> MakeAttributeSP(SourceType* InObject, T (SourceTypeOrBase::*InMethod)(PayloadTypes...), typename TDecay<PayloadTypes>::Type...  InputPayload)
 {
 	return TAttribute<T>::Create(TAttribute<T>::FGetter::CreateSP(InObject, InMethod, MoveTemp(InputPayload)...));
 }
@@ -494,7 +494,7 @@ UE_NODISCARD FORCEINLINE TAttribute<T> MakeAttributeSP(SourceType* InObject, T (
  * Helper function for creating TAttributes from a const member function pointer, accessed through a weak pointer to the shared object
  */
 template<typename T, typename SourceType, typename SourceTypeOrBase, typename... PayloadTypes>
-UE_NODISCARD FORCEINLINE TAttribute<T> MakeAttributeSP(const SourceType* InObject, T (SourceTypeOrBase::*InMethod)(PayloadTypes...) const, typename TDecay<PayloadTypes>::Type...  InputPayload)
+[[nodiscard]] FORCEINLINE TAttribute<T> MakeAttributeSP(const SourceType* InObject, T (SourceTypeOrBase::*InMethod)(PayloadTypes...) const, typename TDecay<PayloadTypes>::Type...  InputPayload)
 {
 	return TAttribute<T>::Create(TAttribute<T>::FGetter::CreateSP(InObject, InMethod, MoveTemp(InputPayload)...));
 }
@@ -504,7 +504,7 @@ UE_NODISCARD FORCEINLINE TAttribute<T> MakeAttributeSP(const SourceType* InObjec
  * TAttribute<float> FloatAttribute = MakeAttributeLambda([]{ return 10.f; });
  */
 template<typename LambdaType, typename... PayloadTypes>
-UE_NODISCARD decltype(auto) MakeAttributeLambda(LambdaType&& InCallable, PayloadTypes&&... InputPayload)
+[[nodiscard]] decltype(auto) MakeAttributeLambda(LambdaType&& InCallable, PayloadTypes&&... InputPayload)
 {
 	typedef decltype(InCallable(DeclVal<PayloadTypes>()...)) T;
 
diff --git a/Engine/Source/Runtime/Core/Public/Misc/HierarchicalLogArchive.h b/Engine/Source/Runtime/Core/Public/Misc/HierarchicalLogArchive.h
--- a/Engine/Source/Runtime/Core/Public/Misc/HierarchicalLogArchive.h
+++ b/Engine/Source/Runtime/Core/Public/Misc/HierarchicalLogArchive.h
@@ -27,7 +27,7 @@ public:
 		WriteLine(InLine, false);
 	}
 
-	UE_NODISCARD FIndentScope PrintIndent(const TCHAR* InLine)
+	[[nodiscard]] FIndentScope PrintIndent(const TCHAR* InLine)
 	{
 		WriteLine(InLine, true);
 		return FIndentScope(this);
@@ -40,7 +40,7 @@ public:
 	}
 
 	template <typename FmtType, typename... Types>
-	UE_NODISCARD FIndentScope PrintfIndent(const FmtType& Fmt, Types... Args)
+	[[nodiscard]] FIndentScope PrintfIndent(const FmtType& Fmt, Types... Args)
 	{
 		WriteLine(FString::Printf(Fmt, Args...), true);
 		return FIndentScope(this);
diff --git a/Engine/Source/Runtime/Datasmith/DatasmithCore/Private/DatasmithSceneXmlReader.cpp b/Engine/Source/Runtime/Datasmith/DatasmithCore/Private/DatasmithSceneXmlReader.cpp
--- a/Engine/Source/Runtime/Datasmith/DatasmithCore/Private/DatasmithSceneXmlReader.cpp
+++ b/Engine/Source/Runtime/Datasmith/DatasmithCore/Private/DatasmithSceneXmlReader.cpp
@@ -117,7 +117,7 @@ void FDatasmithSceneXmlReader::PatchUpVersion(TSharedRef< IDatasmithScene >& Out
 	}
 }
 
-UE_NODISCARD FString FDatasmithSceneXmlReader::UnsanitizeXMLText(const FString& InString) const
+[[nodiscard]] FString FDatasmithSceneXmlReader::UnsanitizeXMLText(const FString& InString) const
 {
 	FString OutString = InString;
 	OutString.ReplaceInline( TEXT("&apos;"), TEXT("'")  );
diff --git a/Engine/Source/Runtime/Datasmith/DatasmithCore/Public/DatasmithSceneXmlReader.h b/Engine/Source/Runtime/Datasmith/DatasmithCore/Public/DatasmithSceneXmlReader.h
--- a/Engine/Source/Runtime/Datasmith/DatasmithCore/Public/DatasmithSceneXmlReader.h
+++ b/Engine/Source/Runtime/Datasmith/DatasmithCore/Public/DatasmithSceneXmlReader.h
@@ -33,7 +33,7 @@ private:
 	bool ParseXmlFile(TSharedRef< IDatasmithScene >& OutScene, bool bInAppend = false);
 	void PatchUpVersion(TSharedRef< IDatasmithScene >& OutScene) const;
 
-	UE_NODISCARD FString UnsanitizeXMLText(const FString& InString) const;
+	[[nodiscard]] FString UnsanitizeXMLText(const FString& InString) const;
 
 	template<typename T> T ValueFromString(const FString& InString) const = delete;
 	FVector VectorFromNode(FXmlNode* InNode, const TCHAR* XName, const TCHAR* YName, const TCHAR* ZName) const;
diff --git a/Engine/Source/Runtime/Engine/Public/Net/DataChannel.h b/Engine/Source/Runtime/Engine/Public/Net/DataChannel.h
--- a/Engine/Source/Runtime/Engine/Public/Net/DataChannel.h
+++ b/Engine/Source/Runtime/Engine/Public/Net/DataChannel.h
@@ -120,7 +120,7 @@ public: \
 	} \
 	/** receives a message of this type from the passed in bunch */ \
 	template<typename... ParamTypes> \
-	UE_NODISCARD static bool Receive(FInBunch& Bunch, ParamTypes&... Params) \
+	[[nodiscard]] static bool Receive(FInBunch& Bunch, ParamTypes&... Params) \
 	{ \
 		FNetControlMessageInfo::ReceiveParams(Bunch, Params...); \
 		return !Bunch.IsError(); \
diff --git a/Engine/Source/Runtime/FieldNotification/Public/FieldNotificationHelpers.h b/Engine/Source/Runtime/FieldNotification/Public/FieldNotificationHelpers.h
--- a/Engine/Source/Runtime/FieldNotification/Public/FieldNotificationHelpers.h
+++ b/Engine/Source/Runtime/FieldNotification/Public/FieldNotificationHelpers.h
@@ -7,7 +7,7 @@
 
 namespace UE::FieldNotification::Helpers
 {
-	UE_NODISCARD FORCEINLINE bool IsValidAsField(const UFunction* InFunction)
+	[[nodiscard]] FORCEINLINE bool IsValidAsField(const UFunction* InFunction)
 	{
 		return InFunction != nullptr
 			&& !InFunction->HasAnyFunctionFlags(FUNC_Net | FUNC_Event)
diff --git a/Engine/Source/Runtime/MovieScene/Public/EntitySystem/MovieSceneEntitySystemTypes.h b/Engine/Source/Runtime/MovieScene/Public/EntitySystem/MovieSceneEntitySystemTypes.h
--- a/Engine/Source/Runtime/MovieScene/Public/EntitySystem/MovieSceneEntitySystemTypes.h
+++ b/Engine/Source/Runtime/MovieScene/Public/EntitySystem/MovieSceneEntitySystemTypes.h
@@ -814,28 +814,28 @@ struct FEntityAllocation
 	/**
 	 * Read type-erased component data for the specified component type
 	 */
-	UE_NODISCARD MOVIESCENE_API TComponentLock<FReadErased> ReadComponentsErased(FComponentTypeID ComponentType) const;
+	[[nodiscard]] MOVIESCENE_API TComponentLock<FReadErased> ReadComponentsErased(FComponentTypeID ComponentType) const;
 
 	/**
 	 * Write type-erased component data for the specified component type
 	 */
-	UE_NODISCARD MOVIESCENE_API TComponentLock<FWriteErased> WriteComponentsErased(FComponentTypeID ComponentType, FEntityAllocationWriteContext InWriteContext) const;
+	[[nodiscard]] MOVIESCENE_API TComponentLock<FWriteErased> WriteComponentsErased(FComponentTypeID ComponentType, FEntityAllocationWriteContext InWriteContext) const;
 
 	/**
 	 * Attempt to read type-erased component data for the specified component type
 	 */
-	UE_NODISCARD MOVIESCENE_API TComponentLock<FReadErasedOptional> TryReadComponentsErased(FComponentTypeID ComponentType) const;
+	[[nodiscard]] MOVIESCENE_API TComponentLock<FReadErasedOptional> TryReadComponentsErased(FComponentTypeID ComponentType) const;
 
 	/**
 	 * Attempt to write type-erased component data for the specified component type
 	 */
-	UE_NODISCARD MOVIESCENE_API TComponentLock<FWriteErasedOptional> TryWriteComponentsErased(FComponentTypeID ComponentType, FEntityAllocationWriteContext InWriteContext) const;
+	[[nodiscard]] MOVIESCENE_API TComponentLock<FWriteErasedOptional> TryWriteComponentsErased(FComponentTypeID ComponentType, FEntityAllocationWriteContext InWriteContext) const;
 
 	/**
 	 * Read typed component data for the specified component type
 	 */
 	template<typename T>
-	UE_NODISCARD TComponentLock<TRead<T>> ReadComponents(TComponentTypeID<T> ComponentType) const
+	[[nodiscard]] TComponentLock<TRead<T>> ReadComponents(TComponentTypeID<T> ComponentType) const
 	{
 		const FComponentHeader& Header = GetComponentHeaderChecked(ComponentType);
 		return TComponentLock<TRead<T>>(&Header, LockMode);
@@ -845,7 +845,7 @@ struct FEntityAllocation
 	 * Write typed component data for the specified component type
 	 */
 	template<typename T>
-	UE_NODISCARD TComponentLock<TReadOptional<T>> TryReadComponents(TComponentTypeID<T> ComponentType) const
+	[[nodiscard]] TComponentLock<TReadOptional<T>> TryReadComponents(TComponentTypeID<T> ComponentType) const
 	{
 		if (const FComponentHeader* Header = FindComponentHeader(ComponentType))
 		{
@@ -858,7 +858,7 @@ struct FEntityAllocation
 	 * Write typed component data for the specified component type
 	 */
 	template<typename T>
-	UE_NODISCARD TComponentLock<TWrite<T>> WriteComponents(TComponentTypeID<T> ComponentType, FEntityAllocationWriteContext InWriteContext) const
+	[[nodiscard]] TComponentLock<TWrite<T>> WriteComponents(TComponentTypeID<T> ComponentType, FEntityAllocationWriteContext InWriteContext) const
 	{
 		const FComponentHeader& Header = GetComponentHeaderChecked(ComponentType);
 		return TComponentLock<TWrite<T>>(&Header, LockMode, InWriteContext);
@@ -868,7 +868,7 @@ struct FEntityAllocation
 	 * Attempt to write typed component data for the specified component type
 	 */
 	template<typename T>
-	UE_NODISCARD TComponentLock<TWriteOptional<T>> TryWriteComponents(TComponentTypeID<T> ComponentType, FEntityAllocationWriteContext InWriteContext) const
+	[[nodiscard]] TComponentLock<TWriteOptional<T>> TryWriteComponents(TComponentTypeID<T> ComponentType, FEntityAllocationWriteContext InWriteContext) const
 	{
 		if (const FComponentHeader* Header = FindComponentHeader(ComponentType))
 		{
diff --git a/Engine/Source/Runtime/SlateCore/Private/FastUpdate/SlateInvalidationRoot.cpp b/Engine/Source/Runtime/SlateCore/Private/FastUpdate/SlateInvalidationRoot.cpp
--- a/Engine/Source/Runtime/SlateCore/Private/FastUpdate/SlateInvalidationRoot.cpp
+++ b/Engine/Source/Runtime/SlateCore/Private/FastUpdate/SlateInvalidationRoot.cpp
@@ -493,7 +493,7 @@ namespace UE::Slate::Private
 			Heap.HeapPush(Element, FWidgetOrderGreater());
 		}
 
-		UE_NODISCARD FSlateInvalidationWidgetHeapElement HeapPeek()
+		[[nodiscard]] FSlateInvalidationWidgetHeapElement HeapPeek()
 		{
 			return Heap.HeapTop();
 		}
@@ -503,7 +503,7 @@ namespace UE::Slate::Private
 			Heap.HeapPopDiscard(FWidgetOrderGreater(), false);
 		}
 
-		UE_NODISCARD bool IsEmpty() const
+		[[nodiscard]] bool IsEmpty() const
 		{
 			return Heap.Num() == 0;
 		}
