Always set FName LLM tag for FName block allocations

Some codepaths that allocate FName blocks, in particular the one used for reserving blocks in advance, did not set the LLM tag for FNames. This unnecessarily distorts some memory stats that rely on LLM and is easily fixed by having all allocations go through a single properly marked codepath.

diff --git a/Engine/Source/Runtime/Core/Private/UObject/UnrealNames.cpp b/Engine/Source/Runtime/Core/Private/UObject/UnrealNames.cpp
--- a/Engine/Source/Runtime/Core/Private/UObject/UnrealNames.cpp
+++ b/Engine/Source/Runtime/Core/Private/UObject/UnrealNames.cpp
@@ -446,8 +446,7 @@ public:
 	/** Initializes all member variables. */
 	FNameEntryAllocator()
 	{
-		LLM_SCOPE(ELLMTag::FName);
-		Blocks[0] = (uint8*)FMemory::MallocPersistentAuxiliary(BlockSizeBytes, alignof(FNameEntry));
+		Blocks[0] = AllocBlock();
 	}
 
 	~FNameEntryAllocator()
@@ -654,12 +653,12 @@ private:
 
 	static uint8* AllocBlock()
 	{
+		LLM_SCOPE(ELLMTag::FName);
 		return (uint8*)FMemory::MallocPersistentAuxiliary(BlockSizeBytes, alignof(FNameEntry));
 	}
 	
 	void AllocateNewBlock()
 	{
-		LLM_SCOPE(ELLMTag::FName);
 		// Null-terminate final entry to allow DebugDump() entry iteration
 #if UE_FNAME_OUTLINE_NUMBER
 		if (CurrentByteCursor + NumberedEntrySize <= BlockSizeBytes)
