Removed PLATFORM_COMPILER_HAS_IF_CONSTEXPR.
Deprecated CONSTEXPR, PLATFORM_COMPILER_HAS_DECLTYPE_AUTO, PLATFORM_COMPILER_HAS_FOLD_EXPRESSIONS and UE_NORETURN.
Made UE_NODISCARD mandatory (not yet deprecated as fixup is needed).

diff --git a/Engine/Source/Runtime/Core/Public/Android/AndroidPlatform.h b/Engine/Source/Runtime/Core/Public/Android/AndroidPlatform.h
--- a/Engine/Source/Runtime/Core/Public/Android/AndroidPlatform.h
+++ b/Engine/Source/Runtime/Core/Public/Android/AndroidPlatform.h
@@ -76,12 +76,6 @@ typedef FAndroidTypes FPlatformTypes;
 #define PLATFORM_ENABLE_VECTORINTRINSICS			1
 #define PLATFORM_ENABLE_VECTORINTRINSICS_NEON		PLATFORM_ANDROID_ARM64
 
-#if __has_feature(cxx_decltype_auto)
-	#define PLATFORM_COMPILER_HAS_DECLTYPE_AUTO 1
-#else
-	#define PLATFORM_COMPILER_HAS_DECLTYPE_AUTO 0
-#endif
-
 // some android platform overrides that sub-platforms can disable
 #ifndef USE_ANDROID_JNI
 	#define USE_ANDROID_JNI							1
diff --git a/Engine/Source/Runtime/Core/Public/Clang/ClangPlatform.h b/Engine/Source/Runtime/Core/Public/Clang/ClangPlatform.h
--- a/Engine/Source/Runtime/Core/Public/Clang/ClangPlatform.h
+++ b/Engine/Source/Runtime/Core/Public/Clang/ClangPlatform.h
@@ -12,12 +12,14 @@
 	#error "Compiler is expected to support if constexpr"
 #endif
 
-#if defined(__cpp_fold_expressions)
-	#define PLATFORM_COMPILER_HAS_FOLD_EXPRESSIONS 1
-#else
+#if !defined(__cpp_fold_expressions)
 	#error "Compiler is expected to support fold expressions"
 #endif
 
+#if !__has_feature(cxx_decltype_auto)
+	#error "Compiler is expected to support decltype(auto)"
+#endif
+
 #define PLATFORM_RETURN_ADDRESS()			__builtin_return_address(0)
 #define PLATFORM_RETURN_ADDRESS_POINTER()	__builtin_frame_address(0)
 
diff --git a/Engine/Source/Runtime/Core/Public/HAL/Platform.h b/Engine/Source/Runtime/Core/Public/HAL/Platform.h
--- a/Engine/Source/Runtime/Core/Public/HAL/Platform.h
+++ b/Engine/Source/Runtime/Core/Public/HAL/Platform.h
@@ -233,20 +233,8 @@
 #ifndef PLATFORM_COMPILER_HAS_TCHAR_WMAIN
 	#define PLATFORM_COMPILER_HAS_TCHAR_WMAIN 0
 #endif
-#ifndef PLATFORM_COMPILER_HAS_DECLTYPE_AUTO
-	#define PLATFORM_COMPILER_HAS_DECLTYPE_AUTO 1
-#endif
-#ifdef _MSC_VER
-	#define PLATFORM_COMPILER_HAS_IF_CONSTEXPR 1
-	#ifndef __clang__
-		#pragma deprecated("PLATFORM_COMPILER_HAS_IF_CONSTEXPR")
-	#endif
-#else
-	#define PLATFORM_COMPILER_HAS_IF_CONSTEXPR 1 UE_DEPRECATED_MACRO(5.1, "PLATFORM_COMPILER_HAS_IF_CONSTEXPR has been deprecated and should be replaced with 1.")
-#endif
-#ifndef PLATFORM_COMPILER_HAS_FOLD_EXPRESSIONS
-	#define PLATFORM_COMPILER_HAS_FOLD_EXPRESSIONS 0
-#endif
+#define PLATFORM_COMPILER_HAS_DECLTYPE_AUTO 1 UE_DEPRECATED_MACRO(5.4, "PLATFORM_COMPILER_HAS_DECLTYPE_AUTO has been deprecated and should be replaced with 1.")
+#define PLATFORM_COMPILER_HAS_FOLD_EXPRESSIONS 1 UE_DEPRECATED_MACRO(5.4, "PLATFORM_COMPILER_HAS_FOLD_EXPRESSIONS has been deprecated and should be replaced with 1.")
 #ifndef PLATFORM_COMPILER_HAS_GENERATED_COMPARISON_OPERATORS
 	#define PLATFORM_COMPILER_HAS_GENERATED_COMPARISON_OPERATORS (__cplusplus >= 202002L)
 #endif
@@ -648,13 +636,22 @@
 #endif
 
 /* Use before a function declaration to warn that callers should not ignore the return value */
-#if !defined(UE_NODISCARD) && defined(__has_cpp_attribute)
+#ifdef __has_cpp_attribute
 	#if __has_cpp_attribute(nodiscard)
-		#define UE_NODISCARD [[nodiscard]]
+		// Define this in when all UE_NODISCARD usage has been replaced and we are ready to deprecate
+		#ifdef _MSC_VER
+			#define UE_NODISCARD [[nodiscard]]
+			#ifndef __clang__
+				#pragma deprecated("UE_NODISCARD")
+			#endif
+		#else
+			#define UE_NODISCARD [[nodiscard]] UE_DEPRECATED_MACRO(5.4, "UE_NODISCARD has been deprecated and should be replaced with [[nodiscard]].")
+		#endif
+	#else
+		#error "Compiler is expected to support [[nodiscard]]"
 	#endif
-#endif
-#ifndef UE_NODISCARD
-	#define UE_NODISCARD
+#else
+	#error "Compiler is expected to support [[nodiscard]]"
 #endif
 
 // Use before a constructor declaration to warn when an unnamed temporary object is ignored, e.g. FScopeLock(CS); instead of FScopeLock Lock(CS);
@@ -672,13 +669,21 @@
 #endif
 
 /* Use before a function declaration to indicate that the function never returns */
-#if !defined(UE_NORETURN) && defined(__has_cpp_attribute)
+#ifdef __has_cpp_attribute
 	#if __has_cpp_attribute(noreturn)
-		#define UE_NORETURN [[noreturn]]
+		#ifdef _MSC_VER
+			#define UE_NORETURN [[noreturn]]
+			#ifndef __clang__
+				#pragma deprecated("UE_NORETURN")
+			#endif
+		#else
+			#define UE_NORETURN [[noreturn]] UE_DEPRECATED_MACRO(5.4, "UE_NORETURN has been deprecated and should be replaced with [[noreturn]].")
+		#endif
+	#else
+		#error "Compiler is expected to support [[noreturn]]"
 	#endif
-#endif
-#ifndef UE_NORETURN
-	#define UE_NORETURN
+#else
+	#error "Compiler is expected to support [[noreturn]]"
 #endif
 
 /* Macro wrapper for the consteval keyword which isn't yet present on all compilers - constexpr
@@ -785,9 +790,7 @@
 #ifndef ABSTRACT
 	#define ABSTRACT
 #endif
-#ifndef CONSTEXPR
-	#define CONSTEXPR constexpr
-#endif
+#define CONSTEXPR constexpr UE_DEPRECATED_MACRO(5.4, "CONSTEXPR has been deprecated and should be replaced with constexpr.")
 #ifndef IN
 	#define IN
 #endif
diff --git a/Engine/Source/Runtime/Core/Public/IOS/IOSPlatform.h b/Engine/Source/Runtime/Core/Public/IOS/IOSPlatform.h
--- a/Engine/Source/Runtime/Core/Public/IOS/IOSPlatform.h
+++ b/Engine/Source/Runtime/Core/Public/IOS/IOSPlatform.h
@@ -82,12 +82,6 @@ typedef FIOSPlatformTypes FPlatformTypes;
 #define PLATFORM_CODE_SECTION(Name)						__attribute__((section("__TEXT,__" Name ",regular,pure_instructions"))) \
 														__attribute__((aligned(4)))
 
-#if __has_feature(cxx_decltype_auto)
-	#define PLATFORM_COMPILER_HAS_DECLTYPE_AUTO 1
-#else
-	#define PLATFORM_COMPILER_HAS_DECLTYPE_AUTO 0
-#endif
-
 //mallocpoison not safe with aligned ansi allocator.  returns the larger unaligned size during Free() which causes writes off the end of the allocation.
 #define UE_USE_MALLOC_FILL_BYTES 0 
 
diff --git a/Engine/Source/Runtime/Core/Public/IntelICX/IntelICXPlatform.h b/Engine/Source/Runtime/Core/Public/IntelICX/IntelICXPlatform.h
--- a/Engine/Source/Runtime/Core/Public/IntelICX/IntelICXPlatform.h
+++ b/Engine/Source/Runtime/Core/Public/IntelICX/IntelICXPlatform.h
@@ -10,9 +10,7 @@
 	#error "Compiler is expected to support if constexpr"
 #endif
 
-#if defined(__cpp_fold_expressions)
-	#define PLATFORM_COMPILER_HAS_FOLD_EXPRESSIONS 1
-#else
+#if !defined(__cpp_fold_expressions)
 	#error "Compiler is expected to support fold expressions"
 #endif
 
diff --git a/Engine/Source/Runtime/Core/Public/MSVC/MSVCPlatform.h b/Engine/Source/Runtime/Core/Public/MSVC/MSVCPlatform.h
--- a/Engine/Source/Runtime/Core/Public/MSVC/MSVCPlatform.h
+++ b/Engine/Source/Runtime/Core/Public/MSVC/MSVCPlatform.h
@@ -12,9 +12,7 @@
 	#error "Compiler is expected to support if constexpr"
 #endif
 
-#if defined(__cpp_fold_expressions)
-	#define PLATFORM_COMPILER_HAS_FOLD_EXPRESSIONS 1
-#else
+#if !defined(__cpp_fold_expressions)
 	#error "Compiler is expected to support fold expressions"
 #endif
 
diff --git a/Engine/Source/Runtime/Core/Public/Mac/MacPlatform.h b/Engine/Source/Runtime/Core/Public/Mac/MacPlatform.h
--- a/Engine/Source/Runtime/Core/Public/Mac/MacPlatform.h
+++ b/Engine/Source/Runtime/Core/Public/Mac/MacPlatform.h
@@ -94,12 +94,6 @@ typedef FMacPlatformTypes FPlatformTypes;
 
 #define PLATFORM_CODE_SECTION(Name)						__attribute__((section("__TEXT,__" Name ",regular,pure_instructions")))
 
-#if __has_feature(cxx_decltype_auto)
-	#define PLATFORM_COMPILER_HAS_DECLTYPE_AUTO			1
-#else
-	#define PLATFORM_COMPILER_HAS_DECLTYPE_AUTO			0
-#endif
-
 // Function type macros.
 #define VARARGS															/* Functions with variable arguments */
 #define CDECL															/* Standard C function */
diff --git a/Engine/Source/Runtime/Core/Public/Unix/UnixPlatform.h b/Engine/Source/Runtime/Core/Public/Unix/UnixPlatform.h
--- a/Engine/Source/Runtime/Core/Public/Unix/UnixPlatform.h
+++ b/Engine/Source/Runtime/Core/Public/Unix/UnixPlatform.h
@@ -85,12 +85,6 @@ typedef FUnixPlatformTypes FPlatformTypes;
 
 #define PLATFORM_ENABLE_POPCNT_INTRINSIC				1
 
-#if __has_feature(cxx_decltype_auto)
-	#define PLATFORM_COMPILER_HAS_DECLTYPE_AUTO 1
-#else
-	#define PLATFORM_COMPILER_HAS_DECLTYPE_AUTO 0
-#endif
-
 // SOCK_CLOEXEC is available on Unix since 2.6.27
 #if LINUX_VERSION_CODE >= KERNEL_VERSION(2,6,27)
 	#define PLATFORM_HAS_BSD_SOCKET_FEATURE_CLOSE_ON_EXEC	1
diff --git a/Engine/Source/Runtime/Core/Public/Windows/WindowsPlatform.h b/Engine/Source/Runtime/Core/Public/Windows/WindowsPlatform.h
--- a/Engine/Source/Runtime/Core/Public/Windows/WindowsPlatform.h
+++ b/Engine/Source/Runtime/Core/Public/Windows/WindowsPlatform.h
@@ -102,7 +102,6 @@ typedef FWindowsPlatformTypes FPlatformTypes;
 #endif
 
 #define PLATFORM_SUPPORTS_STACK_SYMBOLS						1
-#define PLATFORM_COMPILER_HAS_DECLTYPE_AUTO					1
 
 #define PLATFORM_GLOBAL_LOG_CATEGORY						LogWindows
 
@@ -164,11 +163,6 @@ typedef FWindowsPlatformTypes FPlatformTypes;
 
 #pragma warning(disable : 4481) // nonstandard extension used: override specifier 'override'
 
-#if defined(__clang__) || _MSC_VER >= 1900
-	#define CONSTEXPR constexpr
-#else
-	#define CONSTEXPR
-#endif
 #define ABSTRACT abstract
 
 // Strings.
