Disable heart beat thread when -llm or -llmcsv is specified as they are known to cause stalls that lead to heart beat hang detection and an app crash

diff --git a/Engine/Source/Runtime/Core/Private/GenericPlatform/GenericPlatformMisc.cpp b/Engine/Source/Runtime/Core/Private/GenericPlatform/GenericPlatformMisc.cpp
--- a/Engine/Source/Runtime/Core/Private/GenericPlatform/GenericPlatformMisc.cpp
+++ b/Engine/Source/Runtime/Core/Private/GenericPlatform/GenericPlatformMisc.cpp
@@ -1559,7 +1559,8 @@ bool FGenericPlatformMisc::UseRenderThread()
 
 bool FGenericPlatformMisc::AllowThreadHeartBeat()
 {
-	static bool bHeartbeat = !FParse::Param(FCommandLine::Get(), TEXT("noheartbeatthread"));
+	// disable heartbeat thread when explicitly disabled via cmd line or when llm\llmcsv is specified as it's known to cause stalls that would result in a hung detection by a heartbeat thread
+	static bool bHeartbeat = !(FParse::Param(FCommandLine::Get(), TEXT("noheartbeatthread")) || FParse::Param(FCommandLine::Get(), TEXT("llm")) || FParse::Param(FCommandLine::Get(), TEXT("llmcsv")));
 	return bHeartbeat;
 }
 
