Remove unused I/O dispatcher function

diff --git a/Engine/Source/Runtime/Core/Private/IO/IoDispatcher.cpp b/Engine/Source/Runtime/Core/Private/IO/IoDispatcher.cpp
--- a/Engine/Source/Runtime/Core/Private/IO/IoDispatcher.cpp
+++ b/Engine/Source/Runtime/Core/Private/IO/IoDispatcher.cpp
@@ -985,19 +985,6 @@ HasScriptObjectsChunk(FIoDispatcher& Dispatcher)
 	return bHasScriptObjectsChunk;
 }
 
-static bool
-HasUseIoStoreParam()
-{
-#if UE_FORCE_USE_IOSTORE
-    return true;
-#elif WITH_IOSTORE_IN_EDITOR
-    static bool bForceIoStore = FParse::Param(FCommandLine::Get(), TEXT("UseIoStore"));
-    return bForceIoStore;
-#else
-    return false;
-#endif
-}
-
 bool
 FIoDispatcher::IsInitialized()
 {
