[Insights] Fix Unreal Insights Context Switch channel not enabling from Editor Widget correctly

diff --git a/Engine/Plugins/TraceUtilities/Source/EditorTraceUtilities/Private/InsightsStatusBar/SInsightsStatusBar.cpp b/Engine/Plugins/TraceUtilities/Source/EditorTraceUtilities/Private/InsightsStatusBar/SInsightsStatusBar.cpp
--- a/Engine/Plugins/TraceUtilities/Source/EditorTraceUtilities/Private/InsightsStatusBar/SInsightsStatusBar.cpp
+++ b/Engine/Plugins/TraceUtilities/Source/EditorTraceUtilities/Private/InsightsStatusBar/SInsightsStatusBar.cpp
@@ -22,6 +22,7 @@
 #include "Modules/ModuleManager.h"
 #include "ProfilingDebugging/MiscTrace.h"
 #include "ProfilingDebugging/TraceScreenshot.h"
+#include "ProfilingDebugging/PlatformEvents.h"
 #include "SRecentTracesList.h"
 #include "Styling/StyleColors.h"
 #include "ToolMenus.h"
@@ -956,7 +957,11 @@ void SInsightsStatusBarWidget::ToggleChannel_Execute(int32 Index)
 {
 	if (Index < ChannelsInfo.Num())
 	{
-		UE::Trace::ToggleChannel(*ChannelsInfo[Index].Name, !ChannelsInfo[Index].bIsEnabled);
+		const FString& ChannelName = ChannelsInfo[Index].Name;
+		bool bChannelShouldBeEnabled = !ChannelsInfo[Index].bIsEnabled;
+		UE::Trace::ToggleChannel(*ChannelName, bChannelShouldBeEnabled);
+
+		FPlatformEventsTrace::OnTraceChannelUpdated(ChannelName, bChannelShouldBeEnabled);
 	}
 }
 
diff --git a/Engine/Source/Runtime/Core/Private/ProfilingDebugging/Microsoft/EventTracingForWindows.cpp b/Engine/Source/Runtime/Core/Private/ProfilingDebugging/Microsoft/EventTracingForWindows.cpp
--- a/Engine/Source/Runtime/Core/Private/ProfilingDebugging/Microsoft/EventTracingForWindows.cpp
+++ b/Engine/Source/Runtime/Core/Private/ProfilingDebugging/Microsoft/EventTracingForWindows.cpp
@@ -26,8 +26,8 @@ public:
 	FPlatformEvents(uint32 SamplingIntervalUsec);
 	~FPlatformEvents();
 
-	void Enable(EPlatformEvent Event);
-	void Disable(EPlatformEvent Event);
+	void Enable(FPlatformEventsTrace::EEventType Event);
+	void Disable(FPlatformEventsTrace::EEventType Event);
 
 private:
 	virtual uint32 Run() override;
@@ -45,7 +45,7 @@ private:
 	FRunnableThread* Thread = nullptr;
 	uint32 SamplingIntervalUsec;
 
-	EPlatformEvent EnabledEvents = EPlatformEvent::None;
+	FPlatformEventsTrace::EEventType EnabledEvents = FPlatformEventsTrace::EEventType::None;
 
 	struct FTrace
 	{
@@ -128,11 +128,7 @@ void FPlatformEvents::TraceEventCallback(PEVENT_RECORD Event)
 
 			if (ThreadId != 0)
 			{
-				UE_TRACE_LOG(PlatformEvent, ContextSwitch, ContextSwitchChannel)
-					<< ContextSwitch.StartTime(StartTime)
-					<< ContextSwitch.EndTime(EndTime)
-					<< ContextSwitch.ThreadId(ThreadId)
-					<< ContextSwitch.CoreNumber(CoreNumber);
+				FPlatformEventsTrace::OutputContextSwitch(StartTime, EndTime, ThreadId, CoreNumber);
 
 				bool bAlreadyAdded = false;
 				GPlatformEvents->ThreadNameSet.Add(ThreadId, &bAlreadyAdded);
@@ -164,10 +160,7 @@ void FPlatformEvents::TraceEventCallback(PEVENT_RECORD Event)
 							::CloseHandle(ProcessHandle);
 						}
 
-						UE_TRACE_LOG(PlatformEvent, ThreadName, ContextSwitchChannel)
-							<< ThreadName.ThreadId(ThreadId)
-							<< ThreadName.ProcessId(ProcessId)
-							<< ThreadName.Name(Name, NameLen);
+						FPlatformEventsTrace::OutputThreadName(ThreadId, ProcessId, Name, NameLen);
 					}
 				}
 			}
@@ -201,10 +194,7 @@ void FPlatformEvents::TraceEventCallback(PEVENT_RECORD Event)
 		// sometimes we get call stack with only one or two entries, which obviously is not very useful
 		if (bValid && Count > 2)
 		{
-			UE_TRACE_LOG(PlatformEvent, StackSample, StackSamplingChannel)
-				<< StackSample.Time(Payload->TimeStamp)
-				<< StackSample.ThreadId(Payload->ThreadId)
-				<< StackSample.Addresses(Payload->Stack, Count);
+			FPlatformEventsTrace::OutputStackSample(Payload->TimeStamp, Payload->ThreadId, Payload->Stack, Count);
 		}
 	}
 }
@@ -230,7 +220,7 @@ FPlatformEvents::~FPlatformEvents()
 
 /////////////////////////////////////////////////////////////////////
 
-void FPlatformEvents::Enable(EPlatformEvent Event)
+void FPlatformEvents::Enable(FPlatformEventsTrace::EEventType Event)
 {
 	if (bStarting)
 	{
@@ -244,7 +234,7 @@ void FPlatformEvents::Enable(EPlatformEvent Event)
 		return;
 	}
 
-	if ((EnabledEvents & Event) != EPlatformEvent::None)
+	if ((EnabledEvents & Event) != FPlatformEventsTrace::EEventType::None)
 	{
 		// if event is already enabled, do nothing
 		return;
@@ -263,11 +253,11 @@ void FPlatformEvents::Enable(EPlatformEvent Event)
 		// in case ETW was already running, enable the necessary event
 
 		EVENT_TRACE_PROPERTIES* Properties = &Trace.Properties;
-		if (Event == EPlatformEvent::ContextSwitch)
+		if (Event == FPlatformEventsTrace::EEventType::ContextSwitch)
 		{
 			Properties->EnableFlags |= EVENT_TRACE_FLAG_CSWITCH;
 		}
-		else if (Event == EPlatformEvent::StackSampling)
+		else if (Event == FPlatformEventsTrace::EEventType::StackSampling)
 		{
 			Properties->EnableFlags |= EVENT_TRACE_FLAG_PROFILE;
 		}
@@ -284,7 +274,7 @@ void FPlatformEvents::Enable(EPlatformEvent Event)
 
 /////////////////////////////////////////////////////////////////////
 
-void FPlatformEvents::Disable(EPlatformEvent Event)
+void FPlatformEvents::Disable(FPlatformEventsTrace::EEventType Event)
 {
 	if (bStarting)
 	{
@@ -298,7 +288,7 @@ void FPlatformEvents::Disable(EPlatformEvent Event)
 		return;
 	}
 
-	if ((EnabledEvents & Event) == EPlatformEvent::None)
+	if ((EnabledEvents & Event) == FPlatformEventsTrace::EEventType::None)
 	{
 		// if event is already disabled, do nothing
 		return;
@@ -306,7 +296,7 @@ void FPlatformEvents::Disable(EPlatformEvent Event)
 	
 	EnumRemoveFlags(EnabledEvents, Event);
 
-	if (EnabledEvents == EPlatformEvent::None)
+	if (EnabledEvents == FPlatformEventsTrace::EEventType::None)
 	{
 		// in case all events are disabled, stop ETW
 		Stop();
@@ -316,11 +306,11 @@ void FPlatformEvents::Disable(EPlatformEvent Event)
 		// in case ETW still needs to run, disable only specific event
 
 		EVENT_TRACE_PROPERTIES* Properties = &Trace.Properties;
-		if (Event == EPlatformEvent::ContextSwitch)
+		if (Event == FPlatformEventsTrace::EEventType::ContextSwitch)
 		{
 			Properties->EnableFlags &= ~EVENT_TRACE_FLAG_CSWITCH;
 		}
-		else if (Event == EPlatformEvent::StackSampling)
+		else if (Event == FPlatformEventsTrace::EEventType::StackSampling)
 		{
 			Properties->EnableFlags &= ~EVENT_TRACE_FLAG_PROFILE;
 		}
@@ -466,11 +456,11 @@ bool FPlatformEvents::StartETW()
 	Properties->FlushTimer = 1;
 	Properties->EnableFlags = 0;
 
-	if ((EnabledEvents & EPlatformEvent::ContextSwitch) != EPlatformEvent::None)
+	if ((EnabledEvents & FPlatformEventsTrace::EEventType::ContextSwitch) != FPlatformEventsTrace::EEventType::None)
 	{
 		Properties->EnableFlags |= EVENT_TRACE_FLAG_CSWITCH;
 	}
-	if ((EnabledEvents & EPlatformEvent::StackSampling) != EPlatformEvent::None)
+	if ((EnabledEvents & FPlatformEventsTrace::EEventType::StackSampling) != FPlatformEventsTrace::EEventType::None)
 	{
 		Properties->EnableFlags |= EVENT_TRACE_FLAG_PROFILE;
 	}
@@ -540,12 +530,12 @@ void FPlatformEvents::StopETW()
 
 /////////////////////////////////////////////////////////////////////
 
-void PlatformEvents_Init(uint32 SamplingIntervalUsec)
+void FPlatformEventsTrace::Init(uint32 SamplingIntervalUsec)
 {
 	GPlatformEvents = new FPlatformEvents(SamplingIntervalUsec);
 }
 
-void PlatformEvents_Enable(EPlatformEvent Event)
+void FPlatformEventsTrace::Enable(FPlatformEventsTrace::EEventType Event)
 {
 	if (GPlatformEvents)
 	{
@@ -553,7 +543,7 @@ void PlatformEvents_Enable(EPlatformEvent Event)
 	}
 }
 
-void PlatformEvents_Disable(EPlatformEvent Event)
+void FPlatformEventsTrace::Disable(FPlatformEventsTrace::EEventType Event)
 {
 	if (GPlatformEvents)
 	{
@@ -561,7 +551,7 @@ void PlatformEvents_Disable(EPlatformEvent Event)
 	}
 }
 
-void PlatformEvents_Stop()
+void FPlatformEventsTrace::Stop()
 {
 	if (GPlatformEvents)
 	{
diff --git a/Engine/Source/Runtime/Core/Private/ProfilingDebugging/PlatformEvents.cpp b/Engine/Source/Runtime/Core/Private/ProfilingDebugging/PlatformEvents.cpp
--- a/Engine/Source/Runtime/Core/Private/ProfilingDebugging/PlatformEvents.cpp
+++ b/Engine/Source/Runtime/Core/Private/ProfilingDebugging/PlatformEvents.cpp
@@ -8,41 +8,114 @@
 UE_TRACE_CHANNEL_DEFINE(ContextSwitchChannel)
 UE_TRACE_CHANNEL_DEFINE(StackSamplingChannel)
 
-UE_TRACE_EVENT_DEFINE(PlatformEvent, ContextSwitch)
-UE_TRACE_EVENT_DEFINE(PlatformEvent, StackSample)
-UE_TRACE_EVENT_DEFINE(PlatformEvent, ThreadName)
+/////////////////////////////////////////////////////////////////////
+
+// represents time interval when thread was running on specific core
+UE_TRACE_EVENT_BEGIN(PlatformEvent, ContextSwitch, NoSync)
+UE_TRACE_EVENT_FIELD(uint64, StartTime)
+UE_TRACE_EVENT_FIELD(uint64, EndTime)
+UE_TRACE_EVENT_FIELD(uint32, ThreadId)
+UE_TRACE_EVENT_FIELD(uint8, CoreNumber)
+UE_TRACE_EVENT_END()
+
+// represents call stack addresses in stack sampling
+UE_TRACE_EVENT_BEGIN(PlatformEvent, StackSample, NoSync)
+UE_TRACE_EVENT_FIELD(uint64, Time)
+UE_TRACE_EVENT_FIELD(uint32, ThreadId)
+UE_TRACE_EVENT_FIELD(uint64[], Addresses)
+UE_TRACE_EVENT_END()
+
+// thread information for context switches
+// NOTE: in some cases name of process can be empty, for example when there are no
+// privileges to query it or there are no name for process thread belongs to.
+// Depending on platform actual name can be absolute path to process executable.
+UE_TRACE_EVENT_BEGIN(PlatformEvent, ThreadName, NoSync)
+UE_TRACE_EVENT_FIELD(uint32, ThreadId)
+UE_TRACE_EVENT_FIELD(uint32, ProcessId)
+UE_TRACE_EVENT_FIELD(UE::Trace::WideString, Name)
+UE_TRACE_EVENT_END()
 
 /////////////////////////////////////////////////////////////////////
 
-EPlatformEvent PlatformEvents_GetEvent(const FString& Name)
+#if PLATFORM_SUPPORTS_PLATFORM_EVENTS
+
+void FPlatformEventsTrace::OutputContextSwitch(uint64 StartTime, uint64 EndTime, uint32 ThreadId, uint8 CoreNumber)
+{
+	UE_TRACE_LOG(PlatformEvent, ContextSwitch, ContextSwitchChannel)
+		<< ContextSwitch.StartTime(StartTime)
+		<< ContextSwitch.EndTime(EndTime)
+		<< ContextSwitch.ThreadId(ThreadId)
+		<< ContextSwitch.CoreNumber(CoreNumber);
+}
+
+/////////////////////////////////////////////////////////////////////
+
+void FPlatformEventsTrace::OutputStackSample(uint64 Time, uint32 ThreadId, const uint64* Addresses, uint32 AddressCount)
+{
+	UE_TRACE_LOG(PlatformEvent, StackSample, StackSamplingChannel)
+		<< StackSample.Time(Time)
+		<< StackSample.ThreadId(ThreadId)
+		<< StackSample.Addresses(Addresses, AddressCount);
+}
+
+/////////////////////////////////////////////////////////////////////
+
+void FPlatformEventsTrace::OutputThreadName(uint32 ThreadId, uint32 ProcessId, const TCHAR* Name, uint32 NameLen)
+{
+	UE_TRACE_LOG(PlatformEvent, ThreadName, ContextSwitchChannel)
+		<< ThreadName.ThreadId(ThreadId)
+		<< ThreadName.ProcessId(ProcessId)
+		<< ThreadName.Name(Name, NameLen);
+}
+
+/////////////////////////////////////////////////////////////////////
+
+FPlatformEventsTrace::EEventType FPlatformEventsTrace::GetEvent(const FString& Name)
 {
 	if (Name == TEXT("contextswitch"))
 	{
-		return EPlatformEvent::ContextSwitch;
+		return FPlatformEventsTrace::EEventType::ContextSwitch;
 	}
 	else if (Name == TEXT("stacksampling"))
 	{
-		return EPlatformEvent::StackSampling;
+		return FPlatformEventsTrace::EEventType::StackSampling;
 	}
 	else
 	{
-		return EPlatformEvent::None;
+		return FPlatformEventsTrace::EEventType::None;
 	}
 }
 
 /////////////////////////////////////////////////////////////////////
 
-#if PLATFORM_SUPPORTS_PLATFORM_EVENTS
+void FPlatformEventsTrace::OnTraceChannelUpdated(const FString& ChannelName, bool bIsEnabled)
+{
+	FPlatformEventsTrace::EEventType Event = FPlatformEventsTrace::GetEvent(ChannelName.Replace(TEXT("Channel"), TEXT("")));
+
+	if (Event != FPlatformEventsTrace::EEventType::None)
+	{
+		if (bIsEnabled)
+		{
+			FPlatformEventsTrace::Enable(Event);
+		}
+		else
+		{
+			FPlatformEventsTrace::Disable(Event);
+		}
+	}
+}
+
+/////////////////////////////////////////////////////////////////////
 
-void PlatformEvents_PostInit()
+void FPlatformEventsTrace::PostInit()
 {
 	if (TRACE_PRIVATE_CHANNELEXPR_IS_ENABLED(ContextSwitchChannel))
 	{
-		PlatformEvents_Enable(EPlatformEvent::ContextSwitch);
+		Enable(FPlatformEventsTrace::EEventType::ContextSwitch);
 	}
 	else if (TRACE_PRIVATE_CHANNELEXPR_IS_ENABLED(StackSamplingChannel))
 	{
-		PlatformEvents_Enable(EPlatformEvent::StackSampling);
+		Enable(FPlatformEventsTrace::EEventType::StackSampling);
 	}
 }
 
diff --git a/Engine/Source/Runtime/Core/Private/ProfilingDebugging/PlatformEvents.h b/Engine/Source/Runtime/Core/Private/ProfilingDebugging/PlatformEvents.h
deleted file mode 100644
diff --git a/Engine/Source/Runtime/Core/Private/ProfilingDebugging/TraceAuxiliary.cpp b/Engine/Source/Runtime/Core/Private/ProfilingDebugging/TraceAuxiliary.cpp
--- a/Engine/Source/Runtime/Core/Private/ProfilingDebugging/TraceAuxiliary.cpp
+++ b/Engine/Source/Runtime/Core/Private/ProfilingDebugging/TraceAuxiliary.cpp
@@ -392,11 +392,7 @@ bool FTraceAuxiliaryImpl::EnableChannel(const TCHAR* Channel)
 		return false;
 	}
 
-	EPlatformEvent Event = PlatformEvents_GetEvent(Channel);
-	if (Event != EPlatformEvent::None)
-	{
-		PlatformEvents_Enable(Event);
-	}
+	FPlatformEventsTrace::OnTraceChannelUpdated(Channel, true);
 
 	return UE::Trace::ToggleChannel(Channel, true);
 }
@@ -412,11 +408,7 @@ void FTraceAuxiliaryImpl::DisableChannel(const TCHAR* Channel)
 		return;
 	}
 
-	EPlatformEvent Event = PlatformEvents_GetEvent(Channel);
-	if (Event != EPlatformEvent::None)
-	{
-		PlatformEvents_Disable(Event);
-	}
+	FPlatformEventsTrace::OnTraceChannelUpdated(Channel, false);
 
 	UE::Trace::ToggleChannel(Channel, false);
 }
@@ -1535,8 +1527,8 @@ void FTraceAuxiliary::Initialize(const TCHAR* CommandLine)
 	// By default use 1 msec for stack sampling interval.
 	uint32 Microseconds = 1000;
 	FParse::Value(CommandLine, TEXT("-samplinginterval="), Microseconds);
-	PlatformEvents_Init(Microseconds);
-	PlatformEvents_PostInit();
+	FPlatformEventsTrace::Init(Microseconds);
+	FPlatformEventsTrace::PostInit();
 
 #if CSV_PROFILER
 	FCoreDelegates::OnEndFrame.AddRaw(&GTraceAuxiliary, &FTraceAuxiliaryImpl::UpdateCsvStats);
@@ -1590,7 +1582,7 @@ void FTraceAuxiliary::Shutdown()
 
 	// Make sure all platform event functionality has shut down as on some
 	// platforms it impacts whole system, even if application has terminated.
-	PlatformEvents_Stop();
+	FPlatformEventsTrace::Stop();
 #endif
 }
 
diff --git a/Engine/Source/Runtime/Core/Public/ProfilingDebugging/PlatformEvents.h b/Engine/Source/Runtime/Core/Public/ProfilingDebugging/PlatformEvents.h
new file mode 100644
--- /dev/null
+++ b/Engine/Source/Runtime/Core/Public/ProfilingDebugging/PlatformEvents.h
@@ -0,0 +1,66 @@
+// Copyright Epic Games, Inc. All Rights Reserved.
+#pragma once
+
+#include "CoreTypes.h"
+#include "Misc/EnumClassFlags.h"
+#include "Trace/Trace.h"
+#include "Trace/Trace.inl"
+
+class FString;
+
+#if !defined(PLATFORM_SUPPORTS_PLATFORM_EVENTS)
+#	define PLATFORM_SUPPORTS_PLATFORM_EVENTS 0
+#endif
+
+/////////////////////////////////////////////////////////////////////
+
+UE_TRACE_CHANNEL_EXTERN(ContextSwitchChannel)
+UE_TRACE_CHANNEL_EXTERN(StackSamplingChannel)
+
+/////////////////////////////////////////////////////////////////////
+
+class FPlatformEventsTrace
+{
+public:
+	enum class EEventType
+	{
+		None = 0x00,
+		ContextSwitch = 0x01,
+		StackSampling = 0x02,
+	};
+
+	CORE_API static EEventType GetEvent(const FString& Name);
+
+#if PLATFORM_SUPPORTS_PLATFORM_EVENTS
+
+	CORE_API static void Init(uint32 SamplingIntervalUsec);
+	CORE_API static void PostInit();
+	CORE_API static void Enable(EEventType Event);
+	CORE_API static void Disable(EEventType Event);
+	CORE_API static void Stop();
+	CORE_API static void OnTraceChannelUpdated(const FString& ChannelName, bool bIsEnabled);
+
+	CORE_API static void OutputContextSwitch(uint64 StartTime, uint64 EndTime, uint32 ThreadId, uint8 CoreNumber);
+	CORE_API static void OutputStackSample(uint64 Time, uint32 ThreadId, const uint64* Addresses, uint32 AddressCount);
+	CORE_API static void OutputThreadName(uint32 ThreadId, uint32 ProcessId, const TCHAR* Name, uint32 NameLen);
+
+#else // PLATFORM_SUPPORTS_PLATFORM_EVENTS
+
+	CORE_API static void Init(uint32 SamplingIntervalUsec) {}
+	CORE_API static void PostInit() {}
+	CORE_API static void Enable(EEventType Event) {}
+	CORE_API static void Disable(EEventType Event) {}
+	CORE_API static void Stop() {}
+	CORE_API static void OnTraceChannelUpdated(const FString& ChannelName, bool bIsEnabled) {}
+
+	CORE_API static void OutputContextSwitch(uint64 StartTime, uint64 EndTime, uint32 ThreadId, uint8 CoreNumber) {}
+	CORE_API static void OutputStackSample(uint64 Time, uint32 ThreadId, const uint64* Addresses, uint32 AddressCount) {}
+	CORE_API static void OutputThreadName(uint32 ThreadId, uint32 ProcessId, const TCHAR* Name, uint32 NameLen) {}
+
+#endif // PLATFORM_SUPPORTS_PLATFORM_EVENTS
+};
+
+ENUM_CLASS_FLAGS(FPlatformEventsTrace::EEventType)
+
+
+/////////////////////////////////////////////////////////////////////
