Fixed TBitsToSizeType's static assert, which would only fire if passed zero.

diff --git a/Engine/Source/Runtime/Core/Public/Containers/ContainerAllocationPolicies.h b/Engine/Source/Runtime/Core/Public/Containers/ContainerAllocationPolicies.h
--- a/Engine/Source/Runtime/Core/Public/Containers/ContainerAllocationPolicies.h
+++ b/Engine/Source/Runtime/Core/Public/Containers/ContainerAllocationPolicies.h
@@ -600,7 +600,8 @@ struct TAllocatorTraits<TAlignedHeapAllocator<Alignment>> : TAllocatorTraitsBase
 template <int IndexSize>
 struct TBitsToSizeType
 {
-	static_assert(IndexSize, "Unsupported allocator index size.");
+	// Fabricate a compile-time false result that's still dependent on the template parameter
+	static_assert(IndexSize == IndexSize+1, "Unsupported allocator index size.");
 };
 
 template <> struct TBitsToSizeType<8>  { using Type = int8; };
