Replaced explicit memmoves with RelocateConstructItems.
Optimized some repeated rereads of unchanged data mid-loop.

diff --git a/Engine/Source/Runtime/Core/Public/Containers/Array.h b/Engine/Source/Runtime/Core/Public/Containers/Array.h
--- a/Engine/Source/Runtime/Core/Public/Containers/Array.h
+++ b/Engine/Source/Runtime/Core/Public/Containers/Array.h
@@ -1691,18 +1691,15 @@ private:
 			CheckInvariants();
 			checkSlow((Count >= 0) & (Index >= 0) & (Index + Count <= ArrayNum));
 
-			DestructItems(GetData() + Index, Count);
+			ElementType* Dest = GetData() + Index;
+
+			DestructItems(Dest, Count);
 
 			// Skip memmove in the common case that there is nothing to move.
-			SizeType NumToMove = ArrayNum - Index - Count;
+			SizeType NumToMove = (ArrayNum - Index) - Count;
 			if (NumToMove)
 			{
-				FMemory::Memmove
-				(
-					(uint8*)AllocatorInstance.GetAllocation() + (Index)* sizeof(ElementType),
-					(uint8*)AllocatorInstance.GetAllocation() + (Index + Count) * sizeof(ElementType),
-					NumToMove * sizeof(ElementType)
-				);
+				RelocateConstructItems<ElementType>(Dest, Dest + Count, NumToMove);
 			}
 			ArrayNum -= Count;
 
@@ -1750,19 +1747,17 @@ private:
 			CheckInvariants();
 			checkSlow((Count >= 0) & (Index >= 0) & (Index + Count <= ArrayNum));
 
-			DestructItems(GetData() + Index, Count);
+			ElementType* Data = GetData();
+			ElementType* Dest = Data + Index;
+
+			DestructItems(Dest, Count);
 
 			// Replace the elements in the hole created by the removal with elements from the end of the array, so the range of indices used by the array is contiguous.
-			const SizeType NumElementsInHole = Count;
-			const SizeType NumElementsAfterHole = ArrayNum - (Index + Count);
-			const SizeType NumElementsToMoveIntoHole = FPlatformMath::Min(NumElementsInHole, NumElementsAfterHole);
+			const SizeType NumElementsAfterHole = (ArrayNum - Index) - Count;
+			const SizeType NumElementsToMoveIntoHole = FPlatformMath::Min(Count, NumElementsAfterHole);
 			if (NumElementsToMoveIntoHole)
 			{
-				FMemory::Memcpy(
-					(uint8*)AllocatorInstance.GetAllocation() + (Index) * sizeof(ElementType),
-					(uint8*)AllocatorInstance.GetAllocation() + (ArrayNum - NumElementsToMoveIntoHole) * sizeof(ElementType),
-					NumElementsToMoveIntoHole * sizeof(ElementType)
-					);
+				RelocateConstructItems<ElementType>(Dest, Data + (ArrayNum - NumElementsToMoveIntoHole), NumElementsToMoveIntoHole);
 			}
 			ArrayNum -= Count;
 
@@ -2600,33 +2595,35 @@ public:
 			return 0; // nothing to do, loop assumes one item so need to deal with this edge case here
 		}
 
+		ElementType* Data = GetData();
+
 		SizeType WriteIndex = 0;
 		SizeType ReadIndex = 0;
-		bool NotMatch = !::Invoke(Predicate, GetData()[ReadIndex]); // use a ! to guarantee it can't be anything other than zero or one
+		bool bNotMatch = !::Invoke(Predicate, Data[ReadIndex]); // use a ! to guarantee it can't be anything other than zero or one
 		do
 		{
 			SizeType RunStartIndex = ReadIndex++;
-			while (ReadIndex < OriginalNum && NotMatch == !::Invoke(Predicate, GetData()[ReadIndex]))
+			while (ReadIndex < OriginalNum && bNotMatch == !::Invoke(Predicate, Data[ReadIndex]))
 			{
 				ReadIndex++;
 			}
 			SizeType RunLength = ReadIndex - RunStartIndex;
 			checkSlow(RunLength > 0);
-			if (NotMatch)
+			if (bNotMatch)
 			{
 				// this was a non-matching run, we need to move it
 				if (WriteIndex != RunStartIndex)
 				{
-					FMemory::Memmove(&GetData()[WriteIndex], &GetData()[RunStartIndex], sizeof(ElementType)* RunLength); //-V598
+					RelocateConstructItems<ElementType>(Data + WriteIndex, Data + RunStartIndex, RunLength);
 				}
 				WriteIndex += RunLength;
 			}
 			else
 			{
 				// this was a matching run, delete it
-				DestructItems(GetData() + RunStartIndex, RunLength);
+				DestructItems(Data + RunStartIndex, RunLength);
 			}
-			NotMatch = !NotMatch;
+			bNotMatch = !bNotMatch;
 		} while (ReadIndex < OriginalNum);
 
 		ArrayNum = WriteIndex;
