[RunnableThread]
 - Remove default parameters on derived classes that are inconsistent with the base class to force direct callers to specify what behavior they expect.
 - These derivations were never called directly and PVS is complaining about the inconsistency.

diff --git a/Engine/Source/Runtime/Core/Private/HAL/PThreadRunnableThread.h b/Engine/Source/Runtime/Core/Private/HAL/PThreadRunnableThread.h
--- a/Engine/Source/Runtime/Core/Private/HAL/PThreadRunnableThread.h
+++ b/Engine/Source/Runtime/Core/Private/HAL/PThreadRunnableThread.h
@@ -263,7 +263,7 @@ public:
 		// Impossible in pthreads!
 	}
 
-	virtual bool Kill(bool bShouldWait = false) override
+	virtual bool Kill(bool bShouldWait) override
 	{
 		check(Thread && "Did you forget to call Create()?");
 		bool bDidExitOK = true;
diff --git a/Engine/Source/Runtime/Core/Private/Microsoft/MicrosoftRunnableThread.h b/Engine/Source/Runtime/Core/Private/Microsoft/MicrosoftRunnableThread.h
--- a/Engine/Source/Runtime/Core/Private/Microsoft/MicrosoftRunnableThread.h
+++ b/Engine/Source/Runtime/Core/Private/Microsoft/MicrosoftRunnableThread.h
@@ -76,7 +76,7 @@ public:
 		}
 	}
 
-	virtual bool Kill(bool bShouldWait = false) override
+	virtual bool Kill(bool bShouldWait) override
 	{
 		check(Thread && "Did you forget to call Create()?");
 		bool bDidExitOK = true;
