[ParallelFor]
 - Restore ParallelWithPreWork behavior that should always call prework even when number of tasks to execute is 0.

diff --git a/Engine/Source/Runtime/Core/Public/Async/ParallelFor.h b/Engine/Source/Runtime/Core/Public/Async/ParallelFor.h
--- a/Engine/Source/Runtime/Core/Public/Async/ParallelFor.h
+++ b/Engine/Source/Runtime/Core/Public/Async/ParallelFor.h
@@ -116,8 +116,12 @@ namespace ParallelForImpl
 	{
 		if (Num == 0)
 		{
+			// Contract is that prework should always be called even when number of tasks is 0.
+			// We omit the trace scope here to avoid noise when the prework is empty since this amounts to just calling a function anyway with nothing specific to parallelfor itself.
+			CurrentThreadWorkToDoBeforeHelping();
 			return;
 		}
+
 		SCOPE_CYCLE_COUNTER(STAT_ParallelFor);
 		TRACE_CPUPROFILER_EVENT_SCOPE(ParallelFor);
 		check(Num >= 0);
