- Changed the output of "confighash" (via FConfigCacheIni::Dump) to sort the (non-known) files, the sections within the file, and the key/value pairs within the section. This will make diffing two builds much easier.
- Removed useless "Files map" log

diff --git a/Engine/Source/Runtime/Core/Private/Misc/ConfigCacheIni.cpp b/Engine/Source/Runtime/Core/Private/Misc/ConfigCacheIni.cpp
--- a/Engine/Source/Runtime/Core/Private/Misc/ConfigCacheIni.cpp
+++ b/Engine/Source/Runtime/Core/Private/Misc/ConfigCacheIni.cpp
@@ -2961,13 +2961,27 @@ void FConfigCacheIni::Exit()
 void FConfigCacheIni::DumpFile(FOutputDevice& Ar, const FString& Filename, const FConfigFile& File)
 {
 	Ar.Logf(TEXT("FileName: %s"), *Filename);
-	for (FConfigFile::TConstIterator FileIt(File); FileIt; ++FileIt)
+
+	// sort the sections (and keys below) for easier diffing
+	TArray<FString> SectionKeys;
+	File.GetKeys(SectionKeys);
+	SectionKeys.Sort();
+	for (const FString& SectionKey : SectionKeys)
 	{
-		const FConfigSection& Sec = FileIt.Value();
-		Ar.Logf(TEXT("   [%s]"), *FileIt.Key());
-		for (FConfigSectionMap::TConstIterator SecIt(Sec); SecIt; ++SecIt)
+		const FConfigSection& Sec = File[SectionKey];
+		Ar.Logf(TEXT("   [%s]"), *SectionKey);
+
+		TArray<FName> Keys;
+		Sec.GetKeys(Keys);
+		Keys.Sort(FNameLexicalLess());
+		for (FName Key : Keys)
 		{
-			Ar.Logf(TEXT("   %s=%s"), *SecIt.Key().ToString(), *SecIt.Value().GetValue());
+			TArray<FConfigValue> Values;
+			Sec.MultiFind(Key, Values, true);
+			for (const FConfigValue& Value : Values)
+			{
+				Ar.Logf(TEXT("   %s=%s"), *Key.ToString(), *Value.GetValue());
+			}
 		}
 
 		Ar.Log(LINE_TERMINATOR);
@@ -2976,22 +2990,23 @@ void FConfigCacheIni::DumpFile(FOutputDevice& Ar, const FString& Filename, const
 
 void FConfigCacheIni::Dump(FOutputDevice& Ar, const TCHAR* BaseIniName)
 {
-	if (BaseIniName == nullptr)
-	{
-		Ar.Log( TEXT("Files map:") );
-		OtherFiles.Dump( Ar );
-	}
-
 	for (const FConfigCacheIni::FKnownConfigFiles::FKnownConfigFile& File : KnownFiles.Files)
 	{
-		DumpFile(Ar, File.IniName.ToString(), File.IniFile);
+		if (BaseIniName == nullptr || File.IniName == BaseIniName)
+		{
+			DumpFile(Ar, File.IniName.ToString(), File.IniFile);
+		}
 	}
 
-	for (TPair<FString, FConfigFile*> Pair : OtherFiles)
+	// sort the non-known files for easier diffing
+	TArray<FString> Keys;
+	OtherFiles.GetKeys(Keys);
+	Algo::Sort(Keys);
+	for (const FString& Key : Keys)
 	{
-		if (BaseIniName == nullptr || FPaths::GetBaseFilename(Pair.Key) == BaseIniName)
+		if (BaseIniName == nullptr || FPaths::GetBaseFilename(Key) == BaseIniName)
 		{
-			DumpFile(Ar, Pair.Key, *Pair.Value);
+			DumpFile(Ar, Key, *OtherFiles[Key]);
 		}
 	}
 }
