[MemoryWriter] template array allocator type; support 8/16 bit index sizes

diff --git a/Engine/Source/Runtime/Core/Public/Serialization/MemoryWriter.h b/Engine/Source/Runtime/Core/Public/Serialization/MemoryWriter.h
--- a/Engine/Source/Runtime/Core/Public/Serialization/MemoryWriter.h
+++ b/Engine/Source/Runtime/Core/Public/Serialization/MemoryWriter.h
@@ -15,14 +15,14 @@
 /**
  * Archive for storing arbitrary data to the specified memory location
  */
-template <int IndexSize>
-class TMemoryWriter : public FMemoryArchive
-{
-	static_assert(IndexSize == 32 || IndexSize == 64, "Only 32-bit and 64-bit index sizes supported");
-	using IndexSizeType = typename TBitsToSizeType<IndexSize>::Type;
 
+template <typename ArrayAllocatorType>
+class TMemoryWriterBase : public FMemoryArchive
+{
+	using IndexSizeType = typename ArrayAllocatorType::SizeType;
+	static constexpr int32 IndexSize = sizeof(IndexSizeType);
 public:
-	TMemoryWriter( TArray<uint8, TSizedDefaultAllocator<IndexSize>>& InBytes, bool bIsPersistent = false, bool bSetOffset = false, const FName InArchiveName = NAME_None )
+	TMemoryWriterBase( TArray<uint8, ArrayAllocatorType>& InBytes, bool bIsPersistent = false, bool bSetOffset = false, const FName InArchiveName = NAME_None )
 	: FMemoryArchive()
 	, Bytes(InBytes)
 	, ArchiveName(InArchiveName)
@@ -41,11 +41,12 @@ public:
 		if( NumBytesToAdd > 0 )
 		{
 			const int64 NewArrayCount = Bytes.Num() + NumBytesToAdd;
-			if constexpr (IndexSize == 32)
+			if constexpr (IndexSize < 64)
 			{
-				if (NewArrayCount >= MAX_int32)
+				constexpr IndexSizeType MaxValue = std::numeric_limits<IndexSizeType>::max();
+				if (NewArrayCount >= MaxValue)
 				{
-					UE_LOG(LogSerialization, Fatal, TEXT("FMemoryWriter does not support data larger than 2GB. Archive name: %s."), *ArchiveName.ToString());
+					UE_LOG(LogSerialization, Fatal, TEXT("TMemoryWriter with IndexSize=%d does not support data larger than %d bytes. Archive name: %s."), IndexSize, MaxValue, *ArchiveName.ToString());
 				}
 			}
 
@@ -73,14 +74,7 @@ public:
 			return ArchiveName.ToString();
 		}
 
-		if constexpr (IndexSize == 64)
-		{
-			return TEXT("FMemoryWriter64");
-		}
-		else
-		{
-			return TEXT("FMemoryWriter");
-		}
+		return FString::Printf(TEXT("FMemoryWriter%d"), IndexSize);
 	}
 
 	int64 TotalSize() override
@@ -90,12 +84,14 @@ public:
 
 protected:
 
-	TArray<uint8, TSizedDefaultAllocator<IndexSize>>&	Bytes;
+	TArray<uint8, ArrayAllocatorType>&	Bytes;
 
 	/** Archive name, used to debugging, by default set to NAME_None. */
 	const FName ArchiveName;
 };
 
+template <int IndexSize>
+using TMemoryWriter = TMemoryWriterBase<TSizedDefaultAllocator<IndexSize>>;
 
 // FMemoryWriter and FMemoryWriter64 are implemented as derived classes rather than aliases
 // so that forward declarations will work.
