Fix system PLATFORM_VISIONOS define leaking into UE code

diff --git a/Engine/Source/Runtime/Core/Public/Apple/PostAppleSystemHeaders.h b/Engine/Source/Runtime/Core/Public/Apple/PostAppleSystemHeaders.h
--- a/Engine/Source/Runtime/Core/Public/Apple/PostAppleSystemHeaders.h
+++ b/Engine/Source/Runtime/Core/Public/Apple/PostAppleSystemHeaders.h
@@ -3,3 +3,4 @@
 // Restore macros
 #pragma pop_macro("PLATFORM_IOS")
 #pragma pop_macro("PLATFORM_TVOS")
+#pragma pop_macro("PLATFORM_VISIONOS")
\ No newline at end of file
diff --git a/Engine/Source/Runtime/Core/Public/Apple/PreAppleSystemHeaders.h b/Engine/Source/Runtime/Core/Public/Apple/PreAppleSystemHeaders.h
--- a/Engine/Source/Runtime/Core/Public/Apple/PreAppleSystemHeaders.h
+++ b/Engine/Source/Runtime/Core/Public/Apple/PreAppleSystemHeaders.h
@@ -3,3 +3,4 @@
 // Save these macros; mach-o/loader.h redefines them
 #pragma push_macro("PLATFORM_IOS")
 #pragma push_macro("PLATFORM_TVOS")
+#pragma push_macro("PLATFORM_VISIONOS")
