Upgraded PVS-Studio to 7.28
#rnx

diff --git a/Engine/Source/Runtime/Core/Public/Windows/WindowsPlatformCodeAnalysis.h b/Engine/Source/Runtime/Core/Public/Windows/WindowsPlatformCodeAnalysis.h
--- a/Engine/Source/Runtime/Core/Public/Windows/WindowsPlatformCodeAnalysis.h
+++ b/Engine/Source/Runtime/Core/Public/Windows/WindowsPlatformCodeAnalysis.h
@@ -103,7 +103,7 @@
 	#pragma warning(disable : 6240) // warning C6240 : (<expression> && <non-zero constant>) always evaluates to the result of <expression>. Did you intend to use the bitwise-and operator?
 
 	//PVS-Studio settings:
-	//-V::505,542,581,601,623,668,677,690,704,719,720,730,735,751,758,1002,1008,1062
+	//-V::505,542,581,601,623,668,677,690,704,719,720,730,735,751,758,780,781,1002,1008,1055,1062,1100
 	//-V:TRYCOMPRESSION:519,547
 	//-V:check(:501,547,560,605
 	//-V:checkf(:510
@@ -229,6 +229,7 @@
 	//-V:Linker:678
 	//-V:self:678
 	//-V:AccumulateParentID:678
+	//-V:FindChar:679
 
 	// The following classes retain a reference to data supplied in the constructor by the derived class which can not yet be initialized.
 	//-V:FMemoryWriter(:1050
@@ -252,6 +253,17 @@
 	// warning V1051: Consider checking for misprints. It's possible that the 'LayerInfo' should be checked here.
 	//-V::1051
 
+	// V016: User annotation was not applied to a virtual function. To force the annotation, use the 'enable_on_virtual' flag.
+	//-V::016
+
+	// Disabling because incorrectly flagging all TStaticArrays
+	// V557: Array overrun is possible
+	//-V::557
+
+	// Disabling because too many virtuals currently in use in constructors/destructors, need to revist
+	// V1053: Calling the 'foo' virtual function in the constructor/destructor may lead to unexpected result at runtime.
+	//-V::1053
+
 #else // defined(__clang_analyzer__)
 
 // A fake function marked with noreturn that acts as a marker for CA_ASSUME to ensure the
