Include VisionOS in platform debug event handling and don't assume #else is always IOS.

diff --git a/Engine/Source/Runtime/Core/Private/Apple/ApplePlatformDebugEvents.cpp b/Engine/Source/Runtime/Core/Private/Apple/ApplePlatformDebugEvents.cpp
--- a/Engine/Source/Runtime/Core/Private/Apple/ApplePlatformDebugEvents.cpp
+++ b/Engine/Source/Runtime/Core/Private/Apple/ApplePlatformDebugEvents.cpp
@@ -43,10 +43,25 @@
 	#undef PLATFORM_MAC
 	#undef PLATFORM_IOS
 	#undef PLATFORM_TVOS
+	#undef PLATFORM_VISIONOS
 	
 	#define PLATFORM_MAC 1
 	#define PLATFORM_IOS 0
 	#define PLATFORM_TVOS 0
+	#define PLATFORM_VISIONOS 0
+#elif PLATFORM_VISIONOS
+	#include <os/log.h>
+	#include <os/signpost.h>
+	
+	#undef PLATFORM_MAC
+	#undef PLATFORM_IOS
+	#undef PLATFORM_TVOS
+	#undef PLATFORM_VISIONOS
+	
+	#define PLATFORM_MAC 0
+	#define PLATFORM_IOS 1
+	#define PLATFORM_TVOS 0
+	#define PLATFORM_VISIONOS 1
 #elif PLATFORM_TVOS	
 	#include <os/log.h>
 	#include <os/signpost.h>
@@ -54,21 +69,27 @@
 	#undef PLATFORM_MAC
 	#undef PLATFORM_IOS
 	#undef PLATFORM_TVOS
+	#undef PLATFORM_VISIONOS
 	
 	#define PLATFORM_MAC 0
 	#define PLATFORM_IOS 1
 	#define PLATFORM_TVOS 1
-#else
+	#define PLATFORM_VISIONOS 0
+#elif PLATFORM_IOS
 	#include <os/log.h>
 	#include <os/signpost.h>
 	
 	#undef PLATFORM_MAC
 	#undef PLATFORM_IOS
 	#undef PLATFORM_TVOS
+	#undef PLATFORM_VISIONOS
 	
 	#define PLATFORM_MAC 0
 	#define PLATFORM_IOS 1
 	#define PLATFORM_TVOS 0
+	#define PLATFORM_VISIONOS 0
+#else
+	#error "Unknown Apple Platform requested for Profiling"
 #endif
 
 DEFINE_LOG_CATEGORY(LogInstruments)
