Added EAllowShrinking as a better replacement for bAllowShrinking, with optional deprecation.

diff --git a/Engine/Source/Runtime/Core/Private/Containers/String.cpp.inl b/Engine/Source/Runtime/Core/Private/Containers/String.cpp.inl
--- a/Engine/Source/Runtime/Core/Private/Containers/String.cpp.inl
+++ b/Engine/Source/Runtime/Core/Private/Containers/String.cpp.inl
@@ -645,9 +645,9 @@ void UE_STRING_CLASS::InsertAt(int32 Index, const UE_STRING_CLASS& Characters)
 	}
 }
 
-void UE_STRING_CLASS::RemoveAt(int32 Index, int32 Count, bool bAllowShrinking)
+void UE_STRING_CLASS::RemoveAt(int32 Index, int32 Count, EAllowShrinking AllowShrinking)
 {
-	Data.RemoveAt(Index, FMath::Clamp(Count, 0, Len()-Index), bAllowShrinking);
+	Data.RemoveAt(Index, FMath::Clamp(Count, 0, Len()-Index), AllowShrinking);
 }
 
 bool UE_STRING_CLASS::RemoveFromStart(const ElementType* InPrefix, int32 InPrefixLen, ESearchCase::Type SearchCase)
diff --git a/Engine/Source/Runtime/Core/Public/Containers/AllowShrinking.h b/Engine/Source/Runtime/Core/Public/Containers/AllowShrinking.h
new file mode 100644
--- /dev/null
+++ b/Engine/Source/Runtime/Core/Public/Containers/AllowShrinking.h
@@ -0,0 +1,15 @@
+// Copyright Epic Games, Inc. All Rights Reserved.
+
+#pragma once
+
+#include "CoreTypes.h"
+
+// Specifies whether or not a container's removal operation should attempt to auto-shrink the container's reserved memory usage
+enum class EAllowShrinking : uint8
+{
+	No,
+	Yes
+};
+
+// Can comment in this when fixing up existing calls, or made permanent when it's ready to be
+#define UE_ALLOWSHRINKING_BOOL_DEPRECATED(FunctionName) //UE_DEPRECATED(5.4, FunctionName " with a boolean bAllowShrinking has been deprecated - please use the EAllowShrinking enum instead")
diff --git a/Engine/Source/Runtime/Core/Public/Containers/Array.h b/Engine/Source/Runtime/Core/Public/Containers/Array.h
--- a/Engine/Source/Runtime/Core/Public/Containers/Array.h
+++ b/Engine/Source/Runtime/Core/Public/Containers/Array.h
@@ -8,6 +8,7 @@
 #include "HAL/UnrealMemory.h"
 #include "Templates/UnrealTypeTraits.h"
 #include "Templates/UnrealTemplate.h"
+#include "Containers/AllowShrinking.h"
 #include "Containers/ContainerAllocationPolicies.h"
 #include "Containers/ContainerElementTypeCompatibility.h"
 #include "Serialization/Archive.h"
@@ -829,16 +830,21 @@ public:
 	/**
 	 * Pops element from the array.
 	 *
-	 * @param bAllowShrinking If this call allows shrinking of the array during element remove.
+	 * @param AllowShrinking If this call allows shrinking of the array during element remove.
 	 * @returns Popped element.
 	 */
-	FORCEINLINE ElementType Pop(bool bAllowShrinking = true)
+	ElementType Pop(EAllowShrinking AllowShrinking = EAllowShrinking::Yes)
 	{
 		RangeCheck(0);
 		ElementType Result = MoveTempIfPossible(GetData()[ArrayNum - 1]);
-		RemoveAt(ArrayNum - 1, 1, bAllowShrinking);
+		RemoveAt(ArrayNum - 1, 1, AllowShrinking);
 		return Result;
 	}
+	UE_ALLOWSHRINKING_BOOL_DEPRECATED("Pop")
+	FORCEINLINE ElementType Pop(bool bAllowShrinking)
+	{
+		return Pop(bAllowShrinking ? EAllowShrinking::Yes : EAllowShrinking::No);
+	}
 
 	/**
 	 * Pushes element into the array.
@@ -1684,7 +1690,7 @@ public:
 	}
 
 private:
-	void RemoveAtImpl(SizeType Index, SizeType Count, bool bAllowShrinking)
+	void RemoveAtImpl(SizeType Index, SizeType Count, EAllowShrinking AllowShrinking)
 	{
 		if (Count)
 		{
@@ -1705,7 +1711,7 @@ private:
 
 			SlackTrackerNumChanged();
 
-			if (bAllowShrinking)
+			if (AllowShrinking == EAllowShrinking::Yes)
 			{
 				ResizeShrink();
 			}
@@ -1721,7 +1727,7 @@ public:
 	 */
 	FORCEINLINE void RemoveAt(SizeType Index)
 	{
-		RemoveAtImpl(Index, 1, true);
+		RemoveAtImpl(Index, 1, EAllowShrinking::Yes);
 	}
 
 	/**
@@ -1730,17 +1736,23 @@ public:
 	 *
 	 * @param Index Location in array of the element to remove.
 	 * @param Count (Optional) Number of elements to remove. Default is 1.
-	 * @param bAllowShrinking (Optional) Tells if this call can shrink array if suitable after remove. Default is true.
+	 * @param AllowShrinking (Optional) Tells if this call can shrink array if suitable after remove. Default is yes.
 	 */
 	template <typename CountType>
-	FORCEINLINE void RemoveAt(SizeType Index, CountType Count, bool bAllowShrinking = true)
+	FORCEINLINE void RemoveAt(SizeType Index, CountType Count, EAllowShrinking AllowShrinking = EAllowShrinking::Yes)
 	{
 		static_assert(!std::is_same_v<CountType, bool>, "TArray::RemoveAt: unexpected bool passed as the Count argument");
-		RemoveAtImpl(Index, (SizeType)Count, bAllowShrinking);
+		RemoveAtImpl(Index, (SizeType)Count, AllowShrinking);
+	}
+	template <typename CountType>
+	UE_ALLOWSHRINKING_BOOL_DEPRECATED("RemoveAt")
+	FORCEINLINE void RemoveAt(SizeType Index, CountType Count, bool bAllowShrinking)
+	{
+		RemoveAt(Index, Count, bAllowShrinking ? EAllowShrinking::Yes : EAllowShrinking::No);
 	}
 
 private:
-	void RemoveAtSwapImpl(SizeType Index, SizeType Count = 1, bool bAllowShrinking = true)
+	void RemoveAtSwapImpl(SizeType Index, SizeType Count, EAllowShrinking AllowShrinking)
 	{
 		if (Count)
 		{
@@ -1763,7 +1775,7 @@ private:
 
 			SlackTrackerNumChanged();
 
-			if (bAllowShrinking)
+			if (AllowShrinking == EAllowShrinking::Yes)
 			{
 				ResizeShrink();
 			}
@@ -1782,7 +1794,7 @@ public:
 	 */
 	FORCEINLINE void RemoveAtSwap(SizeType Index)
 	{
-		RemoveAtSwapImpl(Index, 1, true);
+		RemoveAtSwapImpl(Index, 1, EAllowShrinking::Yes);
 	}
 
 	/**
@@ -1794,14 +1806,20 @@ public:
 	 *
 	 * @param Index Location in array of the element to remove.
 	 * @param Count (Optional) Number of elements to remove. Default is 1.
-	 * @param bAllowShrinking (Optional) Tells if this call can shrink array if
-	 *                        suitable after remove. Default is true.
+	 * @param AllowShrinking (Optional) Tells if this call can shrink array if
+	 *                        suitable after remove. Default is yes.
 	 */
 	template <typename CountType>
-	FORCEINLINE void RemoveAtSwap(SizeType Index, CountType Count, bool bAllowShrinking = true)
+	FORCEINLINE void RemoveAtSwap(SizeType Index, CountType Count, EAllowShrinking AllowShrinking = EAllowShrinking::Yes)
 	{
 		static_assert(!std::is_same_v<CountType, bool>, "TArray::RemoveAtSwap: unexpected bool passed as the Count argument");
-		RemoveAtSwapImpl(Index, Count, bAllowShrinking);
+		RemoveAtSwapImpl(Index, Count, AllowShrinking);
+	}
+	template <typename CountType>
+	UE_ALLOWSHRINKING_BOOL_DEPRECATED("RemoveAtSwap")
+	FORCEINLINE void RemoveAtSwap(SizeType Index, CountType Count, bool bAllowShrinking)
+	{
+		RemoveAtSwapImpl(Index, Count, bAllowShrinking ? EAllowShrinking::Yes : EAllowShrinking::No);
 	}
 
 	/**
@@ -1862,9 +1880,9 @@ public:
 	 * Resizes array to given number of elements.
 	 *
 	 * @param NewNum New size of the array.
-	 * @param bAllowShrinking Tell if this function can shrink the memory in-use if suitable.
+	 * @param AllowShrinking Tell if this function can shrink the memory in-use if suitable.
 	 */
-	void SetNum(SizeType NewNum, bool bAllowShrinking = true)
+	void SetNum(SizeType NewNum, EAllowShrinking AllowShrinking = EAllowShrinking::Yes)
 	{
 		if (NewNum > Num())
 		{
@@ -1879,18 +1897,23 @@ public:
 		}
 		else if (NewNum < Num())
 		{
-			RemoveAt(NewNum, Num() - NewNum, bAllowShrinking);
+			RemoveAt(NewNum, Num() - NewNum, AllowShrinking);
 		}
 	}
+	UE_ALLOWSHRINKING_BOOL_DEPRECATED("SetNum")
+	FORCEINLINE void SetNum(SizeType NewNum, bool bAllowShrinking)
+	{
+		SetNum(NewNum, bAllowShrinking ? EAllowShrinking::Yes : EAllowShrinking::No);
+	}
 
 	/**
 	 * Resizes array to given number of elements, optionally shrinking it.
 	 * New elements will be zeroed.
 	 *
 	 * @param NewNum New size of the array.
-	 * @param bAllowShrinking Tell if this function can shrink the memory in-use if suitable.
+	 * @param AllowShrinking Tell if this function can shrink the memory in-use if suitable.
 	 */
-	void SetNumZeroed(SizeType NewNum, bool bAllowShrinking = true)
+	void SetNumZeroed(SizeType NewNum, EAllowShrinking AllowShrinking = EAllowShrinking::Yes)
 	{
 		if (NewNum > Num())
 		{
@@ -1903,16 +1926,22 @@ public:
 		}
 		else if (NewNum < Num())
 		{
-			RemoveAt(NewNum, Num() - NewNum, bAllowShrinking);
+			RemoveAt(NewNum, Num() - NewNum, AllowShrinking);
 		}
 	}
+	UE_ALLOWSHRINKING_BOOL_DEPRECATED("SetNumZeroed")
+	FORCEINLINE void SetNumZeroed(SizeType NewNum, bool bAllowShrinking)
+	{
+		SetNumZeroed(NewNum, bAllowShrinking ? EAllowShrinking::Yes : EAllowShrinking::No);
+	}
 
 	/**
 	 * Resizes array to given number of elements. New elements will be uninitialized.
 	 *
 	 * @param NewNum New size of the array.
+	 * @param AllowShrinking Tell if this function can shrink the memory in-use if suitable.
 	 */
-	void SetNumUninitialized(SizeType NewNum, bool bAllowShrinking = true)
+	void SetNumUninitialized(SizeType NewNum, EAllowShrinking AllowShrinking = EAllowShrinking::Yes)
 	{
 		if (NewNum > Num())
 		{
@@ -1925,9 +1954,14 @@ public:
 		}
 		else if (NewNum < Num())
 		{
-			RemoveAt(NewNum, Num() - NewNum, bAllowShrinking);
+			RemoveAt(NewNum, Num() - NewNum, AllowShrinking);
 		}
 	}
+	UE_ALLOWSHRINKING_BOOL_DEPRECATED("SetNumUninitialized")
+	FORCEINLINE void SetNumUninitialized(SizeType NewNum, bool bAllowShrinking)
+	{
+		SetNumUninitialized(NewNum, bAllowShrinking ? EAllowShrinking::Yes : EAllowShrinking::No);
+	}
 
 	/**
 	 * Does nothing except setting the new number of elements in the array. Does not destruct items, does not de-allocate memory.
@@ -2637,11 +2671,11 @@ public:
 	 * Remove all instances that match the predicate
 	 *
 	 * @param Predicate Predicate class instance
-	 * @param bAllowShrinking Tell if this function can shrink the memory in-use if suitable.
+	 * @param AllowShrinking Tell if this function can shrink the memory in-use if suitable.
 	 * @see Remove, RemoveSingle, RemoveSingleSwap, RemoveSwap
 	 */
 	template <class PREDICATE_CLASS>
-	SizeType RemoveAllSwap(const PREDICATE_CLASS& Predicate, bool bAllowShrinking = true)
+	SizeType RemoveAllSwap(const PREDICATE_CLASS& Predicate, EAllowShrinking AllowShrinking = EAllowShrinking::Yes)
 	{
 		bool bRemoved = false;
 		const SizeType OriginalNum = ArrayNum;
@@ -2658,25 +2692,31 @@ public:
 			}
 		}
 
-		if (bRemoved && bAllowShrinking)
+		if (bRemoved && AllowShrinking == EAllowShrinking::Yes)
 		{
 			ResizeShrink();
 		}
 
 		return OriginalNum - ArrayNum;
 	}
+	template <class PREDICATE_CLASS>
+	UE_ALLOWSHRINKING_BOOL_DEPRECATED("RemoveAllSwap")
+	FORCEINLINE SizeType RemoveAllSwap(const PREDICATE_CLASS& Predicate, bool bAllowShrinking)
+	{
+		return RemoveAllSwap(Predicate, bAllowShrinking ? EAllowShrinking::Yes : EAllowShrinking::No);
+	}
 
 	/**
 	 * Removes the first occurrence of the specified item in the array. This version is much more efficient
 	 * O(Count) instead of O(ArrayNum), but does not preserve the order
 	 *
 	 * @param Item The item to remove
-	 * @param bAllowShrinking Tell if this function can shrink the memory in-use if suitable.
+	 * @param AllowShrinking Tell if this function can shrink the memory in-use if suitable.
 	 *
 	 * @returns The number of items removed. For RemoveSingleItem, this is always either 0 or 1.
 	 * @see Add, Insert, Remove, RemoveAll, RemoveAllSwap, RemoveSwap
 	 */
-	SizeType RemoveSingleSwap(const ElementType& Item, bool bAllowShrinking = true)
+	SizeType RemoveSingleSwap(const ElementType& Item, EAllowShrinking AllowShrinking = EAllowShrinking::Yes)
 	{
 		SizeType Index = Find(Item);
 		if (Index == INDEX_NONE)
@@ -2684,23 +2724,31 @@ public:
 			return 0;
 		}
 
-		RemoveAtSwap(Index, 1, bAllowShrinking);
+		RemoveAtSwap(Index, 1, AllowShrinking);
 
 		// Removed one item
 		return 1;
 	}
+	UE_ALLOWSHRINKING_BOOL_DEPRECATED("RemoveSingleSwap")
+	FORCEINLINE SizeType RemoveSingleSwap(const ElementType& Item, bool bAllowShrinking)
+	{
+		return RemoveSingleSwap(Item, bAllowShrinking ? EAllowShrinking::Yes : EAllowShrinking::No);
+	}
 
 	/**
-	 * Removes item from the array.
+	 * Removes all instances of a given item from the array.
 	 *
 	 * This version is much more efficient, because it uses RemoveAtSwap
 	 * internally which is O(Count) instead of RemoveAt which is O(ArrayNum),
 	 * but does not preserve the order.
 	 *
+	 * @param Item The item to remove
+	 * @param AllowShrinking Tell if this function can shrink the memory in-use if suitable.
+	 *
 	 * @returns Number of elements removed.
 	 * @see Add, Insert, Remove, RemoveAll, RemoveAllSwap
 	 */
-	SizeType RemoveSwap(const ElementType& Item, bool bAllowShrinking = true)
+	SizeType RemoveSwap(const ElementType& Item, EAllowShrinking AllowShrinking = EAllowShrinking::Yes)
 	{
 		CheckAddress(&Item);
 
@@ -2715,13 +2763,18 @@ public:
 			}
 		}
 
-		if (bRemoved && bAllowShrinking)
+		if (bRemoved && AllowShrinking == EAllowShrinking::Yes)
 		{
 			ResizeShrink();
 		}
 
 		return OriginalNum - ArrayNum;
 	}
+	UE_ALLOWSHRINKING_BOOL_DEPRECATED("RemoveSwap")
+	FORCEINLINE SizeType RemoveSwap(const ElementType& Item, bool bAllowShrinking)
+	{
+		return RemoveSwap(Item, bAllowShrinking ? EAllowShrinking::Yes : EAllowShrinking::No);
+	}
 
 	/**
 	 * Element-wise array memory swap.
@@ -3271,35 +3324,47 @@ public:
 	 *
 	 * @param OutItem The removed item.
 	 * @param Predicate Predicate class instance.
+	 * @param AllowShrinking Tell if this function can shrink the memory in-use if suitable.
 	 *
 	 * @note: If your array contains raw pointers, they will be automatically dereferenced during heapification.
 	 *        Therefore, your predicate will be passed references rather than pointers.
 	 *        The auto-dereferencing behavior does not occur with smart pointers.
 	 */
 	template <class PREDICATE_CLASS>
-	void HeapPop(ElementType& OutItem, const PREDICATE_CLASS& Predicate, bool bAllowShrinking = true)
+	void HeapPop(ElementType& OutItem, const PREDICATE_CLASS& Predicate, EAllowShrinking AllowShrinking = EAllowShrinking::Yes)
 	{
 		OutItem = MoveTemp((*this)[0]);
-		RemoveAtSwap(0, 1, bAllowShrinking);
+		RemoveAtSwap(0, 1, AllowShrinking);
 
 		TDereferenceWrapper< ElementType, PREDICATE_CLASS> PredicateWrapper(Predicate);
 		AlgoImpl::HeapSiftDown(GetData(), (SizeType)0, Num(), FIdentityFunctor(), PredicateWrapper);
 	}
+	template <class PREDICATE_CLASS>
+	UE_ALLOWSHRINKING_BOOL_DEPRECATED("HeapPop")
+	FORCEINLINE void HeapPop(ElementType& OutItem, const PREDICATE_CLASS& Predicate, bool bAllowShrinking)
+	{
+		HeapPop(OutItem, Predicate, bAllowShrinking ? EAllowShrinking::Yes : EAllowShrinking::No);
+	}
 
 	/** 
 	 * Removes the top element from the heap. Assumes < operator is defined for
 	 * the template type.
 	 *
 	 * @param OutItem The removed item.
-	 * @param bAllowShrinking (Optional) Tells if this call can shrink the array allocation if suitable after the pop. Default is true.
+	 * @param AllowShrinking (Optional) Tells if this call can shrink the array allocation if suitable after the pop. Default is yes.
 	 *
 	 * @note: If your array contains raw pointers, they will be automatically dereferenced during heapification.
 	 *        Therefore, your array will be heapified by the values being pointed to, rather than the pointers' values.
 	 *        The auto-dereferencing behavior does not occur with smart pointers.
 	 */
-	void HeapPop(ElementType& OutItem, bool bAllowShrinking = true)
+	void HeapPop(ElementType& OutItem, EAllowShrinking AllowShrinking = EAllowShrinking::Yes)
 	{
-		HeapPop(OutItem, TLess<ElementType>(), bAllowShrinking);
+		HeapPop(OutItem, TLess<ElementType>(), AllowShrinking);
+	}
+	UE_ALLOWSHRINKING_BOOL_DEPRECATED("HeapPop")
+	FORCEINLINE void HeapPop(ElementType& OutItem, bool bAllowShrinking)
+	{
+		HeapPop(OutItem, bAllowShrinking ? EAllowShrinking::Yes : EAllowShrinking::No);
 	}
 
 	/**
@@ -3317,33 +3382,44 @@ public:
 	 * Removes the top element from the heap.
 	 *
 	 * @param Predicate Predicate class instance.
-	 * @param bAllowShrinking (Optional) Tells if this call can shrink the array allocation if suitable after the discard. Default is true.
+	 * @param AllowShrinking (Optional) Tells if this call can shrink the array allocation if suitable after the discard. Default is yes.
 	 *
 	 * @note: If your array contains raw pointers, they will be automatically dereferenced during heapification.
 	 *        Therefore, your predicate will be passed references rather than pointers.
 	 *        The auto-dereferencing behavior does not occur with smart pointers.
 	 */
 	template <class PREDICATE_CLASS>
-	void HeapPopDiscard(const PREDICATE_CLASS& Predicate, bool bAllowShrinking = true)
+	void HeapPopDiscard(const PREDICATE_CLASS& Predicate, EAllowShrinking AllowShrinking = EAllowShrinking::Yes)
 	{
-		RemoveAtSwap(0, 1, bAllowShrinking);
+		RemoveAtSwap(0, 1, AllowShrinking);
 		TDereferenceWrapper< ElementType, PREDICATE_CLASS> PredicateWrapper(Predicate);
 		AlgoImpl::HeapSiftDown(GetData(), (SizeType)0, Num(), FIdentityFunctor(), PredicateWrapper);
 	}
+	template <class PREDICATE_CLASS>
+	UE_ALLOWSHRINKING_BOOL_DEPRECATED("HeapPopDiscard")
+	FORCEINLINE void HeapPopDiscard(const PREDICATE_CLASS& Predicate, bool bAllowShrinking)
+	{
+		HeapPopDiscard(Predicate, bAllowShrinking ? EAllowShrinking::Yes : EAllowShrinking::No);
+	}
 
 	/** 
 	 * Removes the top element from the heap. Assumes < operator is defined for the template type.
 	 *
-	 * @param bAllowShrinking (Optional) Tells if this call can shrink the array
-	 *		allocation if suitable after the discard. Default is true.
+	 * @param AllowShrinking (Optional) Tells if this call can shrink the array
+	 *		allocation if suitable after the discard. Default is yes.
 	 *
 	 * @note: If your array contains raw pointers, they will be automatically dereferenced during heapification.
 	 *        Therefore, your array will be heapified by the values being pointed to, rather than the pointers' values.
 	 *        The auto-dereferencing behavior does not occur with smart pointers.
 	 */
-	void HeapPopDiscard(bool bAllowShrinking = true)
+	void HeapPopDiscard(EAllowShrinking AllowShrinking = EAllowShrinking::Yes)
 	{
-		HeapPopDiscard(TLess<ElementType>(), bAllowShrinking);
+		HeapPopDiscard(TLess<ElementType>(), AllowShrinking);
+	}
+	UE_ALLOWSHRINKING_BOOL_DEPRECATED("HeapPopDiscard")
+	FORCEINLINE void HeapPopDiscard(bool bAllowShrinking)
+	{
+		HeapPopDiscard(TLess<ElementType>(), bAllowShrinking ? EAllowShrinking::Yes : EAllowShrinking::No);
 	}
 
 	/** 
@@ -3373,37 +3449,48 @@ public:
 	 *
 	 * @param Index Position at which to remove item.
 	 * @param Predicate Predicate class instance.
-	 * @param bAllowShrinking (Optional) Tells if this call can shrink the array allocation
-	 *		if suitable after the remove (default = true).
+	 * @param AllowShrinking (Optional) Tells if this call can shrink the array allocation
+	 *		if suitable after the remove (default = yes).
 	 *
 	 * @note: If your array contains raw pointers, they will be automatically dereferenced during heapification.
 	 *        Therefore, your predicate will be passed references rather than pointers.
 	 *        The auto-dereferencing behavior does not occur with smart pointers.
 	 */
 	template <class PREDICATE_CLASS>
-	void HeapRemoveAt(SizeType Index, const PREDICATE_CLASS& Predicate, bool bAllowShrinking = true)
+	void HeapRemoveAt(SizeType Index, const PREDICATE_CLASS& Predicate, EAllowShrinking AllowShrinking = EAllowShrinking::Yes)
 	{
-		RemoveAtSwap(Index, 1, bAllowShrinking);
+		RemoveAtSwap(Index, 1, AllowShrinking);
 
 		TDereferenceWrapper< ElementType, PREDICATE_CLASS> PredicateWrapper(Predicate);
 		AlgoImpl::HeapSiftDown(GetData(), Index, Num(), FIdentityFunctor(), PredicateWrapper);
 		AlgoImpl::HeapSiftUp(GetData(), (SizeType)0, FPlatformMath::Min(Index, Num() - 1), FIdentityFunctor(), PredicateWrapper);
 	}
+	template <class PREDICATE_CLASS>
+	UE_ALLOWSHRINKING_BOOL_DEPRECATED("HeapRemoveAt")
+	void HeapRemoveAt(SizeType Index, const PREDICATE_CLASS& Predicate, bool bAllowShrinking)
+	{
+		HeapRemoveAt(Index, Predicate, bAllowShrinking ? EAllowShrinking::Yes : EAllowShrinking::No);
+	}
 
 	/**
 	 * Removes an element from the heap. Assumes < operator is defined for the template type.
 	 *
 	 * @param Index Position at which to remove item.
-	 * @param bAllowShrinking (Optional) Tells if this call can shrink the array allocation
-	 *		if suitable after the remove (default = true).
+	 * @param AllowShrinking (Optional) Tells if this call can shrink the array allocation
+	 *		if suitable after the remove (default = yes).
 	 *
 	 * @note: If your array contains raw pointers, they will be automatically dereferenced during heapification.
 	 *        Therefore, your array will be heapified by the values being pointed to, rather than the pointers' values.
 	 *        The auto-dereferencing behavior does not occur with smart pointers.
 	 */
-	void HeapRemoveAt(SizeType Index, bool bAllowShrinking = true)
+	void HeapRemoveAt(SizeType Index, EAllowShrinking AllowShrinking = EAllowShrinking::Yes)
+	{
+		HeapRemoveAt(Index, TLess< ElementType >(), AllowShrinking);
+	}
+	UE_ALLOWSHRINKING_BOOL_DEPRECATED("HeapRemoveAt")
+	FORCEINLINE void HeapRemoveAt(SizeType Index, bool bAllowShrinking)
 	{
-		HeapRemoveAt(Index, TLess< ElementType >(), bAllowShrinking);
+		HeapRemoveAt(Index, bAllowShrinking ? EAllowShrinking::Yes : EAllowShrinking::No);
 	}
 
 	/**
diff --git a/Engine/Source/Runtime/Core/Public/Containers/IndirectArray.h b/Engine/Source/Runtime/Core/Public/Containers/IndirectArray.h
--- a/Engine/Source/Runtime/Core/Public/Containers/IndirectArray.h
+++ b/Engine/Source/Runtime/Core/Public/Containers/IndirectArray.h
@@ -248,10 +248,10 @@ public:
 	 *
 	 * @param Index Location in array of the element to remove.
 	 * @param Count (Optional) Number of elements to remove. Default is 1.
-	 * @param bAllowShrinking (Optional) Tells if this call can shrink array if
-	 *                        suitable after remove. Default is true.
+	 * @param AllowShrinking (Optional) Tells if this call can shrink array if
+	 *                        suitable after remove. Default is yes.
 	 */
-	void RemoveAt(int32 Index, int32 Count = 1, bool bAllowShrinking = true)
+	void RemoveAt(int32 Index, int32 Count = 1, EAllowShrinking AllowShrinking = EAllowShrinking::Yes)
 	{
 		check(Index >= 0);
 		check(Index <= Array.Num());
@@ -262,7 +262,12 @@ public:
 			delete *Element;
 			++Element;
 		}
-		Array.RemoveAt(Index, Count, bAllowShrinking);
+		Array.RemoveAt(Index, Count, AllowShrinking);
+	}
+	UE_ALLOWSHRINKING_BOOL_DEPRECATED("RemoveAt")
+	FORCEINLINE void RemoveAt(int32 Index, int32 Count, bool bAllowShrinking)
+	{
+		RemoveAt(Index, Count, bAllowShrinking ? EAllowShrinking::Yes : EAllowShrinking::No);
 	}
 
 	/**
@@ -274,10 +279,10 @@ public:
 	 *
 	 * @param Index Location in array of the element to remove.
 	 * @param Count (Optional) Number of elements to remove. Default is 1.
-	 * @param bAllowShrinking (Optional) Tells if this call can shrink array if
-	 *                        suitable after remove. Default is true.
+	 * @param AllowShrinking (Optional) Tells if this call can shrink array if
+	 *                        suitable after remove. Default is yes.
 	 */
-	void RemoveAtSwap(int32 Index, int32 Count = 1, bool bAllowShrinking = true)
+	void RemoveAtSwap(int32 Index, int32 Count = 1, EAllowShrinking AllowShrinking = EAllowShrinking::Yes)
 	{
 		check(Index >= 0);
 		check(Index <= Array.Num());
@@ -288,7 +293,12 @@ public:
 			delete *Element;
 			++Element;
 		}
-		Array.RemoveAtSwap(Index, Count, bAllowShrinking);
+		Array.RemoveAtSwap(Index, Count, AllowShrinking);
+	}
+	UE_ALLOWSHRINKING_BOOL_DEPRECATED("RemoveAtSwap")
+	FORCEINLINE void RemoveAtSwap(int32 Index, int32 Count, bool bAllowShrinking)
+	{
+		RemoveAtSwap(Index, Count, bAllowShrinking ? EAllowShrinking::Yes : EAllowShrinking::No);
 	}
 
 	/**
diff --git a/Engine/Source/Runtime/Core/Public/Containers/PagedArray.h b/Engine/Source/Runtime/Core/Public/Containers/PagedArray.h
--- a/Engine/Source/Runtime/Core/Public/Containers/PagedArray.h
+++ b/Engine/Source/Runtime/Core/Public/Containers/PagedArray.h
@@ -321,7 +321,7 @@ public:
 	 * Resizes array to the parameter number of elements.
 	 * The allow shrinking parameter indicates whether the container page allocation can be reduced if possible.
 	 */
-	void SetNum(SizeType NewNum, bool bAllowShrinking = true)
+	void SetNum(SizeType NewNum, EAllowShrinking AllowShrinking = EAllowShrinking::Yes)
 	{
 		const SizeType RequiredPageCount = NumRequiredPages(NewNum);
 		if (NewNum > Num())
@@ -343,7 +343,7 @@ public:
 		else if (NewNum < Num())
 		{
 			SizeType PendingCount = Num() - NewNum;
-			Pages.SetNum(RequiredPageCount, bAllowShrinking);
+			Pages.SetNum(RequiredPageCount, AllowShrinking);
 			if (const SizeType Mod = NewNum % PageTraits::Capacity)
 			{
 				Pages.Last().SetNum(Mod, false);
@@ -351,6 +351,11 @@ public:
 		}
 		Count = NewNum;
 	}
+	UE_ALLOWSHRINKING_BOOL_DEPRECATED("SetNum")
+	FORCEINLINE void SetNum(SizeType NewNum, bool bAllowShrinking)
+	{
+		SetNum(NewNum, bAllowShrinking ? EAllowShrinking::Yes : EAllowShrinking::No);
+	}
 
 	/*
 	 * Fills the memory used by the container elements with the parameter byte value.
@@ -518,26 +523,31 @@ public:
 	/*
 	 * Removes the last element in the container.
 	 */
-	void Pop(bool bAllowShrinking = true)
+	void Pop(EAllowShrinking AllowShrinking = EAllowShrinking::Yes)
 	{
 		CheckValidIndex(0);
 		const SizeType LastIndex = Num() - 1;
 		const SizeType LastPageIndex = GetPageIndex(LastIndex);
 		const SizeType LastIndexInPage = GetPageOffset(LastIndex);
 		Pages[LastPageIndex].RemoveAt(LastIndexInPage, 1, false);
-		if (bAllowShrinking && LastIndexInPage == 0)
+		if (AllowShrinking == EAllowShrinking::Yes && LastIndexInPage == 0)
 		{
 			Pages.SetNum(LastPageIndex);
 		}
 		--Count;
 	}
+	UE_ALLOWSHRINKING_BOOL_DEPRECATED("Pop")
+	FORCEINLINE void Pop(bool bAllowShrinking)
+	{
+		Pop(bAllowShrinking ? EAllowShrinking::Yes : EAllowShrinking::No);
+	}
 
 	/**
 	 * Removes the element at the parameter index position and swaps the last element if existent to the same position
 	 * to ensure the range is contiguous. This method provides efficient removal O(1) but it doesn't preserve the
 	 * insertion order.
 	 */
-	void RemoveAtSwap(SizeType Index, bool bAllowShrinking = true)
+	void RemoveAtSwap(SizeType Index, EAllowShrinking AllowShrinking = EAllowShrinking::Yes)
 	{
 		CheckValidIndex(Index);
 		const SizeType TargetPageIndex = GetPageIndex(Index);
@@ -547,7 +557,7 @@ public:
 		if (TargetPageIndex == LastPageIndex)
 		{
 			Pages[TargetPageIndex].RemoveAtSwap(TargetIndexInPage, 1, false);
-			if (bAllowShrinking && Pages[TargetPageIndex].IsEmpty())
+			if (AllowShrinking == EAllowShrinking::Yes && Pages[TargetPageIndex].IsEmpty())
 			{
 				Pages.SetNum(TargetPageIndex);
 			}
@@ -560,6 +570,11 @@ public:
 		}
 		--Count;
 	}
+	UE_ALLOWSHRINKING_BOOL_DEPRECATED("RemoveAtSwap")
+	FORCEINLINE void RemoveAtSwap(SizeType Index, bool bAllowShrinking)
+	{
+		RemoveAtSwap(Index, bAllowShrinking ? EAllowShrinking::Yes : EAllowShrinking::No);
+	}
 
 	/*
 	 * Empties the container effectively destroying all contained elements and releases any acquired storage.
diff --git a/Engine/Source/Runtime/Core/Public/Containers/ScriptArray.h b/Engine/Source/Runtime/Core/Public/Containers/ScriptArray.h
--- a/Engine/Source/Runtime/Core/Public/Containers/ScriptArray.h
+++ b/Engine/Source/Runtime/Core/Public/Containers/ScriptArray.h
@@ -5,6 +5,7 @@
 #include "CoreTypes.h"
 #include "Misc/AssertionMacros.h"
 #include "HAL/UnrealMemory.h"
+#include "Containers/AllowShrinking.h"
 #include "Containers/ContainerAllocationPolicies.h"
 #include "Containers/Array.h"
 #include <initializer_list>
@@ -100,7 +101,7 @@ public:
 			ResizeTo(ArrayNum, NumBytesPerElement, AlignmentOfElement);
 		}
 	}
-	void SetNumUninitialized(int32 NewNum, int32 NumBytesPerElement, uint32 AlignmentOfElement, bool bAllowShrinking = true)
+	void SetNumUninitialized(int32 NewNum, int32 NumBytesPerElement, uint32 AlignmentOfElement, EAllowShrinking AllowShrinking = EAllowShrinking::Yes)
 	{
 		checkSlow(NewNum >= 0);
 		int32 OldNum = Num();
@@ -110,9 +111,14 @@ public:
 		}
 		else if (NewNum < OldNum)
 		{
-			Remove(NewNum, OldNum - NewNum, NumBytesPerElement, AlignmentOfElement, bAllowShrinking);
+			Remove(NewNum, OldNum - NewNum, NumBytesPerElement, AlignmentOfElement, AllowShrinking);
 		}
 	}
+	UE_ALLOWSHRINKING_BOOL_DEPRECATED("SetNumUninitialized")
+	FORCEINLINE void SetNumUninitialized(int32 NewNum, int32 NumBytesPerElement, uint32 AlignmentOfElement, bool bAllowShrinking)
+	{
+		SetNumUninitialized(NewNum, NumBytesPerElement, AlignmentOfElement, bAllowShrinking ? EAllowShrinking::Yes : EAllowShrinking::No);
+	}
 	void MoveAssign(TScriptArray& Other, int32 NumBytesPerElement, uint32 AlignmentOfElement)
 	{
 		checkSlow(this != &Other);
@@ -178,7 +184,7 @@ public:
 		return ArrayMax - ArrayNum;
 	}
 
-	void Remove( int32 Index, int32 Count, int32 NumBytesPerElement, uint32 AlignmentOfElement, bool bAllowShrinking = true )
+	void Remove( int32 Index, int32 Count, int32 NumBytesPerElement, uint32 AlignmentOfElement, EAllowShrinking AllowShrinking = EAllowShrinking::Yes )
 	{
 		if (Count)
 		{
@@ -202,7 +208,7 @@ public:
 
 			SlackTrackerNumChanged();
 
-			if (bAllowShrinking)
+			if (AllowShrinking == EAllowShrinking::Yes)
 			{
 				ResizeShrink(NumBytesPerElement, AlignmentOfElement);
 			}
@@ -210,6 +216,11 @@ public:
 			checkSlow(ArrayMax >= ArrayNum);
 		}
 	}
+	UE_ALLOWSHRINKING_BOOL_DEPRECATED("Remove")
+	FORCEINLINE void Remove(int32 Index, int32 Count, int32 NumBytesPerElement, uint32 AlignmentOfElement, bool bAllowShrinking)
+	{
+		Remove(Index, Count, NumBytesPerElement, AlignmentOfElement, bAllowShrinking ? EAllowShrinking::Yes : EAllowShrinking::No);
+	}
 	SIZE_T GetAllocatedSize(int32 NumBytesPerElement) const
 	{
 		return ((const typename AllocatorType::ForAnyElementType*)this)->GetAllocatedSize(ArrayMax, NumBytesPerElement);
diff --git a/Engine/Source/Runtime/Core/Public/Containers/Set.h b/Engine/Source/Runtime/Core/Public/Containers/Set.h
--- a/Engine/Source/Runtime/Core/Public/Containers/Set.h
+++ b/Engine/Source/Runtime/Core/Public/Containers/Set.h
@@ -440,7 +440,7 @@ public:
 	{
 		// Empty the elements array, and reallocate it for the expected number of elements.
 		const int32 DesiredHashSize = Allocator::GetNumberOfHashBuckets(ExpectedNumElements);
-		const bool ShouldDoRehash = ShouldRehash(ExpectedNumElements,DesiredHashSize,true);
+		const bool ShouldDoRehash = ShouldRehash(ExpectedNumElements, DesiredHashSize, EAllowShrinking::Yes);
 
 		if (!ShouldDoRehash)
 		{
@@ -530,7 +530,7 @@ public:
 	/** Relaxes the set's hash to a size strictly bounded by the number of elements in the set. */
 	FORCEINLINE void Relax()
 	{
-		ConditionalRehash(Elements.Num(),true);
+		ConditionalRehash(Elements.Num(), EAllowShrinking::Yes);
 	}
 
 	/** 
@@ -718,7 +718,7 @@ private:
 	FORCEINLINE void RehashOrLink(uint32 KeyHash, SetElementType& Element, SizeType ElementIndex)
 	{
 		// Check if the hash needs to be resized.
-		if (!ConditionalRehash(Elements.Num()))
+		if (!ConditionalRehash(Elements.Num(), EAllowShrinking::No))
 		{
 			// If the rehash didn't add the new element to the hash, add it.
 			LinkElement(ElementIndex, Element, KeyHash);
@@ -1505,28 +1505,28 @@ private:
 	 * Checks if the hash has an appropriate number of buckets, and if it should be resized.
 	 * @param NumHashedElements - The number of elements to size the hash for.
 	 * @param DesiredHashSize - Desired size if we should rehash.
-	 * @param bAllowShrinking - true if the hash is allowed to shrink.
+	 * @param AllowShrinking - If the hash is allowed to shrink.
 	 * @return true if the set should berehashed.
 	 */
-	FORCEINLINE bool ShouldRehash(int32 NumHashedElements,int32 DesiredHashSize,bool bAllowShrinking = false) const
+	FORCEINLINE bool ShouldRehash(int32 NumHashedElements, int32 DesiredHashSize, EAllowShrinking AllowShrinking) const
 	{
 		// If the hash hasn't been created yet, or is smaller than the desired hash size, rehash.
 		// If shrinking is allowed and the hash is bigger than the desired hash size, rehash.
-		return ((NumHashedElements > 0 && HashSize < DesiredHashSize) || (bAllowShrinking && HashSize > DesiredHashSize));
+		return ((NumHashedElements > 0 && HashSize < DesiredHashSize) || (AllowShrinking == EAllowShrinking::Yes && HashSize > DesiredHashSize));
 	}
 
 	/**
 	 * Checks if the hash has an appropriate number of buckets, and if not resizes it.
 	 * @param NumHashedElements - The number of elements to size the hash for.
-	 * @param bAllowShrinking - true if the hash is allowed to shrink.
+	 * @param AllowShrinking - If the hash is allowed to shrink.
 	 * @return true if the set was rehashed.
 	 */
-	bool ConditionalRehash(int32 NumHashedElements,bool bAllowShrinking = false) const
+	bool ConditionalRehash(int32 NumHashedElements, EAllowShrinking AllowShrinking) const
 	{
 		// Calculate the desired hash size for the specified number of elements.
 		const int32 DesiredHashSize = Allocator::GetNumberOfHashBuckets(NumHashedElements);
 
-		if (ShouldRehash(NumHashedElements, DesiredHashSize, bAllowShrinking))
+		if (ShouldRehash(NumHashedElements, DesiredHashSize, AllowShrinking))
 		{
 			HashSize = DesiredHashSize;
 			Rehash();
@@ -1644,7 +1644,7 @@ private:
 			, Index(INDEX_NONE)
 		{
 			// The set's hash needs to be initialized to find the elements with the specified key.
-			Set.ConditionalRehash(Set.Elements.Num());
+			Set.ConditionalRehash(Set.Elements.Num(), EAllowShrinking::No);
 			if (Set.HashSize)
 			{
 				NextIndex = Set.GetTypedHash(KeyFuncs::GetKeyHash(Key)).Index;
@@ -2204,7 +2204,7 @@ struct TSetPrivateFriend
 			Set.HashSize = 0;
 
 			// Hash the newly loaded elements.
-			Set.ConditionalRehash(Set.Elements.Num());
+			Set.ConditionalRehash(Set.Elements.Num(), EAllowShrinking::No);
 		}
 
 		return Ar;
@@ -2223,7 +2223,7 @@ struct TSetPrivateFriend
 			Set.HashSize = 0;
 
 			// Hash the newly loaded elements.
-			Set.ConditionalRehash(Set.Elements.Num());
+			Set.ConditionalRehash(Set.Elements.Num(), EAllowShrinking::No);
 		}
  	}
 
diff --git a/Engine/Source/Runtime/Core/Public/Containers/UnrealString.h.inl b/Engine/Source/Runtime/Core/Public/Containers/UnrealString.h.inl
--- a/Engine/Source/Runtime/Core/Public/Containers/UnrealString.h.inl
+++ b/Engine/Source/Runtime/Core/Public/Containers/UnrealString.h.inl
@@ -446,11 +446,16 @@ public:
 	/**
 	 * Removes characters within the string.
 	 *
-	 * @param Index           The index of the first character to remove.
-	 * @param Count           The number of characters to remove.
-	 * @param bAllowShrinking Whether or not to reallocate to shrink the storage after removal.
+	 * @param Index          The index of the first character to remove.
+	 * @param Count          The number of characters to remove.
+	 * @param AllowShrinking Whether or not to reallocate to shrink the storage after removal.
 	 */
-	CORE_API void RemoveAt(int32 Index, int32 Count = 1, bool bAllowShrinking = true);
+	CORE_API void RemoveAt(int32 Index, int32 Count = 1, EAllowShrinking AllowShrinking = EAllowShrinking::Yes);
+	UE_ALLOWSHRINKING_BOOL_DEPRECATED("RemoveAt")
+	FORCEINLINE void RemoveAt(int32 Index, int32 Count, bool bAllowShrinking)
+	{
+		RemoveAt(Index, Count, bAllowShrinking ? EAllowShrinking::Yes : EAllowShrinking::No);
+	}
 
 	/**
 	 * Removes the text from the start of the string if it exists.
@@ -1092,11 +1097,16 @@ public:
 	}
 
 	/** Modifies the string such that it is now the left most given number of characters */
-	FORCEINLINE void LeftInline(int32 Count, bool bAllowShrinking = true)
+	FORCEINLINE void LeftInline(int32 Count, EAllowShrinking AllowShrinking = EAllowShrinking::Yes)
 	{
 		const int32 Length = Len();
 		Count = FMath::Clamp(Count, 0, Length);
-		RemoveAt(Count, Length-Count, bAllowShrinking);
+		RemoveAt(Count, Length-Count, AllowShrinking);
+	}
+	UE_ALLOWSHRINKING_BOOL_DEPRECATED("LeftInline")
+	FORCEINLINE void LeftInline(int32 Count, bool bAllowShrinking)
+	{
+		LeftInline(Count, bAllowShrinking ? EAllowShrinking::Yes : EAllowShrinking::No);
 	}
 
 	/** Returns the left most characters from the string chopping the given number of characters from the end */
@@ -1113,10 +1123,15 @@ public:
 	}
 
 	/** Modifies the string such that it is now the left most characters chopping the given number of characters from the end */
-	FORCEINLINE void LeftChopInline(int32 Count, bool bAllowShrinking = true)
+	FORCEINLINE void LeftChopInline(int32 Count, EAllowShrinking AllowShrinking = EAllowShrinking::Yes)
 	{
 		const int32 Length = Len();
-		RemoveAt(FMath::Clamp(Length-Count, 0, Length), Count, bAllowShrinking);
+		RemoveAt(FMath::Clamp(Length-Count, 0, Length), Count, AllowShrinking);
+	}
+	UE_ALLOWSHRINKING_BOOL_DEPRECATED("LeftChopInline")
+	FORCEINLINE void LeftChopInline(int32 Count, bool bAllowShrinking)
+	{
+		LeftChopInline(Count, bAllowShrinking ? EAllowShrinking::Yes : EAllowShrinking::No);
 	}
 
 	/** Returns the string to the right of the specified location, counting back from the right (end of the word). */
@@ -1133,10 +1148,15 @@ public:
 	}
 
 	/** Modifies the string such that it is now the right most given number of characters */
-	FORCEINLINE void RightInline(int32 Count, bool bAllowShrinking = true)
+	FORCEINLINE void RightInline(int32 Count, EAllowShrinking AllowShrinking = EAllowShrinking::Yes)
 	{
 		const int32 Length = Len();
-		RemoveAt(0, Length-FMath::Clamp(Count,0,Length), bAllowShrinking);
+		RemoveAt(0, Length-FMath::Clamp(Count,0,Length), AllowShrinking);
+	}
+	UE_ALLOWSHRINKING_BOOL_DEPRECATED("RightInline")
+	FORCEINLINE void RightInline(int32 Count, bool bAllowShrinking)
+	{
+		RightInline(Count, bAllowShrinking ? EAllowShrinking::Yes : EAllowShrinking::No);
 	}
 
 	/** Returns the string to the right of the specified location, counting forward from the left (from the beginning of the word). */
@@ -1149,9 +1169,14 @@ public:
 	}
 
 	/** Modifies the string such that it is now the string to the right of the specified location, counting forward from the left (from the beginning of the word). */
-	FORCEINLINE void RightChopInline(int32 Count, bool bAllowShrinking = true)
+	FORCEINLINE void RightChopInline(int32 Count, EAllowShrinking AllowShrinking = EAllowShrinking::Yes)
+	{
+		RemoveAt(0, Count, AllowShrinking);
+	}
+	UE_ALLOWSHRINKING_BOOL_DEPRECATED("RightChopInline")
+	FORCEINLINE void RightChopInline(int32 Count, bool bAllowShrinking)
 	{
-		RemoveAt(0, Count, bAllowShrinking);
+		RightChopInline(Count, bAllowShrinking ? EAllowShrinking::Yes : EAllowShrinking::No);
 	}
 
 	/** Returns the substring from Start position for Count characters. */
@@ -1163,13 +1188,18 @@ public:
 	[[nodiscard]] FORCEINLINE UE_STRING_CLASS Mid(int32 Start) && { return ((UE_STRING_CLASS&&)*this).RightChop(Start); }
 
 	/** Modifies the string such that it is now the substring from Start position for Count characters. */
-	FORCEINLINE void MidInline(int32 Start, int32 Count = MAX_int32, bool bAllowShrinking = true)
+	FORCEINLINE void MidInline(int32 Start, int32 Count = MAX_int32, EAllowShrinking AllowShrinking = EAllowShrinking::Yes)
 	{
 		if (Count != MAX_int32 && int64(Start) + Count < MAX_int32)
 		{
 			LeftInline(Count + Start, false);
 		}
-		RightChopInline(Start, bAllowShrinking);
+		RightChopInline(Start, AllowShrinking);
+	}
+	UE_ALLOWSHRINKING_BOOL_DEPRECATED("MidInline")
+	FORCEINLINE void MidInline(int32 Start, int32 Count, bool bAllowShrinking)
+	{
+		MidInline(Start, Count, bAllowShrinking ? EAllowShrinking::Yes : EAllowShrinking::No);
 	}
 
 	/**
diff --git a/Engine/Source/Runtime/CoreUObject/Private/Serialization/AsyncLoading.cpp b/Engine/Source/Runtime/CoreUObject/Private/Serialization/AsyncLoading.cpp
--- a/Engine/Source/Runtime/CoreUObject/Private/Serialization/AsyncLoading.cpp
+++ b/Engine/Source/Runtime/CoreUObject/Private/Serialization/AsyncLoading.cpp
@@ -1711,7 +1711,7 @@ static int32 GetRandomSerialNumber(int32 MaxVal = MAX_int32)
 	return RandomStream.RandHelper(MaxVal);
 }
 
-void FImportOrImportIndexArray::HeapPop(int32& OutItem, bool bAllowShrinking)
+void FImportOrImportIndexArray::HeapPop(int32& OutItem, EAllowShrinking AllowShrinking)
 {
 	if (GRandomizeLoadOrder)
 	{
@@ -1720,7 +1720,7 @@ void FImportOrImportIndexArray::HeapPop(int32& OutItem, bool bAllowShrinking)
 		RemoveAt(Index, 1, false);
 		return;
 	}
-	TArray<int32>::HeapPop(OutItem, bAllowShrinking);
+	TArray<int32>::HeapPop(OutItem, AllowShrinking);
 }
 
 
diff --git a/Engine/Source/Runtime/CoreUObject/Private/Serialization/AsyncLoading.h b/Engine/Source/Runtime/CoreUObject/Private/Serialization/AsyncLoading.h
--- a/Engine/Source/Runtime/CoreUObject/Private/Serialization/AsyncLoading.h
+++ b/Engine/Source/Runtime/CoreUObject/Private/Serialization/AsyncLoading.h
@@ -329,7 +329,12 @@ struct FAsyncLoadEventArgs
 class FImportOrImportIndexArray : public TArray<int32>
 {
 public:
-	void HeapPop(int32& OutItem, bool bAllowShrinking = true);
+	void HeapPop(int32& OutItem, EAllowShrinking AllowShrinking = EAllowShrinking::Yes);
+	UE_ALLOWSHRINKING_BOOL_DEPRECATED("HeapPop")
+	FORCEINLINE void HeapPop(int32& OutItem, bool bAllowShrinking)
+	{
+		HeapPop(OutItem, bAllowShrinking ? EAllowShrinking::Yes : EAllowShrinking::No);
+	}
 };
 
 /**
