Remove timer during panic

FOutputDeviceRedirector::Panic is called during extraordinary cases like fatal errors. If that error is due to out of memory issues we should not trace timers which could try to allocate trace buffers.

diff --git a/Engine/Source/Runtime/Core/Private/Misc/OutputDeviceRedirector.cpp b/Engine/Source/Runtime/Core/Private/Misc/OutputDeviceRedirector.cpp
--- a/Engine/Source/Runtime/Core/Private/Misc/OutputDeviceRedirector.cpp
+++ b/Engine/Source/Runtime/Core/Private/Misc/OutputDeviceRedirector.cpp
@@ -754,7 +754,6 @@ void FOutputDeviceRedirector::Panic()
 		// or will resume and crash due to a race with the panic thread. Executing on this thread
 		// and having logging for most panic situations with a chance of a crash is preferable to
 		// the alternative of missing logging in a panic situation.
-		TRACE_CPUPROFILER_EVENT_SCOPE(FOutputDeviceRedirector::PanicWait);
 		constexpr double WaitTime = 1.0;
 		for (const double EndTime = FPlatformTime::Seconds() + WaitTime; FPlatformTime::Seconds() < EndTime;)
 		{
