Clean up screenshot comparison report pipeline

#rnx

diff --git a/Engine/Plugins/Tests/EditorTests/Source/EditorTests/Private/UnrealEd/EditorAssetAutomationTests.cpp b/Engine/Plugins/Tests/EditorTests/Source/EditorTests/Private/UnrealEd/EditorAssetAutomationTests.cpp
--- a/Engine/Plugins/Tests/EditorTests/Source/EditorTests/Private/UnrealEd/EditorAssetAutomationTests.cpp
+++ b/Engine/Plugins/Tests/EditorTests/Source/EditorTests/Private/UnrealEd/EditorAssetAutomationTests.cpp
@@ -966,9 +966,9 @@ namespace ImportExportAssetHelper
 			TSharedPtr<SWindow> ActiveWindow = FSlateApplication::Get().GetActiveTopLevelWindow();
 			if (ActiveWindow.IsValid())
 			{
-				FString ScreenshotName;
+				FString ScreenshotPath;
 				const FString TestName = FString::Printf(TEXT("AssetImportExport/Screenshots/%s"), *ImportedAsset->GetName());
-				ScreenshotName = AutomationCommon::GetScreenshotName(TestName);
+				ScreenshotPath = AutomationCommon::GetScreenshotPath(TestName);
 
 				TSharedRef<SWidget> WindowRef = ActiveWindow.ToSharedRef();
 
@@ -979,7 +979,7 @@ namespace ImportExportAssetHelper
 					FAutomationScreenshotData Data;
 					Data.Width = OutImageSize.X;
 					Data.Height = OutImageSize.Y;
-					Data.ScreenshotName = ScreenshotName;
+					Data.ScreenshotPath = ScreenshotPath;
 					FAutomationTestFramework::Get().OnScreenshotCaptured().ExecuteIfBound(OutImageData, Data);
 				}
 
diff --git a/Engine/Plugins/Tests/EditorTests/Source/EditorTests/Private/UnrealEd/EditorAutomationTests.cpp b/Engine/Plugins/Tests/EditorTests/Source/EditorTests/Private/UnrealEd/EditorAutomationTests.cpp
--- a/Engine/Plugins/Tests/EditorTests/Source/EditorTests/Private/UnrealEd/EditorAutomationTests.cpp
+++ b/Engine/Plugins/Tests/EditorTests/Source/EditorTests/Private/UnrealEd/EditorAutomationTests.cpp
@@ -249,7 +249,7 @@ bool FLoadAllMapsInEditorTest::RunTest(const FString& Parameters)
 
 		//Create a screen shot filename and path
 		const FString LoadAllMapsTestName = FString::Printf(TEXT("LoadAllMaps_Editor/%s"), *FPaths::GetBaseFilename(MapName));
-		WindowParameters.ScreenshotName = AutomationCommon::GetScreenshotName(LoadAllMapsTestName);
+		WindowParameters.ScreenshotName = AutomationCommon::GetScreenshotPath(LoadAllMapsTestName);
 
 		//Get the current number of seconds.  This will be used to track how long it took to load the map.
 		MapLoadStartTime = FPlatformTime::Seconds();
diff --git a/Engine/Plugins/Tests/EditorTests/Source/EditorTests/Private/UnrealEd/EditorBuildPromotionTests.cpp b/Engine/Plugins/Tests/EditorTests/Source/EditorTests/Private/UnrealEd/EditorBuildPromotionTests.cpp
--- a/Engine/Plugins/Tests/EditorTests/Source/EditorTests/Private/UnrealEd/EditorBuildPromotionTests.cpp
+++ b/Engine/Plugins/Tests/EditorTests/Source/EditorTests/Private/UnrealEd/EditorBuildPromotionTests.cpp
@@ -2392,7 +2392,7 @@ bool FBuildPromotionPIETest::RunTest(const FString& Parameters)
 	}
 
 	WindowScreenshotParameters ScreenshotParams;
-	ScreenshotParams.ScreenshotName = AutomationCommon::GetScreenshotName(TEXT("EditorBuildPromotion/RunMap"));
+	ScreenshotParams.ScreenshotName = AutomationCommon::GetScreenshotPath(TEXT("EditorBuildPromotion/RunMap"));
 	ScreenshotParams.CurrentWindow = AllWindows[0];
 	//Wait for the play world to come up
 	ADD_LATENT_AUTOMATION_COMMAND(FWaitLatentCommand(1.f));
diff --git a/Engine/Plugins/Tests/RuntimeTests/Source/RuntimeTests/Private/Shaders/ShaderComparisonTests.cpp b/Engine/Plugins/Tests/RuntimeTests/Source/RuntimeTests/Private/Shaders/ShaderComparisonTests.cpp
--- a/Engine/Plugins/Tests/RuntimeTests/Source/RuntimeTests/Private/Shaders/ShaderComparisonTests.cpp
+++ b/Engine/Plugins/Tests/RuntimeTests/Source/RuntimeTests/Private/Shaders/ShaderComparisonTests.cpp
@@ -146,14 +146,14 @@ bool FCompareBasepassShaders::RunTest(const FString& Parameters)
     for (int i = 0; i < NUM_CAMERAS; ++i)
     {
         FString CompareBasepassShadersTestName = FString::Printf(TEXT("CompareBasepassShaders_Game/%s/%d"), *FPaths::GetBaseFilename(MapName), i);
-		ScreenshotFileName[i] = AutomationCommon::GetScreenshotName(CompareBasepassShadersTestName);
+		ScreenshotFileName[i] = AutomationCommon::GetScreenshotPath(CompareBasepassShadersTestName);
         CompareBasepassShadersTestName = FString::Printf(TEXT("Incoming/CompareBasepassShaders_Game/%s/%d"), *FPaths::GetBaseFilename(MapName), i);
-		RealScreenshotFileName[i] = AutomationCommon::GetScreenshotName(CompareBasepassShadersTestName);
+		RealScreenshotFileName[i] = AutomationCommon::GetScreenshotPath(CompareBasepassShadersTestName);
         RealScreenshotFileName[i] = "../../../" + RealScreenshotFileName[i];
         CompareBasepassShadersTestName = FString::Printf(TEXT("CompareBasepassShaders_Game/%s_fp16/%d"), *FPaths::GetBaseFilename(MapName), i);
-		FP16ScreenshotFileName[i] = AutomationCommon::GetScreenshotName(CompareBasepassShadersTestName);
+		FP16ScreenshotFileName[i] = AutomationCommon::GetScreenshotPath(CompareBasepassShadersTestName);
         CompareBasepassShadersTestName = FString::Printf(TEXT("Incoming/CompareBasepassShaders_Game/%s_fp16/%d"), *FPaths::GetBaseFilename(MapName), i);
-		RealFP16ScreenshotFileName[i] = AutomationCommon::GetScreenshotName(CompareBasepassShadersTestName);
+		RealFP16ScreenshotFileName[i] = AutomationCommon::GetScreenshotPath(CompareBasepassShadersTestName);
         RealFP16ScreenshotFileName[i] = "../../../" + RealFP16ScreenshotFileName[i];
     }
 
diff --git a/Engine/Source/Developer/AutomationController/Private/AutomationControllerManager.cpp b/Engine/Source/Developer/AutomationController/Private/AutomationControllerManager.cpp
--- a/Engine/Source/Developer/AutomationController/Private/AutomationControllerManager.cpp
+++ b/Engine/Source/Developer/AutomationController/Private/AutomationControllerManager.cpp
@@ -486,7 +486,7 @@ void FAutomationControllerManager::ReportImageComparisonResult(const FAutomation
 			}
 		}
 
-		Report->AddArtifact(ClusterIndex, CurrentTestPass, FAutomationArtifact(Result.UniqueId, Result.ScreenshotName, EAutomationArtifactType::Comparison, LocalFiles));
+		Report->AddArtifact(ClusterIndex, CurrentTestPass, FAutomationArtifact(Result.UniqueId, Result.ScreenshotPath, EAutomationArtifactType::Comparison, LocalFiles));
 	}
 	else
 	{
@@ -508,7 +508,7 @@ void FAutomationControllerManager::ProcessComparisonQueue()
 			FAutomationWorkerImageComparisonResults ResultMessage(
 				Entry->InstanceId,
 				Result.bSkipAttachingImages ? FGuid() : FGuid::NewGuid(),
-				Result.ScreenshotName,
+				Result.ScreenshotPath,
 				Result.IsNew(),
 				Result.AreSimilar(),
 				Result.MaxLocalDifference,
diff --git a/Engine/Source/Developer/FunctionalTesting/Private/AutomationBlueprintFunctionLibrary.cpp b/Engine/Source/Developer/FunctionalTesting/Private/AutomationBlueprintFunctionLibrary.cpp
--- a/Engine/Source/Developer/FunctionalTesting/Private/AutomationBlueprintFunctionLibrary.cpp
+++ b/Engine/Source/Developer/FunctionalTesting/Private/AutomationBlueprintFunctionLibrary.cpp
@@ -465,7 +465,7 @@ public:
 
 			bool bAttemptToCompareShot = FAutomationTestFramework::Get().OnScreenshotCaptured().ExecuteIfBound(InImageData, Data);
 
-			UE_LOG(AutomationFunctionLibrary, Log, TEXT("Screenshot captured as %s"), *Data.ScreenshotName);
+			UE_LOG(AutomationFunctionLibrary, Log, TEXT("Screenshot captured as %s"), *Data.ScreenshotPath);
 
 			if (GIsAutomationTesting)
 			{
@@ -573,7 +573,7 @@ public:
 
 		bool bAttemptToCompareShot = FAutomationTestFramework::Get().OnScreenshotCaptured().ExecuteIfBound(InImageData, Data);
 
-		UE_LOG(AutomationFunctionLibrary, Log, TEXT("Screenshot captured as %s"), *Data.ScreenshotName);
+		UE_LOG(AutomationFunctionLibrary, Log, TEXT("Screenshot captured as %s"), *Data.ScreenshotPath);
 
 		FAutomationTestFramework::Get().OnScreenshotCompared.AddRaw(this, &FAutomationHighResScreenshotGrabber::OnComparisonComplete);
 	}
@@ -780,9 +780,9 @@ FIntPoint UAutomationBlueprintFunctionLibrary::GetAutomationScreenshotSize(const
 FAutomationScreenshotData UAutomationBlueprintFunctionLibrary::BuildScreenshotData(const FString& MapOrContext, const FString& ScreenShotName, int32 Width, int32 Height)
 {
 	FString TestName = TEXT("");
-	if (FFunctionalTestBase::IsFunctionalTestRunning())
+	if (FAutomationTestFramework::Get().GetCurrentTest()) 
 	{
-		TestName = FFunctionalTestBase::GetRunningTestName();
+		TestName = FAutomationTestFramework::Get().GetCurrentTest()->GetTestFullName();
 	}
 
 #if WITH_AUTOMATION_TESTS
@@ -810,10 +810,16 @@ bool UAutomationBlueprintFunctionLibrary::TakeAutomationScreenshotInternal(UObje
 	return true; //-V773
 }
 
-void UAutomationBlueprintFunctionLibrary::TakeAutomationScreenshot(UObject* WorldContextObject, FLatentActionInfo LatentInfo, const FString& ScreenShotName, const FString& Notes, const FAutomationScreenshotOptions& Options)
+void UAutomationBlueprintFunctionLibrary::TakeAutomationScreenshot(UObject* WorldContextObject, FLatentActionInfo LatentInfo, const FString& InScreenShotName, const FString& Notes, const FAutomationScreenshotOptions& Options)
 {
 	if ( GIsAutomationTesting )
 	{
+		FString ScreenShotName = InScreenShotName;
+		if ( ScreenShotName.IsEmpty() )
+		{
+			ScreenShotName = TEXT("Undefined");
+			UE_LOG(AutomationFunctionLibrary, Warning, TEXT("Screenshot name is empty. Default name will be used."));
+		}
 		if ( UWorld* World = GEngine->GetWorldFromContextObject(WorldContextObject, EGetWorldErrorMode::LogAndReturnNull) )
 		{
 			FLatentActionManager& LatentActionManager = World->GetLatentActionManager();
diff --git a/Engine/Source/Developer/FunctionalTesting/Private/ScreenshotFunctionalTest.cpp b/Engine/Source/Developer/FunctionalTesting/Private/ScreenshotFunctionalTest.cpp
--- a/Engine/Source/Developer/FunctionalTesting/Private/ScreenshotFunctionalTest.cpp
+++ b/Engine/Source/Developer/FunctionalTesting/Private/ScreenshotFunctionalTest.cpp
@@ -168,7 +168,7 @@ void AScreenshotFunctionalTest::OnScreenShotCaptured(int32 InSizeX, int32 InSize
 
 	FAutomationTestFramework::Get().OnScreenshotAndTraceCaptured().ExecuteIfBound(InImageData, CapturedFrameTrace, Data);
 
-	UE_LOG(LogScreenshotFunctionalTest, Log, TEXT("Screenshot captured as %s"), *Data.ScreenshotName);
+	UE_LOG(LogScreenshotFunctionalTest, Log, TEXT("Screenshot captured as %s"), *Data.ScreenshotPath);
 #endif
 }
 
diff --git a/Engine/Source/Developer/FunctionalTesting/Private/ScreenshotFunctionalTestBase.cpp b/Engine/Source/Developer/FunctionalTesting/Private/ScreenshotFunctionalTestBase.cpp
--- a/Engine/Source/Developer/FunctionalTesting/Private/ScreenshotFunctionalTestBase.cpp
+++ b/Engine/Source/Developer/FunctionalTesting/Private/ScreenshotFunctionalTestBase.cpp
@@ -176,7 +176,7 @@ void AScreenshotFunctionalTestBase::OnScreenShotCaptured(int32 InSizeX, int32 In
 
 	FAutomationTestFramework::Get().OnScreenshotAndTraceCaptured().ExecuteIfBound(InImageData, CapturedFrameTrace, Data);
 
-	UE_LOG(LogScreenshotFunctionalTest, Log, TEXT("Screenshot captured as %s"), *Data.ScreenshotName);
+	UE_LOG(LogScreenshotFunctionalTest, Log, TEXT("Screenshot captured as %s"), *Data.ScreenshotPath);
 #endif
 }
 
diff --git a/Engine/Source/Developer/ScreenShotComparison/Private/Models/ScreenComparisonModel.cpp b/Engine/Source/Developer/ScreenShotComparison/Private/Models/ScreenComparisonModel.cpp
--- a/Engine/Source/Developer/ScreenShotComparison/Private/Models/ScreenComparisonModel.cpp
+++ b/Engine/Source/Developer/ScreenShotComparison/Private/Models/ScreenComparisonModel.cpp
@@ -87,11 +87,6 @@ FString FScreenComparisonModel::GetName()
 		{
 			FString VariantSuffix = LoadedMetadata->VariantName.Len() > 0 ? FString::Printf(TEXT(".%s"), *LoadedMetadata->VariantName) : TEXT("");
 			FString NameString = FString::Printf(TEXT("%s.%s%s"), *LoadedMetadata->Context, *LoadedMetadata->ScreenShotName, *VariantSuffix);
-			if ((LoadedMetadata->Context.Len() && LoadedMetadata->TestName.Len())
-				|| !LoadedMetadata->ScreenShotName.Len())
-			{
-				NameString = FString::Printf(TEXT("%s.%s%s"), *LoadedMetadata->Context, *LoadedMetadata->TestName, *VariantSuffix);
-			}
 
 			Name = NameString;
 		}
diff --git a/Engine/Source/Developer/ScreenShotComparisonTools/Private/ScreenShotManager.cpp b/Engine/Source/Developer/ScreenShotComparisonTools/Private/ScreenShotManager.cpp
--- a/Engine/Source/Developer/ScreenShotComparisonTools/Private/ScreenShotManager.cpp
+++ b/Engine/Source/Developer/ScreenShotComparisonTools/Private/ScreenShotManager.cpp
@@ -389,7 +389,7 @@ FImageComparisonResult FScreenShotManager::CompareScreenshot(const FString& InUn
 	ComparisonResult.SourceRHI = IncomingMetaData.Rhi;
 	ComparisonResult.IdealApprovedFolderPath = IdealApprovedFolderPath;
 	// We use the subfolder path to the screenshot as name (before any environment specialization - platform, RHI - are appended).
-	ComparisonResult.ScreenshotName = ResultsSubFolder;
+	ComparisonResult.ScreenshotPath = ResultsSubFolder;
 
 	// Do not save passing variant test screenshots
 	// Disabled for now until more variants are added, since variants now run without baseline tests in lightweight mode and we need to save at least one screenshot
@@ -674,7 +674,7 @@ TSharedPtr<FImageComparisonResult> FScreenShotManager::CompareImageSequence(cons
 			ComparisonResult.SourcePlatform = Metadata.Platform;
 			ComparisonResult.SourceRHI = Metadata.Rhi;
 			ComparisonResult.IdealApprovedFolderPath = GetIdealApprovedFolderForImage(Metadata);;
-			ComparisonResult.ScreenshotName = GetPathComponentForTestImages(Metadata, true);
+			ComparisonResult.ScreenshotPath = GetPathComponentForTestImages(Metadata, true);
 			FString ReportPathOnDisk = FPaths::Combine(ScreenshotResultsFolder, GetPathComponentForTestImages(Metadata, true), Metadata.Platform, GetPathComponentForRHI(Metadata), TEXT("/"));
 
 			ComparisonResult.ReportApprovedFilePath = FPaths::Combine(ReportPathOnDisk, TEXT("Approved.png"));
@@ -750,7 +750,7 @@ void FScreenShotManager::NotifyAutomationTestFrameworkOfImageComparison(const FI
 		ComparisonResult.ReportComparisonFilePath,
 		ComparisonResult.ReportApprovedFilePath,
 		ComparisonResult.ReportIncomingFilePath,
-		ComparisonResult.ScreenshotName
+		ComparisonResult.ScreenshotPath
 	);
 	FAutomationTestFramework::Get().NotifyScreenshotComparisonReport(Results);
 	if (FAutomationTestBase* CurrentTest = FAutomationTestFramework::Get().GetCurrentTest())
diff --git a/Engine/Source/Developer/ScreenShotComparisonTools/Public/ImageComparer.h b/Engine/Source/Developer/ScreenShotComparisonTools/Public/ImageComparer.h
--- a/Engine/Source/Developer/ScreenShotComparisonTools/Public/ImageComparer.h
+++ b/Engine/Source/Developer/ScreenShotComparisonTools/Public/ImageComparer.h
@@ -303,10 +303,10 @@ public:
 	FText ErrorMessage;
 
 	/*
-		Name of the screenshot (includes variant if applicable)
+		Path of the screenshot (includes variant if applicable)
 	*/
 	UPROPERTY()
-	FString ScreenshotName;
+	FString ScreenshotPath;
 
 	/*
 		Whether to skip saving and attaching images to the report for this test
diff --git a/Engine/Source/Runtime/AutomationMessages/Public/AutomationWorkerMessages.h b/Engine/Source/Runtime/AutomationMessages/Public/AutomationWorkerMessages.h
--- a/Engine/Source/Runtime/AutomationMessages/Public/AutomationWorkerMessages.h
+++ b/Engine/Source/Runtime/AutomationMessages/Public/AutomationWorkerMessages.h
@@ -753,7 +753,7 @@ public:
 	FAutomationWorkerImageComparisonResults(
 		FGuid InInstanceId,
 		FGuid InUniqueId,
-		const FString& InName,
+		const FString& InScreenshotPath,
 		bool InIsNew,
 		bool InAreSimilar,
 		double InMaxLocalDifference,
@@ -766,7 +766,7 @@ public:
 	)
 		: FAutomationWorkerMessageBase{ InInstanceId }
 		, UniqueId(InUniqueId)
-		, ScreenshotName(InName)
+		, ScreenshotPath(InScreenshotPath)
 		, bNew(InIsNew)
 		, bSimilar(InAreSimilar)
 		, MaxLocalDifference(InMaxLocalDifference)
@@ -783,9 +783,9 @@ public:
 	UPROPERTY(EditAnywhere, Category="Message")
 	FGuid UniqueId;
 
-	/** The name of the screenshot. */
+	/** The path of the screenshot. */
 	UPROPERTY(EditAnywhere, Category = "Message")
-	FString ScreenshotName;
+	FString ScreenshotPath;
 
 	/** Was this a new image we've never seen before and have no ground truth for? */
 	UPROPERTY(EditAnywhere, Category="Message")
diff --git a/Engine/Source/Runtime/AutomationWorker/Private/AutomationWorkerModule.cpp b/Engine/Source/Runtime/AutomationWorker/Private/AutomationWorkerModule.cpp
--- a/Engine/Source/Runtime/AutomationWorker/Private/AutomationWorkerModule.cpp
+++ b/Engine/Source/Runtime/AutomationWorker/Private/AutomationWorkerModule.cpp
@@ -416,7 +416,7 @@ void FAutomationWorkerModule::HandleScreenShotCompared(const FAutomationWorkerIm
 	// Image comparison finished.
 	FAutomationScreenshotCompareResults CompareResults;
 	CompareResults.UniqueId = Message.UniqueId;
-	CompareResults.ScreenshotName = Message.ScreenshotName;
+	CompareResults.ScreenshotPath = Message.ScreenshotPath;
 	CompareResults.bWasNew = Message.bNew;
 	CompareResults.bWasSimilar = Message.bSimilar;
 	CompareResults.MaxLocalDifference = Message.MaxLocalDifference;
@@ -445,7 +445,7 @@ void FAutomationWorkerModule::HandleScreenShotComparisonReport(const FAutomation
 {
 	FAutomationWorkerImageComparisonResults* Message = FMessageEndpoint::MakeMessage<FAutomationWorkerImageComparisonResults>();
 
-	Message->ScreenshotName = Results.ScreenshotName;
+	Message->ScreenshotPath = Results.ScreenshotPath;
 	Message->UniqueId = Results.UniqueId;
 	Message->bNew = Results.bWasNew;
 	Message->bSimilar = Results.bWasSimilar;
@@ -482,7 +482,7 @@ void FAutomationWorkerModule::HandleScreenShotAndTraceCapturedWithName(const TAr
 	{
 		FAutomationWorkerScreenImage* Message = FMessageEndpoint::MakeMessage<FAutomationWorkerScreenImage>();
 
-		Message->ScreenShotName = Data.ScreenshotName;
+		Message->ScreenShotName = Data.ScreenshotPath;
 		Message->ScreenImage = CompressedBitmap;
 		Message->FrameTrace = CapturedFrameTrace;
 		Message->Metadata = Metadata;
@@ -496,11 +496,11 @@ void FAutomationWorkerModule::HandleScreenShotAndTraceCapturedWithName(const TAr
 		//Save locally
 		const bool bTree = true;
 
-		FString LocalFile = AutomationCommon::GetLocalPathForScreenshot(Data.ScreenshotName);
+		FString LocalFile = AutomationCommon::GetLocalPathForScreenshot(Data.ScreenshotPath);
 		FString LocalTraceFile = FPaths::ChangeExtension(LocalFile, TEXT(".rdc"));
 		FString DirectoryPath = FPaths::GetPath(LocalFile);
 
-		UE_LOG(LogAutomationWorker, Log, TEXT("Saving screenshot %s as %s"),*Data.ScreenshotName, *LocalFile);
+		UE_LOG(LogAutomationWorker, Log, TEXT("Saving screenshot %s as %s"),*Data.ScreenshotPath, *LocalFile);
 
 		if (!IFileManager::Get().MakeDirectory(*DirectoryPath, bTree))
 		{
diff --git a/Engine/Source/Runtime/Core/Private/Misc/AutomationTest.cpp b/Engine/Source/Runtime/Core/Private/Misc/AutomationTest.cpp
--- a/Engine/Source/Runtime/Core/Private/Misc/AutomationTest.cpp
+++ b/Engine/Source/Runtime/Core/Private/Misc/AutomationTest.cpp
@@ -947,7 +947,7 @@ void FAutomationTestFramework::DumpAutomationTestExecutionInfo( const TMap<FStri
 void FAutomationTestFramework::InternalStartTest( const FString& InTestToRun, const FString& InFullTestPath)
 {
 	Parameters.Empty();
-	FullTestPath.Empty();
+	CurrentTestFullPath.Empty();
 
 	FString TestName;
 	if (!InTestToRun.Split(TEXT(" "), &TestName, &Parameters, ESearchCase::CaseSensitive))
@@ -970,7 +970,7 @@ void FAutomationTestFramework::InternalStartTest( const FString& InTestToRun, co
 		StartTime = FPlatformTime::Seconds();
 
 		CurrentTest->SetTestContext(Parameters);
-		FullTestPath = InFullTestPath;
+		CurrentTestFullPath = InFullTestPath;
 
 		// If not a smoke test, log the test has started.
 		uint32 NonSmokeTestFlags = (EAutomationTestFlags::FilterMask & (~EAutomationTestFlags::SmokeFilter));
@@ -978,7 +978,7 @@ void FAutomationTestFramework::InternalStartTest( const FString& InTestToRun, co
 		{
 			if (AutomationTest::bLogTestStateTrace)
 			{
-				UE_LOG(LogAutomationTestStateTrace, Log, TEXT("Test is about to start. Name={%s}"), *FullTestPath);
+				UE_LOG(LogAutomationTestStateTrace, Log, TEXT("Test is about to start. Name={%s}"), *CurrentTestFullPath);
 			}
 			UE_LOG(LogAutomationTest, Log, TEXT("%s %s is starting at %f"), *CurrentTest->GetBeautifiedTestName(), *Parameters, StartTime);
 		}
@@ -1022,7 +1022,7 @@ bool FAutomationTestFramework::InternalStopTest(FAutomationTestExecutionInfo& Ou
 		UE_LOG(LogAutomationTest, Log, TEXT("%s %s ran in %f"), *CurrentTest->GetBeautifiedTestName(), *Parameters, TimeForTest);
 		if (AutomationTest::bLogTestStateTrace)
 		{
-			UE_LOG(LogAutomationTestStateTrace, Log, TEXT("Test has stopped execution. Name={%s}"), *FullTestPath);
+			UE_LOG(LogAutomationTestStateTrace, Log, TEXT("Test has stopped execution. Name={%s}"), *CurrentTestFullPath);
 		}
 	}
 
@@ -1288,11 +1288,14 @@ void FAutomationTestExecutionInfo::AddError(const FString& ErrorMessage)
 FAutomationEvent FAutomationScreenshotCompareResults::ToAutomationEvent() const
 {
 	FAutomationEvent Event(EAutomationEventType::Info, TEXT(""));
+	FString OutputScreenshotName = ScreenshotPath;
+	FPaths::NormalizeDirectoryName(OutputScreenshotName);
+	OutputScreenshotName.ReplaceInline(TEXT("/"), TEXT("."));
 
 	if (bWasNew)
 	{
 		Event.Type = EAutomationEventType::Warning;
-		Event.Message = FString::Printf(TEXT("New Screenshot '%s' was discovered!  Please add a ground truth version of it."), *ScreenshotName);
+		Event.Message = FString::Printf(TEXT("New Screenshot '%s' was discovered!  Please add a ground truth version of it."), *OutputScreenshotName);
 	}
 	else
 	{
@@ -1300,7 +1303,7 @@ FAutomationEvent FAutomationScreenshotCompareResults::ToAutomationEvent() const
 		{
 			Event.Type = EAutomationEventType::Info;
 			Event.Message = FString::Printf(TEXT("Screenshot '%s' was similar!  Global Difference = %f, Max Local Difference = %f"),
-				*ScreenshotName, GlobalDifference, MaxLocalDifference);
+				*OutputScreenshotName, GlobalDifference, MaxLocalDifference);
 		}
 		else
 		{
@@ -1309,11 +1312,11 @@ FAutomationEvent FAutomationScreenshotCompareResults::ToAutomationEvent() const
 			if (ErrorMessage.IsEmpty())
 			{
 				Event.Message = FString::Printf(TEXT("Screenshot '%s' test failed, Screenshots were different!  Global Difference = %f, Max Local Difference = %f"),
-					*ScreenshotName, GlobalDifference, MaxLocalDifference);
+					*OutputScreenshotName, GlobalDifference, MaxLocalDifference);
 			}
 			else
 			{
-				Event.Message = FString::Printf(TEXT("Screenshot '%s' test failed; Error = %s"), *ScreenshotName, *ErrorMessage);
+				Event.Message = FString::Printf(TEXT("Screenshot '%s' test failed; Error = %s"), *OutputScreenshotName, *ErrorMessage);
 			}
 		}
 	}
diff --git a/Engine/Source/Runtime/Core/Public/Misc/AutomationTest.h b/Engine/Source/Runtime/Core/Public/Misc/AutomationTest.h
--- a/Engine/Source/Runtime/Core/Public/Misc/AutomationTest.h
+++ b/Engine/Source/Runtime/Core/Public/Misc/AutomationTest.h
@@ -748,8 +748,8 @@ struct FAutomationScreenshotData
 	bool bIgnoreAntiAliasing;
 	bool bIgnoreColors;
 
-	// Name of the screenshot generated from AutomationCommon::GetScreenShotName()
-	FString ScreenshotName;
+	// Path of the screenshot generated from AutomationCommon::GetScreenShotPath()
+	FString ScreenshotPath;
 
 	FAutomationScreenshotData()
 		: Id()
@@ -795,7 +795,7 @@ struct FAutomationScreenshotCompareResults
 	FString ReportComparisonFilePath;
 	FString ReportApprovedFilePath;
 	FString ReportIncomingFilePath;
-	FString ScreenshotName;
+	FString ScreenshotPath;
 
 	FAutomationScreenshotCompareResults()
 		: UniqueId()
@@ -816,7 +816,7 @@ struct FAutomationScreenshotCompareResults
 		FString InReportComparisonFilePath,
 		FString InReportApprovedFilePath,
 		FString InReportIncomingFilePath,
-		FString InScreenshotName
+		FString InScreenshotPath
 	)
 		: UniqueId(InUniqueId)
 		, ErrorMessage(InErrorMessage)
@@ -828,7 +828,7 @@ struct FAutomationScreenshotCompareResults
 		, ReportComparisonFilePath(InReportComparisonFilePath)
 		, ReportApprovedFilePath(InReportApprovedFilePath)
 		, ReportIncomingFilePath(InReportIncomingFilePath)
-		, ScreenshotName(InScreenshotName)
+		, ScreenshotPath(InScreenshotPath)
 	{ }
 
 	CORE_API FAutomationEvent ToAutomationEvent() const;
@@ -1149,6 +1149,14 @@ public:
 		return CurrentTest;
 	}
 
+	/**
+	 * Returns the actively executing test full path
+	 */
+	FString GetCurrentTestFullPath() const
+	{
+		return CurrentTestFullPath;
+	}
+
 	/**
 	 * Whether to skip stack walk while iterating for listing the tests
 	 */
@@ -1381,7 +1389,7 @@ private:
 	FString Parameters;
 
 	/** Full test path as given by the automation controller of the active test */
-	FString FullTestPath;
+	FString CurrentTestFullPath;
 
 	/** Whether we want to run automation tests on content within the Developer Directories */
 	bool bDeveloperDirectoryIncluded;
@@ -1452,6 +1460,10 @@ public:
 	 * Returns the beautified test name with test context. Should return what is displayed in the Test Automation UI. See GenerateTestNames()
 	 */
 	virtual FString GetTestFullName() const {
+		if (FAutomationTestFramework::Get().GetCurrentTest() == this)
+		{
+			return FAutomationTestFramework::Get().GetCurrentTestFullPath();
+		}
 		if (GetTestContext().IsEmpty()) { return GetBeautifiedTestName(); }
 		return FString::Printf(TEXT("%s.%s"), *GetBeautifiedTestName(), *GetTestContext());
 	}
diff --git a/Engine/Source/Runtime/Engine/Private/Tests/AutomationCommon.cpp b/Engine/Source/Runtime/Engine/Private/Tests/AutomationCommon.cpp
--- a/Engine/Source/Runtime/Engine/Private/Tests/AutomationCommon.cpp
+++ b/Engine/Source/Runtime/Engine/Private/Tests/AutomationCommon.cpp
@@ -92,7 +92,7 @@ namespace AutomationCommon
 	}
 
 	/** Gets a path used for automation testing (PNG sent to the AutomationTest folder) */
-	FString GetScreenshotName(const FString& TestName)
+	FString GetScreenshotPath(const FString& TestName)
 	{
 		FString PathName = TestName / FPlatformProperties::IniPlatformName();
 		PathName = PathName + TEXT("/") + GetRenderDetailsString();
@@ -151,7 +151,7 @@ namespace AutomationCommon
 		// Device's native resolution (we want to use a hardware dump of the frontbuffer at the native resolution so we compare what we actually output rather than what we think we rendered)
 
 		const FString MapAndTest = MapOrContext + TEXT("/") + Data.ScreenShotName;
-		Data.ScreenshotName = GetScreenshotName(MapAndTest);
+		Data.ScreenshotPath = GetScreenshotPath(MapAndTest);
 
 		return Data;
 	}
@@ -165,8 +165,8 @@ namespace AutomationCommon
 			&& FModuleManager::Get().IsModuleLoaded("RenderDocPlugin"))
 		{
 			const FString MapAndTest = MapOrContext / FPaths::MakeValidFileName(TestName, TEXT('_'));
-			FString ScreenshotName = GetScreenshotName(MapAndTest);
-			FString TempCaptureFilePath = FPaths::ChangeExtension(FPaths::ConvertRelativePathToFull(FPaths::AutomationDir() / TEXT("Incoming/") / ScreenshotName), TEXT(".rdc"));
+			FString ScreenshotPath = GetScreenshotPath(MapAndTest);
+			FString TempCaptureFilePath = FPaths::ChangeExtension(FPaths::ConvertRelativePathToFull(FPaths::AutomationDir() / TEXT("Incoming/") / ScreenshotPath), TEXT(".rdc"));
 
 			UE_LOG(LogEngineAutomationTests, Log, TEXT("Taking Frame Trace: %s"), *TempCaptureFilePath);
 
@@ -242,19 +242,25 @@ namespace AutomationCommon
 			: ImageData(InImageData), Initiate(false), TaskCompleted(false)
 		{
 			FString Context = InContext;
-			if (Context.IsEmpty())
+			FString TestName = TEXT("");
+			
+			if (FAutomationTestBase* CurrentTest = FAutomationTestFramework::Get().GetCurrentTest())
 			{
-				if (FAutomationTestBase* CurrentTest = FAutomationTestFramework::Get().GetCurrentTest())
+				TestName = CurrentTest->GetTestFullName();
+				if (Context.IsEmpty()) 
 				{
-					Context = CurrentTest->GetTestContext();
-					if (Context.IsEmpty())
+					if (!CurrentTest->GetTestContext().IsEmpty()) 
+					{
+						Context = CurrentTest->GetTestContext();
+					}
+					else
 					{
 						Context = CurrentTest->GetTestFullName();
 					}
 				}
 			}
 
-			ComparisonParameters = BuildScreenshotData(Context, TEXT(""), InImageName, InWidth, InHeight);
+			ComparisonParameters = BuildScreenshotData(Context, TestName, InImageName, InWidth, InHeight);
 
 			// Copy the relevant data into the metadata for the screenshot.
 			ComparisonParameters.bHasComparisonRules = true;
@@ -301,7 +307,7 @@ namespace AutomationCommon
 			{
 				FAutomationTestFramework::Get().OnScreenshotCaptured().ExecuteIfBound(ImageData, ComparisonParameters);
 
-				UE_LOG(LogEditorAutomationTests, Log, TEXT("Requesting image %s to be compared."), *ComparisonParameters.ScreenshotName);
+				UE_LOG(LogEditorAutomationTests, Log, TEXT("Requesting image %s to be compared."), *ComparisonParameters.ScreenshotPath);
 
 				FAutomationTestFramework::Get().OnScreenshotCompared.AddRaw(this, &FAutomationImageComparisonRequest::OnComparisonComplete);
 				Initiate = true;
@@ -328,7 +334,7 @@ namespace AutomationCommon
 			FAutomationScreenshotData Data;
 			Data.Width = OutImageSize.X;
 			Data.Height = OutImageSize.Y;
-			Data.ScreenshotName = ScreenshotName;
+			Data.ScreenshotPath = ScreenshotName;
 			FAutomationTestFramework::Get().OnScreenshotCaptured().ExecuteIfBound(OutImageData, Data);
 		}
 	}
diff --git a/Engine/Source/Runtime/Engine/Public/Tests/AutomationCommon.h b/Engine/Source/Runtime/Engine/Public/Tests/AutomationCommon.h
--- a/Engine/Source/Runtime/Engine/Public/Tests/AutomationCommon.h
+++ b/Engine/Source/Runtime/Engine/Public/Tests/AutomationCommon.h
@@ -35,7 +35,7 @@ namespace AutomationCommon
 		TestName/PlatformName/DeviceName.png. It's important to understand that a screenshot
 		generated on a device will likely have a different absolute path than the editor so this
 		name should be used with	*/
-	ENGINE_API FString GetScreenshotName(const FString& TestName);
+	ENGINE_API FString GetScreenshotPath(const FString& TestName);
 
 	/** 
 	This function takes the result of GetScreenshotName and will return the complete path to where a
