Respect global suspend in FThreadHeartBeat::IsBeating().

diff --git a/Engine/Source/Runtime/Core/Private/HAL/ThreadHeartBeat.cpp b/Engine/Source/Runtime/Core/Private/HAL/ThreadHeartBeat.cpp
--- a/Engine/Source/Runtime/Core/Private/HAL/ThreadHeartBeat.cpp
+++ b/Engine/Source/Runtime/Core/Private/HAL/ThreadHeartBeat.cpp
@@ -840,14 +840,16 @@ void FThreadHeartBeat::ResumeHeartBeat(bool bAllThreads)
 
 bool FThreadHeartBeat::IsBeating()
 {
-	uint32 ThreadId = FPlatformTLS::GetCurrentThreadId();
-	FScopeLock HeartBeatLock(&HeartBeatCritical);
-	FHeartBeatInfo* HeartBeatInfo = ThreadHeartBeat.Find(ThreadId);
-	if (HeartBeatInfo && HeartBeatInfo->SuspendedCount == 0)
+	if (GlobalSuspendCount.GetValue() == 0)
 	{
-		return true;
+		uint32 ThreadId = FPlatformTLS::GetCurrentThreadId();
+		FScopeLock HeartBeatLock(&HeartBeatCritical);
+		FHeartBeatInfo* HeartBeatInfo = ThreadHeartBeat.Find(ThreadId);
+		if (HeartBeatInfo && HeartBeatInfo->SuspendedCount == 0)
+		{
+			return true;
+		}
 	}
-
 	return false;
 }
 
