Prevent internal compile error.

#rnx

diff --git a/Engine/Source/Runtime/Core/Public/Async/TaskGraphInterfaces.h b/Engine/Source/Runtime/Core/Public/Async/TaskGraphInterfaces.h
--- a/Engine/Source/Runtime/Core/Public/Async/TaskGraphInterfaces.h
+++ b/Engine/Source/Runtime/Core/Public/Async/TaskGraphInterfaces.h
@@ -1145,7 +1145,7 @@ public:
 		}
 
 	private:
-		friend class TGraphTask;
+		friend TGraphTask;
 
 		/** The task that created me to assist with embeded task construction and preparation. **/
 		TGraphTask*						Owner;
