Moved TUseBitwiseSwap into its own header, with updated comment.

diff --git a/Engine/Source/Runtime/Core/Public/Templates/UnrealTemplate.h b/Engine/Source/Runtime/Core/Public/Templates/UnrealTemplate.h
--- a/Engine/Source/Runtime/Core/Public/Templates/UnrealTemplate.h
+++ b/Engine/Source/Runtime/Core/Public/Templates/UnrealTemplate.h
@@ -11,6 +11,7 @@
 #include "Templates/TypeCompatibleBytes.h"
 #include "Templates/Identity.h"
 #include "Traits/IsContiguousContainer.h"
+#include "Traits/UseBitwiseSwap.h"
 #include <type_traits>
 
 /*-----------------------------------------------------------------------------
@@ -594,17 +595,6 @@ UE_INTRINSIC_CAST FORCEINLINE T&& Forward(std::remove_reference_t<T>&& Obj)
 	return (T&&)Obj;
 }
 
-/**
- * A traits class which specifies whether a Swap of a given type should swap the bits or use a traditional value-based swap.
- */
-template <typename T>
-struct TUseBitwiseSwap
-{
-	// We don't use bitwise swapping for 'register' types because this will force them into memory and be slower.
-	enum { Value = !(std::is_enum_v<T> || std::is_pointer_v<T> || std::is_arithmetic_v<T>) };
-};
-
-
 /**
  * Swap two values.  Assumes the types are trivially relocatable.
  */
diff --git a/Engine/Source/Runtime/Core/Public/Traits/UseBitwiseSwap.h b/Engine/Source/Runtime/Core/Public/Traits/UseBitwiseSwap.h
new file mode 100644
--- /dev/null
+++ b/Engine/Source/Runtime/Core/Public/Traits/UseBitwiseSwap.h
@@ -0,0 +1,16 @@
+// Copyright Epic Games, Inc. All Rights Reserved.
+
+#pragma once
+
+#include "CoreTypes.h"
+#include <type_traits>
+
+/**
+ * A traits class which specifies whether a relocation of a given type should use a bitwise function like memcpy or memswap or traditional value-based operations.
+ */
+template <typename T>
+struct TUseBitwiseSwap
+{
+	// We don't use bitwise operations for 'register' types because this will force them into memory and be slower.
+	enum { Value = !(std::is_enum_v<T> || std::is_pointer_v<T> || std::is_arithmetic_v<T>) };
+};
