Fixed "disableunity" build when UE_TRACE_ENABLED is off but UE_TRACE_SERVER_LAUNCH_ENABLED is on.

diff --git a/Engine/Source/Runtime/Core/Private/ProfilingDebugging/TraceAuxiliary.cpp b/Engine/Source/Runtime/Core/Private/ProfilingDebugging/TraceAuxiliary.cpp
--- a/Engine/Source/Runtime/Core/Private/ProfilingDebugging/TraceAuxiliary.cpp
+++ b/Engine/Source/Runtime/Core/Private/ProfilingDebugging/TraceAuxiliary.cpp
@@ -28,6 +28,8 @@
 #endif
 
 #if UE_TRACE_SERVER_LAUNCH_ENABLED
+#include "HAL/PlatformProcess.h"
+#include "Misc/CommandLine.h"
 #include "Misc/Parse.h"
 #include "Misc/Paths.h"
 #endif
