Fix TMatrix2x2 rotation constructor. TQuat2 should be used instead of FQuat2D.

diff --git a/Engine/Source/Runtime/Core/Public/Math/TransformCalculus2D.h b/Engine/Source/Runtime/Core/Public/Math/TransformCalculus2D.h
--- a/Engine/Source/Runtime/Core/Public/Math/TransformCalculus2D.h
+++ b/Engine/Source/Runtime/Core/Public/Math/TransformCalculus2D.h
@@ -442,7 +442,7 @@ public:
 	}
 
 	/** Ctor. initialize from a rotation. */
-	explicit TMatrix2x2(const FQuat2D& Rotation)
+	explicit TMatrix2x2(const TQuat2<T>& Rotation)
 	{
 		T CosAngle = (T)Rotation.GetVector().X;
 		T SinAngle = (T)Rotation.GetVector().Y;
