Removal of NetworkVersion.h from Vector2D.h (which gets included everywhere by coreminimal.h)
So NetworkVersion.h ended up being an important part of coreminimal.h

diff --git a/Engine/Plugins/2D/Paper2D/Source/Paper2DEditor/Public/PaperJSONHelpers.h b/Engine/Plugins/2D/Paper2D/Source/Paper2DEditor/Public/PaperJSONHelpers.h
--- a/Engine/Plugins/2D/Paper2D/Source/Paper2DEditor/Public/PaperJSONHelpers.h
+++ b/Engine/Plugins/2D/Paper2D/Source/Paper2DEditor/Public/PaperJSONHelpers.h
@@ -3,6 +3,7 @@
 #pragma once
 
 #include "Math/Vector2D.h"
+#include "Templates/SharedPointer.h"
 
 class FJsonValue;
 
diff --git a/Engine/Source/Runtime/Core/Public/Math/Vector2D.h b/Engine/Source/Runtime/Core/Public/Math/Vector2D.h
--- a/Engine/Source/Runtime/Core/Public/Math/Vector2D.h
+++ b/Engine/Source/Runtime/Core/Public/Math/Vector2D.h
@@ -9,7 +9,10 @@
 #include "Containers/UnrealString.h"
 #include "Misc/Parse.h"
 #include "Misc/LargeWorldCoordinatesSerializer.h"
+#if UE_ENABLE_INCLUDE_ORDER_DEPRECATED_IN_5_4
 #include "Misc/NetworkVersion.h"
+#endif
+#include "Misc/EngineNetworkCustomVersion.h"
 #include "Math/IntPoint.h"
 #include "Logging/LogMacros.h"
 
