Clean up for two-way enum mapping

diff --git a/Engine/Source/Runtime/Core/Public/Experimental/Containers/TwoWayMap.h b/Engine/Source/Runtime/Core/Public/Experimental/Containers/TwoWayMap.h
--- a/Engine/Source/Runtime/Core/Public/Experimental/Containers/TwoWayMap.h
+++ b/Engine/Source/Runtime/Core/Public/Experimental/Containers/TwoWayMap.h
@@ -41,4 +41,58 @@ struct TTwoWayMap
 protected:
 	inline static TMap<FirstType, SecondType> FirstToSecond = TMap<FirstType, SecondType>();
 	inline static TMap<SecondType, FirstType> SecondToFirst = TMap<SecondType, FirstType>();
+};
+
+template<typename DomainType, typename RangeType>
+struct TBijectionMaps
+{
+protected:
+	struct FPairValue
+	{
+		const DomainType X;
+		const RangeType Y;
+	};
+
+public:
+	TBijectionMaps() = delete;
+
+	constexpr TBijectionMaps(const std::initializer_list<FPairValue>& InitList)
+	{
+		// TODO: Run compile time check to make sure all keys and values are unique
+		for (const FPairValue& Pair : InitList)
+		{
+			if (ensureAlwaysMsgf(!Image.Contains(Pair.X), TEXT("This cannot be a bijection if there is a duplicate domain entry")))
+			{
+				Image.Emplace(Pair.X, Pair.Y);
+			}
+			if (ensureAlwaysMsgf(!PreImage.Contains(Pair.Y), TEXT("This cannot be a bijection if there is a duplicate range entry")))
+			{
+				PreImage.Emplace(Pair.Y, Pair.X);
+			}
+		}
+	}
+
+	constexpr const DomainType* Find(const RangeType& InValue) const
+	{
+		return PreImage.Find(InValue);
+	}
+
+	constexpr const DomainType& FindChecked(const RangeType& InValue) const
+	{
+		return PreImage[InValue];
+	}
+
+	constexpr const RangeType* Find(const DomainType& InValue) const
+	{
+		return Image.Find(InValue);
+	}
+
+	constexpr const RangeType& FindChecked(const DomainType& InValue) const
+	{
+		return Image[InValue];
+	}
+
+private:
+	TMap<DomainType, RangeType> Image = TMap<DomainType, RangeType>();
+	TMap<RangeType, DomainType> PreImage = TMap<RangeType, DomainType>();
 };
\ No newline at end of file
