Add logging for when short term IO stats are reset

diff --git a/Engine/Source/Runtime/Core/Private/HAL/DiskUtilizationTracker.cpp b/Engine/Source/Runtime/Core/Private/HAL/DiskUtilizationTracker.cpp
--- a/Engine/Source/Runtime/Core/Private/HAL/DiskUtilizationTracker.cpp
+++ b/Engine/Source/Runtime/Core/Private/HAL/DiskUtilizationTracker.cpp
@@ -47,6 +47,7 @@ void FDiskUtilizationTracker::StartRead(uint64 InReadBytes, uint64 InSeekDistanc
 	if (bReset)
 	{
 		ShortTermStats.Reset();
+		UE_LOG(LogDiskIO, Display, TEXT("Short term stats reset"));
 		bBreak = true;
 	}
 
