Removing some unused includes in coreminimal.h and color.h from stat.h

diff --git a/Engine/Source/Runtime/Core/Public/CoreMinimal.h b/Engine/Source/Runtime/Core/Public/CoreMinimal.h
--- a/Engine/Source/Runtime/Core/Public/CoreMinimal.h
+++ b/Engine/Source/Runtime/Core/Public/CoreMinimal.h
@@ -72,12 +72,16 @@
 #include "Containers/Array.h"
 #include "Misc/FrameNumber.h"
 #include "Misc/Timespan.h"
+#if UE_ENABLE_INCLUDE_ORDER_DEPRECATED_IN_5_4
 #include "HAL/CriticalSection.h"
+#endif
 #include "Containers/StringConv.h"
 #include "UObject/UnrealNames.h"
 #include "UObject/NameTypes.h"
 #include "Misc/Parse.h"
+#if UE_ENABLE_INCLUDE_ORDER_DEPRECATED_IN_5_4
 #include "Math/Color.h"
+#endif
 #include "Templates/AlignmentTemplates.h"
 #include "Misc/StructBuilder.h"
 #include "Templates/Decay.h"
@@ -91,7 +95,9 @@
 #include "Containers/Set.h"
 #include "Algo/Reverse.h"
 #include "Containers/Map.h"
+#if UE_ENABLE_INCLUDE_ORDER_DEPRECATED_IN_5_4
 #include "Math/ColorList.h"
+#endif
 #include "Math/IntPoint.h"
 #include "Math/IntVector.h"
 #include "Logging/LogCategory.h"
@@ -127,8 +133,10 @@
 #include "Math/TwoVectors.h"
 #include "Math/Edge.h"
 #include "UObject/ObjectVersion.h"
+#if UE_ENABLE_INCLUDE_ORDER_DEPRECATED_IN_5_4
 #include "Math/Plane.h"
 #include "Math/Sphere.h"
+#endif
 #include "Math/CapsuleShape.h"
 #include "Math/Rotator.h"
 #include "Misc/DateTime.h"
@@ -156,13 +164,17 @@
 #include "Math/ScaleMatrix.h"
 #include "Math/MirrorMatrix.h"
 #include "Math/ClipProjectionMatrix.h"
+#if UE_ENABLE_INCLUDE_ORDER_DEPRECATED_IN_5_4
 #include "Math/InterpCurvePoint.h"
 #include "Math/InterpCurve.h"
 #include "Math/CurveEdInterface.h"
+#endif
 #include "Math/Float32.h"
 #include "Math/Float16.h"
+#if UE_ENABLE_INCLUDE_ORDER_DEPRECATED_IN_5_4
 #include "Math/Float16Color.h"
 #include "Math/Vector2DHalf.h"
+#endif
 #include "Math/Transform.h"
 #include "Math/ConvexHull2d.h"
 #include "Math/UnrealMath.h"
\ No newline at end of file
diff --git a/Engine/Source/Runtime/Core/Public/Stats/Stats.h b/Engine/Source/Runtime/Core/Public/Stats/Stats.h
--- a/Engine/Source/Runtime/Core/Public/Stats/Stats.h
+++ b/Engine/Source/Runtime/Core/Public/Stats/Stats.h
@@ -3,7 +3,12 @@
 #pragma once
 
 #include "CoreTypes.h"
+#if UE_ENABLE_INCLUDE_ORDER_DEPRECATED_IN_5_4
 #include "Math/Color.h"
+#endif
+#if ENABLE_STATNAMEDEVENTS
+#include "Math/Color.h"
+#endif
 #include "UObject/NameTypes.h"
 #include "Delegates/Delegate.h"
 #include "HAL/PlatformTime.h"
