Deprecated IfAThenAOrB, IfPThenAOrB and XOR.

diff --git a/Engine/Source/Runtime/Core/Public/Async/TaskGraphInterfaces.h b/Engine/Source/Runtime/Core/Public/Async/TaskGraphInterfaces.h
--- a/Engine/Source/Runtime/Core/Public/Async/TaskGraphInterfaces.h
+++ b/Engine/Source/Runtime/Core/Public/Async/TaskGraphInterfaces.h
@@ -1221,7 +1221,7 @@ private:
 
 		// Fire and forget mode must not have subsequents
 		// Track subsequents mode must have subsequents
-		checkThreadGraph(XOR(TTask::GetSubsequentsMode() == ESubsequentsMode::FireAndForget, IsValidRef(Subsequents)));
+		checkThreadGraph((TTask::GetSubsequentsMode() == ESubsequentsMode::FireAndForget) != IsValidRef(Subsequents));
 
 		if (TTask::GetSubsequentsMode() == ESubsequentsMode::TrackSubsequents)
 		{
diff --git a/Engine/Source/Runtime/Core/Public/Containers/BitArray.h b/Engine/Source/Runtime/Core/Public/Containers/BitArray.h
--- a/Engine/Source/Runtime/Core/Public/Containers/BitArray.h
+++ b/Engine/Source/Runtime/Core/Public/Containers/BitArray.h
@@ -1999,8 +1999,16 @@ private:
 	void FindFirstSetBit()
 	{
 		const uint32 EmptyArrayData = 0;
-		const uint32* ArrayDataA = IfAThenAElseB(ArrayA.GetData(),&EmptyArrayData);
-		const uint32* ArrayDataB = IfAThenAElseB(ArrayB.GetData(),&EmptyArrayData);
+		const uint32* ArrayDataA = ArrayA.GetData();
+		if (!ArrayDataA)
+		{
+			ArrayDataA = &EmptyArrayData;
+		}
+		const uint32* ArrayDataB = ArrayB.GetData();
+		if (!ArrayDataB)
+		{
+			ArrayDataB = &EmptyArrayData;
+		}
 
 		// Advance to the next non-zero uint32.
 		uint32 RemainingBitMask;
diff --git a/Engine/Source/Runtime/Core/Public/Containers/ContainerAllocationPolicies.h b/Engine/Source/Runtime/Core/Public/Containers/ContainerAllocationPolicies.h
--- a/Engine/Source/Runtime/Core/Public/Containers/ContainerAllocationPolicies.h
+++ b/Engine/Source/Runtime/Core/Public/Containers/ContainerAllocationPolicies.h
@@ -113,18 +113,6 @@ class FDefaultBitArrayAllocator;
 template<int IndexSize> class TSizedDefaultAllocator;
 using FDefaultAllocator = TSizedDefaultAllocator<32>;
 
-/** branchless pointer selection
-* return A ? A : B;
-**/
-template<typename ReferencedType>
-ReferencedType* IfAThenAElseB(ReferencedType* A,ReferencedType* B);
-
-/** branchless pointer selection based on predicate
-* return PTRINT(Predicate) ? A : B;
-**/
-template<typename PredicateType,typename ReferencedType>
-ReferencedType* IfPThenAElseB(PredicateType Predicate,ReferencedType* A,ReferencedType* B);
-
 template <typename SizeType>
 FORCEINLINE SizeType DefaultCalculateSlackShrink(SizeType NumElements, SizeType NumAllocatedElements, SIZE_T BytesPerElement, bool bAllowQuantize, uint32 Alignment = DEFAULT_ALIGNMENT)
 {
@@ -888,7 +876,11 @@ public:
 		// FContainerAllocatorInterface
 		FORCEINLINE ElementType* GetAllocation() const
 		{
-			return IfAThenAElseB<ElementType>(SecondaryData.GetAllocation(),GetInlineElements());
+			if (ElementType* Result = SecondaryData.GetAllocation())
+			{
+				return Result;
+			}
+			return GetInlineElements();
 		}
 
 		void ResizeAllocation(SizeType PreviousNumElements, SizeType NumElements,SIZE_T NumBytesPerElement)
diff --git a/Engine/Source/Runtime/Core/Public/Templates/UnrealTemplate.h b/Engine/Source/Runtime/Core/Public/Templates/UnrealTemplate.h
--- a/Engine/Source/Runtime/Core/Public/Templates/UnrealTemplate.h
+++ b/Engine/Source/Runtime/Core/Public/Templates/UnrealTemplate.h
@@ -23,6 +23,7 @@
  * @return If the first parameter provided is non-nullptr, it is returned; otherwise the second parameter is returned.
  */
 template<typename ReferencedType>
+UE_DEPRECATED(5.4, "IfAThenAElseB has been deprecated - use regular C++ constructs instead")
 FORCEINLINE ReferencedType* IfAThenAElseB(ReferencedType* A,ReferencedType* B)
 {
 	return A ? A : B;
@@ -32,12 +33,14 @@ FORCEINLINE ReferencedType* IfAThenAElseB(ReferencedType* A,ReferencedType* B)
 * return PTRINT(Predicate) ? A : B;
 **/
 template<typename PredicateType,typename ReferencedType>
+UE_DEPRECATED(5.4, "IfPThenAElseB has been deprecated - use regular C++ constructs instead")
 FORCEINLINE ReferencedType* IfPThenAElseB(PredicateType Predicate,ReferencedType* A,ReferencedType* B)
 {
 	return Predicate ? A : B;
 }
 
 /** A logical exclusive or function. */
+UE_DEPRECATED(5.4, "XOR has been deprecated - use regular C++ constructs instead")
 inline bool XOR(bool A, bool B)
 {
 	return A != B;
diff --git a/Engine/Source/Runtime/Engine/Private/Components/DirectionalLightComponent.cpp b/Engine/Source/Runtime/Engine/Private/Components/DirectionalLightComponent.cpp
--- a/Engine/Source/Runtime/Engine/Private/Components/DirectionalLightComponent.cpp
+++ b/Engine/Source/Runtime/Engine/Private/Components/DirectionalLightComponent.cpp
@@ -175,7 +175,7 @@ void ComputeShadowCullingVolume(bool bReverseCulling, const FVector* CascadeFrus
 			FVector C = A + LightDirection * (A - B).Size();
 
 			// Account for winding
-			if (XOR(DotA >= 0.0f, bReverseCulling))
+			if ((DotA >= 0.0f) != bReverseCulling)
 			{
 				Planes.Add(FPlane(A, B, C));
 			}
diff --git a/Engine/Source/Runtime/Experimental/Iris/Core/Public/Iris/Serialization/NetSerializerArrayStorage.h b/Engine/Source/Runtime/Experimental/Iris/Core/Public/Iris/Serialization/NetSerializerArrayStorage.h
--- a/Engine/Source/Runtime/Experimental/Iris/Core/Public/Iris/Serialization/NetSerializerArrayStorage.h
+++ b/Engine/Source/Runtime/Experimental/Iris/Core/Public/Iris/Serialization/NetSerializerArrayStorage.h
@@ -72,7 +72,14 @@ public:
 
 		ForElementType() {}
 
-		ElementType* GetAllocation() const { return IfAThenAElseB<ElementType>(SecondaryData.GetAllocation(), GetInlineElements()); }
+		ElementType* GetAllocation() const
+		{
+			if (ElementType* Result = SecondaryData.GetAllocation())
+			{
+				return Result;
+			}
+			return GetInlineElements();
+		}
 
 		void Initialize() { SecondaryData.Initialize(); }
 
