VulkanRHI: enable bindless rendering in Linux builds (still disabled by default at runtime by config).

#rnx

diff --git a/Engine/Source/Runtime/Core/Public/Unix/UnixPlatform.h b/Engine/Source/Runtime/Core/Public/Unix/UnixPlatform.h
--- a/Engine/Source/Runtime/Core/Public/Unix/UnixPlatform.h
+++ b/Engine/Source/Runtime/Core/Public/Unix/UnixPlatform.h
@@ -68,6 +68,7 @@ typedef FUnixPlatformTypes FPlatformTypes;
 #define PLATFORM_ALLOW_ALLOCATIONS_IN_FASYNCWRITER_SERIALIZEBUFFERTOARCHIVE 0
 #define PLATFORM_SUPPORTS_VIRTUAL_TEXTURE_STREAMING		1
 #define PLATFORM_USE_ANSI_POSIX_MALLOC					1
+#define PLATFORM_SUPPORTS_BINDLESS_RENDERING			1
 
 #define PLATFORM_RETURN_ADDRESS_FOR_CALLSTACKTRACING	PLATFORM_RETURN_ADDRESS
 
