- Fixed some bugs with set/unsetcvar at specific priorities, and with other platform cvars/DPs
- Added a GetParentObject() function to IConsoleObject, which is only implemented for other-platform cvars so that their name can be found (due to how the name lookup works, we need to get the name of the parent object, not the other-platform object)

diff --git a/Engine/Source/Runtime/Core/Private/HAL/ConsoleManager.cpp b/Engine/Source/Runtime/Core/Private/HAL/ConsoleManager.cpp
--- a/Engine/Source/Runtime/Core/Private/HAL/ConsoleManager.cpp
+++ b/Engine/Source/Runtime/Core/Private/HAL/ConsoleManager.cpp
@@ -685,39 +685,56 @@ public:
 		ValueArray.Emplace(Tag, LocalCopy);
 	}
 	
-	void Unset(EConsoleVariableFlags SetBy, FName Tag)
+	/**
+	 * Unset the value at the SetBy, and optionally Tag.
+	 * Return true if it had been set, or false if nothing happened
+	 */
+	bool Unset(EConsoleVariableFlags SetBy, FName Tag)
 	{
 		if (!History.Contains(SetBy))
 		{
-			return;
+			return false;
 		}
 		
 		TArray<FTaggedHistoryData>& ValueArray = History[SetBy];
 
+		bool bUnsetSomething = false;
+
 		// first remove the value (if no tag, remove all of them)
-		if (IsArrayPriority(SetBy) && Tag != NAME_None)
+		if (Tag != NAME_None)
 		{
-			// look for the tag in the list
-			for (auto It = ValueArray.CreateIterator(); It; ++It)
+			if (IsArrayPriority(SetBy))
 			{
-				if (It->Key == Tag)
+				// look for the tag in the list
+				for (auto It = ValueArray.CreateIterator(); It; ++It)
 				{
-					It.RemoveCurrent();
-					break;
+					if (It->Key == Tag)
+					{
+						bUnsetSomething = true;
+						It.RemoveCurrent();
+						break;
+					}
+				}
+
+				// toss it if it's now empty
+				if (ValueArray.IsEmpty())
+				{
+					History.Remove(SetBy);
 				}
 			}
-			
-			// toss it if it's now empty
-			if (ValueArray.IsEmpty())
+			else if (ValueArray.Num() > 0 && ValueArray[0].Key == Tag)
 			{
+				bUnsetSomething = true;
 				History.Remove(SetBy);
 			}
 		}
-		else
+		else if (ValueArray.Num() > 0)
 		{
+			bUnsetSomething = true;
 			History.Remove(SetBy);
 		}
-		
+
+		return bUnsetSomething;
 	}
 	
 	const FHistoryData& GetMaxValue(EConsoleVariableFlags& MaxSetBy)
@@ -923,14 +940,18 @@ protected:
 		// this isn't ideal because it could call SetInternal multiple times
 		if (SetBy == ECVF_SetByMask)
 		{
-			#define RECURSE(x) if (IsArrayPriority(ECVF_SetBy##x)) { Unset(ECVF_SetBy##x, Tag); }
+			#define RECURSE(x) Unset(ECVF_SetBy##x, Tag);
 			ENUMERATE_SET_BY(RECURSE)
 			#undef RECURSE
 			
 			return;
 		}
 		
-		PriorityHistory->Unset(SetBy, Tag);
+		// if nothing was unset, no need to perform any more actions
+		if (PriorityHistory->Unset(SetBy, Tag) == false)
+		{
+			return;
+		}
 		
 		uint32 CurrentPri =	(uint32)this->Flags & ECVF_SetByMask;
 		uint32 UnsetPri =	(uint32)SetBy & ECVF_SetByMask;
@@ -948,7 +969,7 @@ protected:
 				NewSetBy = (EConsoleVariableFlags)(NewSetBy | ECVF_Set_SetOnly_Unsafe);
 			}
 			
-			UE_LOG(LogConsoleManager, Display, TEXT(" |-> Unsetting %s = %s"), *IConsoleManager::Get().FindConsoleObjectName(this),
+			UE_LOG(LogConsoleManager, Display, TEXT(" |-> Unsetting %s, now %s"), *IConsoleManager::Get().FindConsoleObjectName(this),
 				*TTypeToString<T>::ToString(MaxValue.GetValueOnGameThread()));
 
 			// and force it to the new value and call any set callbacks
@@ -1057,9 +1078,12 @@ class FConsoleVariable : public FConsoleVariableExtendedData<T>
 	using FConsoleVariableBase::Flags;
 
 public:
-	FConsoleVariable(T DefaultValue, const TCHAR* Help, EConsoleVariableFlags Flags)
+	FConsoleVariable(T DefaultValue, const TCHAR* Help, EConsoleVariableFlags Flags, IConsoleVariable* Parent=nullptr)
 		: FConsoleVariableExtendedData<T>(DefaultValue, Help, Flags)
 		, Data(DefaultValue)
+#if ALLOW_OTHER_PLATFORM_CONFIG
+		, ParentVariable(Parent)
+#endif
 	{
 	}
 
@@ -1096,9 +1120,19 @@ public:
 	virtual class TConsoleVariableData<float>* AsVariableFloat() override { return nullptr; }
 	virtual class TConsoleVariableData<FString>* AsVariableString() override { return nullptr; }
 
+#if ALLOW_OTHER_PLATFORM_CONFIG
+	virtual IConsoleObject* GetParentObject() const
+	{
+		return ParentVariable;
+	}
+#endif
+
 private: // ----------------------------------------------------
 
 	TConsoleVariableData<T> Data;
+#if ALLOW_OTHER_PLATFORM_CONFIG
+	IConsoleVariable* ParentVariable;
+#endif
 
 	const T &Value() const
 	{
@@ -1201,7 +1235,7 @@ FConsoleVariable<T>* FindOrCreateTypedPlatformCVar(FConsoleVariableExtendedData<
 	TSharedPtr<IConsoleVariable> PlatformCVar = CVar->PlatformValues.FindRef(PlatformKey);
 	if (!PlatformCVar.IsValid())
 	{
-		PlatformCVar = TSharedPtr<IConsoleVariable>(new FConsoleVariable(CVar->GetDefaultTypedValue(), TEXT("Platform CVar copy"), CVar->GetFlags()));
+		PlatformCVar = TSharedPtr<IConsoleVariable>(new FConsoleVariable(CVar->GetDefaultTypedValue(), TEXT("Platform CVar copy"), CVar->GetFlags(), CVar));
 		
 		// cache it
 		CVar->PlatformValues.Add(PlatformKey, PlatformCVar);
@@ -2336,10 +2370,18 @@ static void SetUnsetCVar(const TMap<FString, IConsoleObject*>& ConsoleObjects, c
 	}
 	
 	FString PlatformName;
+	FString DeviceProfileName;
 	int32 PlatformDelim = CVarName.Find(TEXT("@"));
 	if (PlatformDelim > 0)
 	{
 		PlatformName = CVarName.Mid(0, PlatformDelim);
+		if (PlatformName.Contains(TEXT("/")))
+		{
+			FString Plat;
+			PlatformName.Split(TEXT("/"), &Plat, &DeviceProfileName);
+			PlatformName = *Plat;
+		}
+
 		CVarName = CVarName.Mid(PlatformDelim + 1);
 	}
 
@@ -2356,7 +2398,7 @@ static void SetUnsetCVar(const TMap<FString, IConsoleObject*>& ConsoleObjects, c
 	// get platform version
 	if (PlatformName.Len())
 	{
-		CVar = CVar->GetPlatformValueVariable(*PlatformName).Get();
+		CVar = CVar->GetPlatformValueVariable(*PlatformName, *DeviceProfileName).Get();
 		if (CVar == nullptr)
 		{
 			Ar.Logf(TEXT("Failed to get CVar for platform %s"), *PlatformName);
@@ -2795,6 +2837,12 @@ FString FConsoleManager::FindConsoleObjectName(const IConsoleObject* InVar) cons
 		}
 	}
 
+	// if we didn't find one, and it has a parent, then give that a try
+	if (InVar->GetParentObject() != nullptr)
+	{
+		return FindConsoleObjectName(InVar->GetParentObject());
+	}
+
 	return FString();
 }
 
@@ -2943,7 +2991,7 @@ void FConsoleManager::UnsetAllConsoleVariablesWithTag(FName Tag, EConsoleVariabl
 	
 	for (IConsoleVariable* Var : *TaggedSet)
 	{
-		Var->Unset(Priority);
+		Var->Unset(Priority, Tag);
 	}
 	
 	UE::ConsoleManager::Private::TaggedCVars.Remove(Tag);
diff --git a/Engine/Source/Runtime/Core/Public/HAL/IConsoleManager.h b/Engine/Source/Runtime/Core/Public/HAL/IConsoleManager.h
--- a/Engine/Source/Runtime/Core/Public/HAL/IConsoleManager.h
+++ b/Engine/Source/Runtime/Core/Public/HAL/IConsoleManager.h
@@ -409,6 +409,14 @@ public:
 		return ((uint32)GetFlags() & (uint32)Value) != 0;
 	}
 
+	/**
+	 * If the object has a parent (for instance the main cvar that owns an other-platform cvar), return it
+	 */
+	virtual IConsoleObject* GetParentObject() const
+	{
+		return nullptr;
+	}
+	
 	/**
 	 * Casts this object to an IConsoleVariable, returns 0 if it's not
 	 */
