[Skeletal Mesh] fix for crash when loading cooked skeletal mesh since the introduction of halfedge buffers

#rnx

diff --git a/Engine/Source/Runtime/Core/Private/UObject/FortniteMainBranchObjectVersion.cpp b/Engine/Source/Runtime/Core/Private/UObject/FortniteMainBranchObjectVersion.cpp
--- a/Engine/Source/Runtime/Core/Private/UObject/FortniteMainBranchObjectVersion.cpp
+++ b/Engine/Source/Runtime/Core/Private/UObject/FortniteMainBranchObjectVersion.cpp
@@ -12,7 +12,7 @@ TMap<FGuid, FGuid> FFortniteMainBranchObjectVersion::GetSystemGuids()
 	SystemGuids.Add(DevGuids.NANITE_DERIVEDDATA_VER, FGuid("D28C003F825F4C178F698938FB0E47F7"));
 	SystemGuids.Add(DevGuids.NIAGARASHADERMAP_DERIVEDDATA_VER, FGuid("C417A50271B3427D98785388266E3B66"));
 	SystemGuids.Add(DevGuids.Niagara_LatestScriptCompileVersion, FGuid("CE33E29C6268469E8E18F1D96A55DFDA"));
-	SystemGuids.Add(DevGuids.SkeletalMeshDerivedDataVersion, FGuid("3BD78FFFB88E4DDF943FB2608A20C0D2"));
+	SystemGuids.Add(DevGuids.SkeletalMeshDerivedDataVersion, FGuid("69ABAF00AC9648E2A2EE57D4CC449CF3"));
 	SystemGuids.Add(DevGuids.STATICMESH_DERIVEDDATA_VER, FGuid("543843B103794E0E9BA4BE80FB602F79"));
 	SystemGuids.Add(DevGuids.MaterialTranslationDDCVersion, FGuid("BBD650804D314CE7A44ADF9BB4F3F5F9"));
 	return SystemGuids;
diff --git a/Engine/Source/Runtime/Engine/Private/SkeletalMeshLODRenderData.cpp b/Engine/Source/Runtime/Engine/Private/SkeletalMeshLODRenderData.cpp
--- a/Engine/Source/Runtime/Engine/Private/SkeletalMeshLODRenderData.cpp
+++ b/Engine/Source/Runtime/Engine/Private/SkeletalMeshLODRenderData.cpp
@@ -868,13 +868,14 @@ void FSkeletalMeshLODRenderData::SerializeStreamedData(FArchive& Ar, USkinnedAss
 	
 	if (Ar.CustomVer(FFortniteMainBranchObjectVersion::GUID) >= FFortniteMainBranchObjectVersion::SkeletalHalfEdgeData)
 	{
-		uint8 MeshDeformerStripFlag = 0;
+		uint8 ClassStripFlag = 0;
+		const uint8 MeshDeformerStripFlag = 1;
 		if (Ar.IsCooking() && !IsMeshDeformerSupportedByCookTargetPlatform(*Ar.CookingTarget()))
 		{
-			MeshDeformerStripFlag = 1;
+			ClassStripFlag |= MeshDeformerStripFlag;
 		}
-
-		FStripDataFlags MeshDeformerStripFlags(Ar, MeshDeformerStripFlag);
+		
+		FStripDataFlags MeshDeformerStripFlags(Ar, ClassStripFlag);
 		
 		if (!MeshDeformerStripFlags.IsClassDataStripped(MeshDeformerStripFlag))
 		{
