Allow tweaking the CSV stream buffering

#rnx

diff --git a/Engine/Source/Runtime/Core/Private/ProfilingDebugging/CsvProfiler.cpp b/Engine/Source/Runtime/Core/Private/ProfilingDebugging/CsvProfiler.cpp
--- a/Engine/Source/Runtime/Core/Private/ProfilingDebugging/CsvProfiler.cpp
+++ b/Engine/Source/Runtime/Core/Private/ProfilingDebugging/CsvProfiler.cpp
@@ -144,6 +144,13 @@ TAutoConsoleVariable<int32> CVarCsvWriteBufferSize(
 	ECVF_Default
 );
 
+TAutoConsoleVariable<int32> CVarCsvStreamFramesToBuffer(
+	TEXT("csv.FramesToBuffer"),
+	128,
+	TEXT("Defines the minimum amount of frames to keep in memory before flushing them."),
+	ECVF_Default
+);
+
 static bool GCsvUseProcessingThread = true;
 static int32 GCsvRepeatCount = 0;
 static int32 GCsvRepeatFrameCount = 0;
@@ -1847,7 +1854,7 @@ class FCsvStreamWriter
 	// frame index less than (WriteFrameIndex - NumFramesToBuffer). NumFramesToBuffer should be large enough to avoid
 	// flushing rows before all the timestamps for that frame have been processed, but small enough to avoid the
 	// additional memory overhead of holding addition rows in memory unnecessarily.
-	const int64 NumFramesToBuffer = 128;
+	int64 NumFramesToBuffer;
 	int64 WriteFrameIndex;
 	int64 ReadFrameIndex;
 
@@ -1861,7 +1868,7 @@ class FCsvStreamWriter
 	uint32 RHIThreadId;
 
 public:
-	FCsvStreamWriter(const TSharedRef<FArchive>& InOutputFile, bool bInContinuousWrites, int32 InBufferSize, bool bInCompressOutput, uint32 RenderThreadId, uint32 RHIThreadId);
+	FCsvStreamWriter(const TSharedRef<FArchive>& InOutputFile, bool bInContinuousWrites, int32 InBufferSize, int64 InNumFramesToBuffer, bool bInCompressOutput, uint32 RenderThreadId, uint32 RHIThreadId);
 	~FCsvStreamWriter();
 
 	void AddSeries(FCsvStatSeries* Series);
@@ -2283,8 +2290,9 @@ private:
 	}
 };
 
-FCsvStreamWriter::FCsvStreamWriter(const TSharedRef<FArchive>& InOutputFile, bool bInContinuousWrites, int32 InBufferSize, bool bInCompressOutput, uint32 InRenderThreadId, uint32 InRHIThreadId)
+FCsvStreamWriter::FCsvStreamWriter(const TSharedRef<FArchive>& InOutputFile, bool bInContinuousWrites, int32 InBufferSize, int64 InNumFramesToBuffer, bool bInCompressOutput, uint32 InRenderThreadId, uint32 InRHIThreadId)
 	: Stream(InOutputFile, InBufferSize, bInCompressOutput)
+	, NumFramesToBuffer(InNumFramesToBuffer)
 	, WriteFrameIndex(-1)
 	, ReadFrameIndex(-1)
 	, bContinuousWrites(bInContinuousWrites)
@@ -2968,8 +2976,8 @@ void FCsvProfiler::BeginFrame()
 				}
 				else
 				{
-					
-					CsvWriter = new FCsvStreamWriter(OutputFile.ToSharedRef(), bContinuousWrites, BufferSize, bCompressOutput, RenderThreadId, RHIThreadId);
+					int64 NumFramesToBuffer = CVarCsvStreamFramesToBuffer.GetValueOnAnyThread();
+					CsvWriter = new FCsvStreamWriter(OutputFile.ToSharedRef(), bContinuousWrites, BufferSize, NumFramesToBuffer, bCompressOutput, RenderThreadId, RHIThreadId);
 
 					NumFramesToCapture = CurrentCommand.Value;
 					GCsvRepeatFrameCount = NumFramesToCapture;
