Version 1.19 of UnrealTraceServer for windows

UnrealTraceServer executable and debug files.
Version: 1.19
Platform: windows, mac, linux
Architecture: x64 (windows, linux), x64/arm64 (mac)
Compiled with musl: false (windows, mac), true (linux)

diff --git a/Engine/Source/Runtime/Core/Private/ProfilingDebugging/TraceAuxiliary.cpp b/Engine/Source/Runtime/Core/Private/ProfilingDebugging/TraceAuxiliary.cpp
--- a/Engine/Source/Runtime/Core/Private/ProfilingDebugging/TraceAuxiliary.cpp
+++ b/Engine/Source/Runtime/Core/Private/ProfilingDebugging/TraceAuxiliary.cpp
@@ -1866,7 +1866,7 @@ static bool LaunchTraceServerCommand(ELaunchTraceServerCommand Command, bool bAd
 ////////////////////////////////////////////////////////////////////////////////
 bool FTraceServerControls::Start()
 {
-	return LaunchTraceServerCommand(ELaunchTraceServerCommand::Fork, false);
+	return LaunchTraceServerCommand(ELaunchTraceServerCommand::Fork, true);
 }
 
 ////////////////////////////////////////////////////////////////////////////////
