Add new RecordEventAtFrameStart method to the CSV profiler

#rnx
#tests Works as advertized on a local build

diff --git a/Engine/Source/Runtime/Core/Private/ProfilingDebugging/CsvProfiler.cpp b/Engine/Source/Runtime/Core/Private/ProfilingDebugging/CsvProfiler.cpp
--- a/Engine/Source/Runtime/Core/Private/ProfilingDebugging/CsvProfiler.cpp
+++ b/Engine/Source/Runtime/Core/Private/ProfilingDebugging/CsvProfiler.cpp
@@ -3610,6 +3610,11 @@ void FCsvProfiler::RecordEvent(int32 CategoryIndex, const FString& EventText)
 	}
 }
 
+void FCsvProfiler::RecordEventAtFrameStart(int32 CategoryIndex, const FString& EventText)
+{
+	RecordEventAtTimestamp(CategoryIndex, EventText, FCsvProfiler::Get()->LastEndFrameTimestamp);
+}
+
 void FCsvProfiler::SetMetadata(const TCHAR* Key, const TCHAR* Value)
 {
 	FCsvProfiler::Get()->SetMetadataInternal(Key, Value, true, EMetadataPersistenceType::Persistent);
diff --git a/Engine/Source/Runtime/Core/Public/ProfilingDebugging/CsvProfiler.h b/Engine/Source/Runtime/Core/Public/ProfilingDebugging/CsvProfiler.h
--- a/Engine/Source/Runtime/Core/Public/ProfilingDebugging/CsvProfiler.h
+++ b/Engine/Source/Runtime/Core/Public/ProfilingDebugging/CsvProfiler.h
@@ -369,6 +369,7 @@ public:
 	CORE_API static void RecordCustomStat(const FName& StatName, uint32 CategoryIndex, int32 Value, const ECsvCustomStatOp CustomStatOp);
 
 	CORE_API static void RecordEvent(int32 CategoryIndex, const FString& EventText);
+	CORE_API static void RecordEventAtFrameStart(int32 CategoryIndex, const FString& EventText);
 	CORE_API static void RecordEventAtTimestamp(int32 CategoryIndex, const FString& EventText, uint64 Cycles64);
 
 	/** Metadata values set with this function will persist between captures. */
