Make the CSV profiler free up temp stats memory when slack data exceeds a threshold in bytes, specified with csv.MaxPerThreadStatDataSlackKB (default 64KB). Setting to 0 cleans up slack immediately but results in a lot more allocations (and can block the memory allocator  excessively in some cases, impacting game performance). This fixes issues on low end platforms where starvation of the CSV profiling thread can result in excessive slack memory that is not cleaned up until capture end.

Also add a debug cvar csv.PauseProcessingThread to pause the processing thread, simulating CPU starvation.

#tests Tested on Switch

diff --git a/Engine/Source/Runtime/Core/Private/ProfilingDebugging/CsvProfiler.cpp b/Engine/Source/Runtime/Core/Private/ProfilingDebugging/CsvProfiler.cpp
--- a/Engine/Source/Runtime/Core/Private/ProfilingDebugging/CsvProfiler.cpp
+++ b/Engine/Source/Runtime/Core/Private/ProfilingDebugging/CsvProfiler.cpp
@@ -151,6 +151,21 @@ TAutoConsoleVariable<int32> CVarCsvStreamFramesToBuffer(
 	ECVF_Default
 );
 
+TAutoConsoleVariable<int32> CVarCsvPauseProcessingThread(
+	TEXT("csv.PauseProcessingThread"),
+	0,
+	TEXT("Debug only - When 1, blocks the processing thread to simulate starvation"),
+	ECVF_Default
+);
+
+TAutoConsoleVariable<int32> CVarMaxPerThreadStatDataSlackKB(
+	TEXT("csv.MaxPerThreadStatDataSlackKB"),
+	64,
+	TEXT("Max amount of per thread slack data to allow during a capture.\r\n")
+	TEXT("Higher values result in better performance due to fewer allocations but higher memory overhead"),
+	ECVF_Default
+);
+
 static bool GCsvUseProcessingThread = true;
 static int32 GCsvRepeatCount = 0;
 static int32 GCsvRepeatFrameCount = 0;
@@ -946,13 +961,20 @@ public:
 	}
 
 	// Called from the consumer thread
-	void PopAll(TArray<T>& ElementsOut)
+	void PopAll(TArray<T>& ElementsOut, int64 MaxSlackMemBytes = -1 )
 	{
 		volatile uint64 CurrentCounterValue = Counter;
 		FPlatformMisc::MemoryBarrier();
 
 		uint32 MaxElementsToPop = uint32(CurrentCounterValue - ConsumerThreadReadIndex);
 
+		// Shrink the output array if it has excessive slack (note: usually this will shrink to 0)
+		int64 SlackMemBytes = ( (int64)ElementsOut.GetSlack() - (int64)MaxElementsToPop ) * (int64)sizeof(T);
+		if ( MaxSlackMemBytes >= 0 && SlackMemBytes > MaxSlackMemBytes )
+		{
+			ElementsOut.Shrink();
+		}
+
 		// Presize the array capacity to avoid memory reallocation.
 		ElementsOut.Reserve(ElementsOut.Num() + MaxElementsToPop);
 
@@ -2050,9 +2072,11 @@ public:
 		
 		check(IsInCsvProcessingThread());
 
-		TimingMarkers.PopAll(OutMarkers);
-		CustomStats.PopAll(OutCustomStats);
-		Events.PopAll(OutEvents);
+		int64 MaxSlackMemBytes = (int64)CVarMaxPerThreadStatDataSlackKB.GetValueOnAnyThread() * 1024;
+
+		TimingMarkers.PopAll(OutMarkers, MaxSlackMemBytes);
+		CustomStats.PopAll(OutCustomStats, MaxSlackMemBytes);
+		Events.PopAll(OutEvents, MaxSlackMemBytes);
 	}
 
 	CSV_PROFILER_INLINE void AddTimestampBegin(const char* StatName, int32 CategoryIndex)
@@ -2555,6 +2579,11 @@ public:
 
 		while (StopCounter.GetValue() == 0)
 		{
+			if (CVarCsvPauseProcessingThread.GetValueOnAnyThread())
+			{
+				FPlatformProcess::Sleep(TimeBetweenUpdatesMS / 1000.0f);
+				continue;
+			}
 			float ElapsedMS = CsvProfiler.ProcessStatData();
 
 			if (GCsvProfilerIsWritingFile)
