Adds missing GetId() member from TSet::TBaseKeyIterator.

Without this implementation, it's not possible to use GetId() on TMap::TKeyIterators/TConstKeyIterators

This missing implementation is specifically helpful for cases where you'd like to find the KV Pair ElementType for a given KeyArgumentType in cases where you may have collisions between keys and the actual KeyArgumentType value is important.

Example case is when dealing with FNames in Editor since they compare on case-insensitive CompareIds but have unique DisplayIds for different casing. If you want to know the casing used when inserted into a map, you can now do the following:

TMap<FName, FAssetPackageData*>::TConstKeyIterator It = CachedPackageData.CreateConstKeyIterator(PackageFNameOfAnyCase);
FSetElementId Id = It.GetId();
if (!Id.IsValidId())
{
    return nullptr;
}

const TPair<FName, FAssetPackageData*> KVPair = CachedPackageData.Get(Id);
FName PackageNameWithCasingUsedUponInsertion = KVPair.Key;

return KVPair.Value;

diff --git a/Engine/Source/Runtime/Core/Private/Containers/ContainersTest.cpp b/Engine/Source/Runtime/Core/Private/Containers/ContainersTest.cpp
--- a/Engine/Source/Runtime/Core/Private/Containers/ContainersTest.cpp
+++ b/Engine/Source/Runtime/Core/Private/Containers/ContainersTest.cpp
@@ -434,6 +434,7 @@ namespace
 					{
 						check(!FoundConstKey);
 						check(It->Key == Key);
+						check(Cont.Get(It.GetId()) == *It);
 						FoundConstKey = &It->Key;
 					}
 
@@ -443,6 +444,7 @@ namespace
 					{
 						check(!FoundKey);
 						check(It->Key == Key);
+						check(Cont.Get(It.GetId()) == *It);
 						FoundKey = &It->Key;
 						It->Value.Str = ReplacedValue;
 					}
diff --git a/Engine/Source/Runtime/Core/Public/Containers/Set.h b/Engine/Source/Runtime/Core/Public/Containers/Set.h
--- a/Engine/Source/Runtime/Core/Public/Containers/Set.h
+++ b/Engine/Source/Runtime/Core/Public/Containers/Set.h
@@ -1688,6 +1688,10 @@ private:
 		}
 
 		// Accessors.
+		[[nodiscard]] FORCEINLINE FSetElementId GetId() const
+		{
+			return FSetElementId(Index);
+		}
 		FORCEINLINE ItElementType* operator->() const
 		{
 			return &Set.Elements[Index].Value;
