Remove unnecessary disable opt on autortfm_did_free.

diff --git a/Engine/Source/Runtime/Core/Private/AutoRTFM/API.cpp b/Engine/Source/Runtime/Core/Private/AutoRTFM/API.cpp
--- a/Engine/Source/Runtime/Core/Private/AutoRTFM/API.cpp
+++ b/Engine/Source/Runtime/Core/Private/AutoRTFM/API.cpp
@@ -270,7 +270,6 @@ extern "C" UE_AUTORTFM_AUTORTFM("RTFM_autortfm_did_allocate") void* autortfm_did
     return Ptr;
 }
 
-UE_DISABLE_OPTIMIZATION
 extern "C" UE_AUTORTFM_AUTORTFM("RTFM_autortfm_did_free") void autortfm_did_free(void* Ptr)
 {
 	// We only need to process did free if we need to track allocation locations.
@@ -294,7 +293,6 @@ extern "C" UE_AUTORTFM_AUTORTFM("RTFM_autortfm_did_free") void autortfm_did_free
 		}
 	}
 }
-UE_ENABLE_OPTIMIZATION
 
 extern "C" UE_AUTORTFM_AUTORTFM("RTFM_autortfm_check_consistency_assuming_no_races") void autortfm_check_consistency_assuming_no_races()
 {
