Fix broken crash user data on out-of-memory

#tests Triggered an OOM locally and confirmed the correct strings in the core dump
#rnx

diff --git a/Engine/Source/Runtime/Core/Private/GenericPlatform/GenericPlatformCrashContext.cpp b/Engine/Source/Runtime/Core/Private/GenericPlatform/GenericPlatformCrashContext.cpp
--- a/Engine/Source/Runtime/Core/Private/GenericPlatform/GenericPlatformCrashContext.cpp
+++ b/Engine/Source/Runtime/Core/Private/GenericPlatform/GenericPlatformCrashContext.cpp
@@ -508,6 +508,11 @@ void FGenericCrashContext::InitializeFromContext(const FSessionContext& Session,
 	bIsInitialized = true;
 }
 
+const FSessionContext& FGenericCrashContext::GetCachedSessionContext()
+{
+	return NCached::Session;
+}
+
 void FGenericCrashContext::CopySharedCrashContext(FSharedCrashContext& Dst)
 {
 	//Copy the session
diff --git a/Engine/Source/Runtime/Core/Public/GenericPlatform/GenericPlatformCrashContext.h b/Engine/Source/Runtime/Core/Public/GenericPlatform/GenericPlatformCrashContext.h
--- a/Engine/Source/Runtime/Core/Public/GenericPlatform/GenericPlatformCrashContext.h
+++ b/Engine/Source/Runtime/Core/Public/GenericPlatform/GenericPlatformCrashContext.h
@@ -401,6 +401,9 @@ public:
 	/** Initialized crash context, using a crash context (e.g. shared from another process). */
 	CORE_API static void InitializeFromContext(const FSessionContext& Context, const TCHAR* EnabledPlugins, const TCHAR* EngineData, const TCHAR* GameData);
 
+	/** Get the current cached session context */
+	CORE_API static const FSessionContext& GetCachedSessionContext();
+
 	/**
 	 * @return true, if the generic crash context has been initialized.
 	 */
