Added "Localization.TurkishUsesCapitalI" to control whether FText::ToUpper uses "big i" or I

Default is to continue to use "big i"

diff --git a/Engine/Source/Runtime/Core/Private/Internationalization/ICUText.cpp b/Engine/Source/Runtime/Core/Private/Internationalization/ICUText.cpp
--- a/Engine/Source/Runtime/Core/Private/Internationalization/ICUText.cpp
+++ b/Engine/Source/Runtime/Core/Private/Internationalization/ICUText.cpp
@@ -12,6 +12,7 @@
 
 #if UE_ENABLE_ICU
 #include "Internationalization/TextHistory.h"
+#include "HAL/IConsoleManager.h"
 
 THIRD_PARTY_INCLUDES_START
 	#include <unicode/utypes.h>
@@ -26,6 +27,13 @@ THIRD_PARTY_INCLUDES_END
 #include "Internationalization/ICUInternationalization.h"
 #include "Internationalization/ICUTextCharacterIterator.h"
 
+static TAutoConsoleVariable<bool> CVarTurkishUsesCapitalI(
+	TEXT("Localization.TurkishUsesCapitalI"),
+	false,
+	TEXT("False: Use 'big i' (CLDR format, default), True: Use I."),
+	ECVF_Default
+	);
+
 FString FTextChronoFormatter::AsDate(const FDateTime& DateTime, const EDateTimeStyle::Type DateStyle, const FString& TimeZone, const FCulture& TargetCulture)
 {
 	FInternationalization& I18N = FInternationalization::Get();
@@ -80,12 +88,26 @@ FString FTextChronoFormatter::AsDateTime(const FDateTime& DateTime, const FStrin
 
 FString FTextTransformer::ToLower(const FString& InStr)
 {
+	FInternationalization& I18N = FInternationalization::Get();
+	checkf(I18N.IsInitialized() == true, TEXT("FInternationalization is not initialized. An FText formatting method was likely used in static object initialization - this is not supported."));
+	
 	return ICUUtilities::ConvertString(ICUUtilities::ConvertString(InStr).toLower());
 }
 
 FString FTextTransformer::ToUpper(const FString& InStr)
 {
-	return ICUUtilities::ConvertString(ICUUtilities::ConvertString(InStr).toUpper());
+	FInternationalization& I18N = FInternationalization::Get();
+	checkf(I18N.IsInitialized() == true, TEXT("FInternationalization is not initialized. An FText formatting method was likely used in static object initialization - this is not supported."));
+
+	FString Result = ICUUtilities::ConvertString(ICUUtilities::ConvertString(InStr).toUpper());
+
+	// The CLDR uses "big i" as the uppercase i in Turkish, however some styles prefer using I
+	if (I18N.GetCurrentLanguage()->GetTwoLetterISOLanguageName() == TEXT("tr") && CVarTurkishUsesCapitalI.AsVariable()->GetBool())
+	{
+		Result.ReplaceCharInline(TEXT('\u0130'), TEXT('I'), ESearchCase::CaseSensitive);
+	}
+
+	return Result;
 }
 
 bool FText::IsWhitespace(const TCHAR Char)
