#ue5 - Optimize TQuat<T>::Slerp_NotNormalized(). Removes 1 branch, based on negative angle identity for Sine function. Added startup tests to validate old versus new function and also test both versions for both positive and negative 'RawCosom' values (not previously tested).

#10448
#tests startup tests and PIE with player animation stressing Slerp()

diff --git a/Engine/Source/Runtime/Core/Private/Math/UnrealMath.cpp b/Engine/Source/Runtime/Core/Private/Math/UnrealMath.cpp
--- a/Engine/Source/Runtime/Core/Private/Math/UnrealMath.cpp
+++ b/Engine/Source/Runtime/Core/Private/Math/UnrealMath.cpp
@@ -1224,41 +1224,32 @@ namespace Math
 	TQuat<T> TQuat<T>::Slerp_NotNormalized(const TQuat<T>& Quat1, const TQuat<T>& Quat2, T Slerp)
 	{
 		// Get cosine of angle between quats.
-		const T RawCosom =
+		T RawCosom =
 			Quat1.X * Quat2.X +
 			Quat1.Y * Quat2.Y +
 			Quat1.Z * Quat2.Z +
 			Quat1.W * Quat2.W;
-		// Unaligned quats - compensate, results in taking shorter route.
-		const T Cosom = FMath::FloatSelect(RawCosom, RawCosom, -RawCosom);
-
-		T Scale0, Scale1;
 
-		if (Cosom < T(0.9999f))
-		{
-			const T Omega = FMath::Acos(Cosom);
-			const T InvSin = T(1.f) / FMath::Sin(Omega);
-			Scale0 = FMath::Sin((T(1.f) - Slerp) * Omega) * InvSin;
-			Scale1 = FMath::Sin(Slerp * Omega) * InvSin;
-		}
-		else
+		// Unaligned quats - compensate, results in taking shorter route.
+		const T Sign = FMath::FloatSelect(RawCosom, static_cast<T>(1.0), static_cast<T>(-1.0));
+		RawCosom *= Sign;
+		
+		T Scale0 = static_cast<T>(1.0) - Slerp;
+		T Scale1 = Slerp * Sign;
+		
+		if (RawCosom < static_cast<T>(0.9999))
 		{
-			// Use linear interpolation.
-			Scale0 = T(1.0f) - Slerp;
-			Scale1 = Slerp;
+			const T Omega = FMath::Acos(RawCosom);
+			const T InvSin = static_cast<T>(1.0) / FMath::Sin(Omega);
+			Scale0 = FMath::Sin(Scale0 * Omega) * InvSin;
+			Scale1 = FMath::Sin(Scale1 * Omega) * InvSin;
 		}
-
-		// In keeping with our flipped Cosom:
-		Scale1 = FMath::FloatSelect(RawCosom, Scale1, -Scale1);
-
-		TQuat<T> Result;
-
-		Result.X = Scale0 * Quat1.X + Scale1 * Quat2.X;
-		Result.Y = Scale0 * Quat1.Y + Scale1 * Quat2.Y;
-		Result.Z = Scale0 * Quat1.Z + Scale1 * Quat2.Z;
-		Result.W = Scale0 * Quat1.W + Scale1 * Quat2.W;
-
-		return Result;
+		
+		return TQuat<T>(
+			Scale0 * Quat1.X + Scale1 * Quat2.X,
+			Scale0 * Quat1.Y + Scale1 * Quat2.Y,
+			Scale0 * Quat1.Z + Scale1 * Quat2.Z,
+			Scale0 * Quat1.W + Scale1 * Quat2.W);
 	}
 
 	template<typename T>
diff --git a/Engine/Source/Runtime/Core/Tests/Math/UnrealMathTest.cpp b/Engine/Source/Runtime/Core/Tests/Math/UnrealMathTest.cpp
--- a/Engine/Source/Runtime/Core/Tests/Math/UnrealMathTest.cpp
+++ b/Engine/Source/Runtime/Core/Tests/Math/UnrealMathTest.cpp
@@ -794,6 +794,54 @@ FQuat4f FindBetweenNormals_5_2(const FVector3f& A, const FVector3f& B)
 	return FindBetween_Helper_5_2(A, B, NormAB);
 }
 
+
+template<typename T>
+UE::Math::TQuat<T> Old_Slerp_NotNormalized(const UE::Math::TQuat<T>& Quat1, const UE::Math::TQuat<T>& Quat2, T Slerp)
+{
+	// Get cosine of angle between quats.
+	const T RawCosom =
+		Quat1.X * Quat2.X +
+		Quat1.Y * Quat2.Y +
+		Quat1.Z * Quat2.Z +
+		Quat1.W * Quat2.W;
+	// Unaligned quats - compensate, results in taking shorter route.
+	const T Cosom = FMath::FloatSelect(RawCosom, RawCosom, -RawCosom);
+
+	T Scale0, Scale1;
+
+	if (Cosom < T(0.9999f))
+	{
+		const T Omega = FMath::Acos(Cosom);
+		const T InvSin = T(1.f) / FMath::Sin(Omega);
+		Scale0 = FMath::Sin((T(1.f) - Slerp) * Omega) * InvSin;
+		Scale1 = FMath::Sin(Slerp * Omega) * InvSin;
+	}
+	else
+	{
+		// Use linear interpolation.
+		Scale0 = T(1.0f) - Slerp;
+		Scale1 = Slerp;
+	}
+
+	// In keeping with our flipped Cosom:
+	Scale1 = FMath::FloatSelect(RawCosom, Scale1, -Scale1);
+
+	UE::Math::TQuat<T> Result;
+
+	Result.X = Scale0 * Quat1.X + Scale1 * Quat2.X;
+	Result.Y = Scale0 * Quat1.Y + Scale1 * Quat2.Y;
+	Result.Z = Scale0 * Quat1.Z + Scale1 * Quat2.Z;
+	Result.W = Scale0 * Quat1.W + Scale1 * Quat2.W;
+
+	return Result;
+}
+
+template<typename T>
+UE::Math::TQuat<T> Old_Slerp(const UE::Math::TQuat<T>& Quat1, const UE::Math::TQuat<T>& Quat2, T Slerp)
+{
+	return Old_Slerp_NotNormalized(Quat1, Quat2, Slerp).GetNormalized();
+}
+
 // ROTATOR TESTS
 
 bool TestRotatorEqual0(const FRotator3f& A, const FRotator3f& B, const float Tolerance)
@@ -3493,6 +3541,7 @@ TEST_CASE_NAMED(FVectorRegisterAbstractionTest, "System::Core::Math::Vector Regi
 
 	// Quat Lerp/Slerp
 	{
+		// Note: Values chosen to have a positive 'RawCosom' within Slerp_NotNormalized()
 		const FRotator3f Rotator0 = FRotator3f(300.0f, -45.0f, 0.0f);
 		const FRotator3f Rotator1 = FRotator3f(10.0f, 270.0f, 30.0f);
 		const float FloatAlpha = 0.25f;
@@ -3506,41 +3555,58 @@ TEST_CASE_NAMED(FVectorRegisterAbstractionTest, "System::Core::Math::Vector Regi
 		Q2 = FMath::Lerp(Q0, Q1, FloatAlpha);
 		Q3 = FQuat4f::Slerp(Q0, Q1, FloatAlpha);
 		LogTest<float>(TEXT("TestQuatLerp"), TestQuatsEqual(Q2, Q3, 1e-6f));
+		Q2 = Old_Slerp(Q0, Q1, FloatAlpha);
+		LogTest<float>(TEXT("TestQuatLerp"), TestQuatsEqual(Q2, Q3, 1e-6f));
 
 		Q2 = FMath::Lerp<FQuat4f>(Q0, Q1, FloatAlpha);
 		Q3 = FQuat4f::Slerp(Q0, Q1, FloatAlpha);
 		LogTest<float>(TEXT("TestQuatLerp"), TestQuatsEqual(Q2, Q3, 1e-6f));
+		Q2 = Old_Slerp(Q0, Q1, FloatAlpha);
+		LogTest<float>(TEXT("TestQuatLerp"), TestQuatsEqual(Q2, Q3, 1e-6f));
 
 		// double alpha
-		Q2 = FMath::Lerp(Q0, Q1, DoubleAlpha);
-		Q3 = FQuat4f::Slerp(Q0, Q1, (float)DoubleAlpha);
+		Q2 = FMath::Lerp(Q0, Q1, -DoubleAlpha);
+		Q3 = FQuat4f::Slerp(Q0, Q1, float(-DoubleAlpha));
+		LogTest<float>(TEXT("TestQuatLerp"), TestQuatsEqual(Q2, Q3, 1e-6f));
+		Q2 = Old_Slerp(Q0, Q1, float(-DoubleAlpha));
 		LogTest<float>(TEXT("TestQuatLerp"), TestQuatsEqual(Q2, Q3, 1e-6f));
 
-		Q2 = FMath::Lerp<FQuat4f>(Q0, Q1, DoubleAlpha);
-		Q3 = FQuat4f::Slerp(Q0, Q1, (float)DoubleAlpha);
+		Q2 = FMath::Lerp<FQuat4f>(Q0, Q1, -DoubleAlpha);
+		Q3 = FQuat4f::Slerp(Q0, Q1, float(-DoubleAlpha));
+		LogTest<float>(TEXT("TestQuatLerp"), TestQuatsEqual(Q2, Q3, 1e-6f));
+		Q2 = Old_Slerp(Q0, Q1, float(-DoubleAlpha));
 		LogTest<float>(TEXT("TestQuatLerp"), TestQuatsEqual(Q2, Q3, 1e-6f));
 
 		// Quat<double>
 		FQuat4d QD0, QD1, QD2, QD3;
-		QD0 = FQuat4d(FRotator3d(Rotator0));
-		QD1 = FQuat4d(FRotator3d(Rotator1));
+		// Note: Values chosen to have a negative 'RawCosom' within Slerp_NotNormalized()
+		QD0 = FQuat4d(FRotator3d(30.0f, 45.0f, 0.0f));
+		QD1 = FQuat4d(FRotator3d(10.0f, -270.0f, -130.0f));
 
 		// double alpha
 		QD2 = FMath::Lerp(QD0, QD1, DoubleAlpha);
 		QD3 = FQuat4d::Slerp(QD0, QD1, DoubleAlpha);
 		LogTest<double>(TEXT("TestQuatLerp"), TestQuatsEqual(QD2, QD3, 1e-6f));
+		QD2 = Old_Slerp(QD0, QD1, DoubleAlpha);
+		LogTest<double>(TEXT("TestQuatLerp"), TestQuatsEqual(QD2, QD3, 1e-6f));
 
 		QD2 = FMath::Lerp<FQuat4d>(QD0, QD1, DoubleAlpha);
 		QD3 = FQuat4d::Slerp(QD0, QD1, DoubleAlpha);
 		LogTest<double>(TEXT("TestQuatLerp"), TestQuatsEqual(QD2, QD3, 1e-6f));
+		QD2 = Old_Slerp(QD0, QD1, DoubleAlpha);
+		LogTest<double>(TEXT("TestQuatLerp"), TestQuatsEqual(QD2, QD3, 1e-6f));
 
 		// float alpha
-		QD2 = FMath::Lerp(QD0, QD1, FloatAlpha);
-		QD3 = FQuat4d::Slerp(QD0, QD1, FloatAlpha);
+		QD2 = FMath::Lerp(QD0, QD1, -FloatAlpha);
+		QD3 = FQuat4d::Slerp(QD0, QD1, -FloatAlpha);
+		LogTest<double>(TEXT("TestQuatLerp"), TestQuatsEqual(QD2, QD3, 1e-6f));
+		QD2 = Old_Slerp(QD0, QD1, double(-FloatAlpha));
 		LogTest<double>(TEXT("TestQuatLerp"), TestQuatsEqual(QD2, QD3, 1e-6f));
 
-		QD2 = FMath::Lerp<FQuat4d>(QD0, QD1, FloatAlpha);
-		QD3 = FQuat4d::Slerp(QD0, QD1, FloatAlpha);
+		QD2 = FMath::Lerp<FQuat4d>(QD0, QD1, -FloatAlpha);
+		QD3 = FQuat4d::Slerp(QD0, QD1, -FloatAlpha);
+		LogTest<double>(TEXT("TestQuatLerp"), TestQuatsEqual(QD2, QD3, 1e-6f));
+		QD2 = Old_Slerp(QD0, QD1, double(-FloatAlpha));
 		LogTest<double>(TEXT("TestQuatLerp"), TestQuatsEqual(QD2, QD3, 1e-6f));
 	}
 
