Fix memorystomp from TLockFreeFixedSizeAllocator_TLSCacheBase when it is accessed after destruction.
#rnx

diff --git a/Engine/Source/Runtime/Core/Public/Containers/LockFreeFixedSizeAllocator.h b/Engine/Source/Runtime/Core/Public/Containers/LockFreeFixedSizeAllocator.h
--- a/Engine/Source/Runtime/Core/Public/Containers/LockFreeFixedSizeAllocator.h
+++ b/Engine/Source/Runtime/Core/Public/Containers/LockFreeFixedSizeAllocator.h
@@ -38,7 +38,7 @@ public:
 	~TLockFreeFixedSizeAllocator_TLSCacheBase()
 	{
 		FPlatformTLS::FreeTlsSlot(TlsSlot);
-		TlsSlot = 0;
+		TlsSlot = FPlatformTLS::InvalidTlsSlot;
 	}
 
 	/**
