LLM and Insights Tracing: Fix the parent declarations of TagTraces to use the parent declared by LLM's TagData rather than constructing it from unique name, which does not always work. Also pass in the DisplayName paths for the tag names rather than the UniqueName. Also fix a crash on editor shutdown due to PlatformTLS being destructed after main before LLM is destructed after main.
#rnx

diff --git a/Engine/Source/Runtime/Core/Private/HAL/LowLevelMemTracker.cpp b/Engine/Source/Runtime/Core/Private/HAL/LowLevelMemTracker.cpp
--- a/Engine/Source/Runtime/Core/Private/HAL/LowLevelMemTracker.cpp
+++ b/Engine/Source/Runtime/Core/Private/HAL/LowLevelMemTracker.cpp
@@ -11,6 +11,7 @@
 #include "HAL/PlatformMemory.h" // for page allocation association.
 #include "LowLevelMemTrackerPrivate.h"
 #include "MemPro/MemProProfiler.h"
+#include "Math/NumericLimits.h"
 #include "Misc/CString.h"
 #include "Misc/FileHelper.h"
 #include "Misc/Fork.h"
@@ -1255,10 +1256,10 @@ public:
 	void Publish(FLowLevelMemTracker& LLMRef, const FTrackerTagSizeMap& TagSizes,
 		const FTagData* OverrideTrackedTotalName, const FTagData* OverrideUntaggedName, int64 TrackedTotal,
 		UE::LLM::ESizeParams SizeParams);
-	static void TraceGenericTags(FLowLevelMemTracker& LLMRef);
 
 private:
 	static const void* GetTagId(const FTagData* TagData);
+	void SendTagDeclaration(const FTagData* TagData);
 
 	ELLMTracker				Tracker;
 	TFastPointerLLMSet<const FTagData*> DeclaredTags;
@@ -1536,6 +1537,44 @@ const TCHAR* ToString(ETagReferenceSource ReferenceSource);
 void SetMemoryStatByFName(FName Name, int64 Amount);
 void ValidateUniqueName(FStringView UniqueName);
 
+typedef TArray<TPair<FLLMInitialisedCallback, UPTRINT>, TInlineAllocator<1>> FInitializedCallbacksArray;
+FInitializedCallbacksArray& GetInitialisedCallbacks()
+{
+	static FInitializedCallbacksArray Array;
+	return Array;
+}
+
+typedef TArray<TPair<FTagCreationCallback, UPTRINT>, TInlineAllocator<1>> FTagCreationCallbacksArray;
+FTagCreationCallbacksArray& GetTagCreationCallbacks()
+{
+	static FTagCreationCallbacksArray Array;
+	return Array;
+}
+
+void FPrivateCallbacks::AddInitialisedCallback(FLLMInitialisedCallback Callback, UPTRINT UserData)
+{
+	FLowLevelMemTracker::Get().BootstrapInitialise();
+
+	FInitializedCallbacksArray& Callbacks = UE::LLMPrivate::GetInitialisedCallbacks();
+	Callbacks.Add(TPair<FLLMInitialisedCallback, UPTRINT> { Callback, UserData });
+}
+
+void FPrivateCallbacks::AddTagCreationCallback(FTagCreationCallback Callback, UPTRINT UserData)
+{
+	FLowLevelMemTracker::Get().BootstrapInitialise();
+
+	FTagCreationCallbacksArray& Callbacks = UE::LLMPrivate::GetTagCreationCallbacks();
+	Callbacks.Add(TPair<FTagCreationCallback, UPTRINT> { Callback, UserData });
+}
+
+void FPrivateCallbacks::RemoveTagCreationCallback(FTagCreationCallback Callback)
+{
+	FLowLevelMemTracker::Get().BootstrapInitialise();
+
+	FTagCreationCallbacksArray& Callbacks = UE::LLMPrivate::GetTagCreationCallbacks();
+	Callbacks.RemoveAll([Callback](TPair<FTagCreationCallback, UPTRINT>& Pair) { return Pair.Key == Callback; });
+}
+
 } // namespace UE::LLMPrivate
 
 static FName GetTagName_CustomName()
@@ -2180,7 +2219,6 @@ void FLowLevelMemTracker::ProcessCommandLine(const TCHAR* CmdLine)
 	bIsDisabled = false;
 	bCsvWriterEnabled = bLocalCsvWriterEnabled;
 	bTraceWriterEnabled = bLocalTraceWriterEnabled;
-	BootstrapInitialise();
 	FinishInitialise();
 
 	// activate tag sets (we ignore None set, it's always on)
@@ -2536,6 +2574,8 @@ void FLowLevelMemTracker::FinishInitialise()
 	{
 		return;
 	}
+	BootstrapInitialise();
+
 #if UE_ENABLE_ARRAY_SLACK_TRACKING
 	GTrackArrayDetailedLock = new FCriticalSection();
 #endif
@@ -2545,8 +2585,15 @@ void FLowLevelMemTracker::FinishInitialise()
 	// FName system construction, which will itself trigger Malloc construction.
 	(void)LLMGetTagUniqueName(ELLMTag::Untagged);
 	InitialiseTagDatas();
+
+	UE::LLMPrivate::FInitializedCallbacksArray& Callbacks = UE::LLMPrivate::GetInitialisedCallbacks();
+	for (const TPair<UE::LLMPrivate::FLLMInitialisedCallback, UPTRINT>& Callback : Callbacks)
+	{
+		Callback.Key(Callback.Value);
+	}
+	Callbacks.Empty();
 }
- 
+
 void FLowLevelMemTracker::InitialiseTagDatas_SetLLMTagNames()
 {
 	using namespace UE::LLMPrivate;
@@ -2894,7 +2941,6 @@ void FLowLevelMemTracker::FinishConstruct(UE::LLMPrivate::FTagData* TagData,
 			TagData->SetParent(ParentData);
 		}
 	}
-	TagData->SetFinishConstructed();
 
 	FTagData* ParentData = const_cast<FTagData*>(TagData->GetParent());
 	if (ParentData)
@@ -2903,6 +2949,24 @@ void FLowLevelMemTracker::FinishConstruct(UE::LLMPrivate::FTagData* TagData,
 		// called and walk up the parent chain.
 		FinishConstruct(ParentData, ReferenceSource);
 	}
+
+	TagData->SetFinishConstructed();
+
+	// Broadcast the tag creation, except for generic tags which are constructed before any subscriber could
+	// possibly have registered. Subscribers must instead read those from GetTrackedTags.
+	if (!TagData->HasEnumTag() || TagData->GetEnumTag() >= ELLMTag::GenericTagCount)
+	{
+		// Leave the critical section while calling the callback, since the callback could be arbitrary
+		// code that calls back into LLM. The TagData pointer is immutable so we do not have to worry about
+		// it disappearing out from under us.
+		TagDataLock.ReadUnlock();
+		for (const TPair<UE::LLMPrivate::FTagCreationCallback, UPTRINT>& Callback :
+			UE::LLMPrivate::GetTagCreationCallbacks())
+		{
+			Callback.Key(TagData, Callback.Value);
+		}
+		TagDataLock.ReadLock();
+	}
 }
 
 TArray<const UE::LLMPrivate::FTagData*> FLowLevelMemTracker::GetTrackedTags(ELLMTagSet TagSet)
@@ -3087,7 +3151,15 @@ FName FLowLevelMemTracker::GetTagDisplayName(const UE::LLMPrivate::FTagData* Tag
 
 FString FLowLevelMemTracker::GetTagDisplayPathName(const UE::LLMPrivate::FTagData* TagData) const
 {
-	return TagData->GetDisplayPath();
+	TStringBuilder<FName::StringBufferSize> Buffer;
+	GetTagDisplayPathName(TagData, Buffer);
+	return FString(Buffer);
+}
+
+void FLowLevelMemTracker::GetTagDisplayPathName(const UE::LLMPrivate::FTagData* TagData,
+	FStringBuilderBase& OutPathName, int32 MaxLen) const
+{
+	TagData->GetDisplayPath(OutPathName, MaxLen);
 }
 
 FName FLowLevelMemTracker::GetTagUniqueName(const UE::LLMPrivate::FTagData* TagData) const
@@ -3095,6 +3167,21 @@ FName FLowLevelMemTracker::GetTagUniqueName(const UE::LLMPrivate::FTagData* TagD
 	return TagData->GetName();
 }
 
+const UE::LLMPrivate::FTagData* FLowLevelMemTracker::GetTagParent(const UE::LLMPrivate::FTagData* TagData) const
+{
+	return TagData->GetParent();
+}
+
+bool FLowLevelMemTracker::GetTagIsEnumTag(const UE::LLMPrivate::FTagData* TagData) const
+{
+	return TagData->HasEnumTag();
+}
+
+ELLMTag FLowLevelMemTracker::GetTagClosestEnumTag(const UE::LLMPrivate::FTagData* TagData) const
+{
+	return TagData->GetEnumTag();
+}
+
 // Deprecated in 5.3
 int64 FLowLevelMemTracker::GetTagAmountForTracker(ELLMTracker Tracker, ELLMTag Tag, bool bPeakAmount)
 {
@@ -3760,9 +3847,9 @@ void FLLMTagDeclaration::ConstructUniqueName()
 namespace UE::LLMPrivate::LLMTagDeclarationInternal
 {
 
-TArray<FLLMTagDeclaration::FCreationCallback, TInlineAllocator<2>>& GetCreationCallbacks()
+TArray<FLLMTagDeclaration::FCreationCallback, TInlineAllocator<1>>& GetCreationCallbacks()
 {
-	static TArray<FLLMTagDeclaration::FCreationCallback, TInlineAllocator<2>> CreationCallbacks;
+	static TArray<FLLMTagDeclaration::FCreationCallback, TInlineAllocator<1>> CreationCallbacks;
 	return CreationCallbacks;
 }
 
@@ -3776,7 +3863,7 @@ FLLMTagDeclaration*& GetList()
 
 void FLLMTagDeclaration::AddCreationCallback(FCreationCallback InCallback)
 {
-	TArray<FCreationCallback, TInlineAllocator<2>>& CreationCallbacks =
+	TArray<FCreationCallback, TInlineAllocator<1>>& CreationCallbacks =
 		UE::LLMPrivate::LLMTagDeclarationInternal::GetCreationCallbacks();
 	if (CreationCallbacks.Num() >= CreationCallbacks.Max())
 	{
@@ -3790,7 +3877,7 @@ void FLLMTagDeclaration::AddCreationCallback(FCreationCallback InCallback)
 
 void FLLMTagDeclaration::ClearCreationCallbacks()
 {
-	TArray<FCreationCallback, TInlineAllocator<2>>& CreationCallbacks =
+	TArray<FCreationCallback, TInlineAllocator<1>>& CreationCallbacks =
 		UE::LLMPrivate::LLMTagDeclarationInternal::GetCreationCallbacks();
 	CreationCallbacks.Empty();
 }
@@ -4302,20 +4389,37 @@ FName FTagData::GetDisplayName() const
 	return DisplayName;
 }
 
-FString FTagData::GetDisplayPath() const
+void FTagData::GetDisplayPath(FStringBuilderBase& Result, int32 MaxLen) const
 {
-	TStringBuilder<256> NameBuffer;
-	AppendDisplayPath(NameBuffer);
-	return FString(NameBuffer);
+	Result.Reset();
+	AppendDisplayPath(Result, MaxLen);
 }
 
-void FTagData::AppendDisplayPath(FStringBuilderBase& Result) const
+void FTagData::AppendDisplayPath(FStringBuilderBase& Result, int32 MaxLen) const
 {
 	if (Parent && Parent->IsUsedAsDisplayParent())
 	{
-		Parent->AppendDisplayPath(Result);
+		Parent->AppendDisplayPath(Result, MaxLen);
+		if (MaxLen >= 0 && Result.Len() + 1 >= MaxLen)
+		{
+			return;
+		}
 		Result << TEXT("/");
 	}
+	if (MaxLen >= 0)
+	{
+		int32 MaxRemainingLen = MaxLen - Result.Len();
+		if (static_cast<int32>(DisplayName.GetStringLength() + 1) > MaxRemainingLen)
+		{
+			if (MaxRemainingLen > 1)
+			{
+				TStringBuilder<FName::StringBufferSize> Buffer;
+				DisplayName.AppendString(Buffer);
+				Result << Buffer.ToView().Left(MaxRemainingLen - 1);
+			}
+			return;
+		}
+	}
 	DisplayName.AppendString(Result);
 }
 
@@ -6046,23 +6150,6 @@ void FLLMTraceWriter::Publish(FLowLevelMemTracker& LLMRef, const FTrackerTagSize
 			<< TrackerSpec.Name(TrackerNames[(uint8)Tracker].Name, NameLen);
 	}
 
-	TStringBuilder<1024> NameBuffer;
-	auto SendTagDeclaration = [this, &NameBuffer](const FTagData* TagData)
-	{
-		if (!TagData || DeclaredTags.Contains(TagData))
-		{
-			return;
-		}
-		DeclaredTags.Add(TagData);
-
-		const FTagData* Parent = TagData->GetParent();
-		NameBuffer.Reset();
-		TagData->AppendDisplayPath(NameBuffer);
-		UE_TRACE_LOG(LLM, TagsSpec, MemTagChannel, NameBuffer.Len() * sizeof(ANSICHAR))
-			<< TagsSpec.TagId(GetTagId(TagData))
-			<< TagsSpec.ParentId(GetTagId(Parent))
-			<< TagsSpec.Name(*NameBuffer, NameBuffer.Len());
-	};
 	SendTagDeclaration(OverrideTrackedTotalTagData);
 	SendTagDeclaration(OverrideUntaggedTagData);
 	for (const TPair<const FTagData*, FTrackerTagSizeData>& It : TagSizes)
@@ -6132,18 +6219,29 @@ void FLLMTraceWriter::Publish(FLowLevelMemTracker& LLMRef, const FTrackerTagSize
 		<< TagValue.Values(TagValues.GetData(), TagCount);
 }
 
-void FLLMTraceWriter::TraceGenericTags(FLowLevelMemTracker& LLMRef)
+void FLLMTraceWriter::SendTagDeclaration(const FTagData* TagData)
 {
-	for (int32 GenericTagIndex = 0; GenericTagIndex < static_cast<int32>(ELLMTag::GenericTagCount); GenericTagIndex++)
+	if (!TagData)
 	{
-		const FTagData* TagData = LLMRef.FindTagData(static_cast<ELLMTag>(GenericTagIndex));
-		FString TagName = TagData->GetDisplayPath();
-		UE_TRACE_LOG(LLM, TagsSpec, MemTagChannel, TagName.Len() * sizeof(ANSICHAR))
-			<< TagsSpec.TagId(GetTagId(TagData))
-			<< TagsSpec.ParentId(GetTagId(TagData->GetParent()))
-			<< TagsSpec.Name(*TagName, TagName.Len());
+		return;
 	}
-}
+	bool bAlreadyInSet;
+	DeclaredTags.Add(TagData, &bAlreadyInSet);
+	if (bAlreadyInSet)
+	{
+		return;
+	}
+
+	const FTagData* Parent = TagData->GetParent();
+	SendTagDeclaration(Parent);
+
+	TStringBuilder<1024> NameBuffer;
+	TagData->AppendDisplayPath(NameBuffer);
+	UE_TRACE_LOG(LLM, TagsSpec, MemTagChannel, NameBuffer.Len() * sizeof(ANSICHAR))
+		<< TagsSpec.TagId(GetTagId(TagData))
+		<< TagsSpec.ParentId(GetTagId(Parent))
+		<< TagsSpec.Name(*NameBuffer, NameBuffer.Len());
+};
 
 // FLLMCsvProfilerWriter implementation.
 FLLMCsvProfilerWriter::FLLMCsvProfilerWriter()
@@ -6211,7 +6309,9 @@ void FLLMCsvProfilerWriter::RecordTagToCsv(int32 CsvCategoryIndex, const FTagDat
 	FName* CsvStatNamePtr = TagDataToCsvStatName.Find(TagData);
 	if (CsvStatNamePtr == nullptr)
 	{
-		NewCsvStatName = FName(TagData->GetDisplayPath());
+		TStringBuilder<FName::StringBufferSize> DisplayPath;
+		TagData->GetDisplayPath(DisplayPath);
+		NewCsvStatName = FName(DisplayPath);
 		TagDataToCsvStatName.Add(TagData, NewCsvStatName);
 		CsvStatNamePtr = &NewCsvStatName;
 	}
diff --git a/Engine/Source/Runtime/Core/Private/HAL/LowLevelMemTrackerPrivate.h b/Engine/Source/Runtime/Core/Private/HAL/LowLevelMemTrackerPrivate.h
--- a/Engine/Source/Runtime/Core/Private/HAL/LowLevelMemTrackerPrivate.h
+++ b/Engine/Source/Runtime/Core/Private/HAL/LowLevelMemTrackerPrivate.h
@@ -95,8 +95,8 @@ namespace LLMPrivate
 		bool IsFinishConstructed() const;
 		FName GetName() const;
 		FName GetDisplayName() const;
-		FString GetDisplayPath() const;
-		void AppendDisplayPath(FStringBuilderBase& Result) const;
+		void GetDisplayPath(FStringBuilderBase& Result, int32 MaxLen=-1) const;
+		void AppendDisplayPath(FStringBuilderBase& Result, int32 MaxLen=-1) const;
 		const FTagData* GetParent() const;
 		FName GetParentName() const;
 		FName GetParentNameSafeBeforeFinishConstruct() const;
diff --git a/Engine/Source/Runtime/Core/Private/ProfilingDebugging/TagTrace.cpp b/Engine/Source/Runtime/Core/Private/ProfilingDebugging/TagTrace.cpp
--- a/Engine/Source/Runtime/Core/Private/ProfilingDebugging/TagTrace.cpp
+++ b/Engine/Source/Runtime/Core/Private/ProfilingDebugging/TagTrace.cpp
@@ -7,6 +7,9 @@
 #include "ProfilingDebugging/MemoryTrace.h"
 #include "HAL/LowLevelMemTracker.h"
 #include "HAL/LowLevelMemTrackerPrivate.h"
+#include "HAL/PlatformString.h"
+#include "Misc/Optional.h"
+#include "Misc/StringBuilder.h"
 #include "Trace/Trace.inl"
 #include "UObject/NameTypes.h"
 
@@ -138,13 +141,25 @@ class FTagTrace
 public:
 					FTagTrace(FMalloc* InMalloc);
 	void			AnnounceGenericTags() const;
-	void 			AnnounceTagDeclarations();
-					void AnnounceSpecialTags() const;
+	void			AnnounceLLMExtendedTags();
+	void			AnnounceSpecialTags() const;
 #if ENABLE_LOW_LEVEL_MEM_TRACKER
-	static void 	OnAnnounceTagDeclaration(FLLMTagDeclaration& TagDeclaration);
+	static void		OnLLMInitialised(UPTRINT pThis);
+	static void		OnAnnounceLLMExtendedTag(const UE::LLMPrivate::FTagData* LLMTag, UPTRINT pThis);
+	int32			AnnounceLLMExtendedTag(FLowLevelMemTracker& Tracker, const UE::LLMPrivate::FTagData* LLMTag);
 #endif
 	int32			AnnounceCustomTag(int32 Tag, int32 ParentTag, const ANSICHAR* Display) const;
-	int32 			AnnounceFNameTag(const FName& TagName);
+	/**
+	 * Send information about the given Tag to the trace if it has not already been sent.
+	 *
+	 * @param UniqueName The uniquename of the tag, guaranteed unique across all other memtrace tags.
+	 * @param TagId The tag to use, or -1 to construct the Tag from UniqueName.
+	 * @param DisplayName Displaystring of the tag when viewing the trace, or nullptr to use UniqueName.
+	 * @param ParentTagId If set, the tag's parent tag, or -1 if the tag has no parent. If not set, ParentTag
+	 *        is calculated by looking for markers in the UniqueName.
+	 * @return The tag that ended up being used for the FName.
+	 */
+	int32 			AnnounceFNameTag(FName UniqueName, int32 TagId, const TCHAR* DisplayName, TOptional<int32> ParentTagId);
 
 private:
 
@@ -180,8 +195,8 @@ FTagTrace::FTagTrace(FMalloc* InMalloc)
 	Malloc = InMalloc;
 	AnnouncedNames.Reserve(1024);
 	AnnounceGenericTags();
-	AnnounceTagDeclarations();
 	AnnounceSpecialTags();
+	AnnounceLLMExtendedTags();
 }
 
 ////////////////////////////////////////////////////////////////////////////////
@@ -202,19 +217,80 @@ void FTagTrace::AnnounceGenericTags() const
 }
 
 ////////////////////////////////////////////////////////////////////////////////
-void FTagTrace::AnnounceTagDeclarations()
+#if ENABLE_LOW_LEVEL_MEM_TRACKER
+void FTagTrace::OnLLMInitialised(UPTRINT pThis)
+{
+	reinterpret_cast<FTagTrace*>(pThis)->AnnounceLLMExtendedTags();
+}
+#endif
+
+////////////////////////////////////////////////////////////////////////////////
+void FTagTrace::AnnounceLLMExtendedTags()
 {
 #if ENABLE_LOW_LEVEL_MEM_TRACKER
-	FLLMTagDeclaration* List = FLLMTagDeclaration::GetList();
-	while (List)
+	// Our first call to AnnounceLLMExtendedTags can occur before FName's singleton has been constructed, and so
+	// LLM reports that it is not yet fully initialised. We need to wait for LLM to have its Tags' FNames, so add
+	// ourselves to its initialisation callback to reexecute this function at that point.
+	FLowLevelMemTracker& Tracker = FLowLevelMemTracker::Get();
+	if (!Tracker.IsInitialized())
 	{
-		OnAnnounceTagDeclaration(*List);
-		List = List->Next;
+		UE::LLMPrivate::FPrivateCallbacks::AddInitialisedCallback(FTagTrace::OnLLMInitialised, reinterpret_cast<UPTRINT>(this));
+		return;
 	}
-	FLLMTagDeclaration::AddCreationCallback(FTagTrace::OnAnnounceTagDeclaration);
+
+	// Announce all of the LLMTags that LLM currently knows about, and sign up to be notified when new ones are added.
+	TArray<const UE::LLMPrivate::FTagData*> Tags = Tracker.GetTrackedTags();
+	for (const UE::LLMPrivate::FTagData* Tag : Tags)
+	{
+		AnnounceLLMExtendedTag(Tracker, Tag);
+	}
+	UE::LLMPrivate::FPrivateCallbacks::AddTagCreationCallback(FTagTrace::OnAnnounceLLMExtendedTag, reinterpret_cast<UPTRINT>(this));
 #endif
 }
 
+////////////////////////////////////////////////////////////////////////////////
+#if ENABLE_LOW_LEVEL_MEM_TRACKER
+void FTagTrace::OnAnnounceLLMExtendedTag(const UE::LLMPrivate::FTagData* LLMTag, UPTRINT pThis)
+{
+	FLowLevelMemTracker& Tracker = FLowLevelMemTracker::Get();
+	reinterpret_cast<FTagTrace*>(pThis)->AnnounceLLMExtendedTag(Tracker, LLMTag);
+}
+
+////////////////////////////////////////////////////////////////////////////////
+int32 FTagTrace::AnnounceLLMExtendedTag(FLowLevelMemTracker& Tracker, const UE::LLMPrivate::FTagData* LLMTag)
+{
+	bool bEnumTag = Tracker.GetTagIsEnumTag(LLMTag);
+	ELLMTag EnumTag = ELLMTag::GenericTagCount;
+	if (bEnumTag)
+	{
+		EnumTag = Tracker.GetTagClosestEnumTag(LLMTag);
+		if (EnumTag < ELLMTag::GenericTagCount)
+		{
+			// We already announced this tag super-early as part of AnnounceGenericTags
+			return static_cast<int32>(EnumTag);
+		}
+		// Not a generic tag, so this is a Platform or Project tag.
+	}
+	int32 ParentTagId = -1;
+	const UE::LLMPrivate::FTagData* ParentLLMTag = Tracker.GetTagParent(LLMTag);
+	if (ParentLLMTag)
+	{
+		ParentTagId = AnnounceLLMExtendedTag(Tracker, ParentLLMTag);
+	}
+
+	TStringBuilder<FName::StringBufferSize+1> DisplayName;
+	Tracker.GetTagDisplayPathName(LLMTag, DisplayName, FName::StringBufferSize);
+	const TCHAR* DisplayNamePtr = *DisplayName;
+	if (DisplayName.Len() >= FName::StringBufferSize)
+	{
+		DisplayNamePtr = nullptr; // Use the UniqueName instead
+	}
+	int32 TagId = bEnumTag ? static_cast<int32>(EnumTag) : -1;
+	return AnnounceFNameTag(Tracker.GetTagUniqueName(LLMTag), TagId, DisplayNamePtr, ParentTagId);
+}
+
+#endif
+
 ////////////////////////////////////////////////////////////////////////////////
 void FTagTrace::AnnounceSpecialTags() const
 {
@@ -232,41 +308,60 @@ void FTagTrace::AnnounceSpecialTags() const
 }
 
 ////////////////////////////////////////////////////////////////////////////////
-#if ENABLE_LOW_LEVEL_MEM_TRACKER
-void FTagTrace::OnAnnounceTagDeclaration(FLLMTagDeclaration& TagDeclaration)
+int32 FTagTrace::AnnounceFNameTag(FName UniqueName, int32 TagId, const TCHAR* DisplayName, TOptional<int32> ParentTagId)
 {
-	TagDeclaration.ConstructUniqueName();
-	GTagTrace->AnnounceFNameTag(TagDeclaration.GetUniqueName());
-}
-#endif
-
-////////////////////////////////////////////////////////////////////////////////
-int32 FTagTrace::AnnounceFNameTag(const FName& Name)
-{
-	const int32 NameIndex = Name.GetDisplayIndex().ToUnstableInt() + FNAME_INDEX_OFFSET;
+	if (TagId == -1)
+	{
+		TagId = UniqueName.GetDisplayIndex().ToUnstableInt() + FNAME_INDEX_OFFSET;
+	}
 
 	// Find or add the item
 	bool bAlreadyInTable;
-	AnnouncedNames.FindOrAdd(NameIndex, true, &bAlreadyInTable);
+	AnnouncedNames.FindOrAdd(TagId, true, &bAlreadyInTable);
 	if (bAlreadyInTable)
 	{
-		return NameIndex;
+		return TagId;
 	}
 
 	// First time encountering this name, announce it
 	ANSICHAR NameString[NAME_SIZE];
-	Name.GetPlainANSIString(NameString);
-
-	int32 ParentTag = -1;
-	FAnsiStringView NameView(NameString);
-	int32 LeafStart;
-	if (NameView.FindLastChar('/', LeafStart))
+	int32 DisplayNameLen = DisplayName == nullptr ? 0 : FCString::Strlen(DisplayName);
+	DisplayNameLen = DisplayNameLen + 1 > NAME_SIZE ? 0 : DisplayNameLen;
+	if (DisplayNameLen > 0)
 	{
-		FName ParentName(NameView.Left(LeafStart));
-		ParentTag = AnnounceFNameTag(ParentName);
+		verify(FPlatformString::Convert(NameString, NAME_SIZE, DisplayName, DisplayNameLen + 1) != nullptr);
 	}
+	else
+	{
+		UniqueName.GetPlainANSIString(NameString);
+	}
+
+	if (!ParentTagId.IsSet())
+	{
+		ParentTagId = -1;
+		if (DisplayNameLen > 0)
+		{
+			// By contract we pull the parent name from the UniqueName rather than the display string.
+			// !ParentTagId.IsSet() does not currently occur when DisplayName is provided so we
+			// handle this case with poor performance rather than adding stacksize or complexity to
+			// handle it efficiently.
+			UniqueName.GetPlainANSIString(NameString);
+		}
 
-	return AnnounceCustomTag(NameIndex, ParentTag, NameString);
+		FAnsiStringView NameView(NameString);
+		int32 LeafStart;
+		if (NameView.FindLastChar('/', LeafStart))
+		{
+			FName ParentName(NameView.Left(LeafStart));
+			ParentTagId = AnnounceFNameTag(ParentName, -1, nullptr, TOptional<int32>());
+		}
+
+		if (DisplayNameLen > 0)
+		{
+			verify(FPlatformString::Convert(NameString, NAME_SIZE, DisplayName, DisplayNameLen + 1) != nullptr);
+		}
+	}
+	return AnnounceCustomTag(TagId, *ParentTagId, NameString);
 }
 
 ////////////////////////////////////////////////////////////////////////////////
@@ -311,7 +406,7 @@ int32 MemoryTrace_AnnounceFNameTag(const FName& TagName)
 #if UE_MEMORY_TAGS_TRACE_ENABLED && UE_TRACE_ENABLED
 	if (GTagTrace)
 	{
-		return GTagTrace->AnnounceFNameTag(TagName);
+		return GTagTrace->AnnounceFNameTag(TagName, -1, nullptr, TOptional<int32>());
 	}
 #endif
 	return -1;
diff --git a/Engine/Source/Runtime/Core/Private/UObject/UnrealNames.cpp b/Engine/Source/Runtime/Core/Private/UObject/UnrealNames.cpp
--- a/Engine/Source/Runtime/Core/Private/UObject/UnrealNames.cpp
+++ b/Engine/Source/Runtime/Core/Private/UObject/UnrealNames.cpp
@@ -2014,6 +2014,7 @@ static FNamePool& GetNamePool()
 
 	FNamePool* Singleton = new (NamePoolData) FNamePool;
 	bNamePoolInitialized = true;
+	LLM(FLowLevelMemTracker::Get().FinishInitialise());
 	return *Singleton;
 }
 
diff --git a/Engine/Source/Runtime/Core/Public/HAL/LowLevelMemTracker.h b/Engine/Source/Runtime/Core/Public/HAL/LowLevelMemTracker.h
--- a/Engine/Source/Runtime/Core/Public/HAL/LowLevelMemTracker.h
+++ b/Engine/Source/Runtime/Core/Public/HAL/LowLevelMemTracker.h
@@ -46,6 +46,7 @@
 
 #include "Containers/Array.h"
 #include "Containers/ArrayView.h"
+#include "Containers/StringFwd.h"
 #include "HAL/CriticalSection.h"
 #include "HAL/PlatformCrt.h"
 #include "HAL/PlatformMisc.h"
@@ -56,6 +57,8 @@
 
 #include <atomic>
 
+class FTagTrace;
+
 #if DO_CHECK
 
 namespace UE::LLMPrivate
@@ -534,6 +537,25 @@ enum class ETagReferenceSource
 	ImplicitParent
 };
 
+typedef void (*FLLMInitialisedCallback)(UPTRINT UserData);
+typedef void (*FTagCreationCallback)(const UE::LLMPrivate::FTagData* TagData, UPTRINT UserData);
+
+/**
+ * Callbacks that can occur during LLM Initialisation. These happen before main and are therefore
+ * dangerous to use, so they are private.
+ */
+struct FPrivateCallbacks
+{
+private:
+	CORE_API static void AddInitialisedCallback(FLLMInitialisedCallback Callback, UPTRINT UserData);
+	// There is no RemoveInitialiseCallback because it is triggered and cleared before Main is called,
+	// there should be no need to remove a callback before then.
+	CORE_API static void AddTagCreationCallback(FTagCreationCallback Callback, UPTRINT UserData);
+	CORE_API static void RemoveTagCreationCallback(FTagCreationCallback Callback);
+
+	friend class ::FTagTrace;
+};
+
 } // UE::LLMPrivate
 
 struct UE_DEPRECATED(4.27, "FLLMCustomTag was an implementation detail that has been modified, switch to FLLMTagInfo or to your own local struct") FLLMCustomTag
@@ -676,10 +698,24 @@ public:
 
 	/** Get the path name for the given FTagData from a chain of its parents' display names. */
 	CORE_API FString GetTagDisplayPathName(const UE::LLMPrivate::FTagData* TagData) const;
+	CORE_API void GetTagDisplayPathName(const UE::LLMPrivate::FTagData* TagData,
+		FStringBuilderBase& OutPathName, int32 MaxLen=-1) const;
 
 	/** Get the unique identifier name for the given FTagData. */
 	CORE_API FName GetTagUniqueName(const UE::LLMPrivate::FTagData* TagData) const;
 
+	/** Return the TagData that is the parent scope of the given TagData, or nullptr if no parent. */
+	CORE_API const UE::LLMPrivate::FTagData* GetTagParent(const UE::LLMPrivate::FTagData* TagData) const;
+
+	/** Return true if and only if the TagData is an ELLMTag or a custom-declared platform or project ELLMTag. */
+	CORE_API bool GetTagIsEnumTag(const UE::LLMPrivate::FTagData* TagData) const;
+
+	/**
+	 * Return the ELLMTag of closest parent (including possibly TagData) which has true==GetTagIsEnumTag
+	 * For FName tags with no ELLMTag parent, they will return ELLMTag::CustomName.
+	 */
+	CORE_API ELLMTag GetTagClosestEnumTag(const UE::LLMPrivate::FTagData* TagData) const;
+
 	/** Get the amount of memory for an ELLMTag from the given tracker. */
 	CORE_API int64 GetTagAmountForTracker(ELLMTracker Tracker, ELLMTag Tag, UE::LLM::ESizeParams SizeParams = UE::LLM::ESizeParams::Default);
 
@@ -724,20 +760,28 @@ public:
 	{
 		return bFullyInitialised;
 	}
+	/**
+	 * Initialize the data required for tracking allocations if it has not already been done.
+	 * This can occur after construction of the singleton FLowLevelMemTracker. 
+	 * This function is called automatically by any LLM function that requires it, but can also be triggered
+	 * by external systems that need to carefully manage static initialization order.
+	 * BootstrapInitialize does NOT make GetTagUniqueName available; that requires FinishInitialize.
+	 */
+	CORE_API void BootstrapInitialise();
+
+	/**
+	 * Finish initialization if not already done. Calls BootstrapInitialise if necessary. After being called,
+	 * all functions in LLM including those returning FNames are available.
+	 * This function should be called as soon as possible after FNames are available so that other tracking
+	 * systems reliant on LLM's FName data can start tracking as soon as possible.
+	 */
+	CORE_API void FinishInitialise();
 
 private:
 	CORE_API FLowLevelMemTracker();
 
 	CORE_API ~FLowLevelMemTracker();
 
-	/**
-	 * Allocation and Setup of the data required for tracking allocations is done as late as possible to prevent
-	 * exercising allocation code too early. Note that as late as possible for tracking allocations is still earlier
-	 * than ParseCommandLine, so we do not complete all initialization in this function (e.g. features required only
-	 * for Update are omitted). Initialization done here will be torn down in ParseCommandLine if LLM is disabled.
-	 */
-	CORE_API void BootstrapInitialise();
-
 	bool IsBootstrapping() const { return bIsBootstrapping; }
 
 	/** Free all memory. This will put the tracker into a permanently disabled state. */
@@ -758,8 +802,6 @@ private:
 	* Creates the subset of tags necessary to record allocations during GMalloc and FName construction
 	*/
 	CORE_API void BootstrapTagDatas();
-	/** Called when we have detected enabled in processcommandline, or as late as possible if callers access features that require full initialisation before then */
-	CORE_API void FinishInitialise();
 	CORE_API void InitialiseTagDatas_SetLLMTagNames();
 	CORE_API void InitialiseTagDatas_FinishRegister();
 	CORE_API void InitialiseTagDatas();
@@ -1001,7 +1043,6 @@ protected:
 	FLLMTagDeclaration* Next = nullptr;
 
 	friend class FLowLevelMemTracker;
-	friend class FTagTrace;
 };
 
 /** Used to filter Tag allocations specifying a Name that matches in a TagSet, or TagSet alone or just by Name. */
