Submitting on behalf of Refik.Karic

ISPC templates integration

diff --git a/Engine/Source/Runtime/Core/Public/Math/Aossoa.isph b/Engine/Source/Runtime/Core/Public/Math/Aossoa.isph
--- a/Engine/Source/Runtime/Core/Public/Math/Aossoa.isph
+++ b/Engine/Source/Runtime/Core/Public/Math/Aossoa.isph
@@ -21,22 +21,14 @@ static const varying int vAOS21 = {1, 3, 5, 7, 9, 11, 13, 15, 17, 19, 21, 23, 25
 #error "No implementation for this target"
 #endif
 
-#define AOS_TO_SOA2_ISPC(T)                                                                                            \
-    unmasked inline void aos_to_soa2_ispc(uniform T src[], varying T *uniform v0, varying T *uniform v1)               \
-    {                                                                                                                  \
-        const varying T src0 = src[programIndex];                                                                      \
-        const varying T src1 = src[programCount + programIndex];                                                       \
-                                                                                                                       \
-        *v0 = shuffle(src0, src1, vAOS20);                                                                             \
-        *v1 = shuffle(src0, src1, vAOS21);                                                                             \
-    }
-
-AOS_TO_SOA2_ISPC(int16)
-AOS_TO_SOA2_ISPC(int32)
-AOS_TO_SOA2_ISPC(float)
-AOS_TO_SOA2_ISPC(int64)
-AOS_TO_SOA2_ISPC(double)
-
+template<typename T>
+unmasked inline void aos_to_soa2_ispc(uniform T src[], varying T* uniform v0, varying T* uniform v1)
+{
+    const varying T src0 = src[programIndex];
+    const varying T src1 = src[programCount + programIndex];
+    *v0 = shuffle(src0, src1, vAOS20);
+    *v1 = shuffle(src0, src1, vAOS21);
+}
 
 #if TARGET_WIDTH == 4
 static const varying int vAOS30 = {0, 3, 1, 4};
@@ -70,76 +62,106 @@ static const varying int vAOS35 = {0, 1, 2, 3, 4, 5, 6, 7, 8, 9, 16, 19, 22, 25,
 #error "No implementation for this target"
 #endif
 
-#if TARGET_WIDTH == 8
-#define AOS_TO_SOA3_ISPC_64BIT(T)                                                                                      \
-    unmasked inline void aos_to_soa3_ispc(uniform T src[], varying T *uniform v0, varying T *uniform v1,               \
-                                        varying T *uniform v2)                                                         \
-    {                                                                                                                  \
-        const uniform float *uniform srcflt = (const uniform float *uniform)src;                                       \
-        const varying float src0 = srcflt[programIndex];                                                               \
-        const varying float src1 = srcflt[programCount + programIndex];                                                \
-        const varying float src2 = srcflt[2 * programCount + programIndex];                                            \
-        const varying float src3 = srcflt[3 * programCount + programIndex];                                            \
-        const varying float src4 = srcflt[4 * programCount + programIndex];                                            \
-        const varying float src5 = srcflt[5 * programCount + programIndex];                                            \
-                                                                                                                       \
-        const varying float t0 = shuffle(src0, src1, vAOS301);                                                         \
-        const varying float t1 = shuffle(src1, src2, vAOS311);                                                         \
-        const varying float s0 = shuffle(src3, src4, vAOS301);                                                         \
-        const varying float s1 = shuffle(src4, src5, vAOS311);                                                         \
-        const varying float t2 = shuffle(src0, src1, vAOS341);                                                         \
-        const varying float u2 = shuffle(t2, src2, vAOS351);                                                           \
-        const varying float s2 = shuffle(src3, src4, vAOS341);                                                         \
-                                                                                                                       \
-        const varying float u0 = shuffle(t0, t1, vAOS321);                                                             \
-        const varying float w0 = shuffle(s0, s1, vAOS321);                                                             \
-        uniform T<TARGET_WIDTH> x0;                                                                                    \
-        *((varying float *uniform)&x0) = u0;                                                                           \
-        *((varying float *uniform)&x0[TARGET_WIDTH/2]) = w0;                                                           \
-        *v0 = *((varying T *uniform)&x0);                                                                              \
-                                                                                                                       \
-        const varying float u1 = shuffle(t0, t1, vAOS331);                                                             \
-        const varying float w1 = shuffle(s0, s1, vAOS331);                                                             \
-        uniform T<TARGET_WIDTH> x1;                                                                                    \
-        *((varying float *uniform)&x1) = u1;                                                                           \
-        *((varying float *uniform)&x1[TARGET_WIDTH/2]) = w1;                                                           \
-        *v1 = *((varying T *uniform)&x1);                                                                              \
-                                                                                                                       \
-        const varying float w2 = shuffle(s2, src5, vAOS351);                                                           \
-        uniform T<TARGET_WIDTH> x2;                                                                                    \
-        *((varying float *uniform)&x2) = u2;                                                                           \
-        *((varying float *uniform)&x2[TARGET_WIDTH/2]) = w2;                                                           \
-        *v2 = *((varying T *uniform)&x2);                                                                              \
-    }
-#endif
-#define AOS_TO_SOA3_ISPC(T)                                                                                            \
-    unmasked inline void aos_to_soa3_ispc(uniform T src[], varying T *uniform v0, varying T *uniform v1,               \
-                                        varying T *uniform v2)                                                         \
-    {                                                                                                                  \
-        const varying T src0 = src[programIndex];                                                                      \
-        const varying T src1 = src[programCount + programIndex];                                                       \
-        const varying T src2 = src[2 * programCount + programIndex];                                                   \
-                                                                                                                       \
-        const varying T t0 = shuffle(src0, src1, vAOS30);                                                              \
-        const varying T t1 = shuffle(src1, src2, vAOS31);                                                              \
-        *v0 = shuffle(t0, t1, vAOS32);                                                                                 \
-        *v1 = shuffle(t0, t1, vAOS33);                                                                                 \
-                                                                                                                       \
-        const varying T t2 = shuffle(src0, src1, vAOS34);                                                              \
-        *v2 = shuffle(t2, src2, vAOS35);                                                                               \
-    }
+template<typename T>
+unmasked inline void aos_to_soa3_ispc(uniform T src[], varying T* uniform v0, varying T* uniform v1,
+    varying T* uniform v2)
+{
+    const varying T src0 = src[programIndex];
+    const varying T src1 = src[programCount + programIndex];
+    const varying T src2 = src[2 * programCount + programIndex];
+
+    const varying T t0 = shuffle(src0, src1, vAOS30);
+    const varying T t1 = shuffle(src1, src2, vAOS31);
+    *v0 = shuffle(t0, t1, vAOS32);
+    *v1 = shuffle(t0, t1, vAOS33);
+
+    const varying T t2 = shuffle(src0, src1, vAOS34);
+    *v2 = shuffle(t2, src2, vAOS35);
+}
 
-AOS_TO_SOA3_ISPC(int16)
-AOS_TO_SOA3_ISPC(int32)
-AOS_TO_SOA3_ISPC(float)
 #if TARGET_WIDTH == 8
-AOS_TO_SOA3_ISPC_64BIT(double)
-AOS_TO_SOA3_ISPC_64BIT(int64)
-#else
-AOS_TO_SOA3_ISPC(double)
-AOS_TO_SOA3_ISPC(int64)
-#endif
+template<>
+unmasked inline void aos_to_soa3_ispc<double>(uniform double src[], varying double* uniform v0, varying double* uniform v1,
+    varying double* uniform v2)
+{
+    const uniform float* uniform srcflt = (const uniform float* uniform)src;
+    const varying float src0 = srcflt[programIndex];
+    const varying float src1 = srcflt[programCount + programIndex];
+    const varying float src2 = srcflt[2 * programCount + programIndex];
+    const varying float src3 = srcflt[3 * programCount + programIndex];
+    const varying float src4 = srcflt[4 * programCount + programIndex];
+    const varying float src5 = srcflt[5 * programCount + programIndex];
+
+    const varying float t0 = shuffle(src0, src1, vAOS301);
+    const varying float t1 = shuffle(src1, src2, vAOS311);
+    const varying float s0 = shuffle(src3, src4, vAOS301);
+    const varying float s1 = shuffle(src4, src5, vAOS311);
+    const varying float t2 = shuffle(src0, src1, vAOS341);
+    const varying float u2 = shuffle(t2, src2, vAOS351);
+    const varying float s2 = shuffle(src3, src4, vAOS341);
+
+    const varying float u0 = shuffle(t0, t1, vAOS321);
+    const varying float w0 = shuffle(s0, s1, vAOS321);
+    uniform double<TARGET_WIDTH> x0;
+    *((varying float *uniform)&x0) = u0;
+    *((varying float *uniform)&x0[TARGET_WIDTH/2]) = w0;
+    *v0 = *((varying double* uniform)&x0);
+
+    const varying float u1 = shuffle(t0, t1, vAOS331);
+    const varying float w1 = shuffle(s0, s1, vAOS331);
+    uniform double<TARGET_WIDTH> x1;
+    *((varying float *uniform)&x1) = u1;
+    *((varying float *uniform)&x1[TARGET_WIDTH/2]) = w1;
+    *v1 = *((varying double* uniform)&x1);
+
+    const varying float w2 = shuffle(s2, src5, vAOS351);
+    uniform double<TARGET_WIDTH> x2;
+    *((varying float *uniform)&x2) = u2;
+    *((varying float *uniform)&x2[TARGET_WIDTH/2]) = w2;
+    *v2 = *((varying double* uniform)&x2);
+}
 
+template<>
+unmasked inline void aos_to_soa3_ispc<int64>(uniform int64 src[], varying int64* uniform v0, varying int64* uniform v1,
+    varying int64* uniform v2)
+{
+    const uniform float* uniform srcflt = (const uniform float* uniform)src;
+    const varying float src0 = srcflt[programIndex];
+    const varying float src1 = srcflt[programCount + programIndex];
+    const varying float src2 = srcflt[2 * programCount + programIndex];
+    const varying float src3 = srcflt[3 * programCount + programIndex];
+    const varying float src4 = srcflt[4 * programCount + programIndex];
+    const varying float src5 = srcflt[5 * programCount + programIndex];
+
+    const varying float t0 = shuffle(src0, src1, vAOS301);
+    const varying float t1 = shuffle(src1, src2, vAOS311);
+    const varying float s0 = shuffle(src3, src4, vAOS301);
+    const varying float s1 = shuffle(src4, src5, vAOS311);
+    const varying float t2 = shuffle(src0, src1, vAOS341);
+    const varying float u2 = shuffle(t2, src2, vAOS351);
+    const varying float s2 = shuffle(src3, src4, vAOS341);
+
+    const varying float u0 = shuffle(t0, t1, vAOS321);
+    const varying float w0 = shuffle(s0, s1, vAOS321);
+    uniform int64<TARGET_WIDTH> x0;
+    *((varying float* uniform) & x0) = u0;
+    *((varying float* uniform) & x0[TARGET_WIDTH / 2]) = w0;
+    *v0 = *((varying int64* uniform) & x0);
+
+    const varying float u1 = shuffle(t0, t1, vAOS331);
+    const varying float w1 = shuffle(s0, s1, vAOS331);
+    uniform int64<TARGET_WIDTH> x1;
+    *((varying float* uniform) & x1) = u1;
+    *((varying float* uniform) & x1[TARGET_WIDTH / 2]) = w1;
+    *v1 = *((varying int64* uniform) & x1);
+
+    const varying float w2 = shuffle(s2, src5, vAOS351);
+    uniform int64<TARGET_WIDTH> x2;
+    *((varying float* uniform) & x2) = u2;
+    *((varying float* uniform) & x2[TARGET_WIDTH / 2]) = w2;
+    *v2 = *((varying int64* uniform) & x2);
+}
+#endif
 
 #if TARGET_WIDTH == 4
 static const varying int vAOS40 = {0, 4, 1, 5};
@@ -147,45 +169,176 @@ static const varying int vAOS41 = {2, 6, 3, 7};
 static const varying int vAOS42 = {0, 1, 4, 5};
 static const varying int vAOS43 = {2, 3, 6, 7};
 #elif TARGET_WIDTH == 8
-static const varying int vAOS40 = {0, 4, 8, 12, 1, 5, 9, 13};
-static const varying int vAOS41 = {2, 6, 10, 14, 3, 7, 11, 15};
-static const varying int vAOS42 = {0, 1, 2, 3, 8, 9, 10, 11};
-static const varying int vAOS43 = {4, 5, 6, 7, 12, 13, 14, 15};
+static const varying int vAOS40 = {0, 1, 2, 3, 8, 9, 10, 11};
+static const varying int vAOS41 = {4, 5, 6, 7, 12, 13, 14, 15};
+static const varying int vAOS42 = {0, 8, 1, 9, 4, 12, 5, 13};
+static const varying int vAOS43 = {2, 10, 3, 11, 6, 14, 7, 15};
+static const varying int vAOS44 = {0, 1, 8, 9, 4, 5, 12, 13};
+static const varying int vAOS45 = {2, 3, 10, 11, 6, 7, 14, 15};
+
+static const varying int vAOS401 = {0, 1, 8, 9, 4, 5, 12, 13};
+static const varying int vAOS411 = {2, 3, 10, 11, 6, 7, 14, 15};
+static const varying int vAOS421 = {0, 1, 2, 3, 8, 9, 10, 11};
+static const varying int vAOS431 = {4, 5, 6, 7, 12, 13, 14, 15};
 #elif TARGET_WIDTH == 16
-static const varying int vAOS40 = {0, 4, 8, 12, 16, 20, 24, 28, 1, 5, 9, 13, 17, 21, 25, 29};
-static const varying int vAOS41 = {2, 6, 10, 14, 18, 22, 26, 30, 3, 7, 11, 15, 19, 23, 27, 31};
-static const varying int vAOS42 = {0, 1, 2, 3, 4, 5, 6, 7, 16, 17, 18, 19, 20, 21, 22, 23};
-static const varying int vAOS43 = {8, 9, 10, 11, 12, 13, 14, 15, 24, 25, 26, 27, 28, 29, 30, 31};
+static const varying int vAOS40 = {0, 16, 1, 17, 4, 20, 5, 21, 8, 24, 9, 25, 12, 28, 13, 29};
+static const varying int vAOS41 = {2, 18, 3, 19, 6, 22, 7, 23, 10, 26, 11, 27, 14, 30, 15, 31};
+static const varying int vAOS42 = {0, 4, 8, 12, 1, 5, 9, 13, 16, 20, 24, 28, 17, 21, 25, 29};
+static const varying int vAOS43 = {2, 6, 10, 14, 3, 7, 11, 15, 18, 22, 26, 30, 19, 23, 27, 31};
 #else
 #error "No implementation for this target"
 #endif
 
-#define AOS_TO_SOA4_ISPC(T)                                                                                            \
-    unmasked inline void aos_to_soa4_ispc(uniform T src[], varying T *uniform v0, varying T *uniform v1,               \
-                                        varying T *uniform v2, varying T *uniform v3)                                  \
-    {                                                                                                                  \
-        const varying T src0 = src[programIndex];                                                                      \
-        const varying T src1 = src[programCount + programIndex];                                                       \
-        const varying T src2 = src[2 * programCount + programIndex];                                                   \
-        const varying T src3 = src[3 * programCount + programIndex];                                                   \
-                                                                                                                       \
-        const varying T t0 = shuffle(src2, src3, vAOS40);                                                              \
-        const varying T t1 = shuffle(src2, src3, vAOS41);                                                              \
-        const varying T t2 = shuffle(src0, src1, vAOS40);                                                              \
-        const varying T t3 = shuffle(src0, src1, vAOS41);                                                              \
-                                                                                                                       \
-        *v0 = shuffle(t2, t0, vAOS42);                                                                                 \
-        *v1 = shuffle(t2, t0, vAOS43);                                                                                 \
-        *v2 = shuffle(t3, t1, vAOS42);                                                                                 \
-        *v3 = shuffle(t3, t1, vAOS43);                                                                                 \
-    }
+#if TARGET_WIDTH == 8
+template<typename T>
+unmasked inline void aos_to_soa4_ispc(uniform T src[], varying T* uniform v0, varying T* uniform v1,
+    varying T* uniform v2, varying T* uniform v3)
+{
+    const varying T src0 = src[programIndex];
+    const varying T src1 = src[programCount + programIndex];
+    const varying T src2 = src[2 * programCount + programIndex];
+    const varying T src3 = src[3 * programCount + programIndex];
+
+    const varying T t0 = shuffle(src0, src2, vAOS40);
+    const varying T t1 = shuffle(src0, src2, vAOS41);
+    const varying T t2 = shuffle(src1, src3, vAOS40);
+    const varying T t3 = shuffle(src1, src3, vAOS41);
+
+    const varying T u0 = shuffle(t0, t1, vAOS42);
+    const varying T u1 = shuffle(t0, t1, vAOS43);
+    const varying T u2 = shuffle(t2, t3, vAOS42);
+    const varying T u3 = shuffle(t2, t3, vAOS43);
+
+    *v0 = shuffle(u0, u2, vAOS44);
+    *v1 = shuffle(u0, u2, vAOS45);
+    *v2 = shuffle(u1, u3, vAOS44);
+    *v3 = shuffle(u1, u3, vAOS45);
+}
+
+template<>
+unmasked inline void aos_to_soa4_ispc<double>(uniform double src[], varying double* uniform v0, varying double* uniform v1,
+    varying double* uniform v2, varying double* uniform v3)
+{
+    const uniform float* uniform srcflt = (const uniform float* uniform)src;
+    const varying float src0 = srcflt[programIndex];
+    const varying float src1 = srcflt[programCount + programIndex];
+    const varying float src2 = srcflt[2 * programCount + programIndex];
+    const varying float src3 = srcflt[3 * programCount + programIndex];
+    const varying float src4 = srcflt[4 * programCount + programIndex];
+    const varying float src5 = srcflt[5 * programCount + programIndex];
+    const varying float src6 = srcflt[6 * programCount + programIndex];
+    const varying float src7 = srcflt[7 * programCount + programIndex];
+
+    const varying float t0 = shuffle(src0, src1, vAOS401);
+    const varying float t1 = shuffle(src0, src1, vAOS411);
+    const varying float t2 = shuffle(src2, src3, vAOS401);
+    const varying float t3 = shuffle(src2, src3, vAOS411);
+    const varying float t4 = shuffle(src4, src5, vAOS401);
+    const varying float t5 = shuffle(src4, src5, vAOS411);
+    const varying float t6 = shuffle(src6, src7, vAOS401);
+    const varying float t7 = shuffle(src6, src7, vAOS411);
+
+    const varying float u0 = shuffle(t0, t2, vAOS421);
+    const varying float w0 = shuffle(t4, t6, vAOS421);
+    uniform double<TARGET_WIDTH> x0;
+    *((varying float* uniform)&x0) = u0;
+    *((varying float* uniform)&x0[TARGET_WIDTH/2]) = w0;
+    *v0 = *((varying double* uniform)&x0);
+
+    const varying float u1 = shuffle(t1, t3, vAOS421);
+    const varying float w1 = shuffle(t5, t7, vAOS421);
+    uniform double<TARGET_WIDTH> x1;
+    *((varying float* uniform)&x1) = u1;
+    *((varying float* uniform)&x1[TARGET_WIDTH/2]) = w1;
+    *v1 = *((varying double* uniform)&x1);
+
+    const varying float u2 = shuffle(t0, t2, vAOS431);
+    const varying float w2 = shuffle(t4, t6, vAOS431);
+    uniform double<TARGET_WIDTH> x2;
+    *((varying float* uniform)&x2) = u2;
+    *((varying float* uniform)&x2[TARGET_WIDTH/2]) = w2;
+    *v2 = *((varying double* uniform)&x2);
+
+    const varying float u3 = shuffle(t1, t3, vAOS431);
+    const varying float w3 = shuffle(t5, t7, vAOS431);
+    uniform double<TARGET_WIDTH> x3;
+    *((varying float* uniform)&x3) = u3;
+    *((varying float* uniform)&x3[TARGET_WIDTH/2]) = w3;
+    *v3 = *((varying double* uniform)&x3);
+}
+
+template<>
+unmasked inline void aos_to_soa4_ispc<int64>(uniform int64 src[], varying int64* uniform v0, varying int64* uniform v1,
+    varying int64* uniform v2, varying int64* uniform v3)
+{
+    const uniform float* uniform srcflt = (const uniform float* uniform)src;
+    const varying float src0 = srcflt[programIndex];
+    const varying float src1 = srcflt[programCount + programIndex];
+    const varying float src2 = srcflt[2 * programCount + programIndex];
+    const varying float src3 = srcflt[3 * programCount + programIndex];
+    const varying float src4 = srcflt[4 * programCount + programIndex];
+    const varying float src5 = srcflt[5 * programCount + programIndex];
+    const varying float src6 = srcflt[6 * programCount + programIndex];
+    const varying float src7 = srcflt[7 * programCount + programIndex];
+
+    const varying float t0 = shuffle(src0, src1, vAOS401);
+    const varying float t1 = shuffle(src0, src1, vAOS411);
+    const varying float t2 = shuffle(src2, src3, vAOS401);
+    const varying float t3 = shuffle(src2, src3, vAOS411);
+    const varying float t4 = shuffle(src4, src5, vAOS401);
+    const varying float t5 = shuffle(src4, src5, vAOS411);
+    const varying float t6 = shuffle(src6, src7, vAOS401);
+    const varying float t7 = shuffle(src6, src7, vAOS411);
+
+    const varying float u0 = shuffle(t0, t2, vAOS421);
+    const varying float w0 = shuffle(t4, t6, vAOS421);
+    uniform int64<TARGET_WIDTH> x0;
+    *((varying float* uniform) & x0) = u0;
+    *((varying float* uniform) & x0[TARGET_WIDTH / 2]) = w0;
+    *v0 = *((varying int64* uniform)&x0);
+
+    const varying float u1 = shuffle(t1, t3, vAOS421);
+    const varying float w1 = shuffle(t5, t7, vAOS421);
+    uniform int64<TARGET_WIDTH> x1;
+    *((varying float* uniform) & x1) = u1;
+    *((varying float* uniform) & x1[TARGET_WIDTH / 2]) = w1;
+    *v1 = *((varying int64* uniform)&x1);
+
+    const varying float u2 = shuffle(t0, t2, vAOS431);
+    const varying float w2 = shuffle(t4, t6, vAOS431);
+    uniform int64<TARGET_WIDTH> x2;
+    *((varying float* uniform) & x2) = u2;
+    *((varying float* uniform) & x2[TARGET_WIDTH / 2]) = w2;
+    *v2 = *((varying int64* uniform)&x2);
+
+    const varying float u3 = shuffle(t1, t3, vAOS431);
+    const varying float w3 = shuffle(t5, t7, vAOS431);
+    uniform int64<TARGET_WIDTH> x3;
+    *((varying float* uniform) & x3) = u3;
+    *((varying float* uniform) & x3[TARGET_WIDTH / 2]) = w3;
+    *v3 = *((varying int64* uniform)&x3);
+}
+#else
+template<typename T>
+unmasked inline void aos_to_soa4_ispc(uniform T src[], varying T *uniform v0, varying T *uniform v1,
+    varying T *uniform v2, varying T *uniform v3)
+{
+    const varying T src0 = src[programIndex];
+    const varying T src1 = src[programCount + programIndex];
+    const varying T src2 = src[2 * programCount + programIndex];
+    const varying T src3 = src[3 * programCount + programIndex];
 
-AOS_TO_SOA4_ISPC(int16)
-AOS_TO_SOA4_ISPC(int32)
-AOS_TO_SOA4_ISPC(float)
-AOS_TO_SOA4_ISPC(int64)
-AOS_TO_SOA4_ISPC(double)
+    const varying T t0 = shuffle(src0, src1, vAOS40);
+    const varying T t1 = shuffle(src0, src1, vAOS41);
+    const varying T t2 = shuffle(src2, src3, vAOS40);
+    const varying T t3 = shuffle(src2, src3, vAOS41);
 
+    *v0 = shuffle(t0, t2, vAOS42);
+    *v1 = shuffle(t0, t2, vAOS43);
+    *v2 = shuffle(t1, t3, vAOS42);
+    *v3 = shuffle(t1, t3, vAOS43);
+}
+#endif
 
 #if TARGET_WIDTH == 4
 static const varying int vAOS60 = {0, 1, 4, 5};
@@ -215,66 +368,351 @@ static const varying int vAOS67 = {0, 1, 2, 3, 4, 18, 24, 30, 8, 9, 10, 11, 12,
 #endif
 
 #if TARGET_WIDTH == 4
-#define AOS_TO_SOA6_ISPC(T)                                                                                            \
-    inline void aos_to_soa6_ispc(uniform T src[], varying T *uniform v0, varying T *uniform v1,                        \
-                                        varying T *uniform v2, varying T *uniform v3,                                  \
-                                        varying T *uniform v4, varying T *uniform v5)                                  \
-    {                                                                                                                  \
-        const varying T src0 = src[programIndex];                                                                      \
-        const varying T src1 = src[programCount + programIndex];                                                       \
-        const varying T src2 = src[2 * programCount + programIndex];                                                   \
-        const varying T src3 = src[3 * programCount + programIndex];                                                   \
-        const varying T src4 = src[4 * programCount + programIndex];                                                   \
-        const varying T src5 = src[5 * programCount + programIndex];                                                   \
-                                                                                                                       \
-        const varying T t0 = shuffle(src0, src1, vAOS62);                                                              \
-        const varying T t1 = shuffle(src3, src4, vAOS62);                                                              \
-        const varying T t2 = shuffle(src0, src2, vAOS63);                                                              \
-        const varying T t3 = shuffle(src3, src5, vAOS63);                                                              \
-        const varying T t4 = shuffle(src1, src2, vAOS62);                                                              \
-        const varying T t5 = shuffle(src4, src5, vAOS62);                                                              \
-                                                                                                                       \
-        *v0 = shuffle(t0, t1, vAOS60);                                                                                 \
-        *v1 = shuffle(t0, t1, vAOS61);                                                                                 \
-        *v2 = shuffle(t2, t3, vAOS60);                                                                                 \
-        *v3 = shuffle(t2, t3, vAOS61);                                                                                 \
-        *v4 = shuffle(t4, t5, vAOS60);                                                                                 \
-        *v5 = shuffle(t4, t5, vAOS61);                                                                                 \
-    }
+template<typename T>
+inline void aos_to_soa6_ispc(uniform T src[], varying T *uniform v0, varying T *uniform v1, varying T *uniform v2, varying T *uniform v3, varying T *uniform v4, varying T *uniform v5)                                  \
+{
+    const varying T src0 = src[programIndex];
+    const varying T src1 = src[programCount + programIndex];
+    const varying T src2 = src[2 * programCount + programIndex];
+    const varying T src3 = src[3 * programCount + programIndex];
+    const varying T src4 = src[4 * programCount + programIndex];
+    const varying T src5 = src[5 * programCount + programIndex];
+                                                     
+    const varying T t0 = shuffle(src0, src1, vAOS62);
+    const varying T t1 = shuffle(src3, src4, vAOS62);
+    const varying T t2 = shuffle(src0, src2, vAOS63);
+    const varying T t3 = shuffle(src3, src5, vAOS63);
+    const varying T t4 = shuffle(src1, src2, vAOS62);
+    const varying T t5 = shuffle(src4, src5, vAOS62);
+
+    *v0 = shuffle(t0, t1, vAOS60);
+    *v1 = shuffle(t0, t1, vAOS61);
+    *v2 = shuffle(t2, t3, vAOS60);
+    *v3 = shuffle(t2, t3, vAOS61);
+    *v4 = shuffle(t4, t5, vAOS60);
+    *v5 = shuffle(t4, t5, vAOS61);
+}
 #else
-#define AOS_TO_SOA6_ISPC(T)                                                                                            \
-    unmasked inline void aos_to_soa6_ispc(uniform T src[], varying T *uniform v0, varying T *uniform v1,               \
-                                        varying T *uniform v2, varying T *uniform v3,                                  \
-                                        varying T *uniform v4, varying T *uniform v5)                                  \
-    {                                                                                                                  \
-        const varying T src0 = src[programIndex];                                                                      \
-        const varying T src1 = src[programCount + programIndex];                                                       \
-        const varying T src2 = src[2 * programCount + programIndex];                                                   \
-        const varying T src3 = src[3 * programCount + programIndex];                                                   \
-        const varying T src4 = src[4 * programCount + programIndex];                                                   \
-        const varying T src5 = src[5 * programCount + programIndex];                                                   \
-                                                                                                                       \
-        const varying T t0 = shuffle(shuffle(src0, src1, vAOS62), src2, vAOS63);                                       \
-        const varying T t1 = shuffle(shuffle(src3, src4, vAOS62), src5, vAOS63);                                       \
-        const varying T t2 = shuffle(shuffle(src0, src1, vAOS64), src2, vAOS65);                                       \
-        const varying T t3 = shuffle(shuffle(src3, src4, vAOS64), src5, vAOS65);                                       \
-        const varying T t4 = shuffle(shuffle(src0, src1, vAOS66), src2, vAOS67);                                       \
-        const varying T t5 = shuffle(shuffle(src3, src4, vAOS66), src5, vAOS67);                                       \
-                                                                                                                       \
-        *v0 = shuffle(t0, t1, vAOS60);                                                                                 \
-        *v1 = shuffle(t0, t1, vAOS61);                                                                                 \
-        *v2 = shuffle(t2, t3, vAOS60);                                                                                 \
-        *v3 = shuffle(t2, t3, vAOS61);                                                                                 \
-        *v4 = shuffle(t4, t5, vAOS60);                                                                                 \
-        *v5 = shuffle(t4, t5, vAOS61);                                                                                 \
+template<typename T>
+unmasked inline void aos_to_soa6_ispc(uniform T src[], varying T *uniform v0, varying T *uniform v1, varying T *uniform v2, varying T *uniform v3, varying T *uniform v4, varying T *uniform v5)                                  \
+{
+    const varying T src0 = src[programIndex];
+    const varying T src1 = src[programCount + programIndex];
+    const varying T src2 = src[2 * programCount + programIndex];
+    const varying T src3 = src[3 * programCount + programIndex];
+    const varying T src4 = src[4 * programCount + programIndex];
+    const varying T src5 = src[5 * programCount + programIndex];
+
+    const varying T t0 = shuffle(shuffle(src0, src1, vAOS62), src2, vAOS63);
+    const varying T t1 = shuffle(shuffle(src3, src4, vAOS62), src5, vAOS63);
+    const varying T t2 = shuffle(shuffle(src0, src1, vAOS64), src2, vAOS65);
+    const varying T t3 = shuffle(shuffle(src3, src4, vAOS64), src5, vAOS65);
+    const varying T t4 = shuffle(shuffle(src0, src1, vAOS66), src2, vAOS67);
+    const varying T t5 = shuffle(shuffle(src3, src4, vAOS66), src5, vAOS67);
+
+    *v0 = shuffle(t0, t1, vAOS60);
+    *v1 = shuffle(t0, t1, vAOS61);
+    *v2 = shuffle(t2, t3, vAOS60);
+    *v3 = shuffle(t2, t3, vAOS61);
+    *v4 = shuffle(t4, t5, vAOS60);
+    *v5 = shuffle(t4, t5, vAOS61);
+}
+#endif
+
+template<typename T>
+unmasked inline void UniformLoad(const uniform T* varying SrcPtr, uniform T* uniform DstPtr)
+{
+    uniform T* uniform Src[programCount];
+
+    foreach(i = 0 ... programCount)
+    {
+        Src[i] = (uniform T* varying)&SrcPtr[0];
     }
+
+    DstPtr[0] = *Src[0];
+    DstPtr[1] = *Src[1];
+    DstPtr[2] = *Src[2];
+    DstPtr[3] = *Src[3];
+#if TARGET_WIDTH == 8 || TARGET_WIDTH == 16
+    DstPtr[4] = *Src[4];
+    DstPtr[5] = *Src[5];
+    DstPtr[6] = *Src[6];
+    DstPtr[7] = *Src[7];
+#endif
+#if TARGET_WIDTH == 16
+    DstPtr[8] = *Src[8];
+    DstPtr[9] = *Src[9];
+    DstPtr[10] = *Src[10];
+    DstPtr[11] = *Src[11];
+    DstPtr[12] = *Src[12];
+    DstPtr[13] = *Src[13];
+    DstPtr[14] = *Src[14];
+    DstPtr[15] = *Src[15];
 #endif
+}
 
-AOS_TO_SOA6_ISPC(int16)
-AOS_TO_SOA6_ISPC(int32)
-AOS_TO_SOA6_ISPC(float)
-AOS_TO_SOA6_ISPC(int64)
-AOS_TO_SOA6_ISPC(double)
+template<typename T>
+unmasked inline void AosToSoa3Explicit(const uniform T* varying SrcPtr, uniform T* uniform DstPtr)
+{
+    uniform T* uniform Src[programCount];
+
+    foreach(i = 0 ... programCount)
+    {
+        Src[i] = (uniform T * varying)&SrcPtr[0];
+    }
+
+#if TARGET_WIDTH == 4
+    // X, X, X, X
+    DstPtr[0] = Src[0][0];
+    DstPtr[1] = Src[1][0];
+    DstPtr[2] = Src[2][0];
+    DstPtr[3] = Src[3][0];
+
+    // Y, Y, Y, Y
+    DstPtr[4] = Src[0][1];
+    DstPtr[5] = Src[1][1];
+    DstPtr[6] = Src[2][1];
+    DstPtr[7] = Src[3][1];
+
+    // Z, Z, Z, Z
+    DstPtr[8] = Src[0][2];
+    DstPtr[9] = Src[1][2];
+    DstPtr[10] = Src[2][2];
+    DstPtr[11] = Src[3][2];
+#elif TARGET_WIDTH == 8
+    // X, X, X, X, X, X, X, X
+    DstPtr[0] = Src[0][0];
+    DstPtr[1] = Src[1][0];
+    DstPtr[2] = Src[2][0];
+    DstPtr[3] = Src[3][0];
+    DstPtr[4] = Src[4][0];
+    DstPtr[5] = Src[5][0];
+    DstPtr[6] = Src[6][0];
+    DstPtr[7] = Src[7][0];
+
+    // Y, Y, Y, Y, Y, Y, Y, Y
+    DstPtr[8] = Src[0][1];
+    DstPtr[9] = Src[1][1];
+    DstPtr[10] = Src[2][1];
+    DstPtr[11] = Src[3][1];
+    DstPtr[12] = Src[4][1];
+    DstPtr[13] = Src[5][1];
+    DstPtr[14] = Src[6][1];
+    DstPtr[15] = Src[7][1];
+
+    // Z, Z, Z, Z, Z, Z, Z, Z
+    DstPtr[16] = Src[0][2];
+    DstPtr[17] = Src[1][2];
+    DstPtr[18] = Src[2][2];
+    DstPtr[19] = Src[3][2];
+    DstPtr[20] = Src[4][2];
+    DstPtr[21] = Src[5][2];
+    DstPtr[22] = Src[6][2];
+    DstPtr[23] = Src[7][2];
+#elif TARGET_WIDTH == 16
+    // X, X, X, X, X, X, X, X, X, X, X, X, X, X, X, X
+    DstPtr[0] = Src[0][0];
+    DstPtr[1] = Src[1][0];
+    DstPtr[2] = Src[2][0];
+    DstPtr[3] = Src[3][0];
+    DstPtr[4] = Src[4][0];
+    DstPtr[5] = Src[5][0];
+    DstPtr[6] = Src[6][0];
+    DstPtr[7] = Src[7][0];
+    DstPtr[8] = Src[8][0];
+    DstPtr[9] = Src[9][0];
+    DstPtr[10] = Src[10][0];
+    DstPtr[11] = Src[11][0];
+    DstPtr[12] = Src[12][0];
+    DstPtr[13] = Src[13][0];
+    DstPtr[14] = Src[14][0];
+    DstPtr[15] = Src[15][0];
 
+    // Y, Y, Y, Y, Y, Y, Y, Y, Y, Y, Y, Y, Y, Y, Y, Y
+    DstPtr[16] = Src[0][1];
+    DstPtr[17] = Src[1][1];
+    DstPtr[18] = Src[2][1];
+    DstPtr[19] = Src[3][1];
+    DstPtr[20] = Src[4][1];
+    DstPtr[21] = Src[5][1];
+    DstPtr[22] = Src[6][1];
+    DstPtr[23] = Src[7][1];
+    DstPtr[24] = Src[8][1];
+    DstPtr[25] = Src[9][1];
+    DstPtr[26] = Src[10][1];
+    DstPtr[27] = Src[11][1];
+    DstPtr[28] = Src[12][1];
+    DstPtr[29] = Src[13][1];
+    DstPtr[30] = Src[14][1];
+    DstPtr[31] = Src[15][1];
+
+    // Z, Z, Z, Z, Z, Z, Z, Z, Z, Z, Z, Z, Z, Z, Z, Z
+    DstPtr[32] = Src[0][2];
+    DstPtr[33] = Src[1][2];
+    DstPtr[34] = Src[2][2];
+    DstPtr[35] = Src[3][2];
+    DstPtr[36] = Src[4][2];
+    DstPtr[37] = Src[5][2];
+    DstPtr[38] = Src[6][2];
+    DstPtr[39] = Src[7][2];
+    DstPtr[40] = Src[8][2];
+    DstPtr[41] = Src[9][2];
+    DstPtr[42] = Src[10][2];
+    DstPtr[43] = Src[11][2];
+    DstPtr[44] = Src[12][2];
+    DstPtr[45] = Src[13][2];
+    DstPtr[46] = Src[14][2];
+    DstPtr[47] = Src[15][2];
+#endif
+}
+
+template<typename T>
+unmasked inline void AosToSoa4Explicit(const uniform T* varying SrcPtr, uniform T* uniform DstPtr)
+{
+    uniform T* uniform Src[programCount];
+
+    foreach(i = 0 ... programCount)
+    {
+        Src[i] = (uniform T* varying)&SrcPtr[0];
+    }
+
+#if TARGET_WIDTH == 4
+    // X, X, X, X
+    DstPtr[0] = Src[0][0];
+    DstPtr[1] = Src[1][0];
+    DstPtr[2] = Src[2][0];
+    DstPtr[3] = Src[3][0];
+
+    // Y, Y, Y, Y
+    DstPtr[4] = Src[0][1];
+    DstPtr[5] = Src[1][1];
+    DstPtr[6] = Src[2][1];
+    DstPtr[7] = Src[3][1];
+
+    // Z, Z, Z, Z
+    DstPtr[8] = Src[0][2];
+    DstPtr[9] = Src[1][2];
+    DstPtr[10] = Src[2][2];
+    DstPtr[11] = Src[3][2];
+
+    // W, W, W, W
+    DstPtr[12] = Src[0][3];
+    DstPtr[13] = Src[1][3];
+    DstPtr[14] = Src[2][3];
+    DstPtr[15] = Src[3][3];
+#elif TARGET_WIDTH == 8
+    // X, X, X, X, X, X, X, X
+    DstPtr[0] = Src[0][0];
+    DstPtr[1] = Src[1][0];
+    DstPtr[2] = Src[2][0];
+    DstPtr[3] = Src[3][0];
+    DstPtr[4] = Src[4][0];
+    DstPtr[5] = Src[5][0];
+    DstPtr[6] = Src[6][0];
+    DstPtr[7] = Src[7][0];
+
+    // Y, Y, Y, Y, Y, Y, Y, Y
+    DstPtr[8] = Src[0][1];
+    DstPtr[9] = Src[1][1];
+    DstPtr[10] = Src[2][1];
+    DstPtr[11] = Src[3][1];
+    DstPtr[12] = Src[4][1];
+    DstPtr[13] = Src[5][1];
+    DstPtr[14] = Src[6][1];
+    DstPtr[15] = Src[7][1];
+
+    // Z, Z, Z, Z, Z, Z, Z, Z
+    DstPtr[16] = Src[0][2];
+    DstPtr[17] = Src[1][2];
+    DstPtr[18] = Src[2][2];
+    DstPtr[19] = Src[3][2];
+    DstPtr[20] = Src[4][2];
+    DstPtr[21] = Src[5][2];
+    DstPtr[22] = Src[6][2];
+    DstPtr[23] = Src[7][2];
+
+    // W, W, W, W, W, W, W, W
+    DstPtr[24] = Src[0][3];
+    DstPtr[25] = Src[1][3];
+    DstPtr[26] = Src[2][3];
+    DstPtr[27] = Src[3][3];
+    DstPtr[28] = Src[4][3];
+    DstPtr[29] = Src[5][3];
+    DstPtr[30] = Src[6][3];
+    DstPtr[31] = Src[7][3];
+#elif TARGET_WIDTH == 16
+    // X, X, X, X, X, X, X, X, X, X, X, X, X, X, X, X
+    DstPtr[0] = Src[0][0];
+    DstPtr[1] = Src[1][0];
+    DstPtr[2] = Src[2][0];
+    DstPtr[3] = Src[3][0];
+    DstPtr[4] = Src[4][0];
+    DstPtr[5] = Src[5][0];
+    DstPtr[6] = Src[6][0];
+    DstPtr[7] = Src[7][0];
+    DstPtr[8] = Src[8][0];
+    DstPtr[9] = Src[9][0];
+    DstPtr[10] = Src[10][0];
+    DstPtr[11] = Src[11][0];
+    DstPtr[12] = Src[12][0];
+    DstPtr[13] = Src[13][0];
+    DstPtr[14] = Src[14][0];
+    DstPtr[15] = Src[15][0];
+
+    // Y, Y, Y, Y, Y, Y, Y, Y, Y, Y, Y, Y, Y, Y, Y, Y
+    DstPtr[16] = Src[0][1];
+    DstPtr[17] = Src[1][1];
+    DstPtr[18] = Src[2][1];
+    DstPtr[19] = Src[3][1];
+    DstPtr[20] = Src[4][1];
+    DstPtr[21] = Src[5][1];
+    DstPtr[22] = Src[6][1];
+    DstPtr[23] = Src[7][1];
+    DstPtr[24] = Src[8][1];
+    DstPtr[25] = Src[9][1];
+    DstPtr[26] = Src[10][1];
+    DstPtr[27] = Src[11][1];
+    DstPtr[28] = Src[12][1];
+    DstPtr[29] = Src[13][1];
+    DstPtr[30] = Src[14][1];
+    DstPtr[31] = Src[15][1];
+
+    // Z, Z, Z, Z, Z, Z, Z, Z, Z, Z, Z, Z, Z, Z, Z, Z
+    DstPtr[32] = Src[0][2];
+    DstPtr[33] = Src[1][2];
+    DstPtr[34] = Src[2][2];
+    DstPtr[35] = Src[3][2];
+    DstPtr[36] = Src[4][2];
+    DstPtr[37] = Src[5][2];
+    DstPtr[38] = Src[6][2];
+    DstPtr[39] = Src[7][2];
+    DstPtr[40] = Src[8][2];
+    DstPtr[41] = Src[9][2];
+    DstPtr[42] = Src[10][2];
+    DstPtr[43] = Src[11][2];
+    DstPtr[44] = Src[12][2];
+    DstPtr[45] = Src[13][2];
+    DstPtr[46] = Src[14][2];
+    DstPtr[47] = Src[15][2];
+
+    // W, W, W, W, W, W, W, W, W, W, W, W, W, W, W, W
+    DstPtr[48] = Src[0][3];
+    DstPtr[49] = Src[1][3];
+    DstPtr[50] = Src[2][3];
+    DstPtr[51] = Src[3][3];
+    DstPtr[52] = Src[4][3];
+    DstPtr[53] = Src[5][3];
+    DstPtr[54] = Src[6][3];
+    DstPtr[55] = Src[7][3];
+    DstPtr[56] = Src[8][3];
+    DstPtr[57] = Src[9][3];
+    DstPtr[58] = Src[10][3];
+    DstPtr[59] = Src[11][3];
+    DstPtr[60] = Src[12][3];
+    DstPtr[61] = Src[13][3];
+    DstPtr[62] = Src[14][3];
+    DstPtr[64] = Src[15][3];
+#endif
+}
 
 #endif
diff --git a/Engine/Source/Runtime/Core/Public/Math/Matrix.isph b/Engine/Source/Runtime/Core/Public/Math/Matrix.isph
--- a/Engine/Source/Runtime/Core/Public/Math/Matrix.isph
+++ b/Engine/Source/Runtime/Core/Public/Math/Matrix.isph
@@ -176,14 +176,8 @@ inline uniform FVector MatrixGetOrigin(const uniform FMatrix &M)
 	return SetVector(M.M[12], M.M[13], M.M[14]);
 }
 
-inline void MatrixGetScaledAxes(const uniform FMatrix44d &M, uniform FVector3d &X, uniform FVector3d &Y, uniform FVector3d &Z)
-{
-	X = SetVector(M.M[0], M.M[1], M.M[2]);
-	Y = SetVector(M.M[4], M.M[5], M.M[6]);
-	Z = SetVector(M.M[8], M.M[9], M.M[10]);
-}
-
-inline void MatrixGetScaledAxes(const uniform FMatrix44f &M, uniform FVector3f &X, uniform FVector3f &Y, uniform FVector3f &Z)
+template <typename T, typename V>
+inline void MatrixGetScaledAxes(const T& M, V &X, V &Y, V &Z)
 {
 	X = SetVector(M.M[0], M.M[1], M.M[2]);
 	Y = SetVector(M.M[4], M.M[5], M.M[6]);
@@ -271,41 +265,27 @@ inline uniform FMatrix MatrixTranspose(const uniform FMatrix& M)
 // we use __m128 to represent 2x2 matrix as A = | A0  A1 |
 //                                              | A2  A3 |
 // 2x2 row major Matrix multiply A*B
-static inline uniform FVector4d Mat2Mul(const uniform FVector4d& vec1, const uniform FVector4d& vec2)
-{
-	return
-		VectorAdd(VectorMultiply( vec1, VectorSwizzle(vec2, 0,3,0,3)),
-			VectorMultiply(VectorSwizzle(vec1, 1,0,3,2), VectorSwizzle(vec2, 2,1,2,1)));
-}
-static inline uniform FVector4f Mat2Mul(const uniform FVector4f& vec1, const uniform FVector4f& vec2)
+template <typename T>
+static inline uniform T Mat2Mul(const uniform T& vec1, const uniform T& vec2)
 {
 	return
 		VectorAdd(VectorMultiply( vec1, VectorSwizzle(vec2, 0,3,0,3)),
 			VectorMultiply(VectorSwizzle(vec1, 1,0,3,2), VectorSwizzle(vec2, 2,1,2,1)));
 }
+
 // 2x2 row major Matrix adjugate multiply (A#)*B
-static inline uniform FVector4d Mat2AdjMul(const uniform FVector4d& vec1, const uniform FVector4d& vec2)
+template <typename T>
+static inline uniform T Mat2AdjMul(const uniform T& vec1, const uniform T& vec2)
 {
 	return
 		VectorSubtract(VectorMultiply(VectorSwizzle(vec1, 3,3,0,0), vec2),
 			VectorMultiply(VectorSwizzle(vec1, 1,1,2,2), VectorSwizzle(vec2, 2,3,0,1)));
 
 }
-static inline uniform FVector4f Mat2AdjMul(const uniform FVector4f& vec1, const uniform FVector4f& vec2)
-{
-	return
-		VectorSubtract(VectorMultiply(VectorSwizzle(vec1, 3,3,0,0), vec2),
-			VectorMultiply(VectorSwizzle(vec1, 1,1,2,2), VectorSwizzle(vec2, 2,3,0,1)));
 
-}
 // 2x2 row major Matrix multiply adjugate A*(B#)
-static inline uniform FVector4d Mat2MulAdj(const uniform FVector4d& vec1, const uniform FVector4d& vec2)
-{
-	return
-		VectorSubtract(VectorMultiply( vec1, VectorSwizzle(vec2, 3,0,3,0)),
-			VectorMultiply(VectorSwizzle(vec1, 1,0,3,2), VectorSwizzle(vec2, 2,1,2,1)));
-}
-static inline uniform FVector4f Mat2MulAdj(const uniform FVector4f& vec1, const uniform FVector4f& vec2)
+template <typename T>
+static inline uniform T Mat2MulAdj(const uniform T& vec1, const uniform T& vec2)
 {
 	return
 		VectorSubtract(VectorMultiply( vec1, VectorSwizzle(vec2, 3,0,3,0)),
@@ -572,33 +552,10 @@ inline uniform FMatrix44f MatrixInverse(const uniform FMatrix44f& M)
 	return Result;
 }
 
-inline uniform FVector4d VectorTransformVector(const uniform FVector4d &VecP, const uniform FMatrix44d &M)
-{
-	uniform FVector4d VTempX, VTempY, VTempZ, VTempW;
-
-	// Splat x,y,z and w
-	VTempX = VectorReplicate(VecP, 0);
-	VTempY = VectorReplicate(VecP, 1);
-	VTempZ = VectorReplicate(VecP, 2);
-	VTempW = VectorReplicate(VecP, 3);
-
-	// Mul by the matrix
-	VTempX = VectorMultiply(VTempX, SetVector4(M.M[0], M.M[1], M.M[2], M.M[3]));
-	VTempY = VectorMultiply(VTempY, SetVector4(M.M[4], M.M[5], M.M[6], M.M[7]));
-	VTempZ = VectorMultiply(VTempZ, SetVector4(M.M[8], M.M[9], M.M[10], M.M[11]));
-	VTempW = VectorMultiply(VTempW, SetVector4(M.M[12], M.M[13], M.M[14], M.M[15]));
-
-	// Add them all together
-	VTempX = VectorAdd(VTempX, VTempY);
-	VTempZ = VectorAdd(VTempZ, VTempW);
-	VTempX = VectorAdd(VTempX, VTempZ);
-
-	return VTempX;
-}
-
-inline uniform FVector4f VectorTransformVector(const uniform FVector4f &VecP, const uniform FMatrix44f &M)
+template <typename T, typename V>
+inline T VectorTransformVector(const T& VecP, const V& M)
 {
-	uniform FVector4f VTempX, VTempY, VTempZ, VTempW;
+	T VTempX, VTempY, VTempZ, VTempW;
 
 	// Splat x,y,z and w
 	VTempX = VectorReplicate(VecP, 0);
@@ -620,34 +577,34 @@ inline uniform FVector4f VectorTransformVector(const uniform FVector4f &VecP, co
 	return VTempX;
 }
 
-inline FVector4 VectorTransformVector(const FVector4 &VecP, const FMatrix &M)
+// Calculate homogeneous transform. W component assumed to be 1.0
+inline FVector MatrixTransformPosition(const FVector &P, const FMatrix &M)
 {
-	FVector4 VTempX, VTempY, VTempZ, VTempW;
+	FVector VTempX, VTempY, VTempZ;
 
-	// Splat x,y,z and w
-	VTempX = VectorReplicate(VecP, 0);
-	VTempY = VectorReplicate(VecP, 1);
-	VTempZ = VectorReplicate(VecP, 2);
-	VTempW = VectorReplicate(VecP, 3);
+	// Splat x,y,z
+	VTempX = SetVector(P.V[0], P.V[0], P.V[0]);
+	VTempY = SetVector(P.V[1], P.V[1], P.V[1]);
+	VTempZ = SetVector(P.V[2], P.V[2], P.V[2]);
 
 	// Mul by the matrix
-	VTempX = VectorMultiply(VTempX, SetVector4(M.M[0], M.M[1], M.M[2], M.M[3]));
-	VTempY = VectorMultiply(VTempY, SetVector4(M.M[4], M.M[5], M.M[6], M.M[7]));
-	VTempZ = VectorMultiply(VTempZ, SetVector4(M.M[8], M.M[9], M.M[10], M.M[11]));
-	VTempW = VectorMultiply(VTempW, SetVector4(M.M[12], M.M[13], M.M[14], M.M[15]));
+	VTempX = VTempX * SetVector(M.M[0], M.M[1], M.M[2]);
+	VTempY = VTempY * SetVector(M.M[4], M.M[5], M.M[6]);
+	VTempZ = VTempZ * SetVector(M.M[8], M.M[9], M.M[10]);
+	const FVector VTempW = SetVector(M.M[12], M.M[13], M.M[14]);
 
 	// Add them all together
-	VTempX = VectorAdd(VTempX, VTempY);
-	VTempZ = VectorAdd(VTempZ, VTempW);
-	VTempX = VectorAdd(VTempX, VTempZ);
+	VTempX = VTempX + VTempY;
+	VTempZ = VTempZ + VTempW;
+	VTempX = VTempX + VTempZ;
 
 	return VTempX;
 }
 
 // Calculate homogeneous transform. W component assumed to be 1.0
-inline FVector MatrixTransformPosition(const FVector &P, const FMatrix &M)
+inline FVector3d MatrixTransformPosition(const FVector3d &P, const uniform FMatrix44d &M)
 {
-	FVector VTempX, VTempY, VTempZ;
+	FVector3d VTempX, VTempY, VTempZ;
 
 	// Splat x,y,z
 	VTempX = SetVector(P.V[0], P.V[0], P.V[0]);
@@ -658,7 +615,7 @@ inline FVector MatrixTransformPosition(const FVector &P, const FMatrix &M)
 	VTempX = VTempX * SetVector(M.M[0], M.M[1], M.M[2]);
 	VTempY = VTempY * SetVector(M.M[4], M.M[5], M.M[6]);
 	VTempZ = VTempZ * SetVector(M.M[8], M.M[9], M.M[10]);
-	const FVector VTempW = SetVector(M.M[12], M.M[13], M.M[14]);
+	const uniform FVector3d VTempW = SetVector(M.M[12], M.M[13], M.M[14]);
 
 	// Add them all together
 	VTempX = VTempX + VTempY;
@@ -668,10 +625,9 @@ inline FVector MatrixTransformPosition(const FVector &P, const FMatrix &M)
 	return VTempX;
 }
 
-// Calculate homogeneous transform. W component assumed to be 1.0
-inline FVector3d MatrixTransformPosition(const FVector3d &P, const uniform FMatrix44d &M)
+inline FVector3f MatrixTransformPosition(const FVector3f &P, const uniform FMatrix44f &M)
 {
-	FVector3d VTempX, VTempY, VTempZ;
+	FVector3f VTempX, VTempY, VTempZ;
 
 	// Splat x,y,z
 	VTempX = SetVector(P.V[0], P.V[0], P.V[0]);
@@ -682,7 +638,7 @@ inline FVector3d MatrixTransformPosition(const FVector3d &P, const uniform FMatr
 	VTempX = VTempX * SetVector(M.M[0], M.M[1], M.M[2]);
 	VTempY = VTempY * SetVector(M.M[4], M.M[5], M.M[6]);
 	VTempZ = VTempZ * SetVector(M.M[8], M.M[9], M.M[10]);
-	const uniform FVector3d VTempW = SetVector(M.M[12], M.M[13], M.M[14]);
+	const uniform FVector3f VTempW = SetVector(M.M[12], M.M[13], M.M[14]);
 
 	// Add them all together
 	VTempX = VTempX + VTempY;
@@ -692,9 +648,10 @@ inline FVector3d MatrixTransformPosition(const FVector3d &P, const uniform FMatr
 	return VTempX;
 }
 
-inline FVector3f MatrixTransformPosition(const FVector3f &P, const uniform FMatrix44f &M)
+// Calculate homogeneous transform. W component assumed to be 1.0
+inline uniform FVector3f MatrixTransformPosition(const uniform FVector3f& P, const uniform FMatrix44f& M)
 {
-	FVector3f VTempX, VTempY, VTempZ;
+	uniform FVector3f VTempX, VTempY, VTempZ;
 
 	// Splat x,y,z
 	VTempX = SetVector(P.V[0], P.V[0], P.V[0]);
@@ -746,17 +703,8 @@ inline uniform FVector3f MatrixInverseTransformVector(const uniform FMatrix44f &
 	return SetVector(VectorTransformVector(SetVector4(V, FLOAT_ZERO), InvSelf));
 }
 
-inline uniform FMatrix44d MatrixReduceAdd(const varying FMatrix44d &M)
-{
-	return SetMatrix(
-		SetVector4(reduce_add(M.M[0]), reduce_add(M.M[1]), reduce_add(M.M[2]), reduce_add(M.M[3])),
-		SetVector4(reduce_add(M.M[4]), reduce_add(M.M[5]), reduce_add(M.M[6]), reduce_add(M.M[7])),
-		SetVector4(reduce_add(M.M[8]), reduce_add(M.M[9]), reduce_add(M.M[10]), reduce_add(M.M[11])),
-		SetVector4(reduce_add(M.M[12]), reduce_add(M.M[13]), reduce_add(M.M[14]), reduce_add(M.M[15]))
-		);
-}
-
-inline uniform FMatrix44f MatrixReduceAdd(const varying FMatrix44f &M)
+template <typename T>
+inline uniform T MatrixReduceAdd(const varying T& M)
 {
 	return SetMatrix(
 		SetVector4(reduce_add(M.M[0]), reduce_add(M.M[1]), reduce_add(M.M[2]), reduce_add(M.M[3])),
diff --git a/Engine/Source/Runtime/Core/Public/Math/Quat.isph b/Engine/Source/Runtime/Core/Public/Math/Quat.isph
--- a/Engine/Source/Runtime/Core/Public/Math/Quat.isph
+++ b/Engine/Source/Runtime/Core/Public/Math/Quat.isph
@@ -80,19 +80,10 @@ inline uniform FVector4 MatrixToQuat(const uniform FMatrix &M)
 * @param Quat2	Pointer to the second quaternion
 * @return Quat1 * Quat2
 */
-inline FVector4 VectorQuaternionMultiply2( const FVector4& Quat1, const FVector4& Quat2 )
+template <typename T>
+inline T VectorQuaternionMultiply2(const T& Quat1, const T& Quat2)
 {
-	FVector4 Result = VectorReplicate(Quat1, 3) * Quat2;
-	Result = VectorMultiplyAdd((VectorReplicate(Quat1, 0) * VectorSwizzle(Quat2, 3,2,1,0)), QMULTI_SIGN_MASK0, Result);
-	Result = VectorMultiplyAdd((VectorReplicate(Quat1, 1) * VectorSwizzle(Quat2, 2,3,0,1)), QMULTI_SIGN_MASK1, Result);
-	Result = VectorMultiplyAdd((VectorReplicate(Quat1, 2) * VectorSwizzle(Quat2, 1,0,3,2)), QMULTI_SIGN_MASK2, Result);
-
-	return Result;
-}
-
-inline uniform FVector4 VectorQuaternionMultiply2( const uniform FVector4& Quat1, const uniform FVector4& Quat2 )
-{
-	uniform FVector4 Result = VectorReplicate(Quat1, 3) * Quat2;
+	T Result = VectorReplicate(Quat1, 3) * Quat2;
 	Result = VectorMultiplyAdd((VectorReplicate(Quat1, 0) * VectorSwizzle(Quat2, 3,2,1,0)), QMULTI_SIGN_MASK0, Result);
 	Result = VectorMultiplyAdd((VectorReplicate(Quat1, 1) * VectorSwizzle(Quat2, 2,3,0,1)), QMULTI_SIGN_MASK1, Result);
 	Result = VectorMultiplyAdd((VectorReplicate(Quat1, 2) * VectorSwizzle(Quat2, 1,0,3,2)), QMULTI_SIGN_MASK2, Result);
@@ -131,28 +122,13 @@ inline uniform FVector4f QuatInverse(const uniform FVector4f &Quat)
 	return Quat * FLOAT_QINV_SIGN_MASK;
 }
 
-inline FVector4d QuatFastLerp(const FVector4d& A, const FVector4d& B, const double Alpha)
+template <typename T, typename F>
+inline T QuatFastLerp(const T& A, const T& B, const F Alpha)
 {
 	// To ensure the 'shortest route', we make sure the dot product between the both rotations is positive.
-	const double DotResult = VectorDot(A, B);
-	const double Bias = select(DotResult >= 0.d, 1.d, -1.d);
-	return (B * Alpha) + (A * (Bias * (1.d - Alpha)));
-}
-
-inline FVector4f QuatFastLerp(const FVector4f& A, const FVector4f& B, const float Alpha)
-{
-	// To ensure the 'shortest route', we make sure the dot product between the both rotations is positive.
-	const float DotResult = VectorDot(A, B);
-	const float Bias = select(DotResult >= 0.f, 1.f, -1.f);
-	return (B * Alpha) + (A * (Bias * (1.f - Alpha)));
-}
-
-inline uniform FVector4 QuatFastLerp(const uniform FVector4& A, const uniform FVector4& B, const uniform FReal Alpha)
-{
-	// To ensure the 'shortest route', we make sure the dot product between the both rotations is positive.
-	const uniform FReal DotResult = VectorDot(A, B);
-	const uniform FReal Bias = select(DotResult >= ZERO, ONE, -ONE);
-	return (B * Alpha) + (A * (Bias * (ONE - Alpha)));
+	const F DotResult = VectorDot(A, B);
+	const F Bias = select(DotResult >= 0, 1, -1);
+	return (B * Alpha) + (A * (Bias * (1 - Alpha)));
 }
 
 // A and B are quaternions.  The result is A + (|A.B| >= 0 ? 1 : -1) * B
@@ -176,103 +152,8 @@ inline uniform FVector4 VectorAccumulateQuaternionShortestPath(const uniform FVe
  * @param VectorW0 Vector to rotate. W component must be zero.
  * @return Vector after rotation by Quat.
  */
-inline FVector4d VectorQuaternionRotateVector(const FVector4d& Quat, const FVector4d& VectorW0)
-{
-	// Q * V * Q.Inverse
-	//const VectorRegister InverseRotation = VectorQuaternionInverse(Quat);
-	//const VectorRegister Temp = VectorQuaternionMultiply2(Quat, VectorW0);
-	//const VectorRegister Rotated = VectorQuaternionMultiply2(Temp, InverseRotation);
-
-	// Equivalence of above can be shown to be:
-	// http://people.csail.mit.edu/bkph/articles/Quaternions.pdf
-	// V' = V + 2w(Q x V) + (2Q x (Q x V))
-	// refactor:
-	// V' = V + w(2(Q x V)) + (Q x (2(Q x V)))
-	// T = 2(Q x V);
-	// V' = V + w*(T) + (Q x T)
-
-	const FVector4d QW = VectorReplicate(Quat, 3);
-	FVector4d T = VectorCross(Quat, VectorW0);
-	T = VectorAdd(T, T);
-	const FVector4d VTemp0 = VectorMultiplyAdd(QW, T, VectorW0);
-	const FVector4d VTemp1 = VectorCross(Quat, T);
-	const FVector4d Rotated = VectorAdd(VTemp0, VTemp1);
-	return Rotated;
-}
-
-inline FVector4f VectorQuaternionRotateVector(const FVector4f& Quat, const FVector4f& VectorW0)
-{
-	// Q * V * Q.Inverse
-	//const VectorRegister InverseRotation = VectorQuaternionInverse(Quat);
-	//const VectorRegister Temp = VectorQuaternionMultiply2(Quat, VectorW0);
-	//const VectorRegister Rotated = VectorQuaternionMultiply2(Temp, InverseRotation);
-
-	// Equivalence of above can be shown to be:
-	// http://people.csail.mit.edu/bkph/articles/Quaternions.pdf
-	// V' = V + 2w(Q x V) + (2Q x (Q x V))
-	// refactor:
-	// V' = V + w(2(Q x V)) + (Q x (2(Q x V)))
-	// T = 2(Q x V);
-	// V' = V + w*(T) + (Q x T)
-
-	const FVector4f QW = VectorReplicate(Quat, 3);
-	FVector4f T = VectorCross(Quat, VectorW0);
-	T = VectorAdd(T, T);
-	const FVector4f VTemp0 = VectorMultiplyAdd(QW, T, VectorW0);
-	const FVector4f VTemp1 = VectorCross(Quat, T);
-	const FVector4f Rotated = VectorAdd(VTemp0, VTemp1);
-	return Rotated;
-}
-
-inline FVector4 VectorQuaternionRotateVector(const uniform FVector4& Quat, const FVector4& VectorW0)
-{
-	// Q * V * Q.Inverse
-	//const VectorRegister InverseRotation = VectorQuaternionInverse(Quat);
-	//const VectorRegister Temp = VectorQuaternionMultiply2(Quat, VectorW0);
-	//const VectorRegister Rotated = VectorQuaternionMultiply2(Temp, InverseRotation);
-
-	// Equivalence of above can be shown to be:
-	// http://people.csail.mit.edu/bkph/articles/Quaternions.pdf
-	// V' = V + 2w(Q x V) + (2Q x (Q x V))
-	// refactor:
-	// V' = V + w(2(Q x V)) + (Q x (2(Q x V)))
-	// T = 2(Q x V);
-	// V' = V + w*(T) + (Q x T)
-
-	const uniform FVector4 QW = VectorReplicate(Quat, 3);
-	FVector4 T = VectorCross(Quat, VectorW0);
-	T = VectorAdd(T, T);
-	const FVector4 VTemp0 = VectorMultiplyAdd(QW, T, VectorW0);
-	const FVector4 VTemp1 = VectorCross(Quat, T);
-	const FVector4 Rotated = VectorAdd(VTemp0, VTemp1);
-	return Rotated;
-}
-
-inline FVector4f VectorQuaternionRotateVector(const uniform FVector4f& Quat, const FVector4f& VectorW0)
-{
-	// Q * V * Q.Inverse
-	//const VectorRegister InverseRotation = VectorQuaternionInverse(Quat);
-	//const VectorRegister Temp = VectorQuaternionMultiply2(Quat, VectorW0);
-	//const VectorRegister Rotated = VectorQuaternionMultiply2(Temp, InverseRotation);
-
-	// Equivalence of above can be shown to be:
-	// http://people.csail.mit.edu/bkph/articles/Quaternions.pdf
-	// V' = V + 2w(Q x V) + (2Q x (Q x V))
-	// refactor:
-	// V' = V + w(2(Q x V)) + (Q x (2(Q x V)))
-	// T = 2(Q x V);
-	// V' = V + w*(T) + (Q x T)
-
-	const uniform FVector4f QW = VectorReplicate(Quat, 3);
-	FVector4f T = VectorCross(Quat, VectorW0);
-	T = VectorAdd(T, T);
-	const FVector4f VTemp0 = VectorMultiplyAdd(QW, T, VectorW0);
-	const FVector4f VTemp1 = VectorCross(Quat, T);
-	const FVector4f Rotated = VectorAdd(VTemp0, VTemp1);
-	return Rotated;
-}
-
-inline uniform FVector4d VectorQuaternionRotateVector(const uniform FVector4d& Quat, const uniform FVector4d& VectorW0)
+template<typename T, typename V>
+inline V VectorQuaternionRotateVector(const T& Quat, const V& VectorW0)
 {
 	// Q * V * Q.Inverse
 	//const VectorRegister InverseRotation = VectorQuaternionInverse(Quat);
@@ -287,40 +168,17 @@ inline uniform FVector4d VectorQuaternionRotateVector(const uniform FVector4d& Q
 	// T = 2(Q x V);
 	// V' = V + w*(T) + (Q x T)
 
-	const uniform FVector4d QW = VectorReplicate(Quat, 3);
-	uniform FVector4d T = VectorCross(Quat, VectorW0);
-	T = VectorAdd(T, T);
-	const uniform FVector4d VTemp0 = VectorMultiplyAdd(QW, T, VectorW0);
-	const uniform FVector4d VTemp1 = VectorCross(Quat, T);
-	const uniform FVector4d Rotated = VectorAdd(VTemp0, VTemp1);
-	return Rotated;
-}
-
-inline uniform FVector4f VectorQuaternionRotateVector(const uniform FVector4f& Quat, const uniform FVector4f& VectorW0)
-{
-	// Q * V * Q.Inverse
-	//const VectorRegister InverseRotation = VectorQuaternionInverse(Quat);
-	//const VectorRegister Temp = VectorQuaternionMultiply2(Quat, VectorW0);
-	//const VectorRegister Rotated = VectorQuaternionMultiply2(Temp, InverseRotation);
-
-	// Equivalence of above can be shown to be:
-	// http://people.csail.mit.edu/bkph/articles/Quaternions.pdf
-	// V' = V + 2w(Q x V) + (2Q x (Q x V))
-	// refactor:
-	// V' = V + w(2(Q x V)) + (Q x (2(Q x V)))
-	// T = 2(Q x V);
-	// V' = V + w*(T) + (Q x T)
-
-	const uniform FVector4f QW = VectorReplicate(Quat, 3);
-	uniform FVector4f T = VectorCross(Quat, VectorW0);
-	T = VectorAdd(T, T);
-	const uniform FVector4f VTemp0 = VectorMultiplyAdd(QW, T, VectorW0);
-	const uniform FVector4f VTemp1 = VectorCross(Quat, T);
-	const uniform FVector4f Rotated = VectorAdd(VTemp0, VTemp1);
+	const V QW = VectorReplicate(Quat, 3);
+	V Q = VectorCross(Quat, VectorW0);
+	Q = VectorAdd(Q, Q);
+	const V VTemp0 = VectorMultiplyAdd(QW, Q, VectorW0);
+	const V VTemp1 = VectorCross(Quat, Q);
+	const V Rotated = VectorAdd(VTemp0, VTemp1);
 	return Rotated;
 }
 
-inline uniform FVector VectorQuaternionRotateVector(const uniform FVector4& Quat, const uniform FVector& V)
+template<>
+inline uniform FVector VectorQuaternionRotateVector<uniform FVector4, uniform FVector>(const uniform FVector4& Quat, const uniform FVector& V)
 {
 	// http://people.csail.mit.edu/bkph/articles/Quaternions.pdf
 	// V' = V + 2w(Q x V) + (2Q x (Q x V))
@@ -335,7 +193,8 @@ inline uniform FVector VectorQuaternionRotateVector(const uniform FVector4& Quat
 	return Result;
 }
 
-inline FVector3d VectorQuaternionRotateVector(const FVector4d& Quat, const FVector3d& V)
+template<>
+inline FVector3d VectorQuaternionRotateVector<FVector4d, FVector3d>(const FVector4d& Quat, const FVector3d& V)
 {
 	// http://people.csail.mit.edu/bkph/articles/Quaternions.pdf
 	// V' = V + 2w(Q x V) + (2Q x (Q x V))
@@ -350,7 +209,8 @@ inline FVector3d VectorQuaternionRotateVector(const FVector4d& Quat, const FVect
 	return Result;
 }
 
-inline FVector3f VectorQuaternionRotateVector(const FVector4f& Quat, const FVector3f& V)
+template<>
+inline FVector3f VectorQuaternionRotateVector<FVector4f, FVector3f>(const FVector4f& Quat, const FVector3f& V)
 {
 	// http://people.csail.mit.edu/bkph/articles/Quaternions.pdf
 	// V' = V + 2w(Q x V) + (2Q x (Q x V))
@@ -365,7 +225,8 @@ inline FVector3f VectorQuaternionRotateVector(const FVector4f& Quat, const FVect
 	return Result;
 }
 
-inline FVector3d VectorQuaternionRotateVector(const uniform FVector4d& Quat, const FVector3d& V)
+template<>
+inline FVector3d VectorQuaternionRotateVector<uniform FVector4d, FVector3d>(const uniform FVector4d& Quat, const FVector3d& V)
 {
 	// http://people.csail.mit.edu/bkph/articles/Quaternions.pdf
 	// V' = V + 2w(Q x V) + (2Q x (Q x V))
@@ -380,7 +241,8 @@ inline FVector3d VectorQuaternionRotateVector(const uniform FVector4d& Quat, con
 	return Result;
 }
 
-inline FVector3f VectorQuaternionRotateVector(const uniform FVector4f& Quat, const FVector3f& V)
+template<>
+inline FVector3f VectorQuaternionRotateVector<uniform FVector4f, FVector3f>(const uniform FVector4f& Quat, const FVector3f& V)
 {
 	// http://people.csail.mit.edu/bkph/articles/Quaternions.pdf
 	// V' = V + 2w(Q x V) + (2Q x (Q x V))
@@ -395,7 +257,8 @@ inline FVector3f VectorQuaternionRotateVector(const uniform FVector4f& Quat, con
 	return Result;
 }
 
-inline uniform FVector8 VectorQuaternionRotateVector(const uniform FVector8& Quat, const uniform FVector8& VectorW0)
+template<>
+inline uniform FVector8 VectorQuaternionRotateVector<uniform FVector8, uniform FVector8>(const uniform FVector8& Quat, const uniform FVector8& VectorW0)
 {
 	const uniform FVector8 QW = VectorReplicate(Quat, 3);
 	uniform FVector8 T = VectorCross(Quat, VectorW0);
diff --git a/Engine/Source/Runtime/Core/Public/Math/Soaaos.isph b/Engine/Source/Runtime/Core/Public/Math/Soaaos.isph
--- a/Engine/Source/Runtime/Core/Public/Math/Soaaos.isph
+++ b/Engine/Source/Runtime/Core/Public/Math/Soaaos.isph
@@ -21,19 +21,12 @@ static const varying int vSOA21 = {8, 24, 9, 25, 10, 26, 11, 27, 12, 28, 13, 29,
 #error "No implementation for this target"
 #endif
 
-#define SOA_TO_AOS2_ISPC(T)                                                                                            \
-    unmasked inline void soa_to_aos2_ispc(const varying T &src0, const varying T &src1, uniform T dst[])               \
-    {                                                                                                                  \
-        dst[programIndex] = shuffle(src0, src1, vSOA20);                                                               \
-        dst[programIndex + programCount] = shuffle(src0, src1, vSOA21);                                                \
-    }
-
-SOA_TO_AOS2_ISPC(int16)
-SOA_TO_AOS2_ISPC(int32)
-SOA_TO_AOS2_ISPC(float)
-SOA_TO_AOS2_ISPC(int64)
-SOA_TO_AOS2_ISPC(double)
-
+template<typename T>
+unmasked inline void soa_to_aos2_ispc(const varying T &src0, const varying T &src1, uniform T dst[])
+{
+    dst[programIndex] = shuffle(src0, src1, vSOA20);
+    dst[programIndex + programCount] = shuffle(src0, src1, vSOA21);
+}
 
 #if TARGET_WIDTH == 4
 static const varying int vSOA30 = {0, 1, 2, 4};
@@ -60,41 +53,39 @@ static const varying int vSOA35 = {26, 0, 5, 27, 1, 6, 28, 2, 7, 29, 3, 8, 30, 4
 #error "No implementation for this target"
 #endif
 
-#define SOA_TO_AOS3_ISPC(T)                                                                                            \
-    unmasked inline void soa_to_aos3_ispc(const varying T &src0, const varying T &src1,                                \
-                                        const varying T &src2, uniform T dst[])                                        \
-    {                                                                                                                  \
-        const varying T s0 = src0;                                                                                     \
-        const varying T s1 = src1;                                                                                     \
-        const varying T s2 = src2;                                                                                     \
-                                                                                                                       \
-        const varying T t0 = shuffle(s0, s1, vSOA30);                                                                  \
-        const varying T t1 = shuffle(s1, s2, vSOA31);                                                                  \
-        dst[programIndex] = shuffle(t0, t1, vSOA32);                                                                   \
-        dst[programCount + programIndex] = shuffle(t0, t1, vSOA33);                                                    \
-                                                                                                                       \
-        const varying T t2 = shuffle(s0, s1, vSOA34);                                                                  \
-        dst[2 * programCount + programIndex] = shuffle(t2, s2, vSOA35);                                                \
-    }
-
-SOA_TO_AOS3_ISPC(int16)
-SOA_TO_AOS3_ISPC(int32)
-SOA_TO_AOS3_ISPC(float)
-SOA_TO_AOS3_ISPC(int64)
-SOA_TO_AOS3_ISPC(double)
-
+template<typename T>
+unmasked inline void soa_to_aos3_ispc(const varying T &src0, const varying T &src1, const varying T &src2, uniform T dst[])
+{
+    const varying T s0 = src0;
+    const varying T s1 = src1;
+    const varying T s2 = src2;
+                              
+    const varying T t0 = shuffle(s0, s1, vSOA30);
+    const varying T t1 = shuffle(s1, s2, vSOA31);
+    dst[programIndex] = shuffle(t0, t1, vSOA32);
+    dst[programCount + programIndex] = shuffle(t0, t1, vSOA33);
 
+    const varying T t2 = shuffle(s0, s1, vSOA34);
+    dst[2 * programCount + programIndex] = shuffle(t2, s2, vSOA35);
+}
 
 #if TARGET_WIDTH == 4
-static const varying int vSOA40 = {0, 2, 4, 6};
-static const varying int vSOA41 = {1, 3, 5, 7};
-static const varying int vSOA42 = {0, 1, 4, 5};
-static const varying int vSOA43 = {2, 3, 6, 7};
+static const varying int vSOA40 = {0, 1, 4, 5};
+static const varying int vSOA41 = {2, 3, 6, 7};
+static const varying int vSOA42 = {0, 4, 1, 5};
+static const varying int vSOA43 = {2, 6, 3, 7};
 #elif TARGET_WIDTH == 8
-static const varying int vSOA40 = {0, 4, 8, 12, 1, 5, 9, 13};
-static const varying int vSOA41 = {2, 6, 10, 14, 3, 7, 11, 15};
-static const varying int vSOA42 = {0, 1, 2, 3, 8, 9, 10, 11};
-static const varying int vSOA43 = {4, 5, 6, 7, 12, 13, 14, 15};
+static const varying int vSOA40 = {0, 8, 1, 9, 4, 12, 5, 13};
+static const varying int vSOA41 = {2, 10, 3, 11, 6, 14, 7, 15};
+static const varying int vSOA42 = {0, 1, 8, 9, 4, 5, 12, 13};
+static const varying int vSOA43 = {2, 3, 10, 11, 6, 7, 14, 15};
+static const varying int vSOA44 = {0, 1, 2, 3, 8, 9, 10, 11};
+static const varying int vSOA45 = {4, 5, 6, 7, 12, 13, 14, 15};
+
+static const varying int vSOA401 = {0, 1, 8, 9, 4, 5, 12, 13};
+static const varying int vSOA411 = {2, 3, 10, 11, 6, 7, 14, 15};
+static const varying int vSOA421 = {0, 1, 2, 3, 8, 9, 10, 11};
+static const varying int vSOA431 = {4, 5, 6, 7, 12, 13, 14, 15};
 #elif TARGET_WIDTH == 16
 static const varying int vSOA40 = {0, 8, 16, 24, 1, 9, 17, 25, 2, 10, 18, 26, 3, 11, 19, 27};
 static const varying int vSOA41 = {4, 12, 20, 28, 5, 13, 21, 29, 6, 14, 22, 30, 7, 15, 23, 31};
@@ -104,27 +95,119 @@ static const varying int vSOA43 = {8, 9, 10, 11, 12, 13, 14, 15, 24, 25, 26, 27,
 #error "No implementation for this target"
 #endif
 
-#define SOA_TO_AOS4_ISPC(T)                                                                                            \
-    unmasked inline void soa_to_aos4_ispc(const varying T &src0, const varying T &src1,                                \
-                                        const varying T &src2, const varying T &src3, uniform T dst[])                 \
-    {                                                                                                                  \
-        const varying T t0 = shuffle(src0, src1, vSOA42);                                                              \
-        const varying T t1 = shuffle(src2, src3, vSOA42);                                                              \
-        const varying T t2 = shuffle(src0, src1, vSOA43);                                                              \
-        const varying T t3 = shuffle(src2, src3, vSOA43);                                                              \
-                                                                                                                       \
-        dst[programIndex] = shuffle(t0, t1, vSOA40);                                                                   \
-        dst[programCount + programIndex] = shuffle(t0, t1, vSOA41);                                                    \
-        dst[2 * programCount + programIndex] = shuffle(t2, t3, vSOA40);                                                \
-        dst[3 * programCount + programIndex] = shuffle(t2, t3, vSOA41);                                                \
-    }
+#if TARGET_WIDTH == 8
+template<typename T>
+unmasked inline void soa_to_aos4_ispc(const varying T& src0, const varying T& src1,
+    const varying T& src2, const varying T& src3, uniform T dst[])
+{
+    const varying T t0 = shuffle(src0, src1, vSOA40);
+    const varying T t1 = shuffle(src2, src3, vSOA40);
+    const varying T t2 = shuffle(t0, t1, vSOA42);
+    const varying T t3 = shuffle(t0, t1, vSOA43);
+
+    const varying T t4 = shuffle(src0, src1, vSOA41);
+    const varying T t5 = shuffle(src2, src3, vSOA41);
+    const varying T t6 = shuffle(t4, t5, vSOA42);
+    const varying T t7 = shuffle(t4, t5, vSOA43);
+
+    dst[programIndex] = shuffle(t2, t3, vSOA44);
+    dst[programCount + programIndex] = shuffle(t6, t7, vSOA44);
+    dst[2 * programCount + programIndex] = shuffle(t2, t3, vSOA45);
+    dst[3 * programCount + programIndex] = shuffle(t6, t7, vSOA45);
+}
+
+template<>
+unmasked inline void soa_to_aos4_ispc<double>(const varying double& src0, const varying double& src1,
+    const varying double& src2, const varying double& src3, uniform double dst[])
+{
+    uniform float *uniform dstflt = (uniform float *uniform)dst;
+    const uniform double<TARGET_WIDTH> x0 = *((const uniform double<TARGET_WIDTH>* uniform)&src0);
+    const uniform double<TARGET_WIDTH> x1 = *((const uniform double<TARGET_WIDTH>* uniform)&src1);
+    const uniform double<TARGET_WIDTH> x2 = *((const uniform double<TARGET_WIDTH>* uniform)&src2);
+    const uniform double<TARGET_WIDTH> x3 = *((const uniform double<TARGET_WIDTH>* uniform)&src3);
+
+    const varying float u0 = *((varying float *uniform)&x0);
+    const varying float u1 = *((varying float *uniform)&x0[TARGET_WIDTH/2]);
+    const varying float u2 = *((varying float *uniform)&x1);
+    const varying float u3 = *((varying float *uniform)&x1[TARGET_WIDTH/2]);
+    const varying float u4 = *((varying float *uniform)&x2);
+    const varying float u5 = *((varying float *uniform)&x2[TARGET_WIDTH/2]);
+    const varying float u6 = *((varying float *uniform)&x3);
+    const varying float u7 = *((varying float *uniform)&x3[TARGET_WIDTH/2]);
+
+    const varying float t0 = shuffle(u0, u2, vSOA401);
+    const varying float t1 = shuffle(u0, u2, vSOA411);
+    const varying float t2 = shuffle(u1, u3, vSOA401);
+    const varying float t3 = shuffle(u1, u3, vSOA411);
+    const varying float t4 = shuffle(u4, u6, vSOA401);
+    const varying float t5 = shuffle(u4, u6, vSOA411);
+    const varying float t6 = shuffle(u5, u7, vSOA401);
+    const varying float t7 = shuffle(u5, u7, vSOA411);
+
+    dstflt[programIndex] = shuffle(t0, t4, vSOA421);
+    dstflt[programCount + programIndex] = shuffle(t1, t5, vSOA421);
+    dstflt[2 * programCount + programIndex] = shuffle(t0, t4, vSOA431);
+    dstflt[3 * programCount + programIndex] = shuffle(t1, t5, vSOA431);
+    dstflt[4 * programCount + programIndex] = shuffle(t2, t6, vSOA421);
+    dstflt[5 * programCount + programIndex] = shuffle(t3, t7, vSOA421);
+    dstflt[6 * programCount + programIndex] = shuffle(t2, t6, vSOA431);
+    dstflt[7 * programCount + programIndex] = shuffle(t3, t7, vSOA431);
+}
+
+template<>
+unmasked inline void soa_to_aos4_ispc<int64>(const varying int64& src0, const varying int64& src1,
+    const varying int64& src2, const varying int64& src3, uniform int64 dst[])
+{
+    uniform float* uniform dstflt = (uniform float* uniform)dst;
+    const uniform int64<TARGET_WIDTH>x0 = *((const uniform int64<TARGET_WIDTH>* uniform)&src0);
+    const uniform int64<TARGET_WIDTH>x1 = *((const uniform int64<TARGET_WIDTH>* uniform)&src1);
+    const uniform int64<TARGET_WIDTH>x2 = *((const uniform int64<TARGET_WIDTH>* uniform)&src2);
+    const uniform int64<TARGET_WIDTH>x3 = *((const uniform int64<TARGET_WIDTH>* uniform)&src3);
 
-SOA_TO_AOS4_ISPC(int16)
-SOA_TO_AOS4_ISPC(int32)
-SOA_TO_AOS4_ISPC(float)
-SOA_TO_AOS4_ISPC(int64)
-SOA_TO_AOS4_ISPC(double)
+    const varying float u0 = *((varying float* uniform)&x0);
+    const varying float u1 = *((varying float* uniform)&x0[TARGET_WIDTH / 2]);
+    const varying float u2 = *((varying float* uniform)&x1);
+    const varying float u3 = *((varying float* uniform)&x1[TARGET_WIDTH / 2]);
+    const varying float u4 = *((varying float* uniform)&x2);
+    const varying float u5 = *((varying float* uniform)&x2[TARGET_WIDTH / 2]);
+    const varying float u6 = *((varying float* uniform)&x3);
+    const varying float u7 = *((varying float* uniform)&x3[TARGET_WIDTH / 2]);
 
+    const varying float t0 = shuffle(u0, u2, vSOA401);
+    const varying float t1 = shuffle(u0, u2, vSOA411);
+    const varying float t2 = shuffle(u1, u3, vSOA401);
+    const varying float t3 = shuffle(u1, u3, vSOA411);
+    const varying float t4 = shuffle(u4, u6, vSOA401);
+    const varying float t5 = shuffle(u4, u6, vSOA411);
+    const varying float t6 = shuffle(u5, u7, vSOA401);
+    const varying float t7 = shuffle(u5, u7, vSOA411);
+
+    dstflt[programIndex] = shuffle(t0, t4, vSOA421);
+    dstflt[programCount + programIndex] = shuffle(t1, t5, vSOA421);
+    dstflt[2 * programCount + programIndex] = shuffle(t0, t4, vSOA431);
+    dstflt[3 * programCount + programIndex] = shuffle(t1, t5, vSOA431);
+    dstflt[4 * programCount + programIndex] = shuffle(t2, t6, vSOA421);
+    dstflt[5 * programCount + programIndex] = shuffle(t3, t7, vSOA421);
+    dstflt[6 * programCount + programIndex] = shuffle(t2, t6, vSOA431);
+    dstflt[7 * programCount + programIndex] = shuffle(t3, t7, vSOA431);
+}
+
+#else
+template<typename T>
+unmasked inline void soa_to_aos4_ispc(const varying T &src0, const varying T &src1,
+    const varying T &src2, const varying T &src3, uniform T dst[])
+{
+    const varying T t0 = shuffle(src0, src1, vSOA42);
+    const varying T t1 = shuffle(src2, src3, vSOA42);
+    const varying T t2 = shuffle(src0, src1, vSOA43);
+    const varying T t3 = shuffle(src2, src3, vSOA43);
+
+    dst[programIndex] = shuffle(t0, t1, vSOA40);
+    dst[programCount + programIndex] = shuffle(t0, t1, vSOA41);
+    dst[2 * programCount + programIndex] = shuffle(t2, t3, vSOA40);
+    dst[3 * programCount + programIndex] = shuffle(t2, t3, vSOA41);
+}
+#endif
 
 #if TARGET_WIDTH == 4
 static const varying int vSOA60 = {0, 1, 4, 5};
@@ -155,52 +238,337 @@ static const varying int vSOA67 = {0, 1, 21, 29, 4, 5, 6, 7, 22, 30, 10, 11, 12,
 #endif
 
 #if TARGET_WIDTH == 4
-#define SOA_TO_AOS6_ISPC(T)                                                                                            \
-    unmasked inline void soa_to_aos6_ispc(const varying T &src0, const varying T &src1,                                \
-                                        const varying T &src2, const varying T &src3,                                  \
-                                        const varying T &src4, const varying T &src5, uniform T dst[])                 \
-    {                                                                                                                  \
-        const varying T t0 = shuffle(src0, src1, vSOA60);                                                              \
-        const varying T t1 = shuffle(src0, src1, vSOA61);                                                              \
-        const varying T t2 = shuffle(src2, src3, vSOA60);                                                              \
-        const varying T t3 = shuffle(src2, src3, vSOA61);                                                              \
-        const varying T t4 = shuffle(src4, src5, vSOA60);                                                              \
-        const varying T t5 = shuffle(src4, src5, vSOA61);                                                              \
-                                                                                                                       \
-        dst[programIndex] = shuffle(t0, t2, vSOA62);                                                                   \
-        dst[programIndex + programCount] = shuffle(t4, t0, vSOA62);                                                    \
-        dst[2 * programCount + programIndex] = shuffle(t2, t4, vSOA63);                                                \
-        dst[3 * programCount + programIndex] = shuffle(t1, t3, vSOA62);                                                \
-        dst[4 * programCount + programIndex] = shuffle(t5, t1, vSOA64);                                                \
-        dst[5 * programCount + programIndex] = shuffle(t3, t5, vSOA63);                                                \
-    }
+template<typename T>
+unmasked inline void soa_to_aos6_ispc(const varying T &src0, const varying T &src1, const varying T &src2, const varying T &src3, const varying T &src4, const varying T &src5, uniform T dst[])
+{
+    const varying T t0 = shuffle(src0, src1, vSOA60);
+    const varying T t1 = shuffle(src0, src1, vSOA61);
+    const varying T t2 = shuffle(src2, src3, vSOA60);
+    const varying T t3 = shuffle(src2, src3, vSOA61);
+    const varying T t4 = shuffle(src4, src5, vSOA60);
+    const varying T t5 = shuffle(src4, src5, vSOA61);
+
+    dst[programIndex] = shuffle(t0, t2, vSOA62);
+    dst[programIndex + programCount] = shuffle(t4, t0, vSOA62);
+    dst[2 * programCount + programIndex] = shuffle(t2, t4, vSOA63);
+    dst[3 * programCount + programIndex] = shuffle(t1, t3, vSOA62);
+    dst[4 * programCount + programIndex] = shuffle(t5, t1, vSOA64);
+    dst[5 * programCount + programIndex] = shuffle(t3, t5, vSOA63);
+}
 #else
-#define SOA_TO_AOS6_ISPC(T)                                                                                            \
-    unmasked inline void soa_to_aos6_ispc(const varying T &src0, const varying T &src1,                                \
-                                        const varying T &src2, const varying T &src3,                                  \
-                                        const varying T &src4, const varying T &src5, uniform T dst[])                 \
-    {                                                                                                                  \
-        const varying T t0 = shuffle(src0, src1, vSOA60);                                                              \
-        const varying T t1 = shuffle(src0, src1, vSOA61);                                                              \
-        const varying T t2 = shuffle(src2, src3, vSOA60);                                                              \
-        const varying T t3 = shuffle(src2, src3, vSOA61);                                                              \
-        const varying T t4 = shuffle(src4, src5, vSOA60);                                                              \
-        const varying T t5 = shuffle(src4, src5, vSOA61);                                                              \
-                                                                                                                       \
-        dst[programIndex] = shuffle(shuffle(t0, t2, vSOA62), t4, vSOA63);                                              \
-        dst[programIndex + programCount] = shuffle(shuffle(t0, t2, vSOA64), t4, vSOA65);                               \
-        dst[2 * programCount + programIndex] = shuffle(shuffle(t0, t2, vSOA66), t4, vSOA67);                           \
-        dst[3 * programCount + programIndex] = shuffle(shuffle(t1, t3, vSOA62), t5, vSOA63);                           \
-        dst[4 * programCount + programIndex] = shuffle(shuffle(t1, t3, vSOA64), t5, vSOA65);                           \
-        dst[5 * programCount + programIndex] = shuffle(shuffle(t1, t3, vSOA66), t5, vSOA67);                           \
+template<typename T>
+unmasked inline void soa_to_aos6_ispc(const varying T &src0, const varying T &src1, const varying T &src2, const varying T &src3, const varying T &src4, const varying T &src5, uniform T dst[])
+{
+    const varying T t0 = shuffle(src0, src1, vSOA60);
+    const varying T t1 = shuffle(src0, src1, vSOA61);
+    const varying T t2 = shuffle(src2, src3, vSOA60);
+    const varying T t3 = shuffle(src2, src3, vSOA61);
+    const varying T t4 = shuffle(src4, src5, vSOA60);
+    const varying T t5 = shuffle(src4, src5, vSOA61);
+
+    dst[programIndex] = shuffle(shuffle(t0, t2, vSOA62), t4, vSOA63);
+    dst[programIndex + programCount] = shuffle(shuffle(t0, t2, vSOA64), t4, vSOA65);
+    dst[2 * programCount + programIndex] = shuffle(shuffle(t0, t2, vSOA66), t4, vSOA67);
+    dst[3 * programCount + programIndex] = shuffle(shuffle(t1, t3, vSOA62), t5, vSOA63);
+    dst[4 * programCount + programIndex] = shuffle(shuffle(t1, t3, vSOA64), t5, vSOA65);
+    dst[5 * programCount + programIndex] = shuffle(shuffle(t1, t3, vSOA66), t5, vSOA67);
+}
+#endif
+
+template<typename T>
+unmasked inline void UniformStore(uniform T* uniform SrcPtr, uniform T* varying DstPtr)
+{
+    uniform T* uniform Dst[programCount];
+
+    foreach(i = 0 ... programCount)
+    {
+        Dst[i] = (uniform T* varying)&DstPtr[0];
     }
+
+    *Dst[0] = SrcPtr[0];
+    *Dst[1] = SrcPtr[1];
+    *Dst[2] = SrcPtr[2];
+    *Dst[3] = SrcPtr[3];
+#if TARGET_WIDTH == 8 || TARGET_WIDTH == 16
+    *Dst[4] = SrcPtr[4];
+    *Dst[5] = SrcPtr[5];
+    *Dst[6] = SrcPtr[6];
+    *Dst[7] = SrcPtr[7];
+#endif
+#if TARGET_WIDTH == 16
+    *Dst[8] = SrcPtr[8];
+    *Dst[9] = SrcPtr[9];
+    *Dst[10] = SrcPtr[10];
+    *Dst[11] = SrcPtr[11];
+    *Dst[12] = SrcPtr[12];
+    *Dst[13] = SrcPtr[13];
+    *Dst[14] = SrcPtr[14];
+    *Dst[15] = SrcPtr[15];
 #endif
+}
 
-SOA_TO_AOS6_ISPC(int16)
-SOA_TO_AOS6_ISPC(int32)
-SOA_TO_AOS6_ISPC(float)
-SOA_TO_AOS6_ISPC(int64)
-SOA_TO_AOS6_ISPC(double)
+template<typename T>
+unmasked inline void SoaToAos3Explicit(uniform T* uniform SrcPtr, uniform T* varying DstPtr)
+{
+    uniform T* uniform Dst[programCount];
 
+    foreach(i = 0 ... programCount)
+    {
+        Dst[i] = (uniform T* varying)&DstPtr[0];
+    }
+
+#if TARGET_WIDTH == 4
+    // X, X, X, X
+    Dst[0][0] = SrcPtr[0];
+    Dst[1][0] = SrcPtr[1];
+    Dst[2][0] = SrcPtr[2];
+    Dst[3][0] = SrcPtr[3];
+
+    // Y, Y, Y, Y
+    Dst[0][1] = SrcPtr[4];
+    Dst[1][1] = SrcPtr[5];
+    Dst[2][1] = SrcPtr[6];
+    Dst[3][1] = SrcPtr[7];
+
+    // Z, Z, Z, Z
+    Dst[0][2] = SrcPtr[8];
+    Dst[1][2] = SrcPtr[9];
+    Dst[2][2] = SrcPtr[10];
+    Dst[3][2] = SrcPtr[11];
+#elif TARGET_WIDTH == 8
+    // X, X, X, X, X, X, X, X
+    Dst[0][0] = SrcPtr[0];
+    Dst[1][0] = SrcPtr[1];
+    Dst[2][0] = SrcPtr[2];
+    Dst[3][0] = SrcPtr[3];
+    Dst[4][0] = SrcPtr[4];
+    Dst[5][0] = SrcPtr[5];
+    Dst[6][0] = SrcPtr[6];
+    Dst[7][0] = SrcPtr[7];
+
+    // Y, Y, Y, Y, Y, Y, Y, Y
+    Dst[0][1] = SrcPtr[8];
+    Dst[1][1] = SrcPtr[9];
+    Dst[2][1] = SrcPtr[10];
+    Dst[3][1] = SrcPtr[11];
+    Dst[4][1] = SrcPtr[12];
+    Dst[5][1] = SrcPtr[13];
+    Dst[6][1] = SrcPtr[14];
+    Dst[7][1] = SrcPtr[15];
+
+    // Z, Z, Z, Z, Z, Z, Z, Z
+    Dst[0][2] = SrcPtr[16];
+    Dst[1][2] = SrcPtr[17];
+    Dst[2][2] = SrcPtr[18];
+    Dst[3][2] = SrcPtr[19];
+    Dst[4][2] = SrcPtr[20];
+    Dst[5][2] = SrcPtr[21];
+    Dst[6][2] = SrcPtr[22];
+    Dst[7][2] = SrcPtr[23];
+#elif TARGET_WIDTH == 16
+    // X, X, X, X, X, X, X, X, X, X, X, X, X, X, X, X
+    Dst[0][0] = SrcPtr[0];
+    Dst[1][0] = SrcPtr[1];
+    Dst[2][0] = SrcPtr[2];
+    Dst[3][0] = SrcPtr[3];
+    Dst[4][0] = SrcPtr[4];
+    Dst[5][0] = SrcPtr[5];
+    Dst[6][0] = SrcPtr[6];
+    Dst[7][0] = SrcPtr[7];
+    Dst[8][0] = SrcPtr[8];
+    Dst[9][0] = SrcPtr[9];
+    Dst[10][0] = SrcPtr[10];
+    Dst[11][0] = SrcPtr[11];
+    Dst[12][0] = SrcPtr[12];
+    Dst[13][0] = SrcPtr[13];
+    Dst[14][0] = SrcPtr[14];
+    Dst[15][0] = SrcPtr[15];
+
+    // Y, Y, Y, Y, Y, Y, Y, Y, Y, Y, Y, Y, Y, Y, Y, Y
+    Dst[0][1] = SrcPtr[16];
+    Dst[1][1] = SrcPtr[17];
+    Dst[2][1] = SrcPtr[18];
+    Dst[3][1] = SrcPtr[19];
+    Dst[4][1] = SrcPtr[20];
+    Dst[5][1] = SrcPtr[21];
+    Dst[6][1] = SrcPtr[22];
+    Dst[7][1] = SrcPtr[23];
+    Dst[8][1] = SrcPtr[24];
+    Dst[9][1] = SrcPtr[25];
+    Dst[10][1] = SrcPtr[26];
+    Dst[11][1] = SrcPtr[27];
+    Dst[12][1] = SrcPtr[28];
+    Dst[13][1] = SrcPtr[29];
+    Dst[14][1] = SrcPtr[30];
+    Dst[15][1] = SrcPtr[31];
+
+    // Z, Z, Z, Z, Z, Z, Z, Z, Z, Z, Z, Z, Z, Z, Z, Z
+    Dst[0][2] = SrcPtr[32];
+    Dst[1][2] = SrcPtr[33];
+    Dst[2][2] = SrcPtr[34];
+    Dst[3][2] = SrcPtr[35];
+    Dst[4][2] = SrcPtr[36];
+    Dst[5][2] = SrcPtr[37];
+    Dst[6][2] = SrcPtr[38];
+    Dst[7][2] = SrcPtr[39];
+    Dst[8][2] = SrcPtr[40];
+    Dst[9][2] = SrcPtr[41];
+    Dst[10][2] = SrcPtr[42];
+    Dst[11][2] = SrcPtr[43];
+    Dst[12][2] = SrcPtr[44];
+    Dst[13][2] = SrcPtr[45];
+    Dst[14][2] = SrcPtr[46];
+    Dst[15][2] = SrcPtr[47];
+#endif
+}
+
+template<typename T>
+unmasked inline void SoaToAos4Explicit(uniform T* uniform SrcPtr, uniform T* varying DstPtr)
+{
+    uniform T* uniform Dst[programCount];
+
+    foreach(i = 0 ... programCount)
+    {
+        Dst[i] = (uniform T* varying)&DstPtr[0];
+    }
+
+#if TARGET_WIDTH == 4
+    // X, X, X, X
+    Dst[0][0] = SrcPtr[0];
+    Dst[1][0] = SrcPtr[1];
+    Dst[2][0] = SrcPtr[2];
+    Dst[3][0] = SrcPtr[3];
+
+    // Y, Y, Y, Y
+    Dst[0][1] = SrcPtr[4];
+    Dst[1][1] = SrcPtr[5];
+    Dst[2][1] = SrcPtr[6];
+    Dst[3][1] = SrcPtr[7];
+
+    // Z, Z, Z, Z
+    Dst[0][2] = SrcPtr[8];
+    Dst[1][2] = SrcPtr[9];
+    Dst[2][2] = SrcPtr[10];
+    Dst[3][2] = SrcPtr[11];
+
+    // W, W, W, W
+    Dst[0][3] = SrcPtr[12];
+    Dst[1][3] = SrcPtr[13];
+    Dst[2][3] = SrcPtr[14];
+    Dst[3][3] = SrcPtr[15];
+#elif TARGET_WIDTH == 8
+    // X, X, X, X, X, X, X, X
+    Dst[0][0] = SrcPtr[0];
+    Dst[1][0] = SrcPtr[1];
+    Dst[2][0] = SrcPtr[2];
+    Dst[3][0] = SrcPtr[3];
+    Dst[4][0] = SrcPtr[4];
+    Dst[5][0] = SrcPtr[5];
+    Dst[6][0] = SrcPtr[6];
+    Dst[7][0] = SrcPtr[7];
+
+    // Y, Y, Y, Y, Y, Y, Y, Y
+    Dst[0][1] = SrcPtr[8];
+    Dst[1][1] = SrcPtr[9];
+    Dst[2][1] = SrcPtr[10];
+    Dst[3][1] = SrcPtr[11];
+    Dst[4][1] = SrcPtr[12];
+    Dst[5][1] = SrcPtr[13];
+    Dst[6][1] = SrcPtr[14];
+    Dst[7][1] = SrcPtr[15];
+
+    // Z, Z, Z, Z, Z, Z, Z, Z
+    Dst[0][2] = SrcPtr[16];
+    Dst[1][2] = SrcPtr[17];
+    Dst[2][2] = SrcPtr[18];
+    Dst[3][2] = SrcPtr[19];
+    Dst[4][2] = SrcPtr[20];
+    Dst[5][2] = SrcPtr[21];
+    Dst[6][2] = SrcPtr[22];
+    Dst[7][2] = SrcPtr[23];
+
+    // W, W, W, W, W, W, W, W
+    Dst[0][3] = SrcPtr[24];
+    Dst[1][3] = SrcPtr[25];
+    Dst[2][3] = SrcPtr[26];
+    Dst[3][3] = SrcPtr[27];
+    Dst[4][3] = SrcPtr[28];
+    Dst[5][3] = SrcPtr[29];
+    Dst[6][3] = SrcPtr[30];
+    Dst[7][3] = SrcPtr[31];
+#elif TARGET_WIDTH == 16
+    // X, X, X, X, X, X, X, X, X, X, X, X, X, X, X, X
+    Dst[0][0] = SrcPtr[0];
+    Dst[1][0] = SrcPtr[1];
+    Dst[2][0] = SrcPtr[2];
+    Dst[3][0] = SrcPtr[3];
+    Dst[4][0] = SrcPtr[4];
+    Dst[5][0] = SrcPtr[5];
+    Dst[6][0] = SrcPtr[6];
+    Dst[7][0] = SrcPtr[7];
+    Dst[8][0] = SrcPtr[8];
+    Dst[9][0] = SrcPtr[9];
+    Dst[10][0] = SrcPtr[10];
+    Dst[11][0] = SrcPtr[11];
+    Dst[12][0] = SrcPtr[12];
+    Dst[13][0] = SrcPtr[13];
+    Dst[14][0] = SrcPtr[14];
+    Dst[15][0] = SrcPtr[15];
+
+    // Y, Y, Y, Y, Y, Y, Y, Y, Y, Y, Y, Y, Y, Y, Y, Y
+    Dst[0][1] = SrcPtr[16];
+    Dst[1][1] = SrcPtr[17];
+    Dst[2][1] = SrcPtr[18];
+    Dst[3][1] = SrcPtr[19];
+    Dst[4][1] = SrcPtr[20];
+    Dst[5][1] = SrcPtr[21];
+    Dst[6][1] = SrcPtr[22];
+    Dst[7][1] = SrcPtr[23];
+    Dst[8][1] = SrcPtr[24];
+    Dst[9][1] = SrcPtr[25];
+    Dst[10][1] = SrcPtr[26];
+    Dst[11][1] = SrcPtr[27];
+    Dst[12][1] = SrcPtr[28];
+    Dst[13][1] = SrcPtr[29];
+    Dst[14][1] = SrcPtr[30];
+    Dst[15][1] = SrcPtr[31];
+
+    // Z, Z, Z, Z, Z, Z, Z, Z, Z, Z, Z, Z, Z, Z, Z, Z
+    Dst[0][2] = SrcPtr[32];
+    Dst[1][2] = SrcPtr[33];
+    Dst[2][2] = SrcPtr[34];
+    Dst[3][2] = SrcPtr[35];
+    Dst[4][2] = SrcPtr[36];
+    Dst[5][2] = SrcPtr[37];
+    Dst[6][2] = SrcPtr[38];
+    Dst[7][2] = SrcPtr[39];
+    Dst[8][2] = SrcPtr[40];
+    Dst[9][2] = SrcPtr[41];
+    Dst[10][2] = SrcPtr[42];
+    Dst[11][2] = SrcPtr[43];
+    Dst[12][2] = SrcPtr[44];
+    Dst[13][2] = SrcPtr[45];
+    Dst[14][2] = SrcPtr[46];
+    Dst[15][2] = SrcPtr[47];
+
+    // W, W, W, W, W, W, W, W, W, W, W, W, W, W, W, W
+    Dst[0][3] = SrcPtr[48];
+    Dst[1][3] = SrcPtr[49];
+    Dst[2][3] = SrcPtr[50];
+    Dst[3][3] = SrcPtr[51];
+    Dst[4][3] = SrcPtr[52];
+    Dst[5][3] = SrcPtr[53];
+    Dst[6][3] = SrcPtr[54];
+    Dst[7][3] = SrcPtr[55];
+    Dst[8][3] = SrcPtr[56];
+    Dst[9][3] = SrcPtr[57];
+    Dst[10][3] = SrcPtr[58];
+    Dst[11][3] = SrcPtr[59];
+    Dst[12][3] = SrcPtr[60];
+    Dst[13][3] = SrcPtr[61];
+    Dst[14][3] = SrcPtr[62];
+    Dst[15][3] = SrcPtr[63];
+#endif
+}
 
 #endif
diff --git a/Engine/Source/Runtime/Core/Public/Math/Transform.isph b/Engine/Source/Runtime/Core/Public/Math/Transform.isph
--- a/Engine/Source/Runtime/Core/Public/Math/Transform.isph
+++ b/Engine/Source/Runtime/Core/Public/Math/Transform.isph
@@ -347,7 +347,7 @@ inline FVector TransformPosition(const uniform FTransform &T, const FVector& V)
 	const FVector4 ScaledVec = VectorMultiply(T.Scale3D, InputVectorW0);
 	const FVector4 RotatedVec = VectorQuaternionRotateVector(T.Rotation, ScaledVec);
 
-	const FVector4 TranslatedVec = VectorAdd(RotatedVec, T.Translation);
+	const FVector4 TranslatedVec = VectorAdd(RotatedVec, (const varying FVector4)T.Translation);
 
 	return SetVector(TranslatedVec);
 }
@@ -363,7 +363,7 @@ inline FVector3f TransformPosition(const uniform FTransform3f &T, const FVector3
 	const FVector4f ScaledVec = VectorMultiply(T.Scale3D, InputVectorW0);
 	const FVector4f RotatedVec = VectorQuaternionRotateVector(T.Rotation, ScaledVec);
 
-	const FVector4f TranslatedVec = VectorAdd(RotatedVec, T.Translation);
+	const FVector4f TranslatedVec = VectorAdd(RotatedVec, (const varying FVector4f)SetVector4(T.Translation.V[0], T.Translation.V[1], T.Translation.V[2], 0.f));
 
 	return SetVector(TranslatedVec);
 }
diff --git a/Engine/Source/Runtime/Core/Public/Math/Vector.isph b/Engine/Source/Runtime/Core/Public/Math/Vector.isph
--- a/Engine/Source/Runtime/Core/Public/Math/Vector.isph
+++ b/Engine/Source/Runtime/Core/Public/Math/Vector.isph
@@ -591,20 +591,68 @@ inline varying FVector3d VectorGather(const uniform FVector3d *varying SrcPtr)
 
 inline varying FVector3f VectorGather(const uniform FVector3f *varying SrcPtr)
 {
-	#pragma ignore warning(perf)
-	return *SrcPtr;
+	varying FVector3f Result;
+
+	if (((1 << TARGET_WIDTH) - 1 ^ lanemask()) == 0)
+	{
+#if TARGET_WIDTH == 4
+		AosToSoa3Explicit((uniform float* varying)SrcPtr, (uniform float* uniform)&Result);
+#else
+		UniformLoad(SrcPtr, (uniform FVector3f* uniform)&Result);
+		aos_to_soa3_ispc((uniform float* uniform)&Result, &Result.V[0], &Result.V[1], &Result.V[2]);
+#endif
+	}
+	else
+	{
+		#pragma ignore warning(perf)
+		Result = *SrcPtr;
+	}
+
+	return Result;
 }
 
 inline void VectorScatter(uniform FVector3d *varying DstPtr, const varying FVector3d &V)
 {
-	#pragma ignore warning(perf)
-	*DstPtr = V;
+	if (((1 << TARGET_WIDTH) - 1 ^ lanemask()) == 0)
+	{
+		unmasked
+		{
+#if TARGET_WIDTH == 4
+			SoaToAos3Explicit((uniform double* uniform)&V, (uniform double* varying)DstPtr);
+#else
+			uniform FVector3d Result[programCount];
+			soa_to_aos3_ispc(V.V[0], V.V[1], V.V[2], (uniform double* uniform)&Result);
+			UniformStore(Result, DstPtr);
+#endif
+		}
+	}
+	else
+	{
+		#pragma ignore warning(perf)
+		*DstPtr = V;
+	}
 }
 
 inline void VectorScatter(uniform FVector3f *varying DstPtr, const varying FVector3f &V)
 {
-	#pragma ignore warning(perf)
-	*DstPtr = V;
+	if (((1 << TARGET_WIDTH) - 1 ^ lanemask()) == 0)
+	{
+		unmasked
+		{
+#if TARGET_WIDTH == 4
+			SoaToAos3Explicit((uniform float* uniform)&V, (uniform float* varying)DstPtr);
+#else
+			uniform FVector3f Result[programCount];
+			soa_to_aos3_ispc(V.V[0], V.V[1], V.V[2], (uniform float* uniform)&Result);
+			UniformStore(Result, DstPtr);
+#endif
+		}
+	}
+	else 
+	{
+		#pragma ignore warning(perf)
+		*DstPtr = V;
+	}
 }
 
 inline varying FVector4d VectorGather(const uniform FVector4d *varying SrcPtr)
@@ -613,39 +661,12 @@ inline varying FVector4d VectorGather(const uniform FVector4d *varying SrcPtr)
 
 	if(((1<<TARGET_WIDTH)-1 ^ lanemask()) == 0)
 	{
-		unmasked
-		{
-			uniform FVector4d *uniform Src[programCount];
-			uniform FVector4d *uniform DstPtr = (uniform FVector4d *uniform)&Result;
-
-			foreach(i = 0 ... programCount)
-			{
-				Src[i] = (uniform FVector4d *varying)&SrcPtr[0];
-			}
-
-			DstPtr[0] = *Src[0];
-			DstPtr[1] = *Src[1];
-			DstPtr[2] = *Src[2];
-			DstPtr[3] = *Src[3];
-#if TARGET_WIDTH == 8 || TARGET_WIDTH == 16
-			DstPtr[4] = *Src[4];
-			DstPtr[5] = *Src[5];
-			DstPtr[6] = *Src[6];
-			DstPtr[7] = *Src[7];
-#endif
-#if TARGET_WIDTH == 16
-			DstPtr[8] = *Src[8];
-			DstPtr[9] = *Src[9];
-			DstPtr[10] = *Src[10];
-			DstPtr[11] = *Src[11];
-			DstPtr[12] = *Src[12];
-			DstPtr[13] = *Src[13];
-			DstPtr[14] = *Src[14];
-			DstPtr[15] = *Src[15];
+#if TARGET_WIDTH == 4
+		AosToSoa4Explicit((const uniform double* varying)SrcPtr, (uniform double* uniform)&Result);
+#else
+		UniformLoad(SrcPtr, (uniform FVector4d* uniform)&Result);
+		aos_to_soa4_ispc((uniform double* uniform)&Result, &Result.V[0], &Result.V[1], &Result.V[2], &Result.V[3]);
 #endif
-
-			aos_to_soa4_ispc((uniform double * uniform)&Result, &Result.V[0], &Result.V[1], &Result.V[2], &Result.V[3]);
-		}
 	}
 	else
 	{
@@ -662,39 +683,12 @@ inline varying FVector4f VectorGather(const uniform FVector4f *varying SrcPtr)
 
 	if(((1<<TARGET_WIDTH)-1 ^ lanemask()) == 0)
 	{
-		unmasked
-		{
-			uniform FVector4f *uniform Src[programCount];
-			uniform FVector4f *uniform DstPtr = (uniform FVector4f *uniform)&Result;
-
-			foreach(i = 0 ... programCount)
-			{
-				Src[i] = (uniform FVector4f *varying)&SrcPtr[0];
-			}
-
-			DstPtr[0] = *Src[0];
-			DstPtr[1] = *Src[1];
-			DstPtr[2] = *Src[2];
-			DstPtr[3] = *Src[3];
-#if TARGET_WIDTH == 8 || TARGET_WIDTH == 16
-			DstPtr[4] = *Src[4];
-			DstPtr[5] = *Src[5];
-			DstPtr[6] = *Src[6];
-			DstPtr[7] = *Src[7];
-#endif
-#if TARGET_WIDTH == 16
-			DstPtr[8] = *Src[8];
-			DstPtr[9] = *Src[9];
-			DstPtr[10] = *Src[10];
-			DstPtr[11] = *Src[11];
-			DstPtr[12] = *Src[12];
-			DstPtr[13] = *Src[13];
-			DstPtr[14] = *Src[14];
-			DstPtr[15] = *Src[15];
+#if TARGET_WIDTH == 4
+		AosToSoa4Explicit((const uniform float* varying)SrcPtr, (uniform float* uniform) & Result);
+#else
+		UniformLoad(SrcPtr, (uniform FVector4f* uniform)&Result);
+		aos_to_soa4_ispc((uniform float* uniform)&Result, &Result.V[0], &Result.V[1], &Result.V[2], &Result.V[3]);
 #endif
-
-			aos_to_soa4_ispc((uniform float * uniform)&Result, &Result.V[0], &Result.V[1], &Result.V[2], &Result.V[3]);
-		}
 	}
 	else
 	{
@@ -711,35 +705,12 @@ inline void VectorScatter(uniform FVector4d *varying DstPtr, const varying FVect
 	{
 		unmasked
 		{
+#if TARGET_WIDTH == 4
+			SoaToAos4Explicit((uniform double* uniform)&V, (uniform double* varying)DstPtr);
+#else
 			uniform FVector4d Result[programCount];
-			soa_to_aos4_ispc(V.V[0], V.V[1], V.V[2], V.V[3], (uniform double *uniform)&Result);
-
-			uniform FVector4d *uniform Dst[programCount];
-
-			foreach(i = 0 ... programCount)
-			{
-				Dst[i] = (uniform FVector4d *varying)&DstPtr[0];
-			}
-
-			*Dst[0] = Result[0];
-			*Dst[1] = Result[1];
-			*Dst[2] = Result[2];
-			*Dst[3] = Result[3];
-#if TARGET_WIDTH == 8 || TARGET_WIDTH == 16
-			*Dst[4] = Result[4];
-			*Dst[5] = Result[5];
-			*Dst[6] = Result[6];
-			*Dst[7] = Result[7];
-#endif
-#if TARGET_WIDTH == 16
-			*Dst[8] = Result[8];
-			*Dst[9] = Result[9];
-			*Dst[10] = Result[10];
-			*Dst[11] = Result[11];
-			*Dst[12] = Result[12];
-			*Dst[13] = Result[13];
-			*Dst[14] = Result[14];
-			*Dst[15] = Result[15];
+			soa_to_aos4_ispc(V.V[0], V.V[1], V.V[2], V.V[3], (uniform double* uniform)&Result);
+			UniformStore(Result, DstPtr);
 #endif
 		}
 	}
@@ -756,35 +727,12 @@ inline void VectorScatter(uniform FVector4f *varying DstPtr, const varying FVect
 	{
 		unmasked
 		{
+#if TARGET_WIDTH == 4
+			SoaToAos4Explicit((uniform float* uniform)&V, (uniform float* varying)DstPtr);
+#else
 			uniform FVector4f Result[programCount];
-			soa_to_aos4_ispc(V.V[0], V.V[1], V.V[2], V.V[3], (uniform float *uniform)&Result);
-
-			uniform FVector4f *uniform Dst[programCount];
-
-			foreach(i = 0 ... programCount)
-			{
-				Dst[i] = (uniform FVector4f *varying)&DstPtr[0];
-			}
-
-			*Dst[0] = Result[0];
-			*Dst[1] = Result[1];
-			*Dst[2] = Result[2];
-			*Dst[3] = Result[3];
-#if TARGET_WIDTH == 8 || TARGET_WIDTH == 16
-			*Dst[4] = Result[4];
-			*Dst[5] = Result[5];
-			*Dst[6] = Result[6];
-			*Dst[7] = Result[7];
-#endif
-#if TARGET_WIDTH == 16
-			*Dst[8] = Result[8];
-			*Dst[9] = Result[9];
-			*Dst[10] = Result[10];
-			*Dst[11] = Result[11];
-			*Dst[12] = Result[12];
-			*Dst[13] = Result[13];
-			*Dst[14] = Result[14];
-			*Dst[15] = Result[15];
+			soa_to_aos4_ispc(V.V[0], V.V[1], V.V[2], V.V[3], (uniform float* uniform)&Result);
+			UniformStore(Result, DstPtr);
 #endif
 		}
 	}
@@ -1049,9 +997,21 @@ inline FVector4 operator*(const FVector4 &A, const uniform FReal F)
 	return Result;
 }
 
-inline FVector4 operator*(const uniform FReal F, const FVector4 &A)
+inline FVector4d operator*(const uniform double F, const FVector4d &A)
 {
-	FVector4 Result;
+	FVector4d Result;
+
+	for(uniform int i = 0; i < 4; i++)
+	{
+		Result.V[i] = F * A.V[i];
+	}
+
+	return Result;
+}
+
+inline FVector4f operator*(const uniform float F, const FVector4f &A)
+{
+	FVector4f Result;
 
 	for(uniform int i = 0; i < 4; i++)
 	{
@@ -1273,9 +1233,21 @@ inline FVector4 operator-(const FVector4 &A, const uniform FReal F)
 	return Result;
 }
 
-inline FVector4 operator-(const uniform FReal F, const FVector4 &A)
+inline FVector4d operator-(const uniform double F, const FVector4d &A)
 {
-	FVector4 Result;
+	FVector4d Result;
+
+	for(uniform int i = 0; i < 4; i++)
+	{
+		Result.V[i] = F - A.V[i];
+	}
+
+	return Result;
+}
+
+inline FVector4f operator-(const uniform float F, const FVector4f &A)
+{
+	FVector4f Result;
 
 	for(uniform int i = 0; i < 4; i++)
 	{
@@ -2412,19 +2384,10 @@ inline uniform bool VectorIsAllNearlyZero(const uniform FVector3f &V, const unif
 	return all(abs(S0) <= Tolerance);
 }
 
-inline FVector3d VectorLerp(const FVector3d& A, const FVector3d& B, const double Alpha)
+template<typename T, typename V>
+inline T VectorLerp(const T& A, const T& B, const V Alpha)
 {
-	return (A + (B-A) * Alpha);
-}
-
-inline FVector3f VectorLerp(const FVector3f& A, const FVector3f& B, const float Alpha)
-{
-	return (A + (B-A) * Alpha);
-}
-
-inline uniform FVector VectorLerp(const uniform FVector& A, const uniform FVector& B, const uniform FReal Alpha)
-{
-	return (A + (B-A) * Alpha);
+	return (A + (B - A) * Alpha);
 }
 
 inline FVector VectorFloor(const FVector &A)
@@ -2432,76 +2395,22 @@ inline FVector VectorFloor(const FVector &A)
 	return SetVector(floor(A.V[0]), floor(A.V[1]), floor(A.V[2]));
 }
 
-inline FVector VectorClamp(const FVector &A, const FVector &B, const FVector &C)
-{
-	return SetVector(clamp(A.V[0], B.V[0], C.V[0]),
-		clamp(A.V[1], B.V[1], C.V[1]),
-		clamp(A.V[2], B.V[2], C.V[2]));
-}
-
-inline FVector VectorClamp(const FVector &A, const uniform FVector &B, const uniform FVector &C)
+template<typename T, typename V>
+inline T VectorClamp(const T &A, const V &B, const V &C)
 {
 	return SetVector(clamp(A.V[0], B.V[0], C.V[0]),
 		clamp(A.V[1], B.V[1], C.V[1]),
 		clamp(A.V[2], B.V[2], C.V[2]));
 }
 
-inline FVector4d VectorAdd(const FVector4d &A, const FVector4d &B)
+template <typename T>
+inline T VectorAdd(const T& A, const T& B)
 {
 	return A + B;
 }
 
-inline FVector4f VectorAdd(const FVector4f &A, const FVector4f &B)
-{
-	return A + B;
-}
-
-inline FVector4 VectorAdd(const FVector4 &A, const uniform FVector4 &B)
-{
-	return A + B;
-}
-
-inline FVector4f VectorAdd(const FVector4f &A, const uniform FVector4f &B)
-{
-	return A + B;
-}
-
-inline uniform FVector4d VectorAdd(const uniform FVector4d &A, const uniform FVector4d &B)
-{
-	return A + B;
-}
-
-inline uniform FVector4f VectorAdd(const uniform FVector4f &A, const uniform FVector4f &B)
-{
-	return A + B;
-}
-
-inline uniform FVector8 VectorAdd(const uniform FVector8 &A, const uniform FVector8 &B)
-{
-	return A + B;
-}
-
-inline FVector4d VectorSubtract(const FVector4d &A, const FVector4d &B)
-{
-	return A - B;
-}
-
-inline FVector4f VectorSubtract(const FVector4f &A, const FVector4f &B)
-{
-	return A - B;
-}
-
-inline FVector4 VectorSubtract(const uniform FVector4 &A, const FVector4 &B)
-{
-	return A - B;
-}
-
-inline uniform FVector4d VectorSubtract(const uniform FVector4d &A, const uniform FVector4d &B)
-{
-	return A - B;
-}
-
-inline uniform FVector4f VectorSubtract(const uniform FVector4f &A, const uniform FVector4f &B)
+template <typename T>
+inline T VectorSubtract(const T& A, const T& B)
 {
 	return A - B;
 }
@@ -2551,22 +2460,8 @@ inline uniform FVector8 VectorMultiply(const uniform FVector8 &A, const uniform
 	return A * B;
 }
 
-inline FVector4d VectorDivide(const FVector4d &A, const FVector4d &B)
-{
-	return A / B;
-}
-
-inline FVector4f VectorDivide(const FVector4f &A, const FVector4f &B)
-{
-	return A / B;
-}
-
-inline uniform FVector4d VectorDivide(const uniform FVector4d &A, const uniform FVector4d &B)
-{
-	return A / B;
-}
-
-inline uniform FVector4f VectorDivide(const uniform FVector4f &A, const uniform FVector4f &B)
+template <typename T>
+inline T VectorDivide(const T& A, const T& B)
 {
 	return A / B;
 }
@@ -2643,6 +2538,18 @@ inline FVector4 VectorMultiplyAdd(const FReal F, const FVector4 &B, const FVecto
 	return Result;
 }
 
+inline FVector4f VectorMultiplyAdd(const float F, const FVector4f& B, const FVector4f& C)
+{
+	FVector4f Result;
+
+	Result.V[0] = F * B.V[0] + C.V[0];
+	Result.V[1] = F * B.V[1] + C.V[1];
+	Result.V[2] = F * B.V[2] + C.V[2];
+	Result.V[3] = F * B.V[3] + C.V[3];
+
+	return Result;
+}
+
 inline uniform FVector4d VectorMultiplyAdd(const uniform FVector4d &A, const uniform FVector4d &B, const uniform FVector4d &C)
 {
 	varying double S0, S1, S2, Result;
@@ -3021,7 +2928,12 @@ inline uniform FReal VectorSum(const uniform FVector &V)
 	return V.V[0] + V.V[1] + V.V[2];
 }
 
-inline FReal VectorSum(const FVector &A)
+inline double VectorSum(const FVector3d &A)
+{
+	return A.V[0] + A.V[1] + A.V[2];
+}
+
+inline float VectorSum(const FVector3f &A)
 {
 	return A.V[0] + A.V[1] + A.V[2];
 }
@@ -3048,37 +2960,8 @@ inline uniform FVector4 VectorReciprocal(const uniform FVector4& V)
 	return *((uniform FVector4 *uniform)&Result);
 }
 
-inline FVector4 VectorReciprocalAccurate(const FVector4& V)
-{
-	// Perform two passes of Newton-Raphson iteration on the hardware estimate
-	//   x1 = x0 - f(x0) / f'(x0)
-	//
-	//    1 / Vec = x
-	// => x * Vec = 1
-	// => F(x) = x * Vec - 1
-	//    F'(x) = Vec
-	// => x1 = x0 - (x0 * Vec - 1) / Vec
-	//
-	// Since 1/Vec is what we're trying to solve, use an estimate for it, x0
-	// => x1 = x0 - (x0 * Vec - 1) * x0 = 2 * x0 - Vec * x0^2
-
-	// Initial estimate
-	const FVector4 x0 = VectorReciprocal(V);
-
-	// First iteration
-	const FVector4 x0Squared = VectorMultiply(x0, x0);
-	const FVector4 x0Times2 = VectorAdd(x0, x0);
-	const FVector4 x1 = VectorSubtract(x0Times2, VectorMultiply(V, x0Squared));
-
-	// Second iteration
-	const FVector4 x1Squared = VectorMultiply(x1, x1);
-	const FVector4 x1Times2 = VectorAdd(x1, x1);
-	const FVector4 x2 = VectorSubtract(x1Times2, VectorMultiply(V, x1Squared));
-
-	return x2;
-}
-
-inline uniform FVector4 VectorReciprocalAccurate(const uniform FVector4& V)
+template <typename T>
+inline T VectorReciprocalAccurate(const T& V)
 {
 	// Perform two passes of Newton-Raphson iteration on the hardware estimate
 	//   x1 = x0 - f(x0) / f'(x0)
@@ -3093,17 +2976,17 @@ inline uniform FVector4 VectorReciprocalAccurate(const uniform FVector4& V)
 	// => x1 = x0 - (x0 * Vec - 1) * x0 = 2 * x0 - Vec * x0^2
 
 	// Initial estimate
-	const uniform FVector4 x0 = VectorReciprocal(V);
+	const T x0 = VectorReciprocal(V);
 
 	// First iteration
-	const uniform FVector4 x0Squared = VectorMultiply(x0, x0);
-	const uniform FVector4 x0Times2 = VectorAdd(x0, x0);
-	const uniform FVector4 x1 = VectorSubtract(x0Times2, VectorMultiply(V, x0Squared));
+	const T x0Squared = VectorMultiply(x0, x0);
+	const T x0Times2 = VectorAdd(x0, x0);
+	const T x1 = VectorSubtract(x0Times2, VectorMultiply(V, x0Squared));
 
 	// Second iteration
-	const uniform FVector4 x1Squared = VectorMultiply(x1, x1);
-	const uniform FVector4 x1Times2 = VectorAdd(x1, x1);
-	const uniform FVector4 x2 = VectorSubtract(x1Times2, VectorMultiply(V, x1Squared));
+	const T x1Squared = VectorMultiply(x1, x1);
+	const T x1Times2 = VectorAdd(x1, x1);
+	const T x2 = VectorSubtract(x1Times2, VectorMultiply(V, x1Squared));
 
 	return x2;
 }
@@ -3255,27 +3138,14 @@ inline uniform FVector VectorReciprocalSqrtAccurate(const uniform FVector& V)
  * @param ElementIndex	Index (0-3) of the element to replicate
  * @return				VectorRegister( Vec[ElementIndex], Vec[ElementIndex], Vec[ElementIndex], Vec[ElementIndex] )
  */
-inline FVector4d VectorReplicate(const FVector4d &Vec, const uniform int ElementIndex)
-{
-	return SetVector4(Vec.V[ElementIndex], Vec.V[ElementIndex], Vec.V[ElementIndex], Vec.V[ElementIndex]);
-}
-
-inline FVector4f VectorReplicate(const FVector4f &Vec, const uniform int ElementIndex)
+template<typename T>
+inline T VectorReplicate(const T& Vec, const uniform int ElementIndex)
 {
 	return SetVector4(Vec.V[ElementIndex], Vec.V[ElementIndex], Vec.V[ElementIndex], Vec.V[ElementIndex]);
 }
 
-inline uniform FVector4d VectorReplicate(const uniform FVector4d &Vec, const uniform int ElementIndex)
-{
-	return SetVector4(Vec.V[ElementIndex], Vec.V[ElementIndex], Vec.V[ElementIndex], Vec.V[ElementIndex]);
-}
-
-inline uniform FVector4f VectorReplicate(const uniform FVector4f &Vec, const uniform int ElementIndex)
-{
-	return SetVector4(Vec.V[ElementIndex], Vec.V[ElementIndex], Vec.V[ElementIndex], Vec.V[ElementIndex]);
-}
-
-inline uniform FVector8 VectorReplicate(const uniform FVector8 &Vec, const uniform int R)
+template<>
+inline uniform FVector8 VectorReplicate<uniform FVector8>(const uniform FVector8 &Vec, const uniform int R)
 {
 	return SetVector8(Vec.V[R], Vec.V[R], Vec.V[R], Vec.V[R], Vec.V[R+4], Vec.V[R+4], Vec.V[R+4], Vec.V[R+4]);
 }
@@ -3346,17 +3216,8 @@ inline uniform FVector8 VectorSwizzle(const uniform FVector8 &Vec, const uniform
  * @param W			Index for which component of Vector2 to use for W (literal 0-3)
  * @return			The swizzled vector
  */
-inline FVector4 VectorShuffle(const FVector4 &Vec1, const FVector4 &Vec2, const uniform int X, const uniform int Y, const uniform int Z, const uniform int W)
-{
-	return SetVector4(Vec1.V[X], Vec1.V[Y], Vec2.V[Z], Vec2.V[W]);
-}
-
-inline uniform FVector4d VectorShuffle(const uniform FVector4d &Vec1, const uniform FVector4d &Vec2, const uniform int X, const uniform int Y, const uniform int Z, const uniform int W)
-{
-	return SetVector4(Vec1.V[X], Vec1.V[Y], Vec2.V[Z], Vec2.V[W]);
-}
-
-inline uniform FVector4f VectorShuffle(const uniform FVector4f &Vec1, const uniform FVector4f &Vec2, const uniform int X, const uniform int Y, const uniform int Z, const uniform int W)
+template <typename T>
+inline T VectorShuffle(const T& Vec1, const T& Vec2, const uniform int X, const uniform int Y, const uniform int Z, const uniform int W)
 {
 	return SetVector4(Vec1.V[X], Vec1.V[Y], Vec2.V[Z], Vec2.V[W]);
 }
@@ -4005,13 +3866,22 @@ inline void VectorSinCos(uniform FVector4 &VSinAngles, uniform FVector4 &VCosAng
 	VCosAngles = VectorMultiply(C, sign);
 }
 
-inline FVector4 VectorLerp(const FVector4 &A, const FVector4 &B, const FReal Alpha)
+template<>
+inline FVector4d VectorLerp<FVector4d, double>(const FVector4d& A, const FVector4d& B, const double Alpha)
 {
-	const FVector4 Delta = VectorSubtract(B, A);
+	const FVector4d Delta = VectorSubtract(B, A);
 	return VectorMultiplyAdd(Alpha, Delta, A);
 }
 
-inline uniform FVector4 VectorLerp(const uniform FVector4 &A, const uniform FVector4 &B, const uniform FReal Alpha)
+template<>
+inline FVector4f VectorLerp<FVector4f, float>(const FVector4f& A, const FVector4f& B, const float Alpha)
+{
+	const FVector4f Delta = VectorSubtract(B, A);
+	return VectorMultiplyAdd(Alpha, Delta, A);
+}
+
+template<>
+inline uniform FVector4 VectorLerp<uniform FVector4, uniform FReal>(const uniform FVector4& A, const uniform FVector4& B, const uniform FReal Alpha)
 {
 	const uniform FVector4 Delta = VectorSubtract(B, A);
 	return VectorMultiplyAdd(Alpha, Delta, A);
@@ -4106,7 +3976,8 @@ inline FVector4 VectorFloor(const FVector4 &A)
 	return SetVector4(floor(A.V[0]), floor(A.V[1]), floor(A.V[2]), floor(A.V[3]));
 }
 
-inline FVector4 VectorClamp(const FVector4 &A, const FVector4 &B, const FVector4 &C)
+template<>
+inline FVector4 VectorClamp<FVector4, FVector4>(const FVector4 &A, const FVector4 &B, const FVector4 &C)
 {
 	return SetVector4(clamp(A.V[0], B.V[0], C.V[0]),
 		clamp(A.V[1], B.V[1], C.V[1]),
@@ -4114,7 +3985,8 @@ inline FVector4 VectorClamp(const FVector4 &A, const FVector4 &B, const FVector4
 		clamp(A.V[3], B.V[3], C.V[3]));
 }
 
-inline FVector4 VectorClamp(const FVector4 &A, const uniform FVector4 &B, const uniform FVector4 &C)
+template<>
+inline FVector4 VectorClamp<FVector4, uniform FVector4>(const FVector4 &A, const uniform FVector4 &B, const uniform FVector4 &C)
 {
 	return SetVector4(clamp(A.V[0], B.V[0], C.V[0]),
 		clamp(A.V[1], B.V[1], C.V[1]),
@@ -4122,37 +3994,26 @@ inline FVector4 VectorClamp(const FVector4 &A, const uniform FVector4 &B, const
 		clamp(A.V[3], B.V[3], C.V[3]));
 }
 
-inline uniform FVector3d VectorReduceAdd(const FVector3d& A)
+template<typename T>
+inline uniform T VectorReduceAdd(const T& A)
 {
 	return SetVector(reduce_add(A.V[0]), reduce_add(A.V[1]), reduce_add(A.V[2]));
 }
 
-inline uniform FVector3f VectorReduceAdd(const FVector3f& A)
-{
-	return SetVector(reduce_add(A.V[0]), reduce_add(A.V[1]), reduce_add(A.V[2]));
-}
-
-inline uniform FVector4 VectorReduceAdd(const FVector4& A)
+template<>
+inline uniform FVector4 VectorReduceAdd<FVector4>(const FVector4& A)
 {
 	return SetVector4(reduce_add(A.V[0]), reduce_add(A.V[1]), reduce_add(A.V[2]), reduce_add(A.V[3]));
 }
 
-inline uniform FVector3d VectorReduceMin(const FVector3d& A)
+template <typename T>
+inline uniform T VectorReduceMin(const T& A)
 {
 	return SetVector(reduce_min(A.V[0]), reduce_min(A.V[1]), reduce_min(A.V[2]));
 }
 
-inline uniform FVector3f VectorReduceMin(const FVector3f& A)
-{
-	return SetVector(reduce_min(A.V[0]), reduce_min(A.V[1]), reduce_min(A.V[2]));
-}
-
-inline uniform FVector3d VectorReduceMax(const FVector3d& A)
-{
-	return SetVector(reduce_max(A.V[0]), reduce_max(A.V[1]), reduce_max(A.V[2]));
-}
-
-inline uniform FVector3f VectorReduceMax(const FVector3f& A)
+template <typename T>
+inline uniform T VectorReduceMax(const T& A)
 {
 	return SetVector(reduce_max(A.V[0]), reduce_max(A.V[1]), reduce_max(A.V[2]));
 }
diff --git a/Engine/Source/Runtime/Core/Public/Math/WideVector.isph b/Engine/Source/Runtime/Core/Public/Math/WideVector.isph
--- a/Engine/Source/Runtime/Core/Public/Math/WideVector.isph
+++ b/Engine/Source/Runtime/Core/Public/Math/WideVector.isph
@@ -86,22 +86,26 @@ unmasked inline uniform WideFVector4 VectorSwizzle(const uniform WideFVector4 &V
 	return Result;
 }
 
-unmasked inline uniform WideFVector4 VectorReplicate(const uniform WideFVector4 &Vec, const uniform int R)
+template<>
+inline uniform WideFVector4 VectorReplicate<uniform WideFVector4>(const uniform WideFVector4 &Vec, const uniform int R)
 {
-#if TARGET_WIDTH == 4
-	const varying int vPerm = { R, R, R, R };
-#elif TARGET_WIDTH == 8
-	const varying int vPerm = { R, R, R, R, R+4, R+4, R+4, R+4 };
-#elif TARGET_WIDTH == 16
-	const varying int vPerm = { R, R, R, R, R+4, R+4, R+4, R+4, R+8, R+8, R+8, R+8, R+12, R+12, R+12, R+12 };
-#endif
-
-	const FReal V = Vec.V[programIndex];
-	const FReal S = shuffle(V, vPerm);
-	uniform WideFVector4 Result;
-	Result.V[programIndex] = S;
-
-	return Result;
+	unmasked
+	{
+	#if TARGET_WIDTH == 4
+		const varying int vPerm = { R, R, R, R };
+	#elif TARGET_WIDTH == 8
+		const varying int vPerm = { R, R, R, R, R+4, R+4, R+4, R+4 };
+	#elif TARGET_WIDTH == 16
+		const varying int vPerm = { R, R, R, R, R+4, R+4, R+4, R+4, R+8, R+8, R+8, R+8, R+12, R+12, R+12, R+12 };
+	#endif
+
+		const FReal V = Vec.V[programIndex];
+		const FReal S = shuffle(V, vPerm);
+		uniform WideFVector4 Result;
+		Result.V[programIndex] = S;
+
+		return Result;
+	}
 }
 
 unmasked inline uniform WideFVector4 VectorCompareGE(const uniform WideFVector4 &A, const uniform WideFVector4 &B)
