FNames: Fix bug in FNameEntryId.CompareLexical and LexicalLess. CompareLexical was converting negative results to 1 and LexicalLess was never returning true.
#rnx

diff --git a/Engine/Source/Runtime/Core/Private/UObject/UnrealNames.cpp b/Engine/Source/Runtime/Core/Private/UObject/UnrealNames.cpp
--- a/Engine/Source/Runtime/Core/Private/UObject/UnrealNames.cpp
+++ b/Engine/Source/Runtime/Core/Private/UObject/UnrealNames.cpp
@@ -2061,7 +2061,7 @@ static int32 CompareDifferentIdsAlphabetically(FNameEntryId AId, FNameEntryId BI
 
 int32 FNameEntryId::CompareLexical(FNameEntryId Rhs) const
 {
-	return Value != Rhs.Value && CompareDifferentIdsAlphabetically(*this, Rhs);
+	return Value != Rhs.Value ? CompareDifferentIdsAlphabetically(*this, Rhs) : 0;
 }
 
 #if !UE_BUILD_SHIPPING && !UE_BUILD_TEST
