Use atomic OR/AND to atomically set/clear bits in a bit array.

Cheaper than a compare-exchange, and avoids races around non-atomic reads and writes.

diff --git a/Engine/Source/Runtime/Core/Public/Containers/BitArray.h b/Engine/Source/Runtime/Core/Public/Containers/BitArray.h
--- a/Engine/Source/Runtime/Core/Public/Containers/BitArray.h
+++ b/Engine/Source/Runtime/Core/Public/Containers/BitArray.h
@@ -225,34 +225,19 @@ public:
 	}
 	FORCEINLINE void AtomicSet(const bool NewValue)
 	{
+		uint32 Current = static_cast<uint32>(FPlatformAtomics::AtomicRead_Relaxed((const volatile int32*)&Data));
 		if(NewValue)
 		{
-			if (!(Data & Mask))
+			if (!(Current & Mask))
 			{
-				while (1)
-				{
-					uint32 Current = Data;
-					uint32 Desired = Current | Mask;
-					if (Current == Desired || FPlatformAtomics::InterlockedCompareExchange((volatile int32*)&Data, (int32)Desired, (int32)Current) == (int32)Current)
-					{
-						return;
-					}
-				}
+				FPlatformAtomics::InterlockedOr((volatile int32*)&Data, (int32)Mask);
 			}
 		}
 		else
 		{
-			if (Data & Mask)
+			if (Current & Mask)
 			{
-				while (1)
-				{
-					uint32 Current = Data;
-					uint32 Desired = Current & ~Mask;
-					if (Current == Desired || FPlatformAtomics::InterlockedCompareExchange((volatile int32*)&Data, (int32)Desired, (int32)Current) == (int32)Current)
-					{
-						return;
-					}
-				}
+				FPlatformAtomics::InterlockedAnd((volatile int32*)&Data, (int32)~Mask);
 			}
 		}
 	}
