Added a scoped override for FMessageLog behavior

FMessageLogScopedOverride can override the log behavior of a given named FMessageLog for the duration of the overrides lifetime, optionally suppressing output log mirroring, or promoting/demoting log categories (eg, to make errors act as warnings, or warnings act as errors).

Note: Like the internal implementation of FMessageLog (FMessageLogViewModel) FMessageLogScopedOverride is not thread-safe and should only be used from the game thread.

diff --git a/Engine/Source/Runtime/Core/Private/Logging/MessageLog.cpp b/Engine/Source/Runtime/Core/Private/Logging/MessageLog.cpp
--- a/Engine/Source/Runtime/Core/Private/Logging/MessageLog.cpp
+++ b/Engine/Source/Runtime/Core/Private/Logging/MessageLog.cpp
@@ -12,6 +12,29 @@ FMessageLog::FMessageSelectionChanged FMessageLog::MessageSelectionChanged;
 
 LLM_DEFINE_TAG(EngineMisc_MessageLog);
 
+namespace UE::MessageLog::Private
+{
+	using FOverrideStackMap = TSortedMap<FName, TArray<const FMessageLogScopedOverride*>, FDefaultAllocator, FNameFastLess>;
+
+	/** Get the internal stack of FMessageLogScopedOverride's */
+	FOverrideStackMap& GetOverrideStack()
+	{
+		static FOverrideStackMap OverrideStack;
+		return OverrideStack;
+	}
+
+	/** Get the current FMessageLogScopedOverride for the given log name, if any */
+	const FMessageLogScopedOverride* GetOverride(const FName LogName)
+	{
+		if (TArray<const FMessageLogScopedOverride*>* NamedOverrideStack = GetOverrideStack().Find(LogName);
+			NamedOverrideStack && NamedOverrideStack->Num() > 0)
+		{
+			return NamedOverrideStack->Top();
+		}
+		return nullptr;
+	}
+}
+
 class FBasicMessageLog : public IMessageLog, public TSharedFromThis<FBasicMessageLog>
 {
 public:
@@ -94,7 +117,8 @@ private:
 };
 
 FMessageLog::FMessageLog( const FName& InLogName )
-	: bSuppressLoggingToOutputLog(false)
+	: LogName(InLogName)
+	, bSuppressLoggingToOutputLog(false)
 {
 	LLM_SCOPE_BYTAG(EngineMisc_MessageLog);
 	
@@ -252,7 +276,59 @@ void FMessageLog::Flush()
 {
 	if(Messages.Num() > 0)
 	{
-		MessageLog->AddMessages(Messages, !bSuppressLoggingToOutputLog);
+		bool bMirrorToOutputLog = !bSuppressLoggingToOutputLog;
+
+		// We defer querying the FMessageLogScopedOverride until Flush to minimize the number of map lookups while adding messages
+		// Though this does mean that messages added during an overload (and not flushed) may log with the wrong severity override
+		if (const FMessageLogScopedOverride* LogOverride = IsInGameThread() ? UE::MessageLog::Private::GetOverride(LogName) : nullptr)
+		{
+			bMirrorToOutputLog = !LogOverride->bSuppressLoggingToOutputLog.Get(bSuppressLoggingToOutputLog);
+
+			if (LogOverride->MessageSeverityRemapping.Num() > 0)
+			{
+				auto ResolveMessageSeverity = [LogOverride](const EMessageSeverity::Type BaseSeverity) -> EMessageSeverity::Type
+				{
+					EMessageSeverity::Type ResolvedSeverity = BaseSeverity;
+					{
+						TArray<EMessageSeverity::Type, TInlineAllocator<4>> CycleDetector;
+						EMessageSeverity::Type DemotedSeverity = BaseSeverity;
+						for (;;)
+						{
+							CycleDetector.Add(ResolvedSeverity);
+							if (const EMessageSeverity::Type* MessageSeverityOverride = LogOverride->MessageSeverityRemapping.Find(ResolvedSeverity))
+							{
+								if (CycleDetector.Contains(*MessageSeverityOverride))
+								{
+									// Remapping cycle!
+									// Since we can't tell whether the intent was to promote or demote this severity, 
+									// we pessimistically assume the intent was to demote and use the most demoted value
+									ResolvedSeverity = DemotedSeverity;
+									break;
+								}
+								ResolvedSeverity = *MessageSeverityOverride;
+								DemotedSeverity = FMath::Max(DemotedSeverity, ResolvedSeverity); // Bigger numbers are lower severity
+								continue;
+							}
+							break;
+						}
+					}
+					return ResolvedSeverity;
+				};
+
+				for (TSharedRef<FTokenizedMessage>& Message : Messages)
+				{
+					const EMessageSeverity::Type ResolvedSeverity = ResolveMessageSeverity(Message->GetSeverity());
+					if (ResolvedSeverity != Message->GetSeverity())
+					{
+						// Need to clone the FTokenizedMessage as the severity value may be shared with other usages of this message
+						Message = Message->Clone();
+						Message->SetSeverity(ResolvedSeverity);
+					}
+				}
+			}
+		}
+
+		MessageLog->AddMessages(Messages, bMirrorToOutputLog);
 		Messages.Empty();
 	}
 }
@@ -291,4 +367,52 @@ const TCHAR* const FMessageLog::GetLogColor( EMessageSeverity::Type InSeverity )
 #endif
 }
 
+
+FMessageLogScopedOverride::FMessageLogScopedOverride(const FName InLogName)
+	: LogName(InLogName)
+{
+	LLM_SCOPE_BYTAG(EngineMisc_MessageLog);
+
+	checkf(IsInGameThread(), TEXT("FMessageLogScopedOverride should only be pushed from the game-thread!"));
+	TArray<const FMessageLogScopedOverride*>& NamedOverrideStack = UE::MessageLog::Private::GetOverrideStack().FindOrAdd(LogName);
+	if (NamedOverrideStack.Num() > 0)
+	{
+		// Set our initial state to the top of the current stack, as that "inherits" the current override state
+		const FMessageLogScopedOverride* CurrentOverride = NamedOverrideStack.Top();
+		bSuppressLoggingToOutputLog = CurrentOverride->bSuppressLoggingToOutputLog;
+		MessageSeverityRemapping = CurrentOverride->MessageSeverityRemapping;
+	}
+	NamedOverrideStack.Push(this);
+}
+
+FMessageLogScopedOverride::~FMessageLogScopedOverride()
+{
+	LLM_SCOPE_BYTAG(EngineMisc_MessageLog);
+
+	checkf(IsInGameThread(), TEXT("FMessageLogScopedOverride should only be popped from the game-thread!"));
+	if (TArray<const FMessageLogScopedOverride*>* NamedOverrideStack = UE::MessageLog::Private::GetOverrideStack().Find(LogName))
+	{
+		const FMessageLogScopedOverride* This = NamedOverrideStack->Pop();
+		checkf(this == This, TEXT("FMessageLogScopedOverride popped an unexpected override off the stack! Out of order destruction of FMessageLogScopedOverride!"));
+	}
+	else
+	{
+		checkf(false, TEXT("FMessageLogScopedOverride popped an unexpected log name off the stack! Out of order destruction of FMessageLogScopedOverride!"));
+	}
+}
+
+FMessageLogScopedOverride& FMessageLogScopedOverride::SuppressLoggingToOutputLog(const bool bShouldSuppress)
+{
+	bSuppressLoggingToOutputLog = bShouldSuppress;
+	return *this;
+}
+
+FMessageLogScopedOverride& FMessageLogScopedOverride::RemapMessageSeverity(const EMessageSeverity::Type SrcSeverity, const EMessageSeverity::Type DestSeverity)
+{
+	LLM_SCOPE_BYTAG(EngineMisc_MessageLog);
+
+	MessageSeverityRemapping.Add(SrcSeverity, DestSeverity);
+	return *this;
+}
+
 #undef LOCTEXT_NAMESPACE
diff --git a/Engine/Source/Runtime/Core/Private/Logging/TokenizedMessage.cpp b/Engine/Source/Runtime/Core/Private/Logging/TokenizedMessage.cpp
--- a/Engine/Source/Runtime/Core/Private/Logging/TokenizedMessage.cpp
+++ b/Engine/Source/Runtime/Core/Private/Logging/TokenizedMessage.cpp
@@ -22,6 +22,13 @@ TSharedRef<FTokenizedMessage> FTokenizedMessage::Create(EMessageSeverity::Type I
 	return Message;
 }
 
+TSharedRef<FTokenizedMessage> FTokenizedMessage::Clone() const
+{
+	TSharedRef<FTokenizedMessage> Message = MakeShared<FTokenizedMessage>(FPrivateToken());
+	*Message = *this;
+	return Message;
+}
+
 FText FTokenizedMessage::ToText() const
 {
 	FText OutMessage;
diff --git a/Engine/Source/Runtime/Core/Public/Logging/MessageLog.h b/Engine/Source/Runtime/Core/Public/Logging/MessageLog.h
--- a/Engine/Source/Runtime/Core/Public/Logging/MessageLog.h
+++ b/Engine/Source/Runtime/Core/Public/Logging/MessageLog.h
@@ -3,6 +3,7 @@
 #pragma once
 
 #include "Containers/Array.h"
+#include "Containers/SortedMap.h"
 #include "CoreTypes.h"
 #include "Delegates/Delegate.h"
 #include "Internationalization/Text.h"
@@ -133,6 +134,9 @@ private:
 	/** The message log we use for output */
 	TSharedPtr<class IMessageLog> MessageLog;
 
+	/** Name of this message log */
+	FName LogName;
+
 	/** Do we want to suppress logging to the output log in addition to the message log? */
 	bool bSuppressLoggingToOutputLog;
 
@@ -143,3 +147,29 @@ private:
 	CORE_API static FMessageSelectionChanged MessageSelectionChanged;
 };
 
+/**
+ * Scoped override for FMessageLog behavior.
+ * This can override the log behavior of a given named FMessageLog for the duration of the overrides lifetime, optionally suppressing 
+ * output log mirroring, or promoting/demoting log categories (eg, to make errors act as warnings, or warnings act as errors).
+ */
+class FMessageLogScopedOverride final
+{
+public:
+	CORE_API explicit FMessageLogScopedOverride(const FName InLogName);
+	CORE_API ~FMessageLogScopedOverride();
+
+	UE_NONCOPYABLE(FMessageLogScopedOverride);
+
+	/** Should we mirror message log messages to the output log during flush? */
+	CORE_API FMessageLogScopedOverride& SuppressLoggingToOutputLog(const bool bShouldSuppress = true);
+
+	/** Map category X to category Y when adding messages to this log */
+	CORE_API FMessageLogScopedOverride& RemapMessageSeverity(const EMessageSeverity::Type SrcSeverity, const EMessageSeverity::Type DestSeverity);
+
+private:
+	friend class FMessageLog;
+
+	FName LogName;
+	TOptional<bool> bSuppressLoggingToOutputLog;
+	TSortedMap<EMessageSeverity::Type, EMessageSeverity::Type> MessageSeverityRemapping;
+};
diff --git a/Engine/Source/Runtime/Core/Public/Logging/TokenizedMessage.h b/Engine/Source/Runtime/Core/Public/Logging/TokenizedMessage.h
--- a/Engine/Source/Runtime/Core/Public/Logging/TokenizedMessage.h
+++ b/Engine/Source/Runtime/Core/Public/Logging/TokenizedMessage.h
@@ -123,6 +123,12 @@ public:
 	 */
 	CORE_API static TSharedRef<FTokenizedMessage> Create(EMessageSeverity::Type InSeverity, const FText& InMessageText = FText());
 
+	/**
+	 * Clone this message.
+	 * @note The message tokens are shared between the original and the clone.
+	 */
+	CORE_API TSharedRef<FTokenizedMessage> Clone() const;
+
 	/** 
 	 * Get this tokenized message as a string
 	 * 
