Fixed locale issue when initializng exporters on computer with non-US locale
  Solution: Assumed console numerical variables are always stored according to the "C" locale
Fixed reading and writing of Datasmith animation values which were not enforcing the use of the "C" locale

diff --git a/Engine/Source/Runtime/Core/Private/HAL/ConsoleManager.cpp b/Engine/Source/Runtime/Core/Private/HAL/ConsoleManager.cpp
--- a/Engine/Source/Runtime/Core/Private/HAL/ConsoleManager.cpp
+++ b/Engine/Source/Runtime/Core/Private/HAL/ConsoleManager.cpp
@@ -21,6 +21,8 @@ ConsoleManager.cpp: console command handling
 #include "HAL/FileManager.h"
 #include "Serialization/ArchiveCountMem.h"
 
+#include <clocale>
+
 DEFINE_LOG_CATEGORY(LogConsoleResponse);
 DEFINE_LOG_CATEGORY_STATIC(LogConsoleManager, Log, All);
 
@@ -29,6 +31,44 @@ namespace UE::ConsoleManager::Private
 	// this tracks the cvars that were  added dynamically with a tag (via plugin or similar)
 	// we use this structure to unset the cvars and update the value when the plugin unloads
 	TMap<FName, TSet<IConsoleVariable*>*> TaggedCVars;
+
+	/**
+	 * Setup a locale for a given scope, the previous locale is restored on scope end.
+	 */
+	struct FConsoleManagerLocaleScope
+	{
+		// Set a custom locale within this object scope
+		FConsoleManagerLocaleScope()
+		{
+			if (const char* saved = std::setlocale(LC_NUMERIC, nullptr))
+			{
+				SavedLocale = TArray<char>(saved, TCString<char>::Strlen(saved) + 1);
+			}
+			std::setlocale(LC_NUMERIC, "C");
+		}
+
+		// restore the captured Locale
+		~FConsoleManagerLocaleScope()
+		{
+			std::setlocale(LC_NUMERIC, SavedLocale.GetData());
+		}
+
+	private:
+		TArray<char> SavedLocale; // Locale previously used, captured, and to be restored
+	};
+
+	template<typename T>
+	void GetValueFromString(T& Value, const TCHAR* Buffer)
+	{
+		TTypeFromString<T>::FromString(Value, Buffer);
+	};
+
+	template<>
+	void GetValueFromString<float>(float& Value, const TCHAR* Buffer)
+	{
+		FConsoleManagerLocaleScope LocaleScope;
+		TTypeFromString<float>::FromString(Value, Buffer);
+	};
 }
 
 static inline bool IsWhiteSpace(TCHAR Value) { return Value == TCHAR(' '); }
@@ -664,7 +704,7 @@ public:
 	void Track(const TCHAR* InValue, EConsoleVariableFlags SetBy, FName Tag)
 	{
 		T LocalCopy;
-		TTypeFromString<T>::FromString(LocalCopy, UE::ConfigUtilities::ConvertValueFromHumanFriendlyValue(InValue));
+		UE::ConsoleManager::Private::GetValueFromString<T>(LocalCopy, UE::ConfigUtilities::ConvertValueFromHumanFriendlyValue(InValue));
 		
 		int Priority = (int)SetBy;
 		TArray<FTaggedHistoryData>& ValueArray = History.FindOrAdd(Priority);
@@ -868,7 +908,7 @@ protected:
 		{
 			// update value
 			T ConvertedValue;
-			TTypeFromString<T>::FromString(ConvertedValue, UE::ConfigUtilities::ConvertValueFromHumanFriendlyValue(InValue));
+			UE::ConsoleManager::Private::GetValueFromString<T>(ConvertedValue, UE::ConfigUtilities::ConvertValueFromHumanFriendlyValue(InValue));
 			// set the value, and push to render thread value as well, but don't trigger callbacks and don't check priorties
 			SetInternalAndUpdateState(ConvertedValue, (EConsoleVariableFlags)(SetBy | ECVF_Set_SetOnly_Unsafe));
 
@@ -1040,6 +1080,7 @@ template<> float FConsoleVariableConversionHelper<float>::GetFloat(float Value)
 }
 template<> FString FConsoleVariableConversionHelper<float>::GetString(float Value)
 {
+	UE::ConsoleManager::Private::FConsoleManagerLocaleScope LocaleScope;
 	return FString::Printf(TEXT("%g"), Value);
 }
 
@@ -1059,7 +1100,7 @@ template<> int32 FConsoleVariableConversionHelper<FString>::GetInt(FString Value
 template<> float FConsoleVariableConversionHelper<FString>::GetFloat(FString Value)
 {
 	float OutValue = 0.0f;
-	TTypeFromString<float>::FromString(OutValue, *Value);
+	UE::ConsoleManager::Private::GetValueFromString<float>(OutValue, *Value);
 	return OutValue;
 }
 template<> FString FConsoleVariableConversionHelper<FString>::GetString(FString Value)
@@ -1100,7 +1141,7 @@ public:
 
 		if (CanChange(SetBy))
 		{
-			TTypeFromString<T>::FromString(Data.ShadowedValue[0], InValue);
+			UE::ConsoleManager::Private::GetValueFromString<T>(Data.ShadowedValue[0], InValue);
 			OnChanged(SetBy, false);
 		}
 	}
@@ -1360,7 +1401,7 @@ public:
 		
 		if(CanChange(SetBy))
 		{
-			TTypeFromString<T>::FromString(MainValue, InValue);
+			UE::ConsoleManager::Private::GetValueFromString<T>(MainValue, InValue);
 			OnChanged(SetBy);
 		}
 	}
@@ -1464,7 +1505,7 @@ public:
 	virtual float GetFloat() const
 	{
 		float Result = 0.0f;
-		TTypeFromString<float>::FromString(Result, *MainValue);
+		UE::ConsoleManager::Private::GetValueFromString<float>(Result, *MainValue);
 		return Result;
 	}
 	virtual FString GetString() const
diff --git a/Engine/Source/Runtime/Datasmith/DatasmithCore/Private/DatasmithAnimationSerializer.cpp b/Engine/Source/Runtime/Datasmith/DatasmithCore/Private/DatasmithAnimationSerializer.cpp
--- a/Engine/Source/Runtime/Datasmith/DatasmithCore/Private/DatasmithAnimationSerializer.cpp
+++ b/Engine/Source/Runtime/Datasmith/DatasmithCore/Private/DatasmithAnimationSerializer.cpp
@@ -1,6 +1,7 @@
 // Copyright Epic Games, Inc. All Rights Reserved.
 
 #include "DatasmithAnimationSerializer.h"
+#include "DatasmithLocaleScope.h"
 #include "DatasmithSceneFactory.h"
 #include "DatasmithUtils.h"
 #include "IDatasmithSceneElements.h"
@@ -446,11 +447,13 @@ namespace DatasmithAnimationJsonSerializerImpl
 
 bool FDatasmithAnimationSerializer::Serialize(const TSharedRef<IDatasmithLevelSequenceElement>& LevelSequence, const TCHAR* FilePath, bool bDebugFormat)
 {
+	FDatasmithLocaleScope LocaleScope;
 	return DatasmithAnimationJsonSerializerImpl::SerializeLevelSequence(LevelSequence, FilePath, bDebugFormat);
 }
 
 bool FDatasmithAnimationSerializer::Deserialize(const TSharedRef<IDatasmithLevelSequenceElement>& LevelSequence, const TCHAR* FilePath)
 {
+	FDatasmithLocaleScope LocaleScope;
 	return DatasmithAnimationJsonSerializerImpl::DeserializeLevelSequence(LevelSequence, FilePath);
 }
 
