Clean up UE_EXTERNAL_PROFILING_ENABLED for platform extensions

#rnx

diff --git a/Engine/Source/Runtime/Core/Public/ProfilingDebugging/ExternalProfiler.h b/Engine/Source/Runtime/Core/Public/ProfilingDebugging/ExternalProfiler.h
--- a/Engine/Source/Runtime/Core/Public/ProfilingDebugging/ExternalProfiler.h
+++ b/Engine/Source/Runtime/Core/Public/ProfilingDebugging/ExternalProfiler.h
@@ -9,8 +9,7 @@
 #include "UObject/NameTypes.h"
 
 #ifndef UE_EXTERNAL_PROFILING_ENABLED
-// temporarily turn off profiler on Switch because of a compiler issue (?) with the thread_local init bools (switching 1 to int helped, 1 did not)
-#define UE_EXTERNAL_PROFILING_ENABLED (!UE_BUILD_SHIPPING && !PLATFORM_SWITCH)
+#define UE_EXTERNAL_PROFILING_ENABLED (!UE_BUILD_SHIPPING)
 #endif
 
 #if UE_EXTERNAL_PROFILING_ENABLED
