ConfigAccessTracking: Skip spurious reads or suppress spurious reports of accessed config values.
Minor formatting changes that are part of the upcoming UE::ConfigAccessTracking code.
#rnx

diff --git a/Engine/Source/Runtime/Core/Private/Misc/ConfigCacheIni.cpp b/Engine/Source/Runtime/Core/Private/Misc/ConfigCacheIni.cpp
--- a/Engine/Source/Runtime/Core/Private/Misc/ConfigCacheIni.cpp
+++ b/Engine/Source/Runtime/Core/Private/Misc/ConfigCacheIni.cpp
@@ -275,32 +275,34 @@ bool FConfigSection::HasQuotes( const FString& Test )
 	return Test.Left(1) == TEXT("\"") && Test.Right(1) == TEXT("\"");
 }
 
-bool FConfigSection::operator==( const FConfigSection& Other ) const
+bool FConfigSection::operator==( const FConfigSection& B ) const
 {
-	if ( Pairs.Num() != Other.Pairs.Num() )
+	const FConfigSection&A = *this;
+	if ( A.Pairs.Num() != B.Pairs.Num() )
 	{
-		return 0;
+		return false;
 	}
 
-	FConfigSectionMap::TConstIterator My(*this), Their(Other);
-	while ( My && Their )
+	FConfigSectionMap::TConstIterator AIter(A), BIter(B);
+	while (AIter && BIter)
 	{
-		if (My.Key() != Their.Key())
+		if (AIter.Key() != BIter.Key())
 		{
-			return 0;
+			return false;
 		}
 
-		const FString& MyValue = My.Value().GetValue(), &TheirValue = Their.Value().GetValue();
-		if ( FCString::Strcmp(*MyValue,*TheirValue) &&
-			(!HasQuotes(MyValue) || FCString::Strcmp(*TheirValue,*MyValue.Mid(1,MyValue.Len()-2))) &&
-			(!HasQuotes(TheirValue) || FCString::Strcmp(*MyValue,*TheirValue.Mid(1,TheirValue.Len()-2))) )
+		const FString& AIterValue = AIter.Value().GetValue();
+		const FString& BIterValue = BIter.Value().GetValue();
+		if ( FCString::Strcmp(*AIterValue,*BIterValue) &&
+			(!HasQuotes(AIterValue) || FCString::Strcmp(*BIterValue,*AIterValue.Mid(1, AIterValue.Len() - 2))) &&
+			(!HasQuotes(BIterValue) || FCString::Strcmp(*AIterValue,*BIterValue.Mid(1, BIterValue.Len() - 2))) )
 		{
-			return 0;
+			return false;
 		}
 
-		++My, ++Their;
+		++AIter, ++BIter;
 	}
-	return 1;
+	return true;
 }
 
 bool FConfigSection::operator!=( const FConfigSection& Other ) const
@@ -308,6 +310,40 @@ bool FConfigSection::operator!=( const FConfigSection& Other ) const
 	return ! (FConfigSection::operator==(Other));
 }
 
+namespace UE::ConfigCacheIni::Private
+{
+
+bool FAccessor::AreSectionsEqualForWriting(const FConfigSection& A, const FConfigSection& B)
+{
+	if (A.Pairs.Num() != B.Pairs.Num())
+	{
+		return false;
+	}
+
+	FConfigSectionMap::TConstIterator AIter(A), BIter(B);
+	while (AIter && BIter)
+	{
+		if (AIter.Key() != BIter.Key())
+		{
+			return false;
+		}
+
+		const FString& AIterValue = AIter.Value().GetValueForWriting();
+		const FString& BIterValue = BIter.Value().GetValueForWriting();
+		if (FCString::Strcmp(*AIterValue, *BIterValue) &&
+			(!FConfigSection::HasQuotes(AIterValue) || FCString::Strcmp(*BIterValue, *AIterValue.Mid(1, AIterValue.Len() - 2))) &&
+			(!FConfigSection::HasQuotes(BIterValue) || FCString::Strcmp(*AIterValue, *BIterValue.Mid(1, BIterValue.Len() - 2))))
+		{
+			return false;
+		}
+
+		++AIter, ++BIter;
+	}
+	return true;
+}
+
+} // namespace UE::ConfigCacheIni::Private
+
 FArchive& operator<<(FArchive& Ar, FConfigSection& ConfigSection)
 {
 	Ar << static_cast<FConfigSection::Super&>(ConfigSection);
@@ -349,17 +385,17 @@ static void ExtractPropertyValue(const FString& FullStructValue, const FString&
 	}
 }
 
-void FConfigSection::HandleAddCommand(FName Key, FString&& Value, bool bAppendValueIfNotArrayOfStructsKeyUsed)
+void FConfigSection::HandleAddCommand(FName ValueName, FString&& Value, bool bAppendValueIfNotArrayOfStructsKeyUsed)
 {
-	if (!HandleArrayOfKeyedStructsCommand(Key, Forward<FString&&>(Value)))
+	if (!HandleArrayOfKeyedStructsCommand(ValueName, Forward<FString&&>(Value)))
 	{
 		if (bAppendValueIfNotArrayOfStructsKeyUsed)
 		{
-			Add(Key, MoveTemp(Value));
+			Add(ValueName, FConfigValue(MoveTemp(Value)));
 		}
 		else
 		{
-			AddUnique(Key, MoveTemp(Value));
+			AddUnique(ValueName, FConfigValue(MoveTemp(Value)));
 		}
 	}
 }
@@ -387,7 +423,10 @@ bool FConfigSection::HandleArrayOfKeyedStructsCommand(FName Key, FString&& Value
 				if (It.Key() == Key)
 				{
 					// now look for the matching ArrayOfStruct Key as the incoming KeyValue
-					ExtractPropertyValue(It.Value().GetValue(), StructKeyMatch, ExistingStructValueKey);
+					{
+						const FString& ItValue = UE::ConfigCacheIni::Private::FAccessor::GetValueForWriting(It.Value()); // Don't report to AccessTracking
+						ExtractPropertyValue(ItValue, StructKeyMatch, ExistingStructValueKey);
+					}
 					if (ExistingStructValueKey == StructKeyValueToMatch)
 					{
 						// we matched the key, so replace the existing value in place (so as not to reorder)
@@ -561,6 +600,8 @@ static bool SaveConfigFileWrapper(const TCHAR* IniFile, const FString& Contents)
 FConfigFile::FConfigFile()
 : Dirty( false )
 , NoSave( false )
+, bHasPlatformName( false )
+, bCanSaveAllSections( true )
 , Name( NAME_None )
 , SourceConfigFile(nullptr)
 {
@@ -961,11 +1002,12 @@ void FConfigFile::CombineFromBuffer(const FString& Buffer, const FString& FileHi
 						FConfigValue* ConfigValue = CurrentSection->Find(Key);
 						if (!ConfigValue)
 						{
-							CurrentSection->Add(Key, MoveTemp(ProcessedValue));
+							CurrentSection->Add(Key,
+								FConfigValue(MoveTemp(ProcessedValue)));
 						}
 						else
 						{
-							*ConfigValue = FConfigValue(MoveTemp(ProcessedValue));
+							*ConfigValue = MoveTemp(ProcessedValue);
 						}
 					}
 				}
@@ -1105,12 +1147,12 @@ void FConfigFile::ProcessInputFileContents(FStringView Contents, const FString&
 					FParse::QuotedString(Value, ProcessedValue);
 
 					// Add this pair to the current FConfigSection
-					CurrentSection->Add(KeyName, *ProcessedValue);
+					CurrentSection->Add(KeyName, FConfigValue(MoveTemp(ProcessedValue)));
 				}
 				else
 				{
 					// Add this pair to the current FConfigSection
-					CurrentSection->Add(KeyName, Value);
+					CurrentSection->Add(KeyName, FConfigValue(Value));
 				}
 			}
 		}
@@ -1376,7 +1418,11 @@ void FConfigFile::AddDynamicLayerToHierarchy(const FString& Filename)
 }
 
 
+namespace UE::ConfigCacheIni::Private
+{
 
+struct FImpl
+{
 /**
  * Check if the provided config section has a property which matches the one we are providing
  *
@@ -1386,7 +1432,7 @@ void FConfigFile::AddDynamicLayerToHierarchy(const FString& Filename)
  *
  * @return True if a property was found in the InSection which matched the Property Name and Value.
  */
-bool DoesConfigPropertyValueMatch(const FConfigSection* InSection, const FName& InPropertyName, const FString& InPropertyValue )
+static bool DoesConfigPropertyValueMatch(const FConfigSection* InSection, const FName& InPropertyName, const FString& InPropertyValue)
 {
 	bool bFoundAMatch = false;
 
@@ -1397,18 +1443,18 @@ bool DoesConfigPropertyValueMatch(const FConfigSection* InSection, const FName&
 		// Start Array check, if the property is in an array, we need to iterate over all properties.
 		for (FConfigSection::TConstKeyIterator It(*InSection, InPropertyName); It && !bFoundAMatch; ++It)
 		{
-			const FString& PropertyValue = It.Value().GetSavedValue();
-			bFoundAMatch = 
+			const FString& PropertyValue = UE::ConfigCacheIni::Private::FAccessor::GetSavedValueForWriting(It.Value());
+			bFoundAMatch =
 				PropertyValue.Len() == InPropertyValue.Len() &&
 				PropertyValue == InPropertyValue;
-				
+
 			// if our properties don't match, run further checks
-			if( !bFoundAMatch )
+			if (!bFoundAMatch)
 			{
 				// Check that the mismatch isn't just a string comparison issue with floats
-				if (bIsInputStringValidFloat && FDefaultValueHelper::IsStringValidFloat( PropertyValue ))
+				if (bIsInputStringValidFloat && FDefaultValueHelper::IsStringValidFloat(PropertyValue))
 				{
-					bFoundAMatch = FCString::Atof( *PropertyValue ) == FCString::Atof( *InPropertyValue );
+					bFoundAMatch = FCString::Atof(*PropertyValue) == FCString::Atof(*InPropertyValue);
 				}
 			}
 		}
@@ -1418,6 +1464,9 @@ bool DoesConfigPropertyValueMatch(const FConfigSection* InSection, const FName&
 	return bFoundAMatch;
 }
 
+}; // struct FImpl
+
+} // namespace UE::ConfigCacheIni::Private
 
 /**
  * Check if the provided property information was set as a commandline override
@@ -1584,7 +1633,16 @@ void FConfigFile::WriteToStringInternal(FString& InOutText, bool bIsADefaultIniW
 	TArray<FString> SectionsToSave;
 	if (SectionsToSaveSection != nullptr)
 	{
-		SectionsToSaveSection->MultiFind("Section", SectionsToSave);
+		// Do not report the read of SectionsToSave. Some ConfigFiles are reallocated without it, and we
+		// log that the section disappeared. But this log is spurious since the only reason it was read was
+		// for the internal save before the FConfigFile is made publicly available.
+		TArray<const FConfigValue*, TInlineAllocator<10>> SectionsToSaveValues;
+		SectionsToSaveSection->MultiFindPointer("Section", SectionsToSaveValues);
+		SectionsToSave.Reserve(SectionsToSaveValues.Num());
+		for (const FConfigValue* ConfigValue : SectionsToSaveValues)
+		{
+			SectionsToSave.Add(UE::ConfigCacheIni::Private::FAccessor::GetValueForWriting(*ConfigValue));
+		}
 	}
 	
 	for( TIterator SectionIterator(*this); SectionIterator; ++SectionIterator )
@@ -1638,7 +1696,9 @@ void FConfigFile::WriteToStringInternal(FString& InOutText, bool bIsADefaultIniW
 				const bool bIsCurrentIniVersion = (SectionName == CurrentIniVersionStr);
 
 				// Check if the property matches the source configs. We do not wanna write it out if so.
-				if ((bIsADefaultIniWrite || bIsCurrentIniVersion || !DoesConfigPropertyValueMatch(SourceConfigSection, PropertyName, PropertyValue)) && !bOptionIsFromCommandline)
+				if ((bIsADefaultIniWrite || bIsCurrentIniVersion ||
+					!UE::ConfigCacheIni::Private::FImpl::DoesConfigPropertyValueMatch(SourceConfigSection, PropertyName, PropertyValue))
+					&& !bOptionIsFromCommandline)
 				{
 					// If this is the first property we are writing of this section, then print the section name
 					if( InOutText.Len() == InitialInOutTextSize )
@@ -1934,14 +1994,14 @@ void FConfigFile::SetString( const TCHAR* Section, const TCHAR* Key, const TCHAR
 	FConfigValue* ConfigValue = Sec->Find( Key );
 	if( ConfigValue == nullptr )
 	{
-		Sec->Add( Key, Value );
+		Sec->Add(Key, FConfigValue(Value));
 		Dirty = true;
 	}
-	// Use GetSavedValueForWriting rather than GetSavedValue to avoid having the is-it-dirty query mark the values as having been accessed for dependency tracking
+	// Use GetSavedValueForWriting rather than GetSavedValue to avoid reporting the value as having been accessed for dependency tracking
 	else if( FCString::Strcmp(*UE::ConfigCacheIni::Private::FAccessor::GetSavedValueForWriting(*ConfigValue),Value)!=0 )
 	{
 		Dirty = true;
-		*ConfigValue = FConfigValue(Value);
+		*ConfigValue = Value;
 	}
 }
 
@@ -1955,14 +2015,14 @@ void FConfigFile::SetText( const TCHAR* Section, const TCHAR* Key, const FText&
 	FConfigValue* ConfigValue = Sec->Find( Key );
 	if( ConfigValue == nullptr )
 	{
-		Sec->Add( Key, StrValue );
+		Sec->Add(Key, FConfigValue(MoveTemp(StrValue)));
 		Dirty = true;
 	}
-	// Use GetSavedValueForWriting rather than GetSavedValue to avoid having the is-it-dirty query mark the values as having been accessed for dependency tracking
+	// Use GetSavedValueForWriting rather than GetSavedValue to avoid reporting the value as having been accessed for dependency tracking
 	else if( FCString::Strcmp(*UE::ConfigCacheIni::Private::FAccessor::GetSavedValueForWriting(*ConfigValue), *StrValue)!=0 )
 	{
 		Dirty = true;
-		*ConfigValue = FConfigValue(StrValue);
+		*ConfigValue = MoveTemp(StrValue);
 	}
 }
 
@@ -2004,7 +2064,7 @@ void FConfigFile::SetArray(const TCHAR* Section, const TCHAR* Key, const TArray<
 
 	for (int32 i = 0; i < Value.Num(); i++)
 	{
-		Sec->Add(Key, *Value[i]);
+		Sec->Add(Key, FConfigValue(Value[i]));
 		Dirty = true;
 	}
 }
@@ -2020,7 +2080,7 @@ bool FConfigFile::AddToSection(const TCHAR* SectionName, FName Key, const FStrin
 bool FConfigFile::AddUniqueToSection(const TCHAR* SectionName, FName Key, const FString& Value)
 {
 	FConfigSection* Section = FindOrAddSectionInternal(SectionName);
-	if (Section->FindPair(Key, Value))
+	if (Section->FindPair(Key, FConfigValue(Value)))
 	{
 		return false;
 	}
@@ -2049,13 +2109,13 @@ bool FConfigFile::RemoveFromSection(const TCHAR* SectionName, FName Key, const F
 {
 	FConfigSection* Section = FindInternal(SectionName);
 	// if it doesn't contain the pair, do nothing
-	if (Section == nullptr || !Section->FindPair(Key, Value))
+	if (Section == nullptr || !Section->FindPair(Key, FConfigValue(Value)))
 	{
 		return false;
 	}
 
 	// remove any copies of the pair
-	Section->Remove(Key, Value);
+	Section->Remove(Key, FConfigValue(Value));
 	Dirty = true;
 	return true;
 }
@@ -2101,13 +2161,13 @@ void FConfigFile::ProcessSourceAndCheckAgainstBackup()
 		FConfigFile BackupFile;
 		ProcessIniContents(*BackupFilename, *BackupFilename, &BackupFile, false, false);
 
-		for( TMap<FString,FConfigSection>::TIterator SectionIterator(*SourceConfigFile); SectionIterator; ++SectionIterator )
+		for (TMap<FString,FConfigSection>::TIterator SectionIterator(*SourceConfigFile); SectionIterator; ++SectionIterator)
 		{
 			const FString& SectionName = SectionIterator.Key();
 			const FConfigSection& SourceSection = SectionIterator.Value();
 			const FConfigSection* BackupSection = BackupFile.FindSection( SectionName );
-
-			if( BackupSection && SourceSection != *BackupSection )
+			
+			if (BackupSection && !UE::ConfigCacheIni::Private::FAccessor::AreSectionsEqualForWriting(SourceSection, *BackupSection))
 			{
 				this->Remove( SectionName );
 				this->Add( SectionName, SourceSection );
@@ -2313,6 +2373,17 @@ FConfigFile* FConfigCacheIni::FindConfigFileWithBaseName(FName BaseName)
 	return nullptr;
 }
 
+FConfigFile& FConfigCacheIni::Add(const FString& Filename, const FConfigFile& File)
+{
+	FConfigFile*& Result = OtherFiles.FindOrAdd(Filename);
+	if (Result)
+	{
+		delete Result;
+	}
+	Result = new FConfigFile(File);
+	return *Result;
+}
+
 bool FConfigCacheIni::ContainsConfigFile(const FConfigFile* ConfigFile) const
 {
 	// Check the normal inis. Note that the FConfigFiles in the map
@@ -2776,14 +2847,14 @@ void FConfigCacheIni::SetText( const TCHAR* Section, const TCHAR* Key, const FTe
 	FConfigValue* ConfigValue = Sec->Find( Key );
 	if( !ConfigValue )
 	{
-		Sec->Add( Key, StrValue );
+		Sec->Add(Key, FConfigValue(MoveTemp(StrValue)));
 		File->Dirty = true;
 	}
-	// Use GetSavedValueForWriting rather than GetSavedValue to avoid having the is-it-dirty query mark the values as having been accessed for dependency tracking
+	// Use GetSavedValueForWriting rather than GetSavedValue to avoid reporting the value as having been accessed for dependency tracking
 	else if( FCString::Strcmp(*UE::ConfigCacheIni::Private::FAccessor::GetSavedValueForWriting(*ConfigValue), *StrValue)!=0 )
 	{
 		File->Dirty = true;
-		*ConfigValue = FConfigValue(StrValue);
+		*ConfigValue = MoveTemp(StrValue);
 	}
 }
 
@@ -3931,7 +4002,7 @@ static void InitializeConfigRemap()
 		}
 		
 		Context.Load(*FPaths::Combine(Pass == 0 ? FPaths::EngineDir() : FPaths::ProjectDir(), TEXT("Config/ConfigRedirects.ini")));
-		
+
 		for (const TPair<FString, FConfigSection>& Section : AsConst(RemapFile))
 		{
 			if (Section.Key == TEXT("SectionNameRemap"))
@@ -4239,11 +4310,15 @@ FString FConfigCacheIni::NormalizeConfigIniPath(const FString& NonNormalizedPath
 FArchive& operator<<(FArchive& Ar, FConfigFile& ConfigFile)
 {
 	bool bHasSourceConfigFile = ConfigFile.SourceConfigFile != nullptr;
+	bool bDirty = ConfigFile.Dirty;
+	bool bNoSave = ConfigFile.NoSave;
+	bool bHasPlatformName = ConfigFile.bHasPlatformName;
 
 	Ar << static_cast<FConfigFile::Super&>(ConfigFile);
-	Ar << ConfigFile.Dirty;
-	Ar << ConfigFile.NoSave;
-	Ar << ConfigFile.bHasPlatformName;
+	Ar << bDirty;
+	Ar << bNoSave;
+	Ar << bHasPlatformName;
+
 	Ar << ConfigFile.Name;
 	Ar << ConfigFile.SourceIniHierarchy;
 	Ar << ConfigFile.SourceEngineConfigDir;
@@ -4262,6 +4337,13 @@ FArchive& operator<<(FArchive& Ar, FConfigFile& ConfigFile)
 	Ar << ConfigFile.PlatformName;
 	Ar << ConfigFile.PerObjectConfigArrayOfStructKeys;
 
+	if (Ar.IsLoading())
+	{
+		ConfigFile.Dirty = bDirty;
+		ConfigFile.NoSave = bNoSave;
+		ConfigFile.bHasPlatformName = bHasPlatformName;
+	}
+
 	return Ar;
 }
 
diff --git a/Engine/Source/Runtime/Core/Private/Misc/ConfigContext.cpp b/Engine/Source/Runtime/Core/Private/Misc/ConfigContext.cpp
--- a/Engine/Source/Runtime/Core/Private/Misc/ConfigContext.cpp
+++ b/Engine/Source/Runtime/Core/Private/Misc/ConfigContext.cpp
@@ -400,9 +400,21 @@ bool FConfigContext::PerformLoad()
 		ConfigFile->PlatformName = Platform;
 		ConfigFile->bHasPlatformName = true;
 
-		// chcek if the config file wants to save all sections
+		// check if the config file wants to save all sections
 		bool bLocalSaveAllSections = false;
-		ConfigFile->GetBool(SectionsToSaveString, SaveAllSectionsKey, bLocalSaveAllSections);
+		// Do not report the read of SectionsToSave. Some ConfigFiles are reallocated without it, and reporting
+		// logs that the section disappeared. But this log is spurious since if the only reason it was read was
+		// for the internal save before the FConfigFile is made publicly available.
+		const FConfigSection* SectionsToSaveSection = ConfigFile->FindSection(SectionsToSaveString);
+		if (SectionsToSaveSection)
+		{
+			const FConfigValue* Value = SectionsToSaveSection->Find(SaveAllSectionsKey);
+			if (Value)
+			{
+				const FString& ValueStr = UE::ConfigCacheIni::Private::FAccessor::GetValueForWriting(*Value);
+				bLocalSaveAllSections = FCString::ToBool(*ValueStr);
+			}
+		}
 
 		// we can always save all sections of a User config file, Editor* (not Editor.ini tho, that is already handled in the normal method)
 		bool bIsUserFile = BaseIniName.Contains(TEXT("User"));
diff --git a/Engine/Source/Runtime/Core/Private/Misc/DataDrivenPlatformInfoRegistry.cpp b/Engine/Source/Runtime/Core/Private/Misc/DataDrivenPlatformInfoRegistry.cpp
--- a/Engine/Source/Runtime/Core/Private/Misc/DataDrivenPlatformInfoRegistry.cpp
+++ b/Engine/Source/Runtime/Core/Private/Misc/DataDrivenPlatformInfoRegistry.cpp
@@ -240,7 +240,8 @@ static void DDPIGetStringArray(const FConfigFile& IniFile, const TCHAR* Key, TAr
 // Gets a string from a section, or empty string if it didn't exist
 static FString GetSectionString(const FConfigSection& Section, FName Key)
 {
-	return Section.FindRef(Key).GetValue();
+	const FConfigValue* Value = Section.Find(Key);
+	return Value ? Value->GetValue() : FString();
 }
 
 #if DDPI_HAS_EXTENDED_PLATFORMINFO_DATA
diff --git a/Engine/Source/Runtime/Core/Public/Misc/ConfigCacheIni.h b/Engine/Source/Runtime/Core/Public/Misc/ConfigCacheIni.h
--- a/Engine/Source/Runtime/Core/Public/Misc/ConfigCacheIni.h
+++ b/Engine/Source/Runtime/Core/Public/Misc/ConfigCacheIni.h
@@ -112,16 +112,8 @@ enum class EKnownIniFile : uint8
 	NumKnownFiles,
 };
 
-namespace UE
-{
-namespace ConfigCacheIni
-{
-namespace Private
-{
-struct FAccessor;
-}
-}
-}
+namespace UE::ConfigCacheIni::Private { struct FAccessor; }
+namespace UE::ConfigCacheIni::Private { struct FImpl; }
 
 class FConfigContext;
 
@@ -206,6 +198,24 @@ public:
 		return *this;
 	}
 
+	FConfigValue& operator=(const TCHAR* RHS)
+	{
+		*this = FString(RHS);
+		return *this;
+	}
+	FConfigValue& operator=(const FString& RHS)
+	{
+		*this = FString(RHS);
+		return *this;
+	}
+	FConfigValue& operator=(FString&& RHS)
+	{
+		SavedValue = MoveTemp(RHS);
+		SavedValueHash = FTextLocalizationResource::HashString(SavedValue);
+		ExpandValueInternal();
+		return *this;
+	}
+
 	// Returns the ini setting with any macros expanded out
 	const FString& GetValue() const 
 	{
@@ -300,8 +310,14 @@ private:
 	/** Internal version of ExpandValue that expands SavedValue into ExpandedValue, or produces an empty ExpandedValue if no expansion occurred. */
 	CORE_API void ExpandValueInternal();
 
-	/** Gets the SavedValue without marking it as having been accessed for e.g. writing out to a ConfigFile to disk */
+	/** Gets the expanded value (GetValue) without marking it as having been accessed for e.g. writing out to a ConfigFile to disk */
 	friend struct UE::ConfigCacheIni::Private::FAccessor;
+	const FString& GetValueForWriting() const
+	{
+		return (ExpandedValue.Len() > 0 ? ExpandedValue : SavedValue);
+	};
+
+	/** Gets the SavedValue without marking it as having been accessed for e.g. writing out to a ConfigFile to disk */
 	const FString& GetSavedValueForWriting() const
 	{
 		return SavedValue;
@@ -315,28 +331,6 @@ private:
 #endif
 };
 
-namespace UE
-{
-namespace ConfigCacheIni
-{
-namespace Private
-{
-/** An accessor class to access functions that should be restricted only to FConfigFileCache Internal use */
-struct FAccessor
-{
-private:
-	friend class ::FConfigCacheIni;
-	friend class ::FConfigFile;
-
-	static const FString& GetSavedValueForWriting(const FConfigValue& ConfigValue)
-	{
-		return ConfigValue.GetSavedValueForWriting();
-	}
-};
-}
-}
-}
-
 typedef TMultiMap<FName,FConfigValue> FConfigSectionMap;
 
 // One section in a config file.
@@ -355,7 +349,7 @@ public:
 	bool operator!=( const FConfigSection& Other ) const;
 
 	// process the '+' and '.' commands, takingf into account ArrayOfStruct unique keys
-	void CORE_API HandleAddCommand(FName Key, FString&& Value, bool bAppendValueIfNotArrayOfStructsKeyUsed);
+	void CORE_API HandleAddCommand(FName ValueName, FString&& Value, bool bAppendValueIfNotArrayOfStructsKeyUsed);
 
 	bool HandleArrayOfKeyedStructsCommand(FName Key, FString&& Value);
 
@@ -382,9 +376,35 @@ public:
 	// look for "array of struct" keys for overwriting single entries of an array
 	TMap<FName, FString> ArrayOfStructKeys;
 
+	friend struct UE::ConfigCacheIni::Private::FAccessor;
 	friend FArchive& operator<<(FArchive& Ar, FConfigSection& ConfigSection);
 };
 
+namespace UE::ConfigCacheIni::Private
+{
+
+/** An accessor class to access functions that should be restricted only to FConfigFileCache Internal use */
+struct FAccessor
+{
+private:
+	friend class ::FConfigCacheIni;
+	friend class ::FConfigFile;
+	friend class ::FConfigSection;
+	friend class ::FConfigContext;
+	friend struct ::UE::ConfigCacheIni::Private::FImpl;
+
+	static const FString& GetValueForWriting(const FConfigValue& ConfigValue)
+	{
+		return ConfigValue.GetValueForWriting();
+	}
+	static const FString& GetSavedValueForWriting(const FConfigValue& ConfigValue)
+	{
+		return ConfigValue.GetSavedValueForWriting();
+	}
+	static bool AreSectionsEqualForWriting(const FConfigSection& A, const FConfigSection& B);
+};
+
+}
 
 #if ALLOW_INI_OVERRIDE_FROM_COMMANDLINE
 // Options which stemmed from the commandline
@@ -394,20 +414,20 @@ struct FConfigCommandlineOverride
 };
 #endif // ALLOW_INI_OVERRIDE_FROM_COMMANDLINE
 
+typedef TMap<FString, FConfigSection> FConfigFileMap;
 
 // One config file.
-typedef TMap<FString, FConfigSection> FConfigFileMap;
 class FConfigFile : private FConfigFileMap
 {
 public:
-	bool Dirty;
-	bool NoSave;
-	bool bHasPlatformName = false;
+	bool Dirty : 1; // = false;
+	bool NoSave : 1; // = false;
+	bool bHasPlatformName : 1; // = false;
 	// by default, we allow saving - this is going to be applied to config files that are not loaded from disk
 	// (when loading, this will get set to false, and then the ini sections will be checked)
-	bool bCanSaveAllSections = true;
+	bool bCanSaveAllSections : 1; // = true;
 
-	/** The name of this config file */	
+	/** The name of this config file */
 	FName Name;
 
 	// The collection of source files which were used to generate this file.
@@ -426,11 +446,17 @@ public:
 	/** The collection of overrides which stemmed from the commandline */
 	TArray<FConfigCommandlineOverride> CommandlineOptions;
 #endif // ALLOW_INI_OVERRIDE_FROM_COMMANDLINE
-	
+
+private:
+	// This holds per-object config class names, with their ArrayOfStructKeys. Since the POC sections are all unique,
+	// we can't track it just in that section. This is expected to be empty/small
+	TMap<FString, TMap<FName, FString> > PerObjectConfigArrayOfStructKeys;
+
+public:
 	CORE_API FConfigFile();
 	FConfigFile( int32 ) {}	// @todo UE-DLL: Workaround for instantiated TMap template during DLLExport (TMap::FindRef)
 	CORE_API ~FConfigFile();
-	
+
 	// looks for a section by name, and creates an empty one if it can't be found
 	UE_DEPRECATED(5.4, "Use FindOrAddConfigSection, and/or use the new AddToSection, etc APIs to modify sections without retrieving the section. See top of ConfigCacheIni.h for more info.")
 	CORE_API FConfigSection* FindOrAddSection(const FString& Name);
@@ -696,11 +722,6 @@ public:
 
 	friend FArchive& operator<<(FArchive& Ar, FConfigFile& ConfigFile);
 private:
-
-	// This holds per-object config class names, with their ArrayOfStructKeys. Since the POC sections are all unique,
-	// we can't track it just in that section. This is expected to be empty/small
-	TMap<FString, TMap<FName, FString> > PerObjectConfigArrayOfStructKeys;
-
 	/** 
 	 * Save the source hierarchy which was loaded out to a backup file so we can check future changes in the base/default configs
 	 */
@@ -852,10 +873,8 @@ public:
 	/** Finds Config file that matches the base name such as "Engine" */
 	CORE_API FConfigFile* FindConfigFileWithBaseName(FName BaseName);
 
-	FConfigFile& Add(const FString& Filename, const FConfigFile& File)
-	{
-		return *OtherFiles.Add(Filename, new FConfigFile(File));
-	}
+	CORE_API FConfigFile& Add(const FString& Filename, const FConfigFile& File);
+
 	int32 Remove(const FString& Filename)
 	{
 		delete OtherFiles.FindRef(Filename);
diff --git a/Engine/Source/Runtime/RHI/Private/DataDrivenShaderPlatformInfo.cpp b/Engine/Source/Runtime/RHI/Private/DataDrivenShaderPlatformInfo.cpp
--- a/Engine/Source/Runtime/RHI/Private/DataDrivenShaderPlatformInfo.cpp
+++ b/Engine/Source/Runtime/RHI/Private/DataDrivenShaderPlatformInfo.cpp
@@ -48,7 +48,8 @@ EShaderPlatform ParseShaderPlatform(const TCHAR* String)
 // Gets a string from a section, or empty string if it didn't exist
 static inline FString GetSectionString(const FConfigSection& Section, FName Key)
 {
-	return Section.FindRef(Key).GetValue();
+	const FConfigValue* Value = Section.Find(Key);
+	return Value ? Value->GetValue() : FString();
 }
 
 // Gets a bool from a section.  It returns the original value if the setting does not exist
