UE Units: Add acceleration (cm/s2 and m/s2)
#rnx
#tests Tested in editor, preflight

diff --git a/Engine/Config/BaseEditor.ini b/Engine/Config/BaseEditor.ini
--- a/Engine/Config/BaseEditor.ini
+++ b/Engine/Config/BaseEditor.ini
@@ -532,6 +532,7 @@ bDisplayUnits=True
 AngleUnits=EUnit::Degrees
 SpeedUnits=EUnit::MetersPerSecond
 AngularSpeedUnits=EUnit::DegreesPerSecond
+AccelerationUnits=EUnit::CentimetersPerSecondSquared
 TemperatureUnits=EUnit::Celsius
 ForceUnits=EUnit::Newtons
 TorqueUnits=EUnit::NewtonMeters
diff --git a/Engine/Programs/UnrealHeaderTool/Config/DefaultEngine.ini b/Engine/Programs/UnrealHeaderTool/Config/DefaultEngine.ini
--- a/Engine/Programs/UnrealHeaderTool/Config/DefaultEngine.ini
+++ b/Engine/Programs/UnrealHeaderTool/Config/DefaultEngine.ini
@@ -70,6 +70,12 @@ gc.MaxObjectsInProgram=262144
 +Units="deg/s"
 +Units="RadiansPerSecond"
 +Units="rad/s"
++Units="CentimetersPerSecondSquared"
++Units="cm/s\u00B2"
++Units="cm/s^2"
++Units="MetersPerSecondSquared"
++Units="m/s\u00B2"
++Units="m/s^2"
 +Units="Celsius"
 +Units="C"
 +Units="degC"
diff --git a/Engine/Source/Editor/UnrealEd/Private/Settings/EditorProjectSettings.cpp b/Engine/Source/Editor/UnrealEd/Private/Settings/EditorProjectSettings.cpp
--- a/Engine/Source/Editor/UnrealEd/Private/Settings/EditorProjectSettings.cpp
+++ b/Engine/Source/Editor/UnrealEd/Private/Settings/EditorProjectSettings.cpp
@@ -109,6 +109,10 @@ void UEditorProjectAppearanceSettings::PostEditChangeProperty( struct FPropertyC
 	{
 		Settings.SetDisplayUnits(EUnitType::AngularSpeed, AngularSpeedUnits);
 	}
+	else if (Name == GET_MEMBER_NAME_CHECKED(UEditorProjectAppearanceSettings, AccelerationUnits))
+	{
+		Settings.SetDisplayUnits(EUnitType::Acceleration, AccelerationUnits);
+	}
 	else if (Name == GET_MEMBER_NAME_CHECKED(UEditorProjectAppearanceSettings, TemperatureUnits))
 	{
 		Settings.SetDisplayUnits(EUnitType::Temperature, TemperatureUnits);
diff --git a/Engine/Source/Editor/UnrealEd/Public/Settings/EditorProjectSettings.h b/Engine/Source/Editor/UnrealEd/Public/Settings/EditorProjectSettings.h
--- a/Engine/Source/Editor/UnrealEd/Public/Settings/EditorProjectSettings.h
+++ b/Engine/Source/Editor/UnrealEd/Public/Settings/EditorProjectSettings.h
@@ -91,6 +91,9 @@ public:
 
 	UPROPERTY(EditAnywhere, config, Category=Units, AdvancedDisplay, meta=(DisplayName="Angular Speed", Tooltip="Choose the units in which to display angular speeds.", ValidEnumValues="DegreesPerSecond, RadiansPerSecond"))
 	EUnit AngularSpeedUnits;
+	
+	UPROPERTY(EditAnywhere, config, Category=Units, AdvancedDisplay, meta=(DisplayName="Acceleration", Tooltip="Choose the units in which to display acceleration.", ValidEnumValues="CentimetersPerSecondSquared, MetersPerSecondSquared"))
+	EUnit AccelerationUnits;
 
 	UPROPERTY(EditAnywhere, config, Category=Units, AdvancedDisplay, meta=(DisplayName="Temperature", Tooltip="Choose the units in which to display temperatures.", ValidEnumValues="Celsius, Farenheit, Kelvin"))
 	EUnit TemperatureUnits;
diff --git a/Engine/Source/Runtime/Core/Private/Math/UnitConversion.cpp b/Engine/Source/Runtime/Core/Private/Math/UnitConversion.cpp
--- a/Engine/Source/Runtime/Core/Private/Math/UnitConversion.cpp
+++ b/Engine/Source/Runtime/Core/Private/Math/UnitConversion.cpp
@@ -50,6 +50,9 @@ constexpr FParseCandidate ParseCandidates[] = {
 
 	{ TEXT("DegreesPerSecond"),		EUnit::DegreesPerSecond },		{ TEXT("deg/s"),	EUnit::DegreesPerSecond },
 	{ TEXT("RadiansPerSecond"),		EUnit::RadiansPerSecond },		{ TEXT("rad/s"),	EUnit::RadiansPerSecond },
+
+	{ TEXT("CentimetersPerSecondSquared"),		EUnit::CentimetersPerSecondSquared },	{ TEXT("cm/s2"),	EUnit::CentimetersPerSecondSquared },{ TEXT("cm/s\u00B2"),	EUnit::CentimetersPerSecondSquared },{ TEXT("cm/s^2"),	EUnit::CentimetersPerSecondSquared },
+	{ TEXT("MetersPerSecondSquared"),			EUnit::MetersPerSecondSquared },		{ TEXT("m/s2"),	EUnit::MetersPerSecondSquared },	{ TEXT("m/s\u00B2"),	EUnit::MetersPerSecondSquared },{ TEXT("m/s^2"),	EUnit::MetersPerSecondSquared },
 		
 	{ TEXT("Celsius"),				EUnit::Celsius },				{ TEXT("C"),		EUnit::Celsius },				{ TEXT("degC"),		EUnit::Celsius },			{ TEXT("\u00B0C"),		EUnit::Celsius },
 	{ TEXT("Farenheit"),			EUnit::Farenheit },				{ TEXT("F"),		EUnit::Farenheit },				{ TEXT("degF"),		EUnit::Farenheit },			{ TEXT("\u00B0F"),		EUnit::Farenheit },
@@ -129,6 +132,8 @@ constexpr const TCHAR* const DisplayStrings[] = {
 
 	TEXT("deg/s"), TEXT("rad/s"),
 
+	TEXT("cm/s\u00B2"), TEXT("m/s\u00B2"),
+
 	TEXT("\u00B0C"), TEXT("\u00B0F"), TEXT("K"),
 
 	TEXT("\u00B5g"), TEXT("mg"), TEXT("g"), TEXT("kg"), TEXT("t"),
@@ -178,7 +183,10 @@ constexpr const TCHAR* const SupportedUnitsStrings[] = {
 	TEXT("MilesPerHour"),			
 
 	TEXT("DegreesPerSecond"),		
-	TEXT("RadiansPerSecond"),		
+	TEXT("RadiansPerSecond"),
+	
+	TEXT("CentimetersPerSecondSquared"),		
+	TEXT("MetersPerSecondSquared"),	
 		
 	TEXT("Celsius"),			
 	TEXT("Farenheit"),			
@@ -260,6 +268,8 @@ constexpr EUnitType UnitTypes[] = {
 	EUnitType::Speed,		EUnitType::Speed,		EUnitType::Speed, 		EUnitType::Speed,
 
 	EUnitType::AngularSpeed, EUnitType::AngularSpeed,
+	
+	EUnitType::Acceleration, EUnitType::Acceleration,
 
 	EUnitType::Temperature,	EUnitType::Temperature,	EUnitType::Temperature,
 
@@ -550,6 +560,7 @@ FUnitSettings::FUnitSettings()
 	DisplayUnits[(uint8)EUnitType::Angle].Add(EUnit::Degrees);
 	DisplayUnits[(uint8)EUnitType::Speed].Add(EUnit::MetersPerSecond);
 	DisplayUnits[(uint8)EUnitType::AngularSpeed].Add(EUnit::DegreesPerSecond);
+	DisplayUnits[(uint8)EUnitType::Acceleration].Add(EUnit::CentimetersPerSecondSquared);
 	DisplayUnits[(uint8)EUnitType::Temperature].Add(EUnit::Celsius);
 	DisplayUnits[(uint8)EUnitType::Mass].Add(EUnit::Kilograms);
 	DisplayUnits[(uint8)EUnitType::Density].Add(EUnit::GramsPerCubicCentimeter);
@@ -732,6 +743,16 @@ namespace UnitConversion
 			default: 						return 1.0;
 		}
 	}
+	
+	double AccelerationUnificationFactor(EUnit From)
+	{
+		// Convert to meters/second2
+		switch (From)
+		{
+		case EUnit::CentimetersPerSecondSquared:	return 0.01f;
+		default: 									return 1.0;
+		}
+	}
 
 	double MassUnificationFactor(EUnit From)
 	{
diff --git a/Engine/Source/Runtime/Core/Public/Math/UnitConversion.h b/Engine/Source/Runtime/Core/Public/Math/UnitConversion.h
--- a/Engine/Source/Runtime/Core/Public/Math/UnitConversion.h
+++ b/Engine/Source/Runtime/Core/Public/Math/UnitConversion.h
@@ -30,6 +30,9 @@ enum class EUnit : uint8
 
 	/** Angular speed units */
 	DegreesPerSecond, RadiansPerSecond,
+	
+	/** Acceleration units */
+	CentimetersPerSecondSquared, MetersPerSecondSquared,
 
 	/** Temperature units */
 	Celsius, Farenheit, Kelvin,
@@ -75,7 +78,7 @@ enum class EUnit : uint8
 /** Enumeration that specifies particular classes of unit */
 enum class EUnitType
 {
-	Distance, Angle, Speed, AngularSpeed, Temperature, Mass, Density, Force, Torque, Frequency, DataSize, LuminousFlux, LuminousIntensity, Illuminance, Luminance, Time, PixelDensity, Multipliers, ExposureValue, Stress,
+	Distance, Angle, Speed, AngularSpeed, Acceleration, Temperature, Mass, Density, Force, Torque, Frequency, DataSize, LuminousFlux, LuminousIntensity, Illuminance, Luminance, Time, PixelDensity, Multipliers, ExposureValue, Stress,
 
 	// Symbolic entry - do not use directly
 	NumberOf,
diff --git a/Engine/Source/Runtime/Core/Public/Math/UnitConversion.inl b/Engine/Source/Runtime/Core/Public/Math/UnitConversion.inl
--- a/Engine/Source/Runtime/Core/Public/Math/UnitConversion.inl
+++ b/Engine/Source/Runtime/Core/Public/Math/UnitConversion.inl
@@ -26,6 +26,8 @@ namespace UnitConversion
 	CORE_API double SpeedUnificationFactor(EUnit From);
 	/** Find the common quantization factor for the specified angular speed unit. Quantizes to DegreesPerSecond. */
 	CORE_API double AngularSpeedUnificationFactor(EUnit From);
+	/** Find the common quantization factor for the specified acceleration unit. Quantizes to m/s2. */
+	CORE_API double AccelerationUnificationFactor(EUnit From);
 	/** Find the common quantization factor for the specified temperature unit. Quantizes to Kelvin. */
 	CORE_API double TemperatureUnificationFactor(EUnit From);
 	/** Find the common quantization factor for the specified mass unit. Quantizes to Grams. */
@@ -84,20 +86,21 @@ T FUnitConversion::Convert(T InValue, EUnit From, EUnit To)
 
 	switch(FUnitConversion::GetUnitType(From))
 	{
-		case EUnitType::Distance:			return InValue * DistanceUnificationFactor(From)		* (1.0 / DistanceUnificationFactor(To));
-		case EUnitType::Angle:				return InValue * AngleUnificationFactor(From) 			* (1.0 / AngleUnificationFactor(To));
-		case EUnitType::Speed:				return InValue * SpeedUnificationFactor(From) 			* (1.0 / SpeedUnificationFactor(To));
-		case EUnitType::AngularSpeed:		return InValue * AngularSpeedUnificationFactor(From)	* (1.0 / AngularSpeedUnificationFactor(To));
-		case EUnitType::Mass:				return InValue * MassUnificationFactor(From) 			* (1.0 / MassUnificationFactor(To));
-		case EUnitType::Density:			return InValue * DensityUnificationFactor(From)			* (1.0 / DensityUnificationFactor(To));
-		case EUnitType::Force:				return InValue * ForceUnificationFactor(From) 			* (1.0 / ForceUnificationFactor(To));
-		case EUnitType::Torque:				return InValue * TorqueUnificationFactor(From)			* (1.0 / TorqueUnificationFactor(To));
-		case EUnitType::Frequency:			return InValue * FrequencyUnificationFactor(From) 		* (1.0 / FrequencyUnificationFactor(To));
-		case EUnitType::DataSize:			return InValue * DataSizeUnificationFactor(From) 		* (1.0 / DataSizeUnificationFactor(To));
-		case EUnitType::LuminousFlux:		return InValue;
-		case EUnitType::Time:				return InValue * TimeUnificationFactor(From) 			* (1.0 / TimeUnificationFactor(To));
-		case EUnitType::Multipliers:		return InValue * MultiplierUnificationFactor(From) 		* (1.0 / MultiplierUnificationFactor(To));
-		case EUnitType::Stress:				return InValue * StressUnificationFactor(From)			* (1.0 / StressUnificationFactor(To));
+		case EUnitType::Distance:			return InValue * DistanceUnificationFactor(From)			* (1.0 / DistanceUnificationFactor(To));
+		case EUnitType::Angle:				return InValue * AngleUnificationFactor(From) 				* (1.0 / AngleUnificationFactor(To));
+		case EUnitType::Speed:				return InValue * SpeedUnificationFactor(From) 				* (1.0 / SpeedUnificationFactor(To));
+		case EUnitType::AngularSpeed:		return InValue * AngularSpeedUnificationFactor(From)		* (1.0 / AngularSpeedUnificationFactor(To));
+		case EUnitType::Acceleration:		return InValue * AccelerationUnificationFactor(From)		* (1.0 / AccelerationUnificationFactor(To));
+		case EUnitType::Mass:				return InValue * MassUnificationFactor(From) 				* (1.0 / MassUnificationFactor(To));
+		case EUnitType::Density:			return InValue * DensityUnificationFactor(From)				* (1.0 / DensityUnificationFactor(To));
+		case EUnitType::Force:				return InValue * ForceUnificationFactor(From) 				* (1.0 / ForceUnificationFactor(To));
+		case EUnitType::Torque:				return InValue * TorqueUnificationFactor(From)				* (1.0 / TorqueUnificationFactor(To));
+		case EUnitType::Frequency:			return InValue * FrequencyUnificationFactor(From) 			* (1.0 / FrequencyUnificationFactor(To));
+		case EUnitType::DataSize:			return InValue * DataSizeUnificationFactor(From) 			* (1.0 / DataSizeUnificationFactor(To));
+		case EUnitType::LuminousFlux:		return InValue;	
+		case EUnitType::Time:				return InValue * TimeUnificationFactor(From) 				* (1.0 / TimeUnificationFactor(To));
+		case EUnitType::Multipliers:		return InValue * MultiplierUnificationFactor(From) 			* (1.0 / MultiplierUnificationFactor(To));
+		case EUnitType::Stress:				return InValue * StressUnificationFactor(From)				* (1.0 / StressUnificationFactor(To));
 		// Temperature conversion is not just a simple multiplication, so needs special treatment
 		case EUnitType::Temperature:
 		{
diff --git a/Engine/Source/Runtime/Core/Tests/Math/UnitConversionTests.cpp b/Engine/Source/Runtime/Core/Tests/Math/UnitConversionTests.cpp
--- a/Engine/Source/Runtime/Core/Tests/Math/UnitConversionTests.cpp
+++ b/Engine/Source/Runtime/Core/Tests/Math/UnitConversionTests.cpp
@@ -47,6 +47,9 @@ TEST_CASE_NAMED(FUnitUnitTests, "System::Core::Math::Unit Conversion", "[Applica
 		{ 90.0,			UE_DOUBLE_HALF_PI,	1e-3,	EUnit::DegreesPerSecond, 				EUnit::RadiansPerSecond								},
 		{ UE_DOUBLE_PI,	180.0,				1e-3,	EUnit::RadiansPerSecond, 				EUnit::DegreesPerSecond								},
 
+		{ 100.0,		1.0,		1e-3,	EUnit::CentimetersPerSecondSquared, 	EUnit::MetersPerSecondSquared								},
+		{ 1.0,		100.0,	1e-3,	EUnit::MetersPerSecondSquared, 		EUnit::CentimetersPerSecondSquared							},
+
 		{ 100, 		212,		0.1,	EUnit::Celsius,										EUnit::Farenheit									},
 		{ 400, 		477.594,	1e-3,	EUnit::Farenheit,									EUnit::Kelvin										},
 		{ 72, 		-201.15,	0.01,	EUnit::Kelvin,										EUnit::Celsius										},
diff --git a/Engine/Source/Runtime/CoreUObject/Public/UObject/NoExportTypes.h b/Engine/Source/Runtime/CoreUObject/Public/UObject/NoExportTypes.h
--- a/Engine/Source/Runtime/CoreUObject/Public/UObject/NoExportTypes.h
+++ b/Engine/Source/Runtime/CoreUObject/Public/UObject/NoExportTypes.h
@@ -370,6 +370,9 @@ enum class EUnit : uint8
 
 	/** Angular speed units */
 	DegreesPerSecond, RadiansPerSecond,
+	
+	/** Acceleration units */
+	CentimetersPerSecondSquared, MetersPerSecondSquared,
 
 	/** Temperature units */
 	Celsius, Farenheit, Kelvin,
