Prevented TSharedFromThis from being placed on a UObject.

diff --git a/Engine/Source/Runtime/Core/Public/Templates/SharedPointer.h b/Engine/Source/Runtime/Core/Public/Templates/SharedPointer.h
--- a/Engine/Source/Runtime/Core/Public/Templates/SharedPointer.h
+++ b/Engine/Source/Runtime/Core/Public/Templates/SharedPointer.h
@@ -1741,7 +1741,13 @@ public:		// @todo: Ideally this would be private, but template sharing problems
 protected:
 
 	/** Hidden stub constructor */
-	TSharedFromThis() { }
+	TSharedFromThis()
+	{
+		// Cause a compile error if ObjectType is a UObject.
+		int32 UObjectTestOverload(const volatile UObject*);
+		int16 UObjectTestOverload(...);
+		static_assert(sizeof(UObjectTestOverload((ObjectType*)nullptr)) == sizeof(int16), "TSharedFromThis is not supported on UObjects");
+	}
 
 	/** Hidden stub copy constructor */
 	TSharedFromThis( TSharedFromThis const& ) { }
