Fixed default initialization of TLS slots to 0 and invalidity check against 0 as 0 is a valid TLS slot index

diff --git a/Engine/Source/Editor/UnrealEd/Private/Cooker/CookPlatformManager.cpp b/Engine/Source/Editor/UnrealEd/Private/Cooker/CookPlatformManager.cpp
--- a/Engine/Source/Editor/UnrealEd/Private/Cooker/CookPlatformManager.cpp
+++ b/Engine/Source/Editor/UnrealEd/Private/Cooker/CookPlatformManager.cpp
@@ -69,7 +69,7 @@ FPlatformData::FPlatformData()
 {
 }
 
-uint32 FPlatformManager::IsInPlatformsLockTLSSlot = 0;
+uint32 FPlatformManager::IsInPlatformsLockTLSSlot = FPlatformTLS::InvalidTlsSlot;
 
 void FPlatformManager::InitializeTls()
 {
diff --git a/Engine/Source/Editor/UnrealEd/Private/Cooker/CookTypes.cpp b/Engine/Source/Editor/UnrealEd/Private/Cooker/CookTypes.cpp
--- a/Engine/Source/Editor/UnrealEd/Private/Cooker/CookTypes.cpp
+++ b/Engine/Source/Editor/UnrealEd/Private/Cooker/CookTypes.cpp
@@ -169,10 +169,10 @@ double FCookerTimer::GetActionTimeTillNow() const
 	return FPlatformTime::Seconds() - ActionStartTime;
 }
 
-static uint32 SchedulerThreadTlsSlot = 0;
+static uint32 SchedulerThreadTlsSlot = FPlatformTLS::InvalidTlsSlot;
 void InitializeTls()
 {
-	if (SchedulerThreadTlsSlot == 0)
+	if (!FPlatformTLS::IsValidTlsSlot(SchedulerThreadTlsSlot))
 	{
 		SchedulerThreadTlsSlot = FPlatformTLS::AllocTlsSlot();
 		SetIsSchedulerThread(true);
diff --git a/Engine/Source/Runtime/Core/Private/Android/AndroidJavaEnv.cpp b/Engine/Source/Runtime/Core/Private/Android/AndroidJavaEnv.cpp
--- a/Engine/Source/Runtime/Core/Private/Android/AndroidJavaEnv.cpp
+++ b/Engine/Source/Runtime/Core/Private/Android/AndroidJavaEnv.cpp
@@ -112,8 +112,8 @@ JNIEnv* AndroidJavaEnv::GetJavaEnv( bool bRequireGlobalThis /*= true*/ )
 
 	// Magic static - *should* be thread safe
 	//Android & pthread specific, bind a destructor for thread exit
-	static uint32 TlsSlot = 0;
-	if (TlsSlot == 0)
+	static uint32 TlsSlot = FPlatformTLS::InvalidTlsSlot;
+	if (!FPlatformTLS::IsValidTlsSlot(TlsSlot))
 	{
 		pthread_key_create((pthread_key_t*)&TlsSlot, &JavaEnvDestructor);
 	}
@@ -141,8 +141,8 @@ JNIEnv* AndroidJavaEnv::GetJavaEnv( bool bRequireGlobalThis /*= true*/ )
 	return (!bRequireGlobalThis || (GlobalObjectRef != nullptr)) ? Env : nullptr;
 #else
 	// register a destructor to detach this thread
-	static uint32 TlsSlot = 0;
-	if (TlsSlot == 0)
+	static uint32 TlsSlot = FPlatformTLS::InvalidTlsSlot;
+	if (!FPlatformTLS::IsValidTlsSlot(TlsSlot))
 	{
 		pthread_key_create((pthread_key_t*)&TlsSlot, &JavaEnvDestructor);
 	}
diff --git a/Engine/Source/Runtime/Core/Private/Async/TaskGraph.cpp b/Engine/Source/Runtime/Core/Private/Async/TaskGraph.cpp
--- a/Engine/Source/Runtime/Core/Private/Async/TaskGraph.cpp
+++ b/Engine/Source/Runtime/Core/Private/Async/TaskGraph.cpp
@@ -468,7 +468,7 @@ public:
 	/** Constructor, initializes everything to unusable values. Meant to be called from a "main" thread. **/
 	FTaskThreadBase()
 		: ThreadId(ENamedThreads::AnyThread)
-		, PerThreadIDTLSSlot(0xffffffff)
+		, PerThreadIDTLSSlot(FPlatformTLS::InvalidTlsSlot)
 		, OwnerWorker(nullptr)
 	{
 		NewTasks.Reset(128);
diff --git a/Engine/Source/Runtime/Core/Private/Containers/LockFreeList.cpp b/Engine/Source/Runtime/Core/Private/Containers/LockFreeList.cpp
--- a/Engine/Source/Runtime/Core/Private/Containers/LockFreeList.cpp
+++ b/Engine/Source/Runtime/Core/Private/Containers/LockFreeList.cpp
@@ -93,7 +93,7 @@ public:
 	~LockFreeLinkAllocator_TLSCache()
 	{
 		FPlatformTLS::FreeTlsSlot(TlsSlot);
-		TlsSlot = 0;
+		TlsSlot = FPlatformTLS::InvalidTlsSlot;
 	}
 
 	/**
diff --git a/Engine/Source/Runtime/Core/Private/HAL/MallocBinned2.cpp b/Engine/Source/Runtime/Core/Private/HAL/MallocBinned2.cpp
--- a/Engine/Source/Runtime/Core/Private/HAL/MallocBinned2.cpp
+++ b/Engine/Source/Runtime/Core/Private/HAL/MallocBinned2.cpp
@@ -142,7 +142,7 @@ static constexpr uint16 SmallBlockSizes[] =
 
 MS_ALIGN(PLATFORM_CACHE_LINE_SIZE) static uint8 UnusedAlignPadding[PLATFORM_CACHE_LINE_SIZE] GCC_ALIGN(PLATFORM_CACHE_LINE_SIZE) = { 0 };
 uint16 FMallocBinned2::SmallBlockSizesReversed[BINNED2_SMALL_POOL_COUNT] = { 0 };
-uint32 FMallocBinned2::Binned2TlsSlot = 0;
+uint32 FMallocBinned2::Binned2TlsSlot = FPlatformTLS::InvalidTlsSlot;
 uint32 FMallocBinned2::OsAllocationGranularity = 0;
 uint32 FMallocBinned2::PageSize = 0;
 FMallocBinned2* FMallocBinned2::MallocBinned2 = nullptr;
@@ -1358,11 +1358,11 @@ void FMallocBinned2::SetupTLSCachesOnCurrentThread()
 	{
 		return;
 	}
-	if (!FMallocBinned2::Binned2TlsSlot)
+	if (!FPlatformTLS::IsValidTlsSlot(FMallocBinned2::Binned2TlsSlot))
 	{
 		FMallocBinned2::Binned2TlsSlot = FPlatformTLS::AllocTlsSlot();
 	}
-	check(FMallocBinned2::Binned2TlsSlot);
+	check(FPlatformTLS::IsValidTlsSlot(FMallocBinned2::Binned2TlsSlot));
 	FPerThreadFreeBlockLists::SetTLS();
 }
 
@@ -1438,7 +1438,7 @@ FMallocBinned2::FBundleNode* FMallocBinned2::FFreeBlockList::PopBundles(uint32 I
 
 void FMallocBinned2::FPerThreadFreeBlockLists::SetTLS()
 {
-	check(FMallocBinned2::Binned2TlsSlot);
+	check(FPlatformTLS::IsValidTlsSlot(FMallocBinned2::Binned2TlsSlot));
 	FPerThreadFreeBlockLists* ThreadSingleton = (FPerThreadFreeBlockLists*)FPlatformTLS::GetTlsValue(FMallocBinned2::Binned2TlsSlot);
 	if (!ThreadSingleton)
 	{
@@ -1454,7 +1454,7 @@ void FMallocBinned2::FPerThreadFreeBlockLists::SetTLS()
 
 void FMallocBinned2::FPerThreadFreeBlockLists::ClearTLS()
 {
-	check(FMallocBinned2::Binned2TlsSlot);
+	check(FPlatformTLS::IsValidTlsSlot(FMallocBinned2::Binned2TlsSlot));
 	FPerThreadFreeBlockLists* ThreadSingleton = (FPerThreadFreeBlockLists*)FPlatformTLS::GetTlsValue(FMallocBinned2::Binned2TlsSlot);
 	if ( ThreadSingleton )
 	{
diff --git a/Engine/Source/Runtime/Core/Private/HAL/MallocBinned3.cpp b/Engine/Source/Runtime/Core/Private/HAL/MallocBinned3.cpp
--- a/Engine/Source/Runtime/Core/Private/HAL/MallocBinned3.cpp
+++ b/Engine/Source/Runtime/Core/Private/HAL/MallocBinned3.cpp
@@ -109,7 +109,7 @@ TAtomic<int32> MemoryRangeFreeTotalCount(0);
 
 MS_ALIGN(PLATFORM_CACHE_LINE_SIZE) static uint8 Binned3UnusedAlignPadding[PLATFORM_CACHE_LINE_SIZE] GCC_ALIGN(PLATFORM_CACHE_LINE_SIZE) = { 0 };
 uint16 FMallocBinned3::SmallBlockSizesReversedShifted[BINNED3_SMALL_POOL_COUNT + 1] = { 0 };
-uint32 FMallocBinned3::Binned3TlsSlot = 0;
+uint32 FMallocBinned3::Binned3TlsSlot = FPlatformTLS::InvalidTlsSlot;
 uint32 FMallocBinned3::OsAllocationGranularity = 0;
 
 #if !BINNED3_USE_SEPARATE_VM_PER_POOL
@@ -1405,11 +1405,11 @@ void FMallocBinned3::SetupTLSCachesOnCurrentThread()
 	{
 		return;
 	}
-	if (!FMallocBinned3::Binned3TlsSlot)
+	if (!FPlatformTLS::IsValidTlsSlot(FMallocBinned3::Binned3TlsSlot))
 	{
 		FMallocBinned3::Binned3TlsSlot = FPlatformTLS::AllocTlsSlot();
 	}
-	check(FMallocBinned3::Binned3TlsSlot);
+	check(FPlatformTLS::IsValidTlsSlot(FMallocBinned3::Binned3TlsSlot));
 	FPerThreadFreeBlockLists::SetTLS();
 }
 
@@ -1488,7 +1488,7 @@ FMallocBinned3::FBundleNode* FMallocBinned3::FFreeBlockList::PopBundles(uint32 I
 
 void FMallocBinned3::FPerThreadFreeBlockLists::SetTLS()
 {
-	check(FMallocBinned3::Binned3TlsSlot);
+	check(FPlatformTLS::IsValidTlsSlot(FMallocBinned3::Binned3TlsSlot));
 	FPerThreadFreeBlockLists* ThreadSingleton = (FPerThreadFreeBlockLists*)FPlatformTLS::GetTlsValue(FMallocBinned3::Binned3TlsSlot);
 	if (!ThreadSingleton)
 	{
@@ -1505,7 +1505,7 @@ void FMallocBinned3::FPerThreadFreeBlockLists::SetTLS()
 
 void FMallocBinned3::FPerThreadFreeBlockLists::ClearTLS()
 {
-	check(FMallocBinned3::Binned3TlsSlot);
+	check(FPlatformTLS::IsValidTlsSlot(FMallocBinned3::Binned3TlsSlot));
 	FPerThreadFreeBlockLists* ThreadSingleton = (FPerThreadFreeBlockLists*)FPlatformTLS::GetTlsValue(FMallocBinned3::Binned3TlsSlot);
 	if (ThreadSingleton)
 	{
diff --git a/Engine/Source/Runtime/Core/Private/HAL/MallocBinnedArena.cpp b/Engine/Source/Runtime/Core/Private/HAL/MallocBinnedArena.cpp
--- a/Engine/Source/Runtime/Core/Private/HAL/MallocBinnedArena.cpp
+++ b/Engine/Source/Runtime/Core/Private/HAL/MallocBinnedArena.cpp
@@ -1151,11 +1151,11 @@ void FMallocBinnedArena::SetupTLSCachesOnCurrentThread()
 	{
 		return;
 	}
-	if (!BinnedArenaTlsSlot)
+	if (!FPlatformTLS::IsValidTlsSlot(BinnedArenaTlsSlot))
 	{
 		BinnedArenaTlsSlot = FPlatformTLS::AllocTlsSlot();
 	}
-	check(BinnedArenaTlsSlot);
+	check(FPlatformTLS::IsValidTlsSlot(BinnedArenaTlsSlot));
 	FPerThreadFreeBlockLists::SetTLS(*this);
 }
 
@@ -1231,7 +1231,7 @@ FMallocBinnedArena::FBundleNode* FMallocBinnedArena::FFreeBlockList::PopBundles(
 void FMallocBinnedArena::FPerThreadFreeBlockLists::SetTLS(FMallocBinnedArena& Allocator)
 {
 	uint32 BinnedArenaTlsSlot = Allocator.BinnedArenaTlsSlot;
-	check(BinnedArenaTlsSlot);
+	check(FPlatformTLS::IsValidTlsSlot(BinnedArenaTlsSlot));
 	FPerThreadFreeBlockLists* ThreadSingleton = (FPerThreadFreeBlockLists*)FPlatformTLS::GetTlsValue(BinnedArenaTlsSlot);
 	if (!ThreadSingleton)
 	{
@@ -1247,7 +1247,7 @@ void FMallocBinnedArena::FPerThreadFreeBlockLists::SetTLS(FMallocBinnedArena& Al
 int64 FMallocBinnedArena::FPerThreadFreeBlockLists::ClearTLS(FMallocBinnedArena& Allocator)
 {
 	uint32 BinnedArenaTlsSlot = Allocator.BinnedArenaTlsSlot;
-	check(BinnedArenaTlsSlot);
+	check(FPlatformTLS::IsValidTlsSlot(BinnedArenaTlsSlot));
 	int64 Result = 0;
 	FPerThreadFreeBlockLists* ThreadSingleton = (FPerThreadFreeBlockLists*)FPlatformTLS::GetTlsValue(BinnedArenaTlsSlot);
 	if (ThreadSingleton)
diff --git a/Engine/Source/Runtime/Core/Private/HAL/MallocBinnedGPU.cpp b/Engine/Source/Runtime/Core/Private/HAL/MallocBinnedGPU.cpp
--- a/Engine/Source/Runtime/Core/Private/HAL/MallocBinnedGPU.cpp
+++ b/Engine/Source/Runtime/Core/Private/HAL/MallocBinnedGPU.cpp
@@ -1095,11 +1095,11 @@ void FMallocBinnedGPU::SetupTLSCachesOnCurrentThread()
 	{
 		return;
 	}
-	if (!BinnedGPUTlsSlot)
+	if (!FPlatformTLS::IsValidTlsSlot(BinnedGPUTlsSlot))
 	{
 		BinnedGPUTlsSlot = FPlatformTLS::AllocTlsSlot();
 	}
-	check(BinnedGPUTlsSlot);
+	check(FPlatformTLS::IsValidTlsSlot(BinnedGPUTlsSlot));
 	FPerThreadFreeBlockLists::SetTLS(*this);
 }
 
@@ -1175,7 +1175,7 @@ FMallocBinnedGPU::FBundleNode* FMallocBinnedGPU::FFreeBlockList::PopBundles(uint
 void FMallocBinnedGPU::FPerThreadFreeBlockLists::SetTLS(FMallocBinnedGPU& Allocator)
 {
 	uint32 BinnedGPUTlsSlot = Allocator.BinnedGPUTlsSlot;
-	check(BinnedGPUTlsSlot);
+	check(FPlatformTLS::IsValidTlsSlot(BinnedGPUTlsSlot));
 	FPerThreadFreeBlockLists* ThreadSingleton = (FPerThreadFreeBlockLists*)FPlatformTLS::GetTlsValue(BinnedGPUTlsSlot);
 	if (!ThreadSingleton)
 	{
@@ -1191,7 +1191,7 @@ void FMallocBinnedGPU::FPerThreadFreeBlockLists::SetTLS(FMallocBinnedGPU& Alloca
 int64 FMallocBinnedGPU::FPerThreadFreeBlockLists::ClearTLS(FMallocBinnedGPU& Allocator)
 {
 	uint32 BinnedGPUTlsSlot = Allocator.BinnedGPUTlsSlot;
-	check(BinnedGPUTlsSlot);
+	check(FPlatformTLS::IsValidTlsSlot(BinnedGPUTlsSlot));
 	int64 Result = 0;
 	FPerThreadFreeBlockLists* ThreadSingleton = (FPerThreadFreeBlockLists*)FPlatformTLS::GetTlsValue(BinnedGPUTlsSlot);
 	if (ThreadSingleton)
diff --git a/Engine/Source/Runtime/Core/Private/HAL/MallocLeakDetection.cpp b/Engine/Source/Runtime/Core/Private/HAL/MallocLeakDetection.cpp
--- a/Engine/Source/Runtime/Core/Private/HAL/MallocLeakDetection.cpp
+++ b/Engine/Source/Runtime/Core/Private/HAL/MallocLeakDetection.cpp
@@ -27,8 +27,8 @@
  */
 struct FMallocLeakDetectionStatics
 {
-	uint32 ContextsTLSID = 0;
-	uint32 SuppressDetectionCountTLSID = 0;
+	uint32 ContextsTLSID = FPlatformTLS::InvalidTlsSlot;
+	uint32 SuppressDetectionCountTLSID = FPlatformTLS::InvalidTlsSlot;
 
 	FMallocLeakDetectionStatics()
 	{
diff --git a/Engine/Source/Runtime/Core/Private/HAL/MallocStomp2.cpp b/Engine/Source/Runtime/Core/Private/HAL/MallocStomp2.cpp
--- a/Engine/Source/Runtime/Core/Private/HAL/MallocStomp2.cpp
+++ b/Engine/Source/Runtime/Core/Private/HAL/MallocStomp2.cpp
@@ -14,7 +14,7 @@ CORE_API FMallocStomp2* GMallocStomp2 = nullptr;
 CORE_API bool GMallocStomp2Enabled = false;
 
 
-uint32 FMallocStomp2::DisableTlsSlot = 0;
+uint32 FMallocStomp2::DisableTlsSlot = FPlatformTLS::InvalidTlsSlot;
 SIZE_T FMallocStomp2::MinSize = 0;
 SIZE_T FMallocStomp2::MaxSize = 0;
 
diff --git a/Engine/Source/Runtime/Core/Private/ProfilingDebugging/CallstackTrace.cpp b/Engine/Source/Runtime/Core/Private/ProfilingDebugging/CallstackTrace.cpp
--- a/Engine/Source/Runtime/Core/Private/ProfilingDebugging/CallstackTrace.cpp
+++ b/Engine/Source/Runtime/Core/Private/ProfilingDebugging/CallstackTrace.cpp
@@ -14,7 +14,7 @@ void	CallstackTrace_InitializeInternal();
 UE_TRACE_CHANNEL_DEFINE(CallstackChannel)
 UE_TRACE_EVENT_DEFINE(Memory, CallstackSpec)
 
-uint32 GCallStackTracingTlsSlotIndex = MAX_uint32;
+uint32 GCallStackTracingTlsSlotIndex = FPlatformTLS::InvalidTlsSlot;
 
 ////////////////////////////////////////////////////////////////////////////////
 void CallstackTrace_Create(class FMalloc* InMalloc)
diff --git a/Engine/Source/Runtime/Core/Private/ProfilingDebugging/CsvProfiler.cpp b/Engine/Source/Runtime/Core/Private/ProfilingDebugging/CsvProfiler.cpp
--- a/Engine/Source/Runtime/Core/Private/ProfilingDebugging/CsvProfiler.cpp
+++ b/Engine/Source/Runtime/Core/Private/ProfilingDebugging/CsvProfiler.cpp
@@ -1982,7 +1982,7 @@ private:
 public:
 	static void InitTls()
 	{
-		if (TlsSlot == 0)
+		if (!FPlatformTLS::IsValidTlsSlot(TlsSlot))
 		{
 			TlsSlot = FPlatformTLS::AllocTlsSlot();
 			FPlatformMisc::MemoryBarrier();
@@ -2226,7 +2226,7 @@ public:
 	TArray<const char*> WaitStatNameStack;
 };
 
-uint32 FCsvProfilerThreadData::TlsSlot = 0;
+uint32 FCsvProfilerThreadData::TlsSlot = FPlatformTLS::InvalidTlsSlot;
 FCriticalSection FCsvProfilerThreadData::TlsCS;
 TArray<FCsvProfilerThreadData::FWeakPtr> FCsvProfilerThreadData::TlsInstances;
 
diff --git a/Engine/Source/Runtime/Core/Private/ProfilingDebugging/StallDetector.cpp b/Engine/Source/Runtime/Core/Private/ProfilingDebugging/StallDetector.cpp
--- a/Engine/Source/Runtime/Core/Private/ProfilingDebugging/StallDetector.cpp
+++ b/Engine/Source/Runtime/Core/Private/ProfilingDebugging/StallDetector.cpp
@@ -400,7 +400,7 @@ public:
 	~ThreadLocalSlotAcquire()
 	{
 		FPlatformTLS::FreeTlsSlot(Slot);
-		Slot = ~0;
+		Slot = FPlatformTLS::InvalidTlsSlot;
 	}
 
 	uint32 GetSlot()
@@ -409,7 +409,7 @@ public:
 	}
 
 private:
-	uint32 Slot = ~0;
+	uint32 Slot = FPlatformTLS::InvalidTlsSlot;
 };
 static ThreadLocalSlotAcquire ThreadLocalSlot;
 
diff --git a/Engine/Source/Runtime/Core/Private/Stats/Stats2.cpp b/Engine/Source/Runtime/Core/Private/Stats/Stats2.cpp
--- a/Engine/Source/Runtime/Core/Private/Stats/Stats2.cpp
+++ b/Engine/Source/Runtime/Core/Private/Stats/Stats2.cpp
@@ -1066,7 +1066,7 @@ void FThreadStatsPool::ReturnToPool( FThreadStats* Instance )
 	FThreadStats
 -----------------------------------------------------------------------------*/
 
-uint32 FThreadStats::TlsSlot = 0;
+uint32 FThreadStats::TlsSlot = FPlatformTLS::InvalidTlsSlot;
 FThreadSafeCounter FThreadStats::PrimaryEnableCounter;
 FThreadSafeCounter FThreadStats::PrimaryEnableUpdateNumber;
 FThreadSafeCounter FThreadStats::PrimaryDisableChangeTagLock;
@@ -1084,7 +1084,7 @@ FThreadStats::FThreadStats():
 {
 	Packet.SetThreadProperties();
 
-	check(TlsSlot && FPlatformTLS::IsValidTlsSlot(TlsSlot));
+	check(FPlatformTLS::IsValidTlsSlot(TlsSlot));
 	FPlatformTLS::SetTlsValue(TlsSlot, this);
 }
 
@@ -1343,10 +1343,10 @@ void FThreadStats::StartThread()
 	// (Must do this before we expose ourselves to other threads via tls).
 	FThreadStatsPool::Get();
 	FStatsThreadState::GetLocalState(); // start up the state
-	if (!TlsSlot)
+	if (!FPlatformTLS::IsValidTlsSlot(TlsSlot))
 	{
 		TlsSlot = FPlatformTLS::AllocTlsSlot();
-		check(TlsSlot);
+		check(FPlatformTLS::IsValidTlsSlot(TlsSlot));
 	}
 
 	check(IsThreadingReady());
diff --git a/Engine/Source/Runtime/Core/Public/Experimental/Containers/HazardPointer.h b/Engine/Source/Runtime/Core/Public/Experimental/Containers/HazardPointer.h
--- a/Engine/Source/Runtime/Core/Public/Experimental/Containers/HazardPointer.h
+++ b/Engine/Source/Runtime/Core/Public/Experimental/Containers/HazardPointer.h
@@ -151,7 +151,7 @@ class FHazardPointerCollection
 	TArray<FTlsData*> AllTlsVariables;
 	TArray<FHazardRecordChunk*> HazardRecordBlocks;
 
-	uint32 CollectablesTlsSlot = 0;
+	uint32 CollectablesTlsSlot = FPlatformTLS::InvalidTlsSlot;
 	std::atomic_uint TotalNumHazardRecords{ HazardChunkSize };
 
 	void Collect(TArray<HazardPointer_Impl::FHazardDeleter>& Collectables);
diff --git a/Engine/Source/Runtime/Core/Public/HAL/MallocBinned2.h b/Engine/Source/Runtime/Core/Public/HAL/MallocBinned2.h
--- a/Engine/Source/Runtime/Core/Public/HAL/MallocBinned2.h
+++ b/Engine/Source/Runtime/Core/Public/HAL/MallocBinned2.h
@@ -409,7 +409,7 @@ class FMallocBinned2 : public FMalloc
 	{
 		FORCEINLINE static FPerThreadFreeBlockLists* Get()
 		{
-			return FMallocBinned2::Binned2TlsSlot ? (FPerThreadFreeBlockLists*)FPlatformTLS::GetTlsValue(FMallocBinned2::Binned2TlsSlot) : nullptr;
+			return FPlatformTLS::IsValidTlsSlot(FMallocBinned2::Binned2TlsSlot) ? (FPerThreadFreeBlockLists*)FPlatformTLS::GetTlsValue(FMallocBinned2::Binned2TlsSlot) : nullptr;
 		}
 		static void SetTLS();
 		static void ClearTLS();
diff --git a/Engine/Source/Runtime/Core/Public/HAL/MallocBinned3.h b/Engine/Source/Runtime/Core/Public/HAL/MallocBinned3.h
--- a/Engine/Source/Runtime/Core/Public/HAL/MallocBinned3.h
+++ b/Engine/Source/Runtime/Core/Public/HAL/MallocBinned3.h
@@ -373,7 +373,7 @@ class CORE_API FMallocBinned3 : public FMalloc
 	{
 		FORCEINLINE static FPerThreadFreeBlockLists* Get()
 		{
-			return FMallocBinned3::Binned3TlsSlot ? (FPerThreadFreeBlockLists*)FPlatformTLS::GetTlsValue(FMallocBinned3::Binned3TlsSlot) : nullptr;
+			return FPlatformTLS::IsValidTlsSlot(FMallocBinned3::Binned3TlsSlot) ? (FPerThreadFreeBlockLists*)FPlatformTLS::GetTlsValue(FMallocBinned3::Binned3TlsSlot) : nullptr;
 		}
 		static void SetTLS();
 		static void ClearTLS();
diff --git a/Engine/Source/Runtime/Core/Public/HAL/MallocBinnedArena.h b/Engine/Source/Runtime/Core/Public/HAL/MallocBinnedArena.h
--- a/Engine/Source/Runtime/Core/Public/HAL/MallocBinnedArena.h
+++ b/Engine/Source/Runtime/Core/Public/HAL/MallocBinnedArena.h
@@ -249,7 +249,7 @@ class CORE_API FMallocBinnedArena final : public FMalloc
 	{
 		FORCEINLINE static FPerThreadFreeBlockLists* Get(uint32 BinnedArenaTlsSlot)
 		{
-			return BinnedArenaTlsSlot ? (FPerThreadFreeBlockLists*)FPlatformTLS::GetTlsValue(BinnedArenaTlsSlot) : nullptr;
+			return FPlatformTLS::IsValidTlsSlot(BinnedArenaTlsSlot) ? (FPerThreadFreeBlockLists*)FPlatformTLS::GetTlsValue(BinnedArenaTlsSlot) : nullptr;
 		}
 		static void SetTLS(FMallocBinnedArena& Allocator);
 		static int64 ClearTLS(FMallocBinnedArena& Allocator);
@@ -645,7 +645,7 @@ public:
 	FArenaParams ArenaParams;
 
 	TArray<uint16> SmallBlockSizesReversedShifted; // this is reversed to get the smallest elements on our main cache line
-	uint32 BinnedArenaTlsSlot;
+	uint32 BinnedArenaTlsSlot = FPlatformTLS::InvalidTlsSlot;
 	uint64 PoolSearchDiv; // if this is zero, the VM turned out to be contiguous anyway so we use a simple subtract and shift
 	uint8* HighestPoolBaseVMPtr; // this is a duplicate of PoolBaseVMPtr[ArenaParams.PoolCount - 1]
 	FPlatformMemory::FPlatformVirtualMemoryBlock PoolBaseVMBlock;
diff --git a/Engine/Source/Runtime/Core/Public/HAL/MallocBinnedGPU.h b/Engine/Source/Runtime/Core/Public/HAL/MallocBinnedGPU.h
--- a/Engine/Source/Runtime/Core/Public/HAL/MallocBinnedGPU.h
+++ b/Engine/Source/Runtime/Core/Public/HAL/MallocBinnedGPU.h
@@ -271,7 +271,7 @@ class CORE_API FMallocBinnedGPU final : public FMalloc
 	{
 		FORCEINLINE static FPerThreadFreeBlockLists* Get(uint32 BinnedGPUTlsSlot)
 		{
-			return BinnedGPUTlsSlot ? (FPerThreadFreeBlockLists*)FPlatformTLS::GetTlsValue(BinnedGPUTlsSlot) : nullptr;
+			return FPlatformTLS::IsValidTlsSlot(BinnedGPUTlsSlot) ? (FPerThreadFreeBlockLists*)FPlatformTLS::GetTlsValue(BinnedGPUTlsSlot) : nullptr;
 		}
 		static void SetTLS(FMallocBinnedGPU& Allocator);
 		static int64 ClearTLS(FMallocBinnedGPU& Allocator);
@@ -618,7 +618,7 @@ public:
 	FArenaParams ArenaParams;
 
 	TArray<uint16> SmallBlockSizesReversedShifted; // this is reversed to get the smallest elements on our main cache line
-	uint32 BinnedGPUTlsSlot;
+	uint32 BinnedGPUTlsSlot = FPlatformTLS::InvalidTlsSlot;
 	uint64 PoolSearchDiv; // if this is zero, the VM turned out to be contiguous anyway so we use a simple subtract and shift
 	uint8* HighestPoolBaseVMPtr; // this is a duplicate of PoolBaseVMPtr[ArenaParams.PoolCount - 1]
 	FPlatformMemory::FPlatformVirtualMemoryBlock PoolBaseVMBlock;
diff --git a/Engine/Source/Runtime/Core/Public/ProfilingDebugging/CallstackTrace.h b/Engine/Source/Runtime/Core/Public/ProfilingDebugging/CallstackTrace.h
--- a/Engine/Source/Runtime/Core/Public/ProfilingDebugging/CallstackTrace.h
+++ b/Engine/Source/Runtime/Core/Public/ProfilingDebugging/CallstackTrace.h
@@ -55,7 +55,7 @@ extern CORE_API uint32 GCallStackTracingTlsSlotIndex;
 */
 inline void* CallstackTrace_GetFallbackPlatformReturnAddressData()
 {
-	if (GCallStackTracingTlsSlotIndex != MAX_uint32)
+	if (FPlatformTLS::IsValidTlsSlot(GCallStackTracingTlsSlotIndex))
 		return FPlatformTLS::GetTlsValue(GCallStackTracingTlsSlotIndex);
 	else
 		return nullptr;
@@ -77,7 +77,7 @@ class FCallStackTraceLimitResolveScope
 public:
 	FORCENOINLINE FCallStackTraceLimitResolveScope()
 	{
-		if (GCallStackTracingTlsSlotIndex != MAX_uint32)
+		if (FPlatformTLS::IsValidTlsSlot(GCallStackTracingTlsSlotIndex))
 		{
 			FPlatformTLS::SetTlsValue(GCallStackTracingTlsSlotIndex, PLATFORM_RETURN_ADDRESS_FOR_CALLSTACKTRACING());
 		}
@@ -85,7 +85,7 @@ public:
 
 	FORCENOINLINE ~FCallStackTraceLimitResolveScope()
 	{
-		if (GCallStackTracingTlsSlotIndex != MAX_uint32)
+		if (FPlatformTLS::IsValidTlsSlot(GCallStackTracingTlsSlotIndex))
 		{
 			FPlatformTLS::SetTlsValue(GCallStackTracingTlsSlotIndex, nullptr);
 		}
diff --git a/Engine/Source/Runtime/CoreUObject/Private/Serialization/AsyncLoading2.cpp b/Engine/Source/Runtime/CoreUObject/Private/Serialization/AsyncLoading2.cpp
--- a/Engine/Source/Runtime/CoreUObject/Private/Serialization/AsyncLoading2.cpp
+++ b/Engine/Source/Runtime/CoreUObject/Private/Serialization/AsyncLoading2.cpp
@@ -2370,14 +2370,14 @@ struct FAsyncLoadingThreadState2
 
 	static void Set(FAsyncLoadingThreadState2* State)
 	{
-		check(TlsSlot != 0);
+		check(FPlatformTLS::IsValidTlsSlot(TlsSlot));
 		check(!FPlatformTLS::GetTlsValue(TlsSlot));
 		FPlatformTLS::SetTlsValue(TlsSlot, State);
 	}
 
 	static FAsyncLoadingThreadState2* Get()
 	{
-		check(TlsSlot != 0);
+		check(FPlatformTLS::IsValidTlsSlot(TlsSlot));
 		return static_cast<FAsyncLoadingThreadState2*>(FPlatformTLS::GetTlsValue(TlsSlot));
 	}
 
