Cherrypick from CL# 31076980

diff --git a/Engine/Source/Runtime/Core/Private/Android/AndroidPlatformFile.cpp b/Engine/Source/Runtime/Core/Private/Android/AndroidPlatformFile.cpp
--- a/Engine/Source/Runtime/Core/Private/Android/AndroidPlatformFile.cpp
+++ b/Engine/Source/Runtime/Core/Private/Android/AndroidPlatformFile.cpp
@@ -154,7 +154,7 @@ FString AndroidRelativeToAbsolutePath(bool bUseInternalBasePath, FString RelPath
 	{
 		
 		do {
-			RelPath.RightChopInline(3, false);
+			RelPath.RightChopInline(3, EAllowShrinking::No);
 		} while (RelPath.StartsWith(TEXT("../"), ESearchCase::CaseSensitive));
 
 		return (bUseInternalBasePath ? GInternalFilePath : GetFileBasePath()) / RelPath;
@@ -2302,7 +2302,7 @@ private:
 			{
 				while (AndroidPath.StartsWith(TEXT("../"), ESearchCase::CaseSensitive))
 				{
-					AndroidPath.RightChopInline(3, false);
+					AndroidPath.RightChopInline(3, EAllowShrinking::No);
 				}
 				AndroidPath.ReplaceInline(FPlatformProcess::BaseDir(), TEXT(""));
 				if (AndroidPath.Equals(TEXT(".."), ESearchCase::CaseSensitive))
