Adding back serialization operator overload that takes an FArchive to the TVector class And FLinearColor. (This was removed 5 years ago).

There are some specific cases where structured archive serialization is not needed, serializing with a plain FArchive is faster.

diff --git a/Engine/Source/Runtime/Core/Public/Math/Color.h b/Engine/Source/Runtime/Core/Public/Math/Color.h
--- a/Engine/Source/Runtime/Core/Public/Math/Color.h
+++ b/Engine/Source/Runtime/Core/Public/Math/Color.h
@@ -89,6 +89,17 @@ struct FLinearColor
 	CORE_API explicit FLinearColor(const FFloat16Color& C);
 
 	// Serializer.
+	
+	friend FArchive& operator<<(FArchive& Ar,FLinearColor& Color)
+	{
+		return Ar << Color.R << Color.G << Color.B << Color.A;
+	}
+
+	bool Serialize( FArchive& Ar )
+	{
+		Ar << *this;
+		return true;
+	}
 
 	friend void operator<<(FStructuredArchive::FSlot Slot, FLinearColor& Color)
 	{
diff --git a/Engine/Source/Runtime/Core/Public/Math/Vector.h b/Engine/Source/Runtime/Core/Public/Math/Vector.h
--- a/Engine/Source/Runtime/Core/Public/Math/Vector.h
+++ b/Engine/Source/Runtime/Core/Public/Math/Vector.h
@@ -1142,11 +1142,16 @@ public:
      * which happen to have been chosen as Cluster starting points.  You want to be able to disregard those.
      */
     static void GenerateClusterCenters(TArray<TVector<T>>& Clusters, const TArray<TVector<T>>& Points, int32 NumIterations, int32 NumConnectionsToBeValid);
-
-    
+	
+	bool Serialize(FArchive& Ar)
+	{
+		Ar << *this;
+		return true;
+	}
+	
     bool Serialize(FStructuredArchive::FSlot Slot)
     {
-        Slot << (TVector<T>&)*this;
+        Slot << *this;
         return true;
     }
 
@@ -1184,6 +1189,57 @@ public:
 	explicit TVector(const TVector<FArg>& From) : TVector<T>((T)From.X, (T)From.Y, (T)From.Z) {}
 };
 
+/**
+ * Serializer for FVector3f.
+ *
+ * @param Ar Serialization Archive.
+ * @param V Vector to serialize.
+ */
+inline FArchive& operator<<(FArchive& Ar, TVector<float>& V)
+{
+	// @warning BulkSerialize: FVector3f is serialized as memory dump
+	// See TArray::BulkSerialize for detailed description of implied limitations.
+	Ar << V.X << V.Y << V.Z;
+
+	V.DiagnosticCheckNaN();
+
+	return Ar;
+}
+
+/**
+ * Serializer for FVector3d.
+ *
+ * @param Ar Serialization Archive.
+ * @param V Vector to serialize.
+ */
+inline FArchive& operator<<(FArchive& Ar, TVector<double>& V)
+{
+	// @warning BulkSerialize: FVector3d is serialized as memory dump
+	// See TArray::BulkSerialize for detailed description of implied limitations.
+	if (Ar.UEVer() >= EUnrealEngineObjectUE5Version::LARGE_WORLD_COORDINATES)
+	{
+		Ar << V.X;
+		Ar << V.Y;
+		Ar << V.Z;
+	}
+	else
+	{
+		checkf(Ar.IsLoading(), TEXT("float -> double conversion applied outside of load!"));
+		// Stored as floats, so serialize float and copy.
+		float X, Y, Z;
+
+		Ar << X;
+		Ar << Y;
+		Ar << Z;
+
+		V = TVector<double>(X, Y, Z);
+	};
+
+	V.DiagnosticCheckNaN();
+
+	return Ar;
+}
+
 /**
  * Structured archive slot serializer for FVector3f.
  *
diff --git a/Engine/Source/Runtime/CoreUObject/Private/UObject/Property.cpp b/Engine/Source/Runtime/CoreUObject/Private/UObject/Property.cpp
--- a/Engine/Source/Runtime/CoreUObject/Private/UObject/Property.cpp
+++ b/Engine/Source/Runtime/CoreUObject/Private/UObject/Property.cpp
@@ -44,6 +44,7 @@ struct TVector3StructOpsTypeTraits : public TStructOpsTypeTraitsBase2<T>
 		WithNetSharedSerialization = true,
 		WithStructuredSerializer = true,
 		WithStructuredSerializeFromMismatchedTag = true,
+		WithSerializer = true,
 	};
 	static constexpr EPropertyObjectReferenceType WithSerializerObjectReferences = EPropertyObjectReferenceType::None;
 };
@@ -294,6 +295,7 @@ struct TStructOpsTypeTraits<FLinearColor> : public TStructOpsTypeTraitsBase2<FLi
 		WithNoInitConstructor = true,
 		WithZeroConstructor = true,
 		WithStructuredSerializer = true,
+		WithSerializer = true,
 	};
 	static constexpr EPropertyObjectReferenceType WithSerializerObjectReferences = EPropertyObjectReferenceType::None;
 };
