Fixed TreatWarningsAsErrors not affecting the verbosity of the underlying log, and added TreatErrorsAsWarnings

This was a regression from 23567776, as previously the platform specific feedback context (eg, FWindowsNativeFeedbackContext) would modify the verbosity that was passed to the underlying GLog/GLogConsole.
This change restores the old behavior at the new code location (FFeedbackContext).

diff --git a/Engine/Source/Runtime/Core/Private/Misc/FeedbackContext.cpp b/Engine/Source/Runtime/Core/Private/Misc/FeedbackContext.cpp
--- a/Engine/Source/Runtime/Core/Private/Misc/FeedbackContext.cpp
+++ b/Engine/Source/Runtime/Core/Private/Misc/FeedbackContext.cpp
@@ -23,6 +23,7 @@ FFeedbackContext::~FFeedbackContext()
 
 void FFeedbackContext::Serialize(const TCHAR* V, ELogVerbosity::Type Verbosity, const FName& Category)
 {
+	Verbosity = ResolveVerbosity(Verbosity);
 	if (IsRunningCommandlet())
 	{
 		if (Verbosity == ELogVerbosity::Error || Verbosity == ELogVerbosity::Warning)
@@ -42,6 +43,7 @@ void FFeedbackContext::Serialize(const TCHAR* V, ELogVerbosity::Type Verbosity,
 
 void FFeedbackContext::Serialize(const TCHAR* V, ELogVerbosity::Type Verbosity, const FName& Category, double Time)
 {
+	Verbosity = ResolveVerbosity(Verbosity);
 	if (IsRunningCommandlet())
 	{
 		if (Verbosity == ELogVerbosity::Error || Verbosity == ELogVerbosity::Warning)
@@ -61,30 +63,46 @@ void FFeedbackContext::Serialize(const TCHAR* V, ELogVerbosity::Type Verbosity,
 
 void FFeedbackContext::SerializeRecord(const UE::FLogRecord& Record)
 {
+	ELogVerbosity::Type Verbosity = ResolveVerbosity(Record.GetVerbosity());
+	TOptional<UE::FLogRecord> RecordOverride;
+	if (Record.GetVerbosity() != Verbosity)
+	{
+		RecordOverride = Record;
+		RecordOverride->SetVerbosity(Verbosity);
+	}
 	if (IsRunningCommandlet())
 	{
-		const ELogVerbosity::Type Verbosity = Record.GetVerbosity();
 		if (Verbosity == ELogVerbosity::Error || Verbosity == ELogVerbosity::Warning)
 		{
-			AddRecordToHistory(Record);
+			AddRecordToHistory(RecordOverride.Get(Record));
 		}
 		if (GLogConsole && !GLog->IsRedirectingTo(GLogConsole))
 		{
-			GLogConsole->SerializeRecord(Record);
+			GLogConsole->SerializeRecord(RecordOverride.Get(Record));
 		}
 	}
 	if (!GLog->IsRedirectingTo(this))
 	{
-		GLog->SerializeRecord(Record);
+		GLog->SerializeRecord(RecordOverride.Get(Record));
 	}
 }
 
-void FFeedbackContext::FormatLine(FStringBuilderBase& Out, const TCHAR* V, ELogVerbosity::Type Verbosity, const FName& Category, double Time, ELogVerbosity::Type* OutVerbosity) const
+ELogVerbosity::Type FFeedbackContext::ResolveVerbosity(ELogVerbosity::Type Verbosity) const
 {
-	if (TreatWarningsAsErrors && Verbosity == ELogVerbosity::Warning)
+	if (Verbosity == ELogVerbosity::Error && TreatErrorsAsWarnings)
 	{
-		Verbosity = ELogVerbosity::Error;
+		return ELogVerbosity::Warning;
 	}
+	if (Verbosity == ELogVerbosity::Warning && TreatWarningsAsErrors)
+	{
+		return ELogVerbosity::Error;
+	}
+	return Verbosity;
+}
+
+void FFeedbackContext::FormatLine(FStringBuilderBase& Out, const TCHAR* V, ELogVerbosity::Type Verbosity, const FName& Category, double Time, ELogVerbosity::Type* OutVerbosity) const
+{
+	Verbosity = ResolveVerbosity(Verbosity);
 	if (OutVerbosity)
 	{
 		*OutVerbosity = Verbosity;
@@ -98,11 +116,7 @@ void FFeedbackContext::FormatLine(FStringBuilderBase& Out, const TCHAR* V, ELogV
 
 void FFeedbackContext::FormatRecordLine(FStringBuilderBase& Out, const UE::FLogRecord& Record, ELogVerbosity::Type* OutVerbosity) const
 {
-	ELogVerbosity::Type Verbosity = Record.GetVerbosity();
-	if (TreatWarningsAsErrors && Verbosity == ELogVerbosity::Warning)
-	{
-		Verbosity = ELogVerbosity::Error;
-	}
+	ELogVerbosity::Type Verbosity = ResolveVerbosity(Record.GetVerbosity());
 	if (OutVerbosity)
 	{
 		*OutVerbosity = Verbosity;
diff --git a/Engine/Source/Runtime/Core/Public/Misc/FeedbackContext.h b/Engine/Source/Runtime/Core/Public/Misc/FeedbackContext.h
--- a/Engine/Source/Runtime/Core/Public/Misc/FeedbackContext.h
+++ b/Engine/Source/Runtime/Core/Public/Misc/FeedbackContext.h
@@ -105,6 +105,8 @@ protected:
 	/** Called to check whether we are playing in editor when starting a slow task */
 	CORE_API virtual bool IsPlayingInEditor() const;
 
+	CORE_API ELogVerbosity::Type ResolveVerbosity(ELogVerbosity::Type Verbosity) const;
+
 	CORE_API void FormatLine(FStringBuilderBase& Out, const TCHAR* V, ELogVerbosity::Type Verbosity, const FName& Category, double Time, ELogVerbosity::Type* OutVerbosity = nullptr) const;
 	CORE_API void FormatRecordLine(FStringBuilderBase& Out, const UE::FLogRecord& Record, ELogVerbosity::Type* OutVerbosity = nullptr) const;
 
@@ -116,7 +118,10 @@ public:
 	virtual TWeakPtr<class SBuildProgressWidget> ShowBuildProgressWindow() {return TWeakPtr<class SBuildProgressWidget>();}
 	virtual void CloseBuildProgressWindow() {}
 
-	bool	TreatWarningsAsErrors = false;
+	/** Promote any logged warnings so that they act as errors */
+	bool TreatWarningsAsErrors = false;
+	/** Demote any logged errors so that they act as warnings; takes priority over TreatWarningsAsErrors */
+	bool TreatErrorsAsWarnings = false;
 
 	CORE_API FFeedbackContext();
 	CORE_API virtual ~FFeedbackContext();
diff --git a/Engine/Source/Runtime/Launch/Private/LaunchEngineLoop.cpp b/Engine/Source/Runtime/Launch/Private/LaunchEngineLoop.cpp
--- a/Engine/Source/Runtime/Launch/Private/LaunchEngineLoop.cpp
+++ b/Engine/Source/Runtime/Launch/Private/LaunchEngineLoop.cpp
@@ -6657,6 +6657,11 @@ bool FEngineLoop::AppInit( )
 		GWarn->TreatWarningsAsErrors = true;
 	}
 
+	if (FParse::Param(FCommandLine::Get(), TEXT("ERRORSASWARNINGS")))
+	{
+		GWarn->TreatErrorsAsWarnings = true;
+	}
+
 	if (FParse::Param(FCommandLine::Get(), TEXT("SILENT")))
 	{
 		GIsSilent = true;
