There is logic error in unreal where anytime touch is triggered, it calls CalculateSurfaceSize() which is not good.. it should be using a cached value since CalculateSurfaceSize should only be called on init and setup.
for now am downgrading the logic level to just be in debug builds.
#android

diff --git a/Engine/Source/Runtime/Core/Public/Android/AndroidPlatform.h b/Engine/Source/Runtime/Core/Public/Android/AndroidPlatform.h
--- a/Engine/Source/Runtime/Core/Public/Android/AndroidPlatform.h
+++ b/Engine/Source/Runtime/Core/Public/Android/AndroidPlatform.h
@@ -100,7 +100,7 @@ typedef FAndroidTypes FPlatformTypes;
 #endif
 
 
-#if (!USE_ANDROID_STANDALONE && !(UE_BUILD_SHIPPING || UE_BUILD_TEST)) || (USE_ANDROID_STANDALONE && !UE_BUILD_SHIPPING)
+#if (!USE_ANDROID_STANDALONE && UE_BUILD_DEBUG) || (USE_ANDROID_STANDALONE && !UE_BUILD_SHIPPING)
 	// M is the scope for the logging such as LogAndroid, STANDALONE_DEBUG_LOGf should be used when using formatted arguments.
 #	define STANDALONE_DEBUG_LOG(M, ...)  FPlatformMisc::LowLevelOutputDebugStringf(TEXT(#M " : "),##__VA_ARGS__);
 #	define STANDALONE_DEBUG_LOGf(M, ...)  FPlatformMisc::LowLevelOutputDebugStringf(TEXT(#M " : ") __VA_ARGS__);
