move the logic that determines whether __builtin_bit_cast should be used into Platform.h

 - allows platform extensions to determine whether it is available

#rnx

diff --git a/Engine/Source/Runtime/Core/Public/Android/AndroidPlatform.h b/Engine/Source/Runtime/Core/Public/Android/AndroidPlatform.h
--- a/Engine/Source/Runtime/Core/Public/Android/AndroidPlatform.h
+++ b/Engine/Source/Runtime/Core/Public/Android/AndroidPlatform.h
@@ -76,6 +76,12 @@ typedef FAndroidTypes FPlatformTypes;
 #define PLATFORM_ENABLE_VECTORINTRINSICS			1
 #define PLATFORM_ENABLE_VECTORINTRINSICS_NEON		PLATFORM_ANDROID_ARM64
 
+// Ensure we can use this builtin - seems to be present on Clang 9, GCC 11 and MSVC 19.26,
+// but gives spurious "non-void function 'BitCast' should return a value" errors on some
+// Mac and Android toolchains when building PCHs, so avoid those.
+#undef PLATFORM_COMPILER_SUPPORTS_BUILTIN_BITCAST
+#define PLATFORM_COMPILER_SUPPORTS_BUILTIN_BITCAST (__clang_major__ >= 13)
+
 // some android platform overrides that sub-platforms can disable
 #ifndef USE_ANDROID_JNI
 	#define USE_ANDROID_JNI							1
diff --git a/Engine/Source/Runtime/Core/Public/Clang/ClangPlatform.h b/Engine/Source/Runtime/Core/Public/Clang/ClangPlatform.h
--- a/Engine/Source/Runtime/Core/Public/Clang/ClangPlatform.h
+++ b/Engine/Source/Runtime/Core/Public/Clang/ClangPlatform.h
@@ -39,3 +39,6 @@
 #define UE_ALLOCATION_FUNCTION_2(SIZE, ALIGN) [[gnu::malloc, gnu::alloc_size(SIZE)]]
 #define UE_ALLOCATION_FUNCTION_X(x, SIZE, ALIGN, FUNC, ...) FUNC
 #define UE_ALLOCATION_FUNCTION(...) UE_ALLOCATION_FUNCTION_X(,##__VA_ARGS__, UE_ALLOCATION_FUNCTION_2(__VA_ARGS__), UE_ALLOCATION_FUNCTION_1(__VA_ARGS__), UE_ALLOCATION_FUNCTION_0(__VA_ARGS__)) 
+
+// Ensure we can use this builtin - seems to be present on Clang 9, GCC 11 and MSVC 19.26,
+#define PLATFORM_COMPILER_SUPPORTS_BUILTIN_BITCAST (__clang_major__ >= 9)
diff --git a/Engine/Source/Runtime/Core/Public/HAL/Platform.h b/Engine/Source/Runtime/Core/Public/HAL/Platform.h
--- a/Engine/Source/Runtime/Core/Public/HAL/Platform.h
+++ b/Engine/Source/Runtime/Core/Public/HAL/Platform.h
@@ -242,6 +242,9 @@
 #ifndef PLATFORM_COMPILER_SUPPORTS_CONSTEXPR_BUILTIN_FILE_AND_LINE
 	#define PLATFORM_COMPILER_SUPPORTS_CONSTEXPR_BUILTIN_FILE_AND_LINE 1
 #endif
+#ifndef PLATFORM_COMPILER_SUPPORTS_BUILTIN_BITCAST
+	#error PLATFORM_COMPILER_SUPPORTS_BUILTIN_BITCAST should have been defined by now
+#endif
 #ifndef PLATFORM_TCHAR_IS_4_BYTES
 	#define PLATFORM_TCHAR_IS_4_BYTES			0
 #endif
diff --git a/Engine/Source/Runtime/Core/Public/IOS/IOSPlatform.h b/Engine/Source/Runtime/Core/Public/IOS/IOSPlatform.h
--- a/Engine/Source/Runtime/Core/Public/IOS/IOSPlatform.h
+++ b/Engine/Source/Runtime/Core/Public/IOS/IOSPlatform.h
@@ -75,6 +75,12 @@ typedef FIOSPlatformTypes FPlatformTypes;
 #define PLATFORM_SUPPORTS_GEOMETRY_SHADERS				0
 #define PLATFORM_SUPPORTS_VIRTUAL_TEXTURE_STREAMING		1
 
+// Ensure we can use this builtin - seems to be present on Clang 9, GCC 11 and MSVC 19.26,
+// but gives spurious "non-void function 'BitCast' should return a value" errors on some
+// Mac and Android toolchains when building PCHs, so avoid those.
+#undef PLATFORM_COMPILER_SUPPORTS_BUILTIN_BITCAST
+#define PLATFORM_COMPILER_SUPPORTS_BUILTIN_BITCAST (__clang_major__ >= 13)
+
 #define PLATFORM_GLOBAL_LOG_CATEGORY					LogIOS
 
 #define PLATFORM_BREAK()                                __builtin_debugtrap()
diff --git a/Engine/Source/Runtime/Core/Public/IntelICX/IntelICXPlatform.h b/Engine/Source/Runtime/Core/Public/IntelICX/IntelICXPlatform.h
--- a/Engine/Source/Runtime/Core/Public/IntelICX/IntelICXPlatform.h
+++ b/Engine/Source/Runtime/Core/Public/IntelICX/IntelICXPlatform.h
@@ -36,4 +36,6 @@
 
 // Intel ICX does not support constexpr for __builtin_FILE and __builtin_LINE
 // Check with (__has_constexpr_builtin(__builtin_FILE) && __has_constexpr_builtin(__builtin_LINE)
-#define PLATFORM_COMPILER_SUPPORTS_CONSTEXPR_BUILTIN_FILE_AND_LINE 0
\ No newline at end of file
+#define PLATFORM_COMPILER_SUPPORTS_CONSTEXPR_BUILTIN_FILE_AND_LINE 0
+
+#define PLATFORM_COMPILER_SUPPORTS_BUILTIN_BITCAST 0
diff --git a/Engine/Source/Runtime/Core/Public/MSVC/MSVCPlatform.h b/Engine/Source/Runtime/Core/Public/MSVC/MSVCPlatform.h
--- a/Engine/Source/Runtime/Core/Public/MSVC/MSVCPlatform.h
+++ b/Engine/Source/Runtime/Core/Public/MSVC/MSVCPlatform.h
@@ -23,3 +23,7 @@
 #if __has_cpp_attribute(msvc::intrinsic)
 #define UE_INTRINSIC_CAST [[msvc::intrinsic]]
 #endif
+
+// Ensure we can use this builtin - seems to be present on Clang 9, GCC 11 and MSVC 19.26
+#define PLATFORM_COMPILER_SUPPORTS_BUILTIN_BITCAST (_MSC_VER >= 1926)
+
diff --git a/Engine/Source/Runtime/Core/Public/Mac/MacPlatform.h b/Engine/Source/Runtime/Core/Public/Mac/MacPlatform.h
--- a/Engine/Source/Runtime/Core/Public/Mac/MacPlatform.h
+++ b/Engine/Source/Runtime/Core/Public/Mac/MacPlatform.h
@@ -80,6 +80,12 @@ typedef FMacPlatformTypes FPlatformTypes;
 
 #define PLATFORM_ENABLE_POPCNT_INTRINSIC				1
 
+// Ensure we can use this builtin - seems to be present on Clang 9, GCC 11 and MSVC 19.26,
+// but gives spurious "non-void function 'BitCast' should return a value" errors on some
+// Mac and Android toolchains when building PCHs, so avoid those.
+#undef PLATFORM_COMPILER_SUPPORTS_BUILTIN_BITCAST
+#define PLATFORM_COMPILER_SUPPORTS_BUILTIN_BITCAST (__clang_major__ >= 13)
+
 #define PLATFORM_GLOBAL_LOG_CATEGORY					LogMac
 
 #if WITH_EDITOR
diff --git a/Engine/Source/Runtime/Core/Public/Templates/TypeCompatibleBytes.h b/Engine/Source/Runtime/Core/Public/Templates/TypeCompatibleBytes.h
--- a/Engine/Source/Runtime/Core/Public/Templates/TypeCompatibleBytes.h
+++ b/Engine/Source/Runtime/Core/Public/Templates/TypeCompatibleBytes.h
@@ -40,7 +40,7 @@ inline ToType BitCast(const FromType& From)
 // Mac and Android toolchains when building PCHs, so avoid those.
 // However, there is a bug in the Clang static analyzer with this builtin: [REDACTED_PROFILE_URL]/llvm-project/issues/69922
 // Don't use it when performing static analysis until the bug is fixed.
-#if !defined(__clang_analyzer__) && ((PLATFORM_APPLE && __clang_major__ >= 13) || (PLATFORM_ANDROID && __clang_major__ >= 13) || (PLATFORM_SWITCH && __clang_major__ >= 13) || (!PLATFORM_APPLE && !PLATFORM_ANDROID && !PLATFORM_SWITCH && defined(__clang__) && __clang_major__ >= 9) || (defined(__GNUC__) && __GNUC__ >= 11) || (defined(_MSC_VER) && _MSC_VER >= 1926))
+#if !defined(__clang_analyzer__) && PLATFORM_COMPILER_SUPPORTS_BUILTIN_BITCAST // can consider replacing with __has_builtin(__builtin_bit_cast) once there's no special cases
 	return __builtin_bit_cast(ToType, From);
 #else
 	TTypeCompatibleBytes<ToType> Result;
