Split all argument for Mac/Linux when launching UTS

All arguments must be separate entries in the argument array. Otherwise they will be passed a single extra argument.

diff --git a/Engine/Source/Runtime/Core/Private/ProfilingDebugging/TraceAuxiliary.cpp b/Engine/Source/Runtime/Core/Private/ProfilingDebugging/TraceAuxiliary.cpp
--- a/Engine/Source/Runtime/Core/Private/ProfilingDebugging/TraceAuxiliary.cpp
+++ b/Engine/Source/Runtime/Core/Private/ProfilingDebugging/TraceAuxiliary.cpp
@@ -1819,6 +1819,7 @@ static bool LaunchTraceServerCommand(ELaunchTraceServerCommand Command, bool bAd
 	}
 
 	TAnsiStringBuilder<64> ForkArg;
+	TAnsiStringBuilder<64> SponsorArg;
 	if (Command == ELaunchTraceServerCommand::Fork)
 	{
 		ForkArg << "fork";
@@ -1829,9 +1830,10 @@ static bool LaunchTraceServerCommand(ELaunchTraceServerCommand Command, bool bAd
 	}
 	if (bAddSponsor)
 	{
-		ForkArg << " --sponsor " << FPlatformProcess::GetCurrentProcessId();
+		SponsorArg  << "--sponsor=" << FPlatformProcess::GetCurrentProcessId();
 	}
 	ForkArg.ToString(); //Ensure zero termination
+	SponsorArg.ToString();
 
 	pid_t UtsPid = fork();
 	if (UtsPid < 0)
@@ -1842,7 +1844,7 @@ static bool LaunchTraceServerCommand(ELaunchTraceServerCommand Command, bool bAd
 	else if (UtsPid == 0)
 	{
 		// Launch UTS from the child process.
-		char* Args[] = { BinPath.GetData(), ForkArg.GetData(), nullptr };
+		char* Args[] = { BinPath.GetData(), ForkArg.GetData(), SponsorArg.GetData(), nullptr };
 		extern char** environ;
 		execve(*BinPath, Args, environ);
 		_exit(0x80 | (errno & 0x7f));
