ConfigAccessTracking: Fix merging of values with ExpandedText (ExpandValueInternal) in config values. The new operator= was not emptying ExpandedValue so it would sometimes leave the old ExpandedValue from the previous config layer that was supposed to be overwritten.
#rnx

diff --git a/Engine/Source/Runtime/Core/Public/Misc/ConfigCacheIni.h b/Engine/Source/Runtime/Core/Public/Misc/ConfigCacheIni.h
--- a/Engine/Source/Runtime/Core/Public/Misc/ConfigCacheIni.h
+++ b/Engine/Source/Runtime/Core/Public/Misc/ConfigCacheIni.h
@@ -212,6 +212,7 @@ public:
 	{
 		SavedValue = MoveTemp(RHS);
 		SavedValueHash = FTextLocalizationResource::HashString(SavedValue);
+		ExpandedValue.Empty();
 		ExpandValueInternal();
 		return *this;
 	}
