* Added InitFromString function to FIntPoint / FIntVector to mirror their ToString function
* Added unit tests for InitFromString for FColor, FLinearColor, FVector, FVector2D, FVector4, FQuat, FRotator

diff --git a/Engine/Source/Runtime/Core/Public/Math/IntPoint.h b/Engine/Source/Runtime/Core/Public/Math/IntPoint.h
--- a/Engine/Source/Runtime/Core/Public/Math/IntPoint.h
+++ b/Engine/Source/Runtime/Core/Public/Math/IntPoint.h
@@ -4,6 +4,7 @@
 
 #include "CoreTypes.h"
 #include "Misc/AssertionMacros.h"
+#include "Misc/Parse.h"
 #include "Math/UnrealMathUtility.h"
 #include "Containers/UnrealString.h"
 #include "Serialization/StructuredArchive.h"
@@ -399,6 +400,23 @@ PRAGMA_ENABLE_DEPRECATION_WARNINGS
 		return FString::Printf(TEXT("X=%s Y=%s"), *LexToString(X), *LexToString(Y));
 	}
 
+	/**
+	 * Initialize this FIntPoint based on an FString. The String is expected to contain X=, Y=
+	 * The FIntPoint will be bogus when InitFromString returns false.
+	 *
+	 * @param	InSourceString	FString containing the color values.
+	 * @return true if the X,Y values were read successfully; false otherwise.
+	 */
+	bool InitFromString(const FString& InSourceString)
+	{
+		X = Y = 0;
+
+		// The initialization is only successful if the X and Y values can all be parsed from the string
+		const bool bSuccessful = FParse::Value(*InSourceString, TEXT("X="), X) && FParse::Value(*InSourceString, TEXT("Y="), Y);
+
+		return bSuccessful;
+	}
+
 	/**
 	 * Divide an int point and round up the result.
 	 *
diff --git a/Engine/Source/Runtime/Core/Public/Math/IntVector.h b/Engine/Source/Runtime/Core/Public/Math/IntVector.h
--- a/Engine/Source/Runtime/Core/Public/Math/IntVector.h
+++ b/Engine/Source/Runtime/Core/Public/Math/IntVector.h
@@ -4,6 +4,7 @@
 
 #include "CoreTypes.h"
 #include "Misc/Crc.h"
+#include "Misc/Parse.h"
 #include "Math/UnrealMathUtility.h"
 #include "Containers/UnrealString.h"
 #include "Serialization/StructuredArchive.h"
@@ -452,6 +453,23 @@ PRAGMA_ENABLE_DEPRECATION_WARNINGS
 		return FString::Printf(TEXT("X=%s Y=%s Z=%s"), *LexToString(X), *LexToString(Y), *LexToString(Z));
 	}
 
+	/**
+	 * Initialize this FIntPoint based on an FString. The String is expected to contain X=, Y=
+	 * The FIntPoint will be bogus when InitFromString returns false.
+	 *
+	 * @param	InSourceString	FString containing the color values.
+	 * @return true if the X,Y values were read successfully; false otherwise.
+	 */
+	bool InitFromString(const FString& InSourceString)
+	{
+		X = Y = 0;
+
+		// The initialization is only successful if the X, Y and Z values can all be parsed from the string
+		const bool bSuccessful = FParse::Value(*InSourceString, TEXT("X="), X) && FParse::Value(*InSourceString, TEXT("Y="), Y) && FParse::Value(*InSourceString, TEXT("Z="), Z);
+
+		return bSuccessful;
+	}
+
 	/**
 	 * Divide an int point and round up the result.
 	 *
diff --git a/Engine/Source/Runtime/Core/Tests/Misc/ParseTests.cpp b/Engine/Source/Runtime/Core/Tests/Misc/ParseTests.cpp
--- a/Engine/Source/Runtime/Core/Tests/Misc/ParseTests.cpp
+++ b/Engine/Source/Runtime/Core/Tests/Misc/ParseTests.cpp
@@ -119,6 +119,350 @@ TEST_CASE("Parse::Value::ToBuffer", "[Parse][Smoke]")
 	}
 }
 
+TEST_CASE("Parse::InitFromString", "[Parse][Smoke]")
+{
+	SECTION("FVector2D")
+	{
+		FVector2D Value(0.0, 1.0);
+		FString Expected(TEXT("X=0.000 Y=1.000"));
+		CHECK(Value.ToString() == Expected);
+
+		// Back-and-forth conversion should work :
+		FVector2D NewValue;
+		CHECK(NewValue.InitFromString(Value.ToString()));
+		CHECK_MESSAGE(*FString::Printf(TEXT("Value:%s, Expected:%s"), *NewValue.ToString(), *Value.ToString()), NewValue.Equals(Value));
+
+		// Permissive formatting should work :
+		CHECK(NewValue.InitFromString(TEXT("X=0     ,Y=1.000000.2")));
+		CHECK_MESSAGE(*FString::Printf(TEXT("Value:%s, Expected:%s"), *NewValue.ToString(), *Value.ToString()), NewValue.Equals(Value));
+		// Out-of-order parameters :
+		CHECK(NewValue.InitFromString(TEXT("Y=1 X=0")));
+		CHECK_MESSAGE(*FString::Printf(TEXT("Value:%s, Expected:%s"), *NewValue.ToString(), *Value.ToString()), NewValue.Equals(Value));
+		// Various formats/delimiters :
+		CHECK(NewValue.InitFromString(TEXT("X=.0;Y=1.000000")));
+		CHECK_MESSAGE(*FString::Printf(TEXT("Value:%s, Expected:%s"), *NewValue.ToString(), *Value.ToString()), NewValue.Equals(Value));
+		CHECK(NewValue.InitFromString(TEXT("X=0A Y=1.000000A")));
+		CHECK_MESSAGE(*FString::Printf(TEXT("Value:%s, Expected:%s"), *NewValue.ToString(), *Value.ToString()), NewValue.Equals(Value));
+		// Missing value == 0 :
+		CHECK(NewValue.InitFromString(TEXT("X= Y=1")));
+		CHECK_MESSAGE(*FString::Printf(TEXT("Value:%s, Expected:%s"), *NewValue.ToString(), *Value.ToString()), NewValue.Equals(Value));
+		CHECK(NewValue.InitFromString(TEXT("X=Y=1")));
+		CHECK_MESSAGE(*FString::Printf(TEXT("Value:%s, Expected:%s"), *NewValue.ToString(), *Value.ToString()), NewValue.Equals(Value));
+		CHECK(NewValue.InitFromString(TEXT("X=A Y=1")));
+		CHECK_MESSAGE(*FString::Printf(TEXT("Value:%s, Expected:%s"), *NewValue.ToString(), *Value.ToString()), NewValue.Equals(Value));
+		// Invalid formats : 
+		CHECK(!NewValue.InitFromString(TEXT("XA= Y=1")));
+		CHECK(!NewValue.InitFromString(TEXT("X=0Y=1")));
+		// Missing component should yield an error :
+		CHECK(!NewValue.InitFromString(TEXT("X=0")));
+		CHECK(!NewValue.InitFromString(TEXT("X=0 A=2")));
+	}
+
+	SECTION("FVector")
+	{
+		FVector Value(0.0, 1.0, 2.0);
+		FString Expected(TEXT("X=0.000 Y=1.000 Z=2.000"));
+		CHECK(Value.ToString() == Expected);
+
+		// Back-and-forth conversion should work :
+		FVector NewValue;
+		CHECK(NewValue.InitFromString(Value.ToString()));
+		CHECK_MESSAGE(*FString::Printf(TEXT("Value:%s, Expected:%s"), *NewValue.ToString(), *Value.ToString()), NewValue.Equals(Value));
+
+		// Permissive formatting should work :
+		CHECK(NewValue.InitFromString(TEXT("X=0     ,Y= 1.000000.2:Z=2")));
+		CHECK_MESSAGE(*FString::Printf(TEXT("Value:%s, Expected:%s"), *NewValue.ToString(), *Value.ToString()), NewValue.Equals(Value));
+		// Out-of-order parameters :
+		CHECK(NewValue.InitFromString(TEXT("Y=1 Z=2 X=0")));
+		CHECK_MESSAGE(*FString::Printf(TEXT("Value:%s, Expected:%s"), *NewValue.ToString(), *Value.ToString()), NewValue.Equals(Value));
+		// Various formats/delimiters :
+		CHECK(NewValue.InitFromString(TEXT("X=.0;Y=1.000000|Z=2")));
+		CHECK_MESSAGE(*FString::Printf(TEXT("Value:%s, Expected:%s"), *NewValue.ToString(), *Value.ToString()), NewValue.Equals(Value));
+		CHECK(NewValue.InitFromString(TEXT("X=0A Y=1.000000A Z=2.")));
+		CHECK_MESSAGE(*FString::Printf(TEXT("Value:%s, Expected:%s"), *NewValue.ToString(), *Value.ToString()), NewValue.Equals(Value));
+		// Missing value == 0 :
+		CHECK(NewValue.InitFromString(TEXT("X= Y=1 Z=2")));
+		CHECK_MESSAGE(*FString::Printf(TEXT("Value:%s, Expected:%s"), *NewValue.ToString(), *Value.ToString()), NewValue.Equals(Value));
+		CHECK(NewValue.InitFromString(TEXT("X=Y=1 Z=2")));
+		CHECK_MESSAGE(*FString::Printf(TEXT("Value:%s, Expected:%s"), *NewValue.ToString(), *Value.ToString()), NewValue.Equals(Value));
+		CHECK(NewValue.InitFromString(TEXT("X=A Y=1 Z=2")));
+		CHECK_MESSAGE(*FString::Printf(TEXT("Value:%s, Expected:%s"), *NewValue.ToString(), *Value.ToString()), NewValue.Equals(Value));
+		// Invalid formats : 
+		CHECK(!NewValue.InitFromString(TEXT("XA= Y=1 Z=2")));
+		CHECK(!NewValue.InitFromString(TEXT("X =0 Y=1 Z=2")));
+		CHECK(!NewValue.InitFromString(TEXT("X=0Y=1Z=2")));
+		// Missing component should yield an error :
+		CHECK(!NewValue.InitFromString(TEXT("X=0 Y=1")));
+		CHECK(!NewValue.InitFromString(TEXT("X=0 Y=1 A=2")));
+	}
+
+	SECTION("FVector4")
+	{
+		FVector4 Value(0.0, 1.0, 2.0, 3.0);
+		FString Expected(TEXT("X=0.000 Y=1.000 Z=2.000 W=3.000"));
+		CHECK(Value.ToString() == Expected);
+
+		// Back-and-forth conversion should work :
+		FVector4 NewValue;
+		CHECK(NewValue.InitFromString(Value.ToString()));
+		CHECK_MESSAGE(*FString::Printf(TEXT("Value:%s, Expected:%s"), *NewValue.ToString(), *Value.ToString()), NewValue.Equals(Value));
+
+		// Permissive formatting should work :
+		CHECK(NewValue.InitFromString(TEXT("X=0     ,Y= 1.000000.2:Z=2 W= 3.")));
+		CHECK_MESSAGE(*FString::Printf(TEXT("Value:%s, Expected:%s"), *NewValue.ToString(), *Value.ToString()), NewValue.Equals(Value));
+		// Out-of-order parameters :
+		CHECK(NewValue.InitFromString(TEXT("Y=1 Z=2 W=3 X=0")));
+		CHECK_MESSAGE(*FString::Printf(TEXT("Value:%s, Expected:%s"), *NewValue.ToString(), *Value.ToString()), NewValue.Equals(Value));
+		// Various formats/delimiters :
+		CHECK(NewValue.InitFromString(TEXT("X=.0;Y=1.000000|Z=2:W=3")));
+		CHECK_MESSAGE(*FString::Printf(TEXT("Value:%s, Expected:%s"), *NewValue.ToString(), *Value.ToString()), NewValue.Equals(Value));
+		CHECK(NewValue.InitFromString(TEXT("X=0A Y=1.000000A Z=2. W= 3")));
+		CHECK_MESSAGE(*FString::Printf(TEXT("Value:%s, Expected:%s"), *NewValue.ToString(), *Value.ToString()), NewValue.Equals(Value));
+		// Missing value == 0 :
+		CHECK(NewValue.InitFromString(TEXT("X= Y=1 Z=2 W=3")));
+		CHECK_MESSAGE(*FString::Printf(TEXT("Value:%s, Expected:%s"), *NewValue.ToString(), *Value.ToString()), NewValue.Equals(Value));
+		CHECK(NewValue.InitFromString(TEXT("X=Y=1 Z=2 W=3")));
+		CHECK_MESSAGE(*FString::Printf(TEXT("Value:%s, Expected:%s"), *NewValue.ToString(), *Value.ToString()), NewValue.Equals(Value));
+		CHECK(NewValue.InitFromString(TEXT("X=A Y=1 Z=2 W=3")));
+		CHECK_MESSAGE(*FString::Printf(TEXT("Value:%s, Expected:%s"), *NewValue.ToString(), *Value.ToString()), NewValue.Equals(Value));
+		// Invalid formats : 
+		CHECK(!NewValue.InitFromString(TEXT("XA= Y=1 Z=2 W=3")));
+		CHECK(!NewValue.InitFromString(TEXT("X=0Y=1Z=2W=3")));
+		// Missing component should yield an error :
+		CHECK(NewValue.InitFromString(TEXT("X=0 Y=1 Z=2"))); // W is optional for FVector4
+		CHECK(!NewValue.InitFromString(TEXT("X=0 Y=1 W=3")));
+		CHECK(!NewValue.InitFromString(TEXT("X=0 Y=1 A=2 W=3")));
+	}
+
+	SECTION("FQuat")
+	{
+		FQuat Value(0.0, 1.0, 2.0, 3.0);
+		FString Expected(TEXT("X=0.000000000 Y=1.000000000 Z=2.000000000 W=3.000000000"));
+		CHECK(Value.ToString() == Expected);
+
+		// Back-and-forth conversion should work :
+		FQuat NewValue;
+		CHECK(NewValue.InitFromString(Value.ToString()));
+		CHECK_MESSAGE(*FString::Printf(TEXT("Value:%s, Expected:%s"), *NewValue.ToString(), *Value.ToString()), NewValue.Equals(Value));
+
+		// Permissive formatting should work :
+		CHECK(NewValue.InitFromString(TEXT("X=0     ,Y= 1.000000.2:Z=2 W= 3.")));
+		CHECK_MESSAGE(*FString::Printf(TEXT("Value:%s, Expected:%s"), *NewValue.ToString(), *Value.ToString()), NewValue.Equals(Value));
+		// Out-of-order parameters :
+		CHECK(NewValue.InitFromString(TEXT("Y=1 Z=2 W=3 X=0")));
+		CHECK_MESSAGE(*FString::Printf(TEXT("Value:%s, Expected:%s"), *NewValue.ToString(), *Value.ToString()), NewValue.Equals(Value));
+		// Various formats/delimiters :
+		CHECK(NewValue.InitFromString(TEXT("X=.0;Y=1.000000|Z=2:W=3")));
+		CHECK_MESSAGE(*FString::Printf(TEXT("Value:%s, Expected:%s"), *NewValue.ToString(), *Value.ToString()), NewValue.Equals(Value));
+		CHECK(NewValue.InitFromString(TEXT("X=0A Y=1.000000A Z=2. W= 3")));
+		CHECK_MESSAGE(*FString::Printf(TEXT("Value:%s, Expected:%s"), *NewValue.ToString(), *Value.ToString()), NewValue.Equals(Value));
+		// Missing value == 0 :
+		CHECK(NewValue.InitFromString(TEXT("X= Y=1 Z=2 W=3")));
+		CHECK_MESSAGE(*FString::Printf(TEXT("Value:%s, Expected:%s"), *NewValue.ToString(), *Value.ToString()), NewValue.Equals(Value));
+		CHECK(NewValue.InitFromString(TEXT("X=Y=1 Z=2 W=3")));
+		CHECK_MESSAGE(*FString::Printf(TEXT("Value:%s, Expected:%s"), *NewValue.ToString(), *Value.ToString()), NewValue.Equals(Value));
+		CHECK(NewValue.InitFromString(TEXT("X=A Y=1 Z=2 W=3")));
+		CHECK_MESSAGE(*FString::Printf(TEXT("Value:%s, Expected:%s"), *NewValue.ToString(), *Value.ToString()), NewValue.Equals(Value));
+		// Invalid formats : 
+		CHECK(!NewValue.InitFromString(TEXT("XA= Y=1 Z=2 W=3")));
+		CHECK(!NewValue.InitFromString(TEXT("X=0Y=1Z=2W=3")));
+		// Missing component should yield an error :
+		CHECK(!NewValue.InitFromString(TEXT("X=0 Y=1 Z=2"))); 
+		CHECK(!NewValue.InitFromString(TEXT("X=0 Y=1 W=3")));
+		CHECK(!NewValue.InitFromString(TEXT("X=0 Y=1 A=2 W=3")));
+	}
+
+	SECTION("FLinearColor")
+	{
+		FLinearColor Value(0.0, 1.0, 2.0, 3.0);
+		FString Expected(TEXT("(R=0.000000,G=1.000000,B=2.000000,A=3.000000)"));
+		CHECK(Value.ToString() == Expected);
+
+		// Back-and-forth conversion should work :
+		FLinearColor NewValue;
+		CHECK(NewValue.InitFromString(Value.ToString()));
+		CHECK_MESSAGE(*FString::Printf(TEXT("Value:%s, Expected:%s"), *NewValue.ToString(), *Value.ToString()), NewValue.Equals(Value));
+
+		// Permissive formatting should work :
+		CHECK(NewValue.InitFromString(TEXT("R=0     ,G= 1.000000.2:B=2 A= 3.")));
+		CHECK_MESSAGE(*FString::Printf(TEXT("Value:%s, Expected:%s"), *NewValue.ToString(), *Value.ToString()), NewValue.Equals(Value));
+		// Out-of-order parameters :
+		CHECK(NewValue.InitFromString(TEXT("G=1 B=2 A=3 R=0")));
+		CHECK_MESSAGE(*FString::Printf(TEXT("Value:%s, Expected:%s"), *NewValue.ToString(), *Value.ToString()), NewValue.Equals(Value));
+		// Various formats/delimiters :
+		CHECK(NewValue.InitFromString(TEXT("R=.0;G=1.000000|B=2:A=3")));
+		CHECK_MESSAGE(*FString::Printf(TEXT("Value:%s, Expected:%s"), *NewValue.ToString(), *Value.ToString()), NewValue.Equals(Value));
+		CHECK(NewValue.InitFromString(TEXT("R=0A G=1.000000A B=2. A= 3")));
+		CHECK_MESSAGE(*FString::Printf(TEXT("Value:%s, Expected:%s"), *NewValue.ToString(), *Value.ToString()), NewValue.Equals(Value));
+		// Missing value == 0 :
+		CHECK(NewValue.InitFromString(TEXT("R= G=1 B=2 A=3")));
+		CHECK_MESSAGE(*FString::Printf(TEXT("Value:%s, Expected:%s"), *NewValue.ToString(), *Value.ToString()), NewValue.Equals(Value));
+		CHECK(NewValue.InitFromString(TEXT("R=G=1 B=2 A=3")));
+		CHECK_MESSAGE(*FString::Printf(TEXT("Value:%s, Expected:%s"), *NewValue.ToString(), *Value.ToString()), NewValue.Equals(Value));
+		CHECK(NewValue.InitFromString(TEXT("R=A G=1 B=2 A=3")));
+		CHECK_MESSAGE(*FString::Printf(TEXT("Value:%s, Expected:%s"), *NewValue.ToString(), *Value.ToString()), NewValue.Equals(Value));
+		// Invalid formats : 
+		CHECK(!NewValue.InitFromString(TEXT("RA= G=1 B=2 A=3")));
+		CHECK(!NewValue.InitFromString(TEXT("R=0G=1B=2W=3")));
+		CHECK(!NewValue.InitFromString(TEXT("R =0 G=1 B=2 W=3")));
+		// Missing component should yield an error :
+		CHECK(NewValue.InitFromString(TEXT("R=0 G=1 B=2"))); // A is optional for FLinearColor
+		CHECK(!NewValue.InitFromString(TEXT("R=0 G=1 A=3")));
+		CHECK(!NewValue.InitFromString(TEXT("R=0 G=1 A=2 A=3")));
+	}
+
+	SECTION("FColor")
+	{
+		FColor Value(0, 1, 2, 3);
+		FString Expected(TEXT("(R=0,G=1,B=2,A=3)"));
+		CHECK(Value.ToString() == Expected);
+
+		// Back-and-forth conversion should work :
+		FColor NewValue;
+		CHECK(NewValue.InitFromString(Value.ToString()));
+		CHECK_MESSAGE(*FString::Printf(TEXT("Value:%s, Expected:%s"), *NewValue.ToString(), *Value.ToString()), NewValue == Value);
+
+		// Permissive formatting should work :
+		CHECK(NewValue.InitFromString(TEXT("R=0     ,G= 1.000000.2:B=2 A= 3.")));
+		CHECK_MESSAGE(*FString::Printf(TEXT("Value:%s, Expected:%s"), *NewValue.ToString(), *Value.ToString()), NewValue == Value);
+		// Out-of-order parameters :
+		CHECK(NewValue.InitFromString(TEXT("G=1 B=2 A=3 R=0")));
+		CHECK_MESSAGE(*FString::Printf(TEXT("Value:%s, Expected:%s"), *NewValue.ToString(), *Value.ToString()), NewValue == Value);
+		// Various formats/delimiters :
+		CHECK(NewValue.InitFromString(TEXT("R=0;G=1|B=2:A=3")));
+		CHECK_MESSAGE(*FString::Printf(TEXT("Value:%s, Expected:%s"), *NewValue.ToString(), *Value.ToString()), NewValue == Value);
+		CHECK(NewValue.InitFromString(TEXT("R=0A G=1A B=2 A= 3")));
+		CHECK_MESSAGE(*FString::Printf(TEXT("Value:%s, Expected:%s"), *NewValue.ToString(), *Value.ToString()), NewValue == Value);
+		// Missing value is invalid for FColor :
+		CHECK(!NewValue.InitFromString(TEXT("R= G=1 B=2 A=3")));
+		CHECK(!NewValue.InitFromString(TEXT("R=G=1 B=2 A=3")));
+		CHECK(!NewValue.InitFromString(TEXT("R=A G=1 B=2 A=3")));
+		// Invalid formats : 
+		CHECK(!NewValue.InitFromString(TEXT("RA= G=1 B=2 A=3")));
+		CHECK(!NewValue.InitFromString(TEXT("R=0G=1B=2A=3")));
+		CHECK(!NewValue.InitFromString(TEXT("R =0 G=1 B=2 A=3")));
+		CHECK(!NewValue.InitFromString(TEXT("R=.0 G=1 B=2 A=3")));
+		// Missing component should yield an error :
+		CHECK(NewValue.InitFromString(TEXT("R=0 G=1 B=2"))); // A is optional for FColor
+		CHECK(!NewValue.InitFromString(TEXT("R=0 G=1 A=3")));
+		CHECK(!NewValue.InitFromString(TEXT("R=0 G=1 A=2 A=3")));
+	}
+
+	SECTION("FRotator")
+	{
+		FRotator Value(0.0, 1.0, 2.0);
+		FString Expected(TEXT("P=0.000000 Y=1.000000 R=2.000000"));
+		CHECK(Value.ToString() == Expected);
+
+		// Back-and-forth conversion should work :
+		FRotator NewValue;
+		CHECK(NewValue.InitFromString(Value.ToString()));
+		CHECK_MESSAGE(*FString::Printf(TEXT("Value:%s, Expected:%s"), *NewValue.ToString(), *Value.ToString()), NewValue.Equals(Value));
+
+		// Permissive formatting should work :
+		CHECK(NewValue.InitFromString(TEXT("P=0     ,Y= 1.000000.2:R=2")));
+		CHECK_MESSAGE(*FString::Printf(TEXT("Value:%s, Expected:%s"), *NewValue.ToString(), *Value.ToString()), NewValue.Equals(Value));
+		// Out-of-order parameters :
+		CHECK(NewValue.InitFromString(TEXT("Y=1 R=2 P=0")));
+		CHECK_MESSAGE(*FString::Printf(TEXT("Value:%s, Expected:%s"), *NewValue.ToString(), *Value.ToString()), NewValue.Equals(Value));
+		// Various formats/delimiters :
+		CHECK(NewValue.InitFromString(TEXT("P=.0;Y=1.000000|R=2")));
+		CHECK_MESSAGE(*FString::Printf(TEXT("Value:%s, Expected:%s"), *NewValue.ToString(), *Value.ToString()), NewValue.Equals(Value));
+		CHECK(NewValue.InitFromString(TEXT("P=0A Y=1.000000A R=2.")));
+		CHECK_MESSAGE(*FString::Printf(TEXT("Value:%s, Expected:%s"), *NewValue.ToString(), *Value.ToString()), NewValue.Equals(Value));
+		// Missing value == 0 :
+		CHECK(NewValue.InitFromString(TEXT("P= Y=1 R=2")));
+		CHECK_MESSAGE(*FString::Printf(TEXT("Value:%s, Expected:%s"), *NewValue.ToString(), *Value.ToString()), NewValue.Equals(Value));
+		CHECK(NewValue.InitFromString(TEXT("P=Y=1 R=2")));
+		CHECK_MESSAGE(*FString::Printf(TEXT("Value:%s, Expected:%s"), *NewValue.ToString(), *Value.ToString()), NewValue.Equals(Value));
+		CHECK(NewValue.InitFromString(TEXT("P=A Y=1 R=2")));
+		CHECK_MESSAGE(*FString::Printf(TEXT("Value:%s, Expected:%s"), *NewValue.ToString(), *Value.ToString()), NewValue.Equals(Value));
+		// Invalid formats : 
+		CHECK(!NewValue.InitFromString(TEXT("PA= Y=1 R=2")));
+		CHECK(!NewValue.InitFromString(TEXT("P=0Y=1R=2")));
+		// Missing component should yield an error :
+		CHECK(!NewValue.InitFromString(TEXT("P=0 Y=1")));
+		CHECK(!NewValue.InitFromString(TEXT("P=0 Y=1 A=2")));
+	}
+
+	SECTION("FIntPoint")
+	{
+		FIntPoint Value(0, 1);
+		FString Expected(TEXT("X=0 Y=1"));
+		CHECK(Value.ToString() == Expected);
+
+		// Back-and-forth conversion should work :
+		FIntPoint NewValue;
+		CHECK(NewValue.InitFromString(Value.ToString()));
+		CHECK_MESSAGE(*FString::Printf(TEXT("Value:%s, Expected:%s"), *NewValue.ToString(), *Value.ToString()), NewValue == Value);
+
+		// Permissive formatting should work :
+		CHECK(NewValue.InitFromString(TEXT("X=0     ,Y= 1.000000.2")));
+		CHECK_MESSAGE(*FString::Printf(TEXT("Value:%s, Expected:%s"), *NewValue.ToString(), *Value.ToString()), NewValue == Value);
+		// Out-of-order parameters :
+		CHECK(NewValue.InitFromString(TEXT("Y=1 X=0")));
+		CHECK_MESSAGE(*FString::Printf(TEXT("Value:%s, Expected:%s"), *NewValue.ToString(), *Value.ToString()), NewValue == Value);
+		// Various formats/delimiters :
+		CHECK(NewValue.InitFromString(TEXT("X=.0;Y=1")));
+		CHECK_MESSAGE(*FString::Printf(TEXT("Value:%s, Expected:%s"), *NewValue.ToString(), *Value.ToString()), NewValue == Value);
+		CHECK(NewValue.InitFromString(TEXT("|X=0:Y=1")));
+		CHECK_MESSAGE(*FString::Printf(TEXT("Value:%s, Expected:%s"), *NewValue.ToString(), *Value.ToString()), NewValue == Value);
+		CHECK(NewValue.InitFromString(TEXT("X=0A Y= 1")));
+		CHECK_MESSAGE(*FString::Printf(TEXT("Value:%s, Expected:%s"), *NewValue.ToString(), *Value.ToString()), NewValue == Value);
+		// Missing value == 0 :
+		CHECK(NewValue.InitFromString(TEXT("X= Y=1")));
+		CHECK_MESSAGE(*FString::Printf(TEXT("Value:%s, Expected:%s"), *NewValue.ToString(), *Value.ToString()), NewValue == Value);
+		CHECK(NewValue.InitFromString(TEXT("X=Y=1")));
+		CHECK_MESSAGE(*FString::Printf(TEXT("Value:%s, Expected:%s"), *NewValue.ToString(), *Value.ToString()), NewValue == Value);
+		CHECK(NewValue.InitFromString(TEXT("X=A Y=1")));
+		CHECK_MESSAGE(*FString::Printf(TEXT("Value:%s, Expected:%s"), *NewValue.ToString(), *Value.ToString()), NewValue == Value);
+		// Invalid formats : 
+		CHECK(!NewValue.InitFromString(TEXT("RA= Y=1")));
+		CHECK(!NewValue.InitFromString(TEXT("X=0Y=1")));
+		CHECK(!NewValue.InitFromString(TEXT("X =0 Y=1")));
+		// Missing component should yield an error :
+		CHECK(!NewValue.InitFromString(TEXT("X=0")));
+		CHECK(!NewValue.InitFromString(TEXT("X=0 A=1")));
+	}
+
+	SECTION("FIntVector")
+	{
+		FIntVector Value(0, 1, 2);
+		FString Expected(TEXT("X=0 Y=1 Z=2"));
+		CHECK(Value.ToString() == Expected);
+
+		// Back-and-forth conversion should work :
+		FIntVector NewValue;
+		CHECK(NewValue.InitFromString(Value.ToString()));
+		CHECK_MESSAGE(*FString::Printf(TEXT("Value:%s, Expected:%s"), *NewValue.ToString(), *Value.ToString()), NewValue == Value);
+
+		// Permissive formatting should work :
+		CHECK(NewValue.InitFromString(TEXT("X=0     ,Y= 1.000000.2:Z=2")));
+		CHECK_MESSAGE(*FString::Printf(TEXT("Value:%s, Expected:%s"), *NewValue.ToString(), *Value.ToString()), NewValue == Value);
+		// Out-of-order parameters :
+		CHECK(NewValue.InitFromString(TEXT("Y=1 Z=2 X=0")));
+		CHECK_MESSAGE(*FString::Printf(TEXT("Value:%s, Expected:%s"), *NewValue.ToString(), *Value.ToString()), NewValue == Value);
+		// Various formats/delimiters :
+		CHECK(NewValue.InitFromString(TEXT("X=.0;Y=1|Z=2")));
+		CHECK_MESSAGE(*FString::Printf(TEXT("Value:%s, Expected:%s"), *NewValue.ToString(), *Value.ToString()), NewValue == Value);
+		CHECK(NewValue.InitFromString(TEXT("X=0A Y=1A Z= 2")));
+		CHECK_MESSAGE(*FString::Printf(TEXT("Value:%s, Expected:%s"), *NewValue.ToString(), *Value.ToString()), NewValue == Value);
+		// Missing value == 0:
+		CHECK(NewValue.InitFromString(TEXT("X= Y=1 Z=2")));
+		CHECK_MESSAGE(*FString::Printf(TEXT("Value:%s, Expected:%s"), *NewValue.ToString(), *Value.ToString()), NewValue == Value);
+		CHECK(NewValue.InitFromString(TEXT("X=Y=1 Z=2")));
+		CHECK_MESSAGE(*FString::Printf(TEXT("Value:%s, Expected:%s"), *NewValue.ToString(), *Value.ToString()), NewValue == Value);
+		CHECK(NewValue.InitFromString(TEXT("X=A Y=1 Z=2")));
+		CHECK_MESSAGE(*FString::Printf(TEXT("Value:%s, Expected:%s"), *NewValue.ToString(), *Value.ToString()), NewValue == Value);
+		// Invalid formats : 
+		CHECK(!NewValue.InitFromString(TEXT("XA= Y=1 Z=2")));
+		CHECK(!NewValue.InitFromString(TEXT("X=0Y=1Z=2")));
+		CHECK(!NewValue.InitFromString(TEXT("X =0 Y=1 Z=2")));
+		// Missing component should yield an error :
+		CHECK(!NewValue.InitFromString(TEXT("X=0 Y=1")));
+		CHECK(!NewValue.InitFromString(TEXT("X=0 Y=1 A=2")));
+	}
+}
 
 TEST_CASE("Parse::GrammaredCLIParse::Callback", "[Smoke]")
 {
