Added thread manager interface for retrieving callstacks without incurring additional allocations (which will be used to avoid failing to allocate memory while crashing out of memory).

diff --git a/Engine/Source/Runtime/Core/Private/HAL/ThreadingBase.cpp b/Engine/Source/Runtime/Core/Private/HAL/ThreadingBase.cpp
--- a/Engine/Source/Runtime/Core/Private/HAL/ThreadingBase.cpp
+++ b/Engine/Source/Runtime/Core/Private/HAL/ThreadingBase.cpp
@@ -521,27 +521,33 @@ const FString& FThreadManager::GetThreadNameInternal(uint32 ThreadId)
 }
 
 #if PLATFORM_SUPPORTS_ALL_THREAD_BACKTRACES
-static void GetAllThreadStackBackTraces_ProcessSingle(
-	uint32 CurThreadId,
-	uint32 ThreadId,
-	const TCHAR* ThreadName,
-	typename FThreadManager::FThreadStackBackTrace& OutStackTrace)
+static TConstArrayView<uint64> ThreadStackBackTraces_PerformStackWalk(uint32 CurThreadId, uint32 ThreadId, TArrayView<uint64> OutProgramCounters)
 {
-	constexpr uint32 MaxDepth = 100;
-	OutStackTrace.ThreadId = ThreadId;
-	OutStackTrace.ThreadName = ThreadName;
-	auto& PCs = OutStackTrace.ProgramCounters;
-	PCs.AddZeroed(MaxDepth);
 	uint32 Depth;
 	if (CurThreadId != ThreadId)
 	{
-		Depth = FPlatformStackWalk::CaptureThreadStackBackTrace(ThreadId, PCs.GetData(), MaxDepth);
+		Depth = FPlatformStackWalk::CaptureThreadStackBackTrace(ThreadId, OutProgramCounters.GetData(), OutProgramCounters.Num());
 	}
 	else
 	{
-		Depth = FPlatformStackWalk::CaptureStackBackTrace(PCs.GetData(), MaxDepth);
+		Depth = FPlatformStackWalk::CaptureStackBackTrace(OutProgramCounters.GetData(), OutProgramCounters.Num());
 	}
-	PCs.SetNum(Depth);
+	return OutProgramCounters.Left(Depth);
+}
+
+static void GetAllThreadStackBackTraces_ProcessSingle(
+	uint32 CurThreadId,
+	uint32 ThreadId,
+	const TCHAR* ThreadName,
+	typename FThreadManager::FThreadStackBackTrace& OutStackTrace)
+{
+	OutStackTrace.ThreadId = ThreadId;
+	OutStackTrace.ThreadName = ThreadName;
+
+	FThreadManager::FThreadStackBackTrace::FProgramCountersArray& PCs = OutStackTrace.ProgramCounters;
+	PCs.SetNumZeroed(FThreadManager::FThreadStackBackTrace::ProgramCountersMaxStackSize);
+	const TConstArrayView<uint64> WrittenProgramCounters = ThreadStackBackTraces_PerformStackWalk(CurThreadId, ThreadId, PCs);
+	PCs.SetNum(WrittenProgramCounters.Num());
 }
 
 void FThreadManager::GetAllThreadStackBackTraces(TArray<FThreadStackBackTrace>& StackTraces)
@@ -561,6 +567,38 @@ void FThreadManager::GetAllThreadStackBackTraces(TArray<FThreadStackBackTrace>&
 		}
 	);
 }
+
+void FThreadManager::ForEachThreadStackBackTrace(TFunctionRef<bool(uint32 ThreadId, const TCHAR* ThreadName, const TConstArrayView<uint64>& StackTrace)> Func)
+{
+	const uint32 CurThreadId = FPlatformTLS::GetCurrentThreadId();
+	FThreadStackBackTrace::FProgramCountersArray ProgramCounterBuffer;
+	ProgramCounterBuffer.SetNumZeroed(FThreadStackBackTrace::ProgramCountersMaxStackSize);
+
+	FScopeLock Lock(&ThreadsCritical);
+
+	{
+		const TConstArrayView<uint64> ProgramCounterOutputArrayView = ThreadStackBackTraces_PerformStackWalk(CurThreadId, GGameThreadId, ProgramCounterBuffer);
+		const bool bContinue = Func(GGameThreadId, TEXT("GameThread"), ProgramCounterOutputArrayView);
+		if (!bContinue)
+		{
+			return;
+		}
+	}
+
+	// Not using ForEachThread since that creates a copy, just verify the callback isn't modifying thread list
+	for (const TPair<uint32, FRunnableThread*>& Pair : Threads)
+	{
+		const uint32 ThreadId = Pair.Key;
+		const FRunnableThread* Thread = Pair.Value;
+		const FString& ThreadName = Thread->GetThreadName();
+		const TConstArrayView<uint64> ProgramCounterOutputArrayView = ThreadStackBackTraces_PerformStackWalk(CurThreadId, ThreadId, ProgramCounterBuffer);
+		const bool bContinue = Func(ThreadId, *ThreadName, ProgramCounterOutputArrayView);
+		if (!bContinue)
+		{
+			return;
+		}
+	}
+}
 #endif
 
 void FThreadManager::ForEachThread(TFunction<void(uint32, FRunnableThread*)> Func)
diff --git a/Engine/Source/Runtime/Core/Public/HAL/ThreadManager.h b/Engine/Source/Runtime/Core/Public/HAL/ThreadManager.h
--- a/Engine/Source/Runtime/Core/Public/HAL/ThreadManager.h
+++ b/Engine/Source/Runtime/Core/Public/HAL/ThreadManager.h
@@ -70,12 +70,29 @@ public:
 #if PLATFORM_SUPPORTS_ALL_THREAD_BACKTRACES
 	struct FThreadStackBackTrace
 	{
+		static constexpr uint32 ProgramCountersMaxStackSize = 100;
+		typedef TArray<uint64, TFixedAllocator<ProgramCountersMaxStackSize>> FProgramCountersArray;
+
 		uint32 ThreadId;
 		FString ThreadName;
-		TArray<uint64, TInlineAllocator<100>> ProgramCounters;
+		FProgramCountersArray ProgramCounters;
 	};
 
 	CORE_API void GetAllThreadStackBackTraces(TArray<FThreadStackBackTrace>& StackTraces);
+
+	/*
+	 * Enumerates through all thread stack backtraces and calls the provided function for each one.
+	 * The callback must return true to continue enumerating, or return false to stop early.
+	 *
+	 * This function is primarily intended to iterate over stack traces in a crashing context and
+	 * avoids allocation of additional memory. It does not, therefore, perform any safety checks to
+	 * ensure that the list of threads is not modified mid-iteration or that the callback does not
+	 * itself allocate memory.
+	 *
+	 * Similarly, the thread name and stack trace array memory are only valid for the duration of the
+	 * callback's execution and must be copied elsewhere if they are to be used beyond its scope.
+	 */
+	CORE_API void ForEachThreadStackBackTrace(TFunctionRef<bool(uint32 ThreadId, const TCHAR* ThreadName, const TConstArrayView<uint64>& StackTrace)> Func);
 #endif
 
 	/**
