Fixing the way that we parse ini config overrides on the command line. This should now properly handle quoted values so long as they are not mixed with commas. Unfortunately, supporting that combination is significantly more complex. Instead, it is recommended that each key/value pair be specified via its own -ini argument.

#tests Verified various combinations of command line ini settings with and without quotes and commas

diff --git a/Engine/Source/Runtime/Core/Private/Misc/ConfigCacheIni.cpp b/Engine/Source/Runtime/Core/Private/Misc/ConfigCacheIni.cpp
--- a/Engine/Source/Runtime/Core/Private/Misc/ConfigCacheIni.cpp
+++ b/Engine/Source/Runtime/Core/Private/Misc/ConfigCacheIni.cpp
@@ -21,7 +21,7 @@
 #include "Misc/ConfigManifest.h"
 #include "Misc/DataDrivenPlatformInfoRegistry.h"
 #include "Misc/StringBuilder.h"
-#include "Misc/Paths.h"
+#include "Misc/PathViews.h"
 #include "ProfilingDebugging/CpuProfilerTrace.h"
 #include "ProfilingDebugging/AssetMetadataTrace.h"
 #include "Serialization/MemoryReader.h"
@@ -1290,7 +1290,7 @@ namespace CommandlineOverrideSpecifiers
 	const auto& IniNameEndIdentifier      = TEXT(":[");
 	const auto& SectionStartIdentifier    = TEXT("[");
 	const auto& PropertyStartIdentifier   = TEXT("]:");
-	const auto& PropertySeperator         = TEXT(",");
+	const auto& PropertySeperator         = TEXT(','); //NOTE: this must be a single character
 	const auto& CustomConfigIdentifier    = TEXT("-CustomConfig=");
 }
 
@@ -1348,54 +1348,89 @@ bool FConfigFile::OverrideFileFromCommandline(FString& Filename)
 void FConfigFile::OverrideFromCommandline(FConfigFile* File, const FString& Filename)
 {
 #if ALLOW_INI_OVERRIDE_FROM_COMMANDLINE
-	FString Settings;
 	// look for this filename on the commandline in the format:
-	//		-ini:IniName:[Section1]:Key1=Value1,[Section2]:Key2=Value2
+	//		-ini:IniName:[Section1]:Key=Value
 	// for example:
-	//		-ini:Engine:[/Script/Engine.Engine]:bSmoothFrameRate=False,[TextureStreaming]:PoolSize=100
+	//		-ini:Engine:[/Script/Engine.Engine]:bSmoothFrameRate=False
 	//			(will update the cache after the final combined engine.ini)
-	const TCHAR* CommandlineStream = FCommandLine::Get();
-	while(FParse::Value(CommandlineStream, *FString::Printf(TEXT("%s%s"), CommandlineOverrideSpecifiers::IniSwitchIdentifier, *FPaths::GetBaseFilename(Filename)), Settings, false))
+
+	TStringBuilder<260> IniSwitchStringBuilder;
+	IniSwitchStringBuilder.Append(CommandlineOverrideSpecifiers::IniSwitchIdentifier);
+	IniSwitchStringBuilder.Append(FPathViews::GetBaseFilename(Filename));
+	IniSwitchStringBuilder.Append(TEXT(":")); // Ensure we only match the exact filename
+
+	// Initial search to early out if the -ini:IniName: pattern doesn't exist anywhere in the string
+	// Cannot use find result directly as text can be found inside another argument
+	if (FCString::Strifind(FCommandLine::Get(), *IniSwitchStringBuilder, true) == nullptr)
 	{
-		// break apart on the commas
-		TArray<FString> SettingPairs;
-		Settings.ParseIntoArray(SettingPairs, CommandlineOverrideSpecifiers::PropertySeperator, true);
-		for (int32 Index = 0; Index < SettingPairs.Num(); Index++)
+		return;
+	}
+
+	// Null terminate
+	const FStringView IniSwitch = *IniSwitchStringBuilder;
+	const TCHAR* RemainingCommandLineStream = FCommandLine::Get();
+
+	FString NextCommandLineArgumentToken;
+	while(FParse::Token(RemainingCommandLineStream, NextCommandLineArgumentToken, /*bUseEscape=*/false))
+	{
+		if (NextCommandLineArgumentToken.StartsWith(IniSwitch))
 		{
-			// set each one, by splitting on the =
-			FString SectionAndKey, Value;
-			if (SettingPairs[Index].Split(TEXT("="), &SectionAndKey, &Value))
+			FString SettingsString = NextCommandLineArgumentToken.RightChop(IniSwitch.Len());
+
+			// break apart on the commas. WARNING: This is supported for legacy reasons only
+			// Providing multiple key-value pairs in a single -ini argument breaks when combined 
+			// with quoted values. Fixing this is non-trivial and likely platform dependent.
+			TArray<FString> SettingPairs;
 			{
-				// now we need to split off the key from the rest of the section name
-				int32 SectionNameEndIndex = SectionAndKey.Find(CommandlineOverrideSpecifiers::PropertyStartIdentifier, ESearchCase::IgnoreCase, ESearchDir::FromEnd);
-				// check for malformed string
-				if (SectionNameEndIndex == INDEX_NONE || SectionNameEndIndex == 0)
+				FString NextSettingToken;
+				const TCHAR* SettingsStream = *SettingsString;
+				while (FParse::Token(SettingsStream, NextSettingToken, false, CommandlineOverrideSpecifiers::PropertySeperator))
 				{
-					continue;
+					SettingPairs.Add(MoveTemp(NextSettingToken));
 				}
+			}
 
-				// Create the commandline override object
-				FConfigCommandlineOverride& CommandlineOption = File->CommandlineOptions[File->CommandlineOptions.Emplace()];
-				CommandlineOption.BaseFileName = *FPaths::GetBaseFilename(Filename);
-				CommandlineOption.Section = SectionAndKey.Left(SectionNameEndIndex);
-				
-				// Remove commandline syntax from the section name.
-				CommandlineOption.Section = CommandlineOption.Section.Replace(CommandlineOverrideSpecifiers::IniNameEndIdentifier, TEXT(""));
-				CommandlineOption.Section = CommandlineOption.Section.Replace(CommandlineOverrideSpecifiers::PropertyStartIdentifier, TEXT(""));
-				CommandlineOption.Section = CommandlineOption.Section.Replace(CommandlineOverrideSpecifiers::SectionStartIdentifier, TEXT(""));
+			for (int32 Index = 0; Index < SettingPairs.Num(); Index++)
+			{
+				// set each one, by splitting on the =
+				FString SectionAndKey, Value;
+				if (SettingPairs[Index].Split(TEXT("="), &SectionAndKey, &Value))
+				{
+					// now we need to split off the key from the rest of the section name
+					int32 SectionNameEndIndex = SectionAndKey.Find(CommandlineOverrideSpecifiers::PropertyStartIdentifier, ESearchCase::IgnoreCase, ESearchDir::FromEnd);
+					// check for malformed string
+					if (SectionNameEndIndex == INDEX_NONE || SectionNameEndIndex == 0)
+					{
+						continue;
+					}
 
-				CommandlineOption.PropertyKey = SectionAndKey.Mid(SectionNameEndIndex + UE_ARRAY_COUNT(CommandlineOverrideSpecifiers::PropertyStartIdentifier) - 1);
-				CommandlineOption.PropertyValue = Value;
+					// Create the commandline override object
+					FConfigCommandlineOverride& CommandlineOption = File->CommandlineOptions[File->CommandlineOptions.Emplace()];
+					CommandlineOption.BaseFileName = *FPaths::GetBaseFilename(Filename);
+					CommandlineOption.Section = SectionAndKey.Left(SectionNameEndIndex);
+					
+					// Remove commandline syntax from the section name.
+					CommandlineOption.Section = CommandlineOption.Section.Replace(CommandlineOverrideSpecifiers::IniNameEndIdentifier, TEXT(""));
+					CommandlineOption.Section = CommandlineOption.Section.Replace(CommandlineOverrideSpecifiers::PropertyStartIdentifier, TEXT(""));
+					CommandlineOption.Section = CommandlineOption.Section.Replace(CommandlineOverrideSpecifiers::SectionStartIdentifier, TEXT(""));
+
+					CommandlineOption.PropertyKey = SectionAndKey.Mid(SectionNameEndIndex + UE_ARRAY_COUNT(CommandlineOverrideSpecifiers::PropertyStartIdentifier) - 1);
+					
+					// If the property value was quoted, remove the quotes
+					if (Value.Len() > 1 && Value.StartsWith(TEXT("\"")) && Value.EndsWith(TEXT("\"")))
+					{
+						CommandlineOption.PropertyValue = Value.Mid(1, Value.Len() - 2);
+					}
+					else
+					{
+						CommandlineOption.PropertyValue = Value;
+					}
 
-				// now put it into this into the cache
-				File->SetString(*CommandlineOption.Section, *CommandlineOption.PropertyKey, *CommandlineOption.PropertyValue);
+					// now put it into this into the cache
+					File->SetString(*CommandlineOption.Section, *CommandlineOption.PropertyKey, *CommandlineOption.PropertyValue);
+				}
 			}
 		}
-
-		// Keep searching for more instances of -ini
-		CommandlineStream = FCString::Stristr(CommandlineStream, CommandlineOverrideSpecifiers::IniSwitchIdentifier);
-		check(CommandlineStream);
-		CommandlineStream++;
 	}
 #endif
 }
diff --git a/Engine/Source/Runtime/Core/Private/Misc/Parse.cpp b/Engine/Source/Runtime/Core/Private/Misc/Parse.cpp
--- a/Engine/Source/Runtime/Core/Private/Misc/Parse.cpp
+++ b/Engine/Source/Runtime/Core/Private/Misc/Parse.cpp
@@ -792,10 +792,10 @@ void FParse::Next( const TCHAR** Stream )
 }
 
 //
-// Grab the next space-delimited string from the input stream.
+// Grab the next space-delimited (or SingleCharacterDelimiter-delimited) string from the input stream.
 // If quoted, gets entire quoted string.
 //
-bool FParse::Token( const TCHAR*& Str, TCHAR* Result, int32 MaxLen, bool UseEscape )
+bool FParse::Token( const TCHAR*& Str, TCHAR* Result, int32 MaxLen, bool UseEscape, const TCHAR SingleCharacterDelimiter/* = TEXT('\0')*/)
 {
 	int32 Len=0;
 
@@ -840,10 +840,21 @@ bool FParse::Token( const TCHAR*& Str, TCHAR* Result, int32 MaxLen, bool UseEsca
 		while (1)
 		{
 			TCHAR Character = *Str;
-			if ((Character == 0) || (FChar::IsWhitespace(Character) && !bInQuote))
+			if (Character == 0)
 			{
 				break;
 			}
+			if (!bInQuote)
+			{
+				if ((SingleCharacterDelimiter != TEXT('\0') && Character == SingleCharacterDelimiter)
+					|| (SingleCharacterDelimiter == TEXT('\0') && (FChar::IsWhitespace(Character))))
+				{
+					// Consume the delimiter. If it's whitespace this isn't critical since we'll consume it at the start
+					// of the next call to Token() but if it's not whitespace we won't, so we better do it now.
+					Str++;
+					break;
+				}
+			}
 			Str++;
 
 			// Preserve escapes if they're in a quoted string (the check for " is in the else to let \" work as expected)
@@ -876,7 +887,7 @@ bool FParse::Token( const TCHAR*& Str, TCHAR* Result, int32 MaxLen, bool UseEsca
 	return Len != 0;
 }
 
-bool FParse::Token( const TCHAR*& Str, FString& Arg, bool UseEscape )
+bool FParse::Token( const TCHAR*& Str, FString& Arg, bool UseEscape, const TCHAR SingleCharacterDelimiter/* = TEXT('\0')*/)
 {
 	Arg.Reset();
 
@@ -920,10 +931,21 @@ bool FParse::Token( const TCHAR*& Str, FString& Arg, bool UseEscape )
 		while (1)
 		{
 			TCHAR Character = *Str;
-			if ((Character == 0) || (FChar::IsWhitespace(Character) && !bInQuote))
+			if (Character == 0)
 			{
 				break;
 			}
+			if (!bInQuote)
+			{
+				if ((SingleCharacterDelimiter != TEXT('\0') && Character == SingleCharacterDelimiter)
+					|| (SingleCharacterDelimiter == TEXT('\0') && (FChar::IsWhitespace(Character))))
+				{
+					// Consume the delimiter. If it's whitespace this isn't critical since we'll consume it at the start
+					// of the next call to Token() but if it's not whitespace we won't, so we better do it now.
+					Str++;
+					break;
+				}
+			}
 			Str++;
 
 			// Preserve escapes if they're in a quoted string (the check for " is in the else to let \" work as expected)
diff --git a/Engine/Source/Runtime/Core/Public/Misc/Parse.h b/Engine/Source/Runtime/Core/Public/Misc/Parse.h
--- a/Engine/Source/Runtime/Core/Public/Misc/Parse.h
+++ b/Engine/Source/Runtime/Core/Public/Misc/Parse.h
@@ -87,10 +87,10 @@ struct FParse
 	 */
 	static CORE_API bool LineExtended(const TCHAR** Stream, FString& Result, int32& LinesConsumed, bool Exact = 0);
 	static CORE_API bool LineExtended(const TCHAR** Stream, FStringBuilderBase& Result, int32& LinesConsumed, bool Exact = 0);
-	/** Grabs the next space-delimited string from the input stream. If quoted, gets entire quoted string. */
-	static CORE_API bool Token( const TCHAR*& Str, TCHAR* Result, int32 MaxLen, bool UseEscape );
-	/** Grabs the next space-delimited string from the input stream. If quoted, gets entire quoted string. */
-	static CORE_API bool Token( const TCHAR*& Str, FString& Arg, bool UseEscape );
+	/** Grabs the next space-delimited (unless SingleCharacterDelimiter is provided) string from the input stream. If quoted, gets entire quoted string. */
+	static CORE_API bool Token( const TCHAR*& Str, TCHAR* Result, int32 MaxLen, bool UseEscape, const TCHAR SingleCharacterDelimiter = TEXT('\0'));
+	/** Grabs the next space-delimited (unless SingleCharacterDelimiter is provided) string from the input stream. If quoted, gets entire quoted string. */
+	static CORE_API bool Token( const TCHAR*& Str, FString& Arg, bool UseEscape, const TCHAR SingleCharacterDelimiter = TEXT('\0'));
 	/** Grabs the next alpha-numeric space-delimited token from the input stream. */
 	static CORE_API bool AlnumToken(const TCHAR*& Str, FString& Arg);
 	/** Grabs the next space-delimited string from the input stream. If quoted, gets entire quoted string. */
