Add -fooLogDuplicates param to allow logging of duplicate foo entries for debug/analysis purposes
#tests Tested on Switch

diff --git a/Engine/Source/Runtime/Core/Private/HAL/IPlatformFileOpenLogWrapper.cpp b/Engine/Source/Runtime/Core/Private/HAL/IPlatformFileOpenLogWrapper.cpp
--- a/Engine/Source/Runtime/Core/Private/HAL/IPlatformFileOpenLogWrapper.cpp
+++ b/Engine/Source/Runtime/Core/Private/HAL/IPlatformFileOpenLogWrapper.cpp
@@ -96,6 +96,10 @@ bool FPlatformFileOpenLog::Initialize(IPlatformFile* Inner, const TCHAR* Command
 		}
 		UE_LOG(LogFileOpenOrder, Log, TEXT("Initialized file open order log : %s"), *LogFilePath);
 	}
+
+	// Log duplicates if requested. This can be used for debugging/analysis purposes to log file access order but shouldn't be used to generate a pak/iostore ordering
+	bLogDuplicates = FParse::Param(CommandLineParam, TEXT("FooLogDuplicates"));
+
 	GPlatformFileOpenLog = this;
 	return true;
 }
diff --git a/Engine/Source/Runtime/Core/Public/HAL/IPlatformFileOpenLogWrapper.h b/Engine/Source/Runtime/Core/Public/HAL/IPlatformFileOpenLogWrapper.h
--- a/Engine/Source/Runtime/Core/Public/HAL/IPlatformFileOpenLogWrapper.h
+++ b/Engine/Source/Runtime/Core/Public/HAL/IPlatformFileOpenLogWrapper.h
@@ -60,12 +60,14 @@ protected:
 	int64					OpenOrder;
 	TMap<FString, int64>	FilenameAccessMap;
 	TArray<IFileHandle*>	LogOutput;
+	bool					bLogDuplicates;
 
 public:
 
 	FPlatformFileOpenLog()
 		: LowerLevel(nullptr)
 		, OpenOrder(0)
+		, bLogDuplicates(false)
 	{
 	}
 
@@ -229,7 +231,15 @@ public:
 	void AddToOpenLog(const TCHAR* Filename)
 	{
 		CriticalSection.Lock();
-		if (FilenameAccessMap.Find(Filename) == nullptr)
+		if (bLogDuplicates)
+		{
+			FString Text = FString::Printf(TEXT("\"%s\"\n"), Filename);
+			for (auto File = LogOutput.CreateIterator(); File; ++File)
+			{
+				(*File)->Write((uint8*)StringCast<ANSICHAR>(*Text).Get(), Text.Len());
+			}
+		}
+		else if (FilenameAccessMap.Find(Filename) == nullptr)
 		{
 			FilenameAccessMap.Emplace(Filename, ++OpenOrder);
 			FString Text = FString::Printf(TEXT("\"%s\" %llu\n"), Filename, OpenOrder);
@@ -255,17 +265,8 @@ public:
 
 	void AddPackageToOpenLog(const TCHAR* Filename)
 	{
-		CriticalSection.Lock();
-		if (FilenameAccessMap.Find(Filename) == nullptr)
-		{
-			FilenameAccessMap.Emplace(Filename, ++OpenOrder);
-			FString Text = FString::Printf(TEXT("\"%s\" %llu\n"), Filename, OpenOrder);
-			for (auto File = LogOutput.CreateIterator(); File; ++File)
-			{
-				(*File)->Write((uint8*)StringCast<ANSICHAR>(*Text).Get(), Text.Len());
-			}
-		}
-		CriticalSection.Unlock();
+		// TODO: deprecate/remove this function?
+		AddToOpenLog(Filename);
 	}
 
 	template <typename FmtType, typename... Types>
