Fixing static analysis warning related to auto& on a TCHAR constant.

diff --git a/Engine/Source/Runtime/Core/Private/Misc/ConfigCacheIni.cpp b/Engine/Source/Runtime/Core/Private/Misc/ConfigCacheIni.cpp
--- a/Engine/Source/Runtime/Core/Private/Misc/ConfigCacheIni.cpp
+++ b/Engine/Source/Runtime/Core/Private/Misc/ConfigCacheIni.cpp
@@ -1290,7 +1290,7 @@ namespace CommandlineOverrideSpecifiers
 	const auto& IniNameEndIdentifier      = TEXT(":[");
 	const auto& SectionStartIdentifier    = TEXT("[");
 	const auto& PropertyStartIdentifier   = TEXT("]:");
-	const auto& PropertySeperator         = TEXT(','); //NOTE: this must be a single character
+	const TCHAR PropertySeperator         = TEXT(','); //NOTE: this must be a single character
 	const auto& CustomConfigIdentifier    = TEXT("-CustomConfig=");
 }
 
