[MTAccessDetector]
 - Fix missing declaration

diff --git a/Engine/Source/Runtime/Core/Public/Misc/MTAccessDetector.h b/Engine/Source/Runtime/Core/Public/Misc/MTAccessDetector.h
--- a/Engine/Source/Runtime/Core/Public/Misc/MTAccessDetector.h
+++ b/Engine/Source/Runtime/Core/Public/Misc/MTAccessDetector.h
@@ -665,6 +665,7 @@ public:
 #define UE_MT_DECLARE_RW_ACCESS_DETECTOR(AccessDetector)
 #define UE_MT_DECLARE_RW_RECURSIVE_ACCESS_DETECTOR(AccessDetector)
 #define UE_MT_DECLARE_RW_FULLY_RECURSIVE_ACCESS_DETECTOR(AccessDetector)
+#define UE_MT_DECLARE_MRSW_RECURSIVE_ACCESS_DETECTOR(AccessDetector)
 
 #define UE_MT_SCOPED_READ_ACCESS(AccessDetector) 
 #define UE_MT_SCOPED_WRITE_ACCESS(AccessDetector)
