Fallback path to standard stackwalking for memory tracing

In some cases our unwinding tables approach doesn't seem to work (e.g. compiled with clang). In those cases fall back to the standard stack walking methods.

diff --git a/Engine/Source/Runtime/Core/Private/ProfilingDebugging/Microsoft/WindowsCallstackTrace.cpp b/Engine/Source/Runtime/Core/Private/ProfilingDebugging/Microsoft/WindowsCallstackTrace.cpp
--- a/Engine/Source/Runtime/Core/Private/ProfilingDebugging/Microsoft/WindowsCallstackTrace.cpp
+++ b/Engine/Source/Runtime/Core/Private/ProfilingDebugging/Microsoft/WindowsCallstackTrace.cpp
@@ -36,6 +36,22 @@
 
 #define BACKTRACE_LOCK_FREE (1 && (BACKTRACE_DBGLVL == 0))
 
+static bool GFullBacktraces = UE_CALLSTACK_TRACE_FULL_CALLSTACKS;
+static bool GModulesAreInitialized = false;
+
+// This implementation is using unwind tables which is results in very fast
+// stack walking. In some cases this is not suitable, and we then fall back
+// to the standard stack walking implementation.
+#if !defined(UE_CALLSTACK_TRACE_USE_UNWIND_TABLES)
+	#if defined(__clang__)
+		#define UE_CALLSTACK_TRACE_USE_UNWIND_TABLES 0
+	#else
+		#define UE_CALLSTACK_TRACE_USE_UNWIND_TABLES 1
+	#endif
+#endif
+
+#if UE_CALLSTACK_TRACE_USE_UNWIND_TABLES
+
 /*
  * Windows' x64 binaries contain a ".pdata" section that describes the location
  * and size of its functions and details on how to unwind them. The unwind
@@ -119,8 +135,7 @@ public:
 	static FBacktracer*		Get();
 	void					AddModule(UPTRINT Base, const TCHAR* Name);
 	void					RemoveModule(UPTRINT Base);
-	uint32					GetBacktraceId(void* AddressOfReturnAddress) const;
-	int32					GetFrameSize(void* FunctionAddress) const;
+	uint32					GetBacktraceId(void* AddressOfReturnAddress);
 
 private:
 	struct FFunction
@@ -190,7 +205,7 @@ private:
 	int32					ModulesNum;
 	int32					ModulesCapacity;
 	FMalloc*				Malloc;
-	FCallstackTracer*		CallstackTracer;
+	FCallstackTracer		CallstackTracer;
 #if BACKTRACE_LOCK_FREE
 	mutable FFunctionLookupSet	FunctionLookups;
 	mutable bool				bReentranceCheck = false;
@@ -203,12 +218,11 @@ private:
 
 ////////////////////////////////////////////////////////////////////////////////
 FBacktracer* FBacktracer::Instance = nullptr;
-static bool GFullBacktraces = UE_CALLSTACK_TRACE_FULL_CALLSTACKS;
-static bool GModulesAreInitialized = false;
 
 ////////////////////////////////////////////////////////////////////////////////
 FBacktracer::FBacktracer(FMalloc* InMalloc)
 	: Malloc(InMalloc)
+	, CallstackTracer(InMalloc)
 #if BACKTRACE_LOCK_FREE
 	, FunctionLookups(InMalloc)
 #endif
@@ -221,8 +235,6 @@ FBacktracer::FBacktracer(FMalloc* InMalloc)
 	Modules = (FModule*)Malloc->Malloc(sizeof(FModule) * ModulesCapacity);
 
 	Instance = this;
-	CallstackTracer = (FCallstackTracer*) Malloc->Malloc(sizeof(FCallstackTracer), alignof(FCallstackTracer));
-	CallstackTracer = new(CallstackTracer) FCallstackTracer(InMalloc);
 }
 
 ////////////////////////////////////////////////////////////////////////////////
@@ -570,7 +582,7 @@ const FBacktracer::FFunction* FBacktracer::LookupFunction(UPTRINT Address, FLook
 }
 
 ////////////////////////////////////////////////////////////////////////////////
-uint32 FBacktracer::GetBacktraceId(void* AddressOfReturnAddress) const
+uint32 FBacktracer::GetBacktraceId(void* AddressOfReturnAddress) 
 {
 	FLookupState LookupState = {};
 	uint64 Frames[256];
@@ -717,28 +729,98 @@ uint32 FBacktracer::GetBacktraceId(void* AddressOfReturnAddress) const
 #endif
 	// Add to queue to be processed. This might block until there is room in the
 	// queue (i.e. the processing thread has caught up processing).
-	return CallstackTracer->AddCallstack(BacktraceEntry);
+	return CallstackTracer.AddCallstack(BacktraceEntry);
 }
 
+#else // UE_CALLSTACK_TRACE_USE_UNWIND_TABLES
+
 ////////////////////////////////////////////////////////////////////////////////
-int32 FBacktracer::GetFrameSize(void* FunctionAddress) const
+class FBacktracer
 {
-	FScopeLock _(&Lock);
+public:
+	FBacktracer(FMalloc* InMalloc);
+	~FBacktracer();
+	static FBacktracer*	Get();
+	uint32 GetBacktraceId(void* AddressOfReturnAddress);
+	void AddModule(UPTRINT Base, const TCHAR* Name) {}
+	void RemoveModule(UPTRINT Base) {}
 
-	FLookupState LookupState = {};
-	const FFunction* Function = LookupFunction(UPTRINT(FunctionAddress), LookupState);
-	if (FunctionAddress == nullptr)
-	{
-		return -1;
-	}
+private:
+	static FBacktracer* Instance;
+	FMalloc* Malloc;
+	FCallstackTracer CallstackTracer;
+};
+
+////////////////////////////////////////////////////////////////////////////////
+FBacktracer* FBacktracer::Instance = nullptr;
+
+////////////////////////////////////////////////////////////////////////////////
+FBacktracer::FBacktracer(FMalloc* InMalloc)
+	: Malloc(InMalloc)
+	, CallstackTracer(InMalloc)
+{
+	Instance = this;
+}
 
-	if (Function->RspBias < 0)
+////////////////////////////////////////////////////////////////////////////////
+FBacktracer::~FBacktracer()
+{
+}
+
+////////////////////////////////////////////////////////////////////////////////
+FBacktracer* FBacktracer::Get()
+{
+	return Instance;
+}
+
+////////////////////////////////////////////////////////////////////////////////
+uint32 FBacktracer::GetBacktraceId(void* AddressOfReturnAddress) 
+{
+#if !UE_BUILD_SHIPPING
+	const uint64 ReturnAddress = *(uint64*)AddressOfReturnAddress;
+	uint64 StackFrames[256];
+	int32 NumStackFrames = FPlatformStackWalk::CaptureStackBackTrace(StackFrames, UE_ARRAY_COUNT(StackFrames));
+	if (NumStackFrames > 0)
 	{
-		return -1;
+		FCallstackTracer::FBacktraceEntry BacktraceEntry;
+		uint64 BacktraceId = 0;
+		uint32 FrameIdx = 0;
+		bool bUseAddress = false;
+		for (int32 Index = 0; Index < NumStackFrames; Index++)
+		{
+			if (!bUseAddress)
+			{
+				// start using backtrace only after ReturnAddress
+				if (StackFrames[Index] == (uint64)ReturnAddress)
+				{
+					bUseAddress = true;
+				}
+			}
+			if (bUseAddress || NumStackFrames == 1)
+			{
+				uint64 RetAddr = StackFrames[Index];
+				StackFrames[FrameIdx++] = RetAddr;
+
+				// This is a simple order-dependent LCG. Should be sufficient enough
+				BacktraceId += RetAddr;
+				BacktraceId *= 0x30be8efa499c249dull;
+			}
+		}
+
+		// Save the collected id
+		BacktraceEntry.Hash = BacktraceId;
+		BacktraceEntry.FrameCount = FrameIdx;
+		BacktraceEntry.Frames = StackFrames;
+
+		// Add to queue to be processed. This might block until there is room in the
+		// queue (i.e. the processing thread has caught up processing).
+		return CallstackTracer.AddCallstack(BacktraceEntry);
 	}
+#endif
 
-	return Function->RspBias;
+	return 0;
 }
+#endif // UE_CALLSTACK_TRACE_USE_UNWIND_TABLES
 
 ////////////////////////////////////////////////////////////////////////////////
 void Modules_Create(FMalloc*);
