Added FGenericPlatformMemory::GetStatsRaw that could be overriden per platform if needed. This function returns raw OS memory stats (on some i.e. on some platforms, GetStats alter available physical memory to take into account cached memory by various allocators)
LLM now calls GetStatsRaw to ensure "Untracked platform" and "Untracked" llm are correct
Removed PLATFORM_USE_CACHED_SLACK_MEMORY_IN_MEMORY_STATS define as it is not needed anymore

#rnx

diff --git a/Engine/Source/Runtime/Core/Private/GenericPlatform/GenericPlatformMemory.cpp b/Engine/Source/Runtime/Core/Private/GenericPlatform/GenericPlatformMemory.cpp
--- a/Engine/Source/Runtime/Core/Private/GenericPlatform/GenericPlatformMemory.cpp
+++ b/Engine/Source/Runtime/Core/Private/GenericPlatform/GenericPlatformMemory.cpp
@@ -274,6 +274,11 @@ FPlatformMemoryStats FGenericPlatformMemory::GetStats()
 	return FPlatformMemoryStats();
 }
 
+FPlatformMemoryStats FGenericPlatformMemory::GetStatsRaw()
+{
+	return FPlatformMemory::GetStats();
+}
+
 void FGenericPlatformMemory::GetStatsForMallocProfiler( FGenericMemoryStats& out_Stats )
 {
 #if	STATS
diff --git a/Engine/Source/Runtime/Core/Private/HAL/LowLevelMemTracker.cpp b/Engine/Source/Runtime/Core/Private/HAL/LowLevelMemTracker.cpp
--- a/Engine/Source/Runtime/Core/Private/HAL/LowLevelMemTracker.cpp
+++ b/Engine/Source/Runtime/Core/Private/HAL/LowLevelMemTracker.cpp
@@ -1924,7 +1924,7 @@ void FLowLevelMemTracker::TickInternal()
 
 	// Cache the amount of memory used early, since some of these functions (FindOrAddTagData) can
 	// cause allocations, which will throw the numbers off slightly.
-	FPlatformMemoryStats PlatformStats = FPlatformMemory::GetStats();
+	FPlatformMemoryStats PlatformStats = FPlatformMemory::GetStatsRaw();
 #if PLATFORM_DESKTOP
 	// virtual is working set + paged out memory.
 	const int64 PlatformProcessMemory = static_cast<int64>(PlatformStats.UsedVirtual);
diff --git a/Engine/Source/Runtime/Core/Private/HAL/MallocBinned2.cpp b/Engine/Source/Runtime/Core/Private/HAL/MallocBinned2.cpp
--- a/Engine/Source/Runtime/Core/Private/HAL/MallocBinned2.cpp
+++ b/Engine/Source/Runtime/Core/Private/HAL/MallocBinned2.cpp
@@ -1505,13 +1505,9 @@ int64 FMallocBinned2::GetTotalAllocatedSmallPoolMemory() const
 
 void FMallocBinned2::GetAllocatorStats( FGenericMemoryStats& OutStats )
 {
-#if PLATFORM_USE_CACHED_SLACK_MEMORY_IN_MEMORY_STATS || BINNED2_ALLOCATOR_STATS
-	NOALLOC_SCOPE_CYCLE_COUNTER(STAT_FMallocBinned2_GetAllocatorStats);
-
-	// Even if we have MB2 stats off, if we want to include cached slack, we need the PageAllocatorFreeCacheSize stat
+	// Even if we have MB2 stats off, cached slack needs to be included as it might be needed by some platforms
 	uint64 OSPageAllocatorCachedFreeSize = CachedOSPageAllocator.GetCachedFreeTotal();
 	OutStats.Add(TEXT("PageAllocatorFreeCacheSize"), OSPageAllocatorCachedFreeSize);
-#endif 
 
 #if BINNED2_ALLOCATOR_STATS
 	int64  TotalAllocatedSmallPoolMemory           = GetTotalAllocatedSmallPoolMemory();
diff --git a/Engine/Source/Runtime/Core/Public/GenericPlatform/GenericPlatformMemory.h b/Engine/Source/Runtime/Core/Public/GenericPlatform/GenericPlatformMemory.h
--- a/Engine/Source/Runtime/Core/Public/GenericPlatform/GenericPlatformMemory.h
+++ b/Engine/Source/Runtime/Core/Public/GenericPlatform/GenericPlatformMemory.h
@@ -348,10 +348,15 @@ struct FGenericPlatformMemory
 	static CORE_API FMalloc* BaseAllocator();
 
 	/**
-	 * @return platform specific current memory statistics.
+	 * @return platform specific current memory statistics. Note: On some platforms, unused allocator cached memory is taken into account in AvailablePhysical. 
 	 */
 	static CORE_API FPlatformMemoryStats GetStats();
 
+	/**
+	 * @return platform specific raw stats.
+	 */
+	static CORE_API FPlatformMemoryStats GetStatsRaw();
+
 	/**
 	* @return memory used for platforms that can do it quickly (without affecting stat unit much)
 	*/
diff --git a/Engine/Source/Runtime/Core/Public/HAL/Platform.h b/Engine/Source/Runtime/Core/Public/HAL/Platform.h
--- a/Engine/Source/Runtime/Core/Public/HAL/Platform.h
+++ b/Engine/Source/Runtime/Core/Public/HAL/Platform.h
@@ -533,10 +533,6 @@
 	#define	PLATFORM_USE_SHOWFLAGS_ALWAYS_BITFIELD				1
 #endif
 
-#ifndef PLATFORM_USE_CACHED_SLACK_MEMORY_IN_MEMORY_STATS
-	#define	PLATFORM_USE_CACHED_SLACK_MEMORY_IN_MEMORY_STATS	0
-#endif
-
 #ifndef PLATFORM_SUPPORTS_LANDSCAPE_VISUAL_MESH_LOD_STREAMING
 	#define PLATFORM_SUPPORTS_LANDSCAPE_VISUAL_MESH_LOD_STREAMING 0 UE_DEPRECATED_MACRO(5.1, "PLATFORM_SUPPORTS_LANDSCAPE_VISUAL_MESH_LOD_STREAMING has been deprecated and should be replaced with 0.")
 #endif
