Added safety check while doing a for-each on FThreadManager::Thread to ensure the list/map is not modified during the provided callback.

The code will only early-return and log if a violation is detected since triggering a crash may result in a deadlock on some platforms.

diff --git a/Engine/Source/Runtime/Core/Private/HAL/ThreadingBase.cpp b/Engine/Source/Runtime/Core/Private/HAL/ThreadingBase.cpp
--- a/Engine/Source/Runtime/Core/Private/HAL/ThreadingBase.cpp
+++ b/Engine/Source/Runtime/Core/Private/HAL/ThreadingBase.cpp
@@ -432,6 +432,21 @@ public:
 };
 uint32 FFakeThread::ThreadIdCounter = 0xffff;
 
+bool FThreadManager::CheckThreadListSafeToContinueIteration()
+{
+	if (bIsThreadListDirty)
+	{
+		UE_LOG(LogCore, Error, TEXT("FThreadManager::Threads was modified during unsafe iteration. Iteration will be aborted."));
+		return false;
+	}
+
+	return true;
+}
+
+void FThreadManager::OnThreadListModified()
+{
+	bIsThreadListDirty = true;
+}
 
 void FThreadManager::AddThread(uint32 ThreadId, FRunnableThread* Thread)
 {
@@ -475,6 +490,7 @@ void FThreadManager::AddThread(uint32 ThreadId, FRunnableThread* Thread)
 	if (!Threads.Contains(ThreadId))
 	{
 		Threads.Add(ThreadId, Thread);
+		OnThreadListModified();
 	}
 }
 
@@ -485,6 +501,7 @@ void FThreadManager::RemoveThread(FRunnableThread* Thread)
 	if (ThreadId)
 	{
 		Threads.Remove(*ThreadId);
+		OnThreadListModified();
 	}
 }
 
@@ -575,17 +592,17 @@ void FThreadManager::ForEachThreadStackBackTrace(TFunctionRef<bool(uint32 Thread
 	ProgramCounterBuffer.SetNumZeroed(FThreadStackBackTrace::ProgramCountersMaxStackSize);
 
 	FScopeLock Lock(&ThreadsCritical);
+	bIsThreadListDirty = false;
 
 	{
 		const TConstArrayView<uint64> ProgramCounterOutputArrayView = ThreadStackBackTraces_PerformStackWalk(CurThreadId, GGameThreadId, ProgramCounterBuffer);
 		const bool bContinue = Func(GGameThreadId, TEXT("GameThread"), ProgramCounterOutputArrayView);
-		if (!bContinue)
+		if (!CheckThreadListSafeToContinueIteration() || !bContinue)
 		{
 			return;
 		}
 	}
 
-	// Not using ForEachThread since that creates a copy, just verify the callback isn't modifying thread list
 	for (const TPair<uint32, FRunnableThread*>& Pair : Threads)
 	{
 		const uint32 ThreadId = Pair.Key;
@@ -593,7 +610,7 @@ void FThreadManager::ForEachThreadStackBackTrace(TFunctionRef<bool(uint32 Thread
 		const FString& ThreadName = Thread->GetThreadName();
 		const TConstArrayView<uint64> ProgramCounterOutputArrayView = ThreadStackBackTraces_PerformStackWalk(CurThreadId, ThreadId, ProgramCounterBuffer);
 		const bool bContinue = Func(ThreadId, *ThreadName, ProgramCounterOutputArrayView);
-		if (!bContinue)
+		if (!CheckThreadListSafeToContinueIteration() || !bContinue)
 		{
 			return;
 		}
diff --git a/Engine/Source/Runtime/Core/Public/HAL/ThreadManager.h b/Engine/Source/Runtime/Core/Public/HAL/ThreadManager.h
--- a/Engine/Source/Runtime/Core/Public/HAL/ThreadManager.h
+++ b/Engine/Source/Runtime/Core/Public/HAL/ThreadManager.h
@@ -25,6 +25,12 @@ class FThreadManager
 	/** List of thread objects to be ticked. */
 	FThreads Threads;
 
+	/* Helper variable for catching unexpected modification of the thread map/list. */
+	bool bIsThreadListDirty = false;
+
+	bool CheckThreadListSafeToContinueIteration();
+	void OnThreadListModified();
+
 public:
 
 	/**
