Remove unnecessary header includes.

diff --git a/Engine/Source/Runtime/Core/Public/AutoRTFM/AutoRTFM.h b/Engine/Source/Runtime/Core/Public/AutoRTFM/AutoRTFM.h
--- a/Engine/Source/Runtime/Core/Public/AutoRTFM/AutoRTFM.h
+++ b/Engine/Source/Runtime/Core/Public/AutoRTFM/AutoRTFM.h
@@ -16,18 +16,19 @@
 #define UE_AUTORTFM_NOAUTORTFM
 #endif
 
-#include <stdbool.h>
+#if defined(UE_AUTORTFM_STANDALONE)
 #include <stdlib.h>
 #include <memory.h>
 #include <type_traits>
-
-#if defined(UE_AUTORTFM_STANDALONE)
 #define UE_AUTORTFM_API
 #define UE_AUTORTFM_FORCEINLINE inline
+#define UE_AUTORTFM_MEMCPY ::memcpy
 #else
 #include <HAL/Platform.h>
+#include <HAL/PlatformMemory.h>
 #define UE_AUTORTFM_API CORE_API
 #define UE_AUTORTFM_FORCEINLINE FORCEINLINE
+#define UE_AUTORTFM_MEMCPY FPlatformMemory::Memcpy
 #endif
 
 #if UE_AUTORTFM
@@ -602,7 +603,7 @@ namespace ForTheRuntime
 	UE_AUTORTFM_FORCEINLINE void WriteMemory(void* DestPtr, void const* SrcPtr, size_t Size)
 	{
 		RecordOpenWrite(DestPtr, Size);
-		memcpy(DestPtr, SrcPtr, Size);
+		UE_AUTORTFM_MEMCPY(DestPtr, SrcPtr, Size);
 	}
 
 	// WriteMemory first records the memory span as written (see RecordOpenWrite) and then copies the specified value into it.
