Add a CVar for Windows to estimate total memory usage in a container

We are currently using an artificial limit on memory usage when running in a container. This limiting mechanism only accounts for the current process, which might miss gigabytes of memory coming from shader compilation or anything else we run in the container: .NET scripts, crash reporting etc. This error leads us to overestimate available memory, so we don't trigger GC or hold off on starting memory intensive processes.

On Windows, containers are implemented using job objects, which are essentially nestable groups of processes. Job objects allow you to set a limit for total committed memory. Because life is cruel, the API does not allow you to query how much memory you currently have committed. We instead estimate it by going over all processes in the job object and sum up their pagefile usage. This estimate could be both too high (shared memory) or too low (kernel objects associated with the process may not be counted), but empirically it is much more accurate than the current version. In my testing, I have seen the first green runs for projects that have been going OOM for months.

Performance-wise, this change is terrible. I have measured that on a 20min cook we spend about 5s just estimating memory usage. This is bad, but ultimately still preferable to failing a cook because we did not realize that we should trigger a GC, for example. In that case we likely re-run the cook multiple times, wasting significantly more than 5s.

(Note that our memory stats use the term "virtual memory" to refer to "committed memory" on Windows.)

diff --git a/Engine/Source/Runtime/Core/Private/Windows/WindowsPlatformMemory.cpp b/Engine/Source/Runtime/Core/Private/Windows/WindowsPlatformMemory.cpp
--- a/Engine/Source/Runtime/Core/Private/Windows/WindowsPlatformMemory.cpp
+++ b/Engine/Source/Runtime/Core/Private/Windows/WindowsPlatformMemory.cpp
@@ -44,6 +44,8 @@
 
 #include <Psapi.h>
 #pragma comment(lib, "psapi.lib")
+#include <jobapi2.h>
+
 
 DECLARE_MEMORY_STAT(TEXT("Windows Specific Memory Stat"),	STAT_WindowsSpecificMemoryStat, STATGROUP_MemoryPlatform);
 
@@ -56,6 +58,13 @@ static FAutoConsoleVariableRef CVarLogPlatformMemoryStats(
 	TEXT("Set a synthetic platform total memory size (in GB) which will be returned as Total and Available memory from GetStats\n"),
 	ECVF_Default);
 
+static bool GWindowsUseContainerMemory = false;
+static FAutoConsoleVariableRef CVarUseContainerMemory(
+	TEXT("memory.WindowsPlatformMemoryUseContainerMemory"),
+	GWindowsUseContainerMemory,
+	TEXT("Set to assume that this process is running in a docker container and take the entire container's memory usage into consideration when computing available memory."),
+	ECVF_Default);
+
 
 #if ENABLE_WIN_ALLOC_TRACKING
 // This allows tracking of allocations that don't happen within the engine's wrappers.
@@ -69,7 +78,25 @@ int WindowsAllocHook(int nAllocType, void *pvData,
 }
 #endif // ENABLE_WIN_ALLOC_TRACKING
 
-
+// Returns whether this process is running as part of a Windows job. Windows jobs are used to run
+// docker containers and can enforce memory limits for a group of processes.
+// The usual system APIs still return the memory limits of the host system, which is not correct.
+static bool IsRunningAsJob()
+{
+	static bool bIsJob = []() {
+		BOOL bIsRunningInJob = false;
+		if (!IsProcessInJob(GetCurrentProcess(), NULL, &bIsRunningInJob))
+		{
+			DWORD ErrNo = GetLastError();
+			UE_LOG(LogHAL, Warning, TEXT("IsProcessInJob(GetCurrentProcess(), NULL, &bIsRunningInJob) failed with GetLastError() = %d"),
+				ErrNo
+			);
+			return false;
+		}
+		return (bool)bIsRunningInJob;
+		}();
+		return bIsJob;
+}
 
 void FWindowsPlatformMemory::Init()
 {
@@ -89,18 +116,16 @@ void FWindowsPlatformMemory::Init()
 	SET_MEMORY_STAT(MCR_PhysicalLLM, 5*GB);	// no upper limit on Windows. Choose 5GB because it's roughly the same as current consoles.
 #endif
 
+	if (IsRunningAsJob())
+	{
+		UE_LOG(LogMemory, Log, TEXT("Process is running as part of a Windows Job with separate resource limits"));
+	}
+
 	const FPlatformMemoryConstants& MemoryConstants = FPlatformMemory::GetConstants();
-#if PLATFORM_32BITS	
 	UE_LOG(LogMemory, Log, TEXT("Memory total: Physical=%.1fGB (%dGB approx) Virtual=%.1fGB"), 
 		(double)MemoryConstants.TotalPhysical/1024.0/1024.0/1024.0,
 		MemoryConstants.TotalPhysicalGB, 
 		(double)MemoryConstants.TotalVirtual/1024.0/1024.0/1024.0 );
-#else
-	// Logging virtual memory size for 64bits is pointless.
-	UE_LOG(LogMemory, Log, TEXT("Memory total: Physical=%.1fGB (%dGB approx)"), 
-		(double)MemoryConstants.TotalPhysical/1024.0/1024.0/1024.0,
-		MemoryConstants.TotalPhysicalGB );
-#endif //PLATFORM_32BITS
 
 	// program size is hard to ascertain and isn't so relevant on Windows. For now just set to zero.
 	LLM(FLowLevelMemTracker::Get().SetProgramSize(0));
@@ -263,6 +288,68 @@ FMalloc* FWindowsPlatformMemory::BaseAllocator()
 	return Instance;
 }
 
+static uint64 EstimateContainerCommittedMemory()
+{
+	// These is unfortunately no way to just get the resources used by a job group. This is an unfortunate hole
+	// in the API provided by Windows. We instead sum up the commit usage of all processes in the job group.
+	// This is slow and overestimates the actual commit usage, but that is close enough and we'd rather be slow
+	// than terminate the process because we are running out of memory.
+	TRACE_CPUPROFILER_EVENT_SCOPE(EstimateContainerCommittedMemory);
+	constexpr int32 MaxProcessesAccountedFor = 256;
+	struct ProcessList {
+		DWORD NumberOfAssignedProcesses;
+		DWORD NumberOfProcessIdsInList;
+		ULONG_PTR ProcessIdList[MaxProcessesAccountedFor];
+	};
+	ProcessList Processes{};
+	DWORD LengthWritten;
+	if (!QueryInformationJobObject(NULL, JobObjectBasicProcessIdList, &Processes, sizeof(Processes), &LengthWritten))
+	{
+		DWORD ErrNo = GetLastError();
+		UE_LOG(LogMemory, Error, TEXT("QueryInformationJobObject(NULL, JobObjectBasicProcessIdList, &Processes, sizeof(Processes), &LengthWritten) failed with GetLastError() = %d"),
+			ErrNo
+		);
+		return 0;
+	}
+	if (Processes.NumberOfProcessIdsInList < Processes.NumberOfAssignedProcesses)
+	{
+		UE_LOG(LogMemory, Warning, TEXT("The number of processes in this container %d is larger than %d, memory accounting may be more inaccurate than usual."),
+			Processes.NumberOfAssignedProcesses,
+			Processes.NumberOfProcessIdsInList
+		);
+	}
+	ensureMsgf(Processes.NumberOfProcessIdsInList <= MaxProcessesAccountedFor,
+		TEXT("Number of processes returned by QueryInformationJobObject is too large: %d (must be below %d)"),
+		Processes.NumberOfProcessIdsInList,
+		MaxProcessesAccountedFor
+	);
+	ensureMsgf(Processes.NumberOfProcessIdsInList > 0, TEXT("Number of processes returned by QueryInformationJobObject is 0"));
+	PROCESS_MEMORY_COUNTERS ProcessMemoryCounters;
+	FPlatformMemory::Memzero(&ProcessMemoryCounters, sizeof(ProcessMemoryCounters));
+	uint64 TotalCommittedMemory = 0;
+	for (int32 ProcessIndex = 0; ProcessIndex < (int32)Processes.NumberOfProcessIdsInList; ProcessIndex++)
+	{
+		HANDLE ProcHandle = OpenProcess(PROCESS_QUERY_LIMITED_INFORMATION, FALSE, (DWORD)Processes.ProcessIdList[ProcessIndex]);
+		if (ProcHandle != INVALID_HANDLE_VALUE)
+		{
+			if (GetProcessMemoryInfo(ProcHandle, &ProcessMemoryCounters, sizeof(ProcessMemoryCounters)))
+			{
+				TotalCommittedMemory += ProcessMemoryCounters.PagefileUsage;
+			}
+		}
+		else
+		{
+			DWORD ErrNo = GetLastError();
+			UE_LOG(LogMemory, Warning, TEXT("OpenProcess(PROCESS_QUERY_LIMITED_INFORMATION, FALSE, %d) failed with GetLastError() = %d"),
+				(DWORD)Processes.ProcessIdList[ProcessIndex],
+				ErrNo
+			);
+		}
+		CloseHandle(ProcHandle);
+	}
+	return TotalCommittedMemory;
+}
+
 FPlatformMemoryStats FWindowsPlatformMemory::GetStats()
 {
 	// release mimalloc held pages back to OS before getting stats
@@ -286,6 +373,8 @@ FPlatformMemoryStats FWindowsPlatformMemory::GetStats()
 	 *	GetPerformanceInfo
 	 *		PPERFORMANCE_INFORMATION 
 	 *		PageSize
+	 * 
+	 * If we are running as part of a Windows job (e.g. in a docker container), we report those limits instead.
 	 */
 
 	// This method is slow, do not call it too often.
@@ -334,6 +423,16 @@ FPlatformMemoryStats FWindowsPlatformMemory::GetStats()
 		MemoryStats.MemoryPressureStatus = MemoryStats.FGenericPlatformMemoryStats::GetMemoryPressureStatus();
 	}
 
+	if (GWindowsUseContainerMemory && IsRunningAsJob())
+	{
+		const uint64 TotalUsage = EstimateContainerCommittedMemory();
+		ensureMsgf(TotalUsage != 0, TEXT("Estimated container memory usage is 0. This should never happen"));
+		if (TotalUsage != 0)
+		{
+			MemoryStats.AvailableVirtual = GetConstants().TotalVirtual - TotalUsage;
+		}
+	}
+
 	if ( GWindowsPlatformMemoryGetStatsLimitTotalGB > 0 )
 	{
 		// if GWindowsPlatformMemoryGetStatsLimitTotalGB is set
@@ -413,6 +512,45 @@ const FPlatformMemoryConstants& FWindowsPlatformMemory::GetConstants()
 		MemoryConstants.AddressLimit = FPlatformMath::RoundUpToPowerOfTwo64(MemoryConstants.TotalPhysical);
 
 		MemoryConstants.TotalPhysicalGB = (uint32)((MemoryConstants.TotalPhysical + 1024 * 1024 * 1024 - 1) / 1024 / 1024 / 1024);
+
+		if (IsRunningAsJob())
+		{
+			JOBOBJECT_EXTENDED_LIMIT_INFORMATION JobInfo{};
+			DWORD LengthWritten{};
+			if (!QueryInformationJobObject(NULL, JobObjectExtendedLimitInformation, &JobInfo, sizeof(JobInfo), &LengthWritten))
+			{
+				DWORD ErrNo = GetLastError();
+				UE_LOG(LogMemory, Error, TEXT("QueryInformationJobObject(NULL, JobObjectExtendedLimitInformation, &JobInfo, sizeof(JobInfo), &LengthWritten) failed with GetLastError() = %d"),
+					ErrNo
+				);
+			}
+			else
+			{
+				if ((JobInfo.BasicLimitInformation.LimitFlags & (JOB_OBJECT_LIMIT_JOB_MEMORY | JOB_OBJECT_LIMIT_PROCESS_MEMORY)) == 0)
+				{
+					UE_LOG(LogMemory, Display, TEXT("Running as part of a job but no memory limit has been set."));
+				}
+				else
+				{
+					uint64 MemoryLimit = MemoryConstants.TotalVirtual;
+					if ((JobInfo.BasicLimitInformation.LimitFlags & JOB_OBJECT_LIMIT_JOB_MEMORY) == JOB_OBJECT_LIMIT_JOB_MEMORY)
+					{
+						MemoryLimit = JobInfo.JobMemoryLimit;
+						UE_LOG(LogMemory, Display, TEXT("Detected a job memory limit of %.1fGB for this job."),
+							(double)MemoryLimit / (1024 * 1024 * 1024)
+						);
+					}
+					if ((JobInfo.BasicLimitInformation.LimitFlags & JOB_OBJECT_LIMIT_PROCESS_MEMORY) == JOB_OBJECT_LIMIT_PROCESS_MEMORY)
+					{
+						MemoryLimit = JobInfo.ProcessMemoryLimit;
+						UE_LOG(LogMemory, Display, TEXT("Detected a per-process memory limit of %.1fGB for this job."),
+							(double)MemoryLimit / (1024 * 1024 * 1024)
+						);
+					}
+					MemoryConstants.TotalVirtual = MemoryLimit;
+				}
+			}
+		}
 	}
 
 	return MemoryConstants;	
