Make memory printouts during cooking more useful

Two changes:
 * we now also print out the available virtual memory occasionally during cooking, not just the used memory. Note that "Virtual Memory" is "committed memory" on Windows in UE-parlance, and the available virtual memory is actually most interesting in a docker container.
 * when an OOM happens, print out the current memory usage in units that don't require 10 digit numbers and are easy to parse. That's GiB for total usages and MiB for allocations.

#changelist validated
#virtualized

diff --git a/Engine/Source/Editor/UnrealEd/Private/CookOnTheFlyServer.cpp b/Engine/Source/Editor/UnrealEd/Private/CookOnTheFlyServer.cpp
--- a/Engine/Source/Editor/UnrealEd/Private/CookOnTheFlyServer.cpp
+++ b/Engine/Source/Editor/UnrealEd/Private/CookOnTheFlyServer.cpp
@@ -1819,8 +1819,12 @@ void UCookOnTheFlyServer::UpdateDisplay(UE::Cook::FTickStackData& StackData, boo
 		{
 			if (!(StackData.TickFlags & ECookTickFlags::HideProgressDisplay) && (GCookProgressDisplay != (int32)ECookProgressDisplayMode::Nothing))
 			{
-				UE_LOG(LogCook, Display, TEXT("Cook Diagnostics: OpenFileHandles=%d, VirtualMemory=%dMiB"),
-					OpenFileHandles, FPlatformMemory::GetStats().UsedVirtual / 1024 / 1024);
+				const FPlatformMemoryStats MemoryStats = FPlatformMemory::GetStats();
+				UE_LOG(LogCook, Display, TEXT("Cook Diagnostics: OpenFileHandles=%d, VirtualMemory=%dMiB, VirtualMemoryAvailable=%dMiB"),
+					OpenFileHandles,
+					MemoryStats.UsedVirtual / 1024 / 1024,
+					MemoryStats.AvailableVirtual / 1024 / 1024
+				);
 				if (CookDirector)
 				{
 					CookDirector->UpdateDisplayDiagnostics();
diff --git a/Engine/Source/Runtime/Core/Private/GenericPlatform/GenericPlatformMemory.cpp b/Engine/Source/Runtime/Core/Private/GenericPlatform/GenericPlatformMemory.cpp
--- a/Engine/Source/Runtime/Core/Private/GenericPlatform/GenericPlatformMemory.cpp
+++ b/Engine/Source/Runtime/Core/Private/GenericPlatform/GenericPlatformMemory.cpp
@@ -219,26 +219,28 @@ void FGenericPlatformMemory::OnOutOfMemory(uint64 Size, uint32 Alignment)
 		FPlatformMemoryStats PlatformMemoryStats = FPlatformMemory::GetStats();
 		if (BackupOOMMemoryPool)
 		{
-			FPlatformMemory::BinnedFreeToOS(BackupOOMMemoryPool, FPlatformMemory::GetBackMemoryPoolSize());
-			UE_LOG(LogMemory, Warning, TEXT("Freeing %d bytes from backup pool to handle out of memory."), FPlatformMemory::GetBackMemoryPoolSize());
+			const uint32 BackupPoolSize = FPlatformMemory::GetBackMemoryPoolSize();
+			FPlatformMemory::BinnedFreeToOS(BackupOOMMemoryPool, BackupPoolSize);
+			UE_LOG(LogMemory, Warning, TEXT("Freeing %d bytes (%.1f MiB) from backup pool to handle out of memory."), BackupPoolSize, double(BackupPoolSize) / (1024 * 1024));
         
 			LLM_IF_ENABLED(FLowLevelMemTracker::Get().OnLowLevelFree(ELLMTracker::Default, BackupOOMMemoryPool));
 			MemoryTrace_Free((uint64)BackupOOMMemoryPool, EMemoryTraceRootHeap::SystemMemory);
 		}
 
 		UE_LOG(LogMemory, Warning, TEXT("MemoryStats:"
-			"\n\tAvailablePhysical %llu"
-			"\n\t AvailableVirtual %llu"
-			"\n\t     UsedPhysical %llu"
-			"\n\t PeakUsedPhysical %llu"
-			"\n\t      UsedVirtual %llu"
-			"\n\t  PeakUsedVirtual %llu"),
-			(uint64)PlatformMemoryStats.AvailablePhysical,
-			(uint64)PlatformMemoryStats.AvailableVirtual,
-			(uint64)PlatformMemoryStats.UsedPhysical,
-			(uint64)PlatformMemoryStats.PeakUsedPhysical,
-			(uint64)PlatformMemoryStats.UsedVirtual,
-			(uint64)PlatformMemoryStats.PeakUsedVirtual);
+			"\n\tAvailablePhysical %llu (%.2f GiB)"
+			"\n\t AvailableVirtual %llu (%.2f GiB)"
+			"\n\t     UsedPhysical %llu (%.2f GiB)"
+			"\n\t PeakUsedPhysical %llu (%.2f GiB)"
+			"\n\t      UsedVirtual %llu (%.2f GiB)"
+			"\n\t  PeakUsedVirtual %llu (%.2f GiB)"),
+			(uint64)PlatformMemoryStats.AvailablePhysical, double(PlatformMemoryStats.AvailablePhysical) / (1024 * 1024 * 1024),
+			(uint64)PlatformMemoryStats.AvailableVirtual, double(PlatformMemoryStats.AvailableVirtual) / (1024 * 1024 * 1024),
+			(uint64)PlatformMemoryStats.UsedPhysical, double(PlatformMemoryStats.UsedPhysical) / (1024 * 1024 * 1024),
+			(uint64)PlatformMemoryStats.PeakUsedPhysical, double(PlatformMemoryStats.PeakUsedPhysical) / (1024 * 1024 * 1024),
+			(uint64)PlatformMemoryStats.UsedVirtual, double(PlatformMemoryStats.UsedVirtual) / (1024 * 1024 * 1024),
+			(uint64)PlatformMemoryStats.PeakUsedVirtual, double(PlatformMemoryStats.PeakUsedVirtual) / (1024 * 1024 * 1024)
+		);
 		if (GWarn)
 		{
 			GMalloc->DumpAllocatorStats(*GWarn);
@@ -248,7 +250,7 @@ void FGenericPlatformMemory::OnOutOfMemory(uint64 Size, uint32 Alignment)
 		FCoreDelegates::GetOutOfMemoryDelegate().Broadcast();
 
 		// ErrorMsg might be unrelated to OoM error in some cases as the code that calls OnOutOfMemory could have called other system functions that modified errno
-		UE_LOG(LogMemory, Fatal, TEXT("Ran out of memory allocating %llu bytes with alignment %u. Last error msg: %s."), Size, Alignment, ErrorMsg);
+		UE_LOG(LogMemory, Fatal, TEXT("Ran out of memory allocating %llu (%.1f MiB) bytes with alignment %u. Last error msg: %s."), Size, double(Size) / (1024 * 1024), Alignment, ErrorMsg);
 	};
 	
 	UE_CALL_ONCE(HandleOOM);
