Deprecate functions with bAllowShrinking parameters.

diff --git a/Engine/Source/Runtime/Core/Public/Containers/AllowShrinking.h b/Engine/Source/Runtime/Core/Public/Containers/AllowShrinking.h
--- a/Engine/Source/Runtime/Core/Public/Containers/AllowShrinking.h
+++ b/Engine/Source/Runtime/Core/Public/Containers/AllowShrinking.h
@@ -11,5 +11,4 @@ enum class EAllowShrinking : uint8
 	Yes
 };
 
-// Can comment in this when fixing up existing calls, or made permanent when it's ready to be
-#define UE_ALLOWSHRINKING_BOOL_DEPRECATED(FunctionName) //UE_DEPRECATED(5.4, FunctionName " with a boolean bAllowShrinking has been deprecated - please use the EAllowShrinking enum instead")
+#define UE_ALLOWSHRINKING_BOOL_DEPRECATED(FunctionName) UE_DEPRECATED_FORENGINE(5.4, FunctionName " with a boolean bAllowShrinking has been deprecated - please use the EAllowShrinking enum instead")
