Added IsInitialized to TNonNullPtr so we can test exceptional initialization scenarios for these pointer types.

diff --git a/Engine/Source/Runtime/Core/Public/Templates/NonNullPointer.h b/Engine/Source/Runtime/Core/Public/Templates/NonNullPointer.h
--- a/Engine/Source/Runtime/Core/Public/Templates/NonNullPointer.h
+++ b/Engine/Source/Runtime/Core/Public/Templates/NonNullPointer.h
@@ -127,6 +127,16 @@ public:
 		return Object;
 	}
 
+	/*
+	 * WARNING: Hack that can be used under extraordinary circumstances. Pointers here 
+	 * should always be valid but might be in the EDefaultConstructNonNullPtr state 
+	 * during initialization.
+	 */
+	FORCEINLINE bool IsInitialized() const
+	{
+		return Object != nullptr;
+	}
+
 private:
 
 	/** The object we're holding a reference to. */
