Update the FSeverityToken of a FTokenizedMessage when updating the severity of the message

#rnx

#changelist validated
#virtualized

diff --git a/Engine/Source/Runtime/Core/Private/Logging/TokenizedMessage.cpp b/Engine/Source/Runtime/Core/Private/Logging/TokenizedMessage.cpp
--- a/Engine/Source/Runtime/Core/Private/Logging/TokenizedMessage.cpp
+++ b/Engine/Source/Runtime/Core/Private/Logging/TokenizedMessage.cpp
@@ -138,7 +138,18 @@ void FTokenizedMessage::SetMessageLink(const TSharedRef<IMessageToken>& InToken)
 
 void FTokenizedMessage::SetSeverity( const EMessageSeverity::Type InSeverity )
 {
-	Severity = InSeverity;
+	if (Severity != InSeverity)
+	{
+		if (MessageTokens.Num() > 0 && MessageTokens[0]->GetType() == EMessageToken::Severity)
+		{
+			if (TSharedRef<FSeverityToken> SeverityToken = StaticCastSharedRef<FSeverityToken>(MessageTokens[0]);
+				SeverityToken->GetSeverity() == Severity)
+			{
+				MessageTokens[0] = FSeverityToken::Create(InSeverity);
+			}
+		}
+		Severity = InSeverity;
+	}
 }
 
 EMessageSeverity::Type FTokenizedMessage::GetSeverity() const
