Change Microsoft crash context to use reduced-overhead stack trace iteration.

#tests Compared runtime context output with -IgnoreDebugger and debug fatal log.

#changelist validated
#virtualized

diff --git a/Engine/Source/Runtime/Core/Private/Microsoft/MicrosoftPlatformCrashContext.cpp b/Engine/Source/Runtime/Core/Private/Microsoft/MicrosoftPlatformCrashContext.cpp
--- a/Engine/Source/Runtime/Core/Private/Microsoft/MicrosoftPlatformCrashContext.cpp
+++ b/Engine/Source/Runtime/Core/Private/Microsoft/MicrosoftPlatformCrashContext.cpp
@@ -10,13 +10,12 @@
 void FMicrosoftPlatformCrashContext::CaptureAllThreadContexts()
 {
 #if PLATFORM_SUPPORTS_ALL_THREAD_BACKTRACES
-	TArray<typename FThreadManager::FThreadStackBackTrace> StackTraces;
-	FThreadManager::Get().GetAllThreadStackBackTraces(StackTraces);
-
-	for (const FThreadManager::FThreadStackBackTrace& Thread : StackTraces)
-	{
-		AddPortableThreadCallStack(Thread.ThreadId, *Thread.ThreadName, Thread.ProgramCounters.GetData(), Thread.ProgramCounters.Num());
-	}
+	FThreadManager::Get().ForEachThreadStackBackTrace(
+		[this](uint32 ThreadId, const TCHAR* ThreadName, const TConstArrayView<uint64>& StackTrace)
+		{
+			AddPortableThreadCallStack(ThreadId, ThreadName, StackTrace.GetData(), StackTrace.Num());
+			return true;
+		});
 #endif
 }
 
