Updated bAllowShrinking deprecation version number.

diff --git a/Engine/Source/Runtime/Core/Public/Containers/AllowShrinking.h b/Engine/Source/Runtime/Core/Public/Containers/AllowShrinking.h
--- a/Engine/Source/Runtime/Core/Public/Containers/AllowShrinking.h
+++ b/Engine/Source/Runtime/Core/Public/Containers/AllowShrinking.h
@@ -11,4 +11,4 @@ enum class EAllowShrinking : uint8
 	Yes
 };
 
-#define UE_ALLOWSHRINKING_BOOL_DEPRECATED(FunctionName) UE_DEPRECATED_FORENGINE(5.4, FunctionName " with a boolean bAllowShrinking has been deprecated - please use the EAllowShrinking enum instead")
+#define UE_ALLOWSHRINKING_BOOL_DEPRECATED(FunctionName) UE_DEPRECATED_FORENGINE(5.5, FunctionName " with a boolean bAllowShrinking has been deprecated - please use the EAllowShrinking enum instead")
