[LLM] Disable hitchdetector when LLM is enabled
It is expected that LLM capture adds CPU overhead during captures, hang/hitch detection made things just worse by spamming logs with unhelpful info.

#tests PIE
#changelist validated
#virtualized

diff --git a/Engine/Source/Runtime/Core/Private/HAL/LowLevelMemTracker.cpp b/Engine/Source/Runtime/Core/Private/HAL/LowLevelMemTracker.cpp
--- a/Engine/Source/Runtime/Core/Private/HAL/LowLevelMemTracker.cpp
+++ b/Engine/Source/Runtime/Core/Private/HAL/LowLevelMemTracker.cpp
@@ -8,11 +8,13 @@
 #include "HAL/LowLevelMemStats.h"
 #include "HAL/FileManager.h"
 #include "HAL/IConsoleManager.h"
+#include "HAL/ThreadHeartBeat.h"
 #include "HAL/PlatformMemory.h" // for page allocation association.
 #include "LowLevelMemTrackerPrivate.h"
 #include "MemPro/MemProProfiler.h"
 #include "Math/NumericLimits.h"
 #include "Misc/CString.h"
+#include "Misc/CoreDelegates.h"
 #include "Misc/FileHelper.h"
 #include "Misc/Fork.h"
 #include "Misc/OutputDeviceRedirector.h"
@@ -2258,6 +2260,16 @@ void FLowLevelMemTracker::ProcessCommandLine(const TCHAR* CmdLine)
 
 	UE_LOG(LogInit, Log, TEXT("LLM enabled CsvWriter: %s TraceWriter: %s"),
 		bCsvWriterEnabled ? TEXT("on") : TEXT("off"), bTraceWriterEnabled ? TEXT("on") : TEXT("off"));
+
+	// Disable hitchdetector because LLM is already slow enough.
+	if (!FParse::Param(CmdLine, TEXT("DetectHitchesWithLLM")))
+	{
+		// Schedule disabling when it's safe to init HitchHeartBeat
+		FCoreDelegates::OnPostEngineInit.AddLambda([]{
+			// Calling Get() will instantiate FGameThreadHitchHeartBeat if it is not already.
+			FGameThreadHitchHeartBeat::Get().Stop();
+		});
+	}
 }
 
 // Return the total amount of memory being tracked
