Fix MRUArray to work with non POD types

diff --git a/Engine/Source/Runtime/Core/Public/Containers/MRUArray.h b/Engine/Source/Runtime/Core/Public/Containers/MRUArray.h
--- a/Engine/Source/Runtime/Core/Public/Containers/MRUArray.h
+++ b/Engine/Source/Runtime/Core/Public/Containers/MRUArray.h
@@ -81,9 +81,7 @@ public:
 	{
 		// Remove any existing copies of the item.
 		this->Remove(Item);
-
-		this->InsertUninitialized(0);
-		(*this)[0] = Item;
+		this->Insert(Item, 0);
 
 		CullArray();
 
