PlatformFileJournal: Fix typos in log statements.
#rnx

diff --git a/Engine/Source/Runtime/AssetRegistry/Private/AssetDataGathererDiscoveryCache.cpp b/Engine/Source/Runtime/AssetRegistry/Private/AssetDataGathererDiscoveryCache.cpp
--- a/Engine/Source/Runtime/AssetRegistry/Private/AssetDataGathererDiscoveryCache.cpp
+++ b/Engine/Source/Runtime/AssetRegistry/Private/AssetDataGathererDiscoveryCache.cpp
@@ -78,7 +78,7 @@ void FAssetDataDiscoveryCache::LoadAndUpdateCache()
 			(!bReadEnabled ? TEXT("read") : TEXT("written"));
 		UE_LOG(LogAssetRegistry, Display,
 			TEXT("PlatformFileJournal is not available on volume '%s' of project directory '%s', so AssetDiscovery cache will not be %s. Unavailability reason:\n\t%s"),
-			*TestVolumeName, *ProjectDir, MissingOperation, *TestError);
+			*TestVolumeName, *FPaths::ConvertRelativePathToFull(ProjectDir), MissingOperation, *TestError);
 	}
 
 	if (!bReadEnabled)
diff --git a/Engine/Source/Runtime/Core/Private/Windows/WindowsPlatformFileJournal.cpp b/Engine/Source/Runtime/Core/Private/Windows/WindowsPlatformFileJournal.cpp
--- a/Engine/Source/Runtime/Core/Private/Windows/WindowsPlatformFileJournal.cpp
+++ b/Engine/Source/Runtime/Core/Private/Windows/WindowsPlatformFileJournal.cpp
@@ -764,8 +764,8 @@ EFileJournalResult GetJournalDescriptor(HANDLE JournalWindowsHandle, const TCHAR
 			{
 				// Can also create using DeviceIoControl with FSCTL_CREATE_USN_JOURNAL
 				*OutError = FString::Printf(TEXT("NTFS Journal is not active for volume '%s'. ")
-					TEXT("Launch cmd.exe as admin and run command `fsutil usn createJournal %s m=<SizeInBytes>`."),
-					TEXT("Recommended <SizeInBytes> is 100000000 (100MB)."),
+					TEXT("Launch cmd.exe as admin and run command `fsutil usn createJournal %s m=<SizeInBytes>`. ")
+					TEXT("Recommended <SizeInBytes> is 1000000000 (1GB)."),
 					VolumeName, VolumeName);
 			}
 			return EFileJournalResult::JournalNotActive;
