Unifying code between MB2 and MB3

diff --git a/Engine/Source/Runtime/Core/Private/HAL/MallocBinned2.cpp b/Engine/Source/Runtime/Core/Private/HAL/MallocBinned2.cpp
--- a/Engine/Source/Runtime/Core/Private/HAL/MallocBinned2.cpp
+++ b/Engine/Source/Runtime/Core/Private/HAL/MallocBinned2.cpp
@@ -63,17 +63,17 @@ static FAutoConsoleVariableRef GMallocBinned2PerThreadCachesCVar(
 	TEXT("Enables per-thread caches of small (<= 32768 byte) allocations from FMallocBinned2")
 	);
 
-int32 GMallocBinned2BundleSize = DEFAULT_GMallocBinned2BundleSize;
+extern int32 GMallocBinnedBundleSize;
 static FAutoConsoleVariableRef GMallocBinned2BundleSizeCVar(
 	TEXT("MallocBinned2.BundleSize"),
-	GMallocBinned2BundleSize,
+	GMallocBinnedBundleSize,
 	TEXT("Max size in bytes of per-block bundles used in the recycling process")
 	);
 
-int32 GMallocBinned2BundleCount = DEFAULT_GMallocBinned2BundleCount;
+extern int32 GMallocBinnedBundleCount;
 static FAutoConsoleVariableRef GMallocBinned2BundleCountCVar(
 	TEXT("MallocBinned2.BundleCount"),
-	GMallocBinned2BundleCount,
+	GMallocBinnedBundleCount,
 	TEXT("Max count in blocks per-block bundles used in the recycling process")
 	);
 
@@ -118,7 +118,6 @@ TAtomic<int64> AllocatedLargePoolMemoryWAlignment(0); // when we allocate at OS
 
 int64 Binned2PoolInfoMemory = 0;
 int64 Binned2HashMemory = 0;
-int64 Binned2TLSMemory = 0;
 #endif
 
 #if BINNED2_ALLOCATOR_STATS_VALIDATION
@@ -142,7 +141,6 @@ static constexpr uint16 SmallBlockSizes[] =
 
 MS_ALIGN(PLATFORM_CACHE_LINE_SIZE) static uint8 UnusedAlignPadding[PLATFORM_CACHE_LINE_SIZE] GCC_ALIGN(PLATFORM_CACHE_LINE_SIZE) = { 0 };
 uint16 FMallocBinned2::SmallBlockSizesReversed[BINNED2_SMALL_POOL_COUNT] = { 0 };
-uint32 FMallocBinned2::Binned2TlsSlot = FPlatformTLS::InvalidTlsSlot;
 uint32 FMallocBinned2::OsAllocationGranularity = 0;
 uint32 FMallocBinned2::PageSize = 0;
 FMallocBinned2* FMallocBinned2::MallocBinned2 = nullptr;
@@ -309,42 +307,6 @@ private:
 };
 
 
-
-/** Hash table struct for retrieving allocation book keeping information */
-struct FMallocBinned2::PoolHashBucket
-{
-	UPTRINT         BucketIndex;
-	FPoolInfo*      FirstPool;
-	PoolHashBucket* Prev;
-	PoolHashBucket* Next;
-
-	PoolHashBucket()
-	{
-		BucketIndex = 0;
-		FirstPool   = nullptr;
-		Prev        = this;
-		Next        = this;
-	}
-
-	void Link(PoolHashBucket* After)
-	{
-		After->Prev = Prev;
-		After->Next = this;
-		Prev ->Next = After;
-		this ->Prev = After;
-	}
-
-	void Unlink()
-	{
-		Next->Prev = Prev;
-		Prev->Next = Next;
-		Prev       = this;
-		Next       = this;
-	}
-};
-
-
-
 struct FMallocBinned2::Private
 {
 	// Implementation. 
@@ -631,7 +593,7 @@ struct FMallocBinned2::Private
 		--RecursionCounter;
 #endif
 #if BINNED2_ALLOCATOR_STATS
-		FMallocBinned2::FPerThreadFreeBlockLists::ConsolidatedMemory += FreeBlockLists->AllocatedMemory;
+		ConsolidatedMemory.fetch_add(FreeBlockLists->AllocatedMemory, std::memory_order_relaxed);
 #endif
 	}
 
@@ -651,9 +613,15 @@ struct FMallocBinned2::Private
 
 FMallocBinned2::Private::FGlobalRecycler FMallocBinned2::Private::GGlobalRecycler;
 
-#if BINNED2_ALLOCATOR_STATS
-int64 FMallocBinned2::FPerThreadFreeBlockLists::ConsolidatedMemory = 0;
-#endif
+void FMallocBinned2::RegisterThreadFreeBlockLists(FPerThreadFreeBlockLists* FreeBlockLists)
+{
+	Private::RegisterThreadFreeBlockLists(FreeBlockLists);
+}
+
+void FMallocBinned2::UnregisterThreadFreeBlockLists(FPerThreadFreeBlockLists* FreeBlockLists)
+{
+	Private::UnregisterThreadFreeBlockLists(FreeBlockLists);
+}
 
 FORCEINLINE void FMallocBinned2::FPoolList::Clear()
 {
@@ -876,34 +844,6 @@ bool FMallocBinned2::IsInternallyThreadSafe() const
 	return true;
 }
 
-static bool PromoteToLargerBin(SIZE_T& Size, uint32& Alignment, const FMallocBinned2& Malloc)
-{
-	// try to promote our allocation request to a larger bin with a matching natural alignment
-	// if requested alignment is larger than BINNED2_MINIMUM_ALIGNMENT but smaller than BINNED2_MAXIMUM_ALIGNMENT
-	// so we don't do a page allocation with a lot of memory waste
-	Alignment = FMath::Max<uint32>(Alignment, BINNED2_MINIMUM_ALIGNMENT);
-	const SIZE_T AlignedSize = Align(Size, Alignment);
-	if (UNLIKELY((AlignedSize <= BINNED2_MAX_SMALL_POOL_SIZE) && (Alignment <= BINNED2_MAXIMUM_ALIGNMENT)))
-	{
-		uint32 PoolIndex = Malloc.BoundSizeToPoolIndex(AlignedSize);
-		do
-		{
-			uint32 BlockSize = Malloc.PoolIndexToBlockSize(PoolIndex);
-			if (IsAligned(BlockSize, Alignment))
-			{
-				// we found a matching pool for our alignment and size requirements, so modify the size request to match
-				Size = SIZE_T(BlockSize);
-				Alignment = BINNED2_MINIMUM_ALIGNMENT;
-				return true;
-			}
-
-			PoolIndex++;
-		} while (PoolIndex < BINNED2_SMALL_POOL_COUNT);
-	}
-
-	return false;
-}
-
 void* FMallocBinned2::MallocSelect(SIZE_T Size, uint32 Alignment, bool bUseSmallPool)
 {
 	if (!bUseSmallPool)
@@ -920,11 +860,11 @@ void* FMallocBinned2::MallocSelect(SIZE_T Size, uint32 Alignment, bool bUseSmall
 
 void* FMallocBinned2::MallocExternalSmall(SIZE_T Size, uint32 Alignment)
 {
-	uint32 PoolIndex = BoundSizeToPoolIndex(Size);
+	uint32 PoolIndex = BoundSizeToPoolIndex(Size, MemSizeToIndex);
 	FPerThreadFreeBlockLists* Lists = GMallocBinned2PerThreadCaches ? FPerThreadFreeBlockLists::Get() : nullptr;
 	if (Lists)
 	{
-		if (Lists->ObtainRecycledPartial(PoolIndex))
+		if (Lists->ObtainRecycledPartial(PoolIndex, Private::GGlobalRecycler))
 		{
 			if (void* Result = Lists->Malloc(PoolIndex))
 			{
@@ -1152,7 +1092,7 @@ void FMallocBinned2::FreeExternal(void* Ptr)
 		FPerThreadFreeBlockLists* Lists = GMallocBinned2PerThreadCaches ? FPerThreadFreeBlockLists::Get() : nullptr;
 		if (Lists)
 		{
-			BundlesToRecycle = Lists->RecycleFullBundle(BasePtr->PoolIndex);
+			BundlesToRecycle = Lists->RecycleFullBundle(BasePtr->PoolIndex, Private::GGlobalRecycler);
 			bool bPushed = Lists->Free(Ptr, PoolIndex, BlockSize);
 			check(bPushed);
 #if BINNED2_ALLOCATOR_STATS
@@ -1358,11 +1298,11 @@ void FMallocBinned2::SetupTLSCachesOnCurrentThread()
 	{
 		return;
 	}
-	if (!FPlatformTLS::IsValidTlsSlot(FMallocBinned2::Binned2TlsSlot))
+	if (!FPlatformTLS::IsValidTlsSlot(FMallocBinned2::BinnedTlsSlot))
 	{
-		FMallocBinned2::Binned2TlsSlot = FPlatformTLS::AllocTlsSlot();
+		FMallocBinned2::BinnedTlsSlot = FPlatformTLS::AllocTlsSlot();
 	}
-	check(FPlatformTLS::IsValidTlsSlot(FMallocBinned2::Binned2TlsSlot));
+	check(FPlatformTLS::IsValidTlsSlot(FMallocBinned2::BinnedTlsSlot));
 	FPerThreadFreeBlockLists::SetTLS();
 }
 
@@ -1373,96 +1313,6 @@ void FMallocBinned2::ClearAndDisableTLSCachesOnCurrentThread()
 	FPerThreadFreeBlockLists::ClearTLS();
 }
 
-
-bool FMallocBinned2::FFreeBlockList::ObtainPartial(uint32 InPoolIndex)
-{
-	if (!PartialBundle.Head)
-	{
-		PartialBundle.Count = 0;
-		PartialBundle.Head = FMallocBinned2::Private::GGlobalRecycler.PopBundle(InPoolIndex);
-		if (PartialBundle.Head)
-		{
-			PartialBundle.Count = PartialBundle.Head->Count;
-			PartialBundle.Head->NextBundle = nullptr;
-			return true;
-		}
-		return false;
-	}
-	return true;
-}
-
-FMallocBinned2::FBundleNode* FMallocBinned2::FFreeBlockList::RecyleFull(uint32 InPoolIndex)
-{
-	FMallocBinned2::FBundleNode* Result = nullptr;
-	if (FullBundle.Head)
-	{
-		FullBundle.Head->Count = FullBundle.Count;
-		if (!FMallocBinned2::Private::GGlobalRecycler.PushBundle(InPoolIndex, FullBundle.Head))
-		{
-			Result = FullBundle.Head;
-			Result->NextBundle = nullptr;
-		}
-		FullBundle.Reset();
-	}
-	return Result;
-}
-
-FMallocBinned2::FBundleNode* FMallocBinned2::FFreeBlockList::PopBundles(uint32 InPoolIndex)
-{
-	FBundleNode* Partial = PartialBundle.Head;
-	if (Partial)
-	{
-		PartialBundle.Reset();
-		Partial->NextBundle = nullptr;
-	}
-
-	FBundleNode* Full = FullBundle.Head;
-	if (Full)
-	{
-		FullBundle.Reset();
-		Full->NextBundle = nullptr;
-	}
-
-	FBundleNode* Result = Partial;
-	if (Result)
-	{
-		Result->NextBundle = Full;
-	}
-	else
-	{
-		Result = Full;
-	}
-
-	return Result;
-}
-
-void FMallocBinned2::FPerThreadFreeBlockLists::SetTLS()
-{
-	check(FPlatformTLS::IsValidTlsSlot(FMallocBinned2::Binned2TlsSlot));
-	FPerThreadFreeBlockLists* ThreadSingleton = (FPerThreadFreeBlockLists*)FPlatformTLS::GetTlsValue(FMallocBinned2::Binned2TlsSlot);
-	if (!ThreadSingleton)
-	{
-		LLM_PLATFORM_SCOPE(ELLMTag::FMalloc);
-		ThreadSingleton = new (FPlatformMemory::BinnedAllocFromOS(Align(sizeof(FPerThreadFreeBlockLists), FMallocBinned2::OsAllocationGranularity))) FPerThreadFreeBlockLists();
-#if BINNED2_ALLOCATOR_STATS
-		Binned2TLSMemory += Align(sizeof(FPerThreadFreeBlockLists), FMallocBinned2::OsAllocationGranularity);
-#endif
-		FPlatformTLS::SetTlsValue(FMallocBinned2::Binned2TlsSlot, ThreadSingleton);
-		FMallocBinned2::Private::RegisterThreadFreeBlockLists(ThreadSingleton);
-	}
-}
-
-void FMallocBinned2::FPerThreadFreeBlockLists::ClearTLS()
-{
-	check(FPlatformTLS::IsValidTlsSlot(FMallocBinned2::Binned2TlsSlot));
-	FPerThreadFreeBlockLists* ThreadSingleton = (FPerThreadFreeBlockLists*)FPlatformTLS::GetTlsValue(FMallocBinned2::Binned2TlsSlot);
-	if ( ThreadSingleton )
-	{
-		FMallocBinned2::Private::UnregisterThreadFreeBlockLists(ThreadSingleton);
-	}
-	FPlatformTLS::SetTlsValue(FMallocBinned2::Binned2TlsSlot, nullptr);
-}
-
 void FMallocBinned2::CanaryTest(const FFreeBlock* Block) const
 {
 #if BINNED2_FORK_SUPPORT
@@ -1496,7 +1346,7 @@ int64 FMallocBinned2::GetTotalAllocatedSmallPoolMemory() const
 		{
 			FreeBlockAllocatedMemory += FreeBlockLists->AllocatedMemory;
 		}
-		FreeBlockAllocatedMemory += FPerThreadFreeBlockLists::ConsolidatedMemory;
+		FreeBlockAllocatedMemory += ConsolidatedMemory.load(std::memory_order_relaxed);
 	}
 
 	return AllocatedSmallPoolMemory.Load(EMemoryOrder::Relaxed) + FreeBlockAllocatedMemory;
@@ -1550,10 +1400,10 @@ void FMallocBinned2::DumpAllocatorStats(class FOutputDevice& Ar)
 	Ar.Logf(TEXT("OS Allocated: %fmb"), ((double)LocalAllocatedLargePoolMemoryWAlignment) / (1024.0f * 1024.0f));
 	Ar.Logf(TEXT("PoolInfo: %fmb"), ((double)Binned2PoolInfoMemory) / (1024.0f * 1024.0f));
 	Ar.Logf(TEXT("Hash: %fmb"), ((double)Binned2HashMemory) / (1024.0f * 1024.0f));
-	Ar.Logf(TEXT("TLS: %fmb"), ((double)Binned2TLSMemory) / (1024.0f * 1024.0f));
+	Ar.Logf(TEXT("TLS: %fmb"), ((double)TLSMemory.load(std::memory_order_relaxed)) / (1024.0f * 1024.0f));
 	Ar.Logf(TEXT("Total allocated from OS: %fmb"),
 		((double)
-			AllocatedOSSmallPoolMemory + AllocatedLargePoolMemoryWAlignment + Binned2PoolInfoMemory + Binned2HashMemory + Binned2TLSMemory
+			AllocatedOSSmallPoolMemory + AllocatedLargePoolMemoryWAlignment + Binned2PoolInfoMemory + Binned2HashMemory + TLSMemory.load(std::memory_order_relaxed)
 			) / (1024.0f * 1024.0f));
 	Ar.Logf(TEXT("Cached free OS pages: %fmb"), ((double)OSPageAllocatorCachedFreeSize) / (1024.0f * 1024.0f));
 
@@ -1585,5 +1435,21 @@ void FMallocBinned2::UpdateStats()
 	FScopedVirtualMallocTimer::UpdateStats();
 }
 
+void* FMallocBinned2::AllocateMetaDataMemory(SIZE_T Size)
+{
+	LLM_PLATFORM_SCOPE(ELLMTag::FMalloc);
+	Size = Align(Size, FMallocBinned2::OsAllocationGranularity);
+	return FPlatformMemory::BinnedAllocFromOS(Size);
+}
+
+void FMallocBinned2::FreeMetaDataMemory(void* Ptr, SIZE_T Size)
+{
+	if (Ptr)
+	{
+		Size = Align(Size, FMallocBinned2::OsAllocationGranularity);
+		FPlatformMemory::BinnedFreeToOS(Ptr, Size);
+	}
+}
+
 
 PRAGMA_RESTORE_UNSAFE_TYPECAST_WARNINGS
diff --git a/Engine/Source/Runtime/Core/Private/HAL/MallocBinned3.cpp b/Engine/Source/Runtime/Core/Private/HAL/MallocBinned3.cpp
--- a/Engine/Source/Runtime/Core/Private/HAL/MallocBinned3.cpp
+++ b/Engine/Source/Runtime/Core/Private/HAL/MallocBinned3.cpp
@@ -38,17 +38,17 @@ static FAutoConsoleVariableRef GMallocBinned3PerThreadCachesCVar(
 	TEXT("Enables per-thread caches of small (<= 32768 byte) allocations from FMallocBinned3")
 	);
 
-int32 GMallocBinned3BundleSize = DEFAULT_GMallocBinned3BundleSize;
+extern int32 GMallocBinnedBundleSize;
 static FAutoConsoleVariableRef GMallocBinned3BundleSizeCVar(
 	TEXT("MallocBinned3.BundleSize"),
-	GMallocBinned3BundleSize,
+	GMallocBinnedBundleSize,
 	TEXT("Max size in bytes of per-block bundles used in the recycling process")
 	);
 
-int32 GMallocBinned3BundleCount = DEFAULT_GMallocBinned3BundleCount;
+extern int32 GMallocBinnedBundleCount;
 static FAutoConsoleVariableRef GMallocBinned3BundleCountCVar(
 	TEXT("MallocBinned3.BundleCount"),
-	GMallocBinned3BundleCount,
+	GMallocBinnedBundleCount,
 	TEXT("Max count in blocks per-block bundles used in the recycling process")
 	);
 
@@ -88,7 +88,6 @@ TAtomic<int64> Binned3Decommits;
 int64 Binned3PoolInfoMemory = 0;
 int64 Binned3HashMemory = 0;
 int64 Binned3FreeBitsMemory = 0;
-int64 Binned3TLSMemory = 0;
 TAtomic<int64> Binned3TotalPoolSearches;
 TAtomic<int64> Binned3TotalPointerTests;
 
@@ -109,7 +108,6 @@ TAtomic<int32> MemoryRangeFreeTotalCount(0);
 
 MS_ALIGN(PLATFORM_CACHE_LINE_SIZE) static uint8 Binned3UnusedAlignPadding[PLATFORM_CACHE_LINE_SIZE] GCC_ALIGN(PLATFORM_CACHE_LINE_SIZE) = { 0 };
 uint16 FMallocBinned3::SmallBlockSizesReversedShifted[BINNED3_SMALL_POOL_COUNT + 1] = { 0 };
-uint32 FMallocBinned3::Binned3TlsSlot = FPlatformTLS::InvalidTlsSlot;
 uint32 FMallocBinned3::OsAllocationGranularity = 0;
 
 #if !BINNED3_USE_SEPARATE_VM_PER_POOL
@@ -319,42 +317,6 @@ public:
 	}
 };
 
-
-/** Hash table struct for retrieving allocation book keeping information */
-struct FMallocBinned3::PoolHashBucket
-{
-	UPTRINT BucketIndex;
-	FPoolInfoLarge* FirstPool;
-	PoolHashBucket* Prev;
-	PoolHashBucket* Next;
-
-	PoolHashBucket()
-	{
-		BucketIndex = 0;
-		FirstPool   = nullptr;
-		Prev        = this;
-		Next        = this;
-	}
-
-	void Link(PoolHashBucket* After)
-	{
-		After->Prev = Prev;
-		After->Next = this;
-		Prev ->Next = After;
-		this ->Prev = After;
-	}
-
-	void Unlink()
-	{
-		Next->Prev = Prev;
-		Prev->Next = Next;
-		Prev       = this;
-		Next       = this;
-	}
-};
-
-
-
 struct FMallocBinned3::Private
 {
 	// Implementation. 
@@ -656,16 +618,22 @@ struct FMallocBinned3::Private
 		FScopeLock Lock(&GetFreeBlockListsRegistrationMutex());
 		GetRegisteredFreeBlockLists().Remove(FreeBlockLists);
 #if BINNED3_ALLOCATOR_STATS
-		FMallocBinned3::FPerThreadFreeBlockLists::ConsolidatedMemory += FreeBlockLists->AllocatedMemory;
+		ConsolidatedMemory.fetch_add(FreeBlockLists->AllocatedMemory, std::memory_order_relaxed);
 #endif
 	}
 };
 
 FMallocBinned3::Private::FGlobalRecycler FMallocBinned3::Private::GGlobalRecycler;
 
-#if BINNED3_ALLOCATOR_STATS
-TAtomic<int64> FMallocBinned3::FPerThreadFreeBlockLists::ConsolidatedMemory;
-#endif
+void FMallocBinned3::RegisterThreadFreeBlockLists(FPerThreadFreeBlockLists* FreeBlockLists)
+{
+	Private::RegisterThreadFreeBlockLists(FreeBlockLists);
+}
+
+void FMallocBinned3::UnregisterThreadFreeBlockLists(FPerThreadFreeBlockLists* FreeBlockLists)
+{
+	Private::UnregisterThreadFreeBlockLists(FreeBlockLists);
+}
 
 FMallocBinned3::FPoolInfoSmall* FMallocBinned3::PushNewPoolToFront(FMallocBinned3::FPoolTable& Table, uint32 InBlockSize, uint32 InPoolIndex, uint32& OutBlockOfBlocksIndex)
 {
@@ -945,8 +913,8 @@ void FMallocBinned3::FreeMetaDataMemory(void *Ptr, SIZE_T InSize)
 	{
 		LLM_IF_ENABLED(FLowLevelMemTracker::Get().OnLowLevelFree(ELLMTracker::Platform, Ptr));
 
-		size_t VirtualAlignedSize = Align(InSize, FPlatformMemory::FPlatformVirtualMemoryBlock::GetVirtualSizeAlignment());
-		FPlatformMemory::FPlatformVirtualMemoryBlock Block(Ptr, VirtualAlignedSize / FPlatformMemory::FPlatformVirtualMemoryBlock::GetVirtualSizeAlignment());
+		InSize = Align(InSize, FPlatformMemory::FPlatformVirtualMemoryBlock::GetVirtualSizeAlignment());
+		FPlatformMemory::FPlatformVirtualMemoryBlock Block(Ptr, InSize / FPlatformMemory::FPlatformVirtualMemoryBlock::GetVirtualSizeAlignment());
 		Block.FreeVirtual();
 	}
 }
@@ -971,40 +939,16 @@ void* FMallocBinned3::MallocExternal(SIZE_T Size, uint32 Alignment)
 		// fallback: test for non-default/mininum alignments and handle a subset of them to reduce memory waste of 64KB page requirements
 		// 			 e.g. code that wants to use aligned avx loads will use 32 byte alignments 
 
-		Alignment			= FMath::Max<uint32>(Alignment, BINNED3_MINIMUM_ALIGNMENT);
-		size_t AlignedSize	= Align(Size, Alignment);
-
-		if ((AlignedSize <= BINNED3_MAX_SMALL_POOL_SIZE) && (Alignment <= BINNED3_MAX_SMALL_POOL_ALIGNMENT))
-		{
-			// Start at the naturally aligned size as that is guaranteed to be aligned for our given alignment in the allocator
-			// Note: We must handle the case where the naturally aligned size is not a valid pool entry and the next higher size
-			//       may not match the same alignment. In that case we do a search upwards as they are likely nearby aligned sizes 
-			//		 for our small range of alignemnts allowed
-			uint32 PoolIndex = BoundSizeToPoolIndex(AlignedSize);
-			do
-			{
-				uint32 BlockSize = PoolIndexToBlockSize(PoolIndex);
-				if (IsAligned(BlockSize, Alignment))
-				{
-					// we found a matching pool for our alignment and size requirements, so modify the size request to match
-					Size = SIZE_T(BlockSize);
-					UsePools = true;
-					Alignment = BINNED3_MINIMUM_ALIGNMENT;
-					break;
-				}
-
-				PoolIndex++;
-			} while (PoolIndex < BINNED3_SMALL_POOL_COUNT);
-		}
+		UsePools = PromoteToLargerBin(Size, Alignment, *this);
 	}
 
 	if (UsePools) 
 	{
-		uint32 PoolIndex = BoundSizeToPoolIndex(Size);
+		uint32 PoolIndex = BoundSizeToPoolIndex(Size, MemSizeToIndex);
 		FPerThreadFreeBlockLists* Lists = GMallocBinned3PerThreadCaches ? FPerThreadFreeBlockLists::Get() : nullptr;
 		if (Lists)
 		{
-			if (Lists->ObtainRecycledPartial(PoolIndex))
+			if (Lists->ObtainRecycledPartial(PoolIndex, Private::GGlobalRecycler))
 			{
 				if (void* Result = Lists->Malloc(PoolIndex))
 				{
@@ -1224,7 +1168,7 @@ void FMallocBinned3::FreeExternal(void* Ptr)
 		FPerThreadFreeBlockLists* Lists = GMallocBinned3PerThreadCaches ? FPerThreadFreeBlockLists::Get() : nullptr;
 		if (Lists)
 		{
-			BundlesToRecycle = Lists->RecycleFullBundle(PoolIndex);
+			BundlesToRecycle = Lists->RecycleFullBundle(PoolIndex, Private::GGlobalRecycler);
 			bool bPushed = Lists->Free(Ptr, PoolIndex, BlockSize);
 			check(bPushed);
 #if BINNED3_ALLOCATOR_STATS
@@ -1405,11 +1349,11 @@ void FMallocBinned3::SetupTLSCachesOnCurrentThread()
 	{
 		return;
 	}
-	if (!FPlatformTLS::IsValidTlsSlot(FMallocBinned3::Binned3TlsSlot))
+	if (!FPlatformTLS::IsValidTlsSlot(FMallocBinned3::BinnedTlsSlot))
 	{
-		FMallocBinned3::Binned3TlsSlot = FPlatformTLS::AllocTlsSlot();
+		FMallocBinned3::BinnedTlsSlot = FPlatformTLS::AllocTlsSlot();
 	}
-	check(FPlatformTLS::IsValidTlsSlot(FMallocBinned3::Binned3TlsSlot));
+	check(FPlatformTLS::IsValidTlsSlot(FMallocBinned3::BinnedTlsSlot));
 	FPerThreadFreeBlockLists::SetTLS();
 }
 
@@ -1423,105 +1367,6 @@ void FMallocBinned3::ClearAndDisableTLSCachesOnCurrentThread()
 	FPerThreadFreeBlockLists::ClearTLS();
 }
 
-
-bool FMallocBinned3::FFreeBlockList::ObtainPartial(uint32 InPoolIndex)
-{
-	if (!PartialBundle.Head)
-	{
-		PartialBundle.Count = 0;
-		PartialBundle.Head = FMallocBinned3::Private::GGlobalRecycler.PopBundle(InPoolIndex);
-		if (PartialBundle.Head)
-		{
-			PartialBundle.Count = PartialBundle.Head->Count;
-			PartialBundle.Head->NextBundle = nullptr;
-			return true;
-		}
-		return false;
-	}
-	return true;
-}
-
-FMallocBinned3::FBundleNode* FMallocBinned3::FFreeBlockList::RecyleFull(uint32 InPoolIndex)
-{
-	FMallocBinned3::FBundleNode* Result = nullptr;
-	if (FullBundle.Head)
-	{
-		FullBundle.Head->Count = FullBundle.Count;
-		if (!FMallocBinned3::Private::GGlobalRecycler.PushBundle(InPoolIndex, FullBundle.Head))
-		{
-			Result = FullBundle.Head;
-			Result->NextBundle = nullptr;
-		}
-		FullBundle.Reset();
-	}
-	return Result;
-}
-
-FMallocBinned3::FBundleNode* FMallocBinned3::FFreeBlockList::PopBundles(uint32 InPoolIndex)
-{
-	FBundleNode* Partial = PartialBundle.Head;
-	if (Partial)
-	{
-		PartialBundle.Reset();
-		Partial->NextBundle = nullptr;
-	}
-
-	FBundleNode* Full = FullBundle.Head;
-	if (Full)
-	{
-		FullBundle.Reset();
-		Full->NextBundle = nullptr;
-	}
-
-	FBundleNode* Result = Partial;
-	if (Result)
-	{
-		Result->NextBundle = Full;
-	}
-	else
-	{
-		Result = Full;
-	}
-
-	return Result;
-}
-
-void FMallocBinned3::FPerThreadFreeBlockLists::SetTLS()
-{
-	check(FPlatformTLS::IsValidTlsSlot(FMallocBinned3::Binned3TlsSlot));
-	FPerThreadFreeBlockLists* ThreadSingleton = (FPerThreadFreeBlockLists*)FPlatformTLS::GetTlsValue(FMallocBinned3::Binned3TlsSlot);
-	if (!ThreadSingleton)
-	{
-		const int64 TLSSize = Align(sizeof(FPerThreadFreeBlockLists), FPlatformMemory::FPlatformVirtualMemoryBlock::GetVirtualSizeAlignment());
-		ThreadSingleton = new (FMallocBinned3::AllocateMetaDataMemory(TLSSize)) FPerThreadFreeBlockLists();
-#if BINNED3_ALLOCATOR_STATS
-		Binned3TLSMemory += TLSSize;
-#endif
-		verify(ThreadSingleton);
-		FPlatformTLS::SetTlsValue(FMallocBinned3::Binned3TlsSlot, ThreadSingleton);
-		FMallocBinned3::Private::RegisterThreadFreeBlockLists(ThreadSingleton);
-	}
-}
-
-void FMallocBinned3::FPerThreadFreeBlockLists::ClearTLS()
-{
-	check(FPlatformTLS::IsValidTlsSlot(FMallocBinned3::Binned3TlsSlot));
-	FPerThreadFreeBlockLists* ThreadSingleton = (FPerThreadFreeBlockLists*)FPlatformTLS::GetTlsValue(FMallocBinned3::Binned3TlsSlot);
-	if (ThreadSingleton)
-	{
-		const int64 TLSSize = Align(sizeof(FPerThreadFreeBlockLists), FPlatformMemory::FPlatformVirtualMemoryBlock::GetVirtualSizeAlignment());
-#if BINNED3_ALLOCATOR_STATS
-		Binned3TLSMemory -= TLSSize;
-#endif
-		FMallocBinned3::Private::UnregisterThreadFreeBlockLists(ThreadSingleton);
-	
-		ThreadSingleton->~FPerThreadFreeBlockLists();
-
-		FMallocBinned3::FreeMetaDataMemory(ThreadSingleton, TLSSize);
-	}
-	FPlatformTLS::SetTlsValue(FMallocBinned3::Binned3TlsSlot, nullptr);
-}
-
 void FMallocBinned3::FFreeBlock::CanaryFail() const
 {
 	UE_LOG(LogMemory, Fatal, TEXT("FMallocBinned3 Attempt to realloc an unrecognized block %p   canary == 0x%x != 0x%x"), (void*)this, (int32)Canary, (int32)FMallocBinned3::FFreeBlock::CANARY_VALUE);
@@ -1537,7 +1382,7 @@ int64 FMallocBinned3::GetTotalAllocatedSmallPoolMemory() const
 		{
 			FreeBlockAllocatedMemory += FreeBlockLists->AllocatedMemory;
 		}
-		FreeBlockAllocatedMemory += FPerThreadFreeBlockLists::ConsolidatedMemory;
+		FreeBlockAllocatedMemory += ConsolidatedMemory.load(std::memory_order_relaxed);
 	}
 
 	return Binned3AllocatedSmallPoolMemory + FreeBlockAllocatedMemory;
@@ -1590,7 +1435,7 @@ void FMallocBinned3::DumpAllocatorStats(class FOutputDevice& Ar)
 	Ar.Logf(TEXT("PoolInfo: %fmb"), ((double)Binned3PoolInfoMemory) / (1024.0f * 1024.0f));
 	Ar.Logf(TEXT("Hash: %fmb"), ((double)Binned3HashMemory) / (1024.0f * 1024.0f));
 	Ar.Logf(TEXT("Free Bits: %fmb"), ((double)Binned3FreeBitsMemory) / (1024.0f * 1024.0f));
-	Ar.Logf(TEXT("TLS: %fmb"), ((double)Binned3TLSMemory) / (1024.0f * 1024.0f));
+	Ar.Logf(TEXT("TLS: %fmb"), ((double)TLSMemory.load(std::memory_order_relaxed) / (1024.0f * 1024.0f));
 	Ar.Logf(TEXT("Slab Commits: %llu"), Binned3Commits.Load());
 	Ar.Logf(TEXT("Slab Decommits: %llu"), Binned3Decommits.Load());
 #if BINNED3_USE_SEPARATE_VM_PER_POOL
@@ -1607,7 +1452,7 @@ void FMallocBinned3::DumpAllocatorStats(class FOutputDevice& Ar)
 #endif
 	Ar.Logf(TEXT("Total allocated from OS: %fmb"), 
 		((double)
-			Binned3AllocatedOSSmallPoolMemory + Binned3AllocatedLargePoolMemoryWAlignment + Binned3PoolInfoMemory + Binned3HashMemory + Binned3FreeBitsMemory + Binned3TLSMemory
+			Binned3AllocatedOSSmallPoolMemory + Binned3AllocatedLargePoolMemoryWAlignment + Binned3PoolInfoMemory + Binned3HashMemory + Binned3FreeBitsMemory + TLSMemory.load(std::memory_order_relaxed)
 			) / (1024.0f * 1024.0f));
 
 
diff --git a/Engine/Source/Runtime/Core/Private/HAL/MallocBinnedCommon.cpp b/Engine/Source/Runtime/Core/Private/HAL/MallocBinnedCommon.cpp
--- a/Engine/Source/Runtime/Core/Private/HAL/MallocBinnedCommon.cpp
+++ b/Engine/Source/Runtime/Core/Private/HAL/MallocBinnedCommon.cpp
@@ -474,4 +474,28 @@ uint32 FBitTree::CountOnes(uint32 UpTo) const
 
 #endif
 
+#if UE_BINNEDCOMMON_ALLOW_RUNTIME_TWEAKING
+
+int32 GMallocBinnedBundleSize = DEFAULT_GMallocBinnedBundleSize;
+static FAutoConsoleVariableRef GMallocBinned3BundleSizeCVar(
+	TEXT("MallocBinned.BundleSize"),
+	GMallocBinnedBundleSize,
+	TEXT("Max size in bytes of per-block bundles used in the recycling process")
+);
+
+int32 GMallocBinnedBundleCount = DEFAULT_GMallocBinnedBundleCount;
+static FAutoConsoleVariableRef GMallocBinned3BundleCountCVar(
+	TEXT("MallocBinned.BundleCount"),
+	GMallocBinnedBundleCount,
+	TEXT("Max count in blocks per-block bundles used in the recycling process")
+);
+
+#endif
+
+uint32 FMallocBinnedCommonBase::BinnedTlsSlot = FPlatformTLS::InvalidTlsSlot;
+#if UE_BINNEDCOMMON_ALLOCATOR_STATS
+std::atomic<int64> FMallocBinnedCommonBase::TLSMemory(0);
+std::atomic<int64> FMallocBinnedCommonBase::ConsolidatedMemory(0);
+#endif
+
 PRAGMA_RESTORE_UNSAFE_TYPECAST_WARNINGS
diff --git a/Engine/Source/Runtime/Core/Public/HAL/MallocBinned2.h b/Engine/Source/Runtime/Core/Public/HAL/MallocBinned2.h
--- a/Engine/Source/Runtime/Core/Public/HAL/MallocBinned2.h
+++ b/Engine/Source/Runtime/Core/Public/HAL/MallocBinned2.h
@@ -9,8 +9,8 @@
 #include "HAL/CriticalSection.h"
 #include "HAL/LowLevelMemTracker.h"
 #include "HAL/MemoryBase.h"
+#include "HAL/MallocBinnedCommon.h"
 #include "HAL/PlatformMath.h"
-#include "HAL/PlatformTLS.h"
 #include "HAL/UnrealMemory.h"
 #include "Math/NumericLimits.h"
 #include "Misc/AssertionMacros.h"
@@ -37,20 +37,10 @@ struct FGenericMemoryStats;
 
 #define DEFAULT_GMallocBinned2PerThreadCaches 1
 #define DEFAULT_GMallocBinned2LockFreeCaches 0
-#define DEFAULT_GMallocBinned2BundleCount 64
 #define DEFAULT_GMallocBinned2AllocExtra 32
 #define BINNED2_MAX_GMallocBinned2MaxBundlesBeforeRecycle 8
 #define DEFAULT_GMallocBinned2MoveOSFreesOffTimeCriticalThreads 1
 
-#if !defined(AGGRESSIVE_MEMORY_SAVING)
-	#error "AGGRESSIVE_MEMORY_SAVING must be defined"
-#endif
-#if AGGRESSIVE_MEMORY_SAVING
-	#define DEFAULT_GMallocBinned2BundleSize 8192
-#else
-	#define DEFAULT_GMallocBinned2BundleSize BINNED2_LARGE_ALLOC
-#endif
-
 // When book keeping is at the end of FFreeBlock, MallocBinned2 cannot tell if the allocation comes from a large allocation (higher than 64KB, also named as "OSAllocation") 
 // or from VeryLargePageAllocator that fell back to FCachedOSPageAllocator. In both cases the allocation (large or small) might be aligned to 64KB.
 // bookKeeping at the end needs to be disabled if we want VeryLargePageAllocator to properly fallback to regular TCachedOSPageAllocator if needed
@@ -64,36 +54,24 @@ struct FGenericMemoryStats;
 #endif
 
 
-#define BINNED2_ALLOW_RUNTIME_TWEAKING 0
+#define BINNED2_ALLOW_RUNTIME_TWEAKING UE_BINNEDCOMMON_ALLOW_RUNTIME_TWEAKING
 #if BINNED2_ALLOW_RUNTIME_TWEAKING
 	extern CORE_API int32 GMallocBinned2PerThreadCaches;
-	extern CORE_API int32 GMallocBinned2BundleSize = DEFAULT_GMallocBinned2BundleSize;
-	extern CORE_API int32 GMallocBinned2BundleCount = DEFAULT_GMallocBinned2BundleCount;
-	extern CORE_API int32 GMallocBinned2MaxBundlesBeforeRecycle = BINNED2_MAX_GMallocBinned2MaxBundlesBeforeRecycle;
-	extern CORE_API int32 GMallocBinned2AllocExtra = DEFAULT_GMallocBinned2AllocExtra;
-	extern CORE_API int32 GMallocBinned2MoveOSFreesOffTimeCriticalThreads = DEFAULT_GMallocBinned2MoveOSFreesOffTimeCriticalThreads;
+	extern CORE_API int32 GMallocBinned2MaxBundlesBeforeRecycle;
+	extern CORE_API int32 GMallocBinned2AllocExtra;
+	extern CORE_API int32 GMallocBinned2MoveOSFreesOffTimeCriticalThreads;
 
 #else
 	#define GMallocBinned2PerThreadCaches DEFAULT_GMallocBinned2PerThreadCaches
-	#define GMallocBinned2BundleSize DEFAULT_GMallocBinned2BundleSize
-	#define GMallocBinned2BundleCount DEFAULT_GMallocBinned2BundleCount
 	#define GMallocBinned2MaxBundlesBeforeRecycle BINNED2_MAX_GMallocBinned2MaxBundlesBeforeRecycle
 	#define GMallocBinned2AllocExtra DEFAULT_GMallocBinned2AllocExtra
 	#define GMallocBinned2MoveOSFreesOffTimeCriticalThreads DEFAULT_GMallocBinned2MoveOSFreesOffTimeCriticalThreads
 #endif
 
-
-#ifndef BINNED2_ALLOCATOR_STATS
-	#if UE_BUILD_SHIPPING && !WITH_EDITOR
-		#define BINNED2_ALLOCATOR_STATS 0	
-	#else
-		#define BINNED2_ALLOCATOR_STATS 1
-	#endif
-#endif
-
-
+#define BINNED2_ALLOCATOR_STATS UE_BINNEDCOMMON_ALLOCATOR_STATS
 #define BINNED2_ALLOCATOR_STATS_VALIDATION (BINNED2_ALLOCATOR_STATS && 0)
 
+
 #if BINNED2_ALLOCATOR_STATS
 //////////////////////////////////////////////////////////////////////////
 // the following don't need a critical section because they are covered by the critical section called Mutex
@@ -128,11 +106,11 @@ enum class EBlockCanary : uint8
 //
 // Optimized virtual memory allocator.
 //
-class FMallocBinned2 : public FMalloc
+class FMallocBinned2 : public TMallocBinnedCommon<FMallocBinned2, BINNED2_MINIMUM_ALIGNMENT, BINNED2_MAXIMUM_ALIGNMENT, BINNED2_MINIMUM_ALIGNMENT_SHIFT, BINNED2_SMALL_POOL_COUNT, BINNED2_MAX_SMALL_POOL_SIZE>
 {
 	// Forward declares.
 	struct FPoolInfo;
-	struct PoolHashBucket;
+	using PoolHashBucket = TPoolHashBucket<FPoolInfo>;
 	struct Private;
 
 	/** Information about a piece of free memory. */
@@ -222,58 +200,6 @@ class FMallocBinned2 : public FMalloc
 		FPoolTable();
 	};
 
-	struct FPtrToPoolMapping
-	{
-		FPtrToPoolMapping()
-			: PtrToPoolPageBitShift(0)
-			, HashKeyShift(0)
-			, PoolMask(0)
-			, MaxHashBuckets(0)
-		{
-		}
-		explicit FPtrToPoolMapping(uint32 InPageSize, uint64 InNumPoolsPerPage, uint64 AddressLimit)
-		{
-			Init(InPageSize, InNumPoolsPerPage, AddressLimit);
-		}
-
-		void Init(uint32 InPageSize, uint64 InNumPoolsPerPage, uint64 AddressLimit)
-		{
-			uint64 PoolPageToPoolBitShift = FPlatformMath::CeilLogTwo64(InNumPoolsPerPage);
-
-			PtrToPoolPageBitShift = FPlatformMath::CeilLogTwo(InPageSize);
-			HashKeyShift          = PtrToPoolPageBitShift + PoolPageToPoolBitShift;
-			PoolMask              = (1ull << PoolPageToPoolBitShift) - 1;
-			MaxHashBuckets        = AddressLimit >> HashKeyShift;
-		}
-
-		FORCEINLINE void GetHashBucketAndPoolIndices(const void* InPtr, uint32& OutBucketIndex, UPTRINT& OutBucketCollision, uint32& OutPoolIndex) const
-		{
-			OutBucketCollision = (UPTRINT)InPtr >> HashKeyShift;
-			OutBucketIndex = uint32(OutBucketCollision & (MaxHashBuckets - 1));
-			OutPoolIndex   = (uint32)(((UPTRINT)InPtr >> PtrToPoolPageBitShift) & PoolMask);
-		}
-
-		FORCEINLINE uint64 GetMaxHashBuckets() const
-		{
-			return MaxHashBuckets;
-		}
-
-	private:
-		/** Shift to apply to a pointer to get the reference from the indirect tables */
-		uint64 PtrToPoolPageBitShift;
-
-		/** Shift required to get required hash table key. */
-		uint64 HashKeyShift;
-
-		/** Used to mask off the bits that have been used to lookup the indirect table */
-		uint64 PoolMask;
-
-		// PageSize dependent constants
-		uint64 MaxHashBuckets;
-	};
-
-	FPtrToPoolMapping PtrToPoolMapping;
-
 	// Pool tables for different pool sizes
 	FPoolTable SmallPoolTables[BINNED2_SMALL_POOL_COUNT];
 
@@ -308,156 +234,6 @@ class FMallocBinned2 : public FMalloc
 #endif
 	}
 
-	// This needs to be small enough to fit inside the smallest allocation handled by MallocBinned2, hence the union.
-	struct FBundleNode
-	{
-		FBundleNode* NextNodeInCurrentBundle;
-
-		// NextBundle ptr is valid when node is stored in FFreeBlockList in a thread-local list of reusable allocations.
-		// Count is valid when node is stored in global recycler and caches the number of nodes in the list formed by NextNodeInCurrentBundle.
-		union
-		{
-			FBundleNode* NextBundle;
-			int32 Count;
-		};
-	};
-
-	struct FBundle
-	{
-		FORCEINLINE FBundle()
-		{
-			Reset();
-		}
-
-		FORCEINLINE void Reset()
-		{
-			Head = nullptr;
-			Count = 0;
-		}
-
-		FORCEINLINE void PushHead(FBundleNode* Node)
-		{
-			Node->NextNodeInCurrentBundle = Head;
-			Node->NextBundle = nullptr;
-			Head = Node;
-			Count++;
-		}
-
-		FORCEINLINE FBundleNode* PopHead()
-		{
-			FBundleNode* Result = Head;
-
-			Count--;
-			Head = Head->NextNodeInCurrentBundle;
-			return Result;
-		}
-
-		FBundleNode* Head;
-		uint32       Count;
-	};
-
-	struct FFreeBlockList
-	{
-		// return true if we actually pushed it
-		FORCEINLINE bool PushToFront(void* InPtr, uint32 InPoolIndex, uint32 InBlockSize)
-		{
-			checkSlow(InPtr);
-
-			if (PartialBundle.Count >= (uint32)GMallocBinned2BundleCount || PartialBundle.Count * InBlockSize >= (uint32)GMallocBinned2BundleSize)
-			{
-				if (FullBundle.Head)
-				{
-					return false;
-				}
-				FullBundle = PartialBundle;
-				PartialBundle.Reset();
-			}
-			PartialBundle.PushHead((FBundleNode*)InPtr);
-			return true;
-		}
-		FORCEINLINE bool CanPushToFront(uint32 InPoolIndex, uint32 InBlockSize)
-		{
-			if (FullBundle.Head && (PartialBundle.Count >= (uint32)GMallocBinned2BundleCount || PartialBundle.Count * InBlockSize >= (uint32)GMallocBinned2BundleSize))
-			{
-				return false;
-			}
-			return true;
-		}
-		FORCEINLINE void* PopFromFront(uint32 InPoolIndex)
-		{
-			if (!PartialBundle.Head)
-			{
-				if (FullBundle.Head)
-				{
-					PartialBundle = FullBundle;
-					FullBundle.Reset();
-				}
-			}
-			return PartialBundle.Head ? PartialBundle.PopHead() : nullptr;
-		}
-
-		// tries to recycle the full bundle, if that fails, it is returned for freeing
-		FBundleNode* RecyleFull(uint32 InPoolIndex);
-		bool ObtainPartial(uint32 InPoolIndex);
-		FBundleNode* PopBundles(uint32 InPoolIndex);
-	private:
-		FBundle PartialBundle;
-		FBundle FullBundle;
-	};
-
-	struct FPerThreadFreeBlockLists
-	{
-		FORCEINLINE static FPerThreadFreeBlockLists* Get()
-		{
-			return FPlatformTLS::IsValidTlsSlot(FMallocBinned2::Binned2TlsSlot) ? (FPerThreadFreeBlockLists*)FPlatformTLS::GetTlsValue(FMallocBinned2::Binned2TlsSlot) : nullptr;
-		}
-		static void SetTLS();
-		static void ClearTLS();
-
-		FPerThreadFreeBlockLists() 
-#if BINNED2_ALLOCATOR_STATS
-			: AllocatedMemory(0) 
-#endif
-		{ 
-		}
-
-		FORCEINLINE void* Malloc(uint32 InPoolIndex)
-		{
-			return FreeLists[InPoolIndex].PopFromFront(InPoolIndex);
-		}
-		// return true if the pointer was pushed
-		FORCEINLINE bool Free(void* InPtr, uint32 InPoolIndex, uint32 InBlockSize)
-		{
-			return FreeLists[InPoolIndex].PushToFront(InPtr, InPoolIndex, InBlockSize);
-		}		
-		// return true if a pointer can be pushed
-		FORCEINLINE bool CanFree(uint32 InPoolIndex, uint32 InBlockSize)
-		{
-			return FreeLists[InPoolIndex].CanPushToFront(InPoolIndex, InBlockSize);
-		}
-		// returns a bundle that needs to be freed if it can't be recycled
-		FBundleNode* RecycleFullBundle(uint32 InPoolIndex)
-		{
-			return FreeLists[InPoolIndex].RecyleFull(InPoolIndex);
-		}
-		// returns true if we have anything to pop
-		bool ObtainRecycledPartial(uint32 InPoolIndex)
-		{
-			return FreeLists[InPoolIndex].ObtainPartial(InPoolIndex);
-		}
-		FBundleNode* PopBundles(uint32 InPoolIndex)
-		{
-			return FreeLists[InPoolIndex].PopBundles(InPoolIndex);
-		}
-#if BINNED2_ALLOCATOR_STATS
-	public:
-		int64 AllocatedMemory;
-		static int64 ConsolidatedMemory;
-#endif
-	private:
-		FFreeBlockList FreeLists[BINNED2_SMALL_POOL_COUNT];
-	};
-
 	static FORCEINLINE FFreeBlock* GetPoolHeaderFromPointer(void* Ptr)
 	{
 #if !UE_USE_VERYLARGEPAGEALLOCATOR || !BINNED2_BOOKKEEPING_AT_THE_END_OF_LARGEBLOCK
@@ -512,7 +288,7 @@ public:
 			FPerThreadFreeBlockLists* Lists = GMallocBinned2PerThreadCaches ? FPerThreadFreeBlockLists::Get() : nullptr;
 			if (Lists)
 			{
-				const uint32 PoolIndex = BoundSizeToPoolIndex(Size);
+				const uint32 PoolIndex = BoundSizeToPoolIndex(Size, MemSizeToIndex);
 				if (void* Result = Lists->Malloc(PoolIndex))
 				{
 #if BINNED2_ALLOCATOR_STATS
@@ -615,7 +391,7 @@ public:
 				}
 				if (bCanFree)
 				{
-					uint32 NewPoolIndex = BoundSizeToPoolIndex(NewSize);
+					uint32 NewPoolIndex = BoundSizeToPoolIndex(NewSize, MemSizeToIndex);
 					uint32 NewBlockSize = PoolIndexToBlockSize(NewPoolIndex);
 					void* Result = NewSize ? Lists->Malloc(NewPoolIndex) : nullptr;
 #if BINNED2_ALLOCATOR_STATS
@@ -715,39 +491,7 @@ public:
 
 	FORCEINLINE virtual SIZE_T QuantizeSize(SIZE_T Count, uint32 Alignment) override
 	{
-		static_assert(DEFAULT_ALIGNMENT <= BINNED2_MINIMUM_ALIGNMENT, "DEFAULT_ALIGNMENT is assumed to be zero"); // used below
-		checkSlow((Alignment & (Alignment - 1)) == 0); // Check the alignment is a power of two
-		SIZE_T SizeOut;
-		if ((Count <= BINNED2_MAX_SMALL_POOL_SIZE) & (Alignment <= BINNED2_MINIMUM_ALIGNMENT)) // one branch, not two
-		{
-			SizeOut = PoolIndexToBlockSize(BoundSizeToPoolIndex(Count));
-			check(SizeOut >= Count);
-			return SizeOut;
-		}
-		Alignment = FMath::Max<uint32>(Alignment, BINNED2_MINIMUM_ALIGNMENT);
-		Count = Align(Count, Alignment);
-		if ((Count <= BINNED2_MAX_SMALL_POOL_SIZE) & (Alignment <= BINNED2_MAXIMUM_ALIGNMENT))
-		{
-			uint32 PoolIndex = BoundSizeToPoolIndex(Count);
-			do
-			{
-				uint32 BlockSize = PoolIndexToBlockSize(PoolIndex);
-				if (IsAligned(BlockSize, Alignment))
-				{
-					SizeOut = SIZE_T(BlockSize);
-					check(SizeOut >= Count);
-					return SizeOut;
-				}
-
-				PoolIndex++;
-			} while (PoolIndex < BINNED2_SMALL_POOL_COUNT);
-		}
-
-		Alignment = FPlatformMath::Max<uint32>(Alignment, OsAllocationGranularity);
-		checkSlow(Alignment <= PageSize);
-		SizeOut = Align(Count, Alignment);
-		check(SizeOut >= Count);
-		return SizeOut;
+		return QuantizeSizeCommon(Count, Alignment, *this);
 	}
 
 	CORE_API virtual bool ValidateHeap() override;
@@ -780,20 +524,17 @@ public:
 	
 	static CORE_API uint16 SmallBlockSizesReversed[BINNED2_SMALL_POOL_COUNT]; // this is reversed to get the smallest elements on our main cache line
 	static CORE_API FMallocBinned2* MallocBinned2;
-	static CORE_API uint32 Binned2TlsSlot;
 	static CORE_API uint32 PageSize;
 	static CORE_API uint32 OsAllocationGranularity;
 	// Mapping of sizes to small table indices
 	static CORE_API uint8 MemSizeToIndex[1 + (BINNED2_MAX_SMALL_POOL_SIZE >> BINNED2_MINIMUM_ALIGNMENT_SHIFT)];
 
-	FORCEINLINE uint32 BoundSizeToPoolIndex(SIZE_T Size) const
-	{
-		auto Index = ((Size + BINNED2_MINIMUM_ALIGNMENT - 1) >> BINNED2_MINIMUM_ALIGNMENT_SHIFT);
-		checkSlow(Index >= 0 && Index <= (BINNED2_MAX_SMALL_POOL_SIZE >> BINNED2_MINIMUM_ALIGNMENT_SHIFT)); // and it should be in the table
-		uint32 PoolIndex = uint32(MemSizeToIndex[Index]);
-		checkSlow(PoolIndex >= 0 && PoolIndex < BINNED2_SMALL_POOL_COUNT);
-		return PoolIndex;
-	}
+	static void RegisterThreadFreeBlockLists(FPerThreadFreeBlockLists* FreeBlockLists);
+	static void UnregisterThreadFreeBlockLists(FPerThreadFreeBlockLists* FreeBlockLists);
+
+	static void* AllocateMetaDataMemory(SIZE_T Size);
+	static void FreeMetaDataMemory(void* Ptr, SIZE_T Size);
+
 	FORCEINLINE uint32 PoolIndexToBlockSize(uint32 PoolIndex) const
 	{
 		return SmallBlockSizesReversed[BINNED2_SMALL_POOL_COUNT - PoolIndex - 1];
diff --git a/Engine/Source/Runtime/Core/Public/HAL/MallocBinned3.h b/Engine/Source/Runtime/Core/Public/HAL/MallocBinned3.h
--- a/Engine/Source/Runtime/Core/Public/HAL/MallocBinned3.h
+++ b/Engine/Source/Runtime/Core/Public/HAL/MallocBinned3.h
@@ -13,7 +13,6 @@
 #include "HAL/MemoryBase.h"
 #include "HAL/PlatformMath.h"
 #include "HAL/PlatformMemory.h"
-#include "HAL/PlatformTLS.h"
 #include "HAL/UnrealMemory.h"
 #include "Math/NumericLimits.h"
 #include "Misc/AssertionMacros.h"
@@ -64,44 +63,21 @@ struct FGenericMemoryStats;
 #endif
 
 #define DEFAULT_GMallocBinned3PerThreadCaches 1
-#define DEFAULT_GMallocBinned3BundleCount 64
 #define DEFAULT_GMallocBinned3AllocExtra 32
 #define BINNED3_MAX_GMallocBinned3MaxBundlesBeforeRecycle 8
 
-#if !defined(AGGRESSIVE_MEMORY_SAVING)
-	#error "AGGRESSIVE_MEMORY_SAVING must be defined"
-#endif
-#if AGGRESSIVE_MEMORY_SAVING
-	#define DEFAULT_GMallocBinned3BundleSize 8192
-#else
-	#define DEFAULT_GMallocBinned3BundleSize 65536
-#endif
-
-
-#define BINNED3_ALLOW_RUNTIME_TWEAKING 0
+#define BINNED3_ALLOW_RUNTIME_TWEAKING UE_BINNEDCOMMON_ALLOW_RUNTIME_TWEAKING
 #if BINNED3_ALLOW_RUNTIME_TWEAKING
 	extern CORE_API int32 GMallocBinned3PerThreadCaches;
-	extern CORE_API int32 GMallocBinned3BundleSize = DEFAULT_GMallocBinned3BundleSize;
-	extern CORE_API int32 GMallocBinned3BundleCount = DEFAULT_GMallocBinned3BundleCount;
-	extern CORE_API int32 GMallocBinned3MaxBundlesBeforeRecycle = BINNED3_MAX_GMallocBinned3MaxBundlesBeforeRecycle;
-	extern CORE_API int32 GMallocBinned3AllocExtra = DEFAULT_GMallocBinned3AllocExtra;
+	extern CORE_API int32 GMallocBinned3MaxBundlesBeforeRecycle;
+	extern CORE_API int32 GMallocBinned3AllocExtra;
 #else
 	#define GMallocBinned3PerThreadCaches DEFAULT_GMallocBinned3PerThreadCaches
-	#define GMallocBinned3BundleSize DEFAULT_GMallocBinned3BundleSize
-	#define GMallocBinned3BundleCount DEFAULT_GMallocBinned3BundleCount
 	#define GMallocBinned3MaxBundlesBeforeRecycle BINNED3_MAX_GMallocBinned3MaxBundlesBeforeRecycle
 	#define GMallocBinned3AllocExtra DEFAULT_GMallocBinned3AllocExtra
 #endif
 
-
-#ifndef BINNED3_ALLOCATOR_STATS
-	#if UE_BUILD_SHIPPING && !WITH_EDITOR
-		#define BINNED3_ALLOCATOR_STATS 0	
-	#else
-		#define BINNED3_ALLOCATOR_STATS 1
-	#endif
-#endif
-
+#define BINNED3_ALLOCATOR_STATS UE_BINNEDCOMMON_ALLOCATOR_STATS
 
 #if BINNED3_ALLOCATOR_STATS
 	#if !(UE_BUILD_SHIPPING || UE_BUILD_TEST)
@@ -113,18 +89,19 @@ struct FGenericMemoryStats;
 	#define BINNED3_ALLOCATOR_PER_BIN_STATS 0
 #endif
 
+
 PRAGMA_DISABLE_UNSAFE_TYPECAST_WARNINGS
 
 //
 // Optimized virtual memory allocator.
 //
-class CORE_API FMallocBinned3 : public FMalloc
+class CORE_API FMallocBinned3 : public TMallocBinnedCommon<FMallocBinned3, BINNED3_MINIMUM_ALIGNMENT, BINNED3_MAX_SMALL_POOL_ALIGNMENT, BINNED3_MINIMUM_ALIGNMENT_SHIFT, BINNED3_SMALL_POOL_COUNT, BINNED3_MAX_SMALL_POOL_SIZE>
 {	
 	// Forward declares.
 	struct FPoolInfoLarge;
 	struct FPoolInfoSmall;
 	struct FPoolTable;
-	struct PoolHashBucket;
+	using PoolHashBucket = TPoolHashBucket<FPoolInfoLarge>;
 	struct Private;
 
 
@@ -219,58 +196,6 @@ class CORE_API FMallocBinned3 : public FMalloc
 #endif
 	};
 
-	struct FPtrToPoolMapping
-	{
-		FPtrToPoolMapping()
-			: PtrToPoolPageBitShift(0)
-			, HashKeyShift(0)
-			, PoolMask(0)
-			, MaxHashBuckets(0)
-		{
-		}
-		explicit FPtrToPoolMapping(uint32 InPageSize, uint64 InNumPoolsPerPage, uint64 AddressLimit)
-		{
-			Init(InPageSize, InNumPoolsPerPage, AddressLimit);
-		}
-
-		void Init(uint32 InPageSize, uint64 InNumPoolsPerPage, uint64 AddressLimit)
-		{
-			uint64 PoolPageToPoolBitShift = FPlatformMath::CeilLogTwo(InNumPoolsPerPage);
-
-			PtrToPoolPageBitShift = FPlatformMath::CeilLogTwo(InPageSize);
-			HashKeyShift          = PtrToPoolPageBitShift + PoolPageToPoolBitShift;
-			PoolMask              = (1ull << PoolPageToPoolBitShift) - 1;
-			MaxHashBuckets        = AddressLimit >> HashKeyShift;
-		}
-
-		FORCEINLINE void GetHashBucketAndPoolIndices(const void* InPtr, uint32& OutBucketIndex, UPTRINT& OutBucketCollision, uint32& OutPoolIndex) const
-		{
-			OutBucketCollision = (UPTRINT)InPtr >> HashKeyShift;
-			OutBucketIndex = uint32(OutBucketCollision & (MaxHashBuckets - 1));
-			OutPoolIndex   = ((UPTRINT)InPtr >> PtrToPoolPageBitShift) & PoolMask;
-		}
-
-		FORCEINLINE uint64 GetMaxHashBuckets() const
-		{
-			return MaxHashBuckets;
-		}
-
-	private:
-		/** Shift to apply to a pointer to get the reference from the indirect tables */
-		uint64 PtrToPoolPageBitShift;
-
-		/** Shift required to get required hash table key. */
-		uint64 HashKeyShift;
-
-		/** Used to mask off the bits that have been used to lookup the indirect table */
-		uint64 PoolMask;
-
-		// PageSize dependent constants
-		uint64 MaxHashBuckets;
-	};
-
-	FPtrToPoolMapping PtrToPoolMapping;
-
 	// Pool tables for different pool sizes
 	FPoolTable SmallPoolTables[BINNED3_SMALL_POOL_COUNT];
 
@@ -282,145 +207,6 @@ class CORE_API FMallocBinned3 : public FMalloc
 
 	FCriticalSection Mutex;
 
-	struct FBundleNode
-	{
-		FBundleNode* NextNodeInCurrentBundle;
-		union
-		{
-			FBundleNode* NextBundle;
-			int32 Count;
-		};
-	};
-
-	struct FBundle
-	{
-		FORCEINLINE FBundle()
-		{
-			Reset();
-		}
-
-		FORCEINLINE void Reset()
-		{
-			Head = nullptr;
-			Count = 0;
-		}
-
-		FORCEINLINE void PushHead(FBundleNode* Node)
-		{
-			Node->NextNodeInCurrentBundle = Head;
-			Node->NextBundle = nullptr;
-			Head = Node;
-			Count++;
-		}
-
-		FORCEINLINE FBundleNode* PopHead()
-		{
-			FBundleNode* Result = Head;
-
-			Count--;
-			Head = Head->NextNodeInCurrentBundle;
-			return Result;
-		}
-
-		FBundleNode* Head;
-		uint32       Count;
-	};
-	static_assert(sizeof(FBundleNode) <= BINNED3_MINIMUM_ALIGNMENT, "Bundle nodes must fit into the smallest block size");
-
-	struct FFreeBlockList
-	{
-		// return true if we actually pushed it
-		FORCEINLINE bool PushToFront(void* InPtr, uint32 InPoolIndex, uint32 InBlockSize)
-		{
-			checkSlow(InPtr);
-
-			if ((PartialBundle.Count >= (uint32)GMallocBinned3BundleCount) | (PartialBundle.Count * InBlockSize >= (uint32)GMallocBinned3BundleSize))
-			{
-				if (FullBundle.Head)
-				{
-					return false;
-				}
-				FullBundle = PartialBundle;
-				PartialBundle.Reset();
-			}
-			PartialBundle.PushHead((FBundleNode*)InPtr);
-			return true;
-		}
-		FORCEINLINE bool CanPushToFront(uint32 InPoolIndex, uint32 InBlockSize)
-		{
-			return !((!!FullBundle.Head) & ((PartialBundle.Count >= (uint32)GMallocBinned3BundleCount) | (PartialBundle.Count * InBlockSize >= (uint32)GMallocBinned3BundleSize)));
-		}
-		FORCEINLINE void* PopFromFront(uint32 InPoolIndex)
-		{
-			if ((!PartialBundle.Head) & (!!FullBundle.Head))
-			{
-				PartialBundle = FullBundle;
-				FullBundle.Reset();
-			}
-			return PartialBundle.Head ? PartialBundle.PopHead() : nullptr;
-		}
-
-		// tries to recycle the full bundle, if that fails, it is returned for freeing
-		FBundleNode* RecyleFull(uint32 InPoolIndex);
-		bool ObtainPartial(uint32 InPoolIndex);
-		FBundleNode* PopBundles(uint32 InPoolIndex);
-	private:
-		FBundle PartialBundle;
-		FBundle FullBundle;
-	};
-
-	struct FPerThreadFreeBlockLists
-	{
-		FORCEINLINE static FPerThreadFreeBlockLists* Get()
-		{
-			return FPlatformTLS::IsValidTlsSlot(FMallocBinned3::Binned3TlsSlot) ? (FPerThreadFreeBlockLists*)FPlatformTLS::GetTlsValue(FMallocBinned3::Binned3TlsSlot) : nullptr;
-		}
-		static void SetTLS();
-		static void ClearTLS();
-
-		FPerThreadFreeBlockLists() 
-#if BINNED3_ALLOCATOR_STATS
-			: AllocatedMemory(0) 
-#endif
-		{ }
-
-		FORCEINLINE void* Malloc(uint32 InPoolIndex)
-		{
-			return FreeLists[InPoolIndex].PopFromFront(InPoolIndex);
-		}
-		// return true if the pointer was pushed
-		FORCEINLINE bool Free(void* InPtr, uint32 InPoolIndex, uint32 InBlockSize)
-		{
-			return FreeLists[InPoolIndex].PushToFront(InPtr, InPoolIndex, InBlockSize);
-		}		
-		// return true if a pointer can be pushed
-		FORCEINLINE bool CanFree(uint32 InPoolIndex, uint32 InBlockSize)
-		{
-			return FreeLists[InPoolIndex].CanPushToFront(InPoolIndex, InBlockSize);
-		}
-		// returns a bundle that needs to be freed if it can't be recycled
-		FBundleNode* RecycleFullBundle(uint32 InPoolIndex)
-		{
-			return FreeLists[InPoolIndex].RecyleFull(InPoolIndex);
-		}
-		// returns true if we have anything to pop
-		bool ObtainRecycledPartial(uint32 InPoolIndex)
-		{
-			return FreeLists[InPoolIndex].ObtainPartial(InPoolIndex);
-		}
-		FBundleNode* PopBundles(uint32 InPoolIndex)
-		{
-			return FreeLists[InPoolIndex].PopBundles(InPoolIndex);
-		}
-#if BINNED3_ALLOCATOR_STATS
-	public:
-		int64 AllocatedMemory;
-		static TAtomic<int64> ConsolidatedMemory;
-#endif
-	private:
-		FFreeBlockList FreeLists[BINNED3_SMALL_POOL_COUNT];
-	};
-
 #if !BINNED3_USE_SEPARATE_VM_PER_POOL
 	FORCEINLINE uint64 PoolIndexFromPtr(const void* Ptr) // returns a uint64 for it can also be used to check if it is an OS allocation
 	{
@@ -527,7 +313,6 @@ class CORE_API FMallocBinned3 : public FMalloc
 
 public:
 
-
 	FMallocBinned3();
 
 	virtual ~FMallocBinned3();
@@ -545,7 +330,7 @@ public:
 			FPerThreadFreeBlockLists* Lists = GMallocBinned3PerThreadCaches ? FPerThreadFreeBlockLists::Get() : nullptr;
 			if (Lists)
 			{
-				uint32 PoolIndex = BoundSizeToPoolIndex(Size);
+				uint32 PoolIndex = BoundSizeToPoolIndex(Size, MemSizeToIndex);
 				uint32 BlockSize = PoolIndexToBlockSize(PoolIndex);
 				Result = Lists->Malloc(PoolIndex);
 #if BINNED3_ALLOCATOR_STATS
@@ -592,7 +377,7 @@ public:
 				}
 				if (bCanFree)
 				{
-					uint32 NewPoolIndex = BoundSizeToPoolIndex(NewSize);
+					uint32 NewPoolIndex = BoundSizeToPoolIndex(NewSize, MemSizeToIndex);
 					uint32 NewBlockSize = PoolIndexToBlockSize(NewPoolIndex);
 					void* Result = NewSize ? Lists->Malloc(NewPoolIndex) : nullptr;
 #if BINNED3_ALLOCATOR_STATS
@@ -661,38 +446,7 @@ public:
 
 	FORCEINLINE virtual SIZE_T QuantizeSize(SIZE_T Count, uint32 Alignment) override
 	{
-		static_assert(DEFAULT_ALIGNMENT <= BINNED3_MINIMUM_ALIGNMENT, "DEFAULT_ALIGNMENT is assumed to be zero"); // used below
-		checkSlow((Alignment & (Alignment - 1)) == 0); // Check the alignment is a power of two
-		SIZE_T SizeOut;
-		if ((Count <= BINNED3_MAX_SMALL_POOL_SIZE) & (Alignment <= BINNED3_MINIMUM_ALIGNMENT)) // one branch, not two
-		{
-			SizeOut = PoolIndexToBlockSize(BoundSizeToPoolIndex(Count));
-			check(SizeOut >= Count);
-			return SizeOut;
-		}
-		Alignment = FMath::Max<uint32>(Alignment, BINNED3_MINIMUM_ALIGNMENT);
-		Count = Align(Count, Alignment);
-		if ((Count <= BINNED3_MAX_SMALL_POOL_SIZE) & (Alignment <= BINNED3_MAX_SMALL_POOL_ALIGNMENT))
-		{
-			uint32 PoolIndex = BoundSizeToPoolIndex(Count);
-			do
-			{
-				uint32 BlockSize = PoolIndexToBlockSize(PoolIndex);
-				if (IsAligned(BlockSize, Alignment))
-				{
-					SizeOut = SIZE_T(BlockSize);
-					check(SizeOut >= Count);
-					return SizeOut;
-				}
-
-				PoolIndex++;
-			} while (PoolIndex < BINNED3_SMALL_POOL_COUNT);
-		}
-		
-		Alignment = FPlatformMath::Max<uint32>(Alignment, OsAllocationGranularity);
-		SizeOut = Align(Count, Alignment);
-		check(SizeOut >= Count);
-		return SizeOut;
+		return QuantizeSizeCommon(Count, Alignment, *this);
 	}
 
 	virtual bool ValidateHeap() override;
@@ -719,8 +473,11 @@ public:
 	// see ((!PoolIndex) | (NewSize > PoolIndexToBlockSize(static_cast<uint32>(PoolIndex - 1)))
 	static uint16 SmallBlockSizesReversedShifted[BINNED3_SMALL_POOL_COUNT + 1]; // this is reversed to get the smallest elements on our main cache line
 	static FMallocBinned3* MallocBinned3;
-	static uint32 Binned3TlsSlot;
 	static uint32 OsAllocationGranularity;
+
+	static void RegisterThreadFreeBlockLists(FPerThreadFreeBlockLists* FreeBlockLists);
+	static void UnregisterThreadFreeBlockLists(FPerThreadFreeBlockLists* FreeBlockLists);
+
 #if !BINNED3_USE_SEPARATE_VM_PER_POOL
 	static uint8* Binned3BaseVMPtr;
 	FPlatformMemory::FPlatformVirtualMemoryBlock Binned3BaseVMBlock;
@@ -733,15 +490,7 @@ public:
 	// Mapping of sizes to small table indices
 	static uint8 MemSizeToIndex[1 + (BINNED3_MAX_SMALL_POOL_SIZE >> BINNED3_MINIMUM_ALIGNMENT_SHIFT)];
 
-	FORCEINLINE uint32 BoundSizeToPoolIndex(SIZE_T Size) 
-	{
-		auto Index = ((Size + BINNED3_MINIMUM_ALIGNMENT - 1) >> BINNED3_MINIMUM_ALIGNMENT_SHIFT);
-		checkSlow(Index >= 0 && Index <= (BINNED3_MAX_SMALL_POOL_SIZE >> BINNED3_MINIMUM_ALIGNMENT_SHIFT)); // and it should be in the table
-		uint32 PoolIndex = uint32(MemSizeToIndex[Index]);
-		checkSlow(PoolIndex >= 0 && PoolIndex < BINNED3_SMALL_POOL_COUNT);
-		return PoolIndex;
-	}
-	FORCEINLINE uint32 PoolIndexToBlockSize(uint32 PoolIndex)
+	FORCEINLINE uint32 PoolIndexToBlockSize(uint32 PoolIndex) const
 	{
 		return uint32(SmallBlockSizesReversedShifted[BINNED3_SMALL_POOL_COUNT - PoolIndex - 1]) << BINNED3_MINIMUM_ALIGNMENT_SHIFT;
 	}
@@ -749,7 +498,6 @@ public:
 	void Commit(uint32 InPoolIndex, void *Ptr, SIZE_T Size);
 	void Decommit(uint32 InPoolIndex, void *Ptr, SIZE_T Size);
 
-private:
 	static void* AllocateMetaDataMemory(SIZE_T Size);
 	static void FreeMetaDataMemory(void *Ptr, SIZE_T Size);
 };
diff --git a/Engine/Source/Runtime/Core/Public/HAL/MallocBinnedCommon.h b/Engine/Source/Runtime/Core/Public/HAL/MallocBinnedCommon.h
--- a/Engine/Source/Runtime/Core/Public/HAL/MallocBinnedCommon.h
+++ b/Engine/Source/Runtime/Core/Public/HAL/MallocBinnedCommon.h
@@ -5,6 +5,7 @@
 #include "Containers/Array.h"
 #include "CoreTypes.h"
 #include "HAL/PlatformMemory.h"
+#include "HAL/PlatformTLS.h"
 #include "Templates/Function.h"
 
 #if PLATFORM_HAS_FPlatformVirtualMemoryBlock
@@ -103,5 +104,447 @@ struct FArenaParams
 
 };
 
+#endif	//~PLATFORM_HAS_FPlatformVirtualMemoryBlock
 
-#endif
\ No newline at end of file
+
+#if !defined(AGGRESSIVE_MEMORY_SAVING)
+#	error "AGGRESSIVE_MEMORY_SAVING must be defined"
+#endif
+
+#if AGGRESSIVE_MEMORY_SAVING
+#	define DEFAULT_GMallocBinnedBundleSize 8192
+#else
+#	define DEFAULT_GMallocBinnedBundleSize 65536
+#endif
+
+#define DEFAULT_GMallocBinnedBundleCount 64
+
+#define UE_BINNEDCOMMON_ALLOW_RUNTIME_TWEAKING 0
+#if UE_BINNEDCOMMON_ALLOW_RUNTIME_TWEAKING
+extern CORE_API int32 GMallocBinnedBundleSize;
+extern CORE_API int32 GMallocBinnedBundleCount;
+#else
+#	define GMallocBinnedBundleSize	DEFAULT_GMallocBinnedBundleSize
+#	define GMallocBinnedBundleCount	DEFAULT_GMallocBinnedBundleCount
+#endif
+
+#ifndef UE_BINNEDCOMMON_ALLOCATOR_STATS
+#	define UE_BINNEDCOMMON_ALLOCATOR_STATS (UE_BUILD_SHIPPING && !WITH_EDITOR)
+#endif
+
+
+class FMallocBinnedCommonBase : public FMalloc
+{
+protected:
+	struct FPtrToPoolMapping
+	{
+		FPtrToPoolMapping()
+			: PtrToPoolPageBitShift(0)
+			, HashKeyShift(0)
+			, PoolMask(0)
+			, MaxHashBuckets(0)
+		{
+		}
+		explicit FPtrToPoolMapping(uint32 InPageSize, uint64 InNumPoolsPerPage, uint64 AddressLimit)
+		{
+			Init(InPageSize, InNumPoolsPerPage, AddressLimit);
+		}
+
+		void Init(uint32 InPageSize, uint64 InNumPoolsPerPage, uint64 AddressLimit)
+		{
+			uint64 PoolPageToPoolBitShift = FPlatformMath::CeilLogTwo64(InNumPoolsPerPage);
+
+			PtrToPoolPageBitShift = FPlatformMath::CeilLogTwo(InPageSize);
+			HashKeyShift = PtrToPoolPageBitShift + PoolPageToPoolBitShift;
+			PoolMask = (1ull << PoolPageToPoolBitShift) - 1;
+			MaxHashBuckets = AddressLimit >> HashKeyShift;
+		}
+
+		FORCEINLINE void GetHashBucketAndPoolIndices(const void* InPtr, uint32& OutBucketIndex, UPTRINT& OutBucketCollision, uint32& OutPoolIndex) const
+		{
+			OutBucketCollision = (UPTRINT)InPtr >> HashKeyShift;
+			OutBucketIndex = uint32(OutBucketCollision & (MaxHashBuckets - 1));
+			OutPoolIndex = uint32(((UPTRINT)InPtr >> PtrToPoolPageBitShift) & PoolMask);
+		}
+
+		FORCEINLINE uint64 GetMaxHashBuckets() const
+		{
+			return MaxHashBuckets;
+		}
+
+	private:
+		/** Shift to apply to a pointer to get the reference from the indirect tables */
+		uint64 PtrToPoolPageBitShift;
+
+		/** Shift required to get required hash table key. */
+		uint64 HashKeyShift;
+
+		/** Used to mask off the bits that have been used to lookup the indirect table */
+		uint64 PoolMask;
+
+		// PageSize dependent constants
+		uint64 MaxHashBuckets;
+	};
+
+	// This needs to be small enough to fit inside the smallest allocation handled by MallocBinned2\3, hence the union.
+	struct FBundleNode
+	{
+		FBundleNode* NextNodeInCurrentBundle;
+
+		// NextBundle ptr is valid when node is stored in FFreeBlockList in a thread-local list of reusable allocations.
+		// Count is valid when node is stored in global recycler and caches the number of nodes in the list formed by NextNodeInCurrentBundle.
+		union
+		{
+			FBundleNode* NextBundle;
+			int32 Count;
+		};
+	};
+
+	struct FBundle
+	{
+		FORCEINLINE FBundle()
+		{
+			Reset();
+		}
+
+		FORCEINLINE void Reset()
+		{
+			Head = nullptr;
+			Count = 0;
+		}
+
+		FORCEINLINE void PushHead(FBundleNode* Node)
+		{
+			Node->NextNodeInCurrentBundle = Head;
+			Node->NextBundle = nullptr;
+			Head = Node;
+			Count++;
+		}
+
+		FORCEINLINE FBundleNode* PopHead()
+		{
+			FBundleNode* Result = Head;
+
+			Count--;
+			Head = Head->NextNodeInCurrentBundle;
+			return Result;
+		}
+
+		FBundleNode* Head;
+		uint32       Count;
+	};
+
+	/** Hash table struct for retrieving allocation book keeping information */
+	template <class T>
+	struct TPoolHashBucket
+	{
+		UPTRINT         BucketIndex;
+		T* FirstPool;
+		TPoolHashBucket* Prev;
+		TPoolHashBucket* Next;
+
+		TPoolHashBucket()
+		{
+			BucketIndex = 0;
+			FirstPool = nullptr;
+			Prev = this;
+			Next = this;
+		}
+
+		void Link(TPoolHashBucket* After)
+		{
+			After->Prev = Prev;
+			After->Next = this;
+			Prev->Next = After;
+			this->Prev = After;
+		}
+
+		void Unlink()
+		{
+			Next->Prev = Prev;
+			Prev->Next = Next;
+			Prev = this;
+			Next = this;
+		}
+	};
+
+	FPtrToPoolMapping PtrToPoolMapping;
+	static CORE_API uint32 BinnedTlsSlot;
+
+#if UE_BINNEDCOMMON_ALLOCATOR_STATS
+	static std::atomic<int64> TLSMemory;
+	static std::atomic<int64> ConsolidatedMemory;
+#endif
+};
+
+template <class AllocType, int MinAlign, int MaxAlign, int MinAlignShift, int NumSmallPools, int MaxSmallPoolSize>
+class TMallocBinnedCommon : public FMallocBinnedCommonBase
+{
+	static_assert(sizeof(FBundleNode) <= MinAlign, "Bundle nodes must fit into the smallest block size");
+
+protected:
+	struct FFreeBlockList
+	{
+		// return true if we actually pushed it
+		FORCEINLINE bool PushToFront(void* InPtr, uint32 InPoolIndex, uint32 InBlockSize)
+		{
+			checkSlow(InPtr);
+
+			if ((PartialBundle.Count >= (uint32)GMallocBinnedBundleCount) | (PartialBundle.Count * InBlockSize >= (uint32)GMallocBinnedBundleSize))
+			{
+				if (FullBundle.Head)
+				{
+					return false;
+				}
+				FullBundle = PartialBundle;
+				PartialBundle.Reset();
+			}
+			PartialBundle.PushHead((FBundleNode*)InPtr);
+			return true;
+		}
+
+		FORCEINLINE bool CanPushToFront(uint32 InPoolIndex, uint32 InBlockSize) const
+		{
+			return !((!!FullBundle.Head) & ((PartialBundle.Count >= (uint32)GMallocBinnedBundleCount) | (PartialBundle.Count * InBlockSize >= (uint32)GMallocBinnedBundleSize)));
+		}
+
+		FORCEINLINE void* PopFromFront(uint32 InPoolIndex)
+		{
+			if ((!PartialBundle.Head) & (!!FullBundle.Head))
+			{
+				PartialBundle = FullBundle;
+				FullBundle.Reset();
+			}
+			return PartialBundle.Head ? PartialBundle.PopHead() : nullptr;
+		}
+
+		// tries to recycle the full bundle, if that fails, it is returned for freeing
+		template <class T>
+		FBundleNode* RecyleFull(uint32 InPoolIndex, T& InGlobalRecycler)
+		{
+			FBundleNode* Result = nullptr;
+			if (FullBundle.Head)
+			{
+				FullBundle.Head->Count = FullBundle.Count;
+				if (!InGlobalRecycler.PushBundle(InPoolIndex, FullBundle.Head))
+				{
+					Result = FullBundle.Head;
+					Result->NextBundle = nullptr;
+				}
+				FullBundle.Reset();
+			}
+			return Result;
+		}
+
+		template <class T>
+		bool ObtainPartial(uint32 InPoolIndex, T& InGlobalRecycler)
+		{
+			if (!PartialBundle.Head)
+			{
+				PartialBundle.Count = 0;
+				PartialBundle.Head = InGlobalRecycler.PopBundle(InPoolIndex);
+				if (PartialBundle.Head)
+				{
+					PartialBundle.Count = PartialBundle.Head->Count;
+					PartialBundle.Head->NextBundle = nullptr;
+					return true;
+				}
+				return false;
+			}
+			return true;
+		}
+
+		FBundleNode* PopBundles(uint32 InPoolIndex)
+		{
+			FBundleNode* Partial = PartialBundle.Head;
+			if (Partial)
+			{
+				PartialBundle.Reset();
+				Partial->NextBundle = nullptr;
+			}
+
+			FBundleNode* Full = FullBundle.Head;
+			if (Full)
+			{
+				FullBundle.Reset();
+				Full->NextBundle = nullptr;
+			}
+
+			FBundleNode* Result = Partial;
+			if (Result)
+			{
+				Result->NextBundle = Full;
+			}
+			else
+			{
+				Result = Full;
+			}
+
+			return Result;
+		}
+
+	private:
+		FBundle PartialBundle;
+		FBundle FullBundle;
+	};
+
+	struct FPerThreadFreeBlockLists
+	{
+		FORCEINLINE static FPerThreadFreeBlockLists* Get()
+		{
+			return BinnedTlsSlot ? (FPerThreadFreeBlockLists*)FPlatformTLS::GetTlsValue(BinnedTlsSlot) : nullptr;
+		}
+
+		static void SetTLS()
+		{
+			check(FPlatformTLS::IsValidTlsSlot(BinnedTlsSlot));
+			FPerThreadFreeBlockLists* ThreadSingleton = (FPerThreadFreeBlockLists*)FPlatformTLS::GetTlsValue(BinnedTlsSlot);
+			if (!ThreadSingleton)
+			{
+				const int64 TLSSize = Align(sizeof(FPerThreadFreeBlockLists), AllocType::OsAllocationGranularity);
+				ThreadSingleton = new (AllocType::AllocateMetaDataMemory(TLSSize)) FPerThreadFreeBlockLists();
+#if UE_BINNEDCOMMON_ALLOCATOR_STATS
+				TLSMemory.fetch_add(TLSSize, std::memory_order_relaxed);
+#endif
+				verify(ThreadSingleton);
+				FPlatformTLS::SetTlsValue(BinnedTlsSlot, ThreadSingleton);
+				AllocType::RegisterThreadFreeBlockLists(ThreadSingleton);
+			}
+		}
+
+		static void ClearTLS()
+		{
+			check(FPlatformTLS::IsValidTlsSlot(BinnedTlsSlot));
+			FPerThreadFreeBlockLists* ThreadSingleton = (FPerThreadFreeBlockLists*)FPlatformTLS::GetTlsValue(BinnedTlsSlot);
+			if (ThreadSingleton)
+			{
+				const int64 TLSSize = Align(sizeof(FPerThreadFreeBlockLists), AllocType::OsAllocationGranularity);
+#if UE_BINNEDCOMMON_ALLOCATOR_STATS
+				TLSMemory.fetch_sub(TLSSize, std::memory_order_relaxed);
+#endif
+				AllocType::UnregisterThreadFreeBlockLists(ThreadSingleton);
+
+				ThreadSingleton->~FPerThreadFreeBlockLists();
+
+				AllocType::FreeMetaDataMemory(ThreadSingleton, TLSSize);
+			}
+			FPlatformTLS::SetTlsValue(BinnedTlsSlot, nullptr);
+		}
+
+		FORCEINLINE void* Malloc(uint32 InPoolIndex)
+		{
+			return FreeLists[InPoolIndex].PopFromFront(InPoolIndex);
+		}
+
+		// return true if the pointer was pushed
+		FORCEINLINE bool Free(void* InPtr, uint32 InPoolIndex, uint32 InBlockSize)
+		{
+			return FreeLists[InPoolIndex].PushToFront(InPtr, InPoolIndex, InBlockSize);
+		}
+
+		// return true if a pointer can be pushed
+		FORCEINLINE bool CanFree(uint32 InPoolIndex, uint32 InBlockSize) const
+		{
+			return FreeLists[InPoolIndex].CanPushToFront(InPoolIndex, InBlockSize);
+		}
+
+		// returns a bundle that needs to be freed if it can't be recycled
+		template <class T>
+		FBundleNode* RecycleFullBundle(uint32 InPoolIndex, T& InGlobalRecycler)
+		{
+			return FreeLists[InPoolIndex].RecyleFull(InPoolIndex, InGlobalRecycler);
+		}
+
+		// returns true if we have anything to pop
+		template <class T>
+		bool ObtainRecycledPartial(uint32 InPoolIndex, T& InGlobalRecycler)
+		{
+			return FreeLists[InPoolIndex].ObtainPartial(InPoolIndex, InGlobalRecycler);
+		}
+
+		FBundleNode* PopBundles(uint32 InPoolIndex)
+		{
+			return FreeLists[InPoolIndex].PopBundles(InPoolIndex);
+		}
+
+#if UE_BINNEDCOMMON_ALLOCATOR_STATS
+	public:
+		int64 AllocatedMemory = 0;
+#endif
+	private:
+		FFreeBlockList FreeLists[NumSmallPools];
+	};
+
+	FORCEINLINE SIZE_T QuantizeSizeCommon(SIZE_T Count, uint32 Alignment, const AllocType& Alloc) const
+	{
+		static_assert(DEFAULT_ALIGNMENT <= MinAlign, "DEFAULT_ALIGNMENT is assumed to be zero"); // used below
+		checkSlow((Alignment & (Alignment - 1)) == 0); // Check the alignment is a power of two
+		SIZE_T SizeOut;
+		if ((Count <= MaxSmallPoolSize) & (Alignment <= MinAlign)) // one branch, not two
+		{
+			SizeOut = Alloc.PoolIndexToBlockSize(BoundSizeToPoolIndex(Count, Alloc.MemSizeToIndex));
+			check(SizeOut >= Count);
+			return SizeOut;
+		}
+		Alignment = FMath::Max<uint32>(Alignment, MinAlign);
+		Count = Align(Count, Alignment);
+		if ((Count <= MaxSmallPoolSize) & (Alignment <= MaxAlign))
+		{
+			uint32 PoolIndex = BoundSizeToPoolIndex(Count, Alloc.MemSizeToIndex);
+			do
+			{
+				uint32 BlockSize = Alloc.PoolIndexToBlockSize(PoolIndex);
+				if (IsAligned(BlockSize, Alignment))
+				{
+					SizeOut = SIZE_T(BlockSize);
+					check(SizeOut >= Count);
+					return SizeOut;
+				}
+
+				PoolIndex++;
+			} while (PoolIndex < NumSmallPools);
+		}
+
+		Alignment = FPlatformMath::Max<uint32>(Alignment, Alloc.OsAllocationGranularity);
+		SizeOut = Align(Count, Alignment);
+		check(SizeOut >= Count);
+		return SizeOut;
+	}
+
+	FORCEINLINE uint32 BoundSizeToPoolIndex(SIZE_T Size, const uint8(&MemSizeToIndex)[1 + (MaxSmallPoolSize >> MinAlignShift)]) const
+	{
+		auto Index = ((Size + MinAlign - 1) >> MinAlignShift);
+		checkSlow(Index >= 0 && Index <= (MaxSmallPoolSize >> MinAlignShift)); // and it should be in the table
+		uint32 PoolIndex = uint32(MemSizeToIndex[Index]);
+		checkSlow(PoolIndex >= 0 && PoolIndex < NumSmallPools);
+		return PoolIndex;
+	}
+
+	bool PromoteToLargerBin(SIZE_T& Size, uint32& Alignment, const AllocType& Alloc) const
+	{
+		// try to promote our allocation request to a larger bin with a matching natural alignment
+		// if requested alignment is larger than MinAlign but smaller than MaxAlign
+		// so we don't do a page allocation with a lot of memory waste
+		Alignment = FMath::Max<uint32>(Alignment, MinAlign);
+		const SIZE_T AlignedSize = Align(Size, Alignment);
+		if (UNLIKELY((AlignedSize <= MaxSmallPoolSize) && (Alignment <= MaxAlign)))
+		{
+			uint32 PoolIndex = BoundSizeToPoolIndex(AlignedSize, Alloc.MemSizeToIndex);
+			do
+			{
+				uint32 BlockSize = Alloc.PoolIndexToBlockSize(PoolIndex);
+				if (IsAligned(BlockSize, Alignment))
+				{
+					// we found a matching pool for our alignment and size requirements, so modify the size request to match
+					Size = SIZE_T(BlockSize);
+					Alignment = MinAlign;
+					return true;
+				}
+
+				PoolIndex++;
+			} while (PoolIndex < NumSmallPools);
+		}
+
+		return false;
+	}
+};
\ No newline at end of file
