Fixed MB3 compilation

diff --git a/Engine/Source/Runtime/Core/Private/HAL/MallocBinned3.cpp b/Engine/Source/Runtime/Core/Private/HAL/MallocBinned3.cpp
--- a/Engine/Source/Runtime/Core/Private/HAL/MallocBinned3.cpp
+++ b/Engine/Source/Runtime/Core/Private/HAL/MallocBinned3.cpp
@@ -1421,7 +1421,7 @@ void FMallocBinned3::DumpAllocatorStats(class FOutputDevice& Ar)
 	Ar.Logf(TEXT("PoolInfo: %fmb"), ((double)Binned3PoolInfoMemory) / (1024.0f * 1024.0f));
 	Ar.Logf(TEXT("Hash: %fmb"), ((double)Binned3HashMemory) / (1024.0f * 1024.0f));
 	Ar.Logf(TEXT("Free Bits: %fmb"), ((double)Binned3FreeBitsMemory) / (1024.0f * 1024.0f));
-	Ar.Logf(TEXT("TLS: %fmb"), ((double)TLSMemory.load(std::memory_order_relaxed) / (1024.0f * 1024.0f));
+	Ar.Logf(TEXT("TLS: %fmb"), ((double)TLSMemory.load(std::memory_order_relaxed)) / (1024.0f * 1024.0f));
 	Ar.Logf(TEXT("Slab Commits: %llu"), Binned3Commits.Load());
 	Ar.Logf(TEXT("Slab Decommits: %llu"), Binned3Decommits.Load());
 #if BINNED3_USE_SEPARATE_VM_PER_POOL
