Fixed MallocBinned Stats define toggle

diff --git a/Engine/Source/Runtime/Core/Public/HAL/MallocBinnedCommon.h b/Engine/Source/Runtime/Core/Public/HAL/MallocBinnedCommon.h
--- a/Engine/Source/Runtime/Core/Public/HAL/MallocBinnedCommon.h
+++ b/Engine/Source/Runtime/Core/Public/HAL/MallocBinnedCommon.h
@@ -129,7 +129,7 @@ extern CORE_API int32 GMallocBinnedBundleCount;
 #endif
 
 #ifndef UE_BINNEDCOMMON_ALLOCATOR_STATS
-#	define UE_BINNEDCOMMON_ALLOCATOR_STATS (UE_BUILD_SHIPPING && !WITH_EDITOR)
+#	define UE_BINNEDCOMMON_ALLOCATOR_STATS (!UE_BUILD_SHIPPING || WITH_EDITOR)
 #endif
 
 
