Fixed includes in MallocBinnedCommon.h

diff --git a/Engine/Source/Runtime/Core/Public/HAL/MallocBinned2.h b/Engine/Source/Runtime/Core/Public/HAL/MallocBinned2.h
--- a/Engine/Source/Runtime/Core/Public/HAL/MallocBinned2.h
+++ b/Engine/Source/Runtime/Core/Public/HAL/MallocBinned2.h
@@ -15,7 +15,6 @@
 #include "Math/NumericLimits.h"
 #include "Misc/AssertionMacros.h"
 #include "Misc/Fork.h"
-#include "Templates/AlignmentTemplates.h"
 #include "Templates/Atomic.h"
 
 struct FGenericMemoryStats;
diff --git a/Engine/Source/Runtime/Core/Public/HAL/MallocBinnedCommon.h b/Engine/Source/Runtime/Core/Public/HAL/MallocBinnedCommon.h
--- a/Engine/Source/Runtime/Core/Public/HAL/MallocBinnedCommon.h
+++ b/Engine/Source/Runtime/Core/Public/HAL/MallocBinnedCommon.h
@@ -2,11 +2,10 @@
 
 #pragma once
 
-#include "Containers/Array.h"
 #include "CoreTypes.h"
 #include "HAL/PlatformMemory.h"
 #include "HAL/PlatformTLS.h"
-#include "Templates/Function.h"
+#include "Templates/AlignmentTemplates.h"
 
 #if PLATFORM_HAS_FPlatformVirtualMemoryBlock
 
@@ -129,7 +128,7 @@ extern CORE_API int32 GMallocBinnedBundleCount;
 #endif
 
 #ifndef UE_BINNEDCOMMON_ALLOCATOR_STATS
-#	define UE_BINNEDCOMMON_ALLOCATOR_STATS (!UE_BUILD_SHIPPING || WITH_EDITOR)
+#	define UE_BINNEDCOMMON_ALLOCATOR_STATS 0
 #endif
 
 
