Fixed missing wrong\includes in MallocBinnedCommon

diff --git a/Engine/Source/Runtime/Core/Public/HAL/MallocBinnedCommon.h b/Engine/Source/Runtime/Core/Public/HAL/MallocBinnedCommon.h
--- a/Engine/Source/Runtime/Core/Public/HAL/MallocBinnedCommon.h
+++ b/Engine/Source/Runtime/Core/Public/HAL/MallocBinnedCommon.h
@@ -3,12 +3,15 @@
 #pragma once
 
 #include "CoreTypes.h"
-#include "HAL/PlatformMemory.h"
 #include "HAL/PlatformTLS.h"
 #include "Templates/AlignmentTemplates.h"
 
 #if PLATFORM_HAS_FPlatformVirtualMemoryBlock
 
+#include "Containers/Array.h"
+#include "HAL/PlatformMemory.h"
+#include "Templates/Function.h"
+
 #define BINNEDCOMMON_MAX_LISTED_SMALL_POOL_SIZE	28672
 #define BINNEDCOMMON_NUM_LISTED_SMALL_POOLS	49
 
