Fix integer overflow when using TScriptArray::InsertZeroed, TScriptArray::Insert,TScriptArray::AddZeroed, TScriptArray::SwapMemory, TScriptArray::CountBytes, TScriptArray::CheckAddress or TScriptArray::Remove

diff --git a/Engine/Source/Runtime/Core/Public/Containers/ScriptArray.h b/Engine/Source/Runtime/Core/Public/Containers/ScriptArray.h
--- a/Engine/Source/Runtime/Core/Public/Containers/ScriptArray.h
+++ b/Engine/Source/Runtime/Core/Public/Containers/ScriptArray.h
@@ -46,7 +46,7 @@ public:
 	void InsertZeroed( int32 Index, int32 Count, int32 NumBytesPerElement, uint32 AlignmentOfElement )
 	{
 		Insert( Index, Count, NumBytesPerElement, AlignmentOfElement );
-		FMemory::Memzero( (uint8*)this->GetAllocation()+Index*NumBytesPerElement, Count*NumBytesPerElement );
+		FMemory::Memzero( (uint8*)this->GetAllocation()+(SIZE_T)Index*NumBytesPerElement, (SIZE_T)Count*NumBytesPerElement );
 	}
 	void Insert( int32 Index, int32 Count, int32 NumBytesPerElement, uint32 AlignmentOfElement )
 	{
@@ -63,9 +63,9 @@ public:
 		}
 		FMemory::Memmove
 		(
-			(uint8*)this->GetAllocation() + (Index+Count )*NumBytesPerElement,
-			(uint8*)this->GetAllocation() + (Index       )*NumBytesPerElement,
-			                                               (OldNum-Index)*NumBytesPerElement
+			(uint8*)this->GetAllocation() + ((SIZE_T)Index+Count )*NumBytesPerElement,
+			(uint8*)this->GetAllocation() + ((SIZE_T)Index       )*NumBytesPerElement,
+			((SIZE_T)OldNum-Index)*NumBytesPerElement
 		);
 
 		SlackTrackerNumChanged();
@@ -89,7 +89,7 @@ public:
 	int32 AddZeroed( int32 Count, int32 NumBytesPerElement, uint32 AlignmentOfElement )
 	{
 		const int32 Index = Add( Count, NumBytesPerElement, AlignmentOfElement );
-		FMemory::Memzero( (uint8*)this->GetAllocation()+Index*NumBytesPerElement, Count*NumBytesPerElement );
+		FMemory::Memzero( (uint8*)this->GetAllocation()+ (SIZE_T)Index*NumBytesPerElement, (SIZE_T)Count*NumBytesPerElement );
 		return Index;
 	}
 	void Shrink( int32 NumBytesPerElement, uint32 AlignmentOfElement )
@@ -158,8 +158,8 @@ public:
 	void SwapMemory(int32 A, int32 B, int32 NumBytesPerElement )
 	{
 		FMemory::Memswap(
-			(uint8*)this->GetAllocation()+(NumBytesPerElement*A),
-			(uint8*)this->GetAllocation()+(NumBytesPerElement*B),
+			(uint8*)this->GetAllocation()+((SIZE_T)NumBytesPerElement*A),
+			(uint8*)this->GetAllocation()+((SIZE_T)NumBytesPerElement*B),
 			NumBytesPerElement
 			);
 	}
@@ -170,11 +170,11 @@ public:
 	}
 	void CountBytes( FArchive& Ar, int32 NumBytesPerElement  ) const
 	{
-		Ar.CountBytes( ArrayNum*NumBytesPerElement, ArrayMax*NumBytesPerElement );
+		Ar.CountBytes( (SIZE_T)ArrayNum*NumBytesPerElement, (SIZE_T)ArrayMax*NumBytesPerElement );
 	}
 	FORCEINLINE void CheckAddress(const void* Addr, int32 NumBytesPerElement) const
 	{
-		checkf((const char*)Addr < (const char*)GetData() || (const char*)Addr >= ((const char*)GetData() + ArrayMax * NumBytesPerElement), TEXT("Attempting to use a container element (%p) which already comes from the container being modified (%p, ArrayMax: %lld, ArrayNum: %lld, SizeofElement: %d)!"), Addr, GetData(), (long long)ArrayMax, (long long)ArrayNum, NumBytesPerElement);
+		checkf((const char*)Addr < (const char*)GetData() || (const char*)Addr >= ((const char*)GetData() + (SIZE_T)ArrayMax * NumBytesPerElement), TEXT("Attempting to use a container element (%p) which already comes from the container being modified (%p, ArrayMax: %lld, ArrayNum: %lld, SizeofElement: %d)!"), Addr, GetData(), (long long)ArrayMax, (long long)ArrayNum, NumBytesPerElement);
 	}
 	/**
 	 * Returns the amount of slack in this array in elements.
@@ -199,9 +199,9 @@ public:
 			{
 				FMemory::Memmove
 					(
-					(uint8*)this->GetAllocation() + (Index)* NumBytesPerElement,
-					(uint8*)this->GetAllocation() + (Index + Count) * NumBytesPerElement,
-					NumToMove * NumBytesPerElement
+					(uint8*)this->GetAllocation() + ((SIZE_T)Index)* NumBytesPerElement,
+					(uint8*)this->GetAllocation() + ((SIZE_T)Index + Count) * NumBytesPerElement,
+					(SIZE_T)NumToMove * NumBytesPerElement
 					);
 			}
 			ArrayNum -= Count;
