Add delegates for tracking of multiprocess worker creation and detachment from within the director of a multiprocess session.  The approach is meant to allow for the fact that workers can possibly come into and out of existence over the lifetime of a session in the future, and to allow some flexibility in the assignment of identifiers.

diff --git a/Engine/Source/Editor/UnrealEd/Private/Cooker/CookDirector.cpp b/Engine/Source/Editor/UnrealEd/Private/Cooker/CookDirector.cpp
--- a/Engine/Source/Editor/UnrealEd/Private/Cooker/CookDirector.cpp
+++ b/Engine/Source/Editor/UnrealEd/Private/Cooker/CookDirector.cpp
@@ -1330,7 +1330,7 @@ FString FCookDirector::GetWorkerCommandLine(FWorkerId WorkerId, int32 ProfileId)
 	Tokens.Insert(TEXT("-run=cook"), 1);
 	Tokens.Insert(TEXT("-cookworker"), 2);
 	Tokens.Insert(FString::Printf(TEXT("-CookProfileId=%d"), ProfileId), 3);
-	Tokens.Insert(FString::Printf(TEXT("-MultiprocessId=%d"), WorkerId.GetRemoteIndex() + 1), 4);
+	Tokens.Insert(FString::Printf(TEXT("-MultiprocessId=%d"), WorkerId.GetMultiprocessId()), 4);
 	check(!WorkerConnectAuthority.IsEmpty()); // This should have been constructed in TryCreateWorkerConnectSocket before any CookWorkerServers could exist to call GetWorkerCommandLine
 	Tokens.Add(FString::Printf(TEXT("-CookDirectorHost=%s"), *WorkerConnectAuthority));
 	Tokens.Add(TEXT("-unattended"));
diff --git a/Engine/Source/Editor/UnrealEd/Private/Cooker/CookWorkerServer.cpp b/Engine/Source/Editor/UnrealEd/Private/Cooker/CookWorkerServer.cpp
--- a/Engine/Source/Editor/UnrealEd/Private/Cooker/CookWorkerServer.cpp
+++ b/Engine/Source/Editor/UnrealEd/Private/Cooker/CookWorkerServer.cpp
@@ -49,6 +49,10 @@ FCookWorkerServer::~FCookWorkerServer()
 
 void FCookWorkerServer::DetachFromRemoteProcess()
 {
+	if (Socket != nullptr)
+	{
+		FCoreDelegates::OnMultiprocessWorkerDetached.Broadcast({WorkerId.GetMultiprocessId()});
+	}
 	Sockets::CloseSocket(Socket);
 	CookWorkerHandle = FProcHandle();
 	CookWorkerProcessId = 0;
@@ -532,6 +536,7 @@ void FCookWorkerServer::LaunchProcess()
 	{
 		UE_LOG(LogCook, Display, TEXT("CookWorkerServer %d launched CookWorker as WorkerId %d and PID %u with commandline \"%s\"."),
 			ProfileId, WorkerId.GetRemoteIndex(), CookWorkerProcessId, *LaunchInfo.WorkerCommandLine);
+		FCoreDelegates::OnMultiprocessWorkerCreated.Broadcast({WorkerId.GetMultiprocessId()});
 		SendToState(EConnectStatus::WaitForConnect);
 	}
 	else
diff --git a/Engine/Source/Runtime/Core/Private/Misc/CoreDelegates.cpp b/Engine/Source/Runtime/Core/Private/Misc/CoreDelegates.cpp
--- a/Engine/Source/Runtime/Core/Private/Misc/CoreDelegates.cpp
+++ b/Engine/Source/Runtime/Core/Private/Misc/CoreDelegates.cpp
@@ -376,6 +376,11 @@ CORE_API void RegisterEncryptionKeyCallback(TEncryptionKeyFunc InCallback)
 TDelegate<const TCHAR*(void)> FCoreDelegates::OnGetBuildURL;
 TDelegate<const TCHAR*(void)> FCoreDelegates::OnGetExecutingJobURL;
 
+#if WITH_EDITOR
+TMulticastDelegate<void(const UE::FMultiprocessCreatedContext&)> FCoreDelegates::OnMultiprocessWorkerCreated;
+TMulticastDelegate<void(const UE::FMultiprocessDetachedContext&)> FCoreDelegates::OnMultiprocessWorkerDetached;
+#endif
+
 FSimpleMulticastDelegate FCoreDelegates::OnParentBeginFork;
 FSimpleMulticastDelegate FCoreDelegates::OnParentPreFork;
 TMulticastDelegate<void(EForkProcessRole /* ProcessRole */)> FCoreDelegates::OnPostFork;
diff --git a/Engine/Source/Runtime/Core/Public/Misc/CoreDelegates.h b/Engine/Source/Runtime/Core/Public/Misc/CoreDelegates.h
--- a/Engine/Source/Runtime/Core/Public/Misc/CoreDelegates.h
+++ b/Engine/Source/Runtime/Core/Public/Misc/CoreDelegates.h
@@ -70,6 +70,19 @@ struct FCrashOverrideParameters
 	CORE_API ~FCrashOverrideParameters();
 };
 
+namespace UE
+{
+struct FMultiprocessCreatedContext
+{
+	int32 Id;
+};
+
+struct FMultiprocessDetachedContext
+{
+	int32 Id;
+};
+}
+
 class FCoreDelegates
 {
 public:
@@ -572,6 +585,13 @@ public:
 	// Extension point for projects to report the URL for a continuous integration job which is currently executing this process. 
 	static CORE_API TDelegate<const TCHAR*()> OnGetExecutingJobURL;
 	
+#if WITH_EDITOR
+	// Called when a subprocess is created for multiprocess operation.
+	static CORE_API TMulticastDelegate<void(const UE::FMultiprocessCreatedContext&)> OnMultiprocessWorkerCreated;
+
+	// Called when a subprocess is detached (but not necessarily terminated) for multiprocess operation.
+	static CORE_API TMulticastDelegate<void(const UE::FMultiprocessDetachedContext&)> OnMultiprocessWorkerDetached;
+#endif
 private:
 
 	// Callbacks for hotfixes
