CompactBinary: Fix AsArray() and AsObject() to store type errors

diff --git a/Engine/Source/Runtime/Core/Private/Tests/Serialization/CompactBinaryTest.cpp b/Engine/Source/Runtime/Core/Private/Tests/Serialization/CompactBinaryTest.cpp
--- a/Engine/Source/Runtime/Core/Private/Tests/Serialization/CompactBinaryTest.cpp
+++ b/Engine/Source/Runtime/Core/Private/Tests/Serialization/CompactBinaryTest.cpp
@@ -70,6 +70,18 @@ UE_CBFIELD_TYPE_ACCESSOR_EX(ObjectId, IsObjectId, AsObjectId, FCbObjectId, const
 UE_CBFIELD_TYPE_ACCESSOR(CustomById, IsCustomById, AsCustomById, FCbCustomById);
 UE_CBFIELD_TYPE_ACCESSOR(CustomByName, IsCustomByName, AsCustomByName, FCbCustomByName);
 
+struct FCbFieldObjectAccessors
+{
+	static constexpr bool (FCbFieldView::*IsType)() const = &FCbFieldView::IsObject;
+	static auto AsType(FCbField& Field, const FCbObject&) { return Field.AsObject(); }
+};
+
+struct FCbFieldArrayAccessors
+{
+	static constexpr bool (FCbFieldView::*IsType)() const = &FCbFieldView::IsArray;
+	static auto AsType(FCbField& Field, const FCbArray&) { return Field.AsArray(); }
+};
+
 struct FCbAttachmentAccessors
 {
 	static constexpr bool (FCbFieldView::*IsType)() const = &FCbFieldView::IsAttachment;
@@ -89,11 +101,21 @@ protected:
 		return TestTrue(What, Actual.Equals(Expected));
 	}
 
+	bool TestEqual(const FString& What, const FCbArray& Actual, const FCbArray& Expected)
+	{
+		return TestTrue(What, Actual.Equals(Expected));
+	}
+
 	bool TestEqual(const FString& What, const FCbObjectView& Actual, const FCbObjectView& Expected)
 	{
 		return TestTrue(What, Actual.Equals(Expected));
 	}
 
+	bool TestEqual(const FString& What, const FCbObject& Actual, const FCbObject& Expected)
+	{
+		return TestTrue(What, Actual.Equals(Expected));
+	}
+
 	bool TestEqual(const FString& What, const FCbCustomById& Actual, const FCbCustomById& Expected)
 	{
 		return TestTrue(What, Actual.Id == Expected.Id && Actual.Data == Expected.Data);
@@ -109,8 +131,8 @@ protected:
 		return TestTrue(What, Actual.EqualBytes(MakeMemoryView(Expected)));
 	}
 
-	template <ECbFieldType FieldType, typename T = TCbFieldValueType<FieldType>, typename AccessorsType = TCbFieldTypeAccessors<FieldType>>
-	void TestFieldNoClone(const TCHAR* What, FCbFieldView& Field, T ExpectedValue = T(), T DefaultValue = T(), ECbFieldError ExpectedError = ECbFieldError::None, const AccessorsType& Accessors = AccessorsType())
+	template <ECbFieldType FieldType, typename T = TCbFieldValueType<FieldType>, typename AccessorsType = TCbFieldTypeAccessors<FieldType>, typename FieldArgType = FCbFieldView>
+	void TestFieldNoClone(const TCHAR* What, FieldArgType& Field, T ExpectedValue = T(), T DefaultValue = T(), ECbFieldError ExpectedError = ECbFieldError::None, const AccessorsType& Accessors = AccessorsType())
 	{
 		TestEqual(FString::Printf(TEXT("FCbFieldView::Is[Type](%s)"), What), Invoke(Accessors.IsType, Field), ExpectedError != ECbFieldError::TypeError);
 		if (ExpectedError == ECbFieldError::None && !Field.IsBool())
@@ -125,7 +147,7 @@ protected:
 	}
 
 	template <ECbFieldType FieldType, typename T = TCbFieldValueType<FieldType>, typename AccessorsType = TCbFieldTypeAccessors<FieldType>>
-	void TestFieldNoClone(const TCHAR* What, TArrayView<const uint8> Value, T ExpectedValue = T(), T DefaultValue = T(), ECbFieldError ExpectedError = ECbFieldError::None, const AccessorsType& Accessors = AccessorsType())
+	void TestFieldBytesNoClone(const TCHAR* What, TArrayView<const uint8> Value, T ExpectedValue = T(), T DefaultValue = T(), ECbFieldError ExpectedError = ECbFieldError::None, const AccessorsType& Accessors = AccessorsType())
 	{
 		FCbFieldView Field(Value.GetData(), FieldType);
 		TestFieldNoClone<FieldType>(What, Field, ExpectedValue, DefaultValue, ExpectedError, Accessors);
@@ -141,7 +163,7 @@ protected:
 	}
 
 	template <ECbFieldType FieldType, typename T = TCbFieldValueType<FieldType>, typename AccessorsType = TCbFieldTypeAccessors<FieldType>>
-	void TestField(const TCHAR* What, TArrayView<const uint8> Value, T ExpectedValue = T(), T DefaultValue = T(), ECbFieldError ExpectedError = ECbFieldError::None, const AccessorsType& Accessors = AccessorsType())
+	void TestFieldBytes(const TCHAR* What, TArrayView<const uint8> Value, T ExpectedValue = T(), T DefaultValue = T(), ECbFieldError ExpectedError = ECbFieldError::None, const AccessorsType& Accessors = AccessorsType())
 	{
 		FCbFieldView Field(Value.GetData(), FieldType);
 		TestEqual(FString::Printf(TEXT("FCbFieldView::GetSize(%s)"), What), Field.GetSize(), uint64(Value.Num()) + uint64(FCbFieldType::HasFieldType(FieldType) ? 0 : 1));
@@ -151,14 +173,14 @@ protected:
 		TestField<FieldType>(What, Field, ExpectedValue, DefaultValue, ExpectedError, Accessors);
 	}
 
-	template <ECbFieldType FieldType, typename T = TCbFieldValueType<FieldType>, typename AccessorsType = TCbFieldTypeAccessors<FieldType>>
-	void TestFieldError(const TCHAR* What, FCbFieldView& Field, ECbFieldError ExpectedError, T ExpectedValue = T(), const AccessorsType& Accessors = AccessorsType())
+	template <ECbFieldType FieldType, typename T = TCbFieldValueType<FieldType>, typename AccessorsType = TCbFieldTypeAccessors<FieldType>, typename FieldArgType = FCbFieldView>
+	void TestFieldError(const TCHAR* What, FieldArgType& Field, ECbFieldError ExpectedError, T ExpectedValue = T(), const AccessorsType& Accessors = AccessorsType())
 	{
 		TestFieldNoClone<FieldType>(What, Field, ExpectedValue, ExpectedValue, ExpectedError, Accessors);
 	}
 
 	template <ECbFieldType FieldType, typename T = TCbFieldValueType<FieldType>, typename AccessorsType = TCbFieldTypeAccessors<FieldType>>
-	void TestFieldError(const TCHAR* What, TArrayView<const uint8> Value, ECbFieldError ExpectedError, T ExpectedValue = T(), const AccessorsType& Accessors = AccessorsType())
+	void TestFieldBytesError(const TCHAR* What, TArrayView<const uint8> Value, ECbFieldError ExpectedError, T ExpectedValue = T(), const AccessorsType& Accessors = AccessorsType())
 	{
 		FCbFieldView Field(Value.GetData(), FieldType);
 		TestFieldError<FieldType>(What, Field, ExpectedError, ExpectedValue, Accessors);
@@ -333,7 +355,7 @@ bool FCbFieldObjectTest::RunTest(const FString& Parameters)
 	};
 
 	// Test FCbFieldView(Object, Empty)
-	TestField<ECbFieldType::Object>(TEXT("Object, Empty"), {0});
+	TestFieldBytes<ECbFieldType::Object>(TEXT("Object, Empty"), {0});
 
 	// Test FCbFieldView(Object, Empty)
 	{
@@ -427,9 +449,10 @@ bool FCbFieldObjectTest::RunTest(const FString& Parameters)
 
 	// Test FCbFieldView(None) as Object
 	{
-		FCbFieldView Field;
-		TestFieldError<ECbFieldType::Object>(TEXT("Object, None"), Field, ECbFieldError::TypeError);
-		FCbField::MakeView(Field).AsObject();
+		FCbFieldView FieldView;
+		TestFieldError<ECbFieldType::Object>(TEXT("Object, None, View"), FieldView, ECbFieldError::TypeError);
+		FCbField Field;
+		TestFieldError<ECbFieldType::Object, FCbObject, FCbFieldObjectAccessors>(TEXT("Object, None"), Field, ECbFieldError::TypeError);
 	}
 
 	// Test FCbObjectView(ObjectWithName) and CreateIterator
@@ -549,7 +572,7 @@ bool FCbFieldArrayTest::RunTest(const FString& Parameters)
 	};
 
 	// Test FCbFieldView(Array, Empty)
-	TestField<ECbFieldType::Array>(TEXT("Array, Empty"), {1, 0});
+	TestFieldBytes<ECbFieldType::Array>(TEXT("Array, Empty"), {1, 0});
 
 	// Test FCbFieldView(Array, Empty)
 	{
@@ -622,9 +645,10 @@ bool FCbFieldArrayTest::RunTest(const FString& Parameters)
 
 	// Test FCbFieldView(None) as Array
 	{
-		FCbFieldView Field;
-		TestFieldError<ECbFieldType::Array>(TEXT("Array, None"), Field, ECbFieldError::TypeError);
-		FCbField::MakeView(Field).AsArray();
+		FCbFieldView FieldView;
+		TestFieldError<ECbFieldType::Array>(TEXT("Array, None, View"), FieldView, ECbFieldError::TypeError);
+		FCbField Field;
+		TestFieldError<ECbFieldType::Array, FCbArray, FCbFieldArrayAccessors>(TEXT("Array, None"), Field, ECbFieldError::TypeError);
 	}
 
 	// Test FCbArrayView(ArrayWithName) and CreateIterator
@@ -711,7 +735,7 @@ bool FCbFieldBinaryTest::RunTest(const FString& Parameters)
 	};
 
 	// Test FCbFieldView(Binary, Empty)
-	TestField<ECbFieldType::Binary>(TEXT("Binary, Empty"), {0});
+	TestFieldBytes<ECbFieldType::Binary>(TEXT("Binary, Empty"), {0});
 
 	// Test FCbFieldView(Binary, Value)
 	{
@@ -743,19 +767,19 @@ IMPLEMENT_CUSTOM_SIMPLE_AUTOMATION_TEST(FCbFieldStringTest, FCbFieldTestBase, "S
 bool FCbFieldStringTest::RunTest(const FString& Parameters)
 {
 	// Test FCbFieldView(String, Empty)
-	TestField<ECbFieldType::String>(TEXT("String, Empty"), {0});
+	TestFieldBytes<ECbFieldType::String>(TEXT("String, Empty"), {0});
 
 	// Test FCbFieldView(String, Value)
 	{
 		const uint8 Value[] = { 3, 'A', 'B', 'C' }; // Size: 3, Data: ABC
-		TestField<ECbFieldType::String>(TEXT("String, Value"), Value, FUtf8StringView(reinterpret_cast<const UTF8CHAR*>(Value) + 1, 3));
+		TestFieldBytes<ECbFieldType::String>(TEXT("String, Value"), Value, FUtf8StringView(reinterpret_cast<const UTF8CHAR*>(Value) + 1, 3));
 	}
 
 	// Test FCbFieldView(String, OutOfRangeSize)
 	{
 		uint8 Value[9];
 		WriteVarUInt(uint64(1) << 31, Value);
-		TestFieldError<ECbFieldType::String>(TEXT("String, OutOfRangeSize"), Value, ECbFieldError::RangeError, UTF8TEXTVIEW("ABC"));
+		TestFieldBytesError<ECbFieldType::String>(TEXT("String, OutOfRangeSize"), Value, ECbFieldError::RangeError, UTF8TEXTVIEW("ABC"));
 	}
 
 	// Test FCbFieldView(None) as String
@@ -882,7 +906,7 @@ bool FCbFieldFloatTest::RunTest(const FString& Parameters)
 	// Test FCbFieldView(Float, 32-bit)
 	{
 		const uint8 Value[] = { 0xc0, 0x12, 0x34, 0x56 }; // -2.28444433f
-		TestField<ECbFieldType::Float32>(TEXT("Float32"), Value, -2.28444433f);
+		TestFieldBytes<ECbFieldType::Float32>(TEXT("Float32"), Value, -2.28444433f);
 
 		FCbFieldView Field(Value, ECbFieldType::Float32);
 		TestField<ECbFieldType::Float64>(TEXT("Float32, AsDouble"), Field, -2.28444433);
@@ -891,7 +915,7 @@ bool FCbFieldFloatTest::RunTest(const FString& Parameters)
 	// Test FCbFieldView(Float, 64-bit)
 	{
 		const uint8 Value[] = { 0xc1, 0x23, 0x45, 0x67, 0x89, 0xab, 0xcd, 0xef }; // -631475.76888888876
-		TestField<ECbFieldType::Float64>(TEXT("Float64"), Value, -631475.76888888876);
+		TestFieldBytes<ECbFieldType::Float64>(TEXT("Float64"), Value, -631475.76888888876);
 
 		FCbFieldView Field(Value, ECbFieldType::Float64);
 		TestFieldError<ECbFieldType::Float32>(TEXT("Float64, AsFloat"), Field, ECbFieldError::RangeError, 8.0f);
@@ -992,10 +1016,10 @@ IMPLEMENT_CUSTOM_SIMPLE_AUTOMATION_TEST(FCbFieldBoolTest, FCbFieldTestBase, "Sys
 bool FCbFieldBoolTest::RunTest(const FString& Parameters)
 {
 	// Test FCbFieldView(Bool, False)
-	TestField<ECbFieldType::BoolFalse>(TEXT("Bool, False"), {}, false, true);
+	TestFieldBytes<ECbFieldType::BoolFalse>(TEXT("Bool, False"), {}, false, true);
 
 	// Test FCbFieldView(Bool, True)
-	TestField<ECbFieldType::BoolTrue>(TEXT("Bool, True"), {}, true, false);
+	TestFieldBytes<ECbFieldType::BoolTrue>(TEXT("Bool, True"), {}, true, false);
 
 	// Test FCbFieldView(None) as Bool
 	{
@@ -1014,10 +1038,10 @@ bool FCbFieldObjectAttachmentTest::RunTest(const FString& Parameters)
 	const FIoHash::ByteArray SequentialBytes{1, 2, 3, 4, 5, 6, 7, 8, 9, 10, 11, 12, 13, 14, 15, 16, 17, 18, 19, 20};
 
 	// Test FCbFieldView(ObjectAttachment, Zero)
-	TestField<ECbFieldType::ObjectAttachment>(TEXT("ObjectAttachment, Zero"), ZeroBytes);
+	TestFieldBytes<ECbFieldType::ObjectAttachment>(TEXT("ObjectAttachment, Zero"), ZeroBytes);
 
 	// Test FCbFieldView(ObjectAttachment, NonZero)
-	TestField<ECbFieldType::ObjectAttachment>(TEXT("ObjectAttachment, NonZero"), SequentialBytes, FIoHash(SequentialBytes));
+	TestFieldBytes<ECbFieldType::ObjectAttachment>(TEXT("ObjectAttachment, NonZero"), SequentialBytes, FIoHash(SequentialBytes));
 
 	// Test FCbFieldView(ObjectAttachment, NonZero) AsAttachment
 	{
@@ -1041,10 +1065,10 @@ bool FCbFieldBinaryAttachmentTest::RunTest(const FString& Parameters)
 	const FIoHash::ByteArray SequentialBytes{1, 2, 3, 4, 5, 6, 7, 8, 9, 10, 11, 12, 13, 14, 15, 16, 17, 18, 19, 20};
 
 	// Test FCbFieldView(BinaryAttachment, Zero)
-	TestField<ECbFieldType::BinaryAttachment>(TEXT("BinaryAttachment, Zero"), ZeroBytes);
+	TestFieldBytes<ECbFieldType::BinaryAttachment>(TEXT("BinaryAttachment, Zero"), ZeroBytes);
 
 	// Test FCbFieldView(BinaryAttachment, NonZero)
-	TestField<ECbFieldType::BinaryAttachment>(TEXT("BinaryAttachment, NonZero"), SequentialBytes, FIoHash(SequentialBytes));
+	TestFieldBytes<ECbFieldType::BinaryAttachment>(TEXT("BinaryAttachment, NonZero"), SequentialBytes, FIoHash(SequentialBytes));
 
 	// Test FCbFieldView(BinaryAttachment, NonZero) AsAttachment
 	{
@@ -1068,10 +1092,10 @@ bool FCbFieldHashTest::RunTest(const FString& Parameters)
 	const FIoHash::ByteArray SequentialBytes{1, 2, 3, 4, 5, 6, 7, 8, 9, 10, 11, 12, 13, 14, 15, 16, 17, 18, 19, 20};
 
 	// Test FCbFieldView(Hash, Zero)
-	TestField<ECbFieldType::Hash>(TEXT("Hash, Zero"), ZeroBytes);
+	TestFieldBytes<ECbFieldType::Hash>(TEXT("Hash, Zero"), ZeroBytes);
 
 	// Test FCbFieldView(Hash, NonZero)
-	TestField<ECbFieldType::Hash>(TEXT("Hash, NonZero"), SequentialBytes, FIoHash(SequentialBytes));
+	TestFieldBytes<ECbFieldType::Hash>(TEXT("Hash, NonZero"), SequentialBytes, FIoHash(SequentialBytes));
 
 	// Test FCbFieldView(None) as Hash
 	{
@@ -1102,10 +1126,10 @@ bool FCbFieldUuidTest::RunTest(const FString& Parameters)
 	const FGuid SequentialGuid(TEXT("00010203-0405-0607-0809-0a0b0c0d0e0f"));
 
 	// Test FCbFieldView(Uuid, Zero)
-	TestField<ECbFieldType::Uuid>(TEXT("Uuid, Zero"), ZeroBytes, FGuid(), SequentialGuid);
+	TestFieldBytes<ECbFieldType::Uuid>(TEXT("Uuid, Zero"), ZeroBytes, FGuid(), SequentialGuid);
 
 	// Test FCbFieldView(Uuid, NonZero)
-	TestField<ECbFieldType::Uuid>(TEXT("Uuid, NonZero"), SequentialBytes, SequentialGuid, FGuid());
+	TestFieldBytes<ECbFieldType::Uuid>(TEXT("Uuid, NonZero"), SequentialBytes, SequentialGuid, FGuid());
 
 	// Test FCbFieldView(None) as Uuid
 	{
@@ -1120,10 +1144,10 @@ IMPLEMENT_CUSTOM_SIMPLE_AUTOMATION_TEST(FCbFieldDateTimeTest, FCbFieldTestBase,
 bool FCbFieldDateTimeTest::RunTest(const FString& Parameters)
 {
 	// Test FCbFieldView(DateTime, Zero)
-	TestField<ECbFieldType::DateTime>(TEXT("DateTime, Zero"), {0, 0, 0, 0, 0, 0, 0, 0});
+	TestFieldBytes<ECbFieldType::DateTime>(TEXT("DateTime, Zero"), {0, 0, 0, 0, 0, 0, 0, 0});
 
 	// Test FCbFieldView(DateTime, 0x1020'3040'5060'7080)
-	TestField<ECbFieldType::DateTime>(TEXT("DateTime, NonZero"), {0x10, 0x20, 0x30, 0x40, 0x50, 0x60, 0x70, 0x80}, int64(0x1020'3040'5060'7080));
+	TestFieldBytes<ECbFieldType::DateTime>(TEXT("DateTime, NonZero"), {0x10, 0x20, 0x30, 0x40, 0x50, 0x60, 0x70, 0x80}, int64(0x1020'3040'5060'7080));
 
 	// Test FCbFieldView(DateTime, Zero) as FDateTime
 	{
@@ -1147,10 +1171,10 @@ IMPLEMENT_CUSTOM_SIMPLE_AUTOMATION_TEST(FCbFieldTimeSpanTest, FCbFieldTestBase,
 bool FCbFieldTimeSpanTest::RunTest(const FString& Parameters)
 {
 	// Test FCbFieldView(TimeSpan, Zero)
-	TestField<ECbFieldType::TimeSpan>(TEXT("TimeSpan, Zero"), {0, 0, 0, 0, 0, 0, 0, 0});
+	TestFieldBytes<ECbFieldType::TimeSpan>(TEXT("TimeSpan, Zero"), {0, 0, 0, 0, 0, 0, 0, 0});
 
 	// Test FCbFieldView(TimeSpan, 0x1020'3040'5060'7080)
-	TestField<ECbFieldType::TimeSpan>(TEXT("TimeSpan, NonZero"), {0x10, 0x20, 0x30, 0x40, 0x50, 0x60, 0x70, 0x80}, int64(0x1020'3040'5060'7080));
+	TestFieldBytes<ECbFieldType::TimeSpan>(TEXT("TimeSpan, NonZero"), {0x10, 0x20, 0x30, 0x40, 0x50, 0x60, 0x70, 0x80}, int64(0x1020'3040'5060'7080));
 
 	// Test FCbFieldView(TimeSpan, Zero) as FTimeSpan
 	{
@@ -1174,10 +1198,10 @@ IMPLEMENT_CUSTOM_SIMPLE_AUTOMATION_TEST(FCbFieldObjectIdTest, FCbFieldTestBase,
 bool FCbFieldObjectIdTest::RunTest(const FString& Parameters)
 {
 	// Test FCbFieldView(ObjectId, Zero)
-	TestField<ECbFieldType::ObjectId>(TEXT("ObjectId, Zero"), {0, 0, 0, 0, 0, 0, 0, 0, 0, 0, 0, 0});
+	TestFieldBytes<ECbFieldType::ObjectId>(TEXT("ObjectId, Zero"), {0, 0, 0, 0, 0, 0, 0, 0, 0, 0, 0, 0});
 
 	// Test FCbFieldView(ObjectId, 0x102030405060708090A0B0C0)
-	TestField<ECbFieldType::ObjectId>(TEXT("ObjectId, NonZero"), {0x10, 0x20, 0x30, 0x40, 0x50, 0x60, 0x70, 0x80, 0x90, 0xA0, 0xB0, 0xC0},
+	TestFieldBytes<ECbFieldType::ObjectId>(TEXT("ObjectId, NonZero"), {0x10, 0x20, 0x30, 0x40, 0x50, 0x60, 0x70, 0x80, 0x90, 0xA0, 0xB0, 0xC0},
 		FCbObjectId(MakeMemoryView<uint8>({0x10, 0x20, 0x30, 0x40, 0x50, 0x60, 0x70, 0x80, 0x90, 0xA0, 0xB0, 0xC0})));
 
 	// Test FCbFieldView(ObjectId, Zero) as FCbObjectId
@@ -1215,33 +1239,33 @@ bool FCbFieldCustomByIdTest::RunTest(const FString& Parameters)
 	// Test FCbFieldView(CustomById, MinId, Empty)
 	{
 		const uint8 Value[] = {1, 0};
-		TestField<ECbFieldType::CustomById>(TEXT("CustomById, MinId, Empty"), Value, FCbCustomById{0});
-		TestField<ECbFieldType::CustomById>(TEXT("CustomById, MinId, Empty, View"), Value, FMemoryView(), MakeMemoryView<uint8>({1, 2, 3}), ECbFieldError::None, FCustomByIdAccessor(0));
-		TestFieldError<ECbFieldType::CustomById>(TEXT("CustomById, MinId, Empty, InvalidId"), Value, ECbFieldError::RangeError, MakeMemoryView<uint8>({1, 2, 3}), FCustomByIdAccessor(MAX_uint64));
+		TestFieldBytes<ECbFieldType::CustomById>(TEXT("CustomById, MinId, Empty"), Value, FCbCustomById{0});
+		TestFieldBytes<ECbFieldType::CustomById>(TEXT("CustomById, MinId, Empty, View"), Value, FMemoryView(), MakeMemoryView<uint8>({1, 2, 3}), ECbFieldError::None, FCustomByIdAccessor(0));
+		TestFieldBytesError<ECbFieldType::CustomById>(TEXT("CustomById, MinId, Empty, InvalidId"), Value, ECbFieldError::RangeError, MakeMemoryView<uint8>({1, 2, 3}), FCustomByIdAccessor(MAX_uint64));
 	}
 
 	// Test FCbFieldView(CustomById, MinId, Value)
 	{
 		const uint8 Value[] = {5, 0, 1, 2, 3, 4};
-		TestFieldNoClone<ECbFieldType::CustomById>(TEXT("CustomById, MinId, Value"), Value, FCbCustomById{0, MakeMemoryView(Value).Right(4)});
-		TestFieldNoClone<ECbFieldType::CustomById>(TEXT("CustomById, MinId, Value, View"), Value, MakeMemoryView(Value).Right(4), FMemoryView(), ECbFieldError::None, FCustomByIdAccessor(0));
-		TestFieldError<ECbFieldType::CustomById>(TEXT("CustomById, MinId, Value, InvalidId"), Value, ECbFieldError::RangeError, MakeMemoryView<uint8>({1, 2, 3}), FCustomByIdAccessor(MAX_uint64));
+		TestFieldBytesNoClone<ECbFieldType::CustomById>(TEXT("CustomById, MinId, Value"), Value, FCbCustomById{0, MakeMemoryView(Value).Right(4)});
+		TestFieldBytesNoClone<ECbFieldType::CustomById>(TEXT("CustomById, MinId, Value, View"), Value, MakeMemoryView(Value).Right(4), FMemoryView(), ECbFieldError::None, FCustomByIdAccessor(0));
+		TestFieldBytesError<ECbFieldType::CustomById>(TEXT("CustomById, MinId, Value, InvalidId"), Value, ECbFieldError::RangeError, MakeMemoryView<uint8>({1, 2, 3}), FCustomByIdAccessor(MAX_uint64));
 	}
 
 	// Test FCbFieldView(CustomById, MaxId, Empty)
 	{
 		const uint8 Value[] = {9, 0xff, 0xff, 0xff, 0xff, 0xff, 0xff, 0xff, 0xff, 0xff};
-		TestField<ECbFieldType::CustomById>(TEXT("CustomById, MaxId, Empty"), Value, FCbCustomById{MAX_uint64});
-		TestField<ECbFieldType::CustomById>(TEXT("CustomById, MaxId, Empty, View"), Value, FMemoryView(), MakeMemoryView<uint8>({1, 2, 3}), ECbFieldError::None, FCustomByIdAccessor(MAX_uint64));
-		TestFieldError<ECbFieldType::CustomById>(TEXT("CustomById, MaxId, Empty, InvalidId"), Value, ECbFieldError::RangeError, MakeMemoryView<uint8>({1, 2, 3}), FCustomByIdAccessor(0));
+		TestFieldBytes<ECbFieldType::CustomById>(TEXT("CustomById, MaxId, Empty"), Value, FCbCustomById{MAX_uint64});
+		TestFieldBytes<ECbFieldType::CustomById>(TEXT("CustomById, MaxId, Empty, View"), Value, FMemoryView(), MakeMemoryView<uint8>({1, 2, 3}), ECbFieldError::None, FCustomByIdAccessor(MAX_uint64));
+		TestFieldBytesError<ECbFieldType::CustomById>(TEXT("CustomById, MaxId, Empty, InvalidId"), Value, ECbFieldError::RangeError, MakeMemoryView<uint8>({1, 2, 3}), FCustomByIdAccessor(0));
 	}
 
 	// Test FCbFieldView(CustomById, MaxId, Value)
 	{
 		const uint8 Value[] = {13, 0xff, 0xff, 0xff, 0xff, 0xff, 0xff, 0xff, 0xff, 0xff, 1, 2, 3, 4};
-		TestFieldNoClone<ECbFieldType::CustomById>(TEXT("CustomById, MaxId, Value"), Value, FCbCustomById{MAX_uint64, MakeMemoryView(Value).Right(4)});
-		TestFieldNoClone<ECbFieldType::CustomById>(TEXT("CustomById, MaxId, Value, View"), Value, MakeMemoryView(Value).Right(4), FMemoryView(), ECbFieldError::None, FCustomByIdAccessor(MAX_uint64));
-		TestFieldError<ECbFieldType::CustomById>(TEXT("CustomById, MaxId, Value, InvalidId"), Value, ECbFieldError::RangeError, MakeMemoryView<uint8>({1, 2, 3}), FCustomByIdAccessor(0));
+		TestFieldBytesNoClone<ECbFieldType::CustomById>(TEXT("CustomById, MaxId, Value"), Value, FCbCustomById{MAX_uint64, MakeMemoryView(Value).Right(4)});
+		TestFieldBytesNoClone<ECbFieldType::CustomById>(TEXT("CustomById, MaxId, Value, View"), Value, MakeMemoryView(Value).Right(4), FMemoryView(), ECbFieldError::None, FCustomByIdAccessor(MAX_uint64));
+		TestFieldBytesError<ECbFieldType::CustomById>(TEXT("CustomById, MaxId, Value, InvalidId"), Value, ECbFieldError::RangeError, MakeMemoryView<uint8>({1, 2, 3}), FCustomByIdAccessor(0));
 	}
 
 	// Test FCbFieldView(None) as CustomById
@@ -1273,17 +1297,17 @@ bool FCbFieldCustomByNameTest::RunTest(const FString& Parameters)
 	// Test FCbFieldView(CustomByName, ABC, Empty)
 	{
 		const uint8 Value[] = {4, 3, 'A', 'B', 'C'};
-		TestField<ECbFieldType::CustomByName>(TEXT("CustomByName, MinId, Empty"), Value, FCbCustomByName{UTF8TEXTVIEW("ABC")});
-		TestField<ECbFieldType::CustomByName>(TEXT("CustomByName, MinId, Empty, View"), Value, FMemoryView(), MakeMemoryView<uint8>({1, 2, 3}), ECbFieldError::None, FCustomByNameAccessor(UTF8TEXTVIEW("ABC")));
-		TestFieldError<ECbFieldType::CustomByName>(TEXT("CustomByName, MinId, Empty, InvalidCase"), Value, ECbFieldError::RangeError, MakeMemoryView<uint8>({1, 2, 3}), FCustomByNameAccessor(UTF8TEXTVIEW("abc")));
+		TestFieldBytes<ECbFieldType::CustomByName>(TEXT("CustomByName, MinId, Empty"), Value, FCbCustomByName{UTF8TEXTVIEW("ABC")});
+		TestFieldBytes<ECbFieldType::CustomByName>(TEXT("CustomByName, MinId, Empty, View"), Value, FMemoryView(), MakeMemoryView<uint8>({1, 2, 3}), ECbFieldError::None, FCustomByNameAccessor(UTF8TEXTVIEW("ABC")));
+		TestFieldBytesError<ECbFieldType::CustomByName>(TEXT("CustomByName, MinId, Empty, InvalidCase"), Value, ECbFieldError::RangeError, MakeMemoryView<uint8>({1, 2, 3}), FCustomByNameAccessor(UTF8TEXTVIEW("abc")));
 	}
 
 	// Test FCbFieldView(CustomByName, ABC, Value)
 	{
 		const uint8 Value[] = {8, 3, 'A', 'B', 'C', 1, 2, 3, 4};
-		TestFieldNoClone<ECbFieldType::CustomByName>(TEXT("CustomByName, MinId, Value"), Value, FCbCustomByName{UTF8TEXTVIEW("ABC"), MakeMemoryView(Value).Right(4)});
-		TestFieldNoClone<ECbFieldType::CustomByName>(TEXT("CustomByName, MinId, Value, View"), Value, MakeMemoryView(Value).Right(4), FMemoryView(), ECbFieldError::None, FCustomByNameAccessor(UTF8TEXTVIEW("ABC")));
-		TestFieldError<ECbFieldType::CustomByName>(TEXT("CustomByName, MinId, Value, InvalidCase"), Value, ECbFieldError::RangeError, MakeMemoryView<uint8>({1, 2, 3}), FCustomByNameAccessor(UTF8TEXTVIEW("abc")));
+		TestFieldBytesNoClone<ECbFieldType::CustomByName>(TEXT("CustomByName, MinId, Value"), Value, FCbCustomByName{UTF8TEXTVIEW("ABC"), MakeMemoryView(Value).Right(4)});
+		TestFieldBytesNoClone<ECbFieldType::CustomByName>(TEXT("CustomByName, MinId, Value, View"), Value, MakeMemoryView(Value).Right(4), FMemoryView(), ECbFieldError::None, FCustomByNameAccessor(UTF8TEXTVIEW("ABC")));
+		TestFieldBytesError<ECbFieldType::CustomByName>(TEXT("CustomByName, MinId, Value, InvalidCase"), Value, ECbFieldError::RangeError, MakeMemoryView<uint8>({1, 2, 3}), FCustomByNameAccessor(UTF8TEXTVIEW("abc")));
 	}
 
 	// Test FCbFieldView(None) as CustomByName
diff --git a/Engine/Source/Runtime/Core/Public/Serialization/CompactBinary.h b/Engine/Source/Runtime/Core/Public/Serialization/CompactBinary.h
--- a/Engine/Source/Runtime/Core/Public/Serialization/CompactBinary.h
+++ b/Engine/Source/Runtime/Core/Public/Serialization/CompactBinary.h
@@ -1425,22 +1425,26 @@ public:
 
 inline FCbObject FCbField::AsObject() &
 {
-	return IsObject() ? FCbObject(AsObjectView(), *this) : FCbObject();
+	const FCbObjectView View = AsObjectView();
+	return !HasError() ? FCbObject(View, *this) : FCbObject();
 }
 
 inline FCbObject FCbField::AsObject() &&
 {
-	return IsObject() ? FCbObject(AsObjectView(), MoveTemp(*this)) : FCbObject();
+	const FCbObjectView View = AsObjectView();
+	return !HasError() ? FCbObject(View, MoveTemp(*this)) : FCbObject();
 }
 
 inline FCbArray FCbField::AsArray() &
 {
-	return IsArray() ? FCbArray(AsArrayView(), *this) : FCbArray();
+	const FCbArrayView View = AsArrayView();
+	return !HasError() ? FCbArray(View, *this) : FCbArray();
 }
 
 inline FCbArray FCbField::AsArray() &&
 {
-	return IsArray() ? FCbArray(AsArrayView(), MoveTemp(*this)) : FCbArray();
+	const FCbArrayView View = AsArrayView();
+	return !HasError() ? FCbArray(View, MoveTemp(*this)) : FCbArray();
 }
 
 inline FSharedBuffer FCbField::AsBinary(const FSharedBuffer& Default) &
