- Added a -setby= option to DumpCVars, to filter on how they were set, like "DumpCVars r. -setby=DeviceProfile" will show all rendering (r. prefix) that were last set by a DeviceProfile

#changelist validated
#virtualized

diff --git a/Engine/Source/Runtime/Core/Private/HAL/ConsoleManager.cpp b/Engine/Source/Runtime/Core/Private/HAL/ConsoleManager.cpp
--- a/Engine/Source/Runtime/Core/Private/HAL/ConsoleManager.cpp
+++ b/Engine/Source/Runtime/Core/Private/HAL/ConsoleManager.cpp
@@ -2299,6 +2299,8 @@ static void DumpObjects(const TMap<FString, IConsoleObject*>& ConsoleObjects, co
 	FString CSVFilename;
 	bool bWriteToCSV = FParse::Value(Params, TEXT("-csv="), CSVFilename);
 	bWriteToCSV = bWriteToCSV || FParse::Param(Params, TEXT("csv"));
+	FString FilterSetBy;
+	FParse::Value(Params, TEXT("-setby="), FilterSetBy);
 	FString Prefix = FParse::Token(Params, false);
 	if (Prefix.StartsWith(TEXT("-")))
 	{
@@ -2371,6 +2373,10 @@ static void DumpObjects(const TMap<FString, IConsoleObject*>& ConsoleObjects, co
 			}
 			if (!bDisplayCommands && CVar != nullptr)
 			{
+				if (FilterSetBy.Len() > 0 && GetConsoleVariableSetByName(CVar->GetFlags()) != FilterSetBy)
+				{
+					continue;
+				}
 				if (bWriteToCSV)
 				{
 					MultiLogf(Log, CSV, TEXT("%s,%s,%s%s"), *Key, *CVar->GetString(), GetConsoleVariableSetByName(CVar->GetFlags()), *Help);
