Reduced book keeping memory in FMallocBinned3

diff --git a/Engine/Source/Runtime/Core/Private/HAL/MallocBinned2.cpp b/Engine/Source/Runtime/Core/Private/HAL/MallocBinned2.cpp
--- a/Engine/Source/Runtime/Core/Private/HAL/MallocBinned2.cpp
+++ b/Engine/Source/Runtime/Core/Private/HAL/MallocBinned2.cpp
@@ -686,10 +686,9 @@ FMallocBinned2::FMallocBinned2()
 	PageSize = Constants.BinnedPageSize;
 	OsAllocationGranularity = Constants.BinnedAllocationGranularity ? Constants.BinnedAllocationGranularity : PageSize;
 	NumPoolsPerPage = PageSize / sizeof(FPoolInfo);
-	PtrToPoolMapping.Init(PageSize, NumPoolsPerPage, Constants.AddressLimit);
+	PtrToPoolMapping.Init(PageSize, NumPoolsPerPage, Constants.AddressStart, Constants.AddressLimit);
 
 	checkf(FMath::IsPowerOfTwo(PageSize), TEXT("OS page size must be a power of two"));
-	checkf(FMath::IsPowerOfTwo(Constants.AddressLimit), TEXT("OS address limit must be a power of two"));
 	checkf(Constants.AddressLimit > PageSize, TEXT("OS address limit must be greater than the page size")); // Check to catch 32 bit overflow in AddressLimit
 	checkf(PageSize % BINNED2_LARGE_ALLOC == 0, TEXT("OS page size must be a multiple of BINNED2_LARGE_ALLOC"));
 
diff --git a/Engine/Source/Runtime/Core/Private/HAL/MallocBinned3.cpp b/Engine/Source/Runtime/Core/Private/HAL/MallocBinned3.cpp
--- a/Engine/Source/Runtime/Core/Private/HAL/MallocBinned3.cpp
+++ b/Engine/Source/Runtime/Core/Private/HAL/MallocBinned3.cpp
@@ -679,15 +679,90 @@
 
 	check(!PLATFORM_32BITS);
 
+	OsAllocationGranularity = FPlatformMemory::FPlatformVirtualMemoryBlock::GetCommitAlignment();
+	checkf(FMath::IsPowerOfTwo(OsAllocationGranularity), TEXT("OS page size must be a power of two"));
+
+	// First thing we try to allocate address space for bins as it might help us to move forward Constants.AddressStart and reduce the amount of available address space for the Large OS Allocs
+	// Available address space is used to reserve hash map that can address all of that range, so less addressable space means less memory is allocated for book keeping
+#if !BINNED3_USE_SEPARATE_VM_PER_POOL
+	Binned3BaseVMBlock = FPlatformMemory::FPlatformVirtualMemoryBlock::AllocateVirtual(BINNED3_SMALL_POOL_COUNT * MAX_MEMORY_PER_BLOCK_SIZE, OsAllocationGranularity);
+	Binned3BaseVMPtr = (uint8*)Binned3BaseVMBlock.GetVirtualPointer();
+	check(IsAligned(Binned3BaseVMPtr, OsAllocationGranularity));
+	verify(Binned3BaseVMPtr);
+#else
+
+	for (uint32 Index = 0; Index < BINNED3_SMALL_POOL_COUNT; ++Index)
+	{
+		FPlatformMemory::FPlatformVirtualMemoryBlock NewBLock = FPlatformMemory::FPlatformVirtualMemoryBlock::AllocateVirtual(MAX_MEMORY_PER_BLOCK_SIZE, OsAllocationGranularity);
+
+		uint8* NewVM = (uint8*)NewBLock.GetVirtualPointer();
+		check(IsAligned(NewVM, OsAllocationGranularity));
+		// insertion sort
+		if (Index && NewVM < PoolBaseVMPtr[Index - 1])
+		{
+			uint32 InsertIndex = 0;
+			for (; InsertIndex < Index; ++InsertIndex)
+			{
+				if (NewVM < PoolBaseVMPtr[InsertIndex])
+				{
+					break;
+				}
+			}
+			check(InsertIndex < Index);
+			for (uint32 MoveIndex = Index; MoveIndex > InsertIndex; --MoveIndex)
+			{
+				PoolBaseVMPtr[MoveIndex] = PoolBaseVMPtr[MoveIndex - 1];
+				PoolBaseVMBlock[MoveIndex] = PoolBaseVMBlock[MoveIndex - 1];
+			}
+			PoolBaseVMPtr[InsertIndex] = NewVM;
+			PoolBaseVMBlock[InsertIndex] = NewBLock;
+		}
+		else
+		{
+			PoolBaseVMPtr[Index] = NewVM;
+			PoolBaseVMBlock[Index] = NewBLock;
+		}
+	}
+	HighestPoolBaseVMPtr = PoolBaseVMPtr[BINNED3_SMALL_POOL_COUNT - 1];
+	uint64 TotalGaps = 0;
+	for (uint32 Index = 0; Index < BINNED3_SMALL_POOL_COUNT - 1; ++Index)
+	{
+		check(PoolBaseVMPtr[Index + 1] > PoolBaseVMPtr[Index]); // we sorted it
+		check(PoolBaseVMPtr[Index + 1] >= PoolBaseVMPtr[Index] + MAX_MEMORY_PER_BLOCK_SIZE); // and blocks are non-overlapping
+		TotalGaps += PoolBaseVMPtr[Index + 1] - (PoolBaseVMPtr[Index] + MAX_MEMORY_PER_BLOCK_SIZE);
+	}
+	if (TotalGaps == 0)
+	{
+		PoolSearchDiv = 0;
+	}
+	else if (TotalGaps < MAX_MEMORY_PER_BLOCK_SIZE)
+	{
+		PoolSearchDiv = MAX_MEMORY_PER_BLOCK_SIZE; // the gaps are not significant, ignoring them should give accurate searches
+	}
+	else
+	{
+		PoolSearchDiv = MAX_MEMORY_PER_BLOCK_SIZE + ((TotalGaps + BINNED3_SMALL_POOL_COUNT - 2) / (BINNED3_SMALL_POOL_COUNT - 1));
+	}
+#endif
+
 	FGenericPlatformMemoryConstants Constants = FPlatformMemory::GetConstants();
+#if !BINNED3_USE_SEPARATE_VM_PER_POOL
+	if (Constants.AddressStart == (uint64)Binned3BaseVMPtr)
+	{
+		Constants.AddressStart += Align(BINNED3_SMALL_POOL_COUNT * MAX_MEMORY_PER_BLOCK_SIZE, OsAllocationGranularity);
+	}
+#else
+	if (!TotalGaps && Constants.AddressStart == (uint64)PoolBaseVMPtr[0])
+	{
+		Constants.AddressStart += Align(BINNED3_SMALL_POOL_COUNT * MAX_MEMORY_PER_BLOCK_SIZE, OsAllocationGranularity);
+	}
+#endif
+
 	// large slab sizes are possible OsAllocationGranularity = 65536;
-	OsAllocationGranularity = FPlatformMemory::FPlatformVirtualMemoryBlock::GetCommitAlignment();
 	NumLargePoolsPerPage = OsAllocationGranularity / sizeof(FPoolInfoLarge);
 	check(OsAllocationGranularity % sizeof(FPoolInfoLarge) == 0);  // these need to divide evenly!
-	PtrToPoolMapping.Init(OsAllocationGranularity, NumLargePoolsPerPage, Constants.AddressLimit);
+	PtrToPoolMapping.Init(OsAllocationGranularity, NumLargePoolsPerPage, Constants.AddressStart, Constants.AddressLimit);
 
-	checkf(FMath::IsPowerOfTwo(OsAllocationGranularity), TEXT("OS page size must be a power of two"));
-	checkf(FMath::IsPowerOfTwo(Constants.AddressLimit), TEXT("OS address limit must be a power of two"));
 	checkf(Constants.AddressLimit > OsAllocationGranularity, TEXT("OS address limit must be greater than the page size")); // Check to catch 32 bit overflow in AddressLimit
 	static_assert(BINNED3_SMALL_POOL_COUNT <= 256, "Small block size array size must fit in a byte");
 	static_assert(sizeof(FFreeBlock) <= BINNED3_MINIMUM_ALIGNMENT, "Free block struct must be small enough to fit into a block.");
@@ -702,6 +777,7 @@
 
 	SmallPoolInfosPerPlatformPage = OsAllocationGranularity / sizeof(FPoolInfoSmall);
 
+	uint32 RequiredMetaMem = 0;
 	for (uint32 Index = 0; Index < BINNED3_SMALL_POOL_COUNT; ++Index)
 	{
 		checkf(Index == 0 || SizeTable[Index - 1].BlockSize < SizeTable[Index].BlockSize, TEXT("Small block sizes must be strictly increasing"));
@@ -719,45 +795,47 @@
 		SmallPoolTables[Index].TotalFreeCount.Store(0);
 #endif
 
-		int64 TotalNumberOfBlocksOfBlocks = MAX_MEMORY_PER_BLOCK_SIZE / (SizeTable[Index].PagesPlatformForBlockOfBlocks * OsAllocationGranularity);
-
-		int64 MaxPoolInfoMemory = Align(sizeof(FPoolInfoSmall**) * (TotalNumberOfBlocksOfBlocks + SmallPoolInfosPerPlatformPage - 1) / SmallPoolInfosPerPlatformPage, OsAllocationGranularity);
-		SmallPoolTables[Index].PoolInfos = (FPoolInfoSmall**)AllocateMetaDataMemory(MaxPoolInfoMemory);
-		FMemory::Memzero(SmallPoolTables[Index].PoolInfos, MaxPoolInfoMemory);
+		const int64 TotalNumberOfBlocksOfBlocks = MAX_MEMORY_PER_BLOCK_SIZE / (SizeTable[Index].PagesPlatformForBlockOfBlocks * OsAllocationGranularity);
+		const uint32 Size = Align(sizeof(FPoolInfoSmall**) * (TotalNumberOfBlocksOfBlocks + SmallPoolInfosPerPlatformPage - 1) / SmallPoolInfosPerPlatformPage, PLATFORM_CACHE_LINE_SIZE);
+		RequiredMetaMem += Size;
 #if BINNED3_ALLOCATOR_STATS
-		Binned3PoolInfoMemory += MaxPoolInfoMemory;
+		Binned3PoolInfoMemory += Size;
 #endif
 
-		{
-			int64 AllocationSize = FBitTree::GetMemoryRequirements(TotalNumberOfBlocksOfBlocks);
-			int64 AlignedAllocationSize = Align(AllocationSize, OsAllocationGranularity);
-
-			{
-				void *Bits = AllocateMetaDataMemory(AlignedAllocationSize);
-				check(Bits);
+		const int64 AllocationSize = Align(FBitTree::GetMemoryRequirements(TotalNumberOfBlocksOfBlocks), PLATFORM_CACHE_LINE_SIZE);
+		RequiredMetaMem += AllocationSize * 2;
 #if BINNED3_ALLOCATOR_STATS
-				Binned3FreeBitsMemory += AlignedAllocationSize;
+		Binned3FreeBitsMemory += AllocationSize * 2;
 #endif
-				SmallPoolTables[Index].BlockOfBlockAllocationBits.FBitTreeInit(TotalNumberOfBlocksOfBlocks, Bits, AlignedAllocationSize, false);
-			}
-			{
-				void *Bits = AllocateMetaDataMemory(AlignedAllocationSize);
-				check(Bits);
-#if BINNED3_ALLOCATOR_STATS
-				Binned3FreeBitsMemory += AlignedAllocationSize;
-#endif
-				SmallPoolTables[Index].BlockOfBlockIsExhausted.FBitTreeInit(TotalNumberOfBlocksOfBlocks, Bits, AlignedAllocationSize, true);
-			}
-		}
 	}
 
+	RequiredMetaMem = Align(RequiredMetaMem, OsAllocationGranularity);
+	uint8* MetaMem = (uint8*)AllocateMetaDataMemory(RequiredMetaMem);
+	const uint8* MetaMemEnd = MetaMem + RequiredMetaMem;
+	FMemory::Memzero(MetaMem, RequiredMetaMem);
 
+	for (uint32 Index = 0; Index < BINNED3_SMALL_POOL_COUNT; ++Index)
+	{
+		const int64 TotalNumberOfBlocksOfBlocks = MAX_MEMORY_PER_BLOCK_SIZE / (SizeTable[Index].PagesPlatformForBlockOfBlocks * OsAllocationGranularity);
+		const uint32 Size = Align(sizeof(FPoolInfoSmall**) * (TotalNumberOfBlocksOfBlocks + SmallPoolInfosPerPlatformPage - 1) / SmallPoolInfosPerPlatformPage, PLATFORM_CACHE_LINE_SIZE);
+
+		SmallPoolTables[Index].PoolInfos = (FPoolInfoSmall**)MetaMem;
+		MetaMem += Size;
+
+		const int64 AllocationSize = Align(FBitTree::GetMemoryRequirements(TotalNumberOfBlocksOfBlocks), PLATFORM_CACHE_LINE_SIZE);
+		SmallPoolTables[Index].BlockOfBlockAllocationBits.FBitTreeInit(TotalNumberOfBlocksOfBlocks, MetaMem, AllocationSize, false);
+		MetaMem += AllocationSize;
+
+		SmallPoolTables[Index].BlockOfBlockIsExhausted.FBitTreeInit(TotalNumberOfBlocksOfBlocks, MetaMem, AllocationSize, true);
+		MetaMem += AllocationSize;
+	}
+	check(MetaMem == MetaMemEnd);
+
 	// Set up pool mappings
 	uint8* IndexEntry = MemSizeToIndex;
 	uint32  PoolIndex  = 0;
 	for (uint32 Index = 0; Index != 1 + (BINNED3_MAX_SMALL_POOL_SIZE >> BINNED3_MINIMUM_ALIGNMENT_SHIFT); ++Index)
 	{
-		
 		uint32 BlockSize = Index << BINNED3_MINIMUM_ALIGNMENT_SHIFT; // inverse of int32 Index = int32((Size >> BINNED3_MINIMUM_ALIGNMENT_SHIFT));
 		while (SizeTable[PoolIndex].BlockSize < BlockSize)
 		{
@@ -767,15 +845,15 @@
 		check(PoolIndex < 256);
 		*IndexEntry++ = uint8(PoolIndex);
 	}
+
 	// now reverse the pool sizes for cache coherency
-
 	for (uint32 Index = 0; Index != BINNED3_SMALL_POOL_COUNT; ++Index)
 	{
 		uint32 Partner = BINNED3_SMALL_POOL_COUNT - Index - 1;
 		SmallBlockSizesReversedShifted[Index] = (SizeTable[Partner].BlockSize >> BINNED3_MINIMUM_ALIGNMENT_SHIFT);
 	}
+
 	uint64 MaxHashBuckets = PtrToPoolMapping.GetMaxHashBuckets();
-
 	{
 		int64 HashAllocSize = Align(MaxHashBuckets * sizeof(PoolHashBucket), OsAllocationGranularity);
 		HashBuckets = (PoolHashBucket*)AllocateMetaDataMemory(HashAllocSize);
@@ -788,67 +866,6 @@
 	DefaultConstructItems<PoolHashBucket>(HashBuckets, MaxHashBuckets);
 	MallocBinned3 = this;
 	GFixedMallocLocationPtr = (FMalloc**)(&MallocBinned3);
-
-#if !BINNED3_USE_SEPARATE_VM_PER_POOL
-	Binned3BaseVMBlock = FPlatformMemory::FPlatformVirtualMemoryBlock::AllocateVirtual(BINNED3_SMALL_POOL_COUNT * MAX_MEMORY_PER_BLOCK_SIZE, OsAllocationGranularity);
-	Binned3BaseVMPtr = (uint8*)Binned3BaseVMBlock.GetVirtualPointer();
-	check(IsAligned(Binned3BaseVMPtr, OsAllocationGranularity));
-	verify(Binned3BaseVMPtr);
-#else
-
-	for (uint32 Index = 0; Index < BINNED3_SMALL_POOL_COUNT; ++Index)
-	{
-		FPlatformMemory::FPlatformVirtualMemoryBlock NewBLock = FPlatformMemory::FPlatformVirtualMemoryBlock::AllocateVirtual(MAX_MEMORY_PER_BLOCK_SIZE, OsAllocationGranularity);
-
-		uint8* NewVM = (uint8*)NewBLock.GetVirtualPointer();
-		check(IsAligned(NewVM, OsAllocationGranularity));
-		// insertion sort
-		if (Index && NewVM < PoolBaseVMPtr[Index - 1])
-		{
-			uint32 InsertIndex = 0;
-			for (; InsertIndex < Index; ++InsertIndex)
-			{
-				if (NewVM < PoolBaseVMPtr[InsertIndex])
-				{
-					break;
-				}
-			}
-			check(InsertIndex < Index);
-			for (uint32 MoveIndex = Index; MoveIndex > InsertIndex; --MoveIndex)
-			{
-				PoolBaseVMPtr[MoveIndex] = PoolBaseVMPtr[MoveIndex - 1];
-				PoolBaseVMBlock[MoveIndex] = PoolBaseVMBlock[MoveIndex - 1];
-			}
-			PoolBaseVMPtr[InsertIndex] = NewVM;
-			PoolBaseVMBlock[InsertIndex] = NewBLock;
-		}
-		else
-		{
-			PoolBaseVMPtr[Index] = NewVM;
-			PoolBaseVMBlock[Index] = NewBLock;
-		}
-	}
-	HighestPoolBaseVMPtr = PoolBaseVMPtr[BINNED3_SMALL_POOL_COUNT - 1];
-	uint64 TotalGaps = 0;
-	for (uint32 Index = 0; Index < BINNED3_SMALL_POOL_COUNT - 1; ++Index)
-	{
-		check(PoolBaseVMPtr[Index + 1] > PoolBaseVMPtr[Index]); // we sorted it
-		check(PoolBaseVMPtr[Index + 1] >= PoolBaseVMPtr[Index] + MAX_MEMORY_PER_BLOCK_SIZE); // and blocks are non-overlapping
-		TotalGaps += PoolBaseVMPtr[Index + 1] - (PoolBaseVMPtr[Index] + MAX_MEMORY_PER_BLOCK_SIZE);
-	}
-	if (TotalGaps == 0)
-	{
-		PoolSearchDiv = 0;
-	}
-	else if (TotalGaps < MAX_MEMORY_PER_BLOCK_SIZE)
-	{
-		PoolSearchDiv = MAX_MEMORY_PER_BLOCK_SIZE; // the gaps are not significant, ignoring them should give accurate searches
-	}
-	else
-	{
-		PoolSearchDiv = MAX_MEMORY_PER_BLOCK_SIZE + ((TotalGaps + BINNED3_SMALL_POOL_COUNT - 2) / (BINNED3_SMALL_POOL_COUNT - 1));
-	}
-#endif
 }
 
 FMallocBinned3::~FMallocBinned3()
diff --git a/Engine/Source/Runtime/Core/Public/GenericPlatform/GenericPlatformMemory.h b/Engine/Source/Runtime/Core/Public/GenericPlatform/GenericPlatformMemory.h
--- a/Engine/Source/Runtime/Core/Public/GenericPlatform/GenericPlatformMemory.h
+++ b/Engine/Source/Runtime/Core/Public/GenericPlatform/GenericPlatformMemory.h
@@ -110,9 +110,17 @@ struct FGenericPlatformMemoryConstants
 	/** This is the "allocation granularity" in Binned malloc terms, i.e. BinnedMalloc will allocate the memory in increments of this value. If zero, Binned will use BinnedPageSize for this value. */
 	SIZE_T BinnedAllocationGranularity = 0;
 
-	// AddressLimit - Second parameter is estimate of the range of addresses expected to be returns by BinnedAllocFromOS(). Binned
-	// Malloc will adjust its internal structures to make lookups for memory allocations O(1) for this range. 
-	// It is ok to go outside this range, lookups will just be a little slower
+	/**
+	 * Starting address for the available virtual address space.
+	 * Can be used with AddressLimit to determine address space range for binned allocators
+	 */
+	uint64 AddressStart = 0;
+
+	/**
+	 * An estimate of the range of addresses expected to be returned by BinnedAllocFromOS(). Binned
+	 * Malloc will adjust its internal structures to make lookups for memory allocations O(1) for this range. 
+	 * It is ok to go outside this range, lookups will just be a little slower
+	 */
 	uint64 AddressLimit = (uint64)0xffffffff + 1;
 
 	/** Approximate physical RAM in GB; 1 on everything except PC. Used for "course tuning", like FPlatformMisc::NumberOfCores(). */
diff --git a/Engine/Source/Runtime/Core/Public/HAL/MallocBinnedCommon.h b/Engine/Source/Runtime/Core/Public/HAL/MallocBinnedCommon.h
--- a/Engine/Source/Runtime/Core/Public/HAL/MallocBinnedCommon.h
+++ b/Engine/Source/Runtime/Core/Public/HAL/MallocBinnedCommon.h
@@ -145,28 +145,32 @@ protected:
 			, HashKeyShift(0)
 			, PoolMask(0)
 			, MaxHashBuckets(0)
+			, AddressSpaceBase(0)
 		{
 		}
-		explicit FPtrToPoolMapping(uint32 InPageSize, uint64 InNumPoolsPerPage, uint64 AddressLimit)
+		explicit FPtrToPoolMapping(uint32 InPageSize, uint64 InNumPoolsPerPage, uint64 AddressBase, uint64 AddressLimit)
 		{
-			Init(InPageSize, InNumPoolsPerPage, AddressLimit);
+			Init(InPageSize, InNumPoolsPerPage, AddressBase, AddressLimit);
 		}
 
-		void Init(uint32 InPageSize, uint64 InNumPoolsPerPage, uint64 AddressLimit)
+		void Init(uint32 InPageSize, uint64 InNumPoolsPerPage, uint64 AddressBase, uint64 AddressLimit)
 		{
 			uint64 PoolPageToPoolBitShift = FPlatformMath::CeilLogTwo64(InNumPoolsPerPage);
 
 			PtrToPoolPageBitShift = FPlatformMath::CeilLogTwo(InPageSize);
 			HashKeyShift = PtrToPoolPageBitShift + PoolPageToPoolBitShift;
 			PoolMask = (1ull << PoolPageToPoolBitShift) - 1;
-			MaxHashBuckets = AddressLimit >> HashKeyShift;
+			MaxHashBuckets = FMath::RoundUpToPowerOfTwo64(AddressLimit - AddressBase) >> HashKeyShift;
+			AddressSpaceBase = AddressBase;
 		}
 
 		FORCEINLINE void GetHashBucketAndPoolIndices(const void* InPtr, uint32& OutBucketIndex, UPTRINT& OutBucketCollision, uint32& OutPoolIndex) const
 		{
-			OutBucketCollision = (UPTRINT)InPtr >> HashKeyShift;
+			check((UPTRINT)InPtr >= AddressSpaceBase);
+			const UPTRINT Ptr = (UPTRINT)InPtr - AddressSpaceBase;
+			OutBucketCollision = Ptr >> HashKeyShift;
 			OutBucketIndex = uint32(OutBucketCollision & (MaxHashBuckets - 1));
-			OutPoolIndex = uint32(((UPTRINT)InPtr >> PtrToPoolPageBitShift) & PoolMask);
+			OutPoolIndex = uint32((Ptr >> PtrToPoolPageBitShift) & PoolMask);
 		}
 
 		FORCEINLINE uint64 GetMaxHashBuckets() const
@@ -186,6 +190,9 @@ protected:
 
 		// PageSize dependent constants
 		uint64 MaxHashBuckets;
+
+		// Base address for any virtual allocations. Can be non 0 on some platforms
+		uint64 AddressSpaceBase;
 	};
 
 	// This needs to be small enough to fit inside the smallest allocation handled by MallocBinned2\3, hence the union.
