Fixed non unity compilation in MallocBinnedCommon due to missing includes

diff --git a/Engine/Source/Runtime/Core/Public/HAL/MallocBinned2.h b/Engine/Source/Runtime/Core/Public/HAL/MallocBinned2.h
--- a/Engine/Source/Runtime/Core/Public/HAL/MallocBinned2.h
+++ b/Engine/Source/Runtime/Core/Public/HAL/MallocBinned2.h
@@ -8,9 +8,7 @@
 #include "HAL/Allocators/PooledVirtualMemoryAllocator.h"
 #include "HAL/CriticalSection.h"
 #include "HAL/LowLevelMemTracker.h"
-#include "HAL/MemoryBase.h"
 #include "HAL/MallocBinnedCommon.h"
-#include "HAL/PlatformMath.h"
 #include "HAL/UnrealMemory.h"
 #include "Math/NumericLimits.h"
 #include "Misc/AssertionMacros.h"
diff --git a/Engine/Source/Runtime/Core/Public/HAL/MallocBinned3.h b/Engine/Source/Runtime/Core/Public/HAL/MallocBinned3.h
--- a/Engine/Source/Runtime/Core/Public/HAL/MallocBinned3.h
+++ b/Engine/Source/Runtime/Core/Public/HAL/MallocBinned3.h
@@ -10,8 +10,6 @@
 #include "HAL/CriticalSection.h"
 #include "HAL/LowLevelMemTracker.h"
 #include "HAL/MallocBinnedCommon.h"
-#include "HAL/MemoryBase.h"
-#include "HAL/PlatformMath.h"
 #include "HAL/PlatformMemory.h"
 #include "HAL/UnrealMemory.h"
 #include "Math/NumericLimits.h"
diff --git a/Engine/Source/Runtime/Core/Public/HAL/MallocBinnedCommon.h b/Engine/Source/Runtime/Core/Public/HAL/MallocBinnedCommon.h
--- a/Engine/Source/Runtime/Core/Public/HAL/MallocBinnedCommon.h
+++ b/Engine/Source/Runtime/Core/Public/HAL/MallocBinnedCommon.h
@@ -3,6 +3,8 @@
 #pragma once
 
 #include "CoreTypes.h"
+#include "HAL/MemoryBase.h"
+#include "HAL/PlatformMath.h"
 #include "HAL/PlatformTLS.h"
 #include "Templates/AlignmentTemplates.h"
 
