Replaced a wron PlatformMath.h with a proper UnrealMathUtility.h include

diff --git a/Engine/Source/Runtime/Core/Public/HAL/MallocBinnedCommon.h b/Engine/Source/Runtime/Core/Public/HAL/MallocBinnedCommon.h
--- a/Engine/Source/Runtime/Core/Public/HAL/MallocBinnedCommon.h
+++ b/Engine/Source/Runtime/Core/Public/HAL/MallocBinnedCommon.h
@@ -4,7 +4,7 @@
 
 #include "CoreTypes.h"
 #include "HAL/MemoryBase.h"
-#include "HAL/PlatformMath.h"
+#include "Math/UnrealMathUtility.h"
 #include "HAL/PlatformTLS.h"
 #include "Templates/AlignmentTemplates.h"
 
