Added LLM_SCOPED_PAUSE_TRACKING_WITH_ENUM_AND_AMOUNT_BYTAG

#rnx

diff --git a/Engine/Source/Runtime/Core/Public/HAL/LowLevelMemTracker.h b/Engine/Source/Runtime/Core/Public/HAL/LowLevelMemTracker.h
--- a/Engine/Source/Runtime/Core/Public/HAL/LowLevelMemTracker.h
+++ b/Engine/Source/Runtime/Core/Public/HAL/LowLevelMemTracker.h
@@ -381,6 +381,7 @@ extern FName LLMGetTagStat(ELLMTag Tag);
 #define LLM_SCOPED_PAUSE_TRACKING(AllocType) 						FLLMPauseScope SCOPE_NAME(ELLMTag::Untagged, false /* bIsStatTag */, 0, ELLMTracker::Max, AllocType);
 #define LLM_SCOPED_PAUSE_TRACKING_FOR_TRACKER(Tracker, AllocType) 	FLLMPauseScope SCOPE_NAME(ELLMTag::Untagged, false /* bIsStatTag */, 0, Tracker, AllocType);
 #define LLM_SCOPED_PAUSE_TRACKING_WITH_ENUM_AND_AMOUNT(Tag, Amount, Tracker, AllocType) FLLMPauseScope SCOPE_NAME(Tag, false /* bIsStatTag */, Amount, Tracker, AllocType);
+#define LLM_SCOPED_PAUSE_TRACKING_WITH_ENUM_AND_AMOUNT_BYTAG(TagDeclName, Amount, Tracker, AllocType) FLLMPauseScope SCOPE_NAME(PREPROCESSOR_JOIN(LLMTagDeclaration_, TagDeclName).GetUniqueName(), false /* bIsStatTag */, Amount, Tracker, AllocType);
 
 /**
  * LLM realloc scope macros. Used when reallocating a pointer and you wish to retain the tagging from the source pointer
@@ -1072,6 +1073,7 @@ struct FLLMTagSetAllocationFilter
 #define LLM_SCOPED_PAUSE_TRACKING(...)
 #define LLM_SCOPED_PAUSE_TRACKING_FOR_TRACKER(...)
 #define LLM_SCOPED_PAUSE_TRACKING_WITH_ENUM_AND_AMOUNT(...)
+#define LLM_SCOPED_PAUSE_TRACKING_WITH_ENUM_AND_AMOUNT_BYTAG(...)
 #define LLM_DUMP_TAG()
 #define LLM_DUMP_PLATFORM_TAG()
 #define LLM_DEFINE_TAG(...)
