[Switch] Add batch deleting of files to time cost of flushing journal on jornaled mounts. Mitigates the cost of per delete journal flush.

#rnx

diff --git a/Engine/Source/Runtime/Core/Public/GenericPlatform/GenericPlatformFile.h b/Engine/Source/Runtime/Core/Public/GenericPlatform/GenericPlatformFile.h
--- a/Engine/Source/Runtime/Core/Public/GenericPlatform/GenericPlatformFile.h
+++ b/Engine/Source/Runtime/Core/Public/GenericPlatform/GenericPlatformFile.h
@@ -371,6 +371,18 @@ public:
 	virtual int64		FileSize(const TCHAR* Filename) = 0;
 	/** Delete a file and return true if the file exists. Will not delete read only files. **/
 	virtual bool		DeleteFile(const TCHAR* Filename) = 0;
+	/** Delete an array of files and return true if ALL deletes are succeeded. **/
+	virtual bool		DeleteFiles(const TArrayView<const TCHAR*>& Filenames)
+	{
+		bool bOneFailed = false;
+
+		for (const TCHAR* File : Filenames)
+		{
+			bOneFailed |= !DeleteFile(File);
+		}
+
+		return !bOneFailed;
+	}
 	/** Return true if the file is read only. **/
 	virtual bool		IsReadOnly(const TCHAR* Filename) = 0;
 	/** Attempt to move a file. Return true if successful. Will not overwrite existing files. **/
diff --git a/Engine/Source/Runtime/Core/Public/HAL/IPlatformFileManagedStorageWrapper.h b/Engine/Source/Runtime/Core/Public/HAL/IPlatformFileManagedStorageWrapper.h
--- a/Engine/Source/Runtime/Core/Public/HAL/IPlatformFileManagedStorageWrapper.h
+++ b/Engine/Source/Runtime/Core/Public/HAL/IPlatformFileManagedStorageWrapper.h
@@ -120,6 +120,12 @@ public:
 		Unlock();
 	}
 
+	FManagedStorageScopeFileLock(const FManagedStorageScopeFileLock&) = delete;
+	FManagedStorageScopeFileLock& operator=(const FManagedStorageScopeFileLock&) = delete;
+
+	FManagedStorageScopeFileLock(FManagedStorageScopeFileLock&& InOther) = delete;
+	FManagedStorageScopeFileLock& operator=(FManagedStorageScopeFileLock&& InOther) = delete;
+	
 	void Unlock();
 
 private:
@@ -879,6 +885,42 @@ public:
 		return bSuccess;
 	}
 
+	virtual bool DeleteFiles(const TArrayView<const TCHAR*>& Filenames) override
+	{
+		if (!IsReady())
+		{
+			return BaseClass::DeleteFiles(Filenames);
+		}
+
+		FPersistentStorageManager& Manager = FPersistentStorageManager::Get();
+
+		TArray<TPair<FPersistentManagedFile, FManagedStorageScopeFileLock>> FileLocks;
+		FileLocks.Reserve(Filenames.Num());
+		for (const TCHAR* Filename : Filenames)
+		{
+			FPersistentManagedFile ManagedFile = Manager.TryManageFile(Filename);
+			FileLocks.Emplace(ManagedFile, MoveTemp(ManagedFile));
+		}
+
+		bool bSuccess = BaseClass::DeleteFiles(Filenames);
+
+		const int32 Count = Filenames.Num();
+		for (int32 i = 0; i < Count; i++)
+		{
+			if (!BaseClass::FileExists(Filenames[i]))
+			{
+				UE_LOG(LogPlatformFileManagedStorage, Display, TEXT("Removing deleted file %s."), *(FileLocks[i].Key.FullFilename));
+				Manager.RemoveFileFromManager(FileLocks[i].Key);
+			}
+			else
+			{
+				UE_LOG(LogPlatformFileManagedStorage, Warning, TEXT("Not removing deleted file %s. It still exists on disk."), *(FileLocks[i].Key.FullFilename));
+			}
+		}
+
+		return bSuccess;
+	}
+
 	virtual bool MoveFile(const TCHAR* To, const TCHAR* From) override
 	{
 		if (!IsReady())
diff --git a/Engine/Source/Runtime/Core/Public/HAL/Platform.h b/Engine/Source/Runtime/Core/Public/HAL/Platform.h
--- a/Engine/Source/Runtime/Core/Public/HAL/Platform.h
+++ b/Engine/Source/Runtime/Core/Public/HAL/Platform.h
@@ -597,6 +597,9 @@
 	#define PLATFORM_MAX_UNIFORM_BUFFER_RANGE (16u*1024u)
 #endif
 
+#ifndef PLATFORM_IMPLEMENTS_BATCH_FILE_DELETE
+	#define PLATFORM_IMPLEMENTS_BATCH_FILE_DELETE 0
+#endif
 // deprecated, do not use
 #define PLATFORM_HAS_THREADSAFE_RHIGetRenderQueryResult	#
 #define PLATFORM_SUPPORTS_RHI_THREAD #
