Query performance counters to estimate shader compilation memory usage

Local shader compilation can frequently take up lots of memory during cooking. We currently have no way to track this effectively in docker containers for debugging. We can get an estimate of the amount of memory used by all shader compile workers using performance counters with a wild card:
 * sum of private bytes gives a lower bound for the amount of memory used,
 * sum of pagefile usage gives an upper bound for the amount of memory used.

I have also found it helpful locally to tick the timers are little bit more frequently and added an option for that.
#tests tested locally by running a cook

diff --git a/Engine/Source/Runtime/Core/Private/Windows/WindowsPlatformPerfCounters.cpp b/Engine/Source/Runtime/Core/Private/Windows/WindowsPlatformPerfCounters.cpp
--- a/Engine/Source/Runtime/Core/Private/Windows/WindowsPlatformPerfCounters.cpp
+++ b/Engine/Source/Runtime/Core/Private/Windows/WindowsPlatformPerfCounters.cpp
@@ -36,20 +36,22 @@ THIRD_PARTY_INCLUDES_START
 THIRD_PARTY_INCLUDES_END
 #include "Windows/HideWindowsPlatformTypes.h"
 
-TRACE_DECLARE_MEMORY_COUNTER(Memory_TotalVirtual,         TEXT("PC / Memory / Total Virtual"));
-TRACE_DECLARE_MEMORY_COUNTER(Memory_TotalPhysical,        TEXT("PC / Memory / Total Physical"));
-TRACE_DECLARE_MEMORY_COUNTER(Memory_UsedPhysical,         TEXT("PC / Memory / Used Physical"));
-TRACE_DECLARE_MEMORY_COUNTER(Memory_UsedVirtual,          TEXT("PC / Memory / Used Virtual"));
-TRACE_DECLARE_MEMORY_COUNTER(Memory_AvailablePhysical,    TEXT("PC / Memory / Available Physical"));
-TRACE_DECLARE_MEMORY_COUNTER(Memory_AvailableVirtual,     TEXT("PC / Memory / Available Virtual"));
-TRACE_DECLARE_MEMORY_COUNTER(Memory_FreeAndZeroPageList,  TEXT("PC / Memory / Free & Zero Page List"));
-TRACE_DECLARE_FLOAT_COUNTER(System_DemandZeroPageFaults,  TEXT("PC / Memory / Demand Zero Page Faults"));
-TRACE_DECLARE_FLOAT_COUNTER(System_GPU_Usage,             TEXT("PC / GPU Engine / Total Usage"));
-TRACE_DECLARE_MEMORY_COUNTER(Memory_UsedDedicatedVRAM,    TEXT("PC / GPU Memory / Used Dedicated"));
-TRACE_DECLARE_FLOAT_COUNTER(System_CPU_Usage,             TEXT("PC / CPU / Total Usage"));
-TRACE_DECLARE_INT_COUNTER(System_Thread_Count,            TEXT("PC / CPU / Total Thread Count"));
-TRACE_DECLARE_INT_COUNTER(System_Network_Sent,            TEXT("PC / Network / Total Bytes Sent"));
-TRACE_DECLARE_INT_COUNTER(System_Network_Received,        TEXT("PC / Network / Total Bytes Received"));
+TRACE_DECLARE_MEMORY_COUNTER(Memory_TotalVirtual,					TEXT("PC / Memory / Total Virtual"));
+TRACE_DECLARE_MEMORY_COUNTER(Memory_TotalPhysical,					TEXT("PC / Memory / Total Physical"));
+TRACE_DECLARE_MEMORY_COUNTER(Memory_UsedPhysical,					TEXT("PC / Memory / Used Physical (Main Process)"));
+TRACE_DECLARE_MEMORY_COUNTER(Memory_UsedVirtual,					TEXT("PC / Memory / Used Virtual (Main Process)"));
+TRACE_DECLARE_MEMORY_COUNTER(Memory_AvailablePhysical,				TEXT("PC / Memory / Available Physical"));
+TRACE_DECLARE_MEMORY_COUNTER(Memory_AvailableVirtual,				TEXT("PC / Memory / Available Virtual"));
+TRACE_DECLARE_MEMORY_COUNTER(Memory_FreeAndZeroPageList,			TEXT("PC / Memory / Free & Zero Page List"));
+TRACE_DECLARE_MEMORY_COUNTER(Memory_PrivateBytesShaderCompilation,	TEXT("PC / Memory / Private Bytes (Local Shader Compile Workers)"));
+TRACE_DECLARE_MEMORY_COUNTER(Memory_UsedVirtualShaderCompilation,	TEXT("PC / Memory / Used Virtual (Local Shader Compile Workers)"));
+TRACE_DECLARE_FLOAT_COUNTER(System_DemandZeroPageFaults,			TEXT("PC / Memory / Demand Zero Page Faults"));
+TRACE_DECLARE_FLOAT_COUNTER(System_GPU_Usage,						TEXT("PC / GPU Engine / Total Usage"));
+TRACE_DECLARE_MEMORY_COUNTER(Memory_UsedDedicatedVRAM,				TEXT("PC / GPU Memory / Used Dedicated"));
+TRACE_DECLARE_FLOAT_COUNTER(System_CPU_Usage,						TEXT("PC / CPU / Total Usage"));
+TRACE_DECLARE_INT_COUNTER(System_Thread_Count,						TEXT("PC / CPU / Total Thread Count"));
+TRACE_DECLARE_INT_COUNTER(System_Network_Sent,						TEXT("PC / Network / Total Bytes Sent"));
+TRACE_DECLARE_INT_COUNTER(System_Network_Received,					TEXT("PC / Network / Total Bytes Received"));
 
 namespace WindowsPlatformPerfCountersImpl 
 {
@@ -57,7 +59,8 @@ namespace WindowsPlatformPerfCountersImpl
 class FPerfCountersThread : public FRunnable
 {
 public:
-	FPerfCountersThread()
+	FPerfCountersThread(float SampleIntervalInSeconds)
+		: SampleIntervalInSeconds(SampleIntervalInSeconds)
 	{
 		Thread = FRunnableThread::Create(this, TEXT("PerfCountersThread"), 512 * 1024, TPri_AboveNormal);
 		check(Thread);
@@ -82,14 +85,16 @@ public:
 	{
 		FCounter Counters[] =
 		{
-			{ L"\\Processor(_Total)\\% Processor Time",      [](double Value) { TRACE_COUNTER_SET(System_CPU_Usage, Value); } },
-			{ L"\\Process(_Total)\\Thread Count",            [](double Value) { TRACE_COUNTER_SET(System_Thread_Count, (int64)Value); } },
-			{ L"\\Memory\\Demand Zero Faults/sec",           [](double Value) { TRACE_COUNTER_SET(System_DemandZeroPageFaults, Value); } },
-			{ L"\\Memory\\Free & Zero Page List Bytes",      [](double Value) { TRACE_COUNTER_SET(Memory_FreeAndZeroPageList, (int64)Value); } },
-			{ L"\\Network Interface(*)\\Bytes Received/sec", [](double Value) { TRACE_COUNTER_SET(System_Network_Received, (int64)Value); } },
-			{ L"\\Network Interface(*)\\Bytes Sent/sec",     [](double Value) { TRACE_COUNTER_SET(System_Network_Sent, (int64)Value); } },
-			{ L"\\GPU Engine(*)\\Utilization Percentage",    [](double Value) { TRACE_COUNTER_SET(System_GPU_Usage, Value); } },
-			{ L"\\GPU Adapter Memory(*)\\Dedicated Usage",   [](double Value) { TRACE_COUNTER_SET(Memory_UsedDedicatedVRAM, (int64)Value); } },
+			{ L"\\Processor(_Total)\\% Processor Time",				[](double Value) { TRACE_COUNTER_SET(System_CPU_Usage, Value); } },
+			{ L"\\Process(_Total)\\Thread Count",					[](double Value) { TRACE_COUNTER_SET(System_Thread_Count, (int64)Value); } },
+			{ L"\\Memory\\Demand Zero Faults/sec",					[](double Value) { TRACE_COUNTER_SET(System_DemandZeroPageFaults, Value); } },
+			{ L"\\Memory\\Free & Zero Page List Bytes",				[](double Value) { TRACE_COUNTER_SET(Memory_FreeAndZeroPageList, (int64)Value); } },
+			{ L"\\Network Interface(*)\\Bytes Received/sec",		[](double Value) { TRACE_COUNTER_SET(System_Network_Received, (int64)Value); } },
+			{ L"\\Network Interface(*)\\Bytes Sent/sec",			[](double Value) { TRACE_COUNTER_SET(System_Network_Sent, (int64)Value); } },
+			{ L"\\GPU Engine(*)\\Utilization Percentage",			[](double Value) { TRACE_COUNTER_SET(System_GPU_Usage, Value); } },
+			{ L"\\GPU Adapter Memory(*)\\Dedicated Usage",			[](double Value) { TRACE_COUNTER_SET(Memory_UsedDedicatedVRAM, (int64)Value); } },
+			{ L"\\Process(ShaderCompileWorker*)\\Page File Bytes",	[](double Value) { TRACE_COUNTER_SET(Memory_UsedVirtualShaderCompilation, (int64)Value); } },
+			{ L"\\Process(ShaderCompileWorker*)\\Private Bytes",	[](double Value) { TRACE_COUNTER_SET(Memory_PrivateBytesShaderCompilation, (int64)Value); } },
 		};
 
 		PDH_HQUERY PdhQueryHandle = nullptr;
@@ -102,12 +107,13 @@ public:
 
 		DWORD CurrentBufferSize = 0;
 		PDH_FMT_COUNTERVALUE_ITEM* Buffer = nullptr;
+		const float SampleInterval = this->SampleIntervalInSeconds;
 
 		while (bContinue)
 		{
 			TRACE_CPUPROFILER_EVENT_SCOPE(FCountersStatsThread::Tick);
 
-			FPlatformProcess::SleepNoStats(0.250);
+			FPlatformProcess::SleepNoStats(SampleInterval);
 
 			FPlatformMemoryStats Stats = FPlatformMemory::GetStats();
 			TRACE_COUNTER_SET(Memory_TotalPhysical, Stats.TotalPhysical);
@@ -159,6 +165,7 @@ public:
 	}
 	std::atomic_bool bContinue{ true };
 	FRunnableThread* Thread = nullptr;
+	float SampleIntervalInSeconds;
 };
 
 static FPerfCountersThread* PerfCountersThread = nullptr;
@@ -172,7 +179,13 @@ void FWindowsPlatformPerfCounters::Init()
 	using namespace WindowsPlatformPerfCountersImpl;
 	if (PerfCountersThread == nullptr && FParse::Param(FCommandLine::Get(), TEXT("perfcounters")))
 	{
-		PerfCountersThread = new FPerfCountersThread();
+		float SampleInterval = 0;
+		FParse::Value(FCommandLine::Get(), TEXT("perfcountersinterval="), SampleInterval);
+		if (SampleInterval <= 0)
+		{
+			SampleInterval = 0.25f;
+		}
+		PerfCountersThread = new FPerfCountersThread(SampleInterval);
 	}
 #endif
 }
