Fixed meta mem check in MB3

diff --git a/Engine/Source/Runtime/Core/Private/HAL/MallocBinned3.cpp b/Engine/Source/Runtime/Core/Private/HAL/MallocBinned3.cpp
--- a/Engine/Source/Runtime/Core/Private/HAL/MallocBinned3.cpp
+++ b/Engine/Source/Runtime/Core/Private/HAL/MallocBinned3.cpp
@@ -829,7 +829,7 @@ FMallocBinned3::FMallocBinned3()
 		SmallPoolTables[Index].BlockOfBlockIsExhausted.FBitTreeInit(TotalNumberOfBlocksOfBlocks, MetaMem, AllocationSize, true);
 		MetaMem += AllocationSize;
 	}
-	check(MetaMem == MetaMemEnd);
+	check(MetaMem <= MetaMemEnd);
 
 	// Set up pool mappings
 	uint8* IndexEntry = MemSizeToIndex;
