Removed "Localization.TurkishUsesCapitalI" again as it caused issues

#rnx

diff --git a/Engine/Source/Runtime/Core/Private/Internationalization/ICUText.cpp b/Engine/Source/Runtime/Core/Private/Internationalization/ICUText.cpp
--- a/Engine/Source/Runtime/Core/Private/Internationalization/ICUText.cpp
+++ b/Engine/Source/Runtime/Core/Private/Internationalization/ICUText.cpp
@@ -27,13 +27,6 @@ THIRD_PARTY_INCLUDES_END
 #include "Internationalization/ICUInternationalization.h"
 #include "Internationalization/ICUTextCharacterIterator.h"
 
-static TAutoConsoleVariable<bool> CVarTurkishUsesCapitalI(
-	TEXT("Localization.TurkishUsesCapitalI"),
-	false,
-	TEXT("False: Use 'big i' (CLDR format, default), True: Use I."),
-	ECVF_Default
-	);
-
 FString FTextChronoFormatter::AsDate(const FDateTime& DateTime, const EDateTimeStyle::Type DateStyle, const FString& TimeZone, const FCulture& TargetCulture)
 {
 	FInternationalization& I18N = FInternationalization::Get();
@@ -99,15 +92,7 @@ FString FTextTransformer::ToUpper(const FString& InStr)
 	FInternationalization& I18N = FInternationalization::Get();
 	checkf(I18N.IsInitialized() == true, TEXT("FInternationalization is not initialized. An FText formatting method was likely used in static object initialization - this is not supported."));
 
-	FString Result = ICUUtilities::ConvertString(ICUUtilities::ConvertString(InStr).toUpper());
-
-	// The CLDR uses "big i" as the uppercase i in Turkish, however some styles prefer using I
-	if (I18N.GetCurrentLanguage()->GetTwoLetterISOLanguageName() == TEXT("tr") && CVarTurkishUsesCapitalI.AsVariable()->GetBool())
-	{
-		Result.ReplaceCharInline(TEXT('\u0130'), TEXT('I'), ESearchCase::CaseSensitive);
-	}
-
-	return Result;
+	return ICUUtilities::ConvertString(ICUUtilities::ConvertString(InStr).toUpper());
 }
 
 bool FText::IsWhitespace(const TCHAR Char)
