Make Linux/Mac use editor version of crash reporter when compiled with WITH_EDITOR

In order to properly separate crashes that happens during development versus crashes that occurr in shipped titles we need to use the editor version of the crash reporter when appropriate.

diff --git a/Engine/Source/Runtime/Core/Private/Mac/MacPlatformMisc.cpp b/Engine/Source/Runtime/Core/Private/Mac/MacPlatformMisc.cpp
--- a/Engine/Source/Runtime/Core/Private/Mac/MacPlatformMisc.cpp
+++ b/Engine/Source/Runtime/Core/Private/Mac/MacPlatformMisc.cpp
@@ -337,7 +337,11 @@ struct FMacApplicationInfo
 		// Cache & create the crash report folder.
 		FString ReportPath = FPaths::ConvertRelativePathToFull(FString::Printf(TEXT("%s"), *(FPaths::GameAgnosticSavedDir() / TEXT("Crashes"))));
 		FCStringAnsi::Strcpy(CrashReportPath, PATH_MAX+1, TCHAR_TO_UTF8(*ReportPath));
+#if WITH_EDITOR
+		FString ReportClient = FPaths::ConvertRelativePathToFull(FPlatformProcess::GenerateApplicationPath(TEXT("CrashReportClientEditor"), EBuildConfiguration::Development));
+#else
 		FString ReportClient = FPaths::ConvertRelativePathToFull(FPlatformProcess::GenerateApplicationPath(TEXT("CrashReportClient"), EBuildConfiguration::Development));
+#endif
 		FCStringAnsi::Strcpy(CrashReportClient, PATH_MAX+1, TCHAR_TO_UTF8(*ReportClient));
 		IFileManager::Get().MakeDirectory(*ReportPath, true);
 		
diff --git a/Engine/Source/Runtime/Core/Private/Unix/UnixPlatformCrashContext.cpp b/Engine/Source/Runtime/Core/Private/Unix/UnixPlatformCrashContext.cpp
--- a/Engine/Source/Runtime/Core/Private/Unix/UnixPlatformCrashContext.cpp
+++ b/Engine/Source/Runtime/Core/Private/Unix/UnixPlatformCrashContext.cpp
@@ -714,15 +714,11 @@ void FUnixCrashContext::GenerateCrashInfoAndLaunchReporter() const
 				FString CrashConfigDstAbsolute = FPaths::Combine(*CrashInfoAbsolute, *CrashConfigFilename);
 				static_cast<void>(IFileManager::Get().Copy(*CrashConfigDstAbsolute, CrashConfigFilePath));	// best effort, so don't care about result
 			}
-
-#if PLATFORM_LINUXARM64
-			// try launching the tool and wait for its exit, if at all
-			const TCHAR * RelativePathToCrashReporter = TEXT("../../../Engine/Binaries/LinuxArm64/CrashReportClient");	// FIXME: painfully hard-coded
+#if WITH_EDITOR
+			FString CrashReportClientPath = FPaths::ConvertRelativePathToFull(FPlatformProcess::GenerateApplicationPath(TEXT("CrashReportClientEditor"), EBuildConfiguration::Development));
 #else
-			// try launching the tool and wait for its exit, if at all
-			const TCHAR * RelativePathToCrashReporter = TEXT("../../../Engine/Binaries/Linux/CrashReportClient");	// FIXME: painfully hard-coded
+			FString CrashReportClientPath = FPaths::ConvertRelativePathToFull(FPlatformProcess::GenerateApplicationPath(TEXT("CrashReportClient"), EBuildConfiguration::Development));
 #endif
-
 			FString CrashReportLogFilename = LogBaseFilename + TEXT("-CRC") + LogExtension;
 			FString CrashReportLogFilepath = FPaths::Combine(*LogFolder, *CrashReportLogFilename);
 			FString CrashReportClientArguments = TEXT(" -Abslog=");
@@ -781,7 +777,7 @@ void FUnixCrashContext::GenerateCrashInfoAndLaunchReporter() const
 						if (UnixCrashReporterTracker::Processes[ProcessIdx].Status.compare_exchange_weak(IsAvailable, IsSpawning))
 						{
 							UnixCrashReporterTracker::Processes[ProcessIdx].Process = FPlatformProcess::CreateProc(
-								RelativePathToCrashReporter, *CrashReportClientArguments, true, false, false, NULL, 0, NULL, NULL);
+								*CrashReportClientPath, *CrashReportClientArguments, true, false, false, NULL, 0, NULL, NULL);
 							
 							++UnixCrashReporterTracker::NumUploadingProcesses;
 							UnixCrashReporterTracker::Processes[ProcessIdx].Status = UnixCrashReporterTracker::SlotStatus::Uploading;
@@ -833,7 +829,7 @@ void FUnixCrashContext::GenerateCrashInfoAndLaunchReporter() const
 			else if (bStartCRCFromEngineHandler)
 			{
 				// spin here until CrashReporter exits
-				FProcHandle RunningProc = FPlatformProcess::CreateProc(RelativePathToCrashReporter, *CrashReportClientArguments, true, false, false, NULL, 0, NULL, NULL);
+				FProcHandle RunningProc = FPlatformProcess::CreateProc(*CrashReportClientPath, *CrashReportClientArguments, true, false, false, NULL, 0, NULL, NULL);
 
 				// do not wait indefinitely - can be more generous about the hitch than in ensure() case
 				// NOTE: Chris.Wood - increased from 3 to 8 mins because server crashes were timing out and getting lost
