Fixed MallocBinned2\3 allocator stats define

diff --git a/Engine/Source/Runtime/Core/Public/HAL/MallocBinned3.h b/Engine/Source/Runtime/Core/Public/HAL/MallocBinned3.h
--- a/Engine/Source/Runtime/Core/Public/HAL/MallocBinned3.h
+++ b/Engine/Source/Runtime/Core/Public/HAL/MallocBinned3.h
@@ -78,11 +78,7 @@ struct FGenericMemoryStats;
 #define BINNED3_ALLOCATOR_STATS UE_BINNEDCOMMON_ALLOCATOR_STATS
 
 #if BINNED3_ALLOCATOR_STATS
-	#if !(UE_BUILD_SHIPPING || UE_BUILD_TEST)
-		#define BINNED3_ALLOCATOR_PER_BIN_STATS 1
-	#else
-		#define BINNED3_ALLOCATOR_PER_BIN_STATS 0
-	#endif
+	#define BINNED3_ALLOCATOR_PER_BIN_STATS !(UE_BUILD_SHIPPING || UE_BUILD_TEST)
 #else
 	#define BINNED3_ALLOCATOR_PER_BIN_STATS 0
 #endif
diff --git a/Engine/Source/Runtime/Core/Public/HAL/MallocBinnedCommon.h b/Engine/Source/Runtime/Core/Public/HAL/MallocBinnedCommon.h
--- a/Engine/Source/Runtime/Core/Public/HAL/MallocBinnedCommon.h
+++ b/Engine/Source/Runtime/Core/Public/HAL/MallocBinnedCommon.h
@@ -133,7 +133,7 @@ extern CORE_API int32 GMallocBinnedBundleCount;
 #endif
 
 #ifndef UE_BINNEDCOMMON_ALLOCATOR_STATS
-#	define UE_BINNEDCOMMON_ALLOCATOR_STATS 0
+#	define UE_BINNEDCOMMON_ALLOCATOR_STATS (!UE_BUILD_SHIPPING || WITH_EDITOR)
 #endif
 
 
