Stopped TArray::RemoveAt() being called with a non-integral size argument, including EAllowShrinking.

diff --git a/Engine/Source/Runtime/Core/Public/Containers/Array.h b/Engine/Source/Runtime/Core/Public/Containers/Array.h
--- a/Engine/Source/Runtime/Core/Public/Containers/Array.h
+++ b/Engine/Source/Runtime/Core/Public/Containers/Array.h
@@ -1738,7 +1738,10 @@ public:
 	 * @param Count (Optional) Number of elements to remove. Default is 1.
 	 * @param AllowShrinking (Optional) Tells if this call can shrink array if suitable after remove. Default is yes.
 	 */
-	template <typename CountType>
+	template <
+		typename CountType
+		UE_REQUIRES(std::is_integral_v<CountType>)
+	>
 	FORCEINLINE void RemoveAt(SizeType Index, CountType Count, EAllowShrinking AllowShrinking = EAllowShrinking::Yes)
 	{
 		static_assert(!std::is_same_v<CountType, bool>, "TArray::RemoveAt: unexpected bool passed as the Count argument");
@@ -1809,7 +1812,10 @@ public:
 	 * @param AllowShrinking (Optional) Tells if this call can shrink array if
 	 *                        suitable after remove. Default is yes.
 	 */
-	template <typename CountType>
+	template <
+		typename CountType
+		UE_REQUIRES(std::is_integral_v<CountType>)
+	>
 	FORCEINLINE void RemoveAtSwap(SizeType Index, CountType Count, EAllowShrinking AllowShrinking = EAllowShrinking::Yes)
 	{
 		static_assert(!std::is_same_v<CountType, bool>, "TArray::RemoveAtSwap: unexpected bool passed as the Count argument");
