I/O - changed dispatcher backend interface from resolving a single to a batch of I/O requests.

#rnx

diff --git a/Engine/Source/Runtime/Core/Internal/IO/IoContainers.h b/Engine/Source/Runtime/Core/Internal/IO/IoContainers.h
new file mode 100644
--- /dev/null
+++ b/Engine/Source/Runtime/Core/Internal/IO/IoContainers.h
@@ -0,0 +1,132 @@
+// Copyright Epic Games, Inc. All Rights Reserved.
+
+#pragma once
+
+#include "Misc/AssertionMacros.h"
+
+namespace UE::Private
+{
+
+template <typename TypeTraits>
+class TInstrusiveListIterator
+{
+public:
+	using ElementType = typename TypeTraits::ElementType;
+
+	explicit TInstrusiveListIterator(ElementType* InElement)
+		: Element(InElement)
+	{ }
+
+	ElementType&	operator*() const	{ check(Element); return *Element; }
+	explicit		operator bool()		{ return Element != nullptr; }
+	void			operator++()		{ check(Element); Element = TypeTraits::GetNext(Element); }
+	bool			operator!=(const TInstrusiveListIterator& Other) const { return Element != Other.Element; }
+
+private:
+	ElementType* Element;
+};
+
+} // namespace UE::Private
+
+template <typename TypeTraits>
+class TInstrusiveList
+{
+public:
+	using ElementType		= typename TypeTraits::ElementType;
+	using FIterator			= UE::Private::TInstrusiveListIterator<TypeTraits>;
+	using FConstIterator	= UE::Private::TInstrusiveListIterator<const TypeTraits>;
+
+	TInstrusiveList() = default;
+	TInstrusiveList(const TInstrusiveList&) = delete;
+	TInstrusiveList(TInstrusiveList&& Other)
+		: Head(Other.Head)
+		, Tail(Other.Tail)
+	{
+		Other.Head = Other.Tail = nullptr;
+	}
+
+	TInstrusiveList& operator=(const TInstrusiveList&) = delete;
+	TInstrusiveList& operator=(TInstrusiveList&& Other)
+	{
+		Head = Other.Head;
+		Tail = Other.Tail;
+		Other.Head = Other.Tail = nullptr;
+		return *this;
+	}
+
+	void AddTail(ElementType* Element)
+	{
+		check(Element != nullptr && TypeTraits::GetNext(Element) == nullptr);
+
+		if (Tail != nullptr)
+		{
+			check(Head != nullptr);
+			TypeTraits::SetNext(Tail, Element);
+			Tail = Element;
+		}
+		else
+		{
+			check(Head == nullptr);
+			Head = Tail = Element;
+		}
+	}
+
+	void AddTail(ElementType* First, ElementType* Last)
+	{
+		check(First && Last);
+		check(TypeTraits::GetNext(First) != nullptr || First == Last);
+
+		if (Tail != nullptr)
+		{
+			check(Head != nullptr);
+			TypeTraits::SetNext(Tail, First);
+			Tail = Last;
+		}
+		else
+		{
+			check(Head == nullptr);
+			Head = First;
+			Tail = Last;
+		}
+	}
+
+	void AddTail(TInstrusiveList&& Other)
+	{
+		if (!Other.IsEmpty())
+		{
+			AddTail(Other.Head, Other.Tail);
+			Other.Head = Other.Tail = nullptr;
+		}
+	}
+
+	ElementType* PopHead()
+	{
+		ElementType* Element = Head;
+		if (Element != nullptr)
+		{
+			Head = TypeTraits::GetNext(Element);
+			if (Head == nullptr)
+			{
+				Tail = nullptr;
+			}
+			TypeTraits::SetNext(Element, nullptr);
+		}
+
+		return Element;
+	}
+
+	bool				IsEmpty() const { return Head == nullptr; }
+	ElementType*		GetHead()		{ return Head; }
+	const ElementType*	GetHead() const { return Head; }
+	ElementType*		GetTail()		{ return Tail; }
+	const ElementType*	GetTail() const { return Tail; }
+
+	FIterator			begin()			{ return FIterator(Head); }
+	FConstIterator		begin() const	{ return FConstIterator(Head); }
+	FIterator			end()			{ return FIterator(nullptr); }
+	FConstIterator		end() const		{ return FConstIterator(nullptr); }
+
+private:
+	ElementType* Head = nullptr;
+	ElementType* Tail = nullptr;
+};
diff --git a/Engine/Source/Runtime/Core/Private/IO/IoDispatcher.cpp b/Engine/Source/Runtime/Core/Private/IO/IoDispatcher.cpp
--- a/Engine/Source/Runtime/Core/Private/IO/IoDispatcher.cpp
+++ b/Engine/Source/Runtime/Core/Private/IO/IoDispatcher.cpp
@@ -1,25 +1,20 @@
 // Copyright Epic Games, Inc. All Rights Reserved.
 
 #include "IO/IoDispatcher.h"
+#include "Async/Mutex.h"
+#include "Async/UniqueLock.h"
+#include "Containers/Ticker.h"
+#include "HAL/Runnable.h"
+#include "HAL/RunnableThread.h"
+#include "HAL/Event.h"
+#include "IO/IoDispatcherBackend.h"
 #include "IO/IoDispatcherPrivate.h"
-#include "IO/IoStore.h"
 #include "IO/IoOffsetLength.h"
 #include "Misc/ScopeRWLock.h"
-#include "Misc/CommandLine.h"
 #include "Misc/CoreDelegates.h"
-#include "Math/RandomStream.h"
 #include "Misc/ScopeLock.h"
-#include "HAL/Runnable.h"
-#include "HAL/RunnableThread.h"
-#include "HAL/PlatformProcess.h"
-#include "Serialization/LargeMemoryReader.h"
-#include "GenericPlatform/GenericPlatformChunkInstall.h"
-#include "HAL/Event.h"
-#include "Async/MappedFileHandle.h"
 #include "ProfilingDebugging/CountersTrace.h"
 #include "ProfilingDebugging/CsvProfiler.h"
-#include "Containers/Ticker.h"
-#include "IO/IoDispatcherBackend.h"
 #include "Templates/Greater.h"
 
 DEFINE_LOG_CATEGORY(LogIoDispatcher);
@@ -336,7 +331,6 @@ public:
 	{
 		RequestAllocator = new FIoRequestAllocator();
 		RequestAllocator->AddRef();
-		DispatcherEvent = FPlatformProcess::GetSynchEventFromPool(false);
 		BackendContext->WakeUpDispatcherThreadDelegate.BindRaw(this, &FIoDispatcherImpl::WakeUpDispatcherThread);
 		BackendContext->bIsMultiThreaded = bInIsMultithreaded;
 		MemoryTrimDelegateHandle = FCoreDelegates::GetMemoryTrimDelegate().AddLambda([this]()
@@ -355,7 +349,6 @@ public:
 		}
 		FCoreDelegates::GetMemoryTrimDelegate().Remove(MemoryTrimDelegateHandle);
 		BackendContext->WakeUpDispatcherThreadDelegate.Unbind();
-		FPlatformProcess::ReturnSynchEventToPool(DispatcherEvent);
 		RequestAllocator->ReleaseRef();
 	}
 
@@ -401,7 +394,7 @@ public:
 		}
 		Request->AddRef();
 		{
-			FScopeLock _(&UpdateLock);
+			UE::TUniqueLock Lock(UpdateMutex);
 			RequestsToCancel.Add(Request);
 		}
 		DispatcherEvent->Trigger();
@@ -415,7 +408,7 @@ public:
 		}
 		Request->AddRef();
 		{
-			FScopeLock _(&UpdateLock);
+			UE::TUniqueLock Lock(UpdateMutex);
 			RequestsToReprioritize.Add(Request);
 		}
 		DispatcherEvent->Trigger();
@@ -582,16 +575,8 @@ public:
 		RequestStats.OnBatchIssued(Batch);
 		
 		{
-			FScopeLock _(&WaitingLock);
-			if (!WaitingRequestsHead)
-			{
-				WaitingRequestsHead = Batch.HeadRequest;
-			}
-			else
-			{
-				WaitingRequestsTail->NextRequest = Batch.HeadRequest;
-			}
-			WaitingRequestsTail = Batch.TailRequest;
+			UE::TUniqueLock Lock(WaitingMutex);
+			WaitingRequests.AddTail(Batch.HeadRequest, Batch.TailRequest);
 		}
 		Batch.HeadRequest = Batch.TailRequest = nullptr;
 		if (BackendContext->bIsMultiThreaded)
@@ -663,7 +648,7 @@ private:
 		FReadScopeLock _(BackendsLock);
 		for (const FBackendAndPriority& Backend : Backends)
 		{
-			FIoRequestImpl* CompletedRequestsHead = Backend.Value->GetCompletedRequests();
+			FIoRequestImpl* CompletedRequestsHead = Backend.Value->GetCompletedIoRequests();
 			while (CompletedRequestsHead)
 			{
 				FIoRequestImpl* NextRequest = CompletedRequestsHead->NextRequest;
@@ -745,32 +730,18 @@ private:
 
 	void ProcessIncomingRequests()
 	{
-		FIoRequestImpl* RequestsToSubmitHead = nullptr;
-		FIoRequestImpl* RequestsToSubmitTail = nullptr;
+		FIoRequestList RequestsToSubmit;
 		//TRACE_CPUPROFILER_EVENT_SCOPE(ProcessIncomingRequests);
 		for (;;)
 		{
 			{
-				FScopeLock _(&WaitingLock);
-				if (WaitingRequestsHead)
-				{
-					if (RequestsToSubmitTail)
-					{
-						RequestsToSubmitTail->NextRequest = WaitingRequestsHead;
-						RequestsToSubmitTail = WaitingRequestsTail;
-					}
-					else
-					{
-						RequestsToSubmitHead = WaitingRequestsHead;
-						RequestsToSubmitTail = WaitingRequestsTail;
-					}
-					WaitingRequestsHead = WaitingRequestsTail = nullptr;
-				}
+				UE::TUniqueLock Lock(WaitingMutex);
+				RequestsToSubmit.AddTail(MoveTemp(WaitingRequests));
 			}
 			TArray<FIoRequestImpl*> LocalRequestsToCancel;
 			TArray<FIoRequestImpl*> LocalRequestsToReprioritize;
 			{
-				FScopeLock _(&UpdateLock);
+				UE::TUniqueLock Lock(UpdateMutex);
 				Swap(LocalRequestsToCancel, RequestsToCancel);
 				Swap(LocalRequestsToReprioritize, RequestsToReprioritize);
 			}
@@ -794,58 +765,69 @@ private:
 				}
 				RequestToRePrioritize->ReleaseRef();
 			}
-			if (!RequestsToSubmitHead)
+			if (RequestsToSubmit.IsEmpty())
 			{
 				return;
 			}
 
-			FIoRequestImpl* Request = RequestsToSubmitHead;
-			RequestsToSubmitHead = RequestsToSubmitHead->NextRequest;
-			Request->NextRequest = nullptr;
-
-			if (!RequestsToSubmitHead)
+			int32 BatchCount = 0;
+			FIoRequestList Batch;
+			while (FIoRequestImpl* Request = RequestsToSubmit.PopHead())
 			{
-				RequestsToSubmitTail = nullptr;
-			}
+				check(Request->NextRequest == nullptr);
+				RequestStats.OnRequestStarted(*Request);
 
-			RequestStats.OnRequestStarted(*Request);
-			if (Request->bCancelled)
-			{
-				CompleteRequest(Request, EIoErrorCode::Cancelled);
-				Request->ReleaseRef();
-				continue;
+				if (Request->bCancelled)
+				{
+					CompleteRequest(Request, EIoErrorCode::Cancelled);
+					Request->ReleaseRef();
+					continue;
+				}
+
+				if (!Request->ChunkId.IsValid())
+				{
+					CompleteRequest(Request, EIoErrorCode::NotFound);
+					Request->ReleaseRef();
+					continue;
+				}
+
+				Batch.AddTail(Request);
+				++BatchCount;
 			}
-			
-			// Make sure that the FIoChunkId in the request is valid before we try to do anything with it.
-			if (Request->ChunkId.IsValid())
+
+			if (BatchCount > 0)
 			{
 				TRACE_CPUPROFILER_EVENT_SCOPE(ResolveRequest);
-				bool bResolved = false;
 				FReadScopeLock _(BackendsLock);
+				
+				FIoRequestList Unresolved;
 				for (const FBackendAndPriority& Backend : Backends)
 				{
-					if (Backend.Value->Resolve(Request))
+					for (FIoRequestImpl& Request : Batch)
+					{
+						Request.Backend = &Backend.Value.Get();
+					}
+					Backend.Value->ResolveIoRequests(MoveTemp(Batch), Unresolved);
+					if (Unresolved.IsEmpty())
 					{
-						bResolved = true;
-						Request->Backend = &Backend.Value.Get();
 						break;
 					}
+					Batch = MoveTemp(Unresolved);
 				}
-				if (!bResolved)
+
+				int32 UnresolvedCount = 0;
+				while (FIoRequestImpl* Request = Unresolved.PopHead())
 				{
+					check(Request->NextRequest == nullptr);
+					Request->Backend = nullptr;
 					CompleteRequest(Request, EIoErrorCode::NotFound);
 					Request->ReleaseRef();
-					continue;
+					UnresolvedCount++;
 				}
+
+				check(UnresolvedCount <= BatchCount);
+				PendingIoRequestsCount += (BatchCount - UnresolvedCount);
 			}
-			else
-			{
-				CompleteRequest(Request, EIoErrorCode::InvalidParameter);
-				Request->ReleaseRef();
-				continue;
-			}
-			
-			++PendingIoRequestsCount;
 		}
 	}
 
@@ -898,11 +880,8 @@ private:
 	FIoRequestAllocator* RequestAllocator = nullptr;
 	FBatchAllocator BatchAllocator;
 	FRunnableThread* Thread = nullptr;
-	FEvent* DispatcherEvent = nullptr;
-	FCriticalSection WaitingLock;
-	FIoRequestImpl* WaitingRequestsHead = nullptr;
-	FIoRequestImpl* WaitingRequestsTail = nullptr;
-	FCriticalSection UpdateLock;
+	FEventRef DispatcherEvent;
+	FIoRequestList WaitingRequests;
 	TArray<FIoRequestImpl*> RequestsToCancel;
 	TArray<FIoRequestImpl*> RequestsToReprioritize;
 	TAtomic<bool> bStopRequested { false };
@@ -911,6 +890,8 @@ private:
 	int64 TotalLoaded = 0;
 	FIoRequestStats RequestStats;
 	bool bIsInitialized = false;
+	UE::FMutex WaitingMutex;
+	UE::FMutex UpdateMutex;
 };
 
 FIoDispatcher::FIoDispatcher()
diff --git a/Engine/Source/Runtime/Core/Public/IO/IoDispatcherBackend.h b/Engine/Source/Runtime/Core/Public/IO/IoDispatcherBackend.h
--- a/Engine/Source/Runtime/Core/Public/IO/IoDispatcherBackend.h
+++ b/Engine/Source/Runtime/Core/Public/IO/IoDispatcherBackend.h
@@ -2,6 +2,7 @@
 
 #pragma once
 
+#include "IO/IoContainers.h"
 #include "IO/IoOffsetLength.h"
 #include "IoDispatcher.h"
 #include "Async/InheritedContext.h"
@@ -13,13 +14,21 @@
 
 struct FIoContainerHeader;
 
+/**
+ * I/O request object.
+ */
 class FIoRequestImpl : private UE::FInheritedContextBase
 {
 public:
+	/** Pointer to the next request that can be used freely by the I/O dispatcher backend(s). */
 	FIoRequestImpl* NextRequest = nullptr;
+	/** Custom data that can be used freely by the I/O dispatcher backend(s). */
 	void* BackendData = nullptr;
+	/** The chunk ID. */
 	FIoChunkId ChunkId;
+	/** Read options. */
 	FIoReadOptions Options;
+	/** I/O dispatcher priority (EIoDispatcherPriority). */
 	int32 Priority = 0;
 
 	FIoRequestImpl(class FIoRequestAllocator& InAllocator)
@@ -28,28 +37,34 @@ public:
 		CaptureInheritedContext();
 	}
 
+	/** Returns whether the request has been cancelled. */
 	bool IsCancelled() const
 	{
 		return bCancelled;
 	}
 
+	/** Mark the request as failed. */
 	void SetFailed()
 	{
 		bFailed = true;
 	}
 
+	/** Returns whether request has a valid buffer. */
 	bool HasBuffer() const
 	{
 		return Buffer.IsSet();
 	}
 
+	/** Creates a new buffer for the request. */
 	CORE_API void CreateBuffer(uint64 Size);
 
+	/** Returns the internal buffer. */
 	FIoBuffer& GetBuffer()
 	{
 		return Buffer.GetValue();
 	}
 
+	/** Sets a new buffer. */
 	void SetResult(FIoBuffer InBuffer)
 	{
 		Buffer.Emplace(InBuffer);
@@ -99,31 +114,103 @@ private:
 	bool bFailed = false;
 };
 
+namespace UE::Private
+{
+
+struct FIoRequestListTraits
+{
+	using ElementType = FIoRequestImpl;
+
+	static FIoRequestImpl* GetNext(const FIoRequestImpl* Element)
+	{
+		return Element->NextRequest;
+	}
+
+	static void SetNext(FIoRequestImpl* Element, FIoRequestImpl* Next)
+	{
+		Element->NextRequest = Next;
+	}
+};
+
+} // namespace UE::Private
+
+using FIoRequestList = TInstrusiveList<UE::Private::FIoRequestListTraits>;
+
 DECLARE_DELEGATE(FWakeUpIoDispatcherThreadDelegate);
 
+/**
+ * Context object used for signalling the I/O dispatcher.
+ */ 
 struct FIoDispatcherBackendContext
 {
+	/** Callback for signalling completed I/O requests. */
 	FWakeUpIoDispatcherThreadDelegate WakeUpDispatcherThreadDelegate;
+	/** Callback for signalling corrupted chunks. */
 	FIoSignatureErrorDelegate SignatureErrorDelegate;
+	/** Whether the I/O dispatcher is running in a separate thread. */
 	bool bIsMultiThreaded;
 };
 
+/**
+ * I/O dispatcher backend interface.
+ *
+ * The following methods are called from the I/O dispatcher thread:
+ *		* ResolveIoRequests
+ *		* CancelIoRequest
+ *		* UpdatePriorityForIoRequest
+ *		* GetCompletedRequests
+ *
+ *	All other methods can be called from any thread.
+ */
 struct IIoDispatcherBackend
 {
+	/**
+	 * Called by the I/O dispatcher once initialized.
+	 * @param Context	Context object used for signalling the I/O dispatcher.
+	 */
 	virtual void Initialize(TSharedRef<const FIoDispatcherBackendContext> Context) = 0;
+
+	/** Called by the I/O dispatcher when shutting down. */
 	virtual void Shutdown() {};
-	virtual bool Resolve(FIoRequestImpl* Request) = 0;
+
+	/**
+	 * Create asynchronous read requests for the list of I/O request(s).
+	 * The request(s) can span multiple chunks.
+	 *
+	 * @param Requests			Request(s) to resolve.
+	 * @param OutUnresolved		Unresolved request(s) returned to the I/O dispatcher.
+	 */
+	virtual void ResolveIoRequests(FIoRequestList Requests, FIoRequestList& OutUnresolved) = 0;
+
+	/** 
+	 * Returns all completed request(s) to the I/O dispatcher. Triggered
+	 * by invoking the wake-up callback on the backend context object.
+	 */
+	virtual FIoRequestImpl* GetCompletedIoRequests() = 0;
+
+	/** Cancel the I/O request. */
 	virtual void CancelIoRequest(FIoRequestImpl* Request) = 0;
+
+	/** Update the priority of the I/O request. */
 	virtual void UpdatePriorityForIoRequest(FIoRequestImpl* Request) = 0;
+
+	/** Returns whether the chunk exists. */
 	virtual bool DoesChunkExist(const FIoChunkId& ChunkId) const = 0;
+
+	/** Returns whether specified range of the chunk exists. */
 	virtual bool DoesChunkExist(const FIoChunkId& ChunkId, const FIoOffsetAndLength& ChunkRange) const { return DoesChunkExist(ChunkId); }
+
+	/** Returns the size of the chunk. */
 	virtual TIoStatusOr<uint64> GetSizeForChunk(const FIoChunkId& ChunkId) const = 0;
+
+	/** Returns the size of the chunk and the size of the requested range. */
 	virtual TIoStatusOr<uint64> GetSizeForChunk(const FIoChunkId& ChunkId, const FIoOffsetAndLength& ChunkRange, uint64& OutAvailable) const
 	{
 		TIoStatusOr<uint64> ChunkSize = GetSizeForChunk(ChunkId);
 		OutAvailable = ChunkSize.IsOk() ? ChunkSize.ValueOrDie() : 0;
 		return ChunkSize;
 	}
-	virtual FIoRequestImpl* GetCompletedRequests() = 0;
+
+	/** Read the chunk as a memory mapped file. */
 	virtual TIoStatusOr<FIoMappedRegion> OpenMapped(const FIoChunkId& ChunkId, const FIoReadOptions& Options) = 0;
 };
diff --git a/Engine/Source/Runtime/CoreUObject/Private/UObject/PackageResourceIoDispatcherBackend.cpp b/Engine/Source/Runtime/CoreUObject/Private/UObject/PackageResourceIoDispatcherBackend.cpp
--- a/Engine/Source/Runtime/CoreUObject/Private/UObject/PackageResourceIoDispatcherBackend.cpp
+++ b/Engine/Source/Runtime/CoreUObject/Private/UObject/PackageResourceIoDispatcherBackend.cpp
@@ -172,15 +172,17 @@ public:
 	~FPackageResourceIoBackend();
 
 	virtual void Initialize(TSharedRef<const FIoDispatcherBackendContext> Context) override;
-	virtual bool Resolve(FIoRequestImpl* Request) override;
+	virtual void ResolveIoRequests(FIoRequestList Requests, FIoRequestList& OutUnresolved) override;
 	virtual void CancelIoRequest(FIoRequestImpl* Request) override;
 	virtual void UpdatePriorityForIoRequest(FIoRequestImpl* Request) override { }
 	virtual bool DoesChunkExist(const FIoChunkId& ChunkId) const override;
 	virtual TIoStatusOr<uint64> GetSizeForChunk(const FIoChunkId& ChunkId) const override;
-	virtual FIoRequestImpl* GetCompletedRequests() override;
+	virtual FIoRequestImpl* GetCompletedIoRequests() override;
 	virtual TIoStatusOr<FIoMappedRegion> OpenMapped(const FIoChunkId& ChunkId, const FIoReadOptions& Options) override;
 
 private:
+	bool Resolve(FIoRequestImpl* Request);
+
 	IPackageResourceManager& ResourceMgr;
 	TSharedPtr<const FIoDispatcherBackendContext> BackendContext;
 	FPendingRequests PendingRequests;
@@ -260,6 +262,17 @@ bool FPackageResourceIoBackend::Resolve(FIoRequestImpl* Request)
 	return true;
 }
 
+void FPackageResourceIoBackend::ResolveIoRequests(FIoRequestList Requests, FIoRequestList& OutUnresolved)
+{
+	while (FIoRequestImpl* Request = Requests.PopHead())
+	{
+		if (Resolve(Request) == false)
+		{
+			OutUnresolved.AddTail(Request);
+		}
+	}
+}
+
 void FPackageResourceIoBackend::CancelIoRequest(FIoRequestImpl* Request)
 {
 	PendingRequests.Cancel(Request);
@@ -298,7 +311,7 @@ TIoStatusOr<uint64> FPackageResourceIoBackend::GetSizeForChunk(const FIoChunkId&
 	return FIoStatus(EIoErrorCode::NotFound);
 }
 
-FIoRequestImpl* FPackageResourceIoBackend::GetCompletedRequests()
+FIoRequestImpl* FPackageResourceIoBackend::GetCompletedIoRequests()
 {
 	FIoRequestImpl* Requests = CompletedRequests.Dequeue();
 
diff --git a/Engine/Source/Runtime/Experimental/IoStoreOnDemand/Private/OnDemandIoDispatcherBackend.cpp b/Engine/Source/Runtime/Experimental/IoStoreOnDemand/Private/OnDemandIoDispatcherBackend.cpp
--- a/Engine/Source/Runtime/Experimental/IoStoreOnDemand/Private/OnDemandIoDispatcherBackend.cpp
+++ b/Engine/Source/Runtime/Experimental/IoStoreOnDemand/Private/OnDemandIoDispatcherBackend.cpp
@@ -1369,14 +1369,14 @@ public:
 	// I/O dispatcher backend
 	virtual void Initialize(TSharedRef<const FIoDispatcherBackendContext> Context) override;
 	virtual void Shutdown() override;
-	virtual bool Resolve(FIoRequestImpl* Request) override;
+	virtual void ResolveIoRequests(FIoRequestList Requests, FIoRequestList& OutUnresolved) override;
 	virtual void CancelIoRequest(FIoRequestImpl* Request) override;
 	virtual void UpdatePriorityForIoRequest(FIoRequestImpl* Request) override;
 	virtual bool DoesChunkExist(const FIoChunkId& ChunkId) const override;
 	virtual bool DoesChunkExist(const FIoChunkId& ChunkId, const FIoOffsetAndLength& ChunkRange) const override;
 	virtual TIoStatusOr<uint64> GetSizeForChunk(const FIoChunkId& ChunkId) const override;
 	virtual TIoStatusOr<uint64> GetSizeForChunk(const FIoChunkId& ChunkId, const FIoOffsetAndLength& ChunkRange, uint64& OutAvailable) const;
-	virtual FIoRequestImpl* GetCompletedRequests() override;
+	virtual FIoRequestImpl* GetCompletedIoRequests() override;
 	virtual TIoStatusOr<FIoMappedRegion> OpenMapped(const FIoChunkId& ChunkId, const FIoReadOptions& Options) override;
 
 	// I/O Http backend
@@ -1393,7 +1393,7 @@ public:
 	virtual void Stop() override;
 
 private:
-
+	bool Resolve(FIoRequestImpl* Request);
 	void MountContainer(FStringView ContainerPath);
 	void UnmountContainer(FStringView ContainerPath);
 
@@ -1843,6 +1843,17 @@ bool FOnDemandIoBackend::Resolve(FIoRequestImpl* Request)
 	return true;
 }
 
+void FOnDemandIoBackend::ResolveIoRequests(FIoRequestList Requests, FIoRequestList& OutUnresolved)
+{
+	while (FIoRequestImpl* Request = Requests.PopHead())
+	{
+		if (Resolve(Request) == false)
+		{
+			OutUnresolved.AddTail(Request);
+		}
+	}
+}
+
 void FOnDemandIoBackend::CancelIoRequest(FIoRequestImpl* Request)
 {
 	if (ChunkRequests.Cancel(Request, Cache.Get()))
@@ -1916,7 +1927,7 @@ TIoStatusOr<uint64> FOnDemandIoBackend::GetSizeForChunk(const FIoChunkId& ChunkI
 	return TIoStatusOr<uint64>(ChunkInfo.Entry->RawSize);
 }
 
-FIoRequestImpl* FOnDemandIoBackend::GetCompletedRequests()
+FIoRequestImpl* FOnDemandIoBackend::GetCompletedIoRequests()
 {
 	FIoRequestImpl* Requests = CompletedRequests.Dequeue();
 
diff --git a/Engine/Source/Runtime/PakFile/Private/IoDispatcherFileBackend.cpp b/Engine/Source/Runtime/PakFile/Private/IoDispatcherFileBackend.cpp
--- a/Engine/Source/Runtime/PakFile/Private/IoDispatcherFileBackend.cpp
+++ b/Engine/Source/Runtime/PakFile/Private/IoDispatcherFileBackend.cpp
@@ -1381,7 +1381,7 @@ bool FFileIoStore::Unmount(const TCHAR* InTocPath)
 
 bool FFileIoStore::Resolve(FIoRequestImpl* Request)
 {
-	FReadScopeLock _(IoStoreReadersLock);
+	// Assumes readers are locked, see ResolveIoRequests
 	for (const TUniquePtr<FFileIoStoreReader>& Reader : IoStoreReaders)
 	{
 		if (const FIoOffsetAndLength* OffsetAndLength = Reader->Resolve(Request->ChunkId))
@@ -1439,6 +1439,18 @@ bool FFileIoStore::Resolve(FIoRequestImpl* Request)
 	return false;
 }
 
+void FFileIoStore::ResolveIoRequests(FIoRequestList Requests, FIoRequestList& OutUnresolved)
+{
+	FReadScopeLock _(IoStoreReadersLock);
+	while (FIoRequestImpl* Request = Requests.PopHead())
+	{
+		if (Resolve(Request) == false)
+		{
+			OutUnresolved.AddTail(Request);
+		}
+	}
+}
+
 void FFileIoStore::CancelIoRequest(FIoRequestImpl* Request)
 {
 	if (Request->BackendData)
@@ -1678,7 +1690,7 @@ void FFileIoStore::FinalizeCompressedBlock(FFileIoStoreCompressedBlock* Compress
 	}
 }
 
-FIoRequestImpl* FFileIoStore::GetCompletedRequests()
+FIoRequestImpl* FFileIoStore::GetCompletedIoRequests()
 {
 	LLM_SCOPE(ELLMTag::FileSystem);
 	//TRACE_CPUPROFILER_EVENT_SCOPE(GetCompletedRequests);
@@ -2460,4 +2472,4 @@ void FFileIoStoreStats::OnBufferAllocated()
 #endif
 }
 
-#endif
\ No newline at end of file
+#endif
diff --git a/Engine/Source/Runtime/PakFile/Private/IoDispatcherFileBackend.h b/Engine/Source/Runtime/PakFile/Private/IoDispatcherFileBackend.h
--- a/Engine/Source/Runtime/PakFile/Private/IoDispatcherFileBackend.h
+++ b/Engine/Source/Runtime/PakFile/Private/IoDispatcherFileBackend.h
@@ -109,12 +109,12 @@ public:
 	void Shutdown() override;
 	TIoStatusOr<FIoContainerHeader> Mount(const TCHAR* InTocPath, int32 Order, const FGuid& EncryptionKeyGuid, const FAES::FAESKey& EncryptionKey);
 	bool Unmount(const TCHAR* InTocPath);
-	bool Resolve(FIoRequestImpl* Request) override;
+	void ResolveIoRequests(FIoRequestList Requests, FIoRequestList& OutUnresolved) override;
 	void CancelIoRequest(FIoRequestImpl* Request) override;
 	void UpdatePriorityForIoRequest(FIoRequestImpl* Request) override;
 	bool DoesChunkExist(const FIoChunkId& ChunkId) const override;
 	TIoStatusOr<uint64> GetSizeForChunk(const FIoChunkId& ChunkId) const;
-	FIoRequestImpl* GetCompletedRequests() override;
+	FIoRequestImpl* GetCompletedIoRequests() override;
 	TIoStatusOr<FIoMappedRegion> OpenMapped(const FIoChunkId& ChunkId, const FIoReadOptions& Options) override;
 	void ReopenAllFileHandles();
 
@@ -157,6 +157,7 @@ private:
 		FFileIoStoreCompressedBlock* CompressedBlock;
 	};
 
+	bool Resolve(FIoRequestImpl* Request);
 	void OnNewPendingRequestsAdded();
 	void ReadBlocks(FFileIoStoreResolvedRequest& ResolvedRequest);
 	void FreeBuffer(FFileIoStoreBuffer& Buffer);
diff --git a/Engine/Source/Runtime/StorageServerClient/Private/StorageServerIoDispatcherBackend.cpp b/Engine/Source/Runtime/StorageServerClient/Private/StorageServerIoDispatcherBackend.cpp
--- a/Engine/Source/Runtime/StorageServerClient/Private/StorageServerIoDispatcherBackend.cpp
+++ b/Engine/Source/Runtime/StorageServerClient/Private/StorageServerIoDispatcherBackend.cpp
@@ -178,6 +178,17 @@ bool FStorageServerIoDispatcherBackend::Resolve(FIoRequestImpl* Request)
 	return true;
 }
 
+void FStorageServerIoDispatcherBackend::ResolveIoRequests(FIoRequestList Requests, FIoRequestList& OutUnresolved)
+{
+	while (FIoRequestImpl* Request = Requests.PopHead())
+	{
+		if (Resolve(Request) == false)
+		{
+			OutUnresolved.AddTail(Request);
+		}
+	}
+}
+
 bool FStorageServerIoDispatcherBackend::DoesChunkExist(const FIoChunkId& ChunkId) const
 {
 	return GetSizeForChunk(ChunkId).IsOk();
@@ -197,7 +208,7 @@ TIoStatusOr<uint64> FStorageServerIoDispatcherBackend::GetSizeForChunk(const FIo
 	}
 }
 
-FIoRequestImpl* FStorageServerIoDispatcherBackend::GetCompletedRequests()
+FIoRequestImpl* FStorageServerIoDispatcherBackend::GetCompletedIoRequests()
 {
 	FScopeLock Lock(&CompletedRequestsCritical);
 	FIoRequestImpl* Result = CompletedRequestsHead;
diff --git a/Engine/Source/Runtime/StorageServerClient/Private/StorageServerIoDispatcherBackend.h b/Engine/Source/Runtime/StorageServerClient/Private/StorageServerIoDispatcherBackend.h
--- a/Engine/Source/Runtime/StorageServerClient/Private/StorageServerIoDispatcherBackend.h
+++ b/Engine/Source/Runtime/StorageServerClient/Private/StorageServerIoDispatcherBackend.h
@@ -21,12 +21,12 @@ public:
 
 	void Initialize(TSharedRef<const FIoDispatcherBackendContext> Context) override;
 	void Shutdown() override;
-	bool Resolve(FIoRequestImpl* Request) override;
+	void ResolveIoRequests(FIoRequestList Requests, FIoRequestList& OutUnresolved) override;
 	void CancelIoRequest(FIoRequestImpl* Request) override {};
 	void UpdatePriorityForIoRequest(FIoRequestImpl* Request) override {};
 	bool DoesChunkExist(const FIoChunkId& ChunkId) const override;
 	TIoStatusOr<uint64> GetSizeForChunk(const FIoChunkId& ChunkId) const override;
-	FIoRequestImpl* GetCompletedRequests() override;
+	FIoRequestImpl* GetCompletedIoRequests() override;
 	TIoStatusOr<FIoMappedRegion> OpenMapped(const FIoChunkId& ChunkId, const FIoReadOptions& Options) override
 	{
 		return FIoStatus(EIoErrorCode::NotFound);
@@ -79,6 +79,7 @@ private:
 		TUniquePtr<FStorageServerSerializationContext> SerializationContext;
 	};
 
+	bool Resolve(FIoRequestImpl* Request);
 	void SubmitBatch(FBatch* Batch);
 	void OnBatchCompleted(FBatch* Batch);
 	bool WaitForBatchToComplete(uint32 WaitTime = MAX_uint32);
