ConfigBranch: Fix Error: Error deleting file '../../../Engine/Config/ConsoleVariables.ini', which occurred on files added through FConfigCacheIni::Find. Mark them as NoSave.
#rnx

diff --git a/Engine/Source/Runtime/Core/Private/Misc/ConfigCacheIni.cpp b/Engine/Source/Runtime/Core/Private/Misc/ConfigCacheIni.cpp
--- a/Engine/Source/Runtime/Core/Private/Misc/ConfigCacheIni.cpp
+++ b/Engine/Source/Runtime/Core/Private/Misc/ConfigCacheIni.cpp
@@ -3094,6 +3094,9 @@ FConfigFile* FConfigCacheIni::Find(const FString& Filename)
 			if (DoesConfigFileExistWrapper(*UnrealFileName))
 			{
 				Result->Read(UnrealFileName);
+				// Files added through Find are treated the same as ReadSingleIntoConfigSystem contexts,
+				// and do not use a hierarchy so they do not use a generatedini and should never be saved.
+				Result->NoSave = true;
 				UE_LOG(LogConfig, Verbose, TEXT("GConfig::Find has loaded file:  %s"), *UnrealFileName);
 			}
 		}
@@ -3211,7 +3214,7 @@ void FConfigCacheIni::Flush(bool bRemoveFromCache, const FString& Filename )
 			{
 				// flush all files
 				for (TPair<FString, FConfigBranch*>& Pair : OtherFiles)
-			{
+				{
 					SaveBranch(*Pair.Value);
 				}
 				for (FConfigBranch& Branch : KnownFiles.Branches)
diff --git a/Engine/Source/Runtime/Core/Private/Misc/ConfigContext.cpp b/Engine/Source/Runtime/Core/Private/Misc/ConfigContext.cpp
--- a/Engine/Source/Runtime/Core/Private/Misc/ConfigContext.cpp
+++ b/Engine/Source/Runtime/Core/Private/Misc/ConfigContext.cpp
@@ -588,7 +588,7 @@ static FString PerformBasicReplacements(const FString& InString, const TCHAR* Ba
 
 static FString PerformExpansionReplacements(const FConfigLayerExpansion& Expansion, const FString& InString)
 {
-	// if there's replacement to do, the output is just the output
+	// if there's no replacement to do, the output is just the input
 	if (Expansion.Before1 == nullptr)
 	{
 		return InString;
@@ -726,11 +726,6 @@ void FConfigContext::AddStaticLayersToHierarchy(TArray<FString>* GatheredLayerFi
 		}
 	}
 	
-	if (OverrideLayers.Num() > 0)
-	{
-		Layers = OverrideLayers.GetData();
-		NumLayers = OverrideLayers.Num();
-	}
 	// let the context override the layers if needed
 	if (OverrideLayers.Num() > 0)
 	{
