- Fixed a config bug where loading into an existing FConfigFile was always returning false
- Changed some post-load fixup to work with existing single files, instead of only Branch files

diff --git a/Engine/Source/Runtime/Core/Private/Misc/ConfigContext.cpp b/Engine/Source/Runtime/Core/Private/Misc/ConfigContext.cpp
--- a/Engine/Source/Runtime/Core/Private/Misc/ConfigContext.cpp
+++ b/Engine/Source/Runtime/Core/Private/Misc/ConfigContext.cpp
@@ -505,18 +505,20 @@ bool FConfigContext::PerformLoad()
 	// now generate and make sure it's up to date (using IniName as a Base for an ini filename)
 	// @todo This bNeedsWrite afaict is always true even if it loaded a completely valid generated/final .ini, and the write below will
 	// just write out the exact same thing it read in!
-	bool bNeedsWrite = GenerateDestIniFile();
+	bool bGeneratedFile = GenerateDestIniFile();
 
-	Branch->InMemoryFile.Name = FName(*BaseIniName);
-	Branch->InMemoryFile.PlatformName = Platform;
-	Branch->InMemoryFile.bHasPlatformName = true;
+	FConfigFile& FinalFile = ExistingFile ? *ExistingFile : Branch->InMemoryFile;
+
+	FinalFile.Name = FName(*BaseIniName);
+	FinalFile.PlatformName = Platform;
+	FinalFile.bHasPlatformName = true;
 
 	// chcek if the config file wants to save all sections
 	bool bLocalSaveAllSections = false;
 	// Do not report the read of SectionsToSave. Some ConfigFiles are reallocated without it, and reporting
 	// logs that the section disappeared. But this log is spurious since if the only reason it was read was
 	// for the internal save before the FConfigFile is made publicly available.
-	const FConfigSection* SectionsToSaveSection = Branch->InMemoryFile.FindSection(SectionsToSaveString);
+	const FConfigSection* SectionsToSaveSection = FinalFile.FindSection(SectionsToSaveString);
 	if (SectionsToSaveSection)
 	{
 		const FConfigValue* Value = SectionsToSaveSection->Find(SaveAllSectionsKey);
@@ -527,13 +529,11 @@ bool FConfigContext::PerformLoad()
 		}
 	}
 
-	Branch->InMemoryFile.GetBool(SectionsToSaveString, SaveAllSectionsKey, bLocalSaveAllSections);
-
 	// we can always save all sections of a User config file, Editor* (not Editor.ini tho, that is already handled in the normal method)
 	bool bIsUserFile = BaseIniName.Contains(TEXT("User"));
 	bool bIsEditorSettingsFile = BaseIniName.Contains(TEXT("Editor")) && BaseIniName != TEXT("Editor");
 
-	Branch->InMemoryFile.bCanSaveAllSections = bLocalSaveAllSections || bIsUserFile || bIsEditorSettingsFile;
+	FinalFile.bCanSaveAllSections = bLocalSaveAllSections || bIsUserFile || bIsEditorSettingsFile;
 
 	// don't write anything to disk in cooked builds - we will always use re-generated INI files anyway.
 	// Note: Unfortunately bAllowGeneratedIniWhenCooked is often true even in shipping builds with cooked data
@@ -549,17 +549,17 @@ bool FConfigContext::PerformLoad()
 		// Check the config system for any changes made to defaults and propagate through to the saved.
 		Branch->InMemoryFile.ProcessSourceAndCheckAgainstBackup();
 
-		if (bNeedsWrite)
+		// don't write anything out if we are reading into an existing file
+		if (bGeneratedFile && ExistingFile == nullptr)
 		{
 			// if it was dirtied during the above function, save it out now
-			Branch->InMemoryFile.Write(DestIniFilename);
+			FinalFile.Write(DestIniFilename);
 		}
 	}
 
 	if (IsInGameThread()) GPerformLoadTime += FPlatformTime::Seconds();
 
-	// GenerateDestIniFile returns true if nothing is loaded, so check if we actually loaded something
-	return Branch->InMemoryFile.Num() > 0;
+	return bGeneratedFile;
 }
 
 
@@ -995,96 +995,94 @@ bool FConfigContext::GenerateDestIniFile()
 		}
 	}
 	
-	// if we loaded nothing at all, we are done
-	if (Branch->InMemoryFile.Num() == 0)
+	// skip over code that doesn't apply when reading into an ExistingFile
+	if (ExistingFile == nullptr && Branch->InMemoryFile.Num() > 0)
 	{
-		return false;
-	}
-	
-	bool bForceRegenerate = false;
+		bool bForceRegenerate = false;
 
-	// New versioning
-	int32 SourceConfigVersionNum = -1;
-	int32 CurrentIniVersion = -1;
-	bool bVersionChanged = false;
+		// New versioning
+		int32 SourceConfigVersionNum = -1;
+		int32 CurrentIniVersion = -1;
+		bool bVersionChanged = false;
 
-	// don't do version checking if we have nothing saved
-	if (Branch->SavedLayer.Num() > 0)
-	{
-		// get the version that was last saved, if any
-		FConfigCommandStreamSection* VersionSection = Branch->SavedLayer.Find(CurrentIniVersionString);
-		if (VersionSection)
+		// don't do version checking if we have nothing saved
+		if (Branch->SavedLayer.Num() > 0)
 		{
-			FConfigValue* VersionKey = VersionSection->Find(*VersionName);
-			if (VersionKey)
+			// get the version that was last saved, if any
+			FConfigCommandStreamSection* VersionSection = Branch->SavedLayer.Find(CurrentIniVersionString);
+			if (VersionSection)
 			{
-				TTypeFromString<int32>::FromString(CurrentIniVersion, *VersionKey->GetValue());
+				FConfigValue* VersionKey = VersionSection->Find(*VersionName);
+				if (VersionKey)
+				{
+					TTypeFromString<int32>::FromString(CurrentIniVersion, *VersionKey->GetValue());
+				}
 			}
-		}
-		
-		// now compare to the source config file
-		if (Branch->CombinedStaticLayers.GetInt(*CurrentIniVersionString, *VersionName, SourceConfigVersionNum))
-		{
-			if (SourceConfigVersionNum > CurrentIniVersion)
+
+			// now compare to the source config file
+			if (Branch->CombinedStaticLayers.GetInt(*CurrentIniVersionString, *VersionName, SourceConfigVersionNum))
 			{
-				UE_LOG(LogInit, Log, TEXT("%s version has been updated. It will be regenerated."), *FPaths::ConvertRelativePathToFull(DestIniFilename));
-				bVersionChanged = true;
+				if (SourceConfigVersionNum > CurrentIniVersion)
+				{
+					UE_LOG(LogInit, Log, TEXT("%s version has been updated. It will be regenerated."), *FPaths::ConvertRelativePathToFull(DestIniFilename));
+					bVersionChanged = true;
+				}
+				else if (SourceConfigVersionNum < CurrentIniVersion)
+				{
+					UE_LOG(LogInit, Warning, TEXT("%s version is later than the source. Since the versions are out of sync, nothing will be done."), *FPaths::ConvertRelativePathToFull(DestIniFilename));
+				}
 			}
-			else if (SourceConfigVersionNum < CurrentIniVersion)
+
+			// Regenerate the ini file?
+			if (FParse::Param(FCommandLine::Get(), TEXT("REGENERATEINIS")) == true)
 			{
-				UE_LOG(LogInit, Warning, TEXT("%s version is later than the source. Since the versions are out of sync, nothing will be done."), *FPaths::ConvertRelativePathToFull(DestIniFilename));
+				bForceRegenerate = true;
 			}
 		}
 
-		// Regenerate the ini file?
-		if (FParse::Param(FCommandLine::Get(), TEXT("REGENERATEINIS")) == true)
+		// Order is important, we want to let force regenerate happen before version change, in case we're trying to wipe everything.
+		//	Version tries to save some info.
+		if (bForceRegenerate)
 		{
-			bForceRegenerate = true;
+			Branch->SavedLayer.Empty();
 		}
-	}
-
-	// Order is important, we want to let force regenerate happen before version change, in case we're trying to wipe everything.
-	//	Version tries to save some info.
-	if (bForceRegenerate)
-	{
-		Branch->SavedLayer.Empty();
-	}
-	else if (bVersionChanged)
-	{
-		// get list of preserved sections (those we want to keep from the Saved, even if the version changed)
-		TArray<FString> PreservedSections;
-		Branch->InMemoryFile.GetArray(*CurrentIniVersionString, *PreserveName, PreservedSections);
-
-		// get the saved keys, and remove non-preserved ones
-		TSet<FString> SavedKeys;
-		Branch->SavedLayer.GetKeys(SavedKeys);
-		for (const FString& SavedSection : SavedKeys)
+		else if (bVersionChanged)
 		{
-			if (!PreservedSections.Contains(SavedSection))
+			// get list of preserved sections (those we want to keep from the Saved, even if the version changed)
+			TArray<FString> PreservedSections;
+			Branch->InMemoryFile.GetArray(*CurrentIniVersionString, *PreserveName, PreservedSections);
+
+			// get the saved keys, and remove non-preserved ones
+			TSet<FString> SavedKeys;
+			Branch->SavedLayer.GetKeys(SavedKeys);
+			for (const FString& SavedSection : SavedKeys)
 			{
-				Branch->SavedLayer.Remove(*SavedSection);
+				if (!PreservedSections.Contains(SavedSection))
+				{
+					Branch->SavedLayer.Remove(*SavedSection);
+				}
 			}
+
+			// make sure current version is saved out (this would only be needed if we preserved the CurrentIniVersionString section, but doesn't hurt to do)
+			Branch->SavedLayer.FindOrAdd(CurrentIniVersionString).Remove(*VersionName);
+			Branch->SavedLayer.FindOrAdd(CurrentIniVersionString).Add(*VersionName, FConfigValue(FString::Printf(TEXT("%d"), SourceConfigVersionNum), FConfigValue::EValueType::Set));
 		}
-		
-		// make sure current version is saved out (this would only be needed if we preserved the CurrentIniVersionString section, but doesn't hurt to do)
-		Branch->SavedLayer.FindOrAdd(CurrentIniVersionString).Remove(*VersionName);
-		Branch->SavedLayer.FindOrAdd(CurrentIniVersionString).Add(*VersionName, FConfigValue(FString::Printf(TEXT("%d"), SourceConfigVersionNum), FConfigValue::EValueType::Set));
-	}
 
-	// now merge in the saved info that is still around after the above logic
-	Branch->InMemoryFile.ApplyFile(&Branch->SavedLayer);
+		// now merge in the saved info that is still around after the above logic
+		Branch->InMemoryFile.ApplyFile(&Branch->SavedLayer);
 
 #if ALLOW_INI_OVERRIDE_FROM_COMMANDLINE
-	// process any commandline overrides
-	FConfigFile::OverrideFromCommandline(&Branch->CommandLineOverrides, BaseIniName);
-	// and push it into the current values
-	Branch->InMemoryFile.ApplyFile(&Branch->CommandLineOverrides);
+		// process any commandline overrides
+		FConfigFile::OverrideFromCommandline(&Branch->CommandLineOverrides, BaseIniName);
+		// and push it into the current values
+		Branch->InMemoryFile.ApplyFile(&Branch->CommandLineOverrides);
 #endif
 
+		//	Branch->CombinedStaticLayers.Cleanup();
+	}
 
-//	Branch->CombinedStaticLayers.Cleanup();
-
-	return bResult;
+	// return true if we actually read anything in
+	return Branch->InMemoryFile.Num() > 0 || (ExistingFile != nullptr && ExistingFile->Num() > 0);
 }
 
 
