Fixed csv metadata for NamedEvents not being set if namedevents was enabled after the csv capture starts.

#tests Win64 test client csv start/stop with stat namedevents/stop with stat namedevents

diff --git a/Engine/Source/Runtime/Core/Private/ProfilingDebugging/CsvProfiler.cpp b/Engine/Source/Runtime/Core/Private/ProfilingDebugging/CsvProfiler.cpp
--- a/Engine/Source/Runtime/Core/Private/ProfilingDebugging/CsvProfiler.cpp
+++ b/Engine/Source/Runtime/Core/Private/ProfilingDebugging/CsvProfiler.cpp
@@ -2839,6 +2839,7 @@ FCsvProfiler::FCsvProfiler()
 	, CaptureFrameNumberRT(0)
 	, CaptureOnEventFrameCount(-1)
 	, bInsertEndFrameAtFrameStart(false)
+	, bNamedEventsWasEnabled(false)
 	, LastEndFrameTimestamp(0)
 	, CaptureEndFrameCount(0)
 	, ProcessingThread(nullptr)
@@ -3061,7 +3062,8 @@ void FCsvProfiler::BeginFrame()
 
 					SetMetadataInternal(TEXT("TargetFramerate"), *FString::FromInt(TargetFPS));
 					SetMetadataInternal(TEXT("StartTimestamp"), *FString::Printf(TEXT("%lld"), FDateTime::UtcNow().ToUnixTimestamp()));
-					SetMetadataInternal(TEXT("NamedEvents"), GCycleStatsShouldEmitNamedEvents ? TEXT("1") : TEXT("0"));
+					SetMetadataInternal(TEXT("NamedEvents"), (GCycleStatsShouldEmitNamedEvents > 0) ? TEXT("1") : TEXT("0"));
+					bNamedEventsWasEnabled = (GCycleStatsShouldEmitNamedEvents > 0);
 
 					GCsvStatCounts = !!CVarCsvStatCounts.GetValueOnGameThread();
 
@@ -3081,6 +3083,12 @@ void FCsvProfiler::BeginFrame()
 			{
 				OnCSVProfileFirstFrameDelegate.Broadcast();
 			}
+
+			if (!bNamedEventsWasEnabled && (GCycleStatsShouldEmitNamedEvents > 0))
+			{
+				bNamedEventsWasEnabled = true;
+				SetMetadataInternal(TEXT("NamedEvents"), TEXT("1"));
+			}
 		}
 	}
 
diff --git a/Engine/Source/Runtime/Core/Public/ProfilingDebugging/CsvProfiler.h b/Engine/Source/Runtime/Core/Public/ProfilingDebugging/CsvProfiler.h
--- a/Engine/Source/Runtime/Core/Public/ProfilingDebugging/CsvProfiler.h
+++ b/Engine/Source/Runtime/Core/Public/ProfilingDebugging/CsvProfiler.h
@@ -499,6 +499,7 @@ private:
 	int32 CaptureOnEventFrameCount;
 
 	bool bInsertEndFrameAtFrameStart;
+	bool bNamedEventsWasEnabled;
 
 	uint64 LastEndFrameTimestamp;
 	uint32 CaptureEndFrameCount;
