Initialize FInternationalization atomically via TLazySingleton static initialization

This required fixing a cycle where FInternationalization::Get() was being called during FInternationalization::Initialize

#rnx

diff --git a/Engine/Source/Runtime/Core/Private/Internationalization/Culture.cpp b/Engine/Source/Runtime/Core/Private/Internationalization/Culture.cpp
--- a/Engine/Source/Runtime/Core/Private/Internationalization/Culture.cpp
+++ b/Engine/Source/Runtime/Core/Private/Internationalization/Culture.cpp
@@ -199,7 +199,7 @@ FString FCulture::CreateCultureName(const FString& LanguageCode, const FString&
 
 FString FCulture::GetCanonicalName(const FString& Name)
 {
-	return FCultureImplementation::GetCanonicalName(Name);
+	return FCultureImplementation::GetCanonicalName(Name, FInternationalization::Get());
 }
 
 const FString& FCulture::GetName() const
diff --git a/Engine/Source/Runtime/Core/Private/Internationalization/ICUCulture.cpp b/Engine/Source/Runtime/Core/Private/Internationalization/ICUCulture.cpp
--- a/Engine/Source/Runtime/Core/Private/Internationalization/ICUCulture.cpp
+++ b/Engine/Source/Runtime/Core/Private/Internationalization/ICUCulture.cpp
@@ -245,9 +245,9 @@ int FICUCultureImplementation::GetLCID() const
 	return ICULocale.getLCID();
 }
 
-FString FICUCultureImplementation::GetCanonicalName(const FString& Name)
+FString FICUCultureImplementation::GetCanonicalName(const FString& Name, FInternationalization& I18N)
 {
-	return ICUUtilities::GetCanonicalCultureName(Name, TEXT("en-US-POSIX"));
+	return ICUUtilities::GetCanonicalCultureName(Name, TEXT("en-US-POSIX"), I18N);
 }
 
 FString FICUCultureImplementation::GetName() const
diff --git a/Engine/Source/Runtime/Core/Private/Internationalization/ICUCulture.h b/Engine/Source/Runtime/Core/Private/Internationalization/ICUCulture.h
--- a/Engine/Source/Runtime/Core/Private/Internationalization/ICUCulture.h
+++ b/Engine/Source/Runtime/Core/Private/Internationalization/ICUCulture.h
@@ -28,6 +28,8 @@ THIRD_PARTY_INCLUDES_END
 
 struct FDecimalNumberFormattingRules;
 
+class FInternationalization;
+
 inline UColAttributeValue UEToICU(const ETextComparisonLevel::Type ComparisonLevel)
 {
 	UColAttributeValue Value;
@@ -194,7 +196,7 @@ public:
 	virtual ETextPluralForm GetPluralForm(double Val, const ETextPluralType PluralType) const override;
 	virtual const TArray<ETextPluralForm>& GetValidPluralForms(const ETextPluralType PluralType) const override;
 
-	static FString GetCanonicalName(const FString& Name);
+	static FString GetCanonicalName(const FString& Name, FInternationalization& I18N);
 
 	TSharedRef<const icu::BreakIterator> GetBreakIterator(const EBreakIteratorType Type);
 	TSharedRef<const icu::Collator, ESPMode::ThreadSafe> GetCollator(const ETextComparisonLevel::Type ComparisonLevel);
diff --git a/Engine/Source/Runtime/Core/Private/Internationalization/ICUInternationalization.cpp b/Engine/Source/Runtime/Core/Private/Internationalization/ICUInternationalization.cpp
--- a/Engine/Source/Runtime/Core/Private/Internationalization/ICUInternationalization.cpp
+++ b/Engine/Source/Runtime/Core/Private/Internationalization/ICUInternationalization.cpp
@@ -568,7 +568,7 @@ TArray<FString> FICUInternationalization::GetPrioritizedCultureNames(const FStri
 	};
 
 	// Apply any culture remapping
-	FString GivenCulture = FCulture::GetCanonicalName(Name);
+	FString GivenCulture = FCultureImplementation::GetCanonicalName(Name, *I18N);
 	IsCultureRemapped(Name, &GivenCulture);
 
 	TArray<FString> PrioritizedCultureNames;
@@ -653,7 +653,7 @@ FCulturePtr FICUInternationalization::GetCulture(const FString& Name)
 
 FCulturePtr FICUInternationalization::FindOrMakeCulture(const FString& Name, const EAllowDefaultCultureFallback AllowDefaultFallback)
 {
-	return FindOrMakeCanonizedCulture(FCulture::GetCanonicalName(Name), AllowDefaultFallback);
+	return FindOrMakeCanonizedCulture(FCultureImplementation::GetCanonicalName(Name, *I18N), AllowDefaultFallback);
 }
 
 FCulturePtr FICUInternationalization::FindOrMakeCanonizedCulture(const FString& Name, const EAllowDefaultCultureFallback AllowDefaultFallback)
diff --git a/Engine/Source/Runtime/Core/Private/Internationalization/Internationalization.cpp b/Engine/Source/Runtime/Core/Private/Internationalization/Internationalization.cpp
--- a/Engine/Source/Runtime/Core/Private/Internationalization/Internationalization.cpp
+++ b/Engine/Source/Runtime/Core/Private/Internationalization/Internationalization.cpp
@@ -21,9 +21,7 @@
 
 FInternationalization& FInternationalization::Get()
 {
-	FInternationalization& Singleton = TLazySingleton<FInternationalization>::Get();
-	Singleton.Initialize();
-	return Singleton;
+	return TLazySingleton<FInternationalization>::Get();
 }
 
 bool FInternationalization::IsAvailable()
@@ -517,6 +515,7 @@ TArray<FCultureRef> FInternationalization::GetAvailableCultures(const TArray<FSt
 FInternationalization::FInternationalization()
 	:	Implementation(this)
 {
+	Initialize();
 }
 
 FInternationalization::~FInternationalization()
diff --git a/Engine/Source/Runtime/Core/Private/Internationalization/InternationalizationUtilities.cpp b/Engine/Source/Runtime/Core/Private/Internationalization/InternationalizationUtilities.cpp
--- a/Engine/Source/Runtime/Core/Private/Internationalization/InternationalizationUtilities.cpp
+++ b/Engine/Source/Runtime/Core/Private/Internationalization/InternationalizationUtilities.cpp
@@ -86,7 +86,7 @@ namespace InternationalizationUtilities
 		return SanitizedCurrencyCode;
 	}
 
-	FString GetCanonicalCultureName(const FString& Name, const FString& FallbackCulture)
+	FString GetCanonicalCultureName(const FString& Name, const FString& FallbackCulture, FInternationalization& I18N)
 	{
 		auto IsLanguageCode = [](const FString& InCode)
 		{
@@ -225,7 +225,7 @@ namespace InternationalizationUtilities
 		const FString SanitizedName = SanitizeCultureCode(Name);
 
 		// If the name matches a custom culture, then just accept it as-is
-		if (FInternationalization::Get().GetCustomCulture(SanitizedName))
+		if (I18N.GetCustomCulture(SanitizedName))
 		{
 			return SanitizedName;
 		}
diff --git a/Engine/Source/Runtime/Core/Private/Internationalization/InternationalizationUtilities.h b/Engine/Source/Runtime/Core/Private/Internationalization/InternationalizationUtilities.h
--- a/Engine/Source/Runtime/Core/Private/Internationalization/InternationalizationUtilities.h
+++ b/Engine/Source/Runtime/Core/Private/Internationalization/InternationalizationUtilities.h
@@ -5,6 +5,8 @@
 #include "CoreTypes.h"
 #include "Containers/UnrealString.h"
 
+class FInternationalization;
+
 namespace InternationalizationUtilities
 {
 	/** Sanitize the given culture code so that it is safe to use */
@@ -23,5 +25,5 @@ namespace InternationalizationUtilities
 	};
 
 	/** Get the canonical version of the given culture code (will also sanitize it) */
-	FString GetCanonicalCultureName(const FString& Name, const FString& FallbackCulture);
+	FString GetCanonicalCultureName(const FString& Name, const FString& FallbackCulture, FInternationalization& I18N);
 }
diff --git a/Engine/Source/Runtime/Core/Private/Internationalization/LegacyCulture.cpp b/Engine/Source/Runtime/Core/Private/Internationalization/LegacyCulture.cpp
--- a/Engine/Source/Runtime/Core/Private/Internationalization/LegacyCulture.cpp
+++ b/Engine/Source/Runtime/Core/Private/Internationalization/LegacyCulture.cpp
@@ -63,11 +63,11 @@ FString FLegacyCultureImplementation::GetName() const
 	return Name;
 }
 
-FString FLegacyCultureImplementation::GetCanonicalName(const FString& Name)
+FString FLegacyCultureImplementation::GetCanonicalName(const FString& Name, FInternationalization& I18N)
 {
 	if (!Name.IsEmpty())
 	{
-		return InternationalizationUtilities::GetCanonicalCultureName(Name, FString());
+		return InternationalizationUtilities::GetCanonicalCultureName(Name, FString(), I18N);
 	}
 	return Name;
 }
diff --git a/Engine/Source/Runtime/Core/Private/Internationalization/LegacyCulture.h b/Engine/Source/Runtime/Core/Private/Internationalization/LegacyCulture.h
--- a/Engine/Source/Runtime/Core/Private/Internationalization/LegacyCulture.h
+++ b/Engine/Source/Runtime/Core/Private/Internationalization/LegacyCulture.h
@@ -10,6 +10,8 @@
 #include "Internationalization/FastDecimalFormat.h"
 
 #if !UE_ENABLE_ICU
+class FInternationalization;
+
 class FLegacyCultureImplementation : public ICultureImplementation
 {
 	friend class FCulture;
@@ -55,7 +57,7 @@ public:
 	virtual ETextPluralForm GetPluralForm(double Val, const ETextPluralType PluralType) const override;
 	virtual const TArray<ETextPluralForm>& GetValidPluralForms(const ETextPluralType PluralType) const override;
 
-	static FString GetCanonicalName(const FString& Name);
+	static FString GetCanonicalName(const FString& Name, FInternationalization& I18N);
 
 private:
 	// Full localized culture name
diff --git a/Engine/Source/Runtime/Core/Private/Internationalization/LegacyInternationalization.cpp b/Engine/Source/Runtime/Core/Private/Internationalization/LegacyInternationalization.cpp
--- a/Engine/Source/Runtime/Core/Private/Internationalization/LegacyInternationalization.cpp
+++ b/Engine/Source/Runtime/Core/Private/Internationalization/LegacyInternationalization.cpp
@@ -72,7 +72,7 @@ void FLegacyInternationalization::GetCultureNames(TArray<FString>& CultureNames)
 
 TArray<FString> FLegacyInternationalization::GetPrioritizedCultureNames(const FString& Name)
 {
-	const FString CanonicalName = FCulture::GetCanonicalName(Name);
+	const FString CanonicalName = FCultureImplementation::GetCanonicalName(Name, *I18N);
 	TArray<FString> PrioritizedCultureNames;
 	
 	if (!CanonicalName.IsEmpty())
