Restore the Windows dll directory cache optimization that was added in cl 23269675 but has been broken since cl 23574121.

diff --git a/Engine/Source/Runtime/Core/Private/Windows/WindowsPlatformProcess.cpp b/Engine/Source/Runtime/Core/Private/Windows/WindowsPlatformProcess.cpp
--- a/Engine/Source/Runtime/Core/Private/Windows/WindowsPlatformProcess.cpp
+++ b/Engine/Source/Runtime/Core/Private/Windows/WindowsPlatformProcess.cpp
@@ -106,11 +106,10 @@ void FWindowsPlatformProcess::AddDllDirectory(const TCHAR* Directory)
 		// enumerate the dir and cache all the dlls
 		{
 			TArray<FString> FoundDllFileNames;
-			IPlatformFile::GetPlatformPhysical().FindFiles(FoundDllFileNames, *NormalizedDirectory, TEXT("*.dll"));
-			for (const FString& DllFileName : FoundDllFileNames)
+			IPlatformFile::GetPlatformPhysical().FindFiles(FoundDllFileNames, *NormalizedDirectory, TEXT(".dll"));
+			for (FString& DllPath : FoundDllFileNames)
 			{
-				TArray<FString>& Paths = SearchPathDllCache.FindOrAdd(*DllFileName);
-				FString DllPath(NormalizedDirectory / DllFileName);
+				TArray<FString>& Paths = SearchPathDllCache.FindOrAdd(FName(FPathViews::GetCleanFilename(DllPath)));
 				FPaths::NormalizeDirectoryName(DllPath);
 				Paths.Add(DllPath);
 			}
