Check for read only channels when enabling channels

Some channels are marked read-only (e.g. memory tracing channels) and cannot be toggled while the application is running. However there was no feedback when users tried to enable those channels using console commands for example or using the public API. With this change a logging channel argument is added to several internal methods in order to output information correctly, to console feedback when needed.

diff --git a/Engine/Source/Runtime/Core/Private/ProfilingDebugging/TraceAuxiliary.cpp b/Engine/Source/Runtime/Core/Private/ProfilingDebugging/TraceAuxiliary.cpp
--- a/Engine/Source/Runtime/Core/Private/ProfilingDebugging/TraceAuxiliary.cpp
+++ b/Engine/Source/Runtime/Core/Private/ProfilingDebugging/TraceAuxiliary.cpp
@@ -106,10 +106,12 @@ public:
 	void AddCommandlineChannels(const TCHAR* ChannelList);
 	void ResetCommandlineChannels();
 	bool HasCommandlineChannels() const { return !CommandlineChannels.IsEmpty(); }
-	void EnableChannels(const TCHAR* ChannelList);
-	void DisableChannels(const TCHAR* ChannelList);
+	void EnableChannels(const TCHAR* ChannelList, const FTraceAuxiliary::FLogCategoryAlias& LogCategory);
+	void DisableChannels(const TCHAR* ChannelList, const FTraceAuxiliary::FLogCategoryAlias& LogCategory);
 	bool Connect(FTraceAuxiliary::EConnectionType Type, const TCHAR* Parameter, const FTraceAuxiliary::FLogCategoryAlias& LogCategory, uint16 SendFlags);
 	bool Stop();
+	void FreezeReadOnlyChannels();
+	static bool IsReadOnlyChannel(const TCHAR* Channel);
 	void ResumeChannels();
 	void PauseChannels();
 	bool IsPaused();
@@ -138,12 +140,12 @@ private:
 		bool bActive = false;
 	};
 
-	void AddChannel(const TCHAR* Name);
-	void RemoveChannel(const TCHAR* Name);
-	template <class T> void ForEachChannel(const TCHAR* ChannelList, bool bResolvePresets, T Callable);
+	void AddChannel(const TCHAR* Name, const FTraceAuxiliary::FLogCategoryAlias& LogCategory);
+	void RemoveChannel(const TCHAR* Name, const FTraceAuxiliary::FLogCategoryAlias& LogCategory);
+	template <class T> void ForEachChannel(const TCHAR* ChannelList, bool bResolvePresets, const FTraceAuxiliary::FLogCategoryAlias& LogCategory, T Callable);
 	static uint32 HashChannelName(const TCHAR* Name);
-	bool EnableChannel(const TCHAR* Channel);
-	void DisableChannel(const TCHAR* Channel);
+	bool EnableChannel(const TCHAR* Channel, const FTraceAuxiliary::FLogCategoryAlias& LogCategory);
+	void DisableChannel(const TCHAR* Channel, const FTraceAuxiliary::FLogCategoryAlias& LogCategory);
 	bool SendToHost(const TCHAR* Host, const FTraceAuxiliary::FLogCategoryAlias& LogCategory, uint16 SendFlags);
 	bool WriteToFile(const TCHAR* Path, const FTraceAuxiliary::FLogCategoryAlias& LogCategory, uint16 SendFlags);
 	bool FinalizeFilePath(const TCHAR* InPath, FString& OutPath, const FTraceAuxiliary::FLogCategoryAlias& LogCategory);
@@ -152,6 +154,7 @@ private:
 	ChannelSet CommandlineChannels;
 	bool bWorkerThreadStarted = false;
 	bool bTruncateFile = false;
+	bool bReadOnlyChannelsFrozen = false;
 	FString PausedPreset;
 
 	struct FCurrentTraceTarget
@@ -180,7 +183,7 @@ bool FTraceAuxiliaryImpl::IsParentProcessAndPreFork()
 ////////////////////////////////////////////////////////////////////////////////
 void FTraceAuxiliaryImpl::AddCommandlineChannels(const TCHAR* ChannelList)
 {
-	ForEachChannel(ChannelList, true, &FTraceAuxiliaryImpl::AddChannel);
+	ForEachChannel(ChannelList, true, LogTrace, &FTraceAuxiliaryImpl::AddChannel);
 }
 
 ////////////////////////////////////////////////////////////////////////////////
@@ -190,36 +193,36 @@ void FTraceAuxiliaryImpl::ResetCommandlineChannels()
 }
 
 ////////////////////////////////////////////////////////////////////////////////
-void FTraceAuxiliaryImpl::EnableChannels(const TCHAR* ChannelList)
+void FTraceAuxiliaryImpl::EnableChannels(const TCHAR* ChannelList, const FTraceAuxiliary::FLogCategoryAlias& LogCategory)
 {
 	if (ChannelList)
 	{
-		ForEachChannel(ChannelList, true, &FTraceAuxiliaryImpl::EnableChannel);
+		ForEachChannel(ChannelList, true, LogCategory, &FTraceAuxiliaryImpl::EnableChannel);
 	}
 }
 
 ////////////////////////////////////////////////////////////////////////////////
-void FTraceAuxiliaryImpl::DisableChannels(const TCHAR* ChannelList)
+void FTraceAuxiliaryImpl::DisableChannels(const TCHAR* ChannelList, const FTraceAuxiliary::FLogCategoryAlias& LogCategory)
 {
 	if (ChannelList)
 	{
-		ForEachChannel(ChannelList, true, &FTraceAuxiliaryImpl::DisableChannel);
+		ForEachChannel(ChannelList, true, LogCategory, &FTraceAuxiliaryImpl::DisableChannel);
 	}
 	else
 	{
 		// Disable all channels.
 		TStringBuilder<128> EnabledChannels;
 		GetActiveChannelsString(EnabledChannels);
-		ForEachChannel(EnabledChannels.ToString(), true, &FTraceAuxiliaryImpl::DisableChannel);
+		ForEachChannel(EnabledChannels.ToString(), true, LogCategory, &FTraceAuxiliaryImpl::DisableChannel);
 	}
 }
 
 ////////////////////////////////////////////////////////////////////////////////
 template<typename T>
-void FTraceAuxiliaryImpl::ForEachChannel(const TCHAR* ChannelList, bool bResolvePresets, T Callable)
+void FTraceAuxiliaryImpl::ForEachChannel(const TCHAR* ChannelList, bool bResolvePresets, const FTraceAuxiliary::FLogCategoryAlias& LogCategory, T Callable)
 {
 	check(ChannelList);
-	UE::String::ParseTokens(ChannelList, TEXT(","), [this, bResolvePresets, Callable] (const FStringView& Token)
+	UE::String::ParseTokens(ChannelList, TEXT(","), [this, bResolvePresets, Callable, &LogCategory] (const FStringView& Token)
 	{
 		TCHAR Name[80];
 		const size_t ChannelNameSize = Token.CopyString(Name, UE_ARRAY_COUNT(Name) - 1);
@@ -231,21 +234,21 @@ void FTraceAuxiliaryImpl::ForEachChannel(const TCHAR* ChannelList, bool bResolve
 			// Check against hard coded presets
 			if(FCString::Stricmp(Name, TEXT("default")) == 0)
 			{
-				ForEachChannel(GDefaultChannels, false, Callable);
+				ForEachChannel(GDefaultChannels, false, LogCategory, Callable);
 			}
 			else if(FCString::Stricmp(Name, TEXT("memory"))== 0)
 			{
-				ForEachChannel(GMemoryChannels, false, Callable);
+				ForEachChannel(GMemoryChannels, false, LogCategory, Callable);
 			}
 			// Check against data driven presets (if available)
 			else if (GConfig && GConfig->GetString(TEXT("Trace.ChannelPresets"), Name, Value, GEngineIni))
 			{
-				ForEachChannel(*Value, false, Callable);
+				ForEachChannel(*Value, false, LogCategory, Callable);
 				return;
 			}
 		}
 
-		Invoke(Callable, this, Name);
+		Invoke(Callable, this, Name, LogCategory);
 	});
 }
 
@@ -262,7 +265,7 @@ uint32 FTraceAuxiliaryImpl::HashChannelName(const TCHAR* Name)
 }
 
 ////////////////////////////////////////////////////////////////////////////////
-void FTraceAuxiliaryImpl::AddChannel(const TCHAR* Name)
+void FTraceAuxiliaryImpl::AddChannel(const TCHAR* Name, const FTraceAuxiliary::FLogCategoryAlias& LogCategory)
 {
 	uint32 Hash = HashChannelName(Name);
 
@@ -276,12 +279,12 @@ void FTraceAuxiliaryImpl::AddChannel(const TCHAR* Name)
 
 	if (IsConnected() && !Value.bActive)
 	{
-		Value.bActive = EnableChannel(*Value.Name);
+		Value.bActive = EnableChannel(*Value.Name, LogCategory);
 	}
 }
 
 ////////////////////////////////////////////////////////////////////////////////
-void FTraceAuxiliaryImpl::RemoveChannel(const TCHAR* Name)
+void FTraceAuxiliaryImpl::RemoveChannel(const TCHAR* Name, const FTraceAuxiliary::FLogCategoryAlias& LogCategory)
 {
 	uint32 Hash = HashChannelName(Name);
 
@@ -293,7 +296,7 @@ void FTraceAuxiliaryImpl::RemoveChannel(const TCHAR* Name)
 
 	if (IsConnected() && Channel.bActive)
 	{
-		DisableChannel(*Channel.Name);
+		DisableChannel(*Channel.Name, LogCategory);
 		Channel.bActive = false;
 	}
 }
@@ -386,7 +389,38 @@ bool FTraceAuxiliaryImpl::Stop()
 }
 
 ////////////////////////////////////////////////////////////////////////////////
-bool FTraceAuxiliaryImpl::EnableChannel(const TCHAR* Channel)
+void FTraceAuxiliaryImpl::FreezeReadOnlyChannels()
+{
+	bReadOnlyChannelsFrozen = true;
+}
+
+////////////////////////////////////////////////////////////////////////////////
+bool FTraceAuxiliaryImpl::IsReadOnlyChannel(const TCHAR* InChannel)
+{
+	const auto ChannelA = StringCast<ANSICHAR, 80>(InChannel);
+	struct FUserData
+	{
+		FAnsiStringView Channel;
+		bool bIsReadOnlyChannel = false;
+	} UserData = { FAnsiStringView(ChannelA.Get(), ChannelA.Length()), false};
+	
+	UE::Trace::EnumerateChannels([](const UE::Trace::FChannelInfo& Info, void* User)
+	{
+		FAnsiStringView NameView = FAnsiStringView(Info.Name).LeftChop(7); // Remove "Channel" suffix
+		FUserData* UserData = (FUserData*) User;
+		if (NameView.Equals(UserData->Channel, ESearchCase::IgnoreCase))
+		{
+			UserData->bIsReadOnlyChannel = Info.bIsReadOnly;
+			return false;
+		}
+		return true;
+	}, &UserData);
+
+	return UserData.bIsReadOnlyChannel;
+}
+
+////////////////////////////////////////////////////////////////////////////////
+bool FTraceAuxiliaryImpl::EnableChannel(const TCHAR* Channel, const FTraceAuxiliary::FLogCategoryAlias& LogCategory)
 {
 	// Channel names have been provided by the user and may not exist yet. As
 	// we want to maintain bActive accurately (channels toggles are reference
@@ -396,13 +430,26 @@ bool FTraceAuxiliaryImpl::EnableChannel(const TCHAR* Channel)
 		return false;
 	}
 
+	// It is not possible to change read only channels once trace is
+	// initialized.
+	if(bReadOnlyChannelsFrozen && IsReadOnlyChannel(Channel))
+	{
+		UE_LOG_REF(
+			LogCategory,
+			Error,
+			TEXT("Channel '%s' is a read only channel. It is only possible to enable on startup using command line parameters."),
+			Channel
+		);
+		return false;
+	}
+
 	FPlatformEventsTrace::OnTraceChannelUpdated(Channel, true);
 
 	return UE::Trace::ToggleChannel(Channel, true);
 }
 
 ////////////////////////////////////////////////////////////////////////////////
-void FTraceAuxiliaryImpl::DisableChannel(const TCHAR* Channel)
+void FTraceAuxiliaryImpl::DisableChannel(const TCHAR* Channel, const FTraceAuxiliary::FLogCategoryAlias& LogCategory)
 {
 	// Channel names have been provided by the user and may not exist yet. As
 	// we want to maintain bActive accurately (channels toggles are reference
@@ -411,6 +458,18 @@ void FTraceAuxiliaryImpl::DisableChannel(const TCHAR* Channel)
 	{
 		return;
 	}
+	
+	// Warn when disabling a read-only channel after init, it will not be possible
+	// to enable again.
+	if(bReadOnlyChannelsFrozen && IsReadOnlyChannel(Channel))
+	{
+		UE_LOG_REF(
+			LogCategory,
+			Warning,
+			TEXT("Channel '%s' is a read only channel. It has been disabled, but cannot be enabled again."),
+			Channel
+		);
+	}
 
 	FPlatformEventsTrace::OnTraceChannelUpdated(Channel, false);
 
@@ -421,7 +480,7 @@ void FTraceAuxiliaryImpl::DisableChannel(const TCHAR* Channel)
 void FTraceAuxiliaryImpl::ResumeChannels()
 {
 	// Enable channels from the "paused" preset.
-	ForEachChannel(*PausedPreset, false, &FTraceAuxiliaryImpl::EnableChannel);
+	ForEachChannel(*PausedPreset, false, LogTrace, &FTraceAuxiliaryImpl::EnableChannel);
 
 	PausedPreset.Empty();
 }
@@ -437,7 +496,7 @@ void FTraceAuxiliaryImpl::PauseChannels()
 	PausedPreset = EnabledChannels.ToString();
 
 	// Disable all "paused" channels.
-	ForEachChannel(*PausedPreset, true, &FTraceAuxiliaryImpl::DisableChannel);
+	ForEachChannel(*PausedPreset, true, LogTrace, &FTraceAuxiliaryImpl::DisableChannel);
 }
 
 ////////////////////////////////////////////////////////////////////////////////
@@ -458,7 +517,7 @@ void FTraceAuxiliaryImpl::EnableCommandlineChannels()
 	{
 		if (!ChannelPair.Value.bActive)
 		{
-			ChannelPair.Value.bActive = EnableChannel(*ChannelPair.Value.Name);
+			ChannelPair.Value.bActive = EnableChannel(*ChannelPair.Value.Name, LogTrace);
 		}
 	}
 }
@@ -469,7 +528,7 @@ void FTraceAuxiliaryImpl::EnableCommandlineChannelsPostInitialize()
 	for (auto& ChannelPair : CommandlineChannels)
 	{
 		// Intentionally enable channel without checking current state.
-		ChannelPair.Value.bActive = EnableChannel(*ChannelPair.Value.Name);
+		ChannelPair.Value.bActive = EnableChannel(*ChannelPair.Value.Name, LogTrace);
 	}
 }
 
@@ -792,6 +851,9 @@ static void SetupInitFromConfig(UE::Trace::FInitializeDesc& OutDesc)
 		return;
 	}
 
+	// Note that these options can only be used when tracing from a forked process (e.g. server).
+	// For a regular process use command line argument -TraceThreadSleepTime and -TraceTailSizeMb
+	
 	int32 SleepTimeConfig = 0;
 	if (GConfig->GetInt(GTraceConfigSection, TEXT("SleepTimeInMS"), SleepTimeConfig, GEngineIni))
 	{
@@ -991,7 +1053,7 @@ static void TraceAuxiliaryEnableChannels(const TArray<FString>& Args)
 		UE_LOG(LogConsoleResponse, Warning, TEXT("Need to provide at least one channel."));
 		return;
 	}
-	GTraceAuxiliary.EnableChannels(*Args[0]);
+	GTraceAuxiliary.EnableChannels(*Args[0], LogConsoleResponse);
 
 	TStringBuilder<128> EnabledChannels;
 	GTraceAuxiliary.GetActiveChannelsString(EnabledChannels);
@@ -1003,11 +1065,11 @@ static void TraceAuxiliaryDisableChannels(const TArray<FString>& Args)
 {
 	if (Args.Num() == 0)
 	{
-		GTraceAuxiliary.DisableChannels(nullptr);
+		GTraceAuxiliary.DisableChannels(nullptr, LogConsoleResponse);
 	}
 	else
 	{
-		GTraceAuxiliary.DisableChannels(*Args[0]);
+		GTraceAuxiliary.DisableChannels(*Args[0], LogConsoleResponse);
 	}
 
 	TStringBuilder<128> EnabledChannels;
@@ -1438,7 +1500,7 @@ void FTraceAuxiliary::Initialize(const TCHAR* CommandLine)
 	{
 		UE_LOG(LogTrace, Log, TEXT("Trace initialization skipped for parent process (pre fork)."));
 
-		GTraceAuxiliary.DisableChannels(nullptr);
+		GTraceAuxiliary.DisableChannels(nullptr, LogTrace);
 
 		// Set our post fork callback up and return - children will pass through and Initialize when they're created.
 		TraceAuxiliaryAddPostForkCallback(CommandLine);
@@ -1569,6 +1631,7 @@ void FTraceAuxiliary::Initialize(const TCHAR* CommandLine)
 
 	UE::Trace::ThreadRegister(TEXT("GameThread"), FPlatformTLS::GetCurrentThreadId(), -1);
 
+	GTraceAuxiliary.FreezeReadOnlyChannels();
 	UE_LOG(LogTrace, Log, TEXT("Finished trace initialization."));
 #endif //UE_TRACE_ENABLED
 }
@@ -1619,7 +1682,7 @@ void FTraceAuxiliary::EnableChannels()
 void FTraceAuxiliary::DisableChannels(const TCHAR* Channels)
 {
 #if UE_TRACE_ENABLED
-	GTraceAuxiliary.DisableChannels(Channels);
+	GTraceAuxiliary.DisableChannels(Channels, LogTrace);
 #endif
 }
 
