Added a workaround for a compiler bug in MSVC 17.9:

https://developercommunity.visualstudio.com/t/VS-[REDACTED_PHONE]-Preview-10-__builtin_arr/10519788

/predefined C++ types (compiler internal)(420): error C2248: 'FHazardPointerCollection::FHazardRecord::FHazardRecord': cannot access private member declared in class 'FHazardPointerCollection::FHazardRecord'

diff --git a/Engine/Source/Runtime/Core/Public/Experimental/Containers/HazardPointer.h b/Engine/Source/Runtime/Core/Public/Experimental/Containers/HazardPointer.h
--- a/Engine/Source/Runtime/Core/Public/Experimental/Containers/HazardPointer.h
+++ b/Engine/Source/Runtime/Core/Public/Experimental/Containers/HazardPointer.h
@@ -83,7 +83,15 @@ class FHazardPointerCollection
 		static constexpr uintptr_t FreeHazardEntry = ~uintptr_t(0);
 
 		std::atomic<uintptr_t> Hazard{ FreeHazardEntry };
+#ifdef _MSC_VER
+	// UE_DEPRECATED - this is a workaround for https://developercommunity.visualstudio.com/t/VS-2022-1790-Preview-10-__builtin_arr/10519788 and should be removed after
+	//                 after 17.9 is no longer supported.
+	public:
+		FHazardRecord() = default;
+	private:
+#else
 		FHazardRecord() = default;
+#endif
 
 		inline void* GetHazard() const
 		{
