ConfigAccessTracking: Fix various assignments of ConfigAccessTracking data in the new FConfigBranch architecture.
#rnx

diff --git a/Engine/Source/Runtime/Core/Private/Misc/ConfigCacheIni.cpp b/Engine/Source/Runtime/Core/Private/Misc/ConfigCacheIni.cpp
--- a/Engine/Source/Runtime/Core/Private/Misc/ConfigCacheIni.cpp
+++ b/Engine/Source/Runtime/Core/Private/Misc/ConfigCacheIni.cpp
@@ -856,6 +856,104 @@ FConfigFile::~FConfigFile()
 	Cleanup();
 }
 
+FConfigFile::FConfigFile(const FConfigFile& Other)
+{
+	*this = Other;
+}
+
+FConfigFile::FConfigFile(FConfigFile&& Other)
+{
+	*this = MoveTemp(Other);
+}
+
+FConfigFile& FConfigFile::operator=(const FConfigFile& Other)
+{
+	this->FConfigFileMap::operator=(Other);
+	Dirty = Other.Dirty;
+	NoSave = Other.NoSave;
+	bHasPlatformName = Other.bHasPlatformName;
+	bCanSaveAllSections = Other.bCanSaveAllSections;
+
+	// LoadType is not copied; each FConfigFile has to set it itself
+
+	Name = Other.Name;
+	PlatformName = Other.PlatformName;
+	Tag = Other.Tag;
+	Branch = Other.Branch;
+
+#if ALLOW_INI_OVERRIDE_FROM_COMMANDLINE
+	CommandlineOptions = Other.CommandlineOptions;
+#endif // ALLOW_INI_OVERRIDE_FROM_COMMANDLINE
+
+	PerObjectConfigArrayOfStructKeys = Other.PerObjectConfigArrayOfStructKeys;
+
+	// FileAccess is not copied; each FConfigFile has to set it itself
+
+	// Update the FileAccess pointers on all Sections and Values that were assigned
+#if UE_WITH_CONFIG_TRACKING
+	UE::ConfigAccessTracking::FFile* LocalFileAccess = GetFileAccess();
+	for (TMap<FString, FConfigSection>::TIterator SectionIterator(*this); SectionIterator; ++SectionIterator)
+	{
+		UE::ConfigAccessTracking::FSection* SectionAccess = nullptr;
+		if (LocalFileAccess)
+		{
+			SectionAccess = new UE::ConfigAccessTracking::FSection(*LocalFileAccess, FStringView(SectionIterator->Key));
+		}
+		SectionIterator->Value.SectionAccess = SectionAccess;
+		for (TPair<FName, FConfigValue>& ValuePair : SectionIterator->Value)
+		{
+			ValuePair.Value.SetSectionAccess(SectionAccess);
+		}
+	}
+#endif
+
+	return *this;
+}
+
+FConfigFile& FConfigFile::operator=(FConfigFile&& Other)
+{
+	this->FConfigFileMap::operator=(MoveTemp(Other));
+	Dirty = Other.Dirty;
+	NoSave = Other.NoSave;
+	bHasPlatformName = Other.bHasPlatformName;
+	bCanSaveAllSections = Other.bCanSaveAllSections;
+
+	// LoadType is not copied; each FConfigFile has to set it itself
+
+	Name = MoveTemp(Other.Name);
+	PlatformName = MoveTemp(Other.PlatformName);
+	Tag = MoveTemp(Other.Tag);
+	Branch = MoveTemp(Other.Branch);
+
+#if ALLOW_INI_OVERRIDE_FROM_COMMANDLINE
+	CommandlineOptions = MoveTemp(Other.CommandlineOptions);
+#endif // ALLOW_INI_OVERRIDE_FROM_COMMANDLINE
+
+	PerObjectConfigArrayOfStructKeys = MoveTemp(Other.PerObjectConfigArrayOfStructKeys);
+
+	// FileAccess is not copied; each FConfigFile has to set it itself
+
+	// Update the FileAccess pointers on all Sections and Values that were assigned
+#if UE_WITH_CONFIG_TRACKING
+	UE::ConfigAccessTracking::FFile* LocalFileAccess = GetFileAccess();
+	for (TMap<FString, FConfigSection>::TIterator SectionIterator(*this); SectionIterator; ++SectionIterator)
+	{
+		UE::ConfigAccessTracking::FSection* SectionAccess = nullptr;
+		if (LocalFileAccess)
+		{
+			SectionAccess = new UE::ConfigAccessTracking::FSection(*LocalFileAccess, FStringView(SectionIterator->Key));
+		}
+		SectionIterator->Value.SectionAccess = SectionAccess;
+		for (TPair<FName, FConfigValue>& ValuePair : SectionIterator->Value)
+		{
+			ValuePair.Value.SectionAccess = SectionAccess;
+		}
+	}
+#endif
+
+	return *this;
+}
+
 void FConfigFile::Cleanup()
 {
 	Empty();
@@ -880,7 +978,7 @@ UE::ConfigAccessTracking::FFile* FConfigFile::GetFileAccess() const
 		{
 			return nullptr;
 		}
-		FileAccess.Set(new UE::ConfigAccessTracking::FFile(this));
+		FileAccess = new UE::ConfigAccessTracking::FFile(this);
 	}
 	return FileAccess.GetReference();
 }
@@ -1153,11 +1251,11 @@ void FConfigFile::ProcessCommand(FConfigSection* Section, FStringView SectionNam
 				FConfigValue* ConfigValue = Section->Find(Key);
 				if (!ConfigValue)
 				{
-					Section->Add(Key, MoveTemp(Value));
+					Section->Add(Key, FConfigValue(Section, Key, MoveTemp(Value)));
 				}
 				else
 				{
-					*ConfigValue = FConfigValue(MoveTemp(Value));
+					*ConfigValue = MoveTemp(Value);
 				}
 			}
 			break;
@@ -1193,6 +1291,25 @@ void FConfigFile::ProcessCommand(FConfigSection* Section, FStringView SectionNam
 	}
 }
 
+#if UE_WITH_CONFIG_TRACKING
+static void ConditionalInitializeLoadType(FConfigFile* File, UE::ConfigAccessTracking::ELoadType LoadType,
+	FName FileName)
+{
+	if (File->LoadType == UE::ConfigAccessTracking::ELoadType::Uninitialized)
+	{
+		File->LoadType = LoadType;
+	}
+	if (File->Name.IsNone())
+	{
+		File->Name = FileName;
+	}
+}
+static void ConditionalInitializeLoadType(FConfigCommandStream* File, UE::ConfigAccessTracking::ELoadType LoadType,
+	FName FileName)
+{
+}
+#endif
+
 template<typename FileType>
 void FillFileFromBuffer(FileType* File, FStringView Buffer, bool bHandleSymbolCommands, const FString& FileHint)
 {
@@ -1203,14 +1320,7 @@ void FillFileFromBuffer(FileType* File, FStringView Buffer, bool bHandleSymbolCo
 	UE_TRACE_METADATA_SCOPE_ASSET_FNAME(FileName, ConfigFileClassName, FileName);
 
 #if UE_WITH_CONFIG_TRACKING
-	if (File->LoadType == UE::ConfigAccessTracking::ELoadType::Uninitialized)
-	{
-		File->LoadType = UE::ConfigAccessTracking::ELoadType::LocalSingleIniFile;
-	}
-	if (File->Name.IsNone())
-	{
-		File->Name = FileName;
-	}
+	ConditionalInitializeLoadType(File, UE::ConfigAccessTracking::ELoadType::LocalSingleIniFile, FileName);
 #endif
 
 	const TCHAR* Ptr = Buffer.GetData();
@@ -3052,7 +3162,18 @@ FConfigBranch& FConfigCacheIni::AddNewBranch(const FString& Filename)
 	FConfigBranch* Branch = new FConfigBranch();
 	Branch->IniName = *FPaths::GetBaseFilename(Filename);
 	Branch->IniPath = Filename;
-	return *OtherFiles.Add(Filename, Branch);
+#if UE_WITH_CONFIG_TRACKING
+	UE::ConfigAccessTracking::FFile* FileAccess = Branch->InMemoryFile.GetFileAccess();
+	if (FileAccess)
+	{
+		FileAccess->SetAsLoadTypeConfigSystem(*this, Branch->InMemoryFile);
+		FileAccess->OverrideFilenameToLoad = FName(FStringView(Filename));
+	}
+#endif
+	FConfigBranch*& Existing = OtherFiles.FindOrAdd(Filename);
+	delete Existing;
+	Existing = Branch;
+	return *Branch;
 }
 
 
@@ -3130,12 +3251,6 @@ FConfigFile* FConfigCacheIni::FindConfigFileWithBaseName(FName BaseName)
 
 FConfigFile& FConfigCacheIni::Add(const FString& Filename, const FConfigFile& File)
 {
-	FConfigBranch** Result = OtherFiles.Find(Filename);
-	if (Result)
-	{
-		delete *Result;
-	}
-
 	FConfigBranch* Branch = new FConfigBranch(File);
 	Branch->IniName = File.Name;
 	Branch->IniPath = Filename;
@@ -3147,7 +3262,10 @@ FConfigFile& FConfigCacheIni::Add(const FString& Filename, const FConfigFile& Fi
 		FileAccess->OverrideFilenameToLoad = FName(FStringView(Filename));
 	}
 #endif
-	return OtherFiles.Add(Filename, Branch)->InMemoryFile;
+	FConfigBranch*& Existing = OtherFiles.FindOrAdd(Filename);
+	delete Existing;
+	Existing = Branch;
+	return Branch->InMemoryFile;
 }
 
 bool FConfigCacheIni::ContainsConfigFile(const FConfigFile* ConfigFile) const
diff --git a/Engine/Source/Runtime/Core/Private/Misc/ConfigContext.cpp b/Engine/Source/Runtime/Core/Private/Misc/ConfigContext.cpp
--- a/Engine/Source/Runtime/Core/Private/Misc/ConfigContext.cpp
+++ b/Engine/Source/Runtime/Core/Private/Misc/ConfigContext.cpp
@@ -478,12 +478,13 @@ bool FConfigContext::PerformLoad()
 	}
 #endif
 
+	FConfigFile& FinalFile = ExistingFile ? *ExistingFile : Branch->InMemoryFile;
 #if UE_WITH_CONFIG_TRACKING
 	// Set the LoadType before calling GenerateDestIniFile, because it will set it if not
 	// already set.
-	if (Branch->InMemoryFile.LoadType == UE::ConfigAccessTracking::ELoadType::Uninitialized)
+	if (FinalFile.LoadType == UE::ConfigAccessTracking::ELoadType::Uninitialized)
 	{
-		Branch->InMemoryFile.LoadType = UE::ConfigAccessTracking::ELoadType::LocalIniFile;
+		FinalFile.LoadType = UE::ConfigAccessTracking::ELoadType::LocalIniFile;
 	}
 #endif
 	
@@ -507,8 +508,6 @@ bool FConfigContext::PerformLoad()
 	// just write out the exact same thing it read in!
 	bool bGeneratedFile = GenerateDestIniFile();
 
-	FConfigFile& FinalFile = ExistingFile ? *ExistingFile : Branch->InMemoryFile;
-
 	FinalFile.Name = FName(*BaseIniName);
 	FinalFile.PlatformName = Platform;
 	FinalFile.bHasPlatformName = true;
diff --git a/Engine/Source/Runtime/Core/Public/Misc/ConfigAccessTracking.h b/Engine/Source/Runtime/Core/Public/Misc/ConfigAccessTracking.h
--- a/Engine/Source/Runtime/Core/Public/Misc/ConfigAccessTracking.h
+++ b/Engine/Source/Runtime/Core/Public/Misc/ConfigAccessTracking.h
@@ -106,27 +106,6 @@ CORE_API void RemoveConfigValueReadCallback(FConfigValueReadCallbackId DelegateH
 namespace Private
 {
 
-/**
- * FFilePtrClearOnCopy: Device to workaround c++'s lack of the ability to create a copyconstructor that appends custom
- * code to the default copyconstructor. FConfigFile uses a default copy constructor, and we don't want to replace that
- * with a manual one that we then have to keep up to date with any property changes. But we want to change the
- * behavior of that copy constructor; FConfigFile->FileAccess is a pointer to a new'd structure that has a backpointer
- * to FConfigFile and so every allocated FConfigFile needs to construct its own copy of it rather than copying the
- * pointer from a source FConfigFile. But that is not (and should not be) the behavior of TRefCountPtr. So we make
- * this subclass of TRefCountPtr with clearing-on-copy behavior, and allocate the FFile* on demand when accessed in
- * the new FConfigFile.
- */
-struct FFilePtrClearOnCopy : public TRefCountPtr<FFile>
-{
-	FFilePtrClearOnCopy() = default;
-	FFilePtrClearOnCopy(FFile* In) : TRefCountPtr<FFile>(In) {}
-	FFilePtrClearOnCopy(const FFilePtrClearOnCopy& Other) {}
-	FFilePtrClearOnCopy(FFilePtrClearOnCopy&& Other) {}
-	FFilePtrClearOnCopy& operator=(const FFilePtrClearOnCopy& Other) { return *this; }
-	FFilePtrClearOnCopy& operator=(FFilePtrClearOnCopy&& Other) { return *this; }
-	void Set(FFile* Ptr) { this->TRefCountPtr<UE::ConfigAccessTracking::FFile>::operator=(Ptr); }
-};
-
 /**
  * Stores the number of existing AddConfigValueReadCallback subscribers. We use it to avoid the cost
  * of a function call when no subscribers are registered.
diff --git a/Engine/Source/Runtime/Core/Public/Misc/ConfigCacheIni.h b/Engine/Source/Runtime/Core/Public/Misc/ConfigCacheIni.h
--- a/Engine/Source/Runtime/Core/Public/Misc/ConfigCacheIni.h
+++ b/Engine/Source/Runtime/Core/Public/Misc/ConfigCacheIni.h
@@ -517,11 +517,6 @@ public:
 	class FConfigBranch* Branch = nullptr;
 	FName Tag;
 	
-	// unused
-#if UE_WITH_CONFIG_TRACKING
-	UE::ConfigAccessTracking::ELoadType LoadType = UE::ConfigAccessTracking::ELoadType::Uninitialized;
-	FName Name;
-#endif
 	uint8 Dirty;
 
 	// used to determine if existing settings should be removed from a FCOnfigFile before applying this stream - used for compatibility with
@@ -571,13 +566,17 @@ private:
 	TMap<FString, TMap<FName, FString> > PerObjectConfigArrayOfStructKeys;
 
 #if UE_WITH_CONFIG_TRACKING
-	mutable UE::ConfigAccessTracking::Private::FFilePtrClearOnCopy FileAccess;
+	mutable TRefCountPtr<UE::ConfigAccessTracking::FFile> FileAccess;
 #endif
 
 public:
 	CORE_API FConfigFile();
 	FConfigFile( int32 ) {}	// @todo UE-DLL: Workaround for instantiated TMap template during DLLExport (TMap::FindRef)
 	CORE_API ~FConfigFile();
+	CORE_API FConfigFile(const FConfigFile& Other);
+	CORE_API FConfigFile(FConfigFile&& Other);
+	CORE_API FConfigFile& operator=(const FConfigFile& Other);
+	CORE_API FConfigFile& operator=(FConfigFile&& Other);
 
 	// looks for a section by name, and creates an empty one if it can't be found
 	UE_DEPRECATED(5.4, "Use FindOrAddConfigSection, and/or use the new AddToSection, etc APIs to modify sections without retrieving the section. See top of ConfigCacheIni.h for more info.")
