- Changed ArrayAddUnique to ArrayAdd for file loads that don't use commands
- Fixed an issue with config override filenames on the commandline

#changelist validated
#virtualized

diff --git a/Engine/Source/Runtime/Core/Private/Misc/ConfigCacheIni.cpp b/Engine/Source/Runtime/Core/Private/Misc/ConfigCacheIni.cpp
--- a/Engine/Source/Runtime/Core/Private/Misc/ConfigCacheIni.cpp
+++ b/Engine/Source/Runtime/Core/Private/Misc/ConfigCacheIni.cpp
@@ -1423,7 +1423,7 @@ void FillFileFromBuffer(FileType* File, FStringView Buffer, bool bHandleSymbolCo
 				{
 					// when we don't want commands, the default action is to add new entries (this is for standalone ini files that have arrays,
 					// without any + cmds) - there's no difference between a single value and an array of 1 (in terms of the Config system)
-					Command = FConfigValue::EValueType::ArrayAddUnique;
+					Command = FConfigValue::EValueType::ArrayAdd;
 				}
 				else
 				{
@@ -4623,8 +4623,8 @@ FString FConfigCacheIni::GetDestIniFilename(const TCHAR* BaseIniName, const TCHA
 	{
 		FString Name(PlatformName ? PlatformName : ANSI_TO_TCHAR(FPlatformProperties::PlatformName()));
 
-		// if the BaseIniName doesn't contain the config dir, put it all together
-		if (FCString::Stristr(BaseIniName, GeneratedConfigDir) != nullptr)
+		// if the BaseIniName doesn't start with the config dir, put it all together
+		if (FString(BaseIniName).StartsWith(GeneratedConfigDir) && FPaths::GetExtension(BaseIniName) == TEXT("ini"))
 		{
 			IniFilename = BaseIniName;
 		}
