Set DetailMode:Epic(3) as the new default. This means that DetailMode:High(2) is now unassigned by the engine, but it can be used by project settings.

Also fixed up some minor things missed from the new detail mode (comments, etc).

diff --git a/Engine/Config/BaseScalability.ini b/Engine/Config/BaseScalability.ini
--- a/Engine/Config/BaseScalability.ini
+++ b/Engine/Config/BaseScalability.ini
@@ -697,7 +697,7 @@ fx.Niagara.QualityLevel=2
 r.TranslucencyLightingVolumeDim=64
 r.RefractionQuality=2
 r.SceneColorFormat=4
-r.DetailMode=2
+r.DetailMode=3
 r.TranslucencyVolumeBlur=1
 r.MaterialQualityLevel=1 ; High quality
 r.SSS.Scale=1
@@ -725,7 +725,7 @@ r.Refraction.OffsetQuality=1
 r.TranslucencyLightingVolumeDim=64
 r.RefractionQuality=2
 r.SceneColorFormat=4
-r.DetailMode=2
+r.DetailMode=3
 r.TranslucencyVolumeBlur=1
 r.MaterialQualityLevel=1 ; High quality
 r.SSS.Scale=1
diff --git a/Engine/Source/Runtime/Core/Private/HAL/ConsoleManager.cpp b/Engine/Source/Runtime/Core/Private/HAL/ConsoleManager.cpp
--- a/Engine/Source/Runtime/Core/Private/HAL/ConsoleManager.cpp
+++ b/Engine/Source/Runtime/Core/Private/HAL/ConsoleManager.cpp
@@ -4033,11 +4033,12 @@ static TAutoConsoleVariable<int32> CVarFreeSkeletalMeshBuffers(
 
 static TAutoConsoleVariable<int32> CVarDetailMode(
 	TEXT("r.DetailMode"),
-	2,
+	3,
 	TEXT("Current detail mode; determines whether components of actors should be updated/ ticked.\n"
-		 " 0: low, show only object with DetailMode low or higher\n"
-		 " 1: medium, show all object with DetailMode medium or higher\n"
-		 " 2: high, show all objects (default)"),
+		" 0: low, show objects with DetailMode low\n"
+		" 1: medium, show objects with DetailMode medium or below\n"
+		" 2: high, show objects with DetailMode high or below\n"
+		" 3: epic, show all objects (default)"),
 	ECVF_Scalability | ECVF_RenderThreadSafe);
 
 static TAutoConsoleVariable<int32> CVarCookOutUnusedDetailModeComponents(
diff --git a/Engine/Source/Runtime/Engine/Classes/Kismet/KismetSystemLibrary.h b/Engine/Source/Runtime/Engine/Classes/Kismet/KismetSystemLibrary.h
--- a/Engine/Source/Runtime/Engine/Classes/Kismet/KismetSystemLibrary.h
+++ b/Engine/Source/Runtime/Engine/Classes/Kismet/KismetSystemLibrary.h
@@ -1627,9 +1627,10 @@ class UKismetSystemLibrary : public UBlueprintFunctionLibrary
 
 	/**
 	 * Get the clamped state of r.DetailMode, see console variable help (allows for scalability, cannot be used in construction scripts)
-	 * 0: low, show only object with DetailMode low or higher
-	 * 1: medium, show all object with DetailMode medium or higher
-	 * 2: high, show all objects
+	 * 0: low, show objects with DetailMode low
+	 * 1: medium, show objects with DetailMode medium or below
+	 * 2: high, show objects with DetailMode high or below
+	 * 3: epic, show all objects
 	 */
 	UFUNCTION(BlueprintPure, Category="Rendering", meta=(UnsafeDuringActorConstruction = "true"))
 	static ENGINE_API int32 GetRenderingDetailMode();
diff --git a/Engine/Source/Runtime/Engine/Private/Particles/ParticleComponents.cpp b/Engine/Source/Runtime/Engine/Private/Particles/ParticleComponents.cpp
--- a/Engine/Source/Runtime/Engine/Private/Particles/ParticleComponents.cpp
+++ b/Engine/Source/Runtime/Engine/Private/Particles/ParticleComponents.cpp
@@ -2787,7 +2787,7 @@ void UParticleSystem::Serialize(FArchive& Ar)
 			if (DeviceProfile->GetConsolidatedCVarValue(TEXT("fx.PruneEmittersOnCookByDetailMode"), CVarDoPrune) && CVarDoPrune == 1)
 			{
 				// get the detail mode from the device platform ini; if it's not there, we assume all detail modes
-				int32 CVarDetailMode = 3;
+				int32 CVarDetailMode = PDM_DefaultValue;
 				if (DeviceProfile->GetConsolidatedCVarValue(TEXT("r.DetailMode"), CVarDetailMode))
 				{
 					CookTargetPlatformDetailModeMask = (1 << CVarDetailMode);
