Added optimized single element TArray Remove* overloads which don't take a count.

diff --git a/Engine/Source/Runtime/Core/Public/Containers/Array.h b/Engine/Source/Runtime/Core/Public/Containers/Array.h
--- a/Engine/Source/Runtime/Core/Public/Containers/Array.h
+++ b/Engine/Source/Runtime/Core/Public/Containers/Array.h
@@ -755,7 +755,24 @@ public:
 		// Template property, branch will be optimized out
 		if constexpr (AllocatorType::RequireRangeCheck)
 		{
-			checkf((Index >= 0) & (Index < ArrayNum),TEXT("Array index out of bounds: %lld from an array of size %lld"),(long long)Index, (long long)ArrayNum); // & for one branch
+			checkf((Index >= 0) & (Index < ArrayNum),TEXT("Array index out of bounds: %lld into an array of size %lld"),(long long)Index, (long long)ArrayNum); // & for one branch
+		}
+	}
+
+	/**
+	 * Checks if a range of indices are in the array range.
+	 *
+	 * @param Index Index of the start of the range to check.
+	 * @param Count Number of elements in the range.
+	 */
+	FORCEINLINE void RangeCheck(SizeType Index, SizeType Count) const
+	{
+		CheckInvariants();
+
+		// Template property, branch will be optimized out
+		if constexpr (AllocatorType::RequireRangeCheck)
+		{
+			checkf((Count >= 0) & (Index >= 0) & (Index + Count <= ArrayNum), TEXT("Array range out of bounds: index %lld and length %lld into an array of size %lld"), (long long)Index, (long long)Count, (long long)ArrayNum); // & for one branch
 		}
 	}
 
@@ -837,7 +854,11 @@ public:
 	{
 		RangeCheck(0);
 		ElementType Result = MoveTempIfPossible(GetData()[ArrayNum - 1]);
-		RemoveAt(ArrayNum - 1, 1, AllowShrinking);
+		RemoveAtImpl(ArrayNum - 1);
+		if (AllowShrinking == EAllowShrinking::Yes)
+		{
+			ResizeShrink();
+		}
 		return Result;
 	}
 	UE_ALLOWSHRINKING_BOOL_DEPRECATED("Pop")
@@ -1690,32 +1711,38 @@ public:
 	}
 
 private:
-	void RemoveAtImpl(SizeType Index, SizeType Count, EAllowShrinking AllowShrinking)
+	void RemoveAtImpl(SizeType Index)
 	{
-		if (Count)
-		{
-			CheckInvariants();
-			checkSlow((Count >= 0) & (Index >= 0) & (Index + Count <= ArrayNum));
+		ElementType* Dest = GetData() + Index;
 
-			ElementType* Dest = GetData() + Index;
+		DestructItem(Dest);
+
+		// Skip relocation in the common case that there is nothing to move.
+		SizeType NumToMove = (ArrayNum - Index) - 1;
+		if (NumToMove)
+		{
+			RelocateConstructItems<ElementType>(Dest, Dest + 1, NumToMove);
+		}
+		--ArrayNum;
 
-			DestructItems(Dest, Count);
+		SlackTrackerNumChanged();
+	}
 
-			// Skip memmove in the common case that there is nothing to move.
-			SizeType NumToMove = (ArrayNum - Index) - Count;
-			if (NumToMove)
-			{
-				RelocateConstructItems<ElementType>(Dest, Dest + Count, NumToMove);
-			}
-			ArrayNum -= Count;
+	void RemoveAtImpl(SizeType Index, SizeType Count)
+	{
+		ElementType* Dest = GetData() + Index;
 
-			SlackTrackerNumChanged();
+		DestructItems(Dest, Count);
 
-			if (AllowShrinking == EAllowShrinking::Yes)
-			{
-				ResizeShrink();
-			}
+		// Skip relocation in the common case that there is nothing to move.
+		SizeType NumToMove = (ArrayNum - Index) - Count;
+		if (NumToMove)
+		{
+			RelocateConstructItems<ElementType>(Dest, Dest + Count, NumToMove);
 		}
+		ArrayNum -= Count;
+
+		SlackTrackerNumChanged();
 	}
 
 public:
@@ -1724,10 +1751,16 @@ public:
 	 * the array.
 	 *
 	 * @param Index Location in array of the element to remove.
+	 * @param AllowShrinking (Optional) Tells if this call can shrink array if suitable after remove. Default is yes.
 	 */
-	FORCEINLINE void RemoveAt(SizeType Index)
+	void RemoveAt(SizeType Index, EAllowShrinking AllowShrinking = EAllowShrinking::Yes)
 	{
-		RemoveAtImpl(Index, 1, EAllowShrinking::Yes);
+		RangeCheck(Index);
+		RemoveAtImpl(Index);
+		if (AllowShrinking == EAllowShrinking::Yes)
+		{
+			ResizeShrink();
+		}
 	}
 
 	/**
@@ -1745,7 +1778,15 @@ public:
 	FORCEINLINE void RemoveAt(SizeType Index, CountType Count, EAllowShrinking AllowShrinking = EAllowShrinking::Yes)
 	{
 		static_assert(!std::is_same_v<CountType, bool>, "TArray::RemoveAt: unexpected bool passed as the Count argument");
-		RemoveAtImpl(Index, (SizeType)Count, AllowShrinking);
+		RangeCheck(Index, Count);
+		if (Count)
+		{
+			RemoveAtImpl(Index, (SizeType)Count);
+			if (AllowShrinking == EAllowShrinking::Yes)
+			{
+				ResizeShrink();
+			}
+		}
 	}
 	template <typename CountType>
 	UE_ALLOWSHRINKING_BOOL_DEPRECATED("RemoveAt")
@@ -1755,34 +1796,42 @@ public:
 	}
 
 private:
-	void RemoveAtSwapImpl(SizeType Index, SizeType Count, EAllowShrinking AllowShrinking)
+	void RemoveAtSwapImpl(SizeType Index)
 	{
-		if (Count)
-		{
-			CheckInvariants();
-			checkSlow((Count >= 0) & (Index >= 0) & (Index + Count <= ArrayNum));
+		ElementType* Data = GetData();
+		ElementType* Dest = Data + Index;
 
-			ElementType* Data = GetData();
-			ElementType* Dest = Data + Index;
+		DestructItem(Dest);
 
-			DestructItems(Dest, Count);
+		// Replace the elements in the hole created by the removal with elements from the end of the array, so the range of indices used by the array is contiguous.
+		const SizeType NumElementsAfterHole = (ArrayNum - Index) - 1;
+		const SizeType NumElementsToMoveIntoHole = FPlatformMath::Min(1, NumElementsAfterHole);
+		if (NumElementsToMoveIntoHole)
+		{
+			RelocateConstructItems<ElementType>(Dest, Data + (ArrayNum - NumElementsToMoveIntoHole), NumElementsToMoveIntoHole);
+		}
+		--ArrayNum;
 
-			// Replace the elements in the hole created by the removal with elements from the end of the array, so the range of indices used by the array is contiguous.
-			const SizeType NumElementsAfterHole = (ArrayNum - Index) - Count;
-			const SizeType NumElementsToMoveIntoHole = FPlatformMath::Min(Count, NumElementsAfterHole);
-			if (NumElementsToMoveIntoHole)
-			{
-				RelocateConstructItems<ElementType>(Dest, Data + (ArrayNum - NumElementsToMoveIntoHole), NumElementsToMoveIntoHole);
-			}
-			ArrayNum -= Count;
+		SlackTrackerNumChanged();
+	}
 
-			SlackTrackerNumChanged();
+	void RemoveAtSwapImpl(SizeType Index, SizeType Count)
+	{
+		ElementType* Data = GetData();
+		ElementType* Dest = Data + Index;
 
-			if (AllowShrinking == EAllowShrinking::Yes)
-			{
-				ResizeShrink();
-			}
+		DestructItems(Dest, Count);
+
+		// Replace the elements in the hole created by the removal with elements from the end of the array, so the range of indices used by the array is contiguous.
+		const SizeType NumElementsAfterHole = (ArrayNum - Index) - Count;
+		const SizeType NumElementsToMoveIntoHole = FPlatformMath::Min(Count, NumElementsAfterHole);
+		if (NumElementsToMoveIntoHole)
+		{
+			RelocateConstructItems<ElementType>(Dest, Data + (ArrayNum - NumElementsToMoveIntoHole), NumElementsToMoveIntoHole);
 		}
+		ArrayNum -= Count;
+
+		SlackTrackerNumChanged();
 	}
 
 public:
@@ -1794,10 +1843,17 @@ public:
 	 * O(ArrayNum)), but does not preserve the order.
 	 *
 	 * @param Index Location in array of the element to remove.
+	 * @param AllowShrinking (Optional) Tells if this call can shrink array if
+	 *                        suitable after remove. Default is yes.
 	 */
-	FORCEINLINE void RemoveAtSwap(SizeType Index)
+	FORCEINLINE void RemoveAtSwap(SizeType Index, EAllowShrinking AllowShrinking = EAllowShrinking::Yes)
 	{
-		RemoveAtSwapImpl(Index, 1, EAllowShrinking::Yes);
+		RangeCheck(Index);
+		RemoveAtSwapImpl(Index);
+		if (AllowShrinking == EAllowShrinking::Yes)
+		{
+			ResizeShrink();
+		}
 	}
 
 	/**
@@ -1819,13 +1875,21 @@ public:
 	FORCEINLINE void RemoveAtSwap(SizeType Index, CountType Count, EAllowShrinking AllowShrinking = EAllowShrinking::Yes)
 	{
 		static_assert(!std::is_same_v<CountType, bool>, "TArray::RemoveAtSwap: unexpected bool passed as the Count argument");
-		RemoveAtSwapImpl(Index, Count, AllowShrinking);
+		RangeCheck(Index, Count);
+		if (Count)
+		{
+		    RemoveAtSwapImpl(Index, Count);
+		    if (AllowShrinking == EAllowShrinking::Yes)
+		    {
+			    ResizeShrink();
+		    }
+		}
 	}
 	template <typename CountType>
 	UE_ALLOWSHRINKING_BOOL_DEPRECATED("RemoveAtSwap")
 	FORCEINLINE void RemoveAtSwap(SizeType Index, CountType Count, bool bAllowShrinking)
 	{
-		RemoveAtSwapImpl(Index, Count, bAllowShrinking ? EAllowShrinking::Yes : EAllowShrinking::No);
+		RemoveAtSwap(Index, Count, bAllowShrinking ? EAllowShrinking::Yes : EAllowShrinking::No);
 	}
 
 	/**
diff --git a/Engine/Source/Runtime/Core/Public/Templates/MemoryOps.h b/Engine/Source/Runtime/Core/Public/Templates/MemoryOps.h
--- a/Engine/Source/Runtime/Core/Public/Templates/MemoryOps.h
+++ b/Engine/Source/Runtime/Core/Public/Templates/MemoryOps.h
@@ -6,6 +6,7 @@
 #include "HAL/UnrealMemory.h"
 #include "Templates/IsTriviallyCopyConstructible.h"
 #include "Templates/UnrealTypeTraits.h"
+#include "Traits/UseBitwiseSwap.h"
 #include <new> // IWYU pragma: export
 
 #include "Templates/IsTriviallyCopyAssignable.h"
@@ -153,6 +154,39 @@ FORCEINLINE void CopyAssignItems(ElementType* Dest, const ElementType* Source, S
 	}
 }
 
+/**
+ * Relocates a single items to a new memory location as a new type. This is a so-called 'destructive move' for which
+ * there is no single operation in C++ but which can be implemented very efficiently in general.
+ *
+ * @param	Dest		The memory location to relocate to.
+ * @param	Source		A pointer to the first item to relocate.
+ */
+template <typename DestinationElementType, typename SourceElementType, typename SizeType>
+FORCEINLINE void RelocateConstructItem(void* Dest, const SourceElementType* Source)
+{
+	// Do a bitwise relocate if TCanBitwiseRelocate says we can, but not if the type involved is known not to be bitwise-swappable
+	if constexpr (UE::Core::Private::MemoryOps::TCanBitwiseRelocate<DestinationElementType, SourceElementType>::Value && (!std::is_same_v<DestinationElementType, SourceElementType> || TUseBitwiseSwap<SourceElementType>::Value))
+	{
+		/* All existing UE containers seem to assume trivial relocatability (i.e. memcpy'able) of their members,
+		 * so we're going to assume that this is safe here.  However, it's not generally possible to assume this
+		 * in general as objects which contain pointers/references to themselves are not safe to be trivially
+		 * relocated.
+		 *
+		 * However, it is not yet possible to automatically infer this at compile time, so we can't enable
+		 * different (i.e. safer) implementations anyway. */
+
+		FMemory::Memmove(Dest, Source, sizeof(SourceElementType));
+	}
+	else
+	{
+		// We need a typedef here because VC won't compile the destructor call below if SourceElementType itself has a member called SourceElementType
+		typedef SourceElementType RelocateConstructItemsElementTypeTypedef;
+
+		new (Dest) DestinationElementType(*Source);
+		Source->RelocateConstructItemsElementTypeTypedef::~RelocateConstructItemsElementTypeTypedef();
+	}
+}
+
 /**
  * Relocates a range of items to a new memory location as a new type. This is a so-called 'destructive move' for which
  * there is no single operation in C++ but which can be implemented very efficiently in general.
