SavePackage Name Sorting: Sort names case-sensitively to prevent indeterminism.
#rnx

diff --git a/Engine/Source/Runtime/Core/Private/UObject/UnrealNames.cpp b/Engine/Source/Runtime/Core/Private/UObject/UnrealNames.cpp
--- a/Engine/Source/Runtime/Core/Private/UObject/UnrealNames.cpp
+++ b/Engine/Source/Runtime/Core/Private/UObject/UnrealNames.cpp
@@ -2025,6 +2025,7 @@ static FNamePool& GetNamePoolPostInit()
 	return (FNamePool&)NamePoolData;
 }
 
+template <ESearchCase::Type SearchCase>
 static int32 CompareDifferentIdsAlphabetically(FNameEntryId AId, FNameEntryId BId)
 {
 	checkSlow(AId != BId);
@@ -2050,18 +2051,32 @@ static int32 CompareDifferentIdsAlphabetically(FNameEntryId AId, FNameEntryId BI
 	}
 
 	int32 MinLen = FMath::Min(AView.Len, BView.Len);
-	if (int32 StrDiff = AView.bIsWide ?	FCStringWide::Strnicmp(AView.Wide, BView.Wide, MinLen) :
-										FCStringAnsi::Strnicmp(AView.Ansi, BView.Ansi, MinLen))
+	int32 StrDiff;
+	if (SearchCase == ESearchCase::IgnoreCase)
 	{
-		return StrDiff;
+		StrDiff = AView.bIsWide ?
+			FCStringWide::Strnicmp(AView.Wide, BView.Wide, MinLen) :
+			FCStringAnsi::Strnicmp(AView.Ansi, BView.Ansi, MinLen);
 	}
+	else
+	{
+		StrDiff = AView.bIsWide ?
+			FCStringWide::Strncmp(AView.Wide, BView.Wide, MinLen) :
+			FCStringAnsi::Strncmp(AView.Ansi, BView.Ansi, MinLen);
+	}
+	if (StrDiff) return StrDiff;
 
 	return AView.Len - BView.Len;
 }
 
 int32 FNameEntryId::CompareLexical(FNameEntryId Rhs) const
 {
-	return Value != Rhs.Value ? CompareDifferentIdsAlphabetically(*this, Rhs) : 0;
+	return Value != Rhs.Value ? CompareDifferentIdsAlphabetically<ESearchCase::IgnoreCase>(*this, Rhs) : 0;
+}
+
+int32 FNameEntryId::CompareLexicalSensitive(FNameEntryId Rhs) const
+{
+	return Value != Rhs.Value ? CompareDifferentIdsAlphabetically<ESearchCase::CaseSensitive>(*this, Rhs) : 0;
 }
 
 #if !UE_BUILD_SHIPPING && !UE_BUILD_TEST
@@ -3286,7 +3301,7 @@ int32 FName::Compare( const FName& Other ) const
 	}
 
 	// Names don't match. This means we don't even need to check numbers.
-	return CompareDifferentIdsAlphabetically(GetComparisonIndex(), Other.GetComparisonIndex());
+	return CompareDifferentIdsAlphabetically<ESearchCase::IgnoreCase>(GetComparisonIndex(), Other.GetComparisonIndex());
 #else	// UE_FNAME_OUTLINE_NUMBER
 	// Names match, check whether numbers match.
 	if (GetComparisonIndex() == Other.GetComparisonIndex())
@@ -3295,7 +3310,7 @@ int32 FName::Compare( const FName& Other ) const
 	}
 
 	// Names don't match. This means we don't even need to check numbers.
-	return CompareDifferentIdsAlphabetically(GetComparisonIndex(), Other.GetComparisonIndex());
+	return CompareDifferentIdsAlphabetically<ESearchCase::IgnoreCase>(GetComparisonIndex(), Other.GetComparisonIndex());
 #endif // UE_FNAME_OUTLINE_NUMBER
 }
 
diff --git a/Engine/Source/Runtime/Core/Public/UObject/NameTypes.h b/Engine/Source/Runtime/Core/Public/UObject/NameTypes.h
--- a/Engine/Source/Runtime/Core/Public/UObject/NameTypes.h
+++ b/Engine/Source/Runtime/Core/Public/UObject/NameTypes.h
@@ -59,10 +59,14 @@ struct FNameEntryId
 		return Value == 0;
 	}
 
-	/** Slow alphabetical order that is stable / deterministic over process runs */
+	/** Slow alphabetical order that is stable / deterministic over process runs, ignores case */
 	CORE_API int32 CompareLexical(FNameEntryId Rhs) const;
 	bool LexicalLess(FNameEntryId Rhs) const { return CompareLexical(Rhs) < 0; }
 
+	/** Slow alphabetical order that is stable / deterministic over process runs, case-sensitive */
+	CORE_API int32 CompareLexicalSensitive(FNameEntryId Rhs) const;
+	bool LexicalSensitiveLess(FNameEntryId Rhs) const { return CompareLexicalSensitive(Rhs) < 0; }
+
 	/** Fast non-alphabetical order that is only stable during this process' lifetime */
 	int32 CompareFast(FNameEntryId Rhs) const { return Value - Rhs.Value; };
 	bool FastLess(FNameEntryId Rhs) const { return CompareFast(Rhs) < 0; }
diff --git a/Engine/Source/Runtime/CoreUObject/Private/UObject/SavePackage2.cpp b/Engine/Source/Runtime/CoreUObject/Private/UObject/SavePackage2.cpp
--- a/Engine/Source/Runtime/CoreUObject/Private/UObject/SavePackage2.cpp
+++ b/Engine/Source/Runtime/CoreUObject/Private/UObject/SavePackage2.cpp
@@ -1065,17 +1065,14 @@ ESavePackageResult CreateLinker(FSaveContext& SaveContext)
 
 struct FNameEntryIdSortHelper
 {
-	/** Comparison function used by Sort */
-	FORCEINLINE bool operator()(const FName& A, const FName& B) const
-	{
-		return A.Compare(B) < 0;
-	}
-
-	/** Comparison function used by Sort */
-	FORCEINLINE bool operator()(FNameEntryId A, FNameEntryId B) const
-	{
-		//@todo Could be implemented without constructing FName but need a would new FNameEntry comparison API
-		return A != B && operator()(FName::CreateFromDisplayId(A, 0), FName::CreateFromDisplayId(B, 0));
+	/** Comparison function used when sorting Names in the package's name table. */
+	bool operator()(FNameEntryId A, FNameEntryId B) const
+	{
+		if (A == B) return false;
+		// Sort by ignore-case first, then by case-sensitive.
+		// So we will get { 'AAA', 'Aaa', 'aaa', 'BBB', 'Bbb', 'bbb' }
+		if (int32 Compare = A.CompareLexical(B)) return Compare < 0;
+		return A.CompareLexicalSensitive(B) < 0;
 	}
 };
 
