Refactor driver denylist parsing and checking code:
* All existing functionality is preserved.
* Additional functionality: allow to set denylists using a regex on the GPU name in addition of per-RHI.
* Additional functionality: these constraints (RHI name, GPU name) can now also be applied to the SuggestedDriverVersion entries.
* Further constraints can easily be added in the FDriverConfigEntryConstraints class if needed in the future.
* Now suggest the recommended driver OR latest drivers in the pop-up message (unless latest driver is denylisted).
* Remove unused code and modernize legacy code.
* Add bespoke unit tests.

#tests tested by QA on a variety of multi-vendor scenarios

diff --git a/Engine/Source/Runtime/Core/Private/Windows/WindowsPlatformMisc.cpp b/Engine/Source/Runtime/Core/Private/Windows/WindowsPlatformMisc.cpp
--- a/Engine/Source/Runtime/Core/Private/Windows/WindowsPlatformMisc.cpp
+++ b/Engine/Source/Runtime/Core/Private/Windows/WindowsPlatformMisc.cpp
@@ -3352,15 +3352,15 @@ static void GetVideoDriverDetailsFromSetup(const FString& DeviceName, bool bVerb
 
 	if (!Out.ProviderName.IsEmpty())
 	{
-		if (Out.ProviderName.Find(TEXT("NVIDIA")) != INDEX_NONE)
+		if (Out.ProviderName.Contains(TEXT("NVIDIA")))
 		{
 			Out.SetNVIDIA();
 		}
-		else if (Out.ProviderName.Find(TEXT("Advanced Micro Devices")) != INDEX_NONE)
+		else if (Out.ProviderName.Contains(TEXT("Advanced Micro Devices")))
 		{
 			Out.SetAMD();
 		}
-		else if (Out.ProviderName.Find(TEXT("Intel")) != INDEX_NONE)	// usually TEXT("Intel Corporation")
+		else if (Out.ProviderName.Contains(TEXT("Intel")))	// usually TEXT("Intel Corporation")
 		{
 			Out.SetIntel();
 		}
@@ -3370,7 +3370,7 @@ static void GetVideoDriverDetailsFromSetup(const FString& DeviceName, bool bVerb
 
 	if(Out.IsNVIDIA())
 	{
-		Out.UserDriverVersion = Out.TrimNVIDIAInternalVersion(Out.InternalDriverVersion);
+		Out.UserDriverVersion = Out.GetNVIDIAUnifiedVersion(Out.InternalDriverVersion);
 	}
 	else if(Out.IsAMD() && !RegistryKey.IsEmpty())
 	{
@@ -3458,7 +3458,7 @@ static void GetVideoDriverDetails(const FString& Key, FGPUDriverInfo& Out)
 
 	if(Out.IsNVIDIA())
 	{
-		Out.UserDriverVersion = Out.TrimNVIDIAInternalVersion(Out.InternalDriverVersion);
+		Out.UserDriverVersion = Out.GetNVIDIAUnifiedVersion(Out.InternalDriverVersion);
 	}
 	else if(Out.IsAMD())
 	{
diff --git a/Engine/Source/Runtime/Core/Public/GenericPlatform/GenericPlatformDriver.h b/Engine/Source/Runtime/Core/Public/GenericPlatform/GenericPlatformDriver.h
--- a/Engine/Source/Runtime/Core/Public/GenericPlatform/GenericPlatformDriver.h
+++ b/Engine/Source/Runtime/Core/Public/GenericPlatform/GenericPlatformDriver.h
@@ -8,686 +8,552 @@
 #include "Math/UnrealMathUtility.h"
 #include "Containers/UnrealString.h"
 #include "Misc/DateTime.h"
+#include "Misc/Optional.h"
 #include "Misc/Parse.h"
 #include "Misc/ConfigCacheIni.h"
+#include "Internationalization/Regex.h"
 
-// TSize integer values separated by '.' e.g. "120.210.11.22"
-// made to compare different GPU driver versions
-// handles whitespace
-// if there are to too many numbers we take the left most ones
-// comparison operators (for driver comparison)
-template <int32 TSize>
-class FMultiInt
+// Represents a full driver version with multiple version numbers,
+// e.g. 30.0.13023.4001.
+class FDriverVersion
 {
 public:
-	// constructor
-	FMultiInt()
-	{
-		for(uint32 i = 0; i < Size; ++i)
-		{
-			Value[i] = 0;
-		}
-	}
+	FDriverVersion() = default;
 
-	FMultiInt(const TCHAR* In)
+	FDriverVersion(const FString& DriverVersionString)
 	{
-		Parse(In);
+		Parse(DriverVersionString);
 	}
 
-	// like Parse but it doesn't alter the pointer
-	void GetValue(const TCHAR* In)
-	{
-		const TCHAR* Cpy = In;
-
-		Parse(Cpy);
-	}
-	// @param In  "123.456" "132", pointer is advanced to be after the data
-	// if there isn't enough input values, the rightmost value gets set
-	void Parse(const TCHAR*& In)
+	void Parse(const FString& DriverVersionString)
 	{
-		check(In);
-		// clear
-		*this = FMultiInt();
-
-		// find out how many numbers we have
-		uint32 NumberCount = 1;
-		{
-			const TCHAR* p = In;
-
-			while(!IsSeparator(*p))
-			{
-				// '.' separates numbers
-				if(*p == TCHAR('.'))
-				{
-					++NumberCount;
-				}
-				++p;
-			}
-		}
+		Values.Empty();
+		FString RemainingString = DriverVersionString;
+		FString CurrentVersionNumber;
 
-		NumberCount = FMath::Min(NumberCount, Size);
-
-		// parse the data
+		uint32_t CurrentVersionNumberIdx = 0;
+		while (RemainingString.Split(TEXT("."), &CurrentVersionNumber, &RemainingString))
 		{
-			const TCHAR* p = In;
-			for(uint32 i = 0; i < NumberCount; ++i)
-			{
-				Value[Size - NumberCount + i] = FCString::Atoi(p);
-
-				// go to next '.', operator or end
-				while(!IsSeparator(*p) && *p != TCHAR('.'))
-				{
-					++p;
-				}
-
-				// jump over '.'
-				if(*p == TCHAR('.'))
-				{
-					++p;
-				}
-			}
-
-			In = p;
+			Values.Add(FCString::Atoi(*CurrentVersionNumber));
 		}
+		Values.Add(FCString::Atoi(*RemainingString));
 	}
 
-	// the base comparison operator, we derive the others from it
-	bool operator>(const FMultiInt<TSize>& rhs) const
+	bool operator>(const FDriverVersion& Other) const
 	{
-		for(uint32 i = 0 ; i < Size; ++i)
+		for (int32 Idx = 0; Idx < Values.Num(); ++Idx)
 		{
-			if(Value[i] > rhs.Value[i])
+			if (Values[Idx] > Other.Values[Idx])
 			{
-				return true; 
+				return true;
 			}
-			if(Value[i] < rhs.Value[i])
+			if (Values[Idx] < Other.Values[Idx])
 			{
-				return false; 
+				return false;
 			}
 		}
 
 		return false;
 	}
-	// could be optimized but doesn't have to be fast
-	bool operator==(const FMultiInt<TSize>& rhs) const
+
+	bool operator==(const FDriverVersion& Other) const
 	{
-		// map to existing operators
-		return !(*this > rhs) && !(rhs > *this);
+		return !(*this > Other) && !(Other > *this);
 	}
-	// could be optimized but doesn't have to be fast
-	bool operator<(const FMultiInt<TSize>& rhs) const
+	bool operator<(const FDriverVersion& Other) const
 	{
-		// map to existing operators
-		return rhs > *this;
+		return Other > *this;
 	}
-	// could be optimized but doesn't have to be fast
-	bool operator!=(const FMultiInt<TSize>& rhs) const
+	bool operator!=(const FDriverVersion& Other) const
 	{
-		// map to existing operators
-		return !(*this == rhs);
+		return !(*this == Other);
 	}
-	// could be optimized but doesn't have to be fast
-	bool operator>=(const FMultiInt<TSize>& rhs) const
+	bool operator>=(const FDriverVersion& Other) const
 	{
-		// map to existing operators
-		return (*this == rhs) || (*this > rhs);
+		return (*this == Other) || (*this > Other);
 	}
-	// could be optimized but doesn't have to be fast
-	bool operator<=(const FMultiInt<TSize>& rhs) const
+	bool operator<=(const FDriverVersion& Other) const
 	{
-		// map to existing operators
-		return (*this == rhs) || (*this < rhs);
+		return (*this == Other) || (*this < Other);
 	}
 
-	static constexpr uint32 Size = TSize;
-
-	// [0]:left .. [Size-1]:right
-	uint32 Value[Size];
+	// Index from the left (major version first).
+	uint32 GetVersionValue(uint32 Index) const
+	{
+		return Values[Index];
+	}
 
-private:
-	bool IsSeparator(const TCHAR c)
+	const TArray<uint32>& GetVersionValues() const
 	{
-		// comparison operators are considered separators (e.g. "1.21 < 0.121")
-		return c == 0
-			|| c == TCHAR('=')
-			|| c == TCHAR('!')
-			|| c == TCHAR('<')
-			|| c == TCHAR('&')
-			|| c == TCHAR('|')
-			|| c == TCHAR('>');
+		return Values;
 	}
+
+private:
+	TArray<uint32> Values;
 };
 
-// used to compare driver versions in Hardware.ini
 enum EComparisonOp
 {
-	ECO_Unknown, ECO_Equal, ECO_NotEqual, ECO_Larger, ECO_LargerThan, ECO_Smaller, ECO_SmallerThan
+	ECO_Unknown, ECO_Equal, ECO_NotEqual, ECO_Greater, ECO_GreaterOrEqual, ECO_Less, ECO_LessOrEqual
 };
 
-// Find the comparison enum based on input text
-// @param In pointer is advanced to be after the data
+// Find the comparison operator in the input string, and advances the In pointer to point to the first character
+// after the comparison operator.
+// Defaults to Equal if no comparison operator is present.
 inline EComparisonOp ParseComparisonOp(const TCHAR*& In)
 {
-	if(In[0] == '=' && In[1] == '=')
+	if (In[0] == '=' && In[1] == '=')
 	{
 		In += 2;
 		return ECO_Equal;
 	}
-	if(In[0] == '!' && In[1] == '=')
+	if (In[0] == '!' && In[1] == '=')
 	{
 		In += 2;
 		return ECO_NotEqual;
 	}
-	if(*In == '>')
+	if (*In == '>')
 	{
 		++In;
-		if(*In == '=')
+		if (*In == '=')
 		{
 			++In;
-			return ECO_LargerThan;
+			return ECO_GreaterOrEqual;
 		}
-		return ECO_Larger;
+		return ECO_Greater;
 	}
-	if(*In == '<')
+	if (*In == '<')
 	{
 		++In;
-		if(*In == '=')
+		if (*In == '=')
 		{
 			++In;
-			return ECO_SmallerThan;
+			return ECO_LessOrEqual;
 		}
-		return ECO_Smaller;
+		return ECO_Less;
 	}
-
-	return ECO_Unknown;
+	return ECO_Equal;
 }
 
-// general comparison with the comparison operator given as enum
-template <class T> bool Compare(const T& A, EComparisonOp Op, const T& B)
+// Compare two objects given the comparison operator.
+template <class T>
+bool CompareWithOp(const T& A, EComparisonOp Op, const T& B)
 {
-	switch(Op)
+	switch (Op)
 	{
-		case ECO_Equal:			return A == B;
-		case ECO_NotEqual:		return A != B;
-		case ECO_Larger:		return A >  B;
-		case ECO_LargerThan:	return A >= B;
-		case ECO_Smaller:		return A <  B;
-		case ECO_SmallerThan:	return A <= B;
+		case ECO_Equal:			 return A == B;
+		case ECO_NotEqual:		 return A != B;
+		case ECO_Greater:		 return A > B;
+		case ECO_GreaterOrEqual: return A >= B;
+		case ECO_Less:		     return A < B;
+		case ECO_LessOrEqual:	 return A <= B;
 	}
 	check(0);
 	return false;
 }
 
-// useful to express very simple math, later we might extend to express a range e.g. ">10 && <12.121"
-// @param InOpWithVersion e.g. "<=220.2"
-// @param CurrentVersion e.g. "219.1"
-inline bool CompareStringOp(const TCHAR* InOpWithMultiInt, const TCHAR* CurrentMultiInt)
-{
-	const TCHAR* p = InOpWithMultiInt;
-
-	EComparisonOp Op = ParseComparisonOp(p);
-
-	// it makes sense to compare for equal if there is no operator
-	if(Op == ECO_Unknown)
-	{
-		Op = ECO_Equal;
-	}
-
-	FMultiInt<6> A, B;
-		
-	A.GetValue(CurrentMultiInt);
-	B.Parse(p);
-
-	return Compare(A, Op, B);
-}
-
-
-// video driver details
 struct FGPUDriverInfo
 {
-	FGPUDriverInfo()
-		: VendorId(0)
-	{
-	}
+	// Hardware vendor ID, e.g. 0x10DE for NVIDIA.
+	uint32 VendorId = 0;
 
-	// DirectX VendorId, 0 if not set, use functions below to set/get
-	uint32 VendorId;
-	// e.g. "NVIDIA GeForce GTX 680" or "AMD Radeon R9 200 / HD 7900 Series"
+	// Name of the graphics device, e.g. "NVIDIA GeForce GTX 680" or "AMD Radeon R9 200 / HD 7900 Series".
 	FString DeviceDescription;
-	// e.g. "NVIDIA" or "Advanced Micro Devices, Inc."
+
+	// Name of the hardware vendor, e.g. "NVIDIA" or "Advanced Micro Devices, Inc."
 	FString ProviderName;
-	// e.g. "15.200.1062.1004"(AMD)
-	// e.g. "9.18.13.4788"(NVIDIA) first number is Windows version (e.g. 7:Vista, 6:XP, 4:Me, 9:Win8(1), 10:Win7), last 5 have the UserDriver version encoded
-	// also called technical version number (https://wiki.mozilla.org/Blocklisting/Blocked_Graphics_Drivers)
-	// TEXT("Unknown") if driver detection failed
-	FString InternalDriverVersion;	
-	// e.g. "Catalyst 15.7.1"(AMD) or "Crimson 15.7.1"(AMD) or "347.88"(NVIDIA)
-	// also called commercial version number (https://wiki.mozilla.org/Blocklisting/Blocked_Graphics_Drivers)
+
+	// Internal driver version, which may differ from the version shown to the user. Set to "Unknown" if the driver detection failed.
+	FString InternalDriverVersion;
+
+	// User-facing driver version.
 	FString UserDriverVersion;
-	// e.g. 3-13-2015
+
+	// Driver date as reported by the driver. Format is MM-DD-YYYY.
 	FString DriverDate;
-	// e.g. D3D11, D3D12
+
+	// Current RHI.
 	FString RHIName;
 
 	bool IsValid() const
 	{
 		return !DeviceDescription.IsEmpty()
 			&& VendorId
-			&& (InternalDriverVersion != TEXT("Unknown"))		// if driver detection code fails
-			&& (InternalDriverVersion != TEXT(""));				// if running on non Windows platform we don't fill in the driver version, later we need to check for the OS as well.
+			&& (InternalDriverVersion != TEXT("Unknown")) // If driver detection code fails.
+			&& (InternalDriverVersion != TEXT(""));		  // If running on a non-Windows platform we don't fill in the driver version.
 	}
-	
-	// set VendorId
+
 	void SetAMD() { VendorId = 0x1002; }
-	// set VendorId
 	void SetIntel() { VendorId = 0x8086; }
-	// set VendorId
 	void SetNVIDIA() { VendorId = 0x10DE; }
-	// get VendorId
+
 	bool IsAMD() const { return VendorId == 0x1002; }
-	// get VendorId
 	bool IsIntel() const { return VendorId == 0x8086; }
-	// get VendorId
 	bool IsNVIDIA() const { return VendorId == 0x10DE; }
 
-	bool IsSameDriverVersionGeneration(const TCHAR* InOpWithMultiInt) const
+	bool IsSameDriverVersionGeneration(const FDriverVersion& UnifiedDriverVersion) const
 	{
 		if (IsIntel())
 		{
-			const TCHAR* p = InOpWithMultiInt;
-			FString DriverVersion = GetUnifiedDriverVersion();
-
-			EComparisonOp Op = ParseComparisonOp(p);
-
-			FMultiInt<6> A, B;
-
-			A.GetValue(*DriverVersion);
-			B.Parse(p);
+			FDriverVersion ThisVersion(GetUnifiedDriverVersion());
 
 			// https://www.intel.com/content/www/us/en/support/articles/000005654/graphics.html
-			// Version format changed in April 2018 starting with xx.xx.100.xxxx
-			if (!((A.Value[4] >= 100) ^ (B.Value[4] >= 100)))
-			{
-				return true;
-			}
-			else
-			{
-				return false;
-			}
+			// Version format changed in April 2018 starting with xx.xx.100.xxxx.
+			// In the unified driver version, that's the first value.
+			return (ThisVersion.GetVersionValue(0) >= 100) == (UnifiedDriverVersion.GetVersionValue(0) >= 100);
 		}
-		
+
 		return true;
 	}
 
-	static FString TrimNVIDIAInternalVersion(const FString& InternalVersion)
+	static FString GetNVIDIAUnifiedVersion(const FString& InternalVersion)
 	{
-		// on the internal driver number: https://forums.geforce.com/default/topic/378546/confusion-over-driver-version-numbers/
-		//   The first 7 shows u that is a Vista driver, 6 that is an XP and 4 that is Me
-		// we don't care about the windows version so we don't look at the front part of the driver version
-		// "9.18.13.4788" -> "347.88"
-		// "10.18.13.4788" -> "347.88"
-		// the following code works with the current numbering scheme, if needed we have to update that
+		// Ignore the Windows/DirectX version by taking the last digits of the internal version
+		// and moving the version dot. Coincidentally, that's the user-facing string. For example:
+		// 9.18.13.4788 -> 3.4788 -> 347.88
+		if (InternalVersion.Len() < 6)
+		{
+			return InternalVersion;
+		}
 
-		// we don't care about the windows version so we don't look at the front part of the driver version
-		// e.g. 36.143
 		FString RightPart = InternalVersion.Right(6);
-
-		// move the dot
 		RightPart = RightPart.Replace(TEXT("."), TEXT(""));
 		RightPart.InsertAt(3, TEXT("."));
 		return RightPart;
 	}
 
+	static FString GetIntelUnifiedVersion(const FString& InternalVersion)
+	{
+		// https://www.intel.com/content/www/us/en/support/articles/000005654/graphics.html
+		// Drop off the OS and DirectX version. For example:
+		// 27.20.100.8935 -> 100.8935
+		int32 DotIndex = InternalVersion.Find(TEXT("."), ESearchCase::CaseSensitive, ESearchDir::FromStart);
+		if (DotIndex != INDEX_NONE)
+		{
+			DotIndex = InternalVersion.Find(TEXT("."), ESearchCase::CaseSensitive, ESearchDir::FromStart, DotIndex + 1);
+			if (DotIndex != INDEX_NONE)
+			{
+				return InternalVersion.RightChop(DotIndex + 1);
+			}
+		}
+		return InternalVersion;
+	}
+
 	FString GetUnifiedDriverVersion() const
 	{
-		// we use the internal version, not the user version to avoid problem where the name was altered 
+		// We use the internal version as the base to avoid instability or changes to the user-facing driver versioning scheme,
+		// and we remove the parts of the version that identify unimportant factors like Windows or DirectX version.
 		const FString& FullVersion = InternalDriverVersion;
 
-		if(IsNVIDIA() && (InternalDriverVersion != UserDriverVersion))
-		{
-			return TrimNVIDIAInternalVersion(FullVersion);
-		}
-		else if(IsAMD())
+		if (IsNVIDIA() && (InternalDriverVersion != UserDriverVersion))
 		{
-			// examples for AMD: "13.12" "15.101.1007" "13.351"
+			return GetNVIDIAUnifiedVersion(FullVersion);
 		}
-		else if(IsIntel())
+		else if (IsIntel())
 		{
-			// https://www.intel.com/content/www/us/en/support/articles/000005654/graphics.html
-			// Drop off the OS and DirectX version
-			// 27.20.100.8935 -> 100.8935
-			int32 DotIndex = FullVersion.Find(TEXT("."), ESearchCase::CaseSensitive, ESearchDir::FromStart);
-			if (DotIndex != INDEX_NONE)
-			{
-				DotIndex = FullVersion.Find(TEXT("."), ESearchCase::CaseSensitive, ESearchDir::FromStart, DotIndex + 1);
-				if (DotIndex != INDEX_NONE)
-				{
-					return FullVersion.RightChop(DotIndex + 1);
-				}
-			}
+			return GetIntelUnifiedVersion(FullVersion);
 		}
 		return FullVersion;
 	}
 };
 
-// one entry in the Hardware.ini file
-struct FDriverDenyListEntry
+static FDateTime ParseMonthDayYearDate(const FString& DateString)
 {
-	// optional, e.g. "<=223.112.21.1", might includes comparison operators, later even things multiple ">12.22 <=12.44"
-	FString DriverVersionString;
-	// optional, e.g. "<=MM-DD-YYYY"
-	FString DriverDateString;
-	// optional, e.g. "D3D11", "D3D12"
-	FString RHIName;
-	// required
-	FString Reason;
-
-	// @param e.g. "DriverVersion=\"361.43\", Reason=\"UE-25096 Viewport flashes black and white when moving in the scene on latest Nvidia drivers\""
-	// At least one specifier of driver date or version is required to check for a match
-	void LoadFromINIString(const TCHAR* In)
-	{
-		FParse::Value(In, TEXT("DriverVersion="), DriverVersionString);
-		FParse::Value(In, TEXT("DriverDate="), DriverDateString);
-		ensure(!DriverVersionString.IsEmpty() || !DriverDateString.IsEmpty());
+	FString Month, DayYear, Day, Year;
+	DateString.Split(TEXT("-"), &Month, &DayYear);
+	DayYear.Split(TEXT("-"), &Day, &Year);
+	return FDateTime(FCString::Atoi(*Year), FCString::Atoi(*Month), FCString::Atoi(*Day));
+}
 
-		FParse::Value(In, TEXT("RHI="), RHIName);
+// Represents an optional set of constraints that a driver configuration entry can include:
+//   * The RHI being used (e.g. "D3D12" will result in the entry to be considered only when running in D3D12).
+//   * The adapter name, using a regular expression (e.g. ".*RTX.*" will result in the entry to be considered only for cards that contain the string "RTX").
+struct FDriverConfigEntryConstraints {
+	TOptional<FString> RHINameConstraint;
+	TOptional<FRegexPattern> AdapterNameRegexConstraint;
 
+	FDriverConfigEntryConstraints() = default;
 
-		// later:
-//		FParse::Value(In, TEXT("DeviceId="), DeviceId);
-//		FParse::Value(In, TEXT("API="), API);
-//		ensure(API == TEXT("DX11"));
+	FDriverConfigEntryConstraints(const TCHAR* Entry)
+	{
+		FString RHINameString;
+		FParse::Value(Entry, TEXT("RHI="), RHINameString);
+		if (!RHINameString.IsEmpty())
+		{
+			RHINameConstraint = RHINameString;
+		}
 
-		FParse::Value(In, TEXT("Reason="), Reason);
-		ensure(!Reason.IsEmpty());
+		FString AdapterNameRegexString;
+		FParse::Value(Entry, TEXT("AdapterNameRegex="), AdapterNameRegexString);
+		if (!AdapterNameRegexString.IsEmpty())
+		{
+			AdapterNameRegexConstraint.Emplace(AdapterNameRegexString);
+		}
 	}
 
-	// test if the given driver version is mentioned in the this entry
-	// @return true:yes, inform used, false otherwise
-	bool Test(const FGPUDriverInfo& Info) const
+	bool AreConstraintsSatisfied(const FGPUDriverInfo& Info, uint32& OutNumSatisfiedConstraints) const
 	{
-		if (IsValid())
+		OutNumSatisfiedConstraints = 0;
+
+		if (RHINameConstraint)
 		{
-			// If RHI specified, ignore if mismatched
-			if (!RHIName.IsEmpty() && RHIName != Info.RHIName)
+			if (*RHINameConstraint != Info.RHIName)
 			{
 				return false;
 			}
+			OutNumSatisfiedConstraints++;
+		}
 
-			if (!DriverVersionString.IsEmpty())
-			{
-				if (Info.IsSameDriverVersionGeneration(*DriverVersionString))
-				{
-					return CompareStringOp(*DriverVersionString, *Info.GetUnifiedDriverVersion());
-				}
-				else
-				{
-					return false;
-				}
-			}
-			else
+		if (AdapterNameRegexConstraint)
+		{
+			FRegexMatcher Matcher(*AdapterNameRegexConstraint, Info.DeviceDescription);
+			if (!Matcher.FindNext())
 			{
-				FString TempDay, TempMonth, TempMonthDay, TempYear;
-
-				// Trim 1-2 character operator then reformat for comparison
-				int32 OpLength = FChar::IsDigit((*DriverDateString)[1]) ? 1 : 2;
-
-				const TCHAR* DriverDateStringChars = *DriverDateString;
-				EComparisonOp Op = ParseComparisonOp(DriverDateStringChars);
-				if (Op == ECO_Unknown)
-				{
-					Op = ECO_Equal;
-				}
-
-				FString DriverDateComparisonOp = DriverDateString.Left(OpLength);
-				FString TrimmedDriverDateString = DriverDateString.RightChop(OpLength);
-				TrimmedDriverDateString.Split(TEXT("-"), &TempMonthDay, &TempYear, ESearchCase::CaseSensitive, ESearchDir::FromEnd);
-				TempMonthDay.Split(TEXT("-"), &TempMonth, &TempDay, ESearchCase::CaseSensitive, ESearchDir::FromEnd);
-				
-				FDateTime TestDate(FCString::Atoi(*TempYear), FCString::Atoi(*TempMonth), FCString::Atoi(*TempDay));
-
-				Info.DriverDate.Split(TEXT("-"), &TempMonthDay, &TempYear, ESearchCase::CaseSensitive, ESearchDir::FromEnd);
-				TempMonthDay.Split(TEXT("-"), &TempMonth, &TempDay, ESearchCase::CaseSensitive, ESearchDir::FromEnd);
-
-				FDateTime InfoDate(FCString::Atoi(*TempYear), FCString::Atoi(*TempMonth), FCString::Atoi(*TempDay));
-
-				return Compare(InfoDate, Op, TestDate);
+				return false;
 			}
+			OutNumSatisfiedConstraints++;
 		}
 
 		return true;
 	}
 
-	bool IsValid() const
+	bool HasConstraints()
 	{
-		return DriverVersionString.Len() > 1 || DriverDateString.Len() > 1;
+		return RHINameConstraint || AdapterNameRegexConstraint;
 	}
+};
 
-	/**
-	 * Returns true if the latest version of the driver is denied by this entry,
-	 * i.e. the comparison op is > or >=.
-	 */
-	bool IsLatestDenied() const
-	{
-		bool bLatestDenied = false;
-		if (IsValid())
+// This corresponds to one DriverDenyList entry in the Hardware.ini configuration file. One entry includes:
+//   * The drivers being denylisted, specified through one of the following:
+//      * A driver version, along with any comparison operator (e.g. "<473.47" would denylist any driver which version is less than 473.47). This uses the "unified" driver version (see FGPUDriverInfo::GetUnifiedDriverVersion).
+//      * A driver date, along with any comparison operator (e.g. "<=10-31-2023" would denylist any driver which release date is on or before Oct. 31, 2023).
+//   * A set of optional constraints that the denylist entry will consider when the driver is being tested (see FDriverConfigEntryConstraints).
+//   * The reason for the driver(s) being denylisted, e.g.visual glitches, stability, or performance issues.
+struct FDriverDenyListEntry : public FDriverConfigEntryConstraints {
+	struct FDriverDateDenyListEntry {
+		EComparisonOp ComparisonOp;
+		FDateTime Date;
+
+		// Expected format is MM-DD-YYYY.
+		static FDriverDateDenyListEntry FromString(const FString& DriverDateStringWithComparisonOp)
 		{
-			const TCHAR* DriverVersionTchar = !DriverVersionString.IsEmpty() ? *DriverVersionString : *DriverDateString;
-			EComparisonOp ComparisonOp = ParseComparisonOp(DriverVersionTchar);
-			bLatestDenied = (ComparisonOp == ECO_Larger) || (ComparisonOp == ECO_LargerThan);
+			FDriverDateDenyListEntry Entry;
+			const TCHAR* Chars = *DriverDateStringWithComparisonOp;
+			Entry.ComparisonOp = ParseComparisonOp(Chars);
+			Entry.Date = ParseMonthDayYearDate(FString(Chars));
+			return Entry;
 		}
-		return bLatestDenied;
-	}
-};
+	};
 
-struct FGPUHardware
-{
-	// const as this is set in the constructor to be conveniently available
-	const FGPUDriverInfo DriverInfo;
+	struct FDriverVersionDenyListEntry {
+		EComparisonOp ComparisonOp;
+		FDriverVersion Version;
 
-	// constructor
-	FGPUHardware(const FGPUDriverInfo InDriverInfo)
-		: DriverInfo(InDriverInfo)
+		static FDriverVersionDenyListEntry FromString(const FString& DriverVersionStringWithComparisonOp)
+		{
+			FDriverVersionDenyListEntry Entry;
+			const TCHAR* Chars = *DriverVersionStringWithComparisonOp;
+			Entry.ComparisonOp = ParseComparisonOp(Chars);
+			Entry.Version.Parse(Chars);
+			return Entry;
+		}
+	};
+
+	FDriverDenyListEntry() = default;
+
+	FDriverDenyListEntry(const TCHAR* Entry)
+		: FDriverConfigEntryConstraints(Entry)
 	{
-		// tests (should be very fast)
-#if !(UE_BUILD_SHIPPING || UE_BUILD_TEST)
+		FString DriverVersionString;
+		FParse::Value(Entry, TEXT("DriverVersion="), DriverVersionString);
+		if (!DriverVersionString.IsEmpty())
 		{
-			FMultiInt<2> A;
-			check(A.Value[0] == 0 && A.Value[1] == 0 && A.Size == 2);
-			A.GetValue(TEXT("18.98"));
-			check(A.Value[0] == 18 && A.Value[1] == 98);
-			A.GetValue(TEXT(""));
-			check(A.Value[0] == 0 && A.Value[1] == 0);
-			A.GetValue(TEXT("98"));
-			check(A.Value[0] == 0 && A.Value[1] == 98);
-			A.GetValue(TEXT("98.34.56"));
-			check(A.Value[0] == 98 && A.Value[1] == 34);
-			A.GetValue(TEXT(" 98 . 034 "));
-			check(A.Value[0] == 98 && A.Value[1] == 34);
-			A.GetValue(TEXT("\t 98\t.\t34\t"));
-			check(A.Value[0] == 98 && A.Value[1] == 34);
-		
-			check(FMultiInt<2>(TEXT("3.07")) == FMultiInt<2>(TEXT("3.07")));
-			check(FMultiInt<2>(TEXT("3.05")) < FMultiInt<2>(TEXT("3.07")));
-			check(FMultiInt<2>(TEXT("3.05")) <= FMultiInt<2>(TEXT("3.07")));
-			check(FMultiInt<2>(TEXT("3.07")) <= FMultiInt<2>(TEXT("3.07")));
-			check(FMultiInt<2>(TEXT("3.08")) > FMultiInt<2>(TEXT("3.07")));
-			check(FMultiInt<2>(TEXT("3.08")) >= FMultiInt<2>(TEXT("3.07")));
-			check(FMultiInt<2>(TEXT("3.07")) >= FMultiInt<2>(TEXT("3.07")));
-			check(FMultiInt<2>(TEXT("3.05")) != FMultiInt<2>(TEXT("3.07")));
-			check(FMultiInt<2>(TEXT("4.05")) > FMultiInt<2>(TEXT("3.07")));
-			check(FMultiInt<2>(TEXT("4.05")) >= FMultiInt<2>(TEXT("3.07")));
-			check(FMultiInt<2>(TEXT("2.05")) < FMultiInt<2>(TEXT("3.07")));
-			check(FMultiInt<2>(TEXT("2.05")) <= FMultiInt<2>(TEXT("3.07")));
-			
-			check(Compare(10, ECO_Equal, 10));
-			check(Compare(10, ECO_NotEqual, 20));
-			check(Compare(20, ECO_Larger, 10));
-			check(Compare(20, ECO_LargerThan, 10));
-			check(Compare(10, ECO_LargerThan, 10));
-			check(Compare(10, ECO_Smaller, 20));
-			check(Compare(10, ECO_SmallerThan, 10));
-
-			check(CompareStringOp(TEXT("<20.10"), TEXT("19.12")));
-			check(CompareStringOp(TEXT("<=20.10"), TEXT("19.12")));
-			check(CompareStringOp(TEXT("<=19.12"), TEXT("19.12")));
-			check(CompareStringOp(TEXT("==19.12"), TEXT("19.12")));
-			check(CompareStringOp(TEXT(">=19.12"), TEXT("19.12")));
-			check(CompareStringOp(TEXT(">=10.12"), TEXT("19.12")));
-			check(CompareStringOp(TEXT("!=20.12"), TEXT("19.12")));
-			check(CompareStringOp(TEXT(">10.12"), TEXT("19.12")));
+			DriverVersion = FDriverVersionDenyListEntry::FromString(DriverVersionString);
+		}
 
-			{
-				FGPUDriverInfo Version;
-				Version.SetNVIDIA();
-				check(Version.IsNVIDIA());
-				check(!Version.IsAMD());
-				check(!Version.IsIntel());
-				Version.InternalDriverVersion = TEXT("10.18.13.4788");
-				check(Version.GetUnifiedDriverVersion() == TEXT("347.88"));
-			}
-			{
-				FGPUDriverInfo Version;
-				Version.SetAMD();
-				check(Version.IsAMD());
-				check(!Version.IsNVIDIA());
-				check(!Version.IsIntel());
-				Version.InternalDriverVersion = TEXT("15.200.1062.1004");
-				check(Version.GetUnifiedDriverVersion() == TEXT("15.200.1062.1004"));
-			}
-			{
-				FGPUDriverInfo Version;
-				Version.SetIntel();
-				check(Version.IsIntel());
-				check(!Version.IsAMD());
-				check(!Version.IsNVIDIA());
-				Version.DeviceDescription = TEXT("Intel(R) HD Graphics 4600");
-				Version.InternalDriverVersion = TEXT("9.18.10.3310");
-				Version.DriverDate = TEXT("9-17-2013");
-				check(Version.GetUnifiedDriverVersion() == TEXT("10.3310"));
-			}
+		FString DriverDateString;
+		FParse::Value(Entry, TEXT("DriverDate="), DriverDateString);
+		if (!DriverDateString.IsEmpty())
+		{
+			DriverDate = FDriverDateDenyListEntry::FromString(DriverDateString);
 		}
-#endif// !(UE_BUILD_SHIPPING || UE_BUILD_TEST)
-	}	
-	
-	// @return a driver version intended to be shown to the user e.g. "15.30.1025.1001 12/17/2015 (Crimson Edition 15.12)"
-	FString GetSuggestedDriverVersion(const FString& InRHIName) const
+
+		ensureMsgf(IsValid(), TEXT("Exactly one of driver date or driver version must be specified in a driver denylist entry"));
+
+		FParse::Value(Entry, TEXT("Reason="), DenylistReason);
+		ensure(!DenylistReason.IsEmpty());
+	}
+
+	// Checks whether the given driver is denylisted by this entry, and also returns the number
+	// of constraints that have been satisfied in OutNumSatisfiedConstraints.
+	bool AppliesToDriver(const FGPUDriverInfo& Info, uint32& OutNumSatisfiedConstraints) const
 	{
-		const FString Section = GetVendorSectionName();
+		OutNumSatisfiedConstraints = 0;
 
-		FString Ret;
+		if (!IsValid())
+		{
+			return false;
+		}
 
-		if(!Section.IsEmpty())
+		// Check constraints first. Only apply the driver version/date check
+		// if all constraints are satisfied.
+		if (!AreConstraintsSatisfied(Info, OutNumSatisfiedConstraints))
 		{
-			TArray<FString> SuggestedDriverVersions;
-			GConfig->GetArray(*Section, TEXT("SuggestedDriverVersion"), SuggestedDriverVersions, GHardwareIni);
+			return false;
+		}
 
-			// Find specific RHI version first
-			if (InRHIName.Len() > 0)
-			{
-				for (const FString& SuggestedVersion : SuggestedDriverVersions)
-				{
-					int32 Found = SuggestedVersion.Find(InRHIName);
-					if (Found != INDEX_NONE && Found > 0 && SuggestedVersion[Found - 1] == ';')
-					{
-						Ret = SuggestedVersion.Left(Found - 1);
-						return Ret;
-					}
-				}
-			}
+		if (DriverVersion && Info.IsSameDriverVersionGeneration(DriverVersion->Version))
+		{
+			FDriverVersion CurrentDriverVersion(Info.GetUnifiedDriverVersion());
+			return CompareWithOp(CurrentDriverVersion, DriverVersion->ComparisonOp, DriverVersion->Version);
+		}
 
-			// Return the first generic one
-			for (const FString& SuggestedVersion : SuggestedDriverVersions)
-			{
-				int32 Found = 0;
-				if (!SuggestedVersion.FindChar(TEXT(';'), Found))
-				{
-					return SuggestedVersion;
-				}
-			}
+		if (DriverDate)
+		{
+			FDateTime CurrentDriverDate = ParseMonthDayYearDate(Info.DriverDate);
+			return CompareWithOp(CurrentDriverDate, DriverDate->ComparisonOp, DriverDate->Date);
 		}
 
-		return Ret;
+		return false;
 	}
 
-	// @return 0 if there is none
-	FDriverDenyListEntry FindDriverDenyListEntry() const
+	// Returns true whether this denylist entry will always apply to the latest drivers available,
+	// i.e. the comparison function is > or >=.
+	bool AppliesToLatestDrivers()
 	{
-		const FString Section = GetVendorSectionName();
-
-		if(!Section.IsEmpty())
+		if (DriverVersion)
 		{
-			TArray<FString> DenyListStrings;
-			GConfig->GetArray(*Section, TEXT("DriverDenyList"), DenyListStrings, GHardwareIni);
+			return DriverVersion->ComparisonOp == ECO_Greater || DriverVersion->ComparisonOp == ECO_GreaterOrEqual;
+		}
+		if (DriverDate)
+		{
+			return DriverDate->ComparisonOp == ECO_Greater || DriverDate->ComparisonOp == ECO_GreaterOrEqual;
+		}
+		return false;
+	}
 
-			for(int32 i = 0; i < DenyListStrings.Num(); ++i)
-			{
-				FDriverDenyListEntry Entry;
+	bool IsValid() const
+	{
+		return DriverVersion.IsSet() != DriverDate.IsSet();
+	}
 
-				const TCHAR* Line = *DenyListStrings[i];
-			
-				ensure(Line[0] == TCHAR('('));
+	TOptional<FDriverVersionDenyListEntry> DriverVersion;
+	TOptional<FDriverDateDenyListEntry> DriverDate;
 
-				Entry.LoadFromINIString(&Line[1]);
+	FString	DenylistReason;
+};
 
-				if(Entry.Test(DriverInfo))
-				{
-					return Entry;
-				}
-			}
+// This corresponds to one SuggestedDriverVersion entry in the Hardware.ini configuration file. One entry includes:
+//   * The suggested driver version.
+//   * A set of optional constraints (see FDriverConfigEntryConstraints).
+struct FSuggestedDriverEntry : public FDriverConfigEntryConstraints {
+
+	FSuggestedDriverEntry() = default;
+	FSuggestedDriverEntry(const TCHAR* Entry)
+		: FDriverConfigEntryConstraints(Entry)
+	{
+		// For backwards compatibility, accept either "DriverVersion=..." or
+		// just the version string without the structured format.
+		FParse::Value(Entry, TEXT("DriverVersion="), SuggestedDriverVersion);
+		if (SuggestedDriverVersion.IsEmpty() && !HasConstraints() && !FString(Entry).Contains(TEXT("=")))
+		{
+			SuggestedDriverVersion = Entry;
 		}
+	}
+
+	// Checks whether this entry applies to the given driver info, and also returns the number
+	// of constraints that have been satisfied in OutNumSatisfiedConstraints.
+	bool AppliesToDriver(const FGPUDriverInfo& Info, uint32& OutNumSatisfiedConstraints)
+	{
+		return IsValid() && AreConstraintsSatisfied(Info, OutNumSatisfiedConstraints);
+	}
+
+	bool IsValid()
+	{
+		return !SuggestedDriverVersion.IsEmpty();
+	}
+
+	FString SuggestedDriverVersion;
+};
 
-		return FDriverDenyListEntry();
+class FGPUDriverHelper
+{
+public:
+	FGPUDriverHelper(const FGPUDriverInfo& InDriverInfo)
+		: DriverInfo(InDriverInfo)
+	{
+	}
+
+	// Returns the best suggested driver version that matches the current configuration, if it exists. 
+	TOptional<FSuggestedDriverEntry> FindSuggestedDriverVersion() const
+	{
+		return GetBestConfigEntry<FSuggestedDriverEntry>(TEXT("SuggestedDriverVersion"));
 	}
 
-	/**
-	 * Returns true if the latest version of the driver is on the deny list.
-	 */
-	bool IsLatestDenied() const
+	// Finds the best driver denylist entry that matches the current driver and configuration, if it exists. 
+	TOptional<FDriverDenyListEntry> FindDriverDenyListEntry() const
+	{
+		return GetBestConfigEntry<FDriverDenyListEntry>(TEXT("DriverDenyList"));
+	}
+
+private:
+	template <typename TEntryType>
+	TOptional<TEntryType> GetBestConfigEntry(const TCHAR* ConfigEntryName) const
 	{
-		bool bLatestDenied = false;
 		const FString Section = GetVendorSectionName();
 
-		if(!Section.IsEmpty())
+		if (Section.IsEmpty())
 		{
-			TArray<FString> DenyListStrings;
-			GConfig->GetArray(*Section, TEXT("DriverDenyList"), DenyListStrings, GHardwareIni);
-
-			for(int32 i = 0; !bLatestDenied && i < DenyListStrings.Num(); ++i)
-			{
-				FDriverDenyListEntry Entry;
+			return {};
+		}
 
-				const TCHAR* Line = *DenyListStrings[i];
-			
-				ensure(Line[0] == TCHAR('('));
+		// Multiple entries can match, but some entries may match on more
+		// constraints (e.g. both RHI and adapter name). Choose the entry
+		// that satisfies the most constraints.
+		TPair<TEntryType, uint32> MostRelevantEntry;
 
-				Entry.LoadFromINIString(&Line[1]);
+		TArray<FString> EntryStrings;
+		GConfig->GetArray(*Section, ConfigEntryName, EntryStrings, GHardwareIni);
+		for (const FString& EntryString : EntryStrings)
+		{
+			ensure(!EntryString.IsEmpty());
+			TEntryType Entry(*EntryString);
 
-				bLatestDenied |= Entry.IsLatestDenied();
+			uint32 NumSatisfiedConstraints = 0;
+			if (Entry.AppliesToDriver(DriverInfo, NumSatisfiedConstraints))
+			{
+				if (NumSatisfiedConstraints >= MostRelevantEntry.Value)
+				{
+					MostRelevantEntry = { Entry, NumSatisfiedConstraints };
+				}
 			}
 		}
-		return bLatestDenied;
+
+		if (MostRelevantEntry.Key.IsValid())
+		{
+			return MostRelevantEntry.Key;
+		}
+		return {};
 	}
 
-	// to get a section name in the Hardware.ini file
-	// @return 0 if not found
+	// Get the corresponding vendor's section name in the Hardware.ini file.
+	// Returns an empty string if not found.
 	FString GetVendorSectionName() const
 	{
 		const TCHAR* Section = nullptr;
 
-		if(DriverInfo.IsNVIDIA())
+		if (DriverInfo.IsNVIDIA())
 		{
 			Section = TEXT("GPU_NVIDIA");
 		}
-		if(DriverInfo.IsAMD())
+		if (DriverInfo.IsAMD())
 		{
 			Section = TEXT("GPU_AMD");
 		}
-		else if(DriverInfo.IsIntel())
+		else if (DriverInfo.IsIntel())
 		{
 			Section = TEXT("GPU_Intel");
-		}	
-		// more GPU vendors can be added on demand
+		}
 		if (!Section)
 		{
 			return TEXT("");
@@ -695,6 +561,6 @@ struct FGPUHardware
 
 		return FString::Printf(TEXT("%s %s"), Section, ANSI_TO_TCHAR(FPlatformProperties::IniPlatformName()));
 	}
-};
-
 
+	FGPUDriverInfo DriverInfo;
+};
diff --git a/Engine/Source/Runtime/Core/Tests/GenericPlatform/GenericPlatformDriverTest.cpp b/Engine/Source/Runtime/Core/Tests/GenericPlatform/GenericPlatformDriverTest.cpp
new file mode 100644
--- /dev/null
+++ b/Engine/Source/Runtime/Core/Tests/GenericPlatform/GenericPlatformDriverTest.cpp
@@ -0,0 +1,303 @@
+// Copyright Epic Games, Inc. All Rights Reserved.
+
+#if WITH_TESTS 
+
+#include "GenericPlatform/GenericPlatformDriver.h"
+#include "Tests/TestHarnessAdapter.h"
+
+static FString MakeVersionString(const TArray<uint32> Values)
+{
+	if (Values.IsEmpty())
+	{
+		return TEXT("");
+	}
+	FString VersionString;
+	for (int32 Idx = 0; Idx < Values.Num(); ++Idx)
+	{
+		VersionString += FString::FromInt(Values[Idx]);
+		if (Idx != Values.Num() - 1)
+		{
+			VersionString += TEXT(".");
+		}
+	}
+	return VersionString;
+}
+
+TEST_CASE_NAMED(DriverVersionTest, "System::Core::GenericPlatform::Driver", "[ApplicationContextMask][SmokeFilter]")
+{
+	SECTION("Driver version parsing")
+	{
+		{
+			FDriverVersion Version;
+			CHECK(Version.GetVersionValues().Num() == 0);
+		}
+		{
+			// All of these should parse the same.
+			TArray<FDriverVersion> Versions;
+			Versions.Emplace(MakeVersionString({ 1,2,3,4 }));
+			Versions.Emplace(TEXT("1.2.3.4"));
+			Versions.Emplace(TEXT("  1.2.3.4"));
+			Versions.Emplace(TEXT("1.2.3.4  "));
+			Versions.Emplace(TEXT("1 . 2. 3 .4"));
+
+			for (const FDriverVersion& Version : Versions)
+			{
+				CHECK(Version.GetVersionValues().Num() == 4);
+				CHECK(Version.GetVersionValue(0) == 1);
+				CHECK(Version.GetVersionValue(1) == 2);
+				CHECK(Version.GetVersionValue(2) == 3);
+				CHECK(Version.GetVersionValue(3) == 4);
+			}
+		}
+	}
+
+	SECTION("Driver version comparisons")
+	{
+		{
+			FDriverVersion A(TEXT("1.2.3.4"));
+			FDriverVersion B(TEXT("1.2.3.4"));
+			CHECK(A == B);
+			CHECK(A <= B);
+			CHECK(A >= B);
+			CHECK(!(A != B));
+			CHECK(!(A < B));
+			CHECK(!(A > B));
+		}
+		{
+			TArray<uint32> BaseVersion = { 1,2,3,4 };
+
+			for (int32 Idx = 0; Idx < 4; ++Idx)
+			{
+				// Make the first version larger by incrementing each dot version.
+				TArray<uint32> LargerVersion = BaseVersion;
+				LargerVersion[Idx] += 1;
+
+				FDriverVersion A(MakeVersionString(LargerVersion));
+				FDriverVersion B(MakeVersionString(BaseVersion));
+				CHECK(A != B);
+				CHECK(A > B);
+				CHECK(A >= B);
+				CHECK(!(A == B));
+				CHECK(!(A < B));
+				CHECK(!(A <= B));
+			}
+		}
+		{
+			TArray<uint32> BaseVersion = { 1,2,3,4 };
+
+			for (int32 Idx = 0; Idx < 4; ++Idx)
+			{
+				// Make the first version smaller by decrementing each dot version.
+				TArray<uint32> SmallerVersion = BaseVersion;
+				SmallerVersion[Idx] -= 1;
+
+				FDriverVersion A(MakeVersionString(SmallerVersion));
+				FDriverVersion B(MakeVersionString(BaseVersion));
+				CHECK(A != B);
+				CHECK(A < B);
+				CHECK(A <= B);
+				CHECK(!(A == B));
+				CHECK(!(A > B));
+				CHECK(!(A >= B));
+			}
+		}
+	}
+
+	SECTION("Unified driver version parsing")
+	{
+		{
+			FGPUDriverInfo Info;
+			Info.SetNVIDIA();
+			Info.InternalDriverVersion = TEXT("31.0.15.5123");
+			CHECK(Info.GetUnifiedDriverVersion() == TEXT("551.23"));
+		}
+		{
+			FGPUDriverInfo Info;
+			Info.SetIntel();
+			Info.InternalDriverVersion = TEXT("9.18.10.3310");
+			CHECK(Info.GetUnifiedDriverVersion() == TEXT("10.3310"));
+		}
+		{
+			FGPUDriverInfo Info;
+			Info.SetAMD();
+			Info.InternalDriverVersion = TEXT("15.200.1062.1004");
+			CHECK(Info.GetUnifiedDriverVersion() == TEXT("15.200.1062.1004"));
+		}
+		{
+			FGPUDriverInfo Info;
+			Info.SetIntel();
+			Info.InternalDriverVersion = TEXT("9.18.10.3310");
+			CHECK(Info.IsSameDriverVersionGeneration(FDriverVersion(TEXT("15.3501"))));
+			CHECK(!Info.IsSameDriverVersionGeneration(FDriverVersion(TEXT("101.5124"))));
+
+			Info.InternalDriverVersion = TEXT("30.0.101.1404");
+			CHECK(!Info.IsSameDriverVersionGeneration(FDriverVersion(TEXT("15.3501"))));
+			CHECK(Info.IsSameDriverVersionGeneration(FDriverVersion(TEXT("100.5124"))));
+		}
+	}
+
+	SECTION("Driver denylist and suggested entry parsing")
+	{
+		{
+			FDriverDenyListEntry Entry(TEXT("(DriverVersion=\"301.12\", RHI=\"D3D12\", AdapterNameRegex=\".*4080.*\", Reason=\"Test version\")"));
+			CHECK(Entry.IsValid());
+			CHECK(Entry.DenylistReason == TEXT("Test version"));
+			CHECK(*Entry.RHINameConstraint == TEXT("D3D12"));
+			CHECK(Entry.AdapterNameRegexConstraint);
+			CHECK(Entry.DriverVersion->ComparisonOp == EComparisonOp::ECO_Equal);
+			CHECK(Entry.DriverVersion->Version == FDriverVersion(MakeVersionString({ 301, 12 })));
+		}
+		{
+			FDriverDenyListEntry Entry(TEXT("(DriverVersion=\"<=311.16\", Reason=\"Test with no constraints\")"));
+			CHECK(Entry.IsValid());
+			CHECK(Entry.DenylistReason == TEXT("Test with no constraints"));
+			CHECK(Entry.DriverVersion->ComparisonOp == EComparisonOp::ECO_LessOrEqual);
+			CHECK(Entry.DriverVersion->Version == FDriverVersion(MakeVersionString({ 311, 16 })));
+		}
+		{
+			FDriverDenyListEntry Entry(TEXT("(DriverDate=\">=12-31-2021\", RHI=\"D3D11\", Reason=\"Test date\")"));
+			CHECK(Entry.IsValid());
+			CHECK(Entry.DenylistReason == TEXT("Test date"));
+			CHECK(*Entry.RHINameConstraint == TEXT("D3D11"));
+			CHECK(Entry.DriverDate->ComparisonOp == EComparisonOp::ECO_GreaterOrEqual);
+			CHECK(Entry.DriverDate->Date == FDateTime(2021, 12, 31));
+			CHECK(Entry.AppliesToLatestDrivers());
+		}
+		{
+			FSuggestedDriverEntry Entry(TEXT("(DriverVersion=\"512.15\", RHI=\"D3D11\", AdapterNameRegex=\".*4050.*\")"));
+			CHECK(Entry.IsValid());
+			CHECK(*Entry.RHINameConstraint == TEXT("D3D11"));
+			CHECK(Entry.AdapterNameRegexConstraint);
+			CHECK(Entry.SuggestedDriverVersion == TEXT("512.15"));
+		}
+		{
+			FSuggestedDriverEntry Entry(TEXT("(DriverVersion=\"101.56\")"));
+			CHECK(Entry.IsValid());
+			CHECK(Entry.SuggestedDriverVersion == TEXT("101.56"));
+		}
+		{
+			FSuggestedDriverEntry Entry(TEXT("111.56"));
+			CHECK(Entry.IsValid());
+			CHECK(Entry.SuggestedDriverVersion == TEXT("111.56"));
+		}
+		{
+			FSuggestedDriverEntry Entry(TEXT("BadString=111.56"));
+			CHECK(!Entry.IsValid());
+		}
+	}
+
+
+	SECTION("Driver denylist and suggested driver entry retrieval")
+	{
+		{
+			GConfig->RemoveKey(TEXT("GPU_NVIDIA Windows"), TEXT("DriverDenylist"), GHardwareIni);
+
+			FGPUDriverInfo Info;
+			Info.SetNVIDIA();
+			Info.RHIName = TEXT("D3D12");
+			Info.InternalDriverVersion = TEXT("31.0.15.5123"); // 551.23
+
+			TArray<FString> DenylistEntries = {
+				TEXT("(DriverVersion=\"551.23\", RHI=\"D3D11\", Reason=\"Test\")"), // Does not apply due to RHI.
+				TEXT("(DriverVersion=\"<550.00\", Reason=\"Test\")"), // Does not apply due to version.
+				TEXT("(DriverVersion=\"<=560\", Reason=\"Test\")"), // Applies (no constraints).
+				TEXT("(DriverVersion=\"<555.55\", RHI=\"D3D12\", Reason=\"Test\")"), // Applies with RHI constraint.
+			};
+
+			GConfig->SetArray(TEXT("GPU_NVIDIA Windows"), TEXT("DriverDenylist"), DenylistEntries, GHardwareIni);
+
+			FGPUDriverHelper Helper(Info);
+			TOptional<FDriverDenyListEntry> Entry = Helper.FindDriverDenyListEntry();
+			CHECK(Entry);
+			CHECK(Entry->IsValid());
+			CHECK(Entry->RHINameConstraint);
+			CHECK(Entry->DriverVersion->ComparisonOp == EComparisonOp::ECO_Less);
+			CHECK(Entry->DriverVersion->Version == FDriverVersion(MakeVersionString({ 555, 55 })));
+		}
+
+		{
+			GConfig->RemoveKey(TEXT("GPU_AMD Windows"), TEXT("DriverDenylist"), GHardwareIni);
+
+			FGPUDriverInfo Info;
+			Info.SetAMD();
+			Info.RHIName = TEXT("D3D11");
+			Info.InternalDriverVersion = TEXT("31.0.15.5123");
+			Info.DriverDate = TEXT("1-15-2019");
+
+			TArray<FString> DenylistEntries = {
+				TEXT("(DriverDate=\"<=1-1-2020\", RHI=\"D3D12\", Reason=\"Test\")"), // Does not apply due to RHI.
+				TEXT("(DriverDate=\"<5-5-2018\", Reason=\"Test\")"), // Does not apply due to date.
+				TEXT("(DriverDate=\"<=12-31-2019\", Reason=\"Test\")"), // Applies (no constraints).
+				TEXT("(DriverDate=\"<1-30-2019\", RHI=\"D3D11\", Reason=\"Test\")"), // Applies with RHI constraint.
+			};
+
+			GConfig->SetArray(TEXT("GPU_AMD Windows"), TEXT("DriverDenylist"), DenylistEntries, GHardwareIni);
+
+			FGPUDriverHelper Helper(Info);
+			TOptional<FDriverDenyListEntry> Entry = Helper.FindDriverDenyListEntry();
+			CHECK(Entry);
+			CHECK(Entry->IsValid());
+			CHECK(Entry->RHINameConstraint);
+			CHECK(Entry->DriverDate->ComparisonOp == EComparisonOp::ECO_Less);
+			CHECK(Entry->DriverDate->Date == FDateTime(2019, 1, 30));
+		}
+
+		{
+			GConfig->RemoveKey(TEXT("GPU_Intel Windows"), TEXT("SuggestedDriverVersion"), GHardwareIni);
+
+			FGPUDriverInfo Info;
+			Info.SetIntel();
+			Info.RHIName = TEXT("D3D12");
+			Info.InternalDriverVersion = TEXT("30.0.101.1404");
+
+			TArray<FString> SuggestedDriverEntries = {
+				TEXT("(DriverVersion=\"30.0.101.5000\", RHI=\"D3D11\")"), // Does not apply due to RHI.
+				TEXT("(DriverVersion=\"30.0.101.1255\")"), // Applies (no constraints).
+				TEXT("(DriverVersion=\"30.0.102.1200\", RHI=\"D3D12\")"), // Applies with RHI constraint.
+			};
+
+			GConfig->SetArray(TEXT("GPU_Intel Windows"), TEXT("SuggestedDriverVersion"), SuggestedDriverEntries, GHardwareIni);
+
+			FGPUDriverHelper Helper(Info);
+			TOptional<FSuggestedDriverEntry> Entry = Helper.FindSuggestedDriverVersion();
+			CHECK(Entry);
+			CHECK(Entry->IsValid());
+			CHECK(Entry->RHINameConstraint);
+			CHECK(Entry->SuggestedDriverVersion == TEXT("30.0.102.1200"));
+		}
+
+		// The following test requires regex support, which may be disabled in tests.
+#if UE_ENABLE_ICU 
+		{
+			GConfig->RemoveKey(TEXT("GPU_NVIDIA Windows"), TEXT("DriverDenylist"), GHardwareIni);
+
+			FGPUDriverInfo Info;
+			Info.SetNVIDIA();
+			Info.RHIName = TEXT("D3D12");
+			Info.DeviceDescription = TEXT("NVIDIA GeForce RTX 4070 ");
+			Info.InternalDriverVersion = TEXT("31.0.15.5123"); // 551.23
+
+			TArray<FString> DenylistEntries = {
+				TEXT("(DriverVersion=\"551.23\", RHI=\"D3D11\", Reason=\"Test\")"), // Does not apply due to RHI.
+				TEXT("(DriverVersion=\">550.00\", RHI=\"D3D12\", AdapterNameRegex=\".*4080.*\", Reason=\"Test\")"), // Does not apply due to name regex.
+				TEXT("(DriverVersion=\"<550.00\", Reason=\"Test\")"), // Does not apply due to version.
+				TEXT("(DriverVersion=\"<555.55\", RHI=\"D3D12\", Reason=\"Test\")"), // Applies with RHI constraint.
+				TEXT("(DriverVersion=\"<=551.23\", RHI=\"D3D12\", AdapterNameRegex=\".*4070.*\" Reason=\"Test\")"), // Applies with RHI and name regex constraints.
+			};
+
+			GConfig->SetArray(TEXT("GPU_NVIDIA Windows"), TEXT("DriverDenylist"), DenylistEntries, GHardwareIni);
+
+			FGPUDriverHelper Helper(Info);
+			TOptional<FDriverDenyListEntry> Entry = Helper.FindDriverDenyListEntry();
+			CHECK(Entry);
+			CHECK(Entry->IsValid());
+			CHECK(Entry->RHINameConstraint);
+			CHECK(Entry->AdapterNameRegexConstraint);
+			CHECK(Entry->DriverVersion->ComparisonOp == EComparisonOp::ECO_LessOrEqual);
+			CHECK(Entry->DriverVersion->Version == FDriverVersion(MakeVersionString({ 551, 23 })));
+		}
+#endif
+	}
+}
+
+#endif //WITH_TESTS
\ No newline at end of file
diff --git a/Engine/Source/Runtime/RHI/Private/DynamicRHI.cpp b/Engine/Source/Runtime/RHI/Private/DynamicRHI.cpp
--- a/Engine/Source/Runtime/RHI/Private/DynamicRHI.cpp
+++ b/Engine/Source/Runtime/RHI/Private/DynamicRHI.cpp
@@ -36,33 +36,25 @@ FDynamicRHI* GDynamicRHI = NULL;
 RHIGetGPUUsageType RHIGetGPUUsage = nullptr;
 bool bDriverDenylistMessageShown = false;
 
-static TAutoConsoleVariable<int32> CVarWarnOfBadDrivers(
+static int32 GWarnOfBadDrivers = true;
+static FAutoConsoleVariableRef CVarWarnOfBadDrivers(
 	TEXT("r.WarnOfBadDrivers"),
-	1,
-	TEXT("On engine startup we can check the current GPU driver and warn the user about issues and suggest a specific version\n")
-	TEXT("The test is fast so this should not cost any performance.\n")
+	GWarnOfBadDrivers,
+	TEXT("Check the current GPU driver on engine startup, warn the user about issues and suggest a specific version.\n")
+	TEXT("The driver denylist is used to check for bad drivers according to their release date and/or driver versions.\n")
 	TEXT(" 0: off\n")
-	TEXT(" 1: a message on startup might appear (default)\n")
-	TEXT(" 2: Simulating the system has a NVIDIA driver on the deny list (UI should appear)\n")
-	TEXT(" 3: Simulating the system has a AMD driver on the deny list (UI should appear)\n")
-	TEXT(" 4: Simulating the system has an allowed AMD driver (no UI should appear)\n")
-	TEXT(" 5: Simulating the system has a Intel driver (no UI should appear)"),
+	TEXT(" 1: check the driver and display a pop-up message if the driver is denylisted (default)"),
 	ECVF_RenderThreadSafe
 	);
 
-static TAutoConsoleVariable<int32> CVarBadDriverWarningIsFatal(
+static bool GBadDriverWarningIsFatal = false;
+static FAutoConsoleVariableRef CVarBadDriverWarningIsFatal(
 	TEXT("r.BadDriverWarningIsFatal"),
-	0,
-	TEXT("If non-zero, trigger a fatal error when warning of bad drivers.\n")
+	GBadDriverWarningIsFatal,
+	TEXT("If non-zero, trigger a fatal error if a denylisted driver is detected.\n")
 	TEXT("For the fatal error to occur, r.WarnOfBadDrivers must be non-zero.\n")
 	TEXT(" 0: off (default)\n")
-	TEXT(" 1: a fatal error occurs after the out of date driver message is dismissed\n"),
-	ECVF_RenderThreadSafe);
-
-static TAutoConsoleVariable<int32> CVarDisableDriverWarningPopupIfGFN(
-	TEXT("r.DisableDriverWarningPopupIfGFN"),
-	1,
-	TEXT("If non-zero, disable driver version warning popup if running on a GFN cloud machine."),
+	TEXT(" 1: a fatal error occurs after the out-of-date driver message is dismissed (non-Shipping only)\n"),
 	ECVF_RenderThreadSafe);
 
 void InitNullRHI()
@@ -101,11 +93,7 @@ void RHIDetectAndWarnOfBadDrivers(bool bHasEditorToken)
 		return;
 	}
 
-	int32 WarnMode = CVarWarnOfBadDrivers.GetValueOnGameThread();
-
 	FGPUDriverInfo DriverInfo;
-
-	// later we should make the globals use the struct directly
 	DriverInfo.VendorId = GRHIVendorId;
 	DriverInfo.DeviceDescription = GRHIAdapterName;
 	DriverInfo.ProviderName = TEXT("Unknown");
@@ -114,50 +102,22 @@ void RHIDetectAndWarnOfBadDrivers(bool bHasEditorToken)
 	DriverInfo.DriverDate = GRHIAdapterDriverDate;
 	DriverInfo.RHIName = GDynamicRHI ? GDynamicRHI->GetName() : FString();
 
-#if !(UE_BUILD_SHIPPING || UE_BUILD_TEST)
-	// for testing
-	if(WarnMode == 2)
-	{
-		DriverInfo.SetNVIDIA();
-		DriverInfo.DeviceDescription = TEXT("Test NVIDIA (bad)");
-		DriverInfo.UserDriverVersion = TEXT("346.43");
-		DriverInfo.InternalDriverVersion = TEXT("9.18.134.643");
-		DriverInfo.DriverDate = TEXT("01-01-1900");
-	}
-	else if(WarnMode == 3)
-	{
-		DriverInfo.SetAMD();
-		DriverInfo.DeviceDescription = TEXT("Test AMD (bad)");
-		DriverInfo.UserDriverVersion = TEXT("Test Catalyst Version");
-		DriverInfo.InternalDriverVersion = TEXT("13.152.1.1000");
-		DriverInfo.DriverDate = TEXT("09-10-13");
-	}
-	else if(WarnMode == 4)
-	{
-		DriverInfo.SetAMD();
-		DriverInfo.DeviceDescription = TEXT("Test AMD (good)");
-		DriverInfo.UserDriverVersion = TEXT("Test Catalyst Version");
-		DriverInfo.InternalDriverVersion = TEXT("15.30.1025.1001");
-		DriverInfo.DriverDate = TEXT("01-01-16");
-	}
-	else if(WarnMode == 5)
-	{
-		DriverInfo.SetIntel();
-		DriverInfo.DeviceDescription = TEXT("Test Intel (good)");
-		DriverInfo.UserDriverVersion = TEXT("Test Intel Version");
-		DriverInfo.InternalDriverVersion = TEXT("8.15.10.2302");
-		DriverInfo.DriverDate = TEXT("01-01-15");
-	}
-#endif // !(UE_BUILD_SHIPPING || UE_BUILD_TEST)
-
-	FGPUHardware DetectedGPUHardware(DriverInfo);
+	FGPUDriverHelper DetectedGPUHardware(DriverInfo);
 
 	// Pre-GCN GPUs usually don't support updating to latest driver
 	// But it is unclear what is the latest version supported as it varies from card to card
 	// So just don't complain if pre-gcn
 	if (DriverInfo.IsValid() && !GRHIDeviceIsAMDPreGCNArchitecture)
 	{
-		FDriverDenyListEntry DenyListEntry = DetectedGPUHardware.FindDriverDenyListEntry();
+		TOptional<FDriverDenyListEntry> DenyListEntry = DetectedGPUHardware.FindDriverDenyListEntry();
+
+		GRHIAdapterDriverOnDenyList = DenyListEntry.IsSet() && DenyListEntry->IsValid();
+		FGenericCrashContext::SetEngineData(TEXT("RHI.DriverDenylisted"), GRHIAdapterDriverOnDenyList ? TEXT("true") : TEXT("false"));
+
+		if(!GRHIAdapterDriverOnDenyList)
+		{
+			return;
+		}
 
 		TArray<FString> DeviceCanUpdateDriverList;
 		GConfig->GetArray(TEXT("Devices"), TEXT("DeviceCanUpdateDriverList"), DeviceCanUpdateDriverList, GHardwareIni);
@@ -186,22 +146,13 @@ void RHIDetectAndWarnOfBadDrivers(bool bHasEditorToken)
 			}
 		}
 
-		GRHIAdapterDriverOnDenyList = DenyListEntry.IsValid();
-		FGenericCrashContext::SetEngineData(TEXT("RHI.DriverDenylisted"), DenyListEntry.IsValid() ? TEXT("true") : TEXT("false"));
-
-		if(!GRHIAdapterDriverOnDenyList)
-		{
-			return;
-		}
-
 		// Only alert users who are capable of updating their driver. Assume vendors with an empty list can always update.
+		// The warning message can also be suppressed with r.WarnOfBadDrivers=0.
 		bool bShowPrompt = bDeviceCanUpdateDriver || !bVendorHasEntries;
-		bShowPrompt = bShowPrompt && !FApp::IsUnattended() && WarnMode != 0;
+		bShowPrompt = bShowPrompt && !FApp::IsUnattended() && GWarnOfBadDrivers != 0;
 
 		if (bShowPrompt)
 		{
-			bool bLatestDenied = DetectedGPUHardware.IsLatestDenied();
-
 			// Note: we don't localize the vendor's name.
 			FString VendorString = DriverInfo.ProviderName;
 			FText HyperlinkText;
@@ -221,39 +172,43 @@ void RHIDetectAndWarnOfBadDrivers(bool bHasEditorToken)
 				HyperlinkText = NSLOCTEXT("MessageDialog", "DriverDownloadLinkIntel", "https://downloadcenter.intel.com/product/80939/Graphics");
 			}
 
-			// format message box UI
 			FFormatNamedArguments Args;
 			Args.Add(TEXT("AdapterName"), FText::FromString(DriverInfo.DeviceDescription));
 			Args.Add(TEXT("Vendor"), FText::FromString(VendorString));
-			Args.Add(TEXT("RHI"), FText::FromString(DenyListEntry.RHIName));
 			Args.Add(TEXT("Hyperlink"), HyperlinkText);
-			Args.Add(TEXT("RecommendedVer"), FText::FromString(DetectedGPUHardware.GetSuggestedDriverVersion(DriverInfo.RHIName)));
 			Args.Add(TEXT("InstalledVer"), FText::FromString(DriverInfo.UserDriverVersion));
 
-			// this message can be suppressed with r.WarnOfBadDrivers=0
-			FText LocalizedMsg;
-			if (bLatestDenied)
+			// Find the best driver version to recommend.
+			TOptional<FSuggestedDriverEntry> SuggestedDriver = DetectedGPUHardware.FindSuggestedDriverVersion();
+			if (SuggestedDriver)
 			{
-				if (!DenyListEntry.RHIName.IsEmpty())
+				// Suggest the latest too, if not denylisted.
+				if (!DenyListEntry->AppliesToLatestDrivers())
 				{
-					LocalizedMsg = FText::Format(NSLOCTEXT("MessageDialog", "LatestVideoCardDriverRHIIssueReport", "The latest version of the {Vendor} graphics driver has known issues in {RHI}.\nPlease install the recommended driver version or switch to a different rendering API.\n\nWould you like to visit the following URL to download the driver?\n\n{Hyperlink}\n\n{AdapterName}\nInstalled: {InstalledVer}\nRecommended: {RecommendedVer}"), Args);
+					Args.Add(TEXT("RecommendedVer"), FText::Format(NSLOCTEXT("MessageDialog", "SuggestedDriverOrLatest", "{0} or latest driver available"), FText::FromString(SuggestedDriver->SuggestedDriverVersion)));
 				}
 				else
 				{
-					LocalizedMsg = FText::Format(NSLOCTEXT("MessageDialog", "LatestVideoCardDriverIssueReport", "The latest version of the {Vendor} graphics driver has known issues.\nPlease install the recommended driver version.\n\nWould you like to visit the following URL to download the driver?\n\n{Hyperlink}\n\n{AdapterName}\nInstalled: {InstalledVer}\nRecommended: {RecommendedVer}"), Args);
+					Args.Add(TEXT("RecommendedVer"), FText::FromString(SuggestedDriver->SuggestedDriverVersion));
 				}
 			}
 			else
 			{
-				if (!DenyListEntry.RHIName.IsEmpty())
-				{
-					LocalizedMsg = FText::Format(NSLOCTEXT("MessageDialog", "VideoCardDriverRHIIssueReport", "The installed version of the {Vendor} graphics driver has known issues in {RHI}.\nPlease install the latest driver version or switch to a different rendering API.\n\nWould you like to visit the following URL to download the driver?\n\n{Hyperlink}\n\n{AdapterName}\nInstalled: {InstalledVer}\nMinimum required: {RecommendedVer}"), Args);
-				}
-				else
-				{
-					LocalizedMsg = FText::Format(NSLOCTEXT("MessageDialog", "VideoCardDriverIssueReport", "The installed version of the {Vendor} graphics driver has known issues.\nPlease install the latest driver version.\n\nWould you like to visit the following URL to download the driver?\n\n{Hyperlink}\n\n{AdapterName}\nInstalled: {InstalledVer}\nMinimum required: {RecommendedVer}"), Args);
-				}
+				ensureMsgf(!DenyListEntry->AppliesToLatestDrivers(), TEXT("Latest drivers are denylisted but no recommended driver driver has been provided"));
+				Args.Add(TEXT("RecommendedVer"), NSLOCTEXT("MessageDialog", "LatestDriver", "latest driver available"));
+			}
+
+			FText LocalizedMsg;
+			if (DenyListEntry->RHINameConstraint)
+			{
+				Args.Add(TEXT("RHI"), FText::FromString(*DenyListEntry->RHINameConstraint));
+				LocalizedMsg = FText::Format(NSLOCTEXT("MessageDialog", "VideoCardDriverRHIIssueReport", "The installed version of the {Vendor} graphics driver has known issues in {RHI}.\nPlease install the recommended driver version or switch to a different rendering API.\n\nWould you like to visit the following URL to download the driver?\n\n{Hyperlink}\n\n{AdapterName}\nInstalled: {InstalledVer}\nRecommended: {RecommendedVer}"), Args);
 			}
+			else
+			{
+				LocalizedMsg = FText::Format(NSLOCTEXT("MessageDialog", "VideoCardDriverIssueReport", "The installed version of the {Vendor} graphics driver has known issues.\nPlease install the recommended driver version.\n\nWould you like to visit the following URL to download the driver?\n\n{Hyperlink}\n\n{AdapterName}\nInstalled: {InstalledVer}\nRecommended: {RecommendedVer}"), Args);
+			}
+
 
 			FText Title = NSLOCTEXT("MessageDialog", "TitleVideoCardDriverIssue", "WARNING: Known issues with graphics driver");
 			EAppReturnType::Type Response = FMessageDialog::Open(EAppMsgType::YesNo, LocalizedMsg, Title);
@@ -262,9 +217,9 @@ void RHIDetectAndWarnOfBadDrivers(bool bHasEditorToken)
 				FPlatformProcess::LaunchURL(*HyperlinkText.ToString(), nullptr, nullptr);
 			}
 #if !UE_BUILD_SHIPPING
-			if (CVarBadDriverWarningIsFatal.GetValueOnGameThread())
+			if (GBadDriverWarningIsFatal)
 			{
-				// Force a fatal error depending on CVar
+				// Force a fatal error depending on CVar.
 				UE_LOG(LogRHI, Fatal, TEXT("Fatal crash requested when graphics drivers are out of date.\n")
 					TEXT("To prevent this crash, please update drivers."));
 			}
@@ -273,7 +228,7 @@ void RHIDetectAndWarnOfBadDrivers(bool bHasEditorToken)
 		}
 		else
 		{
-			UE_LOG(LogRHI, Warning, TEXT("Running with bad GPU drivers but warning dialog will not be shown: bDeviceCanUpdateDriver=%d, VendorHasEntries=%d, IsUnattended=%d, r.WarnOfBadDrivers=%d"), bDeviceCanUpdateDriver, bVendorHasEntries, FApp::IsUnattended(), WarnMode);
+			UE_LOG(LogRHI, Warning, TEXT("Running with bad GPU drivers but warning dialog will not be shown: bDeviceCanUpdateDriver=%d, VendorHasEntries=%d, IsUnattended=%d, r.WarnOfBadDrivers=%d"), bDeviceCanUpdateDriver, bVendorHasEntries, FApp::IsUnattended(), GWarnOfBadDrivers);
 		}
 	}
 }
