Enable some driver denylist tests for Windows only.

diff --git a/Engine/Source/Runtime/Core/Tests/GenericPlatform/GenericPlatformDriverTest.cpp b/Engine/Source/Runtime/Core/Tests/GenericPlatform/GenericPlatformDriverTest.cpp
--- a/Engine/Source/Runtime/Core/Tests/GenericPlatform/GenericPlatformDriverTest.cpp
+++ b/Engine/Source/Runtime/Core/Tests/GenericPlatform/GenericPlatformDriverTest.cpp
@@ -186,7 +186,7 @@ TEST_CASE_NAMED(DriverVersionTest, "System::Core::GenericPlatform::Driver", "[Ap
 		}
 	}
 
-
+#if PLATFORM_WINDOWS
 	SECTION("Driver denylist and suggested driver entry retrieval")
 	{
 		{
@@ -298,6 +298,7 @@ TEST_CASE_NAMED(DriverVersionTest, "System::Core::GenericPlatform::Driver", "[Ap
 		}
 #endif
 	}
+#endif
 }
 
 #endif //WITH_TESTS
\ No newline at end of file
