Add support for + and - when defining config variables on the commandline for array type variables.

diff --git a/Engine/Source/Runtime/Core/Private/Misc/ConfigCacheIni.cpp b/Engine/Source/Runtime/Core/Private/Misc/ConfigCacheIni.cpp
--- a/Engine/Source/Runtime/Core/Private/Misc/ConfigCacheIni.cpp
+++ b/Engine/Source/Runtime/Core/Private/Misc/ConfigCacheIni.cpp
@@ -1834,8 +1834,29 @@ void FConfigFile::OverrideFromCommandline(FConfigFile* File, const FString& File
 					}
 
 					// now put it into this into the cache
-					File->SetString(*CommandlineOption.Section, *CommandlineOption.PropertyKey, *CommandlineOption.PropertyValue);
-				}
+					if (CommandlineOption.PropertyKey.StartsWith(TEXT("-")))
+					{
+						CommandlineOption.PropertyKey.RemoveFromStart(TEXT("-"));
+
+						TArray<FString> ValueArray;
+						File->GetArray(*CommandlineOption.Section, *CommandlineOption.PropertyKey, ValueArray);
+						ValueArray.Remove(CommandlineOption.PropertyValue);
+						File->SetArray(*CommandlineOption.Section, *CommandlineOption.PropertyKey, ValueArray);
+					}
+					else if (CommandlineOption.PropertyKey.StartsWith(TEXT("+")))
+					{
+						CommandlineOption.PropertyKey.RemoveFromStart(TEXT("+"));
+
+						TArray<FString> ValueArray;
+						File->GetArray(*CommandlineOption.Section, *CommandlineOption.PropertyKey, ValueArray);
+						ValueArray.Add(CommandlineOption.PropertyValue);
+						File->SetArray(*CommandlineOption.Section, *CommandlineOption.PropertyKey, ValueArray);
+					}
+					else
+					{
+						File->SetString(*CommandlineOption.Section, *CommandlineOption.PropertyKey, *CommandlineOption.PropertyValue);
+					}
+				}	
 			}
 		}
 	}
