Compiled out namedevents csv metadata changes when UE_SERVER=1 for now.
Server csvs are being filtered out by PRS due to auto trace/framepro captures during the match.

diff --git a/Engine/Source/Runtime/Core/Private/ProfilingDebugging/CsvProfiler.cpp b/Engine/Source/Runtime/Core/Private/ProfilingDebugging/CsvProfiler.cpp
--- a/Engine/Source/Runtime/Core/Private/ProfilingDebugging/CsvProfiler.cpp
+++ b/Engine/Source/Runtime/Core/Private/ProfilingDebugging/CsvProfiler.cpp
@@ -3087,7 +3087,13 @@ void FCsvProfiler::BeginFrame()
 			if (!bNamedEventsWasEnabled && (GCycleStatsShouldEmitNamedEvents > 0))
 			{
 				bNamedEventsWasEnabled = true;
+#if !UE_SERVER
+				// Servers with -EnableMetrics run perf collection a few times per match which includes
+				// insights/framepro captures and namedevents on.
+				// We don't want these server csvs to be filtered out by PRS, so
+				// excluding this logic from server builds for now.
 				SetMetadataInternal(TEXT("NamedEvents"), TEXT("1"));
+#endif
 			}
 		}
 	}
