Added FString::RemoveAt(Index, AllowShrinking), to match the TArray::RemoveAt overload.

diff --git a/Engine/Source/Runtime/Core/Private/Containers/String.cpp.inl b/Engine/Source/Runtime/Core/Private/Containers/String.cpp.inl
--- a/Engine/Source/Runtime/Core/Private/Containers/String.cpp.inl
+++ b/Engine/Source/Runtime/Core/Private/Containers/String.cpp.inl
@@ -645,6 +645,15 @@ void UE_STRING_CLASS::InsertAt(int32 Index, const UE_STRING_CLASS& Characters)
 	}
 }
 
+void UE_STRING_CLASS::RemoveAt(int32 Index, EAllowShrinking AllowShrinking)
+{
+	// The clamping behavior below that we are inheriting from the (Index, Count, AllowShrinking) overload is unfortunate, as it pessimizes this call.
+	if (FMath::Clamp(1, 0, Len()-Index) == 1)
+	{
+		Data.RemoveAt(Index, AllowShrinking);
+	}
+}
+
 void UE_STRING_CLASS::RemoveAt(int32 Index, int32 Count, EAllowShrinking AllowShrinking)
 {
 	Data.RemoveAt(Index, FMath::Clamp(Count, 0, Len()-Index), AllowShrinking);
diff --git a/Engine/Source/Runtime/Core/Public/Containers/UnrealString.h.inl b/Engine/Source/Runtime/Core/Public/Containers/UnrealString.h.inl
--- a/Engine/Source/Runtime/Core/Public/Containers/UnrealString.h.inl
+++ b/Engine/Source/Runtime/Core/Public/Containers/UnrealString.h.inl
@@ -443,6 +443,14 @@ public:
 	CORE_API void InsertAt(int32 Index, ElementType Character);
 	CORE_API void InsertAt(int32 Index, const UE_STRING_CLASS& Characters);
 
+	/**
+	 * Removes a character from the string.
+	 *
+	 * @param Index          The index of the character to remove.
+	 * @param AllowShrinking Whether or not to reallocate to shrink the storage after removal.
+	 */
+	CORE_API void RemoveAt(int32 Index, EAllowShrinking AllowShrinking = EAllowShrinking::Yes);
+
 	/**
 	 * Removes characters within the string.
 	 *
@@ -450,7 +458,7 @@ public:
 	 * @param Count          The number of characters to remove.
 	 * @param AllowShrinking Whether or not to reallocate to shrink the storage after removal.
 	 */
-	CORE_API void RemoveAt(int32 Index, int32 Count = 1, EAllowShrinking AllowShrinking = EAllowShrinking::Yes);
+	CORE_API void RemoveAt(int32 Index, int32 Count, EAllowShrinking AllowShrinking = EAllowShrinking::Yes);
 	UE_ALLOWSHRINKING_BOOL_DEPRECATED("RemoveAt")
 	FORCEINLINE void RemoveAt(int32 Index, int32 Count, bool bAllowShrinking)
 	{
