2 small build fixes for tests.
Found by right click, run tests from the ide.

diff --git a/Engine/Source/Runtime/Core/Tests/HAL/PlatformStringTest.cpp b/Engine/Source/Runtime/Core/Tests/HAL/PlatformStringTest.cpp
--- a/Engine/Source/Runtime/Core/Tests/HAL/PlatformStringTest.cpp
+++ b/Engine/Source/Runtime/Core/Tests/HAL/PlatformStringTest.cpp
@@ -6,6 +6,9 @@
 
 #if WITH_TESTS
 
+#include "Misc/CString.h"
+#include "Templates/UnrealTemplate.h"
+
 template <typename CharType, SIZE_T Size>
 static void InvokePlatformStringGetVarArgs(CharType (&Dest)[Size], const CharType* Fmt, ...)
 {
diff --git a/Engine/Source/Runtime/Core/Tests/HAL/ThreadTest.cpp b/Engine/Source/Runtime/Core/Tests/HAL/ThreadTest.cpp
--- a/Engine/Source/Runtime/Core/Tests/HAL/ThreadTest.cpp
+++ b/Engine/Source/Runtime/Core/Tests/HAL/ThreadTest.cpp
@@ -11,6 +11,7 @@
 #include "HAL/ThreadSafeCounter.h"
 #include "HAL/ThreadSingleton.h"
 #include "HAL/Event.h"
+#include "HAL/PlatformProcess.h"
 #include "Containers/Queue.h"
 #include "Containers/StringConv.h"
 
