- Removed log spam

diff --git a/Engine/Source/Runtime/Core/Private/Misc/ConfigContext.cpp b/Engine/Source/Runtime/Core/Private/Misc/ConfigContext.cpp
--- a/Engine/Source/Runtime/Core/Private/Misc/ConfigContext.cpp
+++ b/Engine/Source/Runtime/Core/Private/Misc/ConfigContext.cpp
@@ -939,7 +939,6 @@ bool FConfigContext::LoadIniFileHierarchy()
 
 		if (ChangeTracker != nullptr && ChangeTracker->bTrackLoadedFiles)
 		{
-			UE_LOG(LogConfig, Display, TEXT("Tracking ini %s"), *IniFileName);
 			ChangeTracker->LoadedFiles.Add(IniFileName);
 		}
 	}
