Added RemoveAt*(Index, AllowShrinking) overloads to TObservableArray, TIndirectArray and FPreviewAssetAttachContainer to match TArray.

diff --git a/Engine/Source/Runtime/Core/Public/Containers/IndirectArray.h b/Engine/Source/Runtime/Core/Public/Containers/IndirectArray.h
--- a/Engine/Source/Runtime/Core/Public/Containers/IndirectArray.h
+++ b/Engine/Source/Runtime/Core/Public/Containers/IndirectArray.h
@@ -243,15 +243,32 @@ public:
 	}
 
 	/**
-	 * Removes an element (or elements) at given location optionally shrinking
+	 * Removes an element at given location optionally shrinking
 	 * the array.
 	 *
 	 * @param Index Location in array of the element to remove.
-	 * @param Count (Optional) Number of elements to remove. Default is 1.
 	 * @param AllowShrinking (Optional) Tells if this call can shrink array if
 	 *                        suitable after remove. Default is yes.
 	 */
-	void RemoveAt(int32 Index, int32 Count = 1, EAllowShrinking AllowShrinking = EAllowShrinking::Yes)
+	void RemoveAt(int32 Index, EAllowShrinking AllowShrinking = EAllowShrinking::Yes)
+	{
+		check(Index >= 0);
+		check(Index < Array.Num());
+		T** Element = GetData() + Index;
+		delete *Element;
+		Array.RemoveAt(Index, AllowShrinking);
+	}
+
+	/**
+	 * Removes elements at given location optionally shrinking
+	 * the array.
+	 *
+	 * @param Index Location in array of the first element to remove.
+	 * @param Count Number of elements to remove.
+	 * @param AllowShrinking (Optional) Tells if this call can shrink array if
+	 *                        suitable after remove. Default is yes.
+	 */
+	void RemoveAt(int32 Index, int32 Count, EAllowShrinking AllowShrinking = EAllowShrinking::Yes)
 	{
 		check(Index >= 0);
 		check(Index <= Array.Num());
@@ -271,18 +288,38 @@ public:
 	}
 
 	/**
-	 * Removes an element (or elements) at given location optionally shrinking
+	 * Removes an element at given location optionally shrinking
 	 * the array.
 	 *
 	 * This version is much more efficient than RemoveAt (O(Count) instead of
 	 * O(ArrayNum)), but does not preserve the order.
 	 *
 	 * @param Index Location in array of the element to remove.
-	 * @param Count (Optional) Number of elements to remove. Default is 1.
 	 * @param AllowShrinking (Optional) Tells if this call can shrink array if
 	 *                        suitable after remove. Default is yes.
 	 */
-	void RemoveAtSwap(int32 Index, int32 Count = 1, EAllowShrinking AllowShrinking = EAllowShrinking::Yes)
+	void RemoveAtSwap(int32 Index, EAllowShrinking AllowShrinking = EAllowShrinking::Yes)
+	{
+		check(Index >= 0);
+		check(Index < Array.Num());
+		T** Element = GetData() + Index;
+		delete *Element;
+		Array.RemoveAtSwap(Index, AllowShrinking);
+	}
+
+	/**
+	 * Removes elements at given location optionally shrinking
+	 * the array.
+	 *
+	 * This version is much more efficient than RemoveAt (O(Count) instead of
+	 * O(ArrayNum)), but does not preserve the order.
+	 *
+	 * @param Index Location in array of the rirst element to remove.
+	 * @param Count Number of elements to remove.
+	 * @param AllowShrinking (Optional) Tells if this call can shrink array if
+	 *                        suitable after remove. Default is yes.
+	 */
+	void RemoveAtSwap(int32 Index, int32 Count, EAllowShrinking AllowShrinking = EAllowShrinking::Yes)
 	{
 		check(Index >= 0);
 		check(Index <= Array.Num());
diff --git a/Engine/Source/Runtime/Engine/Classes/Animation/PreviewAssetAttachComponent.h b/Engine/Source/Runtime/Engine/Classes/Animation/PreviewAssetAttachComponent.h
--- a/Engine/Source/Runtime/Engine/Classes/Animation/PreviewAssetAttachComponent.h
+++ b/Engine/Source/Runtime/Engine/Classes/Animation/PreviewAssetAttachComponent.h
@@ -139,7 +139,8 @@ public:
 	/** 
 	 * RemoveAtSwap passthrough
 	 */
-	ENGINE_API void RemoveAtSwap( int32 Index, int32 Count = 1, EAllowShrinking AllowShrinking = EAllowShrinking::Yes);
+	ENGINE_API void RemoveAtSwap( int32 Index, EAllowShrinking AllowShrinking = EAllowShrinking::Yes);
+	ENGINE_API void RemoveAtSwap( int32 Index, int32 Count, EAllowShrinking AllowShrinking = EAllowShrinking::Yes);
 	UE_ALLOWSHRINKING_BOOL_DEPRECATED("RemoveAtSwap")
 	FORCEINLINE void RemoveAtSwap(int32 Index, int32 Count, bool bAllowShrinking)
 	{
diff --git a/Engine/Source/Runtime/Engine/Private/Animation/PreviewAssetAttachComponent.cpp b/Engine/Source/Runtime/Engine/Private/Animation/PreviewAssetAttachComponent.cpp
--- a/Engine/Source/Runtime/Engine/Private/Animation/PreviewAssetAttachComponent.cpp
+++ b/Engine/Source/Runtime/Engine/Private/Animation/PreviewAssetAttachComponent.cpp
@@ -36,7 +36,7 @@ void FPreviewAssetAttachContainer::RemoveAttachedObject( UObject* ObjectToRemove
 
 		if( Pair.GetAttachedObject() == ObjectToRemove && Pair.AttachedTo == AttachName )
 		{
-			AttachedObjects.RemoveAtSwap( i, 1, EAllowShrinking::No);
+			AttachedObjects.RemoveAtSwap( i, EAllowShrinking::No);
 			break;
 		}
 	}
@@ -84,7 +84,12 @@ TIterator FPreviewAssetAttachContainer::CreateIterator()
 	return AttachedObjects.CreateIterator();
 }
 
-void FPreviewAssetAttachContainer::RemoveAtSwap( int32 Index, int32 Count /* = 1 */, EAllowShrinking AllowShrinking /*=EAllowShrinking::Yes */ )
+void FPreviewAssetAttachContainer::RemoveAtSwap( int32 Index, EAllowShrinking AllowShrinking /*=EAllowShrinking::Yes */ )
+{
+	AttachedObjects.RemoveAtSwap( Index, AllowShrinking );
+}
+
+void FPreviewAssetAttachContainer::RemoveAtSwap( int32 Index, int32 Count, EAllowShrinking AllowShrinking /*=EAllowShrinking::Yes */ )
 {
 	AttachedObjects.RemoveAtSwap( Index, Count, AllowShrinking );
 }
diff --git a/Engine/Source/Runtime/SlateCore/Public/Containers/ObservableArray.h b/Engine/Source/Runtime/SlateCore/Public/Containers/ObservableArray.h
--- a/Engine/Source/Runtime/SlateCore/Public/Containers/ObservableArray.h
+++ b/Engine/Source/Runtime/SlateCore/Public/Containers/ObservableArray.h
@@ -331,13 +331,20 @@ public:
 		return RemoveSingleSwap(Item, bAllowShrinking ? EAllowShrinking::Yes : EAllowShrinking::No);
 	}
 
-	void RemoveAt(SizeType Index)
+	void RemoveAt(SizeType Index, EAllowShrinking AllowShrinking = EAllowShrinking::Yes)
 	{
-		RemoveAt(Index, 1, EAllowShrinking::Yes);
-	}
+		check((Index >= 0) & (Index < Num()));
+
+		ElementType RemovedElement = MoveTemp(Array[Index]);
 
-	template <typename CountType>
-	void RemoveAt(SizeType Index, CountType NumToRemove = 1, EAllowShrinking AllowShrinking = EAllowShrinking::Yes)
+		Array.RemoveAt(Index, AllowShrinking);
+		ArrayChangedDelegate.Broadcast(ObservableArrayChangedArgsType::MakeRemoveAction({ &RemovedElement, 1 }, Index));
+	}
+	template <
+		typename CountType
+		UE_REQUIRES(std::is_integral_v<CountType>)
+	>
+	void RemoveAt(SizeType Index, CountType NumToRemove, EAllowShrinking AllowShrinking = EAllowShrinking::Yes)
 	{
 		static_assert(!std::is_same_v<CountType, bool>, "TObservableArray::RemoveAt: unexpected bool passed as the Count argument");
 		check((NumToRemove > 0) & (Index >= 0) & (Index + NumToRemove <= Num()));
@@ -371,13 +378,21 @@ public:
 		RemoveAt(Index, NumToRemove, bAllowShrinking ? EAllowShrinking::Yes : EAllowShrinking::No);
 	}
 
-	void RemoveAtSwap(SizeType Index)
+	void RemoveAtSwap(SizeType Index, EAllowShrinking AllowShrinking = EAllowShrinking::Yes)
 	{
-		RemoveAtSwap(Index, 1, EAllowShrinking::Yes);
+		SizeType PreviousNum = Array.Num();
+		SizeType SwapAmount = FPlatformMath::Min(1, PreviousNum - (Index + 1));
+
+		ElementType RemovedElement = MoveTemp(Array[Index]);
+
+		Array.RemoveAtSwap(Index, AllowShrinking);
+		ArrayChangedDelegate.Broadcast(ObservableArrayChangedArgsType::MakeRemoveSwapAction({ &RemovedElement, 1 }, Index, PreviousNum - SwapAmount));
 	}
-	
-	template <typename CountType>
-	void RemoveAtSwap(SizeType Index, CountType NumToRemove = 1, EAllowShrinking AllowShrinking = EAllowShrinking::Yes)
+	template <
+		typename CountType
+		UE_REQUIRES(std::is_integral_v<CountType>)
+	>
+	void RemoveAtSwap(SizeType Index, CountType NumToRemove, EAllowShrinking AllowShrinking = EAllowShrinking::Yes)
 	{
 		static_assert(!std::is_same_v<CountType, bool>, "TObservableArray::RemoveAtSwap: unexpected bool passed as the Count argument");
 		check((NumToRemove > 0) & (Index >= 0) & (Index + NumToRemove <= Num()));
