- Fixed OSC plugin errors from objc.h #defining Nil, which UE doesn't need

diff --git a/Engine/Source/Runtime/Core/Public/Mac/MacSystemIncludes.h b/Engine/Source/Runtime/Core/Public/Mac/MacSystemIncludes.h
--- a/Engine/Source/Runtime/Core/Public/Mac/MacSystemIncludes.h
+++ b/Engine/Source/Runtime/Core/Public/Mac/MacSystemIncludes.h
@@ -51,6 +51,9 @@
 #include <sys/sysctl.h>
 #include <malloc/malloc.h>
 
+// we don't want this in UE, but objc.h defines it to nullptr
+#undef Nil
+
 struct tagRECT
 {
 	int32 left;
