Made UE_STATIC_ASSERT_COMPLETE_TYPE take an optional message, like regular static_assert does.

diff --git a/Engine/Source/Runtime/Core/Public/Misc/StaticAssertCompleteType.h b/Engine/Source/Runtime/Core/Public/Misc/StaticAssertCompleteType.h
--- a/Engine/Source/Runtime/Core/Public/Misc/StaticAssertCompleteType.h
+++ b/Engine/Source/Runtime/Core/Public/Misc/StaticAssertCompleteType.h
@@ -24,7 +24,7 @@ namespace UEStaticAssertCompleteType_Private
 }
 
 // Causes a compile error if a type is incomplete
-#define UE_STATIC_ASSERT_COMPLETE_TYPE(TypeToCheck, Message) static_assert(sizeof(UEStaticAssertCompleteType_Private::TUEStaticAssertTypeChecker<TypeToCheck>::Func()), Message)
+#define UE_STATIC_ASSERT_COMPLETE_TYPE(TypeToCheck, ...) static_assert(sizeof(UEStaticAssertCompleteType_Private::TUEStaticAssertTypeChecker<TypeToCheck>::Func()), ##__VA_ARGS__)
 
 // Tests
 
