Expose the crash report tags and adding a new `VerseRuntimeError` crash context for Verse runtime error reports. Also exposes a new `GetGameName` function for standardizing the way the game name is retrieved for use in crash reports.

diff --git a/Engine/Source/Runtime/Core/Private/GenericPlatform/GenericPlatformCrashContext.cpp b/Engine/Source/Runtime/Core/Private/GenericPlatform/GenericPlatformCrashContext.cpp
--- a/Engine/Source/Runtime/Core/Private/GenericPlatform/GenericPlatformCrashContext.cpp
+++ b/Engine/Source/Runtime/Core/Private/GenericPlatform/GenericPlatformCrashContext.cpp
@@ -75,7 +75,7 @@ static bool NeedsEscape(FStringView Str)
 	return false;
 }
 
-static const TCHAR* AttendedStatusToString(EUnattendedStatus Status)
+const TCHAR* AttendedStatusToString(const EUnattendedStatus Status)
 {
 	switch(Status)
 	{
@@ -190,7 +190,24 @@ const TCHAR* const FGenericCrashContext::RuntimePropertiesTag = TEXT( "RuntimePr
 const TCHAR* const FGenericCrashContext::PlatformPropertiesTag = TEXT( "PlatformProperties" );
 const TCHAR* const FGenericCrashContext::EngineDataTag = TEXT( "EngineData" );
 const TCHAR* const FGenericCrashContext::GameDataTag = TEXT( "GameData" );
+const TCHAR* const FGenericCrashContext::GameNameTag = TEXT( "GameName" );
 const TCHAR* const FGenericCrashContext::EnabledPluginsTag = TEXT("EnabledPlugins");
+const TCHAR* const FGenericCrashContext::CrashVersionTag = TEXT("CrashVersion");
+const TCHAR* const FGenericCrashContext::ExecutionGuidTag = TEXT("ExecutionGuid");
+const TCHAR* const FGenericCrashContext::CrashGuidTag = TEXT("CrashGUID");
+const TCHAR* const FGenericCrashContext::IsEnsureTag = TEXT("IsEnsure");
+const TCHAR* const FGenericCrashContext::IsStallTag = TEXT("IsStall");
+const TCHAR* const FGenericCrashContext::IsAssertTag = TEXT("IsAssert");
+const TCHAR* const FGenericCrashContext::CrashTypeTag = TEXT("CrashType");
+const TCHAR* const FGenericCrashContext::ErrorMessageTag = TEXT("ErrorMessage");
+const TCHAR* const FGenericCrashContext::CrashReporterMessageTag = TEXT("CrashReporterMessage");
+const TCHAR* const FGenericCrashContext::AttendedStatusTag = TEXT("CrashReporterMessage");
+const TCHAR* const FGenericCrashContext::SecondsSinceStartTag = TEXT("SecondsSinceStart");
+const TCHAR* const FGenericCrashContext::BuildVersionTag = TEXT("BuildVersion");
+const TCHAR* const FGenericCrashContext::CallStackTag = TEXT("CallStack");
+const TCHAR* const FGenericCrashContext::PortableCallStackTag = TEXT("PCallStack");
+const TCHAR* const FGenericCrashContext::PortableCallStackHashTag = TEXT("PCallStackHash");
+const TCHAR* const FGenericCrashContext::IsRequestingExitTag = TEXT("IsRequestingExit");
 const TCHAR* const FGenericCrashContext::UEMinidumpName = TEXT( "UEMinidump.dmp" );
 const TCHAR* const FGenericCrashContext::NewLineTag = TEXT( "&nl;" );
 
@@ -202,6 +219,7 @@ const TCHAR* const FGenericCrashContext::CrashTypeGPU = TEXT("GPUCrash");
 const TCHAR* const FGenericCrashContext::CrashTypeHang = TEXT("Hang");
 const TCHAR* const FGenericCrashContext::CrashTypeAbnormalShutdown = TEXT("AbnormalShutdown");
 const TCHAR* const FGenericCrashContext::CrashTypeOutOfMemory = TEXT("OutOfMemory");
+const TCHAR* const FGenericCrashContext::CrashTypeVerseRuntimeError = TEXT("VerseRuntimeError");
 
 const TCHAR* const FGenericCrashContext::EngineModeExUnknown = TEXT("Unset");
 const TCHAR* const FGenericCrashContext::EngineModeExDirty = TEXT("Dirty");
@@ -513,6 +531,12 @@ const FSessionContext& FGenericCrashContext::GetCachedSessionContext()
 	return NCached::Session;
 }
 
+FString FGenericCrashContext::GetGameName()
+{
+	return FString::Printf(TEXT("UE-%s"), FApp::GetProjectName());
+}
+
+
 void FGenericCrashContext::CopySharedCrashContext(FSharedCrashContext& Dst)
 {
 	//Copy the session
@@ -742,7 +766,7 @@ void FGenericCrashContext::SerializeTempCrashContextToFile()
 void FGenericCrashContext::SerializeSessionContext(FString& Buffer)
 {
 	AddCrashPropertyInternal(Buffer, TEXT("ProcessId"), NCached::Session.ProcessId);
-	AddCrashPropertyInternal(Buffer, TEXT("SecondsSinceStart"), NCached::Session.SecondsSinceStart);
+	AddCrashPropertyInternal(Buffer, FGenericCrashContext::SecondsSinceStartTag, NCached::Session.SecondsSinceStart);
 
 	AddCrashPropertyInternal(Buffer, TEXT("IsInternalBuild"), NCached::Session.bIsInternalBuild);
 	AddCrashPropertyInternal(Buffer, TEXT("IsPerforceBuild"), NCached::Session.bIsPerforceBuild);
@@ -751,18 +775,18 @@ void FGenericCrashContext::SerializeSessionContext(FString& Buffer)
 
 	if (FCString::Strlen(NCached::Session.GameName) > 0)
 	{
-		AddCrashPropertyInternal(Buffer, TEXT("GameName"), NCached::Session.GameName);
+		AddCrashPropertyInternal(Buffer, FGenericCrashContext::GameNameTag, NCached::Session.GameName);
 	}
 	else
 	{
 		const TCHAR* ProjectName = FApp::GetProjectName();
 		if (ProjectName != nullptr && ProjectName[0] != 0)
 		{
-			AddCrashPropertyInternal(Buffer, TEXT("GameName"), *FString::Printf(TEXT("UE-%s"), ProjectName));
+			AddCrashPropertyInternal(Buffer, FGenericCrashContext::GameNameTag, *FString::Printf(TEXT("UE-%s"), ProjectName));
 		}
 		else
 		{
-			AddCrashPropertyInternal(Buffer, TEXT("GameName"), TEXT(""));
+			AddCrashPropertyInternal(Buffer, FGenericCrashContext::GameNameTag, TEXT(""));
 		}
 	}
 	AddCrashPropertyInternal(Buffer, TEXT("ExecutableName"), NCached::Session.ExecutableName);
@@ -782,12 +806,12 @@ void FGenericCrashContext::SerializeSessionContext(FString& Buffer)
 	AddCrashPropertyInternal(Buffer, TEXT("CommandLine"), NCached::Session.CommandLine);
 	AddCrashPropertyInternal(Buffer, TEXT("LanguageLCID"), NCached::Session.LanguageLCID);
 	AddCrashPropertyInternal(Buffer, TEXT("AppDefaultLocale"), NCached::Session.DefaultLocale);
-	AddCrashPropertyInternal(Buffer, TEXT("BuildVersion"), NCached::Session.BuildVersion);
+	AddCrashPropertyInternal(Buffer, BuildVersionTag, NCached::Session.BuildVersion);
 	AddCrashPropertyInternal(Buffer, TEXT("Symbols"), NCached::Session.SymbolsLabel);
 	AddCrashPropertyInternal(Buffer, TEXT("IsUERelease"), NCached::Session.bIsUERelease);
 
 	// Need to set this at the time of the crash to check if requesting exit had been called
-	AddCrashPropertyInternal(Buffer, TEXT("IsRequestingExit"), NCached::Session.bIsExitRequested);
+	AddCrashPropertyInternal(Buffer, FGenericCrashContext::IsRequestingExitTag, NCached::Session.bIsExitRequested);
 
 	// Remove periods from user names to match AutoReporter user names
 	// The name prefix is read by CrashRepository.AddNewCrash in the website code
@@ -864,17 +888,17 @@ void FGenericCrashContext::SerializeContentToBuffer() const
 	AddHeader(CommonBuffer);
 
 	BeginSection( CommonBuffer, RuntimePropertiesTag );
-	AddCrashProperty( TEXT( "CrashVersion" ), (int32)ECrashDescVersions::VER_3_CrashContext );
-	AddCrashProperty( TEXT( "ExecutionGuid" ), *ExecutionGuid.ToString() );
-	AddCrashProperty( TEXT( "CrashGUID" ), (const TCHAR*)CrashGUID);
-
-	AddCrashProperty( TEXT( "IsEnsure" ), (Type == ECrashContextType::Ensure) );
-	AddCrashProperty( TEXT( "IsStall"), (Type == ECrashContextType::Stall) );
-	AddCrashProperty( TEXT( "IsAssert" ), (Type == ECrashContextType::Assert) );
-	AddCrashProperty( TEXT( "CrashType" ), GetCrashTypeString(Type) );
-	AddCrashProperty( TEXT( "ErrorMessage" ), ErrorMessage );
-	AddCrashProperty( TEXT( "CrashReporterMessage" ), NCached::Session.CrashReportClientRichText );
-	AddCrashProperty( TEXT( "AttendedStatus"), NCached::Session.AttendedStatus);
+	AddCrashProperty( CrashVersionTag, (int32)ECrashDescVersions::VER_3_CrashContext );
+	AddCrashProperty( ExecutionGuidTag, *ExecutionGuid.ToString() );
+	AddCrashProperty( CrashGuidTag, (const TCHAR*)CrashGUID);
+
+	AddCrashProperty( IsEnsureTag, (Type == ECrashContextType::Ensure) );
+	AddCrashProperty( IsStallTag, (Type == ECrashContextType::Stall) );
+	AddCrashProperty( IsAssertTag, (Type == ECrashContextType::Assert) );
+	AddCrashProperty( CrashTypeTag, GetCrashTypeString(Type) );
+	AddCrashProperty( ErrorMessageTag, ErrorMessage );
+	AddCrashProperty( CrashReporterMessageTag, NCached::Session.CrashReportClientRichText );
+	AddCrashProperty( AttendedStatusTag, NCached::Session.AttendedStatus);
 
 	SerializeSessionContext(CommonBuffer);
 
@@ -1075,13 +1099,13 @@ void FGenericCrashContext::AddPortableCallStack() const
 {	
 	if (CallStack.Num() == 0)
 	{
-		AddCrashProperty(TEXT("PCallStack"), TEXT(""));
+		AddCrashProperty(PortableCallStackTag, TEXT(""));
 		return;
 	}
 
-	BeginSection(CommonBuffer, TEXT("PCallStack"));
+	BeginSection(CommonBuffer, PortableCallStackTag);
 	AppendPortableCallstack(CommonBuffer, CallStack);
-	EndSection(CommonBuffer, TEXT("PCallStack"));
+	EndSection(CommonBuffer, PortableCallStackTag);
 }
 
 void FGenericCrashContext::AddGPUBreadcrumbs() const
@@ -1219,6 +1243,8 @@ const TCHAR* FGenericCrashContext::GetCrashTypeString(ECrashContextType Type)
 		return CrashTypeAbnormalShutdown;
 	case ECrashContextType::OutOfMemory:
 		return CrashTypeOutOfMemory;
+	case ECrashContextType::VerseRuntimeError:
+		return CrashTypeVerseRuntimeError;
 	default:
 		return CrashTypeCrash;
 	}
diff --git a/Engine/Source/Runtime/Core/Public/GenericPlatform/GenericPlatformCrashContext.h b/Engine/Source/Runtime/Core/Public/GenericPlatform/GenericPlatformCrashContext.h
--- a/Engine/Source/Runtime/Core/Public/GenericPlatform/GenericPlatformCrashContext.h
+++ b/Engine/Source/Runtime/Core/Public/GenericPlatform/GenericPlatformCrashContext.h
@@ -118,6 +118,7 @@ enum class ECrashContextType
 	Hang,
 	OutOfMemory,
 	AbnormalShutdown,
+	VerseRuntimeError,
 
 	Max
 };
@@ -140,6 +141,8 @@ enum class EUnattendedStatus : uint8
 	Unattended
 };
 
+const TCHAR* AttendedStatusToString(const EUnattendedStatus Status);
+
 #define CR_MAX_ERROR_MESSAGE_CHARS 2048
 #define CR_MAX_DIRECTORY_CHARS 256
 #define CR_MAX_SYMBOL_CHARS 128
@@ -374,9 +377,27 @@ public:
 	CORE_API static const TCHAR* const PlatformPropertiesTag;
 	CORE_API static const TCHAR* const EngineDataTag;
 	CORE_API static const TCHAR* const GameDataTag;
+	CORE_API static const TCHAR* const GameNameTag;
 	CORE_API static const TCHAR* const EnabledPluginsTag;
 	CORE_API static const TCHAR* const UEMinidumpName;
 	CORE_API static const TCHAR* const NewLineTag;
+	CORE_API static const TCHAR* const CrashVersionTag;
+	CORE_API static const TCHAR* const ExecutionGuidTag;
+	CORE_API static const TCHAR* const CrashGuidTag;
+	CORE_API static const TCHAR* const IsEnsureTag;
+	CORE_API static const TCHAR* const IsStallTag;
+	CORE_API static const TCHAR* const IsAssertTag;
+	CORE_API static const TCHAR* const CrashTypeTag;
+	CORE_API static const TCHAR* const ErrorMessageTag;
+	CORE_API static const TCHAR* const CrashReporterMessageTag;
+	CORE_API static const TCHAR* const AttendedStatusTag;
+	CORE_API static const TCHAR* const SecondsSinceStartTag;
+	CORE_API static const TCHAR* const BuildVersionTag;
+	CORE_API static const TCHAR* const CallStackTag;
+	CORE_API static const TCHAR* const PortableCallStackTag;
+	CORE_API static const TCHAR* const PortableCallStackHashTag;
+	CORE_API static const TCHAR* const IsRequestingExitTag;
+
 	static constexpr inline int32 CrashGUIDLength = 128;
 
 	CORE_API static const TCHAR* const CrashTypeCrash;
@@ -387,6 +408,7 @@ public:
 	CORE_API static const TCHAR* const CrashTypeHang;
 	CORE_API static const TCHAR* const CrashTypeAbnormalShutdown;
 	CORE_API static const TCHAR* const CrashTypeOutOfMemory;
+	CORE_API static const TCHAR* const CrashTypeVerseRuntimeError;
 
 	CORE_API static const TCHAR* const EngineModeExUnknown;
 	CORE_API static const TCHAR* const EngineModeExDirty;
@@ -404,6 +426,9 @@ public:
 	/** Get the current cached session context */
 	CORE_API static const FSessionContext& GetCachedSessionContext();
 
+	/** Gets the current standardized game name for use in a Crash Reporter report. */
+	CORE_API static FString GetGameName();
+
 	/**
 	 * @return true, if the generic crash context has been initialized.
 	 */
@@ -502,6 +527,12 @@ public:
 		AddCrashPropertyInternal(CommonBuffer, PropertyName, Value);
 	}
 
+	template <typename Type>
+	static void AddCrashProperty(FString& Buffer, const TCHAR* PropertyName, const Type& Value)
+	{
+		AddCrashPropertyInternal(Buffer, PropertyName, Value);
+	}
+
 	/** Escapes and appends specified text to XML string */
 	CORE_API static void AppendEscapedXMLString(FString& OutBuffer, FStringView Text );
 	CORE_API static void AppendEscapedXMLString(FStringBuilderBase& OutBuffer, FStringView Text);
@@ -659,6 +690,9 @@ public:
 	{
 		switch (Type)
 		{
+		// Verse runtime errors only halt the Verse runtime itself; they do not result in a crash.
+		// Certain runtime errors may be recoverable from in the future.
+		case ECrashContextType::VerseRuntimeError: [[fallthrough]];
 		case ECrashContextType::Ensure:
 			return true;
 		case ECrashContextType::Stall:
@@ -736,6 +770,7 @@ private:
 	/** Add module/pdb information to the crash report xml */
 	void AddModules() const;
 
+public:  // Allows this helper functionality to be present for clients to write their own types of crash reports.
 	/** Writes header information to the buffer. */
 	static void AddHeader(FString& Buffer);
 
@@ -746,6 +781,7 @@ private:
 	static void EndSection(FString& Buffer, const TCHAR* SectionName);
 	static void AddSection(FString& Buffer, const TCHAR* SectionName, const FString& SectionContent);
 
+private:
 	/** Called once when GConfig is initialized. Opportunity to cache values from config. */
 	static void InitializeFromConfig();
 
