Remove driver denylist tests that relied on setting config values.
These values can be cached beyond the test and interfere with normal use of the editor.

#rnx
#tests preflight

diff --git a/Engine/Source/Runtime/Core/Tests/GenericPlatform/GenericPlatformDriverTest.cpp b/Engine/Source/Runtime/Core/Tests/GenericPlatform/GenericPlatformDriverTest.cpp
--- a/Engine/Source/Runtime/Core/Tests/GenericPlatform/GenericPlatformDriverTest.cpp
+++ b/Engine/Source/Runtime/Core/Tests/GenericPlatform/GenericPlatformDriverTest.cpp
@@ -185,120 +185,6 @@ TEST_CASE_NAMED(DriverVersionTest, "System::Core::GenericPlatform::Driver", "[Ap
 			CHECK(!Entry.IsValid());
 		}
 	}
-
-#if PLATFORM_WINDOWS
-	SECTION("Driver denylist and suggested driver entry retrieval")
-	{
-		{
-			GConfig->RemoveKey(TEXT("GPU_NVIDIA Windows"), TEXT("DriverDenylist"), GHardwareIni);
-
-			FGPUDriverInfo Info;
-			Info.SetNVIDIA();
-			Info.RHIName = TEXT("D3D12");
-			Info.InternalDriverVersion = TEXT("31.0.15.5123"); // 551.23
-
-			TArray<FString> DenylistEntries = {
-				TEXT("(DriverVersion=\"551.23\", RHI=\"D3D11\", Reason=\"Test\")"), // Does not apply due to RHI.
-				TEXT("(DriverVersion=\"<550.00\", Reason=\"Test\")"), // Does not apply due to version.
-				TEXT("(DriverVersion=\"<=560\", Reason=\"Test\")"), // Applies (no constraints).
-				TEXT("(DriverVersion=\"<555.55\", RHI=\"D3D12\", Reason=\"Test\")"), // Applies with RHI constraint.
-			};
-
-			GConfig->SetArray(TEXT("GPU_NVIDIA Windows"), TEXT("DriverDenylist"), DenylistEntries, GHardwareIni);
-
-			FGPUDriverHelper Helper(Info);
-			TOptional<FDriverDenyListEntry> Entry = Helper.FindDriverDenyListEntry();
-			CHECK(Entry);
-			CHECK(Entry->IsValid());
-			CHECK(Entry->RHINameConstraint);
-			CHECK(Entry->DriverVersion->ComparisonOp == EComparisonOp::ECO_Less);
-			CHECK(Entry->DriverVersion->Version == FDriverVersion(MakeVersionString({ 555, 55 })));
-		}
-
-		{
-			GConfig->RemoveKey(TEXT("GPU_AMD Windows"), TEXT("DriverDenylist"), GHardwareIni);
-
-			FGPUDriverInfo Info;
-			Info.SetAMD();
-			Info.RHIName = TEXT("D3D11");
-			Info.InternalDriverVersion = TEXT("31.0.15.5123");
-			Info.DriverDate = TEXT("1-15-2019");
-
-			TArray<FString> DenylistEntries = {
-				TEXT("(DriverDate=\"<=1-1-2020\", RHI=\"D3D12\", Reason=\"Test\")"), // Does not apply due to RHI.
-				TEXT("(DriverDate=\"<5-5-2018\", Reason=\"Test\")"), // Does not apply due to date.
-				TEXT("(DriverDate=\"<=12-31-2019\", Reason=\"Test\")"), // Applies (no constraints).
-				TEXT("(DriverDate=\"<1-30-2019\", RHI=\"D3D11\", Reason=\"Test\")"), // Applies with RHI constraint.
-			};
-
-			GConfig->SetArray(TEXT("GPU_AMD Windows"), TEXT("DriverDenylist"), DenylistEntries, GHardwareIni);
-
-			FGPUDriverHelper Helper(Info);
-			TOptional<FDriverDenyListEntry> Entry = Helper.FindDriverDenyListEntry();
-			CHECK(Entry);
-			CHECK(Entry->IsValid());
-			CHECK(Entry->RHINameConstraint);
-			CHECK(Entry->DriverDate->ComparisonOp == EComparisonOp::ECO_Less);
-			CHECK(Entry->DriverDate->Date == FDateTime(2019, 1, 30));
-		}
-
-		{
-			GConfig->RemoveKey(TEXT("GPU_Intel Windows"), TEXT("SuggestedDriverVersion"), GHardwareIni);
-
-			FGPUDriverInfo Info;
-			Info.SetIntel();
-			Info.RHIName = TEXT("D3D12");
-			Info.InternalDriverVersion = TEXT("30.0.101.1404");
-
-			TArray<FString> SuggestedDriverEntries = {
-				TEXT("(DriverVersion=\"30.0.101.5000\", RHI=\"D3D11\")"), // Does not apply due to RHI.
-				TEXT("(DriverVersion=\"30.0.101.1255\")"), // Applies (no constraints).
-				TEXT("(DriverVersion=\"30.0.102.1200\", RHI=\"D3D12\")"), // Applies with RHI constraint.
-			};
-
-			GConfig->SetArray(TEXT("GPU_Intel Windows"), TEXT("SuggestedDriverVersion"), SuggestedDriverEntries, GHardwareIni);
-
-			FGPUDriverHelper Helper(Info);
-			TOptional<FSuggestedDriverEntry> Entry = Helper.FindSuggestedDriverVersion();
-			CHECK(Entry);
-			CHECK(Entry->IsValid());
-			CHECK(Entry->RHINameConstraint);
-			CHECK(Entry->SuggestedDriverVersion == TEXT("30.0.102.1200"));
-		}
-
-		// The following test requires regex support, which may be disabled in tests.
-#if UE_ENABLE_ICU 
-		{
-			GConfig->RemoveKey(TEXT("GPU_NVIDIA Windows"), TEXT("DriverDenylist"), GHardwareIni);
-
-			FGPUDriverInfo Info;
-			Info.SetNVIDIA();
-			Info.RHIName = TEXT("D3D12");
-			Info.DeviceDescription = TEXT("NVIDIA GeForce RTX 4070 ");
-			Info.InternalDriverVersion = TEXT("31.0.15.5123"); // 551.23
-
-			TArray<FString> DenylistEntries = {
-				TEXT("(DriverVersion=\"551.23\", RHI=\"D3D11\", Reason=\"Test\")"), // Does not apply due to RHI.
-				TEXT("(DriverVersion=\">550.00\", RHI=\"D3D12\", AdapterNameRegex=\".*4080.*\", Reason=\"Test\")"), // Does not apply due to name regex.
-				TEXT("(DriverVersion=\"<550.00\", Reason=\"Test\")"), // Does not apply due to version.
-				TEXT("(DriverVersion=\"<555.55\", RHI=\"D3D12\", Reason=\"Test\")"), // Applies with RHI constraint.
-				TEXT("(DriverVersion=\"<=551.23\", RHI=\"D3D12\", AdapterNameRegex=\".*4070.*\" Reason=\"Test\")"), // Applies with RHI and name regex constraints.
-			};
-
-			GConfig->SetArray(TEXT("GPU_NVIDIA Windows"), TEXT("DriverDenylist"), DenylistEntries, GHardwareIni);
-
-			FGPUDriverHelper Helper(Info);
-			TOptional<FDriverDenyListEntry> Entry = Helper.FindDriverDenyListEntry();
-			CHECK(Entry);
-			CHECK(Entry->IsValid());
-			CHECK(Entry->RHINameConstraint);
-			CHECK(Entry->AdapterNameRegexConstraint);
-			CHECK(Entry->DriverVersion->ComparisonOp == EComparisonOp::ECO_LessOrEqual);
-			CHECK(Entry->DriverVersion->Version == FDriverVersion(MakeVersionString({ 551, 23 })));
-		}
-#endif
-	}
-#endif
 }
 
 #endif //WITH_TESTS
\ No newline at end of file
