- Fixed a non-unity build issue when UE_RELATIVE_BASE_DIR is set

diff --git a/Engine/Source/Runtime/Core/Private/Windows/WindowsPlatformProcess.cpp b/Engine/Source/Runtime/Core/Private/Windows/WindowsPlatformProcess.cpp
--- a/Engine/Source/Runtime/Core/Private/Windows/WindowsPlatformProcess.cpp
+++ b/Engine/Source/Runtime/Core/Private/Windows/WindowsPlatformProcess.cpp
@@ -9,6 +9,7 @@
 #include "GenericPlatform/GenericPlatformFile.h"
 #include "HAL/PlatformAffinity.h"
 #include "HAL/PlatformMisc.h"
+#include "HAL/FileManager.h"
 #include "HAL/UnrealMemory.h"
 #include "Internationalization/Internationalization.h"
 #include "Logging/LogMacros.h"
