Fixed NoPCH build for FoundationTests

#rnx

diff --git a/Engine/Source/Runtime/Core/Tests/HAL/PlatformStringTest.cpp b/Engine/Source/Runtime/Core/Tests/HAL/PlatformStringTest.cpp
--- a/Engine/Source/Runtime/Core/Tests/HAL/PlatformStringTest.cpp
+++ b/Engine/Source/Runtime/Core/Tests/HAL/PlatformStringTest.cpp
@@ -2,6 +2,8 @@
 
 #include "HAL/PlatformString.h"
 
+#include "Misc/CString.h"
+#include "Templates/UnrealTemplate.h"
 #include "Tests/TestHarnessAdapter.h"
 
 #if WITH_TESTS
diff --git a/Engine/Source/Runtime/Core/Tests/HAL/ThreadTest.cpp b/Engine/Source/Runtime/Core/Tests/HAL/ThreadTest.cpp
--- a/Engine/Source/Runtime/Core/Tests/HAL/ThreadTest.cpp
+++ b/Engine/Source/Runtime/Core/Tests/HAL/ThreadTest.cpp
@@ -6,6 +6,7 @@
 
 #if WITH_TESTS
 
+#include "HAL/PlatformProcess.h"
 #include "HAL/Thread.h"
 #include "HAL/ThreadSafeBool.h"
 #include "HAL/ThreadSafeCounter.h"
