Improved error messages when using MemoryOps functions with undefined types, e.g. when trying to instantiate the destructor of TArray<UndefinedType>.

diff --git a/Engine/Source/Runtime/Core/Public/Templates/MemoryOps.h b/Engine/Source/Runtime/Core/Public/Templates/MemoryOps.h
--- a/Engine/Source/Runtime/Core/Public/Templates/MemoryOps.h
+++ b/Engine/Source/Runtime/Core/Public/Templates/MemoryOps.h
@@ -39,7 +39,11 @@ namespace UE::Core::Private::MemoryOps
 template <typename ElementType, typename SizeType>
 FORCEINLINE void DefaultConstructItems(void* Address, SizeType Count)
 {
-	if constexpr (TIsZeroConstructType<ElementType>::Value)
+	if constexpr (sizeof(ElementType) == 0)
+	{
+		// Should never get here, but this construct should improve the error messages we get when we try to call this function with incomplete types
+	}
+	else if constexpr (TIsZeroConstructType<ElementType>::Value)
 	{
 		FMemory::Memset(Address, 0, sizeof(ElementType) * Count);
 	}
@@ -65,7 +69,11 @@ FORCEINLINE void DefaultConstructItems(void* Address, SizeType Count)
 template <typename ElementType>
 FORCEINLINE void DestructItem(ElementType* Element)
 {
-	if constexpr (!TIsTriviallyDestructible<ElementType>::Value)
+	if constexpr (sizeof(ElementType) == 0)
+	{
+		// Should never get here, but this construct should improve the error messages we get when we try to call this function with incomplete types
+	}
+	else if constexpr (!TIsTriviallyDestructible<ElementType>::Value)
 	{
 		// We need a typedef here because VC won't compile the destructor call below if ElementType itself has a member called ElementType
 		typedef ElementType DestructItemsElementTypeTypedef;
@@ -85,7 +93,11 @@ FORCEINLINE void DestructItem(ElementType* Element)
 template <typename ElementType, typename SizeType>
 FORCEINLINE void DestructItems(ElementType* Element, SizeType Count)
 {
-	if constexpr (!TIsTriviallyDestructible<ElementType>::Value)
+	if constexpr (sizeof(ElementType) == 0)
+	{
+		// Should never get here, but this construct should improve the error messages we get when we try to call this function with incomplete types
+	}
+	else if constexpr (!TIsTriviallyDestructible<ElementType>::Value)
 	{
 		while (Count)
 		{
@@ -109,7 +121,11 @@ FORCEINLINE void DestructItems(ElementType* Element, SizeType Count)
 template <typename DestinationElementType, typename SourceElementType, typename SizeType>
 FORCEINLINE void ConstructItems(void* Dest, const SourceElementType* Source, SizeType Count)
 {
-	if constexpr (TIsBitwiseConstructible<DestinationElementType, SourceElementType>::Value)
+	if constexpr (sizeof(DestinationElementType) == 0 || sizeof(SourceElementType) == 0)
+	{
+		// Should never get here, but this construct should improve the error messages we get when we try to call this function with incomplete types
+	}
+	else if constexpr (TIsBitwiseConstructible<DestinationElementType, SourceElementType>::Value)
 	{
 		if (Count)
 		{
@@ -138,7 +154,11 @@ FORCEINLINE void ConstructItems(void* Dest, const SourceElementType* Source, Siz
 template <typename ElementType, typename SizeType>
 FORCEINLINE void CopyAssignItems(ElementType* Dest, const ElementType* Source, SizeType Count)
 {
-	if constexpr (TIsTriviallyCopyAssignable<ElementType>::Value)
+	if constexpr (sizeof(ElementType) == 0)
+	{
+		// Should never get here, but this construct should improve the error messages we get when we try to call this function with incomplete types
+	}
+	else if constexpr (TIsTriviallyCopyAssignable<ElementType>::Value)
 	{
 		FMemory::Memcpy(Dest, Source, sizeof(ElementType) * Count);
 	}
@@ -164,8 +184,12 @@ FORCEINLINE void CopyAssignItems(ElementType* Dest, const ElementType* Source, S
 template <typename DestinationElementType, typename SourceElementType, typename SizeType>
 FORCEINLINE void RelocateConstructItem(void* Dest, const SourceElementType* Source)
 {
+	if constexpr (sizeof(DestinationElementType) == 0 || sizeof(SourceElementType) == 0)
+	{
+		// Should never get here, but this construct should improve the error messages we get when we try to call this function with incomplete types
+	}
 	// Do a bitwise relocate if TCanBitwiseRelocate says we can, but not if the type involved is known not to be bitwise-swappable
-	if constexpr (UE::Core::Private::MemoryOps::TCanBitwiseRelocate<DestinationElementType, SourceElementType>::Value && (!std::is_same_v<DestinationElementType, SourceElementType> || TUseBitwiseSwap<SourceElementType>::Value))
+	else if constexpr (UE::Core::Private::MemoryOps::TCanBitwiseRelocate<DestinationElementType, SourceElementType>::Value && (!std::is_same_v<DestinationElementType, SourceElementType> || TUseBitwiseSwap<SourceElementType>::Value))
 	{
 		/* All existing UE containers seem to assume trivial relocatability (i.e. memcpy'able) of their members,
 		 * so we're going to assume that this is safe here.  However, it's not generally possible to assume this
@@ -198,7 +222,11 @@ FORCEINLINE void RelocateConstructItem(void* Dest, const SourceElementType* Sour
 template <typename DestinationElementType, typename SourceElementType, typename SizeType>
 FORCEINLINE void RelocateConstructItems(void* Dest, const SourceElementType* Source, SizeType Count)
 {
-	if constexpr (UE::Core::Private::MemoryOps::TCanBitwiseRelocate<DestinationElementType, SourceElementType>::Value)
+	if constexpr (sizeof(DestinationElementType) == 0 || sizeof(SourceElementType) == 0)
+	{
+		// Should never get here, but this construct should improve the error messages we get when we try to call this function with incomplete types
+	}
+	else if constexpr (UE::Core::Private::MemoryOps::TCanBitwiseRelocate<DestinationElementType, SourceElementType>::Value)
 	{
 		/* All existing UE containers seem to assume trivial relocatability (i.e. memcpy'able) of their members,
 		 * so we're going to assume that this is safe here.  However, it's not generally possible to assume this
@@ -235,7 +263,11 @@ FORCEINLINE void RelocateConstructItems(void* Dest, const SourceElementType* Sou
 template <typename ElementType, typename SizeType>
 FORCEINLINE void MoveConstructItems(void* Dest, const ElementType* Source, SizeType Count)
 {
-	if constexpr (TIsTriviallyCopyConstructible<ElementType>::Value)
+	if constexpr (sizeof(ElementType) == 0)
+	{
+		// Should never get here, but this construct should improve the error messages we get when we try to call this function with incomplete types
+	}
+	else if constexpr (TIsTriviallyCopyConstructible<ElementType>::Value)
 	{
 		FMemory::Memmove(Dest, Source, sizeof(ElementType) * Count);
 	}
@@ -261,7 +293,11 @@ FORCEINLINE void MoveConstructItems(void* Dest, const ElementType* Source, SizeT
 template <typename ElementType, typename SizeType>
 FORCEINLINE void MoveAssignItems(ElementType* Dest, const ElementType* Source, SizeType Count)
 {
-	if constexpr (TIsTriviallyCopyAssignable<ElementType>::Value)
+	if constexpr (sizeof(ElementType) == 0)
+	{
+		// Should never get here, but this construct should improve the error messages we get when we try to call this function with incomplete types
+	}
+	else if constexpr (TIsTriviallyCopyAssignable<ElementType>::Value)
 	{
 		FMemory::Memmove(Dest, Source, sizeof(ElementType) * Count);
 	}
@@ -280,7 +316,11 @@ FORCEINLINE void MoveAssignItems(ElementType* Dest, const ElementType* Source, S
 template <typename ElementType, typename SizeType>
 FORCEINLINE bool CompareItems(const ElementType* A, const ElementType* B, SizeType Count)
 {
-	if constexpr (TTypeTraits<ElementType>::IsBytewiseComparable)
+	if constexpr (sizeof(ElementType) == 0)
+	{
+		// Should never get here, but this construct should improve the error messages we get when we try to call this function with incomplete types
+	}
+	else if constexpr (TTypeTraits<ElementType>::IsBytewiseComparable)
 	{
 		return !Count || !FMemory::Memcmp(A, B, sizeof(ElementType) * Count);
 	}
