- Worked around SA warning about using a #define in an if condition

diff --git a/Engine/Source/Runtime/Core/Private/Misc/ConfigContext.cpp b/Engine/Source/Runtime/Core/Private/Misc/ConfigContext.cpp
--- a/Engine/Source/Runtime/Core/Private/Misc/ConfigContext.cpp
+++ b/Engine/Source/Runtime/Core/Private/Misc/ConfigContext.cpp
@@ -977,8 +977,11 @@ bool FConfigContext::GenerateDestIniFile()
 		return false;
 	}
 	
+#if !IS_PROGRAM
 	// Don't try to load any generated files from disk in cooked builds. We will always use the re-generated INIs.
-	if (IS_PROGRAM || !FPlatformProperties::RequiresCookedData() || bAllowGeneratedIniWhenCooked)
+	// Programs also always want this, so skip the check for Programs
+	if (!FPlatformProperties::RequiresCookedData() || bAllowGeneratedIniWhenCooked)
+#endif
 	{
 		if (DestIniFilename.Len() > 0)
 		{
