FCompression interface update to 64 bit.
-> Done because UE disables 32 bit down cast - silently could overflow.
-> General use should automatically work.
-> Deprecate some old 32 bit functions not generally used.
-> Deprecate/rename some ancient flags for future removal.
-> Update CompressorData to uint ptr for future use
-> Remove some stuff deprecated in [REDACTED_PHONE])

diff --git a/Engine/Source/Editor/UnrealEd/Private/Factories/PackFactory.cpp b/Engine/Source/Editor/UnrealEd/Private/Factories/PackFactory.cpp
--- a/Engine/Source/Editor/UnrealEd/Private/Factories/PackFactory.cpp
+++ b/Engine/Source/Editor/UnrealEd/Private/Factories/PackFactory.cpp
@@ -102,13 +102,18 @@ namespace PackFactoryHelper
 		}
 
 		FGuardedInt32 GuardedWorkingSize = FGuardedInt32(Entry.CompressionBlockSize);
-		int32 MaxCompressionBlockSize = FCompression::GetMaximumCompressedSize(CompressionMethod.GetValue(), GuardedWorkingSize.Get(0));
-		if (MaxCompressionBlockSize < 0)
+		int64 MaxCompressionBlockSize64 = 0;
+		if (!FCompression::GetMaximumCompressedSize(CompressionMethod.GetValue(), MaxCompressionBlockSize64, GuardedWorkingSize.Get(0)))
 		{
 			return false;
 		}
 
-		GuardedWorkingSize += MaxCompressionBlockSize;
+		if (MaxCompressionBlockSize64 < 0)
+		{
+			return false;
+		}
+
+		GuardedWorkingSize += MaxCompressionBlockSize64;
 		int32 WorkingSize = GuardedWorkingSize.Get(0);
 		if (WorkingSize <= 0)
 		{
@@ -172,7 +177,7 @@ namespace PackFactoryHelper
 				FAES::DecryptData(PersistentBuffer.GetData(), SizeToRead, Key);
 			}
 
-			uint8* UncompressedBuffer = PersistentBuffer.GetData() + MaxCompressionBlockSize;
+			uint8* UncompressedBuffer = PersistentBuffer.GetData() + MaxCompressionBlockSize64;
 			if (!FCompression::UncompressMemory(CompressionMethod.GetValue(), UncompressedBuffer, UncompressedBlockSize, PersistentBuffer.GetData(), CompressedBlockSize))
 			{
 				return false;
diff --git a/Engine/Source/Runtime/Core/Private/Compression/OodleDataCompressionFormat.cpp b/Engine/Source/Runtime/Core/Private/Compression/OodleDataCompressionFormat.cpp
--- a/Engine/Source/Runtime/Core/Private/Compression/OodleDataCompressionFormat.cpp
+++ b/Engine/Source/Runtime/Core/Private/Compression/OodleDataCompressionFormat.cpp
@@ -202,13 +202,18 @@ struct FOodleDataCompressionFormat : ICompressionFormat
 		return FString::Printf(TEXT("C_%s_CL_%s_%s"), *GetCompressorString(), *GetCompressionLevelString(), OODLE_DERIVEDDATA_VER);
 	}
 
-	virtual bool Compress(void* OutCompressedBuffer, int32& OutCompressedSize, const void* InUncompressedBuffer, int32 InUncompressedSize, int32 InCompressionData, ECompressionFlags Flags) override
+	virtual bool Compress(void* OutCompressedBuffer, int64& OutCompressedSize, const void* InUncompressedBuffer, int64 InUncompressedSize, uintptr_t InCompressionData, ECompressionFlags Flags) override
 	{
 		// OutCompressedSize is read-write
-		int32 CompressedBufferSize = OutCompressedSize;
+		int64 CompressedBufferSize = OutCompressedSize;
 
 		// CompressedSize should be >= GetCompressedBufferSize(UncompressedSize, CompressionData)
-		check(CompressedBufferSize >= GetCompressedBufferSize(InUncompressedSize, InCompressionData));
+		int64 CheckCompressedBufferSize;
+		if (!GetCompressedBufferSize(CheckCompressedBufferSize, InUncompressedSize, InCompressionData))
+		{
+			return false;
+		}
+		check(CompressedBufferSize >= CheckCompressedBufferSize);
 
 		if ( Flags & COMPRESS_ForPackaging )
 		{
@@ -232,7 +237,7 @@ struct FOodleDataCompressionFormat : ICompressionFormat
 				{
 					// verify we can decode data made with DLL-enocder using the current decoder :
 
-					int32 DecodeSpaceSize = InUncompressedSize + 16; // FuzzSafe_Yes so shouldn't need big padding even on Oodle v5
+					int64 DecodeSpaceSize = InUncompressedSize + 16; // FuzzSafe_Yes so shouldn't need big padding even on Oodle v5
 					void * DecodeSpace = FMemory::Malloc(DecodeSpaceSize);
 					check( DecodeSpace != nullptr );
 
@@ -261,7 +266,7 @@ struct FOodleDataCompressionFormat : ICompressionFormat
 			}
 			else
 			{
-				OutCompressedSize = (int32) Result;
+				OutCompressedSize = Result;
 				return true;
 			}
 		}
@@ -291,32 +296,31 @@ struct FOodleDataCompressionFormat : ICompressionFormat
 			}
 			else
 			{
-				OutCompressedSize = (int32) Result;
+				OutCompressedSize = Result;
 				return true;
 			}
 		}
 	
 	}
 
-	virtual bool Uncompress(void* OutUncompressedBuffer, int32& OutUncompressedSize, const void* InCompressedBuffer, int32 InCompressedSize, int32 CompressionData) override
+	virtual bool Uncompress(void* OutUncompressedBuffer, int64 UncompressedSize, const void* InCompressedBuffer, int64 InCompressedSize, uintptr_t CompressionData) override
 	{
-		// OutUncompressedSize is read-write
 		return FOodleDataCompression::Decompress(
-			OutUncompressedBuffer, OutUncompressedSize,
+			OutUncompressedBuffer, UncompressedSize,
 			InCompressedBuffer,InCompressedSize);
 	}
 
-	virtual int32 GetCompressedBufferSize(int32 UncompressedSize, int32 CompressionData) override
+	virtual bool GetCompressedBufferSize(int64& OutBufferNeeded, int64 UncompressedSize, uintptr_t CompressionData) override
 	{
 		// CompressionData is not used
-		int32 Needed = (int32)OodleLZ_GetCompressedBufferSizeNeeded(Compressor, UncompressedSize);
+		int64 Needed = OodleLZ_GetCompressedBufferSizeNeeded(Compressor, UncompressedSize);
 
 		if ( OodleLZCompressFuncPtr != nullptr )
 		{
 			// older versions of Oodle needed larger padding
 			// old rule was 274 per 256KB :
 
-			int32 NumSeekChunks = (UncompressedSize + OODLELZ_BLOCK_LEN-1)/OODLELZ_BLOCK_LEN;
+			int64 NumSeekChunks = (UncompressedSize + OODLELZ_BLOCK_LEN-1)/OODLELZ_BLOCK_LEN;
 
 			Needed += 274*NumSeekChunks;
 
@@ -325,7 +329,8 @@ struct FOodleDataCompressionFormat : ICompressionFormat
 			//	so you'd need to know what version to use
 		}
 
-		return Needed;
+		OutBufferNeeded = Needed;
+		return true;
 	}
 };
 
@@ -334,7 +339,7 @@ extern ICompressionFormat * CreateOodleDataCompressionFormat();
 ICompressionFormat * CreateOodleDataCompressionFormat()
 {
 	// set up Oodle for packaging by parsing the command line options
-	// this is now down on first use of a pluggable ICompressionFormat
+	// this is now done on first use of a pluggable ICompressionFormat
 
 
 	// settings to use in non-tools context (eg. runtime game encoding) :
diff --git a/Engine/Source/Runtime/Core/Private/IO/IoStore.cpp b/Engine/Source/Runtime/Core/Private/IO/IoStore.cpp
--- a/Engine/Source/Runtime/Core/Private/IO/IoStore.cpp
+++ b/Engine/Source/Runtime/Core/Private/IO/IoStore.cpp
@@ -1498,18 +1498,24 @@ private:
 	{
 		TRACE_CPUPROFILER_EVENT_SCOPE(CompressBlock);
 		check(Block->CompressionMethod != NAME_None);
-		int32 CompressedBlockSize = int32(Block->IoBuffer->DataSize());
+		uint64 CompressedBlockSize = Block->IoBuffer->DataSize();
 		bool bCompressed;
 		{
-			bCompressed = FCompression::CompressMemoryIfWorthDecompressing(
+			if (!FCompression::CompressMemoryIfWorthDecompressing(
 				Block->CompressionMethod,
-				WriterContext->WriterSettings.CompressionMinBytesSaved,
+				bCompressed,
+				(int64)WriterContext->WriterSettings.CompressionMinBytesSaved,
 				WriterContext->WriterSettings.CompressionMinPercentSaved,
 				Block->IoBuffer->Data(),
-				CompressedBlockSize,
+				(int64&)CompressedBlockSize,
 				Block->UncompressedData,
-				static_cast<int32>(Block->UncompressedSize),
-				COMPRESS_ForPackaging);
+				(int64)Block->UncompressedSize,
+				COMPRESS_ForPackaging))
+			{
+				UE_LOG(LogIoStore, Error, TEXT("Compression failed: Method=%s, CompressedSize=0x%llx, UncompressedSize=0x%llx"), 
+					*Block->CompressionMethod.ToString(), CompressedBlockSize, Block->UncompressedSize);
+				bCompressed = false;
+			}
 		}
 		if (!bCompressed)
 		{
diff --git a/Engine/Source/Runtime/Core/Private/Misc/Compression.cpp b/Engine/Source/Runtime/Core/Private/Misc/Compression.cpp
--- a/Engine/Source/Runtime/Core/Private/Misc/Compression.cpp
+++ b/Engine/Source/Runtime/Core/Private/Misc/Compression.cpp
@@ -70,7 +70,7 @@ static uint32 appGZIPVersion()
  * @param	BitWindow					Bit window to use in compression
  * @return true if compression succeeds, false if it fails because CompressedBuffer was too small or other reasons
  */
-static bool appCompressMemoryZLIB(void* CompressedBuffer, int32& CompressedSize, const void* UncompressedBuffer, int32 UncompressedSize, int32 BitWindow, int32 CompLevel)
+static bool appCompressMemoryZLIB(void* CompressedBuffer, int64& CompressedSize, const void* UncompressedBuffer, int64 UncompressedSize, int32 BitWindow, int32 CompLevel)
 {
 	TRACE_CPUPROFILER_EVENT_SCOPE(appCompressMemoryZLIB);
 	DECLARE_SCOPE_CYCLE_COUNTER(TEXT("Compress Memory ZLIB"), STAT_appCompressMemoryZLIB, STATGROUP_Compression);
@@ -80,33 +80,49 @@ static bool appCompressMemoryZLIB(void* CompressedBuffer, int32& CompressedSize,
 
 	CompLevel = FMath::Clamp(CompLevel, Z_DEFAULT_COMPRESSION, Z_BEST_COMPRESSION);
 
-	// Zlib wants to use unsigned long.
-	unsigned long ZCompressedSize = CompressedSize;
-	unsigned long ZUncompressedSize = UncompressedSize;
-	bool bOperationSucceeded = false;
-
 	// Compress data
 	// If using the default Zlib bit window, use the zlib routines, otherwise go manual with deflate2
 	if (BitWindow == 0 || BitWindow == DEFAULT_ZLIB_BIT_WINDOW)
 	{
-		bOperationSucceeded = compress2((uint8*)CompressedBuffer, &ZCompressedSize, (const uint8*)UncompressedBuffer, ZUncompressedSize, CompLevel) == Z_OK ? true : false;
+		if (!IntFitsIn<uLongf>(CompressedSize) ||
+			!IntFitsIn<uLong>(UncompressedSize))
+		{
+			UE_LOG(LogCompression, Error, TEXT("Requested a ZLIB compression that doesn't fit in uLong bits"));
+			return false;
+		}
+		uLongf ZCompressedSize = (uLongf)CompressedSize;
+		uLong ZUncompressedSize = (uLong)UncompressedSize;
+		bool bOperationSucceeded = compress2((uint8*)CompressedBuffer, &ZCompressedSize, (const uint8*)UncompressedBuffer, ZUncompressedSize, CompLevel) == Z_OK ? true : false;
+		CompressedSize = ZCompressedSize;
+		return bOperationSucceeded;
 	}
 	else
 	{
+		// Stream wants uInt not uLong
 		z_stream stream;
+
+		if (!IntFitsIn<decltype(stream.avail_out)>(CompressedSize) ||
+			!IntFitsIn<decltype(stream.avail_in)>(UncompressedSize))
+		{
+			UE_LOG(LogCompression, Error, TEXT("Requested a ZLIB compression that doesn't fit in uInt bits"));
+			return false;
+		}
+
 		stream.next_in = (Bytef*)UncompressedBuffer;
-		stream.avail_in = (uInt)ZUncompressedSize;
+		stream.avail_in = (uInt)UncompressedSize;
 		stream.next_out = (Bytef*)CompressedBuffer;
-		stream.avail_out = (uInt)ZCompressedSize;
+		stream.avail_out = (uInt)CompressedSize;
 		stream.zalloc = &zalloc;
 		stream.zfree = &zfree;
 		stream.opaque = Z_NULL;
 
+		bool bOperationSucceeded = false;
+
 		if (ensure(Z_OK == deflateInit2(&stream, CompLevel, Z_DEFLATED, BitWindow, MAX_MEM_LEVEL, Z_DEFAULT_STRATEGY)))
 		{
 			if (ensure(Z_STREAM_END == deflate(&stream, Z_FINISH)))
 			{
-				ZCompressedSize = stream.total_out;
+				CompressedSize = stream.total_out;
 				if (ensure(Z_OK == deflateEnd(&stream)))
 				{
 					bOperationSucceeded = true;
@@ -117,14 +133,12 @@ static bool appCompressMemoryZLIB(void* CompressedBuffer, int32& CompressedSize,
 				deflateEnd(&stream);
 			}
 		}
-	}
 
-	// Propagate compressed size from intermediate variable back into out variable.
-	CompressedSize = ZCompressedSize;
-	return bOperationSucceeded;
+		return bOperationSucceeded;
+	}
 }
 
-static bool appCompressMemoryGZIP(void* CompressedBuffer, int32& CompressedSize, const void* UncompressedBuffer, int32 UncompressedSize)
+static bool appCompressMemoryGZIP(void* CompressedBuffer, int64& CompressedSize, const void* UncompressedBuffer, int64 UncompressedSize)
 {
 	TRACE_CPUPROFILER_EVENT_SCOPE(appCompressMemoryGZIP);
 	DECLARE_SCOPE_CYCLE_COUNTER( TEXT( "Compress Memory GZIP" ), STAT_appCompressMemoryGZIP, STATGROUP_Compression );
@@ -134,6 +148,13 @@ static bool appCompressMemoryGZIP(void* CompressedBuffer, int32& CompressedSize,
 	gzipstream.zfree = &zfree;
 	gzipstream.opaque = Z_NULL;
 
+	if (!IntFitsIn<decltype(gzipstream.avail_in)>(UncompressedSize) ||
+		!IntFitsIn<decltype(gzipstream.avail_out)>(CompressedSize))
+	{
+		UE_LOG(LogCompression, Error, TEXT("Requested a ZLIB compression that doesn't fit in uInt bits"));
+		return false;
+	}
+
 	// Setup input buffer
 	gzipstream.next_in = (uint8*)UncompressedBuffer;
 	gzipstream.avail_in = UncompressedSize;
@@ -202,11 +223,18 @@ static int appCompressMemoryBoundGZIP(int32 UncompressedSize)
  * @param	CompressedSize				Size of CompressedBuffer data in bytes
  * @return true if compression succeeds, false if it fails because CompressedBuffer was too small or other reasons
  */
-bool appUncompressMemoryGZIP(void* UncompressedBuffer, int32 UncompressedSize, const void* CompressedBuffer, int32 CompressedSize)
+static bool appUncompressMemoryGZIP(void* UncompressedBuffer, int64 UncompressedSize, const void* CompressedBuffer, int64 CompressedSize)
 {
 	TRACE_CPUPROFILER_EVENT_SCOPE(appUncompressMemoryGZIP);
 	DECLARE_SCOPE_CYCLE_COUNTER(TEXT("Uncompress Memory GZIP"), STAT_appUncompressMemoryGZIP, STATGROUP_Compression);
 
+	if (!IntFitsIn<decltype(z_stream::avail_in)>(UncompressedSize) ||
+		!IntFitsIn<decltype(z_stream::avail_out)>(CompressedSize))
+	{
+		UE_LOG(LogCompression, Error, TEXT("GZIP compression: can't fit in unsigned long: 0x%llx or 0x%llx"), UncompressedSize, CompressedSize);
+		return false;
+	}
+
 	// Zlib wants to use unsigned long.
 	unsigned long ZCompressedSize = CompressedSize;
 	unsigned long ZUncompressedSize = UncompressedSize;
@@ -264,11 +292,18 @@ bool appUncompressMemoryGZIP(void* UncompressedBuffer, int32 UncompressedSize, c
  * @param	CompressedSize				Size of CompressedBuffer data in bytes
  * @return true if compression succeeds, false if it fails because CompressedBuffer was too small or other reasons
  */
-bool appUncompressMemoryZLIB( void* UncompressedBuffer, int32 UncompressedSize, const void* CompressedBuffer, int32 CompressedSize, int32 BitWindow )
+static bool appUncompressMemoryZLIB( void* UncompressedBuffer, int64 UncompressedSize, const void* CompressedBuffer, int64 CompressedSize, int32 BitWindow )
 {
 	TRACE_CPUPROFILER_EVENT_SCOPE(appUncompressMemoryZLIB);
 	DECLARE_SCOPE_CYCLE_COUNTER( TEXT( "Uncompress Memory ZLIB" ), STAT_appUncompressMemoryZLIB, STATGROUP_Compression );
 
+	if (!IntFitsIn<decltype(z_stream::avail_in)>(UncompressedSize) ||
+		!IntFitsIn<decltype(z_stream::avail_out)>(CompressedSize))
+	{
+		UE_LOG(LogCompression, Error, TEXT("ZLIB compression: can't fit in unsigned long: 0x%llx or 0x%llx"), UncompressedSize, CompressedSize);
+		return false;
+	}
+
 	// Zlib wants to use unsigned long.
 	unsigned long ZCompressedSize	= CompressedSize;
 	unsigned long ZUncompressedSize	= UncompressedSize;
@@ -321,7 +356,7 @@ bool appUncompressMemoryZLIB( void* UncompressedBuffer, int32 UncompressedSize,
 	return bOperationSucceeded;
 }
 
-bool appUncompressMemoryStreamZLIB(void* UncompressedBuffer, int32 UncompressedSize, IMemoryReadStream* Stream, int64 StreamOffset, int32 CompressedSize, int32 BitWindow)
+static bool appUncompressMemoryStreamZLIB(void* UncompressedBuffer, int64 UncompressedSize, IMemoryReadStream* Stream, int64 StreamOffset, int64 CompressedSize, int32 BitWindow)
 {
 	TRACE_CPUPROFILER_EVENT_SCOPE(appUncompressMemoryStreamZLIB);
 	DECLARE_SCOPE_CYCLE_COUNTER(TEXT("Uncompress Memory ZLIB"), STAT_appUncompressMemoryZLIB, STATGROUP_Compression);
@@ -331,6 +366,13 @@ bool appUncompressMemoryStreamZLIB(void* UncompressedBuffer, int32 UncompressedS
 	const void* ChunkMemory = Stream->Read(ChunkSize, StreamOffset + ChunkOffset, CompressedSize);
 	ChunkOffset += ChunkSize;
 
+	if (!IntFitsIn<decltype(z_stream::avail_in)>(ChunkSize) ||
+		!IntFitsIn<decltype(z_stream::avail_out)>(UncompressedSize))
+	{
+		UE_LOG(LogCompression, Error, TEXT("ZLIB compression: can't fit in unsigned long: 0x%llx or 0x%llx"), UncompressedSize, ChunkSize);
+		return false;
+	}
+
 	z_stream stream;
 	stream.zalloc = &zalloc;
 	stream.zfree = &zfree;
@@ -357,6 +399,12 @@ bool appUncompressMemoryStreamZLIB(void* UncompressedBuffer, int32 UncompressedS
 			ChunkOffset += ChunkSize;
 			check(ChunkOffset <= CompressedSize);
 
+			if (!IntFitsIn<decltype(z_stream::avail_in)>(ChunkSize))
+			{
+				UE_LOG(LogCompression, Error, TEXT("ZLIB compression: can't fit in unsigned long: 0x%llx"), ChunkSize);
+				return false;
+			}
+
 			stream.next_in = (uint8*)ChunkMemory;
 			stream.avail_in = ChunkSize;
 		}
@@ -461,12 +509,12 @@ FName FCompression::GetCompressionFormatFromDeprecatedFlags(ECompressionFlags Fl
 {
 	switch (Flags & COMPRESS_DeprecatedFormatFlagsMask)
 	{
-		case COMPRESS_ZLIB:
+		case COMPRESS_ZLIB_DEPRECATED:
 			return NAME_Zlib;
-		case COMPRESS_GZIP:
+		case COMPRESS_GZIP_DEPRECATED:
 			return NAME_Gzip;
 		// COMPRESS_Custom was a temporary solution to third party compression before we had plugins working, and it was only ever used with oodle, we just assume Oodle with Custom
-		case COMPRESS_Custom:
+		case COMPRESS_Custom_DEPRECATED:
 			return NAME_Oodle;
 		default:
 			break;
@@ -475,75 +523,191 @@ FName FCompression::GetCompressionFormatFromDeprecatedFlags(ECompressionFlags Fl
 	return NAME_None;
 }
 
-int32 FCompression::GetMaximumCompressedSize(FName FormatName, int32 UncompressedSize, ECompressionFlags Flags, int32 CompressionData)
+bool FCompression::GetMaximumCompressedSize(FName FormatName, int64& OutMaxCompressedSize, int64 UncompressedSize, uintptr_t CompressionData)
 {
+	check(UncompressedSize >= 0);
+	if (UncompressedSize < 0)
+	{
+		OutMaxCompressedSize = -1;
+		UE_LOG(LogCompression, Error, TEXT("Negative value passed to GetMaximumCompressedSize (0x%llx)"), UncompressedSize);
+		return false;
+	}
+
 	if (FormatName == NAME_None)
 	{
-		return UncompressedSize;
+		OutMaxCompressedSize = UncompressedSize;
+		return true;
 	}
-	else if ( FormatName == NAME_Oodle )
+	else if (FormatName == NAME_Oodle)
 	{
 		//	avoid calling CompressMemoryBound in the Decoder because it creates an ICompressionFormat for Oodle
-		//	and initializes encoders
-
-		return FOodleDataCompression::GetMaximumCompressedSize(UncompressedSize);
+		//	and initializes encoders (and also is a different value!)
+		// This should be codec independent because it's referring to how much it gets compressed, not the buffer space needed
+		// to compress... and we should be OK for overflow because this just means future oodle size checks will fail, not that
+		// anything will get stomped.
+		OutMaxCompressedSize = FOodleDataCompression::GetMaximumCompressedSize(UncompressedSize);
+		return true;
 	}
-	else
+
+	// If we don't have anything better to use then we just use the compressed buffer size, which is almost certainly
+	// too big but also all we can go on.
+	return CompressMemoryBound(FormatName, OutMaxCompressedSize, UncompressedSize, CompressionData);	
+}
+
+int32 FCompression::GetMaximumCompressedSize(FName FormatName, int32 UncompressedSize, ECompressionFlags Flags, int32 CompressionData)
+{
+	int64 MaxCompressedSize = -1;
+	bool bSucceeded = GetMaximumCompressedSize(FormatName, MaxCompressedSize, UncompressedSize, CompressionData);
+	if (!bSucceeded ||
+		!IntFitsIn<int32>(MaxCompressedSize))
 	{
-		return CompressMemoryBound(FormatName,UncompressedSize,Flags,CompressionData);
+		UE_LOG(LogCompression, Fatal, TEXT("GetMaximumCompressedSize failed, check sizes/format (%d, %s)"), UncompressedSize, *FormatName.ToString());
+		return -1;
 	}
+	return (int32)MaxCompressedSize;
 }
 
-int32 FCompression::CompressMemoryBound(FName FormatName, int32 UncompressedSize, ECompressionFlags Flags, int32 CompressionData)
+bool FCompression::CompressMemoryBound(FName FormatName, int64& OutBufferSizeRequired, int64 UncompressedSize, uintptr_t CompressionData)
 {
-	int32 CompressionBound = UncompressedSize;
-	
+	// Init to a garbage value so if they don't pay attention to the return value they 
+	// crash allocating a massive buffer.
+	OutBufferSizeRequired = -1;
+
+	check(UncompressedSize >= 0);
+	if (UncompressedSize < 0)
+	{
+		UE_LOG(LogCompression, Error, TEXT("Negative value passed to CompressMemoryBound (0x%llx)"), UncompressedSize);
+		return false;
+	}
+
 	if (FormatName == NAME_None)
 	{
-		return UncompressedSize;
+		OutBufferSizeRequired = UncompressedSize;
+		return true;
 	}
 	else if (FormatName == NAME_Zlib)
 	{
+		if (!IntFitsIn<uLong>(UncompressedSize))
+		{
+			UE_LOG(LogCompression, Error, TEXT("Zlib doesn't support >32 bit sizes (0x%llx)"), UncompressedSize);
+			return false;
+		}
+
 		// Zlib's compressBounds gives a better (smaller) value, but only for the default bit window.
 		if (CompressionData == 0 || CompressionData == DEFAULT_ZLIB_BIT_WINDOW)
 		{
-			CompressionBound = compressBound(UncompressedSize);
+			OutBufferSizeRequired = compressBound(UncompressedSize);
 		}
 		else
 		{
 			// Calculate pessimistic bounds for compression. This value is calculated based on the algorithm used in deflate2.
-			CompressionBound = UncompressedSize + ((UncompressedSize + 7) >> 3) + ((UncompressedSize + 63) >> 6) + 5 + 6;
+			OutBufferSizeRequired = UncompressedSize + ((UncompressedSize + 7) >> 3) + ((UncompressedSize + 63) >> 6) + 5 + 6;
+			if (OutBufferSizeRequired < 0)
+			{
+				UE_LOG(LogCompression, Error, TEXT("Zlib CompressMemoryBound calculated negative value 0x%llx -> 0x%llx"), UncompressedSize, OutBufferSizeRequired);
+				return false;
+			}
 		}
+		return true;
 	}
 	else if (FormatName == NAME_Gzip)
 	{
-		// Calculate gzip bounds for compression.
-		CompressionBound = appCompressMemoryBoundGZIP(UncompressedSize);
+		if (!IntFitsIn<int32>(UncompressedSize))
+		{
+			UE_LOG(LogCompression, Error, TEXT("Gzip doesn't support >32 bit sizes (0x%llx)"), UncompressedSize);
+			return false;
+		}
+
+		OutBufferSizeRequired = appCompressMemoryBoundGZIP((int32)UncompressedSize);
+		if (OutBufferSizeRequired < 0)
+		{
+			UE_LOG(LogCompression, Error, TEXT("Gzip CompressMemoryBound calculated negative value 0x%llx -> 0x%llx"), UncompressedSize, OutBufferSizeRequired);
+			return false;
+		}
+		return true;
 	}
 	else if (FormatName == NAME_LZ4)
 	{
-		// hardcoded lz4
-		CompressionBound = LZ4_compressBound(UncompressedSize);
-	}
-	else
-	{
-		ICompressionFormat* Format = GetCompressionFormat(FormatName);
-		if (Format)
+		if (UncompressedSize > LZ4_MAX_INPUT_SIZE)
+		{
+			UE_LOG(LogCompression, Error, TEXT("LZ4 doesn't support >32 bit sizes (0x%llx) max is 0x%x"), UncompressedSize, LZ4_MAX_INPUT_SIZE);
+			return false;
+		}
+		OutBufferSizeRequired = LZ4_compressBound((int)UncompressedSize);
+		if (OutBufferSizeRequired < 0)
 		{
-			CompressionBound = Format->GetCompressedBufferSize(UncompressedSize, CompressionData);
+			UE_LOG(LogCompression, Error, TEXT("LZ4 CompressMemoryBound calculated negative value 0x%llx -> 0x%llx"), UncompressedSize, OutBufferSizeRequired);
+			return false;
 		}
+		return true;
+	}
+
+	ICompressionFormat* Format = GetCompressionFormat(FormatName);
+	if (!Format)
+	{
+		return false;
+	}
+
+	if (!Format->GetCompressedBufferSize(OutBufferSizeRequired, UncompressedSize, CompressionData))
+	{
+		UE_LOG(LogCompression, Error, TEXT("GetCompressedBufferSize for format %s failed to return compression bound: check bits needed? (0x%llx)"), *FormatName.ToString(), UncompressedSize);
+		return false;
 	}
 
-	// this will fail if there was an int32 overflow (CompressionBound will be negative)
-	check( CompressionBound >= UncompressedSize );
+	if (OutBufferSizeRequired < 0)
+	{
+		UE_LOG(LogCompression, Error, TEXT("%s CompressMemoryBound calculated negative value 0x%llx -> 0x%llx"), *FormatName.ToString(), UncompressedSize, OutBufferSizeRequired);
+		return false;
+	}
 
-	return CompressionBound;
+	return true;
 }
 
+// 32 thunk to 64 bit.
+int32 FCompression::CompressMemoryBound(FName FormatName, int32 UncompressedSize, ECompressionFlags Flags, int32 CompressionData)
+{
+	int64 BufferSizeNeeded = 0;
+	bool bSucceeded = CompressMemoryBound(FormatName, BufferSizeNeeded, UncompressedSize, CompressionData);
+	if (!bSucceeded ||
+		!IntFitsIn<int32>(BufferSizeNeeded))
+	{
+		UE_LOG(LogCompression, Fatal, TEXT("CompressMemoryBound failed, check sizes/format (%d, %s)"), UncompressedSize, *FormatName.ToString());
+		return -1;
+	}
 
+	return (int32)BufferSizeNeeded;
+}
 
+// 32 thunk to 64 bit.
 bool FCompression::CompressMemoryIfWorthDecompressing(FName FormatName, int32 MinBytesSaved, int32 MinPercentSaved, void* CompressedBuffer, int32& CompressedSize, const void* UncompressedBuffer, int32 UncompressedSize, ECompressionFlags Flags, int32 CompressionData)
 {
+	int64 CompressedSize64 = CompressedSize;
+	bool bWasCompressed = false;
+	bool bSucceeded = CompressMemoryIfWorthDecompressing(FormatName, bWasCompressed, MinBytesSaved, MinPercentSaved, CompressedBuffer, CompressedSize64, UncompressedBuffer, UncompressedSize, Flags, CompressionData);
+
+	// Behavior change: compression failure now FATALs instead of returning "wasn't compressed" (in addition to the type narrowing now caught)
+	if (!bSucceeded ||
+		!IntFitsIn<int32>(CompressedSize64))
+	{
+		UE_LOG(LogCompression, Fatal, TEXT("CompressMemoryIfWorthDecompressing failed, check sizes/format (%d, %s)"), UncompressedSize, *FormatName.ToString());
+		return false;
+	}
+
+	CompressedSize = (int32)CompressedSize64;
+	return bWasCompressed;
+}
+
+bool FCompression::CompressMemoryIfWorthDecompressing(FName FormatName, bool& bOutWasCompressed, int64 MinBytesSaved, int32 MinPercentSaved, void* CompressedBuffer, int64& CompressedSize, const void* UncompressedBuffer, int64 UncompressedSize, ECompressionFlags Flags, uintptr_t CompressionData)
+{
+	// init to false so that if they ignore the return they just pass it uncompressed.
+	bOutWasCompressed = false;
+	if (UncompressedSize < -1 ||
+		CompressedSize < -1)
+	{
+		UE_LOG(LogCompression, Error, TEXT("Negative value passed to CompressMemoryIfWorthDecompressing (0x%llx / 0x%llx)"), UncompressedSize, CompressedSize);
+		return false;
+	}
+
 	// returns false if we could compress,
 	//	but it's not worth the time to decompress
 	//	you should store the data uncompressed instead
@@ -554,7 +718,8 @@ bool FCompression::CompressMemoryIfWorthDecompressing(FName FormatName, int32 Mi
 		// no need to even try encoding
 		// also saves encode time
 		// NOTE : this check applies even for compressor who say "bNeedsWorthItCheck = false" , eg. Oodle
-		return false;
+		bOutWasCompressed = false;
+		return true;
 	}
 	
 	bool bNeedsWorthItCheck = false;
@@ -572,7 +737,10 @@ bool FCompression::CompressMemoryIfWorthDecompressing(FName FormatName, int32 Mi
 	{
 		ICompressionFormat* Format = GetCompressionFormat(FormatName);
 
-		if ( ! Format ) return false;
+		if ( ! Format )
+		{
+			return false;
+		}
 
 		bNeedsWorthItCheck = ! Format->DoesOwnWorthDecompressingCheck();
 	}
@@ -590,7 +758,8 @@ bool FCompression::CompressMemoryIfWorthDecompressing(FName FormatName, int32 Mi
 		// ICompressionFormat does own "worth it" check, don't do our own
 		// do check for expansion because that's how they signal "not worth it"
 		//	(CompressMemory is not allowed to return false)
-		return ( CompressedSize < UncompressedSize );
+		bOutWasCompressed = ( CompressedSize < UncompressedSize );
+		return true;
 	}
 
 	// we got compression, but do we want it ?
@@ -601,18 +770,33 @@ bool FCompression::CompressMemoryIfWorthDecompressing(FName FormatName, int32 Mi
 
 	// must save at least MinBytesSaved regardless of percentage (for small files)
 	// also checks CompressedSize >= UncompressedSize
-	int32 BytesSaved = UncompressedSize - CompressedSize;
+	int64 BytesSaved = UncompressedSize - CompressedSize;
 	if ( BytesSaved < MinBytesSaved )
 	{
-		return false;
+		bOutWasCompressed = false;
+		return true;
 	}
 
 	// Check the saved compression ratio, if it's too low just store uncompressed. 
 	// For example, saving 64 KB per 1 MB is about 6%.
-	return (int64)BytesSaved * 100 >= (int64)UncompressedSize * MinPercentSaved;
+	bOutWasCompressed = BytesSaved * 100 >= UncompressedSize * MinPercentSaved;
+	return true;
 }
 
 bool FCompression::CompressMemory(FName FormatName, void* CompressedBuffer, int32& CompressedSize, const void* UncompressedBuffer, int32 UncompressedSize, ECompressionFlags Flags, int32 CompressionData)
+{
+	int64 CompressedSize64 = CompressedSize;
+	bool bSucceeded = CompressMemory(FormatName, CompressedBuffer, CompressedSize64, UncompressedBuffer, UncompressedSize, Flags, CompressionData);
+	if (!IntFitsIn<int32>(CompressedSize64))
+	{
+		UE_LOG(LogCompression, Error, TEXT("Compressing 32 bit memory size ended up a 64 bit size! %d -> %lld"), UncompressedSize, CompressedSize64);
+		return false;
+	}
+	CompressedSize = (int32)CompressedSize64;
+	return bSucceeded;
+}
+
+bool FCompression::CompressMemory(FName FormatName, void* CompressedBuffer, int64& CompressedSize, const void* UncompressedBuffer, int64 UncompressedSize, ECompressionFlags Flags, uintptr_t CompressionData)
 {
 	TRACE_CPUPROFILER_EVENT_SCOPE(FCompression::CompressMemory);
 	uint64 CompressorStartTime = FPlatformTime::Cycles64();
@@ -632,6 +816,12 @@ bool FCompression::CompressMemory(FName FormatName, void* CompressedBuffer, int3
 	else if (FormatName == NAME_LZ4)
 	{
 		// hardcoded lz4
+		if (UncompressedSize > LZ4_MAX_INPUT_SIZE)
+		{
+			UE_LOG(LogCompression, Error, TEXT("LZ4 can't compress larger than 0x%x (passed 0x%llx)"), LZ4_MAX_INPUT_SIZE, UncompressedSize);
+			return false;
+		}
+
 		CompressedSize = LZ4_compress_HC((const char*)UncompressedBuffer, (char*)CompressedBuffer, UncompressedSize, CompressedSize, LZ4HC_CLEVEL_MAX);
 		bCompressSucceeded = CompressedSize > 0;
 	}
@@ -693,7 +883,7 @@ FString FCompression::GetCompressorDDCSuffix(FName FormatName)
 
 DECLARE_FLOAT_ACCUMULATOR_STAT(TEXT("Uncompressor total time"),STAT_UncompressorTime,STATGROUP_Compression);
 
-bool FCompression::UncompressMemory(FName FormatName, void* UncompressedBuffer, int32 UncompressedSize, const void* CompressedBuffer, int32 CompressedSize, ECompressionFlags Flags, int32 CompressionData)
+bool FCompression::UncompressMemory(FName FormatName, void* UncompressedBuffer, int64 UncompressedSize, const void* CompressedBuffer, int64 CompressedSize, ECompressionFlags Flags, uintptr_t CompressionData)
 {
 	TRACE_CPUPROFILER_EVENT_SCOPE(FCompression::UncompressMemory);
 	SCOPED_NAMED_EVENT(FCompression_UncompressMemory, FColor::Cyan);
@@ -714,6 +904,13 @@ bool FCompression::UncompressMemory(FName FormatName, void* UncompressedBuffer,
 	}
 	else if (FormatName == NAME_LZ4)
 	{
+		if (!IntFitsIn<int>(UncompressedSize) ||
+			!IntFitsIn<int>(CompressedSize))
+		{
+			UE_LOG(LogCompression, Error, TEXT("LZ4 can't fit in int: 0x%llx or 0x%llx"), CompressedSize, UncompressedSize);
+			return false;
+		}
+
 		// hardcoded lz4
 		bUncompressSucceeded = LZ4_decompress_safe((const char*)CompressedBuffer, (char*)UncompressedBuffer, CompressedSize, UncompressedSize) > 0;
 	}
@@ -729,8 +926,6 @@ bool FCompression::UncompressMemory(FName FormatName, void* UncompressedBuffer,
 		ICompressionFormat* Format = GetCompressionFormat(FormatName);
 		if (Format)
 		{
-			// @todo Oodle : UncompressedSize is read-write to Format->Uncompress but that's not actually used
-			//	-> get rid of that
 			bUncompressSucceeded = Format->Uncompress(UncompressedBuffer, UncompressedSize, CompressedBuffer, CompressedSize, CompressionData);
 		}
 	}
@@ -777,7 +972,7 @@ bool FCompression::UncompressMemory(FName FormatName, void* UncompressedBuffer,
 	return bUncompressSucceeded;
 }
 
-bool FCompression::UncompressMemoryStream(FName FormatName, void* UncompressedBuffer, int32 UncompressedSize, IMemoryReadStream* Stream, int64 StreamOffset, int32 CompressedSize, ECompressionFlags Flags, int32 CompressionData)
+bool FCompression::UncompressMemoryStream(FName FormatName, void* UncompressedBuffer, int64 UncompressedSize, IMemoryReadStream* Stream, int64 StreamOffset, int64 CompressedSize, ECompressionFlags Flags, uintptr_t CompressionData)
 {
 	TRACE_CPUPROFILER_EVENT_SCOPE(FCompression::UncompressMemoryStream);
 
@@ -819,18 +1014,6 @@ bool FCompression::UncompressMemoryStream(FName FormatName, void* UncompressedBu
 	FCompressedGrowableBuffer.
 -----------------------------------------------------------------------------*/
 
-/**
- * Constructor
- *
- * @param	InMaxPendingBufferSize	Max chunk size to compress in uncompressed bytes
- * @param	InCompressionFlags		Compression flags to compress memory with
- */
-FCompressedGrowableBuffer::FCompressedGrowableBuffer(FCompressedGrowableBuffer::EVS2015Redirector, int32 InMaxPendingBufferSize, ECompressionFlags InCompressionFlags)
-	: FCompressedGrowableBuffer(InMaxPendingBufferSize, FCompression::GetCompressionFormatFromDeprecatedFlags(InCompressionFlags), InCompressionFlags)
-{
-
-}
-
 FCompressedGrowableBuffer::FCompressedGrowableBuffer( int32 InMaxPendingBufferSize, FName InCompressionFormat, ECompressionFlags InCompressionFlags )
 :	MaxPendingBufferSize( InMaxPendingBufferSize )
 ,	CompressionFormat( InCompressionFormat )
@@ -1016,7 +1199,6 @@ bool FCompression::VerifyCompressionFlagsValid(int32 InCompressionFlags)
 }
 
 
-
 /***********************
   Deprecated functions
 ***********************/
diff --git a/Engine/Source/Runtime/Core/Private/Serialization/Archive.cpp b/Engine/Source/Runtime/Core/Private/Serialization/Archive.cpp
--- a/Engine/Source/Runtime/Core/Private/Serialization/Archive.cpp
+++ b/Engine/Source/Runtime/Core/Private/Serialization/Archive.cpp
@@ -768,7 +768,9 @@ void FArchive::SerializeCompressedNew(void* V, int64 Length, FName CompressionFo
 			if ((Flags & COMPRESS_DeprecatedFormatFlagsMask) != 0)
 			{
 				UE_LOG(LogSerialization, Warning, TEXT("Old style compression flags are being used with FAsyncCompressionChunk, please update any code using this!"));
+				PRAGMA_DISABLE_DEPRECATION_WARNINGS
 				CompressionFormatToDecode = FCompression::GetCompressionFormatFromDeprecatedFlags(Flags);
+				PRAGMA_ENABLE_DEPRECATION_WARNINGS
 			}
 
 			if (CompressionFormatToDecode == NAME_Zlib && FPlatformProperties::GetZlibReplacementFormat() != nullptr)
@@ -919,7 +921,9 @@ void FArchive::SerializeCompressedNew(void* V, int64 Length, FName CompressionFo
 		{
 			check( CompressionFormatToEncode == NAME_Zlib );
 			UE_LOG(LogSerialization, Warning, TEXT("Old style compression flags are being used with FAsyncCompressionChunk, please update any code using this!"));
+			PRAGMA_DISABLE_DEPRECATION_WARNINGS
 			CompressionFormatToEncode = FCompression::GetCompressionFormatFromDeprecatedFlags(Flags);
+			PRAGMA_ENABLE_DEPRECATION_WARNINGS
 		}
 
 		// if there's a cooking target, and it wants to replace Zlib compression with another format, use it. When loading, 
diff --git a/Engine/Source/Runtime/Core/Private/Serialization/ArchiveSaveCompressedProxy.cpp b/Engine/Source/Runtime/Core/Private/Serialization/ArchiveSaveCompressedProxy.cpp
--- a/Engine/Source/Runtime/Core/Private/Serialization/ArchiveSaveCompressedProxy.cpp
+++ b/Engine/Source/Runtime/Core/Private/Serialization/ArchiveSaveCompressedProxy.cpp
@@ -15,11 +15,6 @@ PRAGMA_DISABLE_UNSAFE_TYPECAST_WARNINGS
 /*----------------------------------------------------------------------------
 	FArchiveSaveCompressedProxy
 ----------------------------------------------------------------------------*/
-FArchiveSaveCompressedProxy::FArchiveSaveCompressedProxy(FArchiveSaveCompressedProxy::EVS2015Redirector, TArray<uint8>& InCompressedData, ECompressionFlags InCompressionFlags)
-	: FArchiveSaveCompressedProxy(InCompressedData, FCompression::GetCompressionFormatFromDeprecatedFlags(InCompressionFlags), InCompressionFlags)
-{
-}
-
 FArchiveSaveCompressedProxy::FArchiveSaveCompressedProxy(TArray<uint8>& InCompressedData, FName InCompressionFormat, ECompressionFlags InCompressionFlags)
 	: CompressedData(InCompressedData)
 	, CompressionFormat(InCompressionFormat)
diff --git a/Engine/Source/Runtime/Core/Public/Misc/CompressedGrowableBuffer.h b/Engine/Source/Runtime/Core/Public/Misc/CompressedGrowableBuffer.h
--- a/Engine/Source/Runtime/Core/Public/Misc/CompressedGrowableBuffer.h
+++ b/Engine/Source/Runtime/Core/Public/Misc/CompressedGrowableBuffer.h
@@ -12,16 +12,6 @@
  */
 struct FCompressedGrowableBuffer
 {
-private:
-	/**
-	 * This enum and the following constructor is a workaround for VC compiler bug that prevents using attributes
-	 * on constructors without inline implementation. This should be removed when the deprecated ctor is removed.
-	 */
-	enum EVS2015Redirector
-	{
-		Redirect
-	};
-	CORE_API FCompressedGrowableBuffer(EVS2015Redirector, int32 MaxPendingBufferSize, ECompressionFlags CompressionFlags);
 public:
 	/**
 	 * Constructor
@@ -29,12 +19,7 @@ public:
 	 * @param	MaxPendingBufferSize	Max chunk size to compress in uncompressed bytes
 	 * @param	CompressionFlags		Compression flags to compress memory with
 	 */
-	UE_DEPRECATED(4.21, "Use FName version of FCompressedGrowableBuffer constructor")
-	FCompressedGrowableBuffer(int32 MaxPendingBufferSize, ECompressionFlags CompressionFlags)
-		// Make sure to remove the EVS2015Redirector constructor when this constructor is removed
-		: FCompressedGrowableBuffer(EVS2015Redirector::Redirect, MaxPendingBufferSize, CompressionFlags)
-	{}
-	CORE_API FCompressedGrowableBuffer(int32 MaxPendingBufferSize, FName COmpressionFormat, ECompressionFlags CompressionFlags=COMPRESS_None);
+	CORE_API FCompressedGrowableBuffer(int32 MaxPendingBufferSize, FName CompressionFormat, ECompressionFlags CompressionFlags=COMPRESS_None);
 
 	/**
 	 * Locks the buffer for reading. Needs to be called before calls to Access and needs
diff --git a/Engine/Source/Runtime/Core/Public/Misc/Compression.h b/Engine/Source/Runtime/Core/Public/Misc/Compression.h
--- a/Engine/Source/Runtime/Core/Public/Misc/Compression.h
+++ b/Engine/Source/Runtime/Core/Public/Misc/Compression.h
@@ -14,8 +14,8 @@ class IMemoryReadStream;
 template <typename T> class TAtomic;
 
 // Define global current platform default to current platform.  
-// DEPRECATED, USE NAME_Zlib
-#define COMPRESS_Default			COMPRESS_ZLIB
+// DEPRECATED, USE NAME_Zlib - was deprecated in 2018, leaving for grep until 5.6 or so.
+//#define COMPRESS_Default			COMPRESS_ZLIB
 
 /**
  * Chunk size serialization code splits data into. The loading value CANNOT be changed without resaving all
@@ -27,6 +27,15 @@ template <typename T> class TAtomic;
 
 struct FCompression
 {
+	/**
+	* CompressionData is per-compressor data passed through this API and is interpreted by the codecs.
+	* For builtins:
+	*	NAME_Zlib	BitWindow
+	*	NAME_LZ4	Unused
+	*	NAME_Gzip	Unused
+	*	NAME_Oodle	Unused
+	*/
+
 	/** Time spent compressing data in cycles. */
 	CORE_API static TAtomic<uint64> CompressorTimeCycles;
 	/** Number of bytes before compression.		*/
@@ -44,29 +53,41 @@ struct FCompression
 	CORE_API static uint32 GetCompressorVersion(FName FormatName);
 
 	/**
-	 * Thread-safe abstract compression routine to query memory requirements for a compression operation.
+	 * Thread-safe abstract compression routine to query buffer requirements for a compression operation.
 	 * This is the minimize size to allocate the buffer for CompressMemory (encoding).
 	 * Use GetMaximumCompressedSize at decode to know how large a compressed buffer may be.
 	 *
+	 * Despite the name, this has nothing to do with how much memory will be allocated during compression.
+	 * 
 	 * @param	FormatName					Name of the compression format
 	 * @param	UncompressedSize			Size of uncompressed data in bytes
-	 * @param	Flags						Flags to control what method to use and optionally control memory vs speed
-	 * @param	CompressionData				Additional compression parameter (specifies BitWindow value for ZLIB compression format)
-	 * @return The maximum possible bytes needed for compression of data buffer of size UncompressedSize
-	 */
-	CORE_API static int32 CompressMemoryBound(FName FormatName, int32 UncompressedSize, ECompressionFlags Flags=COMPRESS_NoFlags, int32 CompressionData=0);
+	 * @param	CompressionData				Additional compression parameter - see comments a top of FCompression struct.
+	 * @param	OutBufferSizeRequired		The maximum possible bytes needed for compression of a data buffer of the given size.
+	 * @return  false if the bounds couldn't be determined - e.g. the format is unknown or the size isn't supported by the format.
+	*/
+	CORE_API static bool CompressMemoryBound(FName InFormatName, int64& OutBufferSizeRequired, int64 InUncompressedSizeBytes, uintptr_t InCompressionData = 0);
+
+	// Thunks to the 64 bit, prefer that as this can only Fatal on error!
+	CORE_API static int32 CompressMemoryBound(FName FormatName, int32 UncompressedSize, ECompressionFlags Flags = COMPRESS_NoFlags, int32 CompressionData = 0);
 	
 	/**
-	 * Thread-safe abstract compression routine to query maximum compressed size that could be made.
+	 * Thread-safe abstract compression routine to query maximum compressed size that could be made. This is used for making buffers that could hold
+	 * any compressed buffer that started uncompressed with the given size. It's distinct from the buffer needed to compress! If you are calling CompressMemory then
+	 * you need to pass a buffer sized based on CompressMemoryBound!
+	 * 
 	 * CompressMemoryBound is strictly greater equal GetMaximumCompressedSize.
 	 *
 	 * @param	FormatName					Name of the compression format
-	 * @param	UncompressedSize			Size of uncompressed data in bytes
-	 * @param	Flags						Flags to control what method to use and optionally control memory vs speed
-	 * @param	CompressionData				Additional compression parameter (specifies BitWindow value for ZLIB compression format)
-	 * @return The maximum possible size of valid compressed data made by this format
+	 * @param	OutMaxCompressedSize		The maximum possible size of valid compressed data made by this format
+	 * @param	InUncompressedSizeBytes		Size of uncompressed data in bytes	 
+	 * @param	CompressionData				Additional compression parameter - see comments a top of FCompression struct.
+	 * @return success
 	 */
-	CORE_API static int32 GetMaximumCompressedSize(FName FormatName, int32 UncompressedSize, ECompressionFlags Flags=COMPRESS_NoFlags, int32 CompressionData=0);
+	CORE_API static bool GetMaximumCompressedSize(FName InFormatName, int64& OutMaxCompressedSize, int64 InUncompressedSizeBytes, uintptr_t InCompressionData=0);
+
+	// Thunks to the 64 bit version, prefer that as this can only Fatal on error!
+	UE_DEPRECATED(5.5, "Use the 64 bit version")
+	CORE_API static int32 GetMaximumCompressedSize(FName FormatName, int32 UncompressedSize, ECompressionFlags Flags = COMPRESS_NoFlags, int32 CompressionData = 0);
 
 	/**
 	 * Thread-safe abstract compression routine. Compresses memory from uncompressed buffer and writes it to compressed
@@ -80,17 +101,19 @@ struct FCompression
 	 * @param	UncompressedBuffer			Buffer containing uncompressed data
 	 * @param	UncompressedSize			Size of uncompressed data in bytes
 	 * @param	Flags						Flags to control what method to use and optionally control memory vs speed
-	 * @param	CompressionData				Additional compression parameter (specifies BitWindow value for ZLIB compression format)
+	 * @param	CompressionData				Additional compression parameter - see comments a top of FCompression struct.
 	 * @return true if compression succeeds, false if it fails because CompressedBuffer was too small or other reasons
 	 */
-	CORE_API static bool CompressMemory(FName FormatName, void* CompressedBuffer, int32& CompressedSize, const void* UncompressedBuffer, int32 UncompressedSize, ECompressionFlags Flags=COMPRESS_NoFlags, int32 CompressionData=0);
+	CORE_API static bool CompressMemory(FName FormatName, void* CompressedBuffer, int64& CompressedSize, const void* UncompressedBuffer, int64 UncompressedSize, ECompressionFlags Flags = COMPRESS_NoFlags, uintptr_t CompressionData = 0);
+
+	// Thunks to 64 bit CompressMemory, favor that.
+	CORE_API static bool CompressMemory(FName FormatName, void* CompressedBuffer, int32& CompressedSize, const void* UncompressedBuffer, int32 UncompressedSize, ECompressionFlags Flags = COMPRESS_NoFlags, int32 CompressionData = 0);
 	
 	/**
-	* Same as CompressMemory but evaluates if the compression gain is worth the runtime decode time
-	* returns false if the size saving is not worth it (also if CompressedSize >= UncompressedSize)
-	* if false is returned, send the data uncompressed instead
+	* Same as CompressMemory but evaluates if the compression gain is worth the runtime decode time.
 	 *
 	 * @param	FormatName					Name of the compression format
+	 * @param	bOutWasCompressed [out]		Whether the compression was worth it. The data is still compressed.
 	 * @param	MinBytesSaved				Minimum amount of bytes which should be saved when performing compression, otherwise false is returned
 	 * @param	MinPercentSaved				Minimum percentage of the buffer which should be saved when performing compression, otherwise false is returned
 	 * @param	CompressedBuffer			Buffer compressed data is going to be written to
@@ -98,10 +121,13 @@ struct FCompression
 	 * @param	UncompressedBuffer			Buffer containing uncompressed data
 	 * @param	UncompressedSize			Size of uncompressed data in bytes
 	 * @param	Flags						Flags to control what method to use and optionally control memory vs speed
-	 * @param	CompressionData				Additional compression parameter (specifies BitWindow value for ZLIB compression format)
+	 * @param	CompressionData				Additional compression parameter - see comments a top of FCompression struct.
 	 * @return true if compression succeeds, false if it fails because CompressedBuffer was too small or other reasons
 	 */
-	CORE_API static bool CompressMemoryIfWorthDecompressing(FName FormatName, int32 MinBytesSaved, int32 MinPercentSaved, void* CompressedBuffer, int32& CompressedSize, const void* UncompressedBuffer, int32 UncompressedSize, ECompressionFlags Flags=COMPRESS_NoFlags, int32 CompressionData=0);
+	CORE_API static bool CompressMemoryIfWorthDecompressing(FName FormatName, bool& bOutWasCompressed, int64 MinBytesSaved, int32 MinPercentSaved, void* CompressedBuffer, int64& CompressedSize, const void* UncompressedBuffer, int64 UncompressedSize, ECompressionFlags Flags = COMPRESS_NoFlags, uintptr_t CompressionData = 0);
+
+	// 32 bit thunk. This legacy function doens't separate bOutWasCompressed vs general compress failure and will return false in both cases.
+	CORE_API static bool CompressMemoryIfWorthDecompressing(FName FormatName, int32 MinBytesSaved, int32 MinPercentSaved, void* CompressedBuffer, int32& CompressedSize, const void* UncompressedBuffer, int32 UncompressedSize, ECompressionFlags Flags = COMPRESS_NoFlags, int32 CompressionData = 0);
 
 	/**
 	 * Thread-safe abstract decompression routine. Uncompresses memory from compressed buffer and writes it to uncompressed
@@ -116,9 +142,15 @@ struct FCompression
 	 * @param	CompressionData				Additional decompression parameter (specifies BitWindow value for ZLIB compression format)
 	 * @return true if compression succeeds, false if it fails because CompressedBuffer was too small or other reasons
 	 */
-	CORE_API static bool UncompressMemory(FName FormatName, void* UncompressedBuffer, int32 UncompressedSize, const void* CompressedBuffer, int32 CompressedSize, ECompressionFlags Flags=COMPRESS_NoFlags, int32 CompressionData=0);
+	CORE_API static bool UncompressMemory(FName FormatName, void* UncompressedBuffer, int64 UncompressedSize, const void* CompressedBuffer, int64 CompressedSize, ECompressionFlags Flags=COMPRESS_NoFlags, uintptr_t CompressionData=0);
+
+	/**
+	* Wraps UncompressMemory to provide access with a IMemoryReadStream interface, attempting to avoid copying in to a contiguous buffer doe decompression
+	* if the format supports it. **If not, a CompressedSize temp buffer will be allocated.**
+	*/
+	CORE_API static bool UncompressMemoryStream(FName FormatName, void* UncompressedBuffer, int64 UncompressedSize, IMemoryReadStream* Stream, int64 StreamOffset, int64 CompressedSize, ECompressionFlags Flags = COMPRESS_NoFlags, uintptr_t CompressionData = 0);
+
 
-	CORE_API static bool UncompressMemoryStream(FName FormatName, void* UncompressedBuffer, int32 UncompressedSize, IMemoryReadStream* Stream, int64 StreamOffset, int32 CompressedSize, ECompressionFlags Flags = COMPRESS_NoFlags, int32 CompressionData = 0);
 	/**
 	 * Returns a string which can be used to identify if a format has become out of date
 	 *
@@ -127,7 +159,6 @@ struct FCompression
 	 */
 	CORE_API static FString GetCompressorDDCSuffix(FName FormatName);
 
-
 	/**
 	 * Checks to see if a format will be usable, so that a fallback can be used
 	 * @param FormatName The name of the format to test
@@ -140,6 +171,8 @@ struct FCompression
 	 */
 	CORE_API static bool VerifyCompressionFlagsValid(int32 InCompressionFlags);
 
+	// deprecating this because we "should", but all internal references to these flags were deprecated in 4.21 so hopefully no one is using them anymore!
+	UE_DEPRECATED(5.5, "Use the direct FName instead e.g NAME_Oodle")
 	CORE_API static FName GetCompressionFormatFromDeprecatedFlags(ECompressionFlags DeprecatedFlags);
 
 private:
diff --git a/Engine/Source/Runtime/Core/Public/Misc/CompressionFlags.h b/Engine/Source/Runtime/Core/Public/Misc/CompressionFlags.h
--- a/Engine/Source/Runtime/Core/Public/Misc/CompressionFlags.h
+++ b/Engine/Source/Runtime/Core/Public/Misc/CompressionFlags.h
@@ -11,14 +11,15 @@
 */
 enum ECompressionFlags
 {
-	/** No compression																*/
+	/** No flags																*/
 	COMPRESS_None = 0x00,
+
 	/** Compress with ZLIB - DEPRECATED, USE FNAME									*/
-	COMPRESS_ZLIB = 0x01,
+	COMPRESS_ZLIB_DEPRECATED = 0x01,
 	/** Compress with GZIP - DEPRECATED, USE FNAME									*/
-	COMPRESS_GZIP = 0x02,
+	COMPRESS_GZIP_DEPRECATED = 0x02,
 	/** Compress with user defined callbacks - DEPRECATED, USE FNAME                */
-	COMPRESS_Custom = 0x04,
+	COMPRESS_Custom_DEPRECATED = 0x04,
 	/** Joint of the previous ones to determine if old flags are being used			*/
 	COMPRESS_DeprecatedFormatFlagsMask = 0xF,
 
diff --git a/Engine/Source/Runtime/Core/Public/Misc/ICompressionFormat.h b/Engine/Source/Runtime/Core/Public/Misc/ICompressionFormat.h
--- a/Engine/Source/Runtime/Core/Public/Misc/ICompressionFormat.h
+++ b/Engine/Source/Runtime/Core/Public/Misc/ICompressionFormat.h
@@ -14,10 +14,66 @@
 struct ICompressionFormat : public IModularFeature, public IModuleInterface
 {
 	virtual FName GetCompressionFormatName() = 0;
-	virtual bool Compress(void* CompressedBuffer, int32& CompressedSize, const void* UncompressedBuffer, int32 UncompressedSize, int32 CompressionData, ECompressionFlags Flags) = 0;
-	virtual bool Uncompress(void* UncompressedBuffer, int32& UncompressedSize, const void* CompressedBuffer, int32 CompressedSize, int32 CompressionData) = 0;
-	virtual int32 GetCompressedBufferSize(int32 UncompressedSize, int32 CompressionData) = 0;
+
+	UE_DEPRECATED(5.5, "Switch to 64 bit version for memory sizes")
+	virtual bool Compress(void* CompressedBuffer, int32& CompressedSize, const void* UncompressedBuffer, int32 UncompressedSize, int32 CompressionData, ECompressionFlags Flags) { return false; };
+
+	// Default 64 bit implementation wraps legacy 32 bit implementation with size checks - will be removed when deprecated function is removed
+	virtual bool Compress(void* CompressedBuffer, int64& CompressedSize, const void* UncompressedBuffer, int64 UncompressedSize, uintptr_t CompressionData, ECompressionFlags Flags)
+	{
+		if (!IntFitsIn<int32>(UncompressedSize) ||
+			!IntFitsIn<int32>(CompressedSize) ||
+			!IntFitsIn<int32>(CompressionData))
+		{
+			return false;
+		}
+		PRAGMA_DISABLE_DEPRECATION_WARNINGS;
+		return Compress(CompressedBuffer, (int32&)CompressedSize, UncompressedBuffer, (int32)UncompressedSize, (int32)CompressionData, Flags);
+		PRAGMA_ENABLE_DEPRECATION_WARNINGS;
+	}
+
+	UE_DEPRECATED(5.5, "Switch to 64 bit version for memory sizes")
+	virtual bool Uncompress(void* UncompressedBuffer, int32& UncompressedSize, const void* CompressedBuffer, int32 CompressedSize, int32 CompressionData) {return false;};
+
+	// Default 64 bit implementation wraps legacy 32 bit implementation with size checks - will be removed when deprecated function is removed
+	virtual bool Uncompress(void* UncompressedBuffer, int64 UncompressedSize, const void* CompressedBuffer, int64 CompressedSize, uintptr_t CompressionData)
+	{
+		if (!IntFitsIn<int32>(UncompressedSize) ||
+			!IntFitsIn<int32>(CompressedSize) ||
+			!IntFitsIn<int32>(CompressionData))
+		{
+			return false;
+		}
+		PRAGMA_DISABLE_DEPRECATION_WARNINGS;
+		return Uncompress(UncompressedBuffer, (int32&)UncompressedSize, CompressedBuffer, (int32)CompressedSize, (int32)CompressionData);
+		PRAGMA_ENABLE_DEPRECATION_WARNINGS;
+	}
+
+	UE_DEPRECATED(5.5, "Switch to 64 bit version for memory sizes")
+	virtual int32 GetCompressedBufferSize(int32 UncompressedSize, int32 CompressionData) {return -1;};
+
+	virtual bool GetCompressedBufferSize(int64& OutBufferSize, int64 UncompressedSize, uintptr_t CompressionData)
+	{
+		if (!IntFitsIn<int32>(UncompressedSize) ||
+			!IntFitsIn<int32>(CompressionData))
+		{
+			checkf(0, TEXT("%s::GetCompressedBufferSize can't handle 64 bits - format needs to upgrade to the new API."), *GetCompressionFormatName().ToString());
+			return false;
+		}
+		
+		PRAGMA_DISABLE_DEPRECATION_WARNINGS;
+		OutBufferSize = GetCompressedBufferSize((int32)UncompressedSize, (int32)CompressionData);
+		PRAGMA_ENABLE_DEPRECATION_WARNINGS;
+		return true;
+	}
+
 	virtual uint32 GetVersion() = 0;
 	virtual FString GetDDCKeySuffix() = 0;
+
+	// Returns whether the compression format internally decides whether the decreased size is worth the CPU cost
+	// of decompressing it. If the format does this check it should report this by returning a compressed buffer
+	// larger than the decompressed buffer - this is usually automatic as the format would presumably then send the
+	// data uncompressed wrapped in the format's container, which would increase the size.
+	// See FCompression::CompressMemoryIfWorthDecompressing.
 	virtual bool DoesOwnWorthDecompressingCheck() = 0;
 };
diff --git a/Engine/Source/Runtime/Core/Public/Serialization/ArchiveSaveCompressedProxy.h b/Engine/Source/Runtime/Core/Public/Serialization/ArchiveSaveCompressedProxy.h
--- a/Engine/Source/Runtime/Core/Public/Serialization/ArchiveSaveCompressedProxy.h
+++ b/Engine/Source/Runtime/Core/Public/Serialization/ArchiveSaveCompressedProxy.h
@@ -17,29 +17,14 @@
  * FArchive Proxy to transparently write out compressed data to an array.
  */
 class FArchiveSaveCompressedProxy : public FArchive
-{
-	/**
-	 * This enum and the following constructor is a workaround for VC compiler bug that prevents using attributes
-	 * on constructors without inline implementation. This should be removed when the deprecated ctor is removed.
-	 */
-	enum EVS2015Redirector
-	{
-		Redirect
-	};
-	CORE_API FArchiveSaveCompressedProxy(EVS2015Redirector, TArray<uint8>& InCompressedData, ECompressionFlags InCompressionFlags);
-
+{	
 public:
 	/** 
 	 * Constructor, initializing all member variables and allocating temp memory.
 	 *
 	 * @param	InCompressedData [ref]	Array of bytes that is going to hold compressed data
 	 * @param	InCompressionFormat		Compression format to use
-	 */
-	UE_DEPRECATED(4.21, "Use the FName version of FArchiveSaveCompressedProxy constructor")
-	FArchiveSaveCompressedProxy(TArray<uint8>& InCompressedData, ECompressionFlags InCompressionFlags)
-		// Make sure to remove the EVS2015Redirector constructor when this constructor is removed
-		: FArchiveSaveCompressedProxy(EVS2015Redirector::Redirect, InCompressedData, InCompressionFlags)
-	{}
+	 */	
 	CORE_API FArchiveSaveCompressedProxy( TArray<uint8>& InCompressedData, FName InCompressionFormat, ECompressionFlags InCompressionFlags=COMPRESS_None);
 
 	/** Destructor, flushing array if needed. Also frees temporary memory. */
diff --git a/Engine/Source/Runtime/CoreUObject/Public/Serialization/BulkData.h b/Engine/Source/Runtime/CoreUObject/Public/Serialization/BulkData.h
--- a/Engine/Source/Runtime/CoreUObject/Public/Serialization/BulkData.h
+++ b/Engine/Source/Runtime/CoreUObject/Public/Serialization/BulkData.h
@@ -882,7 +882,11 @@ public:
 	/**
 	 * Sets whether we should store the data compressed on disk.
 	 *
-	 * @param CompressionFlags	Flags to use for compressing the data. Use COMPRESS_NONE for no compression, or something like COMPRESS_ZLIB to compress the data
+	 * @param CompressionFormat		The format to use for compression e.g. NAME_None or NAME_Oodle.
+	 * 
+	 * This should likely not be used - bulk data that gets deployed should be compressed via UnrealPak so that
+	 * it can leverage hardware compression as available, and editor bulk data is compressed by default and should use 
+	 * FEditorBulkData::SetCompressionOptions to control that if necessary.
 	 */
 	COREUOBJECT_API void StoreCompressedOnDisk( FName CompressionFormat );
 
diff --git a/Engine/Source/Runtime/PakFile/Private/IPlatformFilePak.cpp b/Engine/Source/Runtime/PakFile/Private/IPlatformFilePak.cpp
--- a/Engine/Source/Runtime/PakFile/Private/IPlatformFilePak.cpp
+++ b/Engine/Source/Runtime/PakFile/Private/IPlatformFilePak.cpp
@@ -5547,7 +5547,14 @@ public:
 			TEXT("See ExtensionsToNotUsePluginCompression in [Pak] section of Engine.ini to add more extensions."),
 			*CompressionMethod.ToString(), TEXT("Unknown"));
 
-		int64 WorkingBufferRequiredSize = FCompression::GetMaximumCompressedSize(CompressionMethod,CompressionBlockSize);
+		// Get how large the biggest compressed block can be and use it to size our read buffers.
+		int64 WorkingBufferRequiredSize;
+		if (!FCompression::GetMaximumCompressedSize(CompressionMethod, WorkingBufferRequiredSize, CompressionBlockSize))
+		{
+			// We should have checked() above, except when checks are disabled. Can't do anything here as there's no failure path out.
+			LowLevelFatalError(TEXT("Failed to get compressed size for compression method: %s. Check plugin is loaded."), *CompressionMethod.ToString());
+		}
+
 		if ( CompressionMethod != NAME_Oodle )
 		{
 			// an amount to extra allocate, in case one block's compressed size is bigger than GetMaximumCompressedSize
diff --git a/Engine/Source/Runtime/PakFile/Public/IPlatformFilePak.h b/Engine/Source/Runtime/PakFile/Public/IPlatformFilePak.h
--- a/Engine/Source/Runtime/PakFile/Public/IPlatformFilePak.h
+++ b/Engine/Source/Runtime/PakFile/Public/IPlatformFilePak.h
@@ -505,15 +505,15 @@ struct FPakEntry
 			{
 				CompressionMethodIndex = 0;
 			}
-			else if (LegacyCompressionMethod & COMPRESS_ZLIB)
+			else if (LegacyCompressionMethod & COMPRESS_ZLIB_DEPRECATED)
 			{
 				CompressionMethodIndex = 1;
 			}
-			else if (LegacyCompressionMethod & COMPRESS_GZIP)
+			else if (LegacyCompressionMethod & COMPRESS_GZIP_DEPRECATED)
 			{
 				CompressionMethodIndex = 2;
 			}
-			else if (LegacyCompressionMethod & COMPRESS_Custom)
+			else if (LegacyCompressionMethod & COMPRESS_Custom_DEPRECATED)
 			{
 				CompressionMethodIndex = 3;
 			}
