Deprecated TIsTriviallyDestructible.

diff --git a/Engine/Plugins/Runtime/ReplicationSystemTestPlugin/Source/Private/Tests/Serialization/TestNetSerializerFixture.h b/Engine/Plugins/Runtime/ReplicationSystemTestPlugin/Source/Private/Tests/Serialization/TestNetSerializerFixture.h
--- a/Engine/Plugins/Runtime/ReplicationSystemTestPlugin/Source/Private/Tests/Serialization/TestNetSerializerFixture.h
+++ b/Engine/Plugins/Runtime/ReplicationSystemTestPlugin/Source/Private/Tests/Serialization/TestNetSerializerFixture.h
@@ -8,7 +8,7 @@
 #include "Iris/Serialization/NetBitStreamWriter.h"
 #include "Iris/Serialization/NetSerializer.h"
 #include "Templates/IsPODType.h"
-#include "Templates/IsTriviallyDestructible.h"
+#include <type_traits>
 
 namespace UE::Net
 {
@@ -86,7 +86,7 @@ public:
 
 	~TTestNetSerializerFixture()
 	{
-		if (!TIsTriviallyDestructible<SourceType>::Value)
+		if constexpr (!std::is_trivially_destructible_v<SourceType>)
 		{
 			reinterpret_cast<SourceType*>(SourceBuffer[0])->~SourceType();
 			reinterpret_cast<SourceType*>(SourceBuffer[1])->~SourceType();
@@ -265,3 +265,7 @@ bool TTestNetSerializerFixture<ConfigPrinter, SourceType>::TestCloneDynamicState
 }
 
 }
+
+#if UE_ENABLE_INCLUDE_ORDER_DEPRECATED_IN_5_5
+#include "Templates/IsTriviallyDestructible.h"
+#endif
diff --git a/Engine/Source/Editor/UnrealEd/Private/UnrealEdPrivatePCH.h b/Engine/Source/Editor/UnrealEd/Private/UnrealEdPrivatePCH.h
--- a/Engine/Source/Editor/UnrealEd/Private/UnrealEdPrivatePCH.h
+++ b/Engine/Source/Editor/UnrealEd/Private/UnrealEdPrivatePCH.h
@@ -40,7 +40,6 @@
 #include "GenericPlatform/GenericPlatformMath.h"
 #include "Templates/MemoryOps.h"
 #include "Templates/IsTriviallyCopyAssignable.h"
-#include "Templates/IsTriviallyDestructible.h"
 #include "Math/NumericLimits.h"
 #include "Serialization/Archive.h"
 #include "Templates/IsEnumClass.h"
@@ -653,3 +652,7 @@
 #include "PackageTools.h"
 #include "Commandlets/GatherTextCommandletBase.h"
 #include "Dialogs/Dialogs.h"
+
+#if UE_ENABLE_INCLUDE_ORDER_DEPRECATED_IN_5_5
+#include "Templates/IsTriviallyDestructible.h"
+#endif
diff --git a/Engine/Source/Runtime/AIModule/Classes/BehaviorTree/BTNode.h b/Engine/Source/Runtime/AIModule/Classes/BehaviorTree/BTNode.h
--- a/Engine/Source/Runtime/AIModule/Classes/BehaviorTree/BTNode.h
+++ b/Engine/Source/Runtime/AIModule/Classes/BehaviorTree/BTNode.h
@@ -10,6 +10,7 @@
 #include "BehaviorTree/BlackboardAssetProvider.h"
 #include "GameplayTaskOwnerInterface.h"
 #include "Tasks/AITask.h"
+#include <type_traits>
 #include "BTNode.generated.h"
 
 class AActor;
@@ -319,7 +320,7 @@ T* UBTNode::InitializeNodeMemory(uint8* NodeMemory, EBTMemoryInit::Type InitType
 template <typename T>
 void UBTNode::CleanupNodeMemory(uint8* NodeMemory, EBTMemoryClear::Type CleanupType) const
 {
-	if constexpr (!TIsTriviallyDestructible<T>::Value)
+	if constexpr (!std::is_trivially_destructible_v<T>)
 	{
 		if (CleanupType == EBTMemoryClear::Destroy)
 		{
diff --git a/Engine/Source/Runtime/AIModule/Private/Navigation/RecastGraphAStar.cpp b/Engine/Source/Runtime/AIModule/Private/Navigation/RecastGraphAStar.cpp
--- a/Engine/Source/Runtime/AIModule/Private/Navigation/RecastGraphAStar.cpp
+++ b/Engine/Source/Runtime/AIModule/Private/Navigation/RecastGraphAStar.cpp
@@ -5,6 +5,7 @@
 #include "NavigationSystem.h"
 #include "NavMesh/PImplRecastNavMesh.h"
 #include "NavMesh/RecastNavMesh.h"
+#include <type_traits>
 
 #include UE_INLINE_GENERATED_CPP_BY_NAME(RecastGraphAStar)
 
@@ -71,7 +72,7 @@ dtStatus FRecastGraphWrapper::ConvertToRecastStatus(const FRecastAStar& Algo, co
 	return FindPathStatus;
 }
 
-static_assert(TIsTriviallyDestructible<FRecastAStarSearchNode>::Value == true, "FRecastAStarSearchNode must be trivially destructible");
+static_assert(std::is_trivially_destructible_v<FRecastAStarSearchNode>, "FRecastAStarSearchNode must be trivially destructible");
 
 dtPolyRef FRecastAStarResult::SetPathInfo(const int32 Index, const FRecastAStarSearchNode& SearchNode)
 {
diff --git a/Engine/Source/Runtime/Core/Public/Core.h b/Engine/Source/Runtime/Core/Public/Core.h
--- a/Engine/Source/Runtime/Core/Public/Core.h
+++ b/Engine/Source/Runtime/Core/Public/Core.h
@@ -66,7 +66,6 @@ MONOLITHIC_HEADER_BOILERPLATE()
 #include "Templates/IsClass.h"
 #include "Templates/TypeCompatibleBytes.h"
 #include "Traits/IsContiguousContainer.h"
-#include "Templates/IsTriviallyDestructible.h"
 #include "Misc/CString.h"
 #include "Templates/IsEnumClass.h"
 #include "GenericPlatform/GenericPlatformProperties.h"
@@ -312,3 +311,7 @@ MONOLITHIC_HEADER_BOILERPLATE()
 #include "Misc/Optional.h"
 #include "Misc/EnumClassFlags.h"
 #include "Misc/EnumRange.h"
+
+#if UE_ENABLE_INCLUDE_ORDER_DEPRECATED_IN_5_5
+#include "Templates/IsTriviallyDestructible.h"
+#endif
diff --git a/Engine/Source/Runtime/Core/Public/CoreMinimal.h b/Engine/Source/Runtime/Core/Public/CoreMinimal.h
--- a/Engine/Source/Runtime/Core/Public/CoreMinimal.h
+++ b/Engine/Source/Runtime/Core/Public/CoreMinimal.h
@@ -48,7 +48,6 @@
 #include "Math/NumericLimits.h"
 #include "HAL/PlatformMath.h"
 #include "Templates/IsTriviallyCopyAssignable.h"
-#include "Templates/IsTriviallyDestructible.h"
 #include "Templates/MemoryOps.h"
 #include "Containers/ContainerAllocationPolicies.h"
 #include "Templates/IsEnumClass.h"
@@ -177,4 +176,8 @@
 #endif
 #include "Math/Transform.h"
 #include "Math/ConvexHull2d.h"
-#include "Math/UnrealMath.h"
\ No newline at end of file
+#include "Math/UnrealMath.h"
+
+#if UE_ENABLE_INCLUDE_ORDER_DEPRECATED_IN_5_5
+#include "Templates/IsTriviallyDestructible.h"
+#endif
diff --git a/Engine/Source/Runtime/Core/Public/CoreSharedPCH.h b/Engine/Source/Runtime/Core/Public/CoreSharedPCH.h
--- a/Engine/Source/Runtime/Core/Public/CoreSharedPCH.h
+++ b/Engine/Source/Runtime/Core/Public/CoreSharedPCH.h
@@ -228,7 +228,6 @@
 #include "Templates/IsSigned.h"
 #include "Templates/IsTriviallyCopyAssignable.h"
 #include "Templates/IsTriviallyCopyConstructible.h"
-#include "Templates/IsTriviallyDestructible.h"
 #include "Templates/Less.h"
 #include "Templates/MemoryOps.h"
 #include "Templates/PointerIsConvertibleFromTo.h"
@@ -254,4 +253,8 @@
 #include "UObject/ScriptDelegates.h"
 #include "UObject/UnrealNames.h"
 #include "UObject/UObjectHierarchyFwd.h"
-#include "UObject/WeakObjectPtrTemplates.h"
\ No newline at end of file
+#include "UObject/WeakObjectPtrTemplates.h"
+
+#if UE_ENABLE_INCLUDE_ORDER_DEPRECATED_IN_5_5
+#include "Templates/IsTriviallyDestructible.h"
+#endif
diff --git a/Engine/Source/Runtime/Core/Public/Serialization/CompactBinary.h b/Engine/Source/Runtime/Core/Public/Serialization/CompactBinary.h
--- a/Engine/Source/Runtime/Core/Public/Serialization/CompactBinary.h
+++ b/Engine/Source/Runtime/Core/Public/Serialization/CompactBinary.h
@@ -15,9 +15,9 @@
 #include "Misc/EnumClassFlags.h"
 #include "String/BytesToHex.h"
 #include "Templates/Function.h"
-#include "Templates/IsTriviallyDestructible.h"
 #include "Templates/RemoveReference.h"
 #include "Templates/UnrealTemplate.h"
+#include <type_traits>
 
 template <typename CharType> class TStringBuilderBase;
 
@@ -853,7 +853,7 @@ protected:
 	/** Assign a field from a pointer to its data and an optional externally-provided type. */
 	inline void Assign(const void* InData, const ECbFieldType InType)
 	{
-		static_assert(TIsTriviallyDestructible<FCbFieldView>::Value,
+		static_assert(std::is_trivially_destructible_v<FCbFieldView>,
 			"This optimization requires FCbFieldView to be trivially destructible!");
 		new(this) FCbFieldView(InData, InType);
 	}
@@ -1460,3 +1460,7 @@ inline FSharedBuffer FCbField::AsBinary(const FSharedBuffer& Default) &&
 }
 
 ///////////////////////////////////////////////////////////////////////////////////////////////////////////////////////
+
+#if UE_ENABLE_INCLUDE_ORDER_DEPRECATED_IN_5_5
+#include "Templates/IsTriviallyDestructible.h"
+#endif
diff --git a/Engine/Source/Runtime/Core/Public/Serialization/MemoryLayout.h b/Engine/Source/Runtime/Core/Public/Serialization/MemoryLayout.h
--- a/Engine/Source/Runtime/Core/Public/Serialization/MemoryLayout.h
+++ b/Engine/Source/Runtime/Core/Public/Serialization/MemoryLayout.h
@@ -11,7 +11,6 @@
 #include "Templates/EnableIf.h"
 #include "Templates/IsAbstract.h"
 #include "Templates/IsPolymorphic.h"
-#include "Templates/IsTriviallyDestructible.h"
 #include "Templates/Models.h"
 #include "Templates/UnrealTemplate.h"
 
@@ -856,3 +855,7 @@ struct FPlatformTypeLayoutParameters
 	 */
 	CORE_API void AppendKeyString(FString& KeyString) const;
 };
+
+#if UE_ENABLE_INCLUDE_ORDER_DEPRECATED_IN_5_5
+#include "Templates/IsTriviallyDestructible.h"
+#endif
diff --git a/Engine/Source/Runtime/Core/Public/Templates/IsTrivial.h b/Engine/Source/Runtime/Core/Public/Templates/IsTrivial.h
--- a/Engine/Source/Runtime/Core/Public/Templates/IsTrivial.h
+++ b/Engine/Source/Runtime/Core/Public/Templates/IsTrivial.h
@@ -3,9 +3,9 @@
 #pragma once
 
 #include "Templates/AndOrNot.h"
-#include "Templates/IsTriviallyDestructible.h"
 #include "Templates/IsTriviallyCopyConstructible.h"
 #include "Templates/IsTriviallyCopyAssignable.h"
+#include <type_traits>
 
 /**
  * Traits class which tests if a type is trivial.
@@ -13,5 +13,9 @@
 template <typename T>
 struct TIsTrivial
 {
-	enum { Value = TAnd<TIsTriviallyDestructible<T>, TIsTriviallyCopyConstructible<T>, TIsTriviallyCopyAssignable<T>>::Value };
+	enum { Value = TAndValue<std::is_trivially_destructible_v<T>, TIsTriviallyCopyConstructible<T>, TIsTriviallyCopyAssignable<T>>::Value };
 };
+
+#if UE_ENABLE_INCLUDE_ORDER_DEPRECATED_IN_5_5
+#include "Templates/IsTriviallyDestructible.h"
+#endif
diff --git a/Engine/Source/Runtime/Core/Public/Templates/IsTriviallyDestructible.h b/Engine/Source/Runtime/Core/Public/Templates/IsTriviallyDestructible.h
--- a/Engine/Source/Runtime/Core/Public/Templates/IsTriviallyDestructible.h
+++ b/Engine/Source/Runtime/Core/Public/Templates/IsTriviallyDestructible.h
@@ -1,13 +1,15 @@
 // Copyright Epic Games, Inc. All Rights Reserved.
 
 #pragma once
+
+#include "CoreTypes.h"
 #include <type_traits>
 
 /**
  * Traits class which tests if a type has a trivial destructor.
  */
 template <typename T>
-struct TIsTriviallyDestructible
+struct UE_DEPRECATED(5.5, "TIsTriviallyDestructible has been deprecated, please use std::is_trivially_destructible_v instead.") TIsTriviallyDestructible
 {
 	enum { Value = std::is_trivially_destructible_v<T> };
 };
diff --git a/Engine/Source/Runtime/Core/Public/Templates/MemoryOps.h b/Engine/Source/Runtime/Core/Public/Templates/MemoryOps.h
--- a/Engine/Source/Runtime/Core/Public/Templates/MemoryOps.h
+++ b/Engine/Source/Runtime/Core/Public/Templates/MemoryOps.h
@@ -4,30 +4,22 @@
 
 #include "CoreTypes.h"
 #include "HAL/UnrealMemory.h"
+#include "Templates/IsTriviallyCopyAssignable.h"
 #include "Templates/IsTriviallyCopyConstructible.h"
 #include "Templates/UnrealTypeTraits.h"
 #include "Traits/UseBitwiseSwap.h"
 #include <new> // IWYU pragma: export
-
-#include "Templates/IsTriviallyCopyAssignable.h"
-#include "Templates/IsTriviallyDestructible.h"
+#include <type_traits>
 
 
 namespace UE::Core::Private::MemoryOps
 {
 	template <typename DestinationElementType, typename SourceElementType>
-	struct TCanBitwiseRelocate
-	{
-		enum
-		{
-			Value =
-				std::is_same_v<DestinationElementType, SourceElementType> ||
-				TAnd<
-					TIsBitwiseConstructible<DestinationElementType, SourceElementType>,
-					TIsTriviallyDestructible<SourceElementType>
-				>::Value
-		};
-	};
+	constexpr inline bool TCanBitwiseRelocate_V =
+		std::is_same_v<DestinationElementType, SourceElementType> || (
+			TIsBitwiseConstructible<DestinationElementType, SourceElementType>::Value &&
+			std::is_trivially_destructible_v<SourceElementType>
+		);
 }
 
 /**
@@ -73,7 +65,7 @@ FORCEINLINE void DestructItem(ElementType* Element)
 	{
 		// Should never get here, but this construct should improve the error messages we get when we try to call this function with incomplete types
 	}
-	else if constexpr (!TIsTriviallyDestructible<ElementType>::Value)
+	else if constexpr (!std::is_trivially_destructible_v<ElementType>)
 	{
 		// We need a typedef here because VC won't compile the destructor call below if ElementType itself has a member called ElementType
 		typedef ElementType DestructItemsElementTypeTypedef;
@@ -97,7 +89,7 @@ FORCEINLINE void DestructItems(ElementType* Element, SizeType Count)
 	{
 		// Should never get here, but this construct should improve the error messages we get when we try to call this function with incomplete types
 	}
-	else if constexpr (!TIsTriviallyDestructible<ElementType>::Value)
+	else if constexpr (!std::is_trivially_destructible_v<ElementType>)
 	{
 		while (Count)
 		{
@@ -188,8 +180,8 @@ FORCEINLINE void RelocateConstructItem(void* Dest, const SourceElementType* Sour
 	{
 		// Should never get here, but this construct should improve the error messages we get when we try to call this function with incomplete types
 	}
-	// Do a bitwise relocate if TCanBitwiseRelocate says we can, but not if the type involved is known not to be bitwise-swappable
-	else if constexpr (UE::Core::Private::MemoryOps::TCanBitwiseRelocate<DestinationElementType, SourceElementType>::Value && (!std::is_same_v<DestinationElementType, SourceElementType> || TUseBitwiseSwap<SourceElementType>::Value))
+	// Do a bitwise relocate if TCanBitwiseRelocate_V says we can, but not if the type involved is known not to be bitwise-swappable
+	else if constexpr (UE::Core::Private::MemoryOps::TCanBitwiseRelocate_V<DestinationElementType, SourceElementType> && (!std::is_same_v<DestinationElementType, SourceElementType> || TUseBitwiseSwap<SourceElementType>::Value))
 	{
 		/* All existing UE containers seem to assume trivial relocatability (i.e. memcpy'able) of their members,
 		 * so we're going to assume that this is safe here.  However, it's not generally possible to assume this
@@ -226,7 +218,7 @@ FORCEINLINE void RelocateConstructItems(void* Dest, const SourceElementType* Sou
 	{
 		// Should never get here, but this construct should improve the error messages we get when we try to call this function with incomplete types
 	}
-	else if constexpr (UE::Core::Private::MemoryOps::TCanBitwiseRelocate<DestinationElementType, SourceElementType>::Value)
+	else if constexpr (UE::Core::Private::MemoryOps::TCanBitwiseRelocate_V<DestinationElementType, SourceElementType>)
 	{
 		/* All existing UE containers seem to assume trivial relocatability (i.e. memcpy'able) of their members,
 		 * so we're going to assume that this is safe here.  However, it's not generally possible to assume this
@@ -341,3 +333,7 @@ FORCEINLINE bool CompareItems(const ElementType* A, const ElementType* B, SizeTy
 		return true;
 	}
 }
+
+#if UE_ENABLE_INCLUDE_ORDER_DEPRECATED_IN_5_5
+#include "Templates/IsTriviallyDestructible.h"
+#endif
diff --git a/Engine/Source/Runtime/CoreUObject/Public/UObject/Class.h b/Engine/Source/Runtime/CoreUObject/Public/UObject/Class.h
--- a/Engine/Source/Runtime/CoreUObject/Public/UObject/Class.h
+++ b/Engine/Source/Runtime/CoreUObject/Public/UObject/Class.h
@@ -49,7 +49,6 @@
 #include "Templates/IsAbstract.h"
 #include "Templates/IsEnum.h"
 #include "Templates/IsPODType.h"
-#include "Templates/IsTriviallyDestructible.h"
 #include "Templates/IsUECoreType.h"
 #include "Templates/Models.h"
 #include "Templates/Tuple.h"
@@ -71,6 +70,7 @@
 #include "UObject/UObjectGlobals.h"
 #include "UObject/UnrealNames.h"
 #include "UObject/ObjectPtr.h"
+#include <type_traits>
 
 #if UE_ENABLE_INCLUDE_ORDER_DEPRECATED_IN_5_2
 #include "Misc/PackageAccessTracking.h"
@@ -1197,7 +1197,7 @@ public:
 		{
 			constexpr FCapabilities Capabilities {
 				(TIsPODType<CPPSTRUCT>::Value ? CPF_IsPlainOldData : CPF_None)
-				| (TIsTriviallyDestructible<CPPSTRUCT>::Value ? CPF_NoDestructor : CPF_None)
+				| (std::is_trivially_destructible_v<CPPSTRUCT> ? CPF_NoDestructor : CPF_None)
 				| (TIsZeroConstructType<CPPSTRUCT>::Value ? CPF_ZeroConstructor : CPF_None)
 				| (TModels_V<CGetTypeHashable, CPPSTRUCT> ? CPF_HasGetValueTypeHash : CPF_None),
 				TTraits::WithSerializerObjectReferences,
@@ -4235,3 +4235,7 @@ UE_DECLARE_CORE_VARIANT_TYPE(Sphere3, Sphere);
 #if UE_ENABLE_INCLUDE_ORDER_DEPRECATED_IN_5_2
 #include "CoreMinimal.h"
 #endif
+
+#if UE_ENABLE_INCLUDE_ORDER_DEPRECATED_IN_5_5
+#include "Templates/IsTriviallyDestructible.h"
+#endif
diff --git a/Engine/Source/Runtime/CoreUObject/Public/UObject/UnrealType.h b/Engine/Source/Runtime/CoreUObject/Public/UObject/UnrealType.h
--- a/Engine/Source/Runtime/CoreUObject/Public/UObject/UnrealType.h
+++ b/Engine/Source/Runtime/CoreUObject/Public/UObject/UnrealType.h
@@ -38,7 +38,6 @@
 #include "Templates/IsIntegral.h"
 #include "Templates/IsPODType.h"
 #include "Templates/IsSigned.h"
-#include "Templates/IsTriviallyDestructible.h"
 #include "Templates/IsUEnumClass.h"
 #include "Templates/MemoryOps.h"
 #include "Templates/Models.h"
@@ -65,6 +64,7 @@
 #include "UObject/UObjectGlobals.h"
 #include "UObject/UnrealNames.h"
 #include "UObject/WeakObjectPtr.h"
+#include <type_traits>
 
 class FBlake3;
 class FOutputDevice;
@@ -1309,7 +1309,7 @@ protected:
 	{
 		return 
 			(TIsPODType<TCppType>::Value ? CPF_IsPlainOldData : CPF_None) 
-			| (TIsTriviallyDestructible<TCppType>::Value ? CPF_NoDestructor : CPF_None) 
+			| (std::is_trivially_destructible_v<TCppType> ? CPF_NoDestructor : CPF_None) 
 			| (TIsZeroConstructType<TCppType>::Value ? CPF_ZeroConstructor : CPF_None)
 			| (TModels_V<CGetTypeHashable, TCppType> ? CPF_HasGetValueTypeHash : CPF_None);
 
@@ -7169,3 +7169,7 @@ class COREUOBJECT_API UMulticastInlineDelegatePropertyWrapper : public UMulticas
 #if UE_ENABLE_INCLUDE_ORDER_DEPRECATED_IN_5_2
 #include "CoreMinimal.h"
 #endif
+
+#if UE_ENABLE_INCLUDE_ORDER_DEPRECATED_IN_5_5
+#include "Templates/IsTriviallyDestructible.h"
+#endif
diff --git a/Engine/Source/Runtime/Engine/Private/EnginePrivatePCH.h b/Engine/Source/Runtime/Engine/Private/EnginePrivatePCH.h
--- a/Engine/Source/Runtime/Engine/Private/EnginePrivatePCH.h
+++ b/Engine/Source/Runtime/Engine/Private/EnginePrivatePCH.h
@@ -40,7 +40,6 @@
 #include "GenericPlatform/GenericPlatformMath.h"
 #include "Templates/MemoryOps.h"
 #include "Templates/IsTriviallyCopyAssignable.h"
-#include "Templates/IsTriviallyDestructible.h"
 #include "Math/NumericLimits.h"
 #include "Serialization/Archive.h"
 #include "Templates/IsEnumClass.h"
@@ -531,3 +530,7 @@
 #include "Particles/ParticleSystem.h"
 #include "Particles/Emitter.h"
 #include "TimerManager.h"
+
+#if UE_ENABLE_INCLUDE_ORDER_DEPRECATED_IN_5_5
+#include "Templates/IsTriviallyDestructible.h"
+#endif
diff --git a/Engine/Source/Runtime/Experimental/Chaos/Public/Chaos/Framework/ScratchBuffer.h b/Engine/Source/Runtime/Experimental/Chaos/Public/Chaos/Framework/ScratchBuffer.h
--- a/Engine/Source/Runtime/Experimental/Chaos/Public/Chaos/Framework/ScratchBuffer.h
+++ b/Engine/Source/Runtime/Experimental/Chaos/Public/Chaos/Framework/ScratchBuffer.h
@@ -4,7 +4,7 @@
 #include "Chaos/Core.h"
 #include "Containers/Array.h"
 #include "Templates/AlignmentTemplates.h"
-#include "Templates/IsTriviallyDestructible.h"
+#include <type_traits>
 
 #define CHAOS_SCRATCHBUFFER_CHECKSENTINEL (DO_CHECK)
 
@@ -50,7 +50,7 @@ namespace Chaos
 
 			template<typename T> T* AllocUninitialized()
 			{
-				static_assert(TIsTriviallyDestructible<T>::Value, "FScratchBuffer only supports trivially destructible types");
+				static_assert(std::is_trivially_destructible_v<T>, "FScratchBuffer only supports trivially destructible types");
 
 				void* Address = AllocAligned(sizeof(T), alignof(T));
 				return (T*)Address;
@@ -58,7 +58,7 @@ namespace Chaos
 
 			template<typename T> T* AllocArrayUninitialized(const int32 Num)
 			{
-				static_assert(TIsTriviallyDestructible<T>::Value, "FScratchBuffer only supports trivially destructible types");
+				static_assert(std::is_trivially_destructible_v<T>, "FScratchBuffer only supports trivially destructible types");
 
 				const size_t AlignedSize = Align(sizeof(T), alignof(T));
 				void* Address = AllocAligned(Num * AlignedSize, alignof(T));
@@ -172,4 +172,8 @@ namespace Chaos
 		};
 
 	}
-}
\ No newline at end of file
+}
+
+#if UE_ENABLE_INCLUDE_ORDER_DEPRECATED_IN_5_5
+#include "Templates/IsTriviallyDestructible.h"
+#endif
diff --git a/Engine/Source/Runtime/Experimental/ChaosCore/Public/Chaos/ObjectPool.h b/Engine/Source/Runtime/Experimental/ChaosCore/Public/Chaos/ObjectPool.h
--- a/Engine/Source/Runtime/Experimental/ChaosCore/Public/Chaos/ObjectPool.h
+++ b/Engine/Source/Runtime/Experimental/ChaosCore/Public/Chaos/ObjectPool.h
@@ -6,7 +6,6 @@
 #include "HAL/UnrealMemory.h"
 #include "Math/UnrealMathUtility.h"
 #include "Templates/AlignmentTemplates.h"
-#include "Templates/IsTriviallyDestructible.h"
 #include "Templates/UnrealTemplate.h"
 
 namespace Chaos
@@ -15,11 +14,11 @@ namespace Chaos
 
 	// The type needs to have a destructor run when the pool pointer is freed
 	template<typename T>
-	using TRequiresDestructor = std::enable_if_t<!TIsTriviallyDestructible<T>::Value>;
+	using TRequiresDestructor = std::enable_if_t<!std::is_trivially_destructible_v<T>>;
 
 	// The type needs no destruction and can just be abandoned on free
 	template<typename T>
-	using TTrivialDestruct = std::enable_if_t<TIsTriviallyDestructible<T>::Value>;
+	using TTrivialDestruct = std::enable_if_t<std::is_trivially_destructible_v<T>>;
 
 	template<typename ObjectType>
 	class TObjectPool
@@ -591,4 +590,8 @@ namespace Chaos
 	private:
 		FObjectPool* Pool;
 	};
-}
\ No newline at end of file
+}
+
+#if UE_ENABLE_INCLUDE_ORDER_DEPRECATED_IN_5_5
+#include "Templates/IsTriviallyDestructible.h"
+#endif
diff --git a/Engine/Source/Runtime/MovieScene/Public/EntitySystem/MovieSceneComponentTypeInfo.h b/Engine/Source/Runtime/MovieScene/Public/EntitySystem/MovieSceneComponentTypeInfo.h
--- a/Engine/Source/Runtime/MovieScene/Public/EntitySystem/MovieSceneComponentTypeInfo.h
+++ b/Engine/Source/Runtime/MovieScene/Public/EntitySystem/MovieSceneComponentTypeInfo.h
@@ -109,7 +109,7 @@ struct FComponentTypeInfo
 	uint8 Alignment;
 	/** 1 byte - Type flags for the component type  */
 	uint8 bIsZeroConstructType : 1;       // Whether TIsZeroConstructType<T> is true
-	uint8 bIsTriviallyDestructable : 1;   // Whether TIsTriviallyDestructible<T> is true
+	uint8 bIsTriviallyDestructable : 1;   // Whether std::is_trivially_destructible_v<T> is true
 	uint8 bIsTriviallyCopyAssignable : 1; // Whether TIsTriviallyCopyAssignable<T> is true
 	uint8 bIsPreserved : 1;               // Whether this component should be preserved when an entity containing it is replaced or overwritten
 	uint8 bIsCopiedToOutput : 1;          // Whether this component should be copied to an output entity if there are now multiple contributors to the same property/state
diff --git a/Engine/Source/Runtime/MovieScene/Public/EntitySystem/MovieSceneEntityFactoryTemplates.h b/Engine/Source/Runtime/MovieScene/Public/EntitySystem/MovieSceneEntityFactoryTemplates.h
--- a/Engine/Source/Runtime/MovieScene/Public/EntitySystem/MovieSceneEntityFactoryTemplates.h
+++ b/Engine/Source/Runtime/MovieScene/Public/EntitySystem/MovieSceneEntityFactoryTemplates.h
@@ -17,6 +17,8 @@
 #include "EntitySystem/MovieSceneSequenceInstanceHandle.h"
 #include "EntitySystem/MovieSceneEntitySystemDirectedGraph.h"
 
+#include <type_traits>
+
 
 namespace UE
 {
@@ -220,7 +222,7 @@ FComponentTypeInfo FComponentRegistry::MakeComponentTypeInfoWithoutComponentOps(
 	NewTypeInfo.Sizeof = ComponentTypeSize;
 	NewTypeInfo.Alignment = Alignment;
 	NewTypeInfo.bIsZeroConstructType = TIsZeroConstructType<T>::Value;
-	NewTypeInfo.bIsTriviallyDestructable = TIsTriviallyDestructible<T>::Value;
+	NewTypeInfo.bIsTriviallyDestructable = std::is_trivially_destructible_v<T>;
 	NewTypeInfo.bIsTriviallyCopyAssignable = TIsTriviallyCopyAssignable<T>::Value;
 	NewTypeInfo.bIsPreserved = EnumHasAnyFlags(Params.Flags, EComponentTypeFlags::Preserved);
 	NewTypeInfo.bIsCopiedToOutput = EnumHasAnyFlags(Params.Flags, EComponentTypeFlags::CopyToOutput);
diff --git a/Engine/Source/Runtime/MovieScene/Public/EntitySystem/MovieScenePropertyRegistry.h b/Engine/Source/Runtime/MovieScene/Public/EntitySystem/MovieScenePropertyRegistry.h
--- a/Engine/Source/Runtime/MovieScene/Public/EntitySystem/MovieScenePropertyRegistry.h
+++ b/Engine/Source/Runtime/MovieScene/Public/EntitySystem/MovieScenePropertyRegistry.h
@@ -18,7 +18,6 @@
 #include "Misc/TVariant.h"
 #include "Stats/Stats.h"
 #include "Stats/Stats2.h"
-#include "Templates/IsTriviallyDestructible.h"
 #include "Templates/SharedPointer.h"
 #include "Templates/UnrealTemplate.h"
 #include "Templates/UnrealTypeTraits.h"
@@ -310,7 +309,6 @@ private:
 } // namespace MovieScene
 } // namespace UE
 
-
-
-
-
+#if UE_ENABLE_INCLUDE_ORDER_DEPRECATED_IN_5_5
+#include "Templates/IsTriviallyDestructible.h"
+#endif
diff --git a/Engine/Source/Runtime/Net/Core/Public/Net/Core/Misc/ResizableCircularQueue.h b/Engine/Source/Runtime/Net/Core/Public/Net/Core/Misc/ResizableCircularQueue.h
--- a/Engine/Source/Runtime/Net/Core/Public/Net/Core/Misc/ResizableCircularQueue.h
+++ b/Engine/Source/Runtime/Net/Core/Public/Net/Core/Misc/ResizableCircularQueue.h
@@ -6,8 +6,8 @@
 #include "Misc/AssertionMacros.h"
 #include "Containers/Array.h"
 #include "Templates/IsPODType.h"
-#include "Templates/IsTriviallyDestructible.h"
 #include "Templates/MemoryOps.h"
+#include <type_traits>
 
 /**
  * Simple ResizableCircularQueue.
@@ -106,7 +106,7 @@ private:
 	enum : uint32
 	{
 		bConstructElements = (TIsPODType<T>::Value ? 0U : 1U),
-		bDestructElements = (TIsTriviallyDestructible<T>::Value ? 0U : 1U),
+		bDestructElements = (std::is_trivially_destructible_v<T> ? 0U : 1U),
 	};
 
 #if WITH_DEV_AUTOMATION_TESTS
@@ -349,3 +349,7 @@ void TResizableCircularQueue<T, AllocatorT>::Empty()
 	Storage.SetNumUnsafeInternal(0);
 	Storage.Empty();
 }
+
+#if UE_ENABLE_INCLUDE_ORDER_DEPRECATED_IN_5_5
+#include "Templates/IsTriviallyDestructible.h"
+#endif
diff --git a/Engine/Source/Runtime/Renderer/Private/RayTracing/RayTracing.cpp b/Engine/Source/Runtime/Renderer/Private/RayTracing/RayTracing.cpp
--- a/Engine/Source/Runtime/Renderer/Private/RayTracing/RayTracing.cpp
+++ b/Engine/Source/Runtime/Renderer/Private/RayTracing/RayTracing.cpp
@@ -15,6 +15,7 @@
 #include "Materials/MaterialRenderProxy.h"
 #include "Experimental/Containers/SherwoodHashTable.h"
 #include "Async/ParallelFor.h"
+#include <type_traits>
 
 static int32 GRayTracingSceneCaptures = -1;
 static FAutoConsoleVariableRef CVarRayTracingSceneCaptures(
@@ -1425,7 +1426,7 @@ namespace RayTracing
 	}
 }
 
-static_assert(TIsTriviallyDestructible<RayTracing::FRelevantPrimitive>::Value == true, "FRelevantPrimitive must be trivially destructible");
+static_assert(std::is_trivially_destructible_v<RayTracing::FRelevantPrimitive>, "FRelevantPrimitive must be trivially destructible");
 template <> struct TIsPODType<RayTracing::FRelevantPrimitive> { enum { Value = true }; }; // Necessary to use TChunkedArray::CopyToLinearArray
 
 #endif //RHI_RAYTRACING
diff --git a/Engine/Source/Runtime/Renderer/Private/RayTracing/RayTracingMaterialHitShaders.h b/Engine/Source/Runtime/Renderer/Private/RayTracing/RayTracingMaterialHitShaders.h
--- a/Engine/Source/Runtime/Renderer/Private/RayTracing/RayTracingMaterialHitShaders.h
+++ b/Engine/Source/Runtime/Renderer/Private/RayTracing/RayTracingMaterialHitShaders.h
@@ -14,6 +14,7 @@
 #include "RayTracingInstanceMask.h"
 #include "RayTracingPayloadType.h"
 #include "ShaderParameterStruct.h"
+#include <type_traits>
 
 enum class ERayTracingMeshCommandsMode : uint8;
 
@@ -226,7 +227,7 @@ public:
 		static constexpr uint32 MaxNum = 1024;
 
 		// Note: constructors for elements of this array are called explicitly in AllocateInternal(). Destructors are not called.
-		static_assert(TIsTriviallyDestructible<FRayTracingLocalShaderBindings>::Value, "FRayTracingLocalShaderBindings must be trivially destructible, as no destructor will be called.");
+		static_assert(std::is_trivially_destructible_v<FRayTracingLocalShaderBindings>, "FRayTracingLocalShaderBindings must be trivially destructible, as no destructor will be called.");
 		FRayTracingLocalShaderBindings Bindings[MaxNum];
 		FChunk* Next;
 		uint32 Num;
@@ -253,7 +254,7 @@ private:
 		{
 			FChunk* OldChunk = CurrentChunk;
 
-			static_assert(TIsTriviallyDestructible<FChunk>::Value, "Chunk must be trivially destructible, as no destructor will be called.");
+			static_assert(std::is_trivially_destructible_v<FChunk>, "Chunk must be trivially destructible, as no destructor will be called.");
 			CurrentChunk = (FChunk*)ParameterMemory.Alloc(sizeof(FChunk), alignof(FChunk));
 			CurrentChunk->Next = nullptr;
 			CurrentChunk->Num = 0;
diff --git a/Engine/Source/Runtime/Renderer/Private/SceneRendering.cpp b/Engine/Source/Runtime/Renderer/Private/SceneRendering.cpp
--- a/Engine/Source/Runtime/Renderer/Private/SceneRendering.cpp
+++ b/Engine/Source/Runtime/Renderer/Private/SceneRendering.cpp
@@ -101,6 +101,7 @@
 #include "LightFunctionAtlas.h"
 #include "EnvironmentComponentsFlags.h"
 #include "VolumetricCloudProxy.h"
+#include <type_traits>
 
 /*-----------------------------------------------------------------------------
 	Globals
@@ -2224,7 +2225,7 @@ FViewInfo* FViewInfo::CreateSnapshot() const
 	}
 	
 	// Ensure the internal state is maintained, needed because we've just Memcpy'd the member data.
-	static_assert(TIsTriviallyDestructible<FGPUScenePrimitiveCollector>::Value != 0, "The destructor is not invoked properly because of FMemory::Memcpy(*Result, *this) above");
+	static_assert(std::is_trivially_destructible_v<FGPUScenePrimitiveCollector> != 0, "The destructor is not invoked properly because of FMemory::Memcpy(*Result, *this) above");
 	Result->DynamicPrimitiveCollector = FGPUScenePrimitiveCollector(DynamicPrimitiveCollector);
 
 	Result->bIsSnapshot = true;
diff --git a/Engine/Source/Runtime/SlateCore/Public/FastUpdate/WidgetProxy.h b/Engine/Source/Runtime/SlateCore/Public/FastUpdate/WidgetProxy.h
--- a/Engine/Source/Runtime/SlateCore/Public/FastUpdate/WidgetProxy.h
+++ b/Engine/Source/Runtime/SlateCore/Public/FastUpdate/WidgetProxy.h
@@ -13,6 +13,7 @@
 #include "Layout/Clipping.h"
 #include "Layout/FlowDirection.h"
 #include "Rendering/DrawElements.h"
+#include <type_traits>
 
 class SWidget;
 class FPaintArgs;
@@ -228,7 +229,7 @@ static_assert(sizeof(FWidgetProxy) <= 32, "FWidgetProxy should be 32 bytes");
 #endif
 
 #if !UE_SLATE_WITH_WIDGETPROXY_WEAKPTR
-static_assert(TIsTriviallyDestructible<FWidgetProxy>::Value == true, "FWidgetProxy must be trivially destructible");
+static_assert(std::is_trivially_destructible_v<FWidgetProxy>, "FWidgetProxy must be trivially destructible");
 template <> struct TIsPODType<FWidgetProxy> { enum { Value = true }; };
 #endif
 
diff --git a/Engine/Source/Runtime/SlateCore/Public/Rendering/RenderingCommon.h b/Engine/Source/Runtime/SlateCore/Public/Rendering/RenderingCommon.h
--- a/Engine/Source/Runtime/SlateCore/Public/Rendering/RenderingCommon.h
+++ b/Engine/Source/Runtime/SlateCore/Public/Rendering/RenderingCommon.h
@@ -13,6 +13,7 @@
 #include "Rendering/DrawElementCoreTypes.h"
 #include "Rendering/SlateRendererTypes.h"
 #include "SlateGlobals.h"
+#include <type_traits>
 #include <utility>
 
 #include "RenderingCommon.generated.h"
@@ -364,7 +365,7 @@ private:
 };
 
 template<> struct TIsPODType<FSlateVertex> { enum { Value = true }; };
-static_assert(TIsTriviallyDestructible<FSlateVertex>::Value == true, "FSlateVertex should be trivially destructible");
+static_assert(std::is_trivially_destructible_v<FSlateVertex>, "FSlateVertex should be trivially destructible");
 static_assert(std::is_trivially_copyable_v<FSlateVertex> == true, "FSlateVertex should be trivially copyable");
 
 /** Stores an aligned rect as shorts. */
@@ -424,7 +425,7 @@ struct FShortRect
 };
 
 template<> struct TIsPODType<FShortRect> { enum { Value = true }; };
-static_assert(TIsTriviallyDestructible<FShortRect>::Value == true, "FShortRect should be trivially destructible");
+static_assert(std::is_trivially_destructible_v<FShortRect>, "FShortRect should be trivially destructible");
 
 namespace UE::Slate
 {
diff --git a/Engine/Source/Runtime/SlateCore/Public/Rendering/SlateRenderBatch.h b/Engine/Source/Runtime/SlateCore/Public/Rendering/SlateRenderBatch.h
--- a/Engine/Source/Runtime/SlateCore/Public/Rendering/SlateRenderBatch.h
+++ b/Engine/Source/Runtime/SlateCore/Public/Rendering/SlateRenderBatch.h
@@ -5,6 +5,7 @@
 #include "RenderingCommon.h"
 #include "Layout/Clipping.h"
 #include "Templates/UnrealTemplate.h"
+#include <type_traits>
 
 class FSlateElementBatch;
 class FSlateDrawLayerHandle;
@@ -237,4 +238,4 @@ public:
 };
 
 static_assert(TIsTriviallyCopyConstructible<FSlateRenderBatch>::Value == true, "FSlateRenderBatch must be mem copyable");
-static_assert(TIsTriviallyDestructible<FSlateRenderBatch>::Value == true, "FSlateRenderBatch must be trivially destructible");
+static_assert(std::is_trivially_destructible_v<FSlateRenderBatch>, "FSlateRenderBatch must be trivially destructible");
