[LockFreeList]
  - Replace platform atomics by std::atomic that allow more fine grained memory ordering
  - Fix AtomicRead that was doing a CAS on windows platforms instead of a less costly read
  - Remove TSAN_SAFE so we can actually test the code for thread safety

#rnx

diff --git a/Engine/Source/Runtime/Core/Private/Containers/LockFreeList.cpp b/Engine/Source/Runtime/Core/Private/Containers/LockFreeList.cpp
--- a/Engine/Source/Runtime/Core/Private/Containers/LockFreeList.cpp
+++ b/Engine/Source/Runtime/Core/Private/Containers/LockFreeList.cpp
@@ -102,7 +102,7 @@ public:
 	* @return Pointer to the allocated memory.
 	* @see Free
 	*/
-	TLinkPtr Pop() TSAN_SAFE
+	TLinkPtr Pop()
 	{
 		FThreadLocalCache& TLS = GetTLS();
 
@@ -123,8 +123,8 @@ public:
 					{
 						TLink* Event = FLockFreeLinkPolicy::IndexToLink(FirstIndex + Index);
 						Event->DoubleNext.Init();
-						Event->SingleNext = 0;
-						Event->Payload = (void*)UPTRINT(TLS.PartialBundle);
+						Event->SingleNext.store(0, std::memory_order_relaxed);
+						Event->Payload.store((void*)UPTRINT(TLS.PartialBundle), std::memory_order_relaxed);
 						TLS.PartialBundle = FLockFreeLinkPolicy::IndexToPtr(FirstIndex + Index);
 					}
 				}
@@ -133,11 +133,11 @@ public:
 		}
 		TLinkPtr Result = TLS.PartialBundle;
 		TLink* ResultP = FLockFreeLinkPolicy::DerefLink(TLS.PartialBundle);
-		TLS.PartialBundle = TLinkPtr(UPTRINT(ResultP->Payload));
+		TLS.PartialBundle = TLinkPtr(UPTRINT(ResultP->Payload.load(std::memory_order_relaxed)));
 		TLS.NumPartial--;
 		//checkLockFreePointerList(TLS.NumPartial >= 0 && ((!!TLS.NumPartial) == (!!TLS.PartialBundle)));
-		ResultP->Payload = nullptr;
-		checkLockFreePointerList(!ResultP->DoubleNext.GetPtr() && !ResultP->SingleNext);
+		ResultP->Payload.store(nullptr, std::memory_order_relaxed);
+		checkLockFreePointerList(!ResultP->DoubleNext.GetPtr() && !ResultP->SingleNext.load(std::memory_order_relaxed));
 		return Result;
 	}
 
@@ -147,7 +147,7 @@ public:
 	* @param Item The item to free.
 	* @see Allocate
 	*/
-	void Push(TLinkPtr Item) TSAN_SAFE
+	void Push(TLinkPtr Item)
 	{
 		FThreadLocalCache& TLS = GetTLS();
 		if (TLS.NumPartial >= NUM_PER_BUNDLE)
@@ -163,8 +163,8 @@ public:
 		}
 		TLink* ItemP = FLockFreeLinkPolicy::DerefLink(Item);
 		ItemP->DoubleNext.SetPtr(0);
-		ItemP->SingleNext = 0;
-		ItemP->Payload = (void*)UPTRINT(TLS.PartialBundle);
+		ItemP->SingleNext.store(0, std::memory_order_relaxed);
+		ItemP->Payload.store((void*)UPTRINT(TLS.PartialBundle), std::memory_order_relaxed);
 		TLS.PartialBundle = Item;
 		TLS.NumPartial++;
 	}
@@ -223,11 +223,11 @@ void FLockFreeLinkPolicy::FreeLockFreeLink(FLockFreeLinkPolicy::TLinkPtr Item)
 	GetLockFreeAllocator().Push(Item);
 }
 
-FLockFreeLinkPolicy::TLinkPtr FLockFreeLinkPolicy::AllocLockFreeLink() TSAN_SAFE
+FLockFreeLinkPolicy::TLinkPtr FLockFreeLinkPolicy::AllocLockFreeLink()
 {
 	FLockFreeLinkPolicy::TLinkPtr Result = GetLockFreeAllocator().Pop();
 	// this can only really be a mem stomp
-	checkLockFreePointerList(Result && !FLockFreeLinkPolicy::DerefLink(Result)->DoubleNext.GetPtr() && !FLockFreeLinkPolicy::DerefLink(Result)->Payload && !FLockFreeLinkPolicy::DerefLink(Result)->SingleNext);
+	checkLockFreePointerList(Result && !FLockFreeLinkPolicy::DerefLink(Result)->DoubleNext.GetPtr() && !FLockFreeLinkPolicy::DerefLink(Result)->Payload.load(std::memory_order_relaxed) && !FLockFreeLinkPolicy::DerefLink(Result)->SingleNext.load(std::memory_order_relaxed));
 	return Result;
 }
 
diff --git a/Engine/Source/Runtime/Core/Public/Containers/LockFreeList.h b/Engine/Source/Runtime/Core/Public/Containers/LockFreeList.h
--- a/Engine/Source/Runtime/Core/Public/Containers/LockFreeList.h
+++ b/Engine/Source/Runtime/Core/Public/Containers/LockFreeList.h
@@ -15,6 +15,8 @@
 #include "Templates/AlignmentTemplates.h"
 #include "Templates/Function.h"
 
+#include <atomic>
+
 CORE_API DECLARE_LOG_CATEGORY_EXTERN(LogLockFreeList, Log, All);
 
 // what level of checking to perform...normally checkLockFreePointerList but could be ensure or check
@@ -139,17 +141,17 @@ struct FIndexedPointer
 	void Init()
 	{
 		static_assert(((MAX_LOCK_FREE_LINKS - 1) & MAX_LOCK_FREE_LINKS) == 0, "MAX_LOCK_FREE_LINKS must be a power of two");
-		Ptrs = 0;
+		Ptrs.store(0, std::memory_order_relaxed);
 	}
 	FORCEINLINE void SetAll(uint32 Ptr, uint64 CounterAndState)
 	{
 		checkLockFreePointerList(Ptr < MAX_LOCK_FREE_LINKS && CounterAndState < (uint64(1) << (64 - MAX_LOCK_FREE_LINKS_AS_BITS)));
-		Ptrs = (uint64(Ptr) | (CounterAndState << MAX_LOCK_FREE_LINKS_AS_BITS));
+		Ptrs.store(uint64(Ptr) | (CounterAndState << MAX_LOCK_FREE_LINKS_AS_BITS), std::memory_order_relaxed);
 	}
 
 	FORCEINLINE uint32 GetPtr() const
 	{
-		return uint32(Ptrs & (MAX_LOCK_FREE_LINKS - 1));
+		return uint32(Ptrs.load(std::memory_order_relaxed) & (MAX_LOCK_FREE_LINKS - 1));
 	}
 
 	FORCEINLINE void SetPtr(uint32 To)
@@ -159,7 +161,7 @@ struct FIndexedPointer
 
 	FORCEINLINE uint64 GetCounterAndState() const
 	{
-		return (Ptrs >> MAX_LOCK_FREE_LINKS_AS_BITS);
+		return (Ptrs.load(std::memory_order_relaxed) >> MAX_LOCK_FREE_LINKS_AS_BITS);
 	}
 
 	FORCEINLINE void SetCounterAndState(uint64 To)
@@ -193,35 +195,36 @@ struct FIndexedPointer
 	FORCEINLINE void AtomicRead(const FIndexedPointer& Other)
 	{
 		checkLockFreePointerList(IsAligned(&Ptrs, 8) && IsAligned(&Other.Ptrs, 8));
-		Ptrs = uint64(FPlatformAtomics::AtomicRead((volatile const int64*)&Other.Ptrs));
+		Ptrs.store(Other.Ptrs.load(std::memory_order_acquire), std::memory_order_relaxed);
 		TestCriticalStall();
 	}
 
 	FORCEINLINE bool InterlockedCompareExchange(const FIndexedPointer& Exchange, const FIndexedPointer& Comparand)
 	{
 		TestCriticalStall();
-		return uint64(FPlatformAtomics::InterlockedCompareExchange((volatile int64*)&Ptrs, Exchange.Ptrs, Comparand.Ptrs)) == Comparand.Ptrs;
+		uint64 Expected = Comparand.Ptrs.load(std::memory_order_relaxed);
+		return Ptrs.compare_exchange_strong(Expected, Exchange.Ptrs.load(std::memory_order_relaxed), std::memory_order_acq_rel, std::memory_order_relaxed);
 	}
 
 	FORCEINLINE bool operator==(const FIndexedPointer& Other) const
 	{
-		return Ptrs == Other.Ptrs;
+		return Ptrs.load(std::memory_order_relaxed) == Other.Ptrs.load(std::memory_order_relaxed);
 	}
 	FORCEINLINE bool operator!=(const FIndexedPointer& Other) const
 	{
-		return Ptrs != Other.Ptrs;
+		return Ptrs.load(std::memory_order_relaxed) != Other.Ptrs.load(std::memory_order_relaxed);
 	}
 
 private:
-	uint64 Ptrs;
+	std::atomic<uint64> Ptrs;
 
 } GCC_ALIGN(8);
 
 struct FIndexedLockFreeLink
 {
-	FIndexedPointer DoubleNext;
-	void *Payload;
-	uint32 SingleNext;
+	FIndexedPointer     DoubleNext;
+	std::atomic<void*>  Payload;
+	std::atomic<uint32> SingleNext;
 };
 
 // there is a version of this code that uses 128 bit atomics to avoid the indirection, that is why we have this policy class at all.
@@ -278,7 +281,7 @@ public:
 		Head.Init();
 	}
 
-	void Push(TLinkPtr Item) TSAN_SAFE
+	void Push(TLinkPtr Item)
 	{
 		while (true)
 		{
@@ -287,7 +290,7 @@ public:
 			TDoublePtr NewHead;
 			NewHead.AdvanceCounterAndState(LocalHead, TABAInc);
 			NewHead.SetPtr(Item);
-			FLockFreeLinkPolicy::DerefLink(Item)->SingleNext = LocalHead.GetPtr();
+			FLockFreeLinkPolicy::DerefLink(Item)->SingleNext.store(LocalHead.GetPtr(), std::memory_order_relaxed);
 			if (Head.InterlockedCompareExchange(NewHead, LocalHead))
 			{
 				break;
@@ -295,7 +298,7 @@ public:
 		}
 	}
 
-	bool PushIf(TFunctionRef<TLinkPtr(uint64)> AllocateIfOkToPush) TSAN_SAFE
+	bool PushIf(TFunctionRef<TLinkPtr(uint64)> AllocateIfOkToPush)
 	{
 		static_assert(TABAInc > 1, "method should not be used for lists without state");
 		while (true)
@@ -311,7 +314,7 @@ public:
 
 			TDoublePtr NewHead;
 			NewHead.AdvanceCounterAndState(LocalHead, TABAInc);
-			FLockFreeLinkPolicy::DerefLink(Item)->SingleNext = LocalHead.GetPtr();
+			FLockFreeLinkPolicy::DerefLink(Item)->SingleNext.store(LocalHead.GetPtr(), std::memory_order_relaxed);
 			NewHead.SetPtr(Item);
 			if (Head.InterlockedCompareExchange(NewHead, LocalHead))
 			{
@@ -322,7 +325,7 @@ public:
 	}
 
 
-	TLinkPtr Pop() TSAN_SAFE
+	TLinkPtr Pop()
 	{
 		TLinkPtr Item = 0;
 		while (true)
@@ -337,17 +340,17 @@ public:
 			TDoublePtr NewHead;
 			NewHead.AdvanceCounterAndState(LocalHead, TABAInc);
 			TLink* ItemP = FLockFreeLinkPolicy::DerefLink(Item);
-			NewHead.SetPtr(ItemP->SingleNext);
+			NewHead.SetPtr(ItemP->SingleNext.load(std::memory_order_relaxed));
 			if (Head.InterlockedCompareExchange(NewHead, LocalHead))
 			{
-				ItemP->SingleNext = 0;
+				ItemP->SingleNext.store(0, std::memory_order_relaxed);
 				break;
 			}
 		}
 		return Item;
 	}
 
-	TLinkPtr PopAll() TSAN_SAFE
+	TLinkPtr PopAll()
 	{
 		TLinkPtr Item = 0;
 		while (true)
@@ -370,7 +373,7 @@ public:
 		return Item;
 	}
 
-	TLinkPtr PopAllAndChangeState(TFunctionRef<uint64(uint64)> StateChange) TSAN_SAFE
+	TLinkPtr PopAllAndChangeState(TFunctionRef<uint64(uint64)> StateChange)
 	{
 		static_assert(TABAInc > 1, "method should not be used for lists without state");
 		TLinkPtr Item = 0;
@@ -430,14 +433,14 @@ public:
 		RootList.Reset();
 	}
 
-	void Push(T* InPayload) TSAN_SAFE
+	void Push(T* InPayload)
 	{
 		TLinkPtr Item = FLockFreeLinkPolicy::AllocLockFreeLink();
-		FLockFreeLinkPolicy::DerefLink(Item)->Payload = InPayload;
+		FLockFreeLinkPolicy::DerefLink(Item)->Payload.store(InPayload, std::memory_order_relaxed);
 		RootList.Push(Item);
 	}
 
-	bool PushIf(T* InPayload, TFunctionRef<bool(uint64)> OkToPush) TSAN_SAFE
+	bool PushIf(T* InPayload, TFunctionRef<bool(uint64)> OkToPush)
 	{
 		TLinkPtr Item = 0;
 
@@ -448,7 +451,7 @@ public:
 				if (!Item)
 				{
 					Item = FLockFreeLinkPolicy::AllocLockFreeLink();
-					FLockFreeLinkPolicy::DerefLink(Item)->Payload = InPayload;
+					FLockFreeLinkPolicy::DerefLink(Item)->Payload.store(InPayload, std::memory_order_relaxed);
 				}
 				return Item;
 			}
@@ -467,40 +470,56 @@ public:
 	}
 
 
-	T* Pop() TSAN_SAFE
+	T* Pop()
 	{
 		TLinkPtr Item = RootList.Pop();
 		T* Result = nullptr;
 		if (Item)
 		{
-			Result = (T*)FLockFreeLinkPolicy::DerefLink(Item)->Payload;
+			Result = (T*)FLockFreeLinkPolicy::DerefLink(Item)->Payload.load(std::memory_order_relaxed);
 			FLockFreeLinkPolicy::FreeLockFreeLink(Item);
 		}
 		return Result;
 	}
 
-	void PopAll(TArray<T*>& OutArray) TSAN_SAFE
+	template <typename ContainerType>
+	void PopAll(ContainerType& OutContainer)
+	{
+		TLinkPtr Links = RootList.PopAll();
+		while (Links)
+		{
+			TLink* LinksP = FLockFreeLinkPolicy::DerefLink(Links);
+			OutContainer.Add((T*)LinksP->Payload.load(std::memory_order_relaxed));
+			TLinkPtr Del = Links;
+			Links = LinksP->SingleNext.load(std::memory_order_relaxed);
+			FLockFreeLinkPolicy::FreeLockFreeLink(Del);
+		}
+	}
+
+	template <typename FunctorType>
+	void PopAllAndApply(FunctorType InFunctor)
 	{
 		TLinkPtr Links = RootList.PopAll();
 		while (Links)
 		{
 			TLink* LinksP = FLockFreeLinkPolicy::DerefLink(Links);
-			OutArray.Add((T*)LinksP->Payload);
+			InFunctor((T*)LinksP->Payload.load(std::memory_order_relaxed));
 			TLinkPtr Del = Links;
-			Links = LinksP->SingleNext;
+			Links = LinksP->SingleNext.load(std::memory_order_relaxed);
 			FLockFreeLinkPolicy::FreeLockFreeLink(Del);
 		}
 	}
 
-	void PopAllAndChangeState(TArray<T*>& OutArray, TFunctionRef<uint64(uint64)> StateChange) TSAN_SAFE
+	template <typename ContainerType>
+	void PopAllAndChangeState(ContainerType& OutContainer, TFunctionRef<uint64(uint64)> StateChange)
 	{
 		TLinkPtr Links = RootList.PopAllAndChangeState(StateChange);
 		while (Links)
 		{
 			TLink* LinksP = FLockFreeLinkPolicy::DerefLink(Links);
-			OutArray.Add((T*)LinksP->Payload);
+			OutContainer.Add((T*)LinksP->Payload.load(std::memory_order_relaxed));
 			TLinkPtr Del = Links;
-			Links = LinksP->SingleNext;
+			Links = LinksP->SingleNext.load(std::memory_order_relaxed);
 			FLockFreeLinkPolicy::FreeLockFreeLink(Del);
 		}
 	}
@@ -550,10 +569,10 @@ public:
 		FLockFreeLinkPolicy::FreeLockFreeLink(Head.GetPtr());
 	}
 
-	void Push(T* InPayload) TSAN_SAFE
+	void Push(T* InPayload)
 	{
 		TLinkPtr Item = FLockFreeLinkPolicy::AllocLockFreeLink();
-		FLockFreeLinkPolicy::DerefLink(Item)->Payload = InPayload;
+		FLockFreeLinkPolicy::DerefLink(Item)->Payload.store(InPayload, std::memory_order_relaxed);
 		TDoublePtr LocalTail;
 		while (true)
 		{
@@ -595,7 +614,7 @@ public:
 		}
 	}
 
-	T* Pop() TSAN_SAFE
+	T* Pop()
 	{
 		T* Result = nullptr;
 		TDoublePtr LocalHead;
@@ -625,7 +644,7 @@ public:
 				else
 				{
 					TestCriticalStall();
-					Result = (T*)FLockFreeLinkPolicy::DerefLink(LocalNext.GetPtr())->Payload;
+					Result = (T*)FLockFreeLinkPolicy::DerefLink(LocalNext.GetPtr())->Payload.load(std::memory_order_relaxed);
 					TDoublePtr NewHead;
 					NewHead.AdvanceCounterAndState(LocalHead, TABAInc);
 					NewHead.SetPtr(LocalNext.GetPtr());
@@ -640,11 +659,12 @@ public:
 		return Result;
 	}
 
-	void PopAll(TArray<T*>& OutArray)
+	template <typename ContainerType>
+	void PopAll(ContainerType& OutContainer)
 	{
 		while (T* Item = Pop())
 		{
-			OutArray.Add(Item);
+			OutContainer.Add(Item);
 		}
 	}
 
@@ -845,11 +865,21 @@ public:
 	*
 	*	@param Output The array to hold the returned items. Must be empty.
 	*/
-	void PopAll(TArray<T *>& Output)
+	template <typename ContainerType>
+	void PopAll(ContainerType& Output)
 	{
 		FLockFreePointerListLIFOBase<T, TPaddingForCacheContention>::PopAll(Output);
 	}
 
+	/**
+	*	Pop all items from the list and call a functor for each of them.
+	*/
+	template <typename FunctorType>
+	void PopAllAndApply(FunctorType InFunctor)
+	{
+		FLockFreePointerListLIFOBase<T, TPaddingForCacheContention>::PopAllAndApply(InFunctor);
+	}
+
 	/**
 	*	Check if the list is empty.
 	*
@@ -904,7 +934,8 @@ public:
 	*
 	*	@param Output The array to hold the returned items. Must be empty.
 	*/
-	void PopAll(TArray<T *>& Output)
+	template <typename ContainerType>
+	void PopAll(ContainerType& Output)
 	{
 		FLockFreePointerFIFOBase<T, TPaddingForCacheContention>::PopAll(Output);
 	}
@@ -952,7 +983,8 @@ public:
 	*
 	*	@param Output The array to hold the returned items. Must be empty.
 	*/
-	void PopAllAndClose(TArray<T *>& Output)
+	template <typename ContainerType>
+	void PopAllAndClose(ContainerType& Output)
 	{
 		auto CheckOpenAndClose = [](uint64 State) -> uint64
 		{
