Remove whitespace from GPU breadcrumbs fields in crash report.

Use AddCrashProperty instead of AddSection to avoid inserting whitespace in properties.

#tests tested locally

diff --git a/Engine/Source/Runtime/Core/Private/GenericPlatform/GenericPlatformCrashContext.cpp b/Engine/Source/Runtime/Core/Private/GenericPlatform/GenericPlatformCrashContext.cpp
--- a/Engine/Source/Runtime/Core/Private/GenericPlatform/GenericPlatformCrashContext.cpp
+++ b/Engine/Source/Runtime/Core/Private/GenericPlatform/GenericPlatformCrashContext.cpp
@@ -1120,17 +1120,17 @@ void FGenericCrashContext::AddGPUBreadcrumbs() const
 	// We use a version indicator for the format used by the breadcrumbs
 	// string, so that parsers can know what to expect and don't break
 	// if changes are made in the format exported by the engine.
-	AddSection(CommonBuffer, TEXT("FormatVersion"), FGPUBreadcrumbCrashData::Version);
-	AddSection(CommonBuffer, TEXT("Source"), NCached::GPUBreadcrumbs->SourceName);
+	AddCrashProperty(TEXT("FormatVersion"), FGPUBreadcrumbCrashData::Version);
+	AddCrashProperty(TEXT("Source"), NCached::GPUBreadcrumbs->SourceName);
 
 	for (auto& [Queue, Breadcrumbs] : NCached::GPUBreadcrumbs->Queues)
 	{
 		BeginSection(CommonBuffer, TEXT("Queue"));
 
-		AddSection(CommonBuffer, TEXT("Name"), Queue);
-		AddSection(CommonBuffer, TEXT("FullHash"), Breadcrumbs.FullHash.ToString());
-		AddSection(CommonBuffer, TEXT("ActiveHash"), Breadcrumbs.ActiveHash.ToString());
-		AddSection(CommonBuffer, TEXT("Breadcrumbs"), Breadcrumbs.BreadcrumbString);
+		AddCrashProperty(TEXT("Name"), Queue);
+		AddCrashProperty(TEXT("FullHash"), Breadcrumbs.FullHash.ToString());
+		AddCrashProperty(TEXT("ActiveHash"), Breadcrumbs.ActiveHash.ToString());
+		AddCrashProperty(TEXT("Breadcrumbs"), Breadcrumbs.BreadcrumbString);
 
 		EndSection(CommonBuffer, TEXT("Queue"));
 	}
