CompressedBuffer: Added an optional flag to decompress blocks to an intermediate buffer

diff --git a/Engine/Source/Runtime/Core/Private/Compression/CompressedBuffer.cpp b/Engine/Source/Runtime/Core/Private/Compression/CompressedBuffer.cpp
--- a/Engine/Source/Runtime/Core/Private/Compression/CompressedBuffer.cpp
+++ b/Engine/Source/Runtime/Core/Private/Compression/CompressedBuffer.cpp
@@ -173,7 +173,7 @@ class FDecoder
 {
 public:
 	virtual uint64 GetHeaderSize(const FHeader& Header) const = 0;
-	virtual bool TryDecompressTo(FDecoderContext& Context, const FDecoderSource& Source, const FHeader& Header, FMemoryView HeaderView, uint64 RawOffset, FMutableMemoryView RawView) const = 0;
+	virtual bool TryDecompressTo(FDecoderContext& Context, const FDecoderSource& Source, const FHeader& Header, FMemoryView HeaderView, uint64 RawOffset, FMutableMemoryView RawView, ECompressedBufferDecompressFlags Flags) const = 0;
 	virtual FCompositeBuffer DecompressToComposite(FDecoderContext& Context, const FDecoderSource& Source, const FHeader& Header, const FMemoryView HeaderView, uint64 RawOffset, uint64 RawSize) const = 0;
 };
 
@@ -211,7 +211,8 @@ public:
 		const FHeader& Header,
 		const FMemoryView HeaderView,
 		const uint64 RawOffset,
-		const FMutableMemoryView RawView) const final
+		const FMutableMemoryView RawView,
+		const ECompressedBufferDecompressFlags Flags) const final
 	{
 		if (Header.Method == EMethod::None &&
 			RawOffset <= Header.TotalRawSize &&
@@ -428,7 +429,7 @@ public:
 		return sizeof(FHeader) + sizeof(uint32) * uint64(Header.BlockCount);
 	}
 
-	bool TryDecompressTo(FDecoderContext& Context, const FDecoderSource& Source, const FHeader& Header, FMemoryView HeaderView, uint64 RawOffset, FMutableMemoryView RawView) const final;
+	bool TryDecompressTo(FDecoderContext& Context, const FDecoderSource& Source, const FHeader& Header, FMemoryView HeaderView, uint64 RawOffset, FMutableMemoryView RawView, ECompressedBufferDecompressFlags Flags) const final;
 	FCompositeBuffer DecompressToComposite(FDecoderContext& Context, const FDecoderSource& Source, const FHeader& Header, const FMemoryView HeaderView, uint64 RawOffset, uint64 RawSize) const final;
 
 protected:
@@ -436,7 +437,7 @@ protected:
 
 #if WITH_EDITORONLY_DATA
 private:
-	bool TryParallelDecompressTo(FDecoderContext& Context, const FDecoderSource& Source, const FHeader& Header, FMemoryView HeaderView, uint64 RawOffset, FMutableMemoryView RawView) const;
+	bool TryParallelDecompressTo(FDecoderContext& Context, const FDecoderSource& Source, const FHeader& Header, FMemoryView HeaderView, uint64 RawOffset, FMutableMemoryView RawView, ECompressedBufferDecompressFlags Flags) const;
 #endif
 };
 
@@ -446,7 +447,8 @@ bool FBlockDecoder::TryDecompressTo(
 	const FHeader& Header,
 	const FMemoryView HeaderView,
 	const uint64 RawOffset,
-	const FMutableMemoryView RawView) const
+	const FMutableMemoryView RawView,
+	const ECompressedBufferDecompressFlags Flags) const
 {
 	if (Header.TotalRawSize < RawOffset + RawView.GetSize())
 	{
@@ -459,7 +461,7 @@ bool FBlockDecoder::TryDecompressTo(
 #if WITH_EDITORONLY_DATA
 	if (RawView.GetSize() >= ParallelDecodeMinRawSize && LastBlockIndex - FirstBlockIndex > 1 && Source.SupportsParallelRead())
 	{
-		return TryParallelDecompressTo(Context, Source, Header, HeaderView, RawOffset, RawView);
+		return TryParallelDecompressTo(Context, Source, Header, HeaderView, RawOffset, RawView, Flags);
 	}
 #endif
 
@@ -490,7 +492,7 @@ bool FBlockDecoder::TryDecompressTo(
 			else
 			{
 				FMutableMemoryView RawBlock;
-				if (RawTarget.GetSize() == RawBlockSize)
+				if (RawTarget.GetSize() == RawBlockSize && !EnumHasAnyFlags(Flags, ECompressedBufferDecompressFlags::IntermediateBuffer))
 				{
 					RawBlock = RawTarget;
 				}
@@ -510,7 +512,7 @@ bool FBlockDecoder::TryDecompressTo(
 					return false;
 				}
 
-				if (RawTarget.GetSize() != RawBlockSize)
+				if (RawTarget.GetSize() != RawBlockSize || EnumHasAnyFlags(Flags, ECompressedBufferDecompressFlags::IntermediateBuffer))
 				{
 					RawTarget.CopyFrom(RawBlock.Mid(RawBlockOffset, RawTarget.GetSize()));
 				}
@@ -534,7 +536,8 @@ FORCENOINLINE bool FBlockDecoder::TryParallelDecompressTo(
 	const FHeader& Header,
 	const FMemoryView HeaderView,
 	const uint64 RawOffset,
-	const FMutableMemoryView RawView) const
+	const FMutableMemoryView RawView,
+	const ECompressedBufferDecompressFlags Flags) const
 {
 	if (Header.TotalRawSize < RawOffset + RawView.GetSize())
 	{
@@ -557,7 +560,7 @@ FORCENOINLINE bool FBlockDecoder::TryParallelDecompressTo(
 		CompressedOffsets.Add(CompressedOffsets.Last() + NETWORK_ORDER32(CompressedBlockSizes[BlockIndex]));
 	}
 
-	const auto TryDecompressBlock = [this, &Source, &Header, RawOffset, RawView, CompressedBlockSizes, CompressedOffsets = CompressedOffsets.GetData(), FirstBlockIndex](FDecoderContext& Context, const uint32 BlockIndex)
+	const auto TryDecompressBlock = [this, &Source, &Header, RawOffset, RawView, Flags, CompressedBlockSizes, CompressedOffsets = CompressedOffsets.GetData(), FirstBlockIndex](FDecoderContext& Context, const uint32 BlockIndex)
 	{
 		const auto BlocksToBytes = [&Header](uint64 BlockCount) -> uint64 { return BlockCount << Header.BlockSizeExponent; };
 		const bool bFirstBlock = (BlockIndex == FirstBlockIndex);
@@ -578,7 +581,7 @@ FORCENOINLINE bool FBlockDecoder::TryParallelDecompressTo(
 			else
 			{
 				FMutableMemoryView RawBlock;
-				if (RawTarget.GetSize() == RawBlockSize)
+				if (RawTarget.GetSize() == RawBlockSize && !EnumHasAnyFlags(Flags, ECompressedBufferDecompressFlags::IntermediateBuffer))
 				{
 					RawBlock = RawTarget;
 				}
@@ -598,7 +601,7 @@ FORCENOINLINE bool FBlockDecoder::TryParallelDecompressTo(
 					return false;
 				}
 
-				if (RawTarget.GetSize() != RawBlockSize)
+				if (RawTarget.GetSize() != RawBlockSize || EnumHasAnyFlags(Flags, ECompressedBufferDecompressFlags::IntermediateBuffer))
 				{
 					RawTarget.CopyFrom(RawBlock.Mid(RawBlockOffset, RawTarget.GetSize()));
 				}
@@ -655,7 +658,7 @@ FCompositeBuffer FBlockDecoder::DecompressToComposite(
 	const uint64 RawSize) const
 {
 	FUniqueBuffer Buffer = FUniqueBuffer::Alloc(RawSize);
-	if (TryDecompressTo(Context, Source, Header, HeaderView, RawOffset, Buffer))
+	if (TryDecompressTo(Context, Source, Header, HeaderView, RawOffset, Buffer, ECompressedBufferDecompressFlags::None))
 	{
 		return FCompositeBuffer(Buffer.MoveToShared());
 	}
@@ -1130,11 +1133,11 @@ bool FCompressedBuffer::TryGetCompressParameters(
 	return false;
 }
 
-bool FCompressedBuffer::TryDecompressTo(const FMutableMemoryView RawView) const
+bool FCompressedBuffer::TryDecompressTo(const FMutableMemoryView RawView, ECompressedBufferDecompressFlags Flags) const
 {
 	if (CompressedData && RawView.GetSize() == GetRawSize())
 	{
-		return FCompressedBufferReader(*this).TryDecompressTo(RawView);
+		return FCompressedBufferReader(*this).TryDecompressTo(RawView, 0, Flags);
 	}
 	return false;
 }
@@ -1274,7 +1277,7 @@ bool FCompressedBufferReader::TryGetCompressParameters(
 	return false;
 }
 
-bool FCompressedBufferReader::TryDecompressTo(const FMutableMemoryView RawView, const uint64 RawOffset)
+bool FCompressedBufferReader::TryDecompressTo(const FMutableMemoryView RawView, const uint64 RawOffset, ECompressedBufferDecompressFlags Flags)
 {
 	using namespace UE::CompressedBuffer::Private;
 	FHeader Header;
@@ -1290,7 +1293,7 @@ bool FCompressedBufferReader::TryDecompressTo(const FMutableMemoryView RawView,
 					SourceArchive
 						? ImplicitConv<const FDecoderSource&>(FArchiveDecoderSource(*SourceArchive, Context.HeaderOffset))
 						: ImplicitConv<const FDecoderSource&>(FBufferDecoderSource(SourceBuffer->GetCompressed())),
-					Header, HeaderView, RawOffset, RawView))
+					Header, HeaderView, RawOffset, RawView, Flags))
 				{
 					return true;
 				}
@@ -1318,7 +1321,7 @@ FSharedBuffer FCompressedBufferReader::Decompress(const uint64 RawOffset, const
 					SourceArchive
 						? ImplicitConv<const FDecoderSource&>(FArchiveDecoderSource(*SourceArchive, Context.HeaderOffset))
 						: ImplicitConv<const FDecoderSource&>(FBufferDecoderSource(SourceBuffer->GetCompressed())),
-					Header, HeaderView, RawOffset, RawData))
+					Header, HeaderView, RawOffset, RawData, ECompressedBufferDecompressFlags::None))
 				{
 					return RawData.MoveToShared();
 				}
diff --git a/Engine/Source/Runtime/Core/Public/Compression/CompressedBuffer.h b/Engine/Source/Runtime/Core/Public/Compression/CompressedBuffer.h
--- a/Engine/Source/Runtime/Core/Public/Compression/CompressedBuffer.h
+++ b/Engine/Source/Runtime/Core/Public/Compression/CompressedBuffer.h
@@ -7,6 +7,7 @@
 #include "Memory/CompositeBuffer.h"
 #include "Memory/MemoryFwd.h"
 #include "Memory/SharedBuffer.h"
+#include "Misc/EnumClassFlags.h"
 #include "Templates/RemoveReference.h"
 #include "Templates/UnrealTemplate.h"
 
@@ -21,6 +22,20 @@ namespace UE::CompressedBuffer::Private { struct FHeader; }
 using ECompressedBufferCompressionLevel = FOodleDataCompression::ECompressionLevel;
 using ECompressedBufferCompressor = FOodleDataCompression::ECompressor;
 
+enum class ECompressedBufferDecompressFlags : uint32
+{
+	None = 0,
+
+	/**
+	 * Decompress each block to an intermediate buffer before copying it to the target address.
+	 *
+	 * Use this flag to maintain performance when decompressing to uncached or write-combined memory.
+	 */
+	IntermediateBuffer = 1 << 0,
+};
+
+ENUM_CLASS_FLAGS(ECompressedBufferDecompressFlags);
+
 /**
  * A compressed buffer stores compressed data in a self-contained format.
  *
@@ -138,7 +153,8 @@ public:
 	 *
 	 * @return True if the requested range was decompressed, otherwise false.
 	 */
-	[[nodiscard]] CORE_API bool TryDecompressTo(FMutableMemoryView RawView) const;
+	[[nodiscard]] CORE_API bool TryDecompressTo(FMutableMemoryView RawView,
+		ECompressedBufferDecompressFlags Flags = ECompressedBufferDecompressFlags::None) const;
 
 	/**
 	 * Decompress into an owned buffer.
@@ -249,7 +265,8 @@ public:
 	 * @param RawOffset   The offset into the raw data from which to decompress.
 	 * @return True if the requested range was decompressed, otherwise false.
 	 */
-	[[nodiscard]] CORE_API bool TryDecompressTo(FMutableMemoryView RawView, uint64 RawOffset = 0);
+	[[nodiscard]] CORE_API bool TryDecompressTo(FMutableMemoryView RawView, uint64 RawOffset = 0,
+		ECompressedBufferDecompressFlags Flags = ECompressedBufferDecompressFlags::None);
 
 	/**
 	 * Decompress into an owned buffer.
