Add some more stdlib functions to AutoRTFM

diff --git a/Engine/Source/Runtime/Core/Private/AutoRTFM/StdLibSupport.cpp b/Engine/Source/Runtime/Core/Private/AutoRTFM/StdLibSupport.cpp
--- a/Engine/Source/Runtime/Core/Private/AutoRTFM/StdLibSupport.cpp
+++ b/Engine/Source/Runtime/Core/Private/AutoRTFM/StdLibSupport.cpp
@@ -320,6 +320,8 @@ UE_AUTORTFM_REGISTER_OPEN_FUNCTION(TlsSetValue);
 
 #if PLATFORM_LINUX
 UE_AUTORTFM_REGISTER_SELF_FUNCTION(clock_gettime);
+UE_AUTORTFM_REGISTER_SELF_FUNCTION(gettimeofday);
+UE_AUTORTFM_REGISTER_SELF_FUNCTION(gmtime_r);
 UE_AUTORTFM_REGISTER_SELF_FUNCTION(bcmp);
 #endif // PLATFORM_LINUX
 
