Adding DISABLED_* options to the LowLevelTestAdapter.
So Adapted tests may be disabled.

-preflight 65d5d72fd40bf4c4d9587078 - engine preflight

diff --git a/Engine/Source/Runtime/Core/Public/Misc/LowLevelTestAdapter.h b/Engine/Source/Runtime/Core/Public/Misc/LowLevelTestAdapter.h
--- a/Engine/Source/Runtime/Core/Public/Misc/LowLevelTestAdapter.h
+++ b/Engine/Source/Runtime/Core/Public/Misc/LowLevelTestAdapter.h
@@ -108,6 +108,19 @@ FString CaptureExpressionsAndValues(const FString& InExpressions, ArgTypes&&...
 #define TEST_CASE(PrettyName, TFlags) TEST_CASE_NAMED_STR(TEST_CASE_GENERATED_NAME_UNIQUE, TEST_CASE_GENERATED_NAME_UNIQUE_STR, PrettyName, TFlags)
 #define TEST_CASE_NAMED(ClassName, PrettyName, TFlags) TEST_CASE_NAMED_STR(ClassName, #ClassName, PrettyName, TFlags)
 
+// Both python and oodle don't trust __LINE__ for unique names, and use __COUNTER__ where possible
+#ifdef __COUNTER__
+#define MAKE_UNIQUE_IDENT(str) LLT_JOIN(str, __COUNTER__)
+#else
+#define MAKE_UNIQUE_IDENT(str) LLT_JOIN(str, __LINE__)
+#endif
+
+// DISABLED_ makes a unique name for either a function or a lambda such that the linker should strip them.
+#define DISABLED_TEST_CASE(...)						static void MAKE_UNIQUE_IDENT(disabled_test_()
+#define DISABLED_TEST_CASE_NAMED(ClassName, ...)	static void MAKE_UNIQUE_IDENT(disabled_test_)()
+#define DISABLED_SCENARIO(...)						static void MAKE_UNIQUE_IDENT(disabled_scenario_)()
+#define DISABLED_SECTION(...)						auto MAKE_UNIQUE_IDENT(disabled_section_) = []()
+
 //-V:CHECK:571,501,547
 #define CHECK(Expr) if (!(Expr)) { FAutomationTestFramework::Get().GetCurrentTest()->AddError(TEXT("Condition failed")); }
 //-V:CHECK_FALSE:571,501,547
