- Tentatively add detector for hybrid hard drives.
- Log error message when physical drives can't be enumerated.

diff --git a/Engine/Source/Runtime/Core/Private/Windows/WindowsPlatformMisc.cpp b/Engine/Source/Runtime/Core/Private/Windows/WindowsPlatformMisc.cpp
--- a/Engine/Source/Runtime/Core/Private/Windows/WindowsPlatformMisc.cpp
+++ b/Engine/Source/Runtime/Core/Private/Windows/WindowsPlatformMisc.cpp
@@ -416,6 +416,14 @@ namespace
 
 	TArray<StorageDevice> StorageDevices;
 
+	static void LogStorageInformationWarning(HRESULT HRes,  const TCHAR* message)
+	{
+		IErrorInfo* Error = nullptr;
+		GetErrorInfo(0, &Error);
+		_com_error error(HRes, Error);
+		UE_LOG(LogWindows, Warning, TEXT("%s [%s]"), message, error.ErrorMessage());
+	}
+
 	static bool CollectStorageInformation()
 	{
 		IWbemLocator* WbemLocator = nullptr;
@@ -428,18 +436,21 @@ namespace
 		HRESULT hres = CoInitializeSecurity(NULL, -1, NULL, NULL, RPC_C_AUTHN_LEVEL_DEFAULT, RPC_C_IMP_LEVEL_IMPERSONATE, NULL, EOAC_NONE, NULL);
 		if (FAILED(hres))
 		{
+			LogStorageInformationWarning(hres, TEXT("Error initializing COM"));
 			FWindowsPlatformMisc::CoUninitialize();
 			return false;
 		}
 		hres = CoCreateInstance(CLSID_WbemLocator, 0, CLSCTX_INPROC_SERVER, IID_IWbemLocator, (LPVOID*)&WbemLocator);
 		if (FAILED(hres))
 		{
+			LogStorageInformationWarning(hres, TEXT("Error creating Wbem instance"));
 			FWindowsPlatformMisc::CoUninitialize();
 			return false;
 		}
 		hres = WbemLocator->ConnectServer(_bstr_t(L"ROOT\\microsoft\\windows\\storage"), NULL, NULL, 0, NULL, 0, 0, &WbemServices);
 		if (FAILED(hres))
 		{
+			LogStorageInformationWarning(hres, TEXT("Error connecting to storage service"));
 			WbemLocator->Release();
 			FWindowsPlatformMisc::CoUninitialize();
 			return false;
@@ -456,6 +467,7 @@ namespace
 		);
 		if (FAILED(hres))
 		{
+			LogStorageInformationWarning(hres, TEXT("Error setting authentication information"));
 			WbemServices->Release();
 			WbemLocator->Release();
 			FWindowsPlatformMisc::CoUninitialize();
@@ -473,6 +485,7 @@ namespace
 
 		if (FAILED(hres))
 		{
+			LogStorageInformationWarning(hres, TEXT("Error listing partitions"));
 			WbemServices->Release();
 			WbemLocator->Release();
 			FWindowsPlatformMisc::CoUninitialize();
@@ -485,7 +498,12 @@ namespace
 			ULONG Returned;
 			IWbemClassObject* StorageWbemObject[10]{};
 			hres = StorageEnumerator->Next(WBEM_INFINITE, 10, StorageWbemObject, &Returned);
-			if (FAILED(hres) || Returned == 0)
+			if (FAILED(hres))
+			{
+				LogStorageInformationWarning(hres, TEXT("Error iterating over partitions"));
+				break;
+			}
+			else if (Returned == 0)
 			{
 				break;
 			}
@@ -517,21 +535,40 @@ namespace
 					{
 						VARIANT Drive;
 						hres = PartitionObject->Get(L"DriveLetter", 0, &Drive, NULL, NULL);
-						if (SUCCEEDED(hres) && Drive.uiVal != 0)
+						if (SUCCEEDED(hres))
 						{
-							VARIANT SerialNumber;
-							hres = DiskObject->Get(L"SerialNumber", 0, &SerialNumber, NULL, NULL);
-							if (SUCCEEDED(hres))
+							if (Drive.uiVal != 0)
 							{
-								StorageDevices.Emplace(SerialNumber.bstrVal, Drive.uiVal);
-								VariantClear(&SerialNumber);
+								VARIANT SerialNumber;
+								hres = DiskObject->Get(L"SerialNumber", 0, &SerialNumber, NULL, NULL);
+								if (SUCCEEDED(hres))
+								{
+									StorageDevices.Emplace(SerialNumber.bstrVal, Drive.uiVal);
+									VariantClear(&SerialNumber);
+								}
+								else
+								{
+									LogStorageInformationWarning(hres, TEXT("Error retrieving serial number"));
+								}
+								VariantClear(&Drive);
 							}
-							VariantClear(&Drive);
+						}
+						else
+						{
+							LogStorageInformationWarning(hres, TEXT("Error retrieving drive letter"));
 						}
 						DiskObject->Release();
 					}
+					else
+					{
+						LogStorageInformationWarning(hres, TEXT("Error retrieving disk information"));
+					}
 					PartitionObject->Release();
 				}
+				else
+				{
+					LogStorageInformationWarning(hres, TEXT("Error retrieving partition information"));
+				}
 
 				VariantClear(&Disk);
 				VariantClear(&Partition);
@@ -549,6 +586,7 @@ namespace
 
 		if (FAILED(hres))
 		{
+			LogStorageInformationWarning(hres, TEXT("Error when querying physical disks"));
 			WbemServices->Release();
 			WbemLocator->Release();
 			FWindowsPlatformMisc::CoUninitialize();
@@ -560,7 +598,12 @@ namespace
 			ULONG Returned;
 			IWbemClassObject* StorageWbemObject = nullptr;
 			hres = StorageEnumerator->Next(WBEM_INFINITE, 1, &StorageWbemObject, &Returned);
-			if (Returned == 0 || FAILED(hres))
+			if (FAILED(hres))
+			{
+				LogStorageInformationWarning(hres, TEXT("Error when iterating over physical disks"));
+				break;
+			}
+			else if (Returned == 0)
 			{
 				break;
 			}
@@ -569,10 +612,12 @@ namespace
 			VARIANT SerialNumber;
 			VARIANT MediaType;
 			VARIANT BusType;
+			VARIANT SpindleSpeed;
 
 			StorageWbemObject->Get(L"SerialNumber", 0, &SerialNumber, NULL, NULL);
 			StorageWbemObject->Get(L"MediaType", 0, &MediaType, NULL, NULL);
 			StorageWbemObject->Get(L"BusType", 0, &BusType, NULL, NULL);
+			StorageWbemObject->Get(L"SpindleSpeed", 0, &SpindleSpeed, NULL, NULL);
 
 			FString Serial(SerialNumber.bstrVal);
 			for (auto& StorageDevice : StorageDevices)
@@ -589,6 +634,10 @@ namespace
 						{
 							StorageDevice.Stats.DriveType = EStorageDeviceType::NVMe;
 						}
+						else if (SpindleSpeed.uintVal != 0)
+						{
+							StorageDevice.Stats.DriveType = EStorageDeviceType::Hybrid;
+						}
 						else
 						{
 							StorageDevice.Stats.DriveType = EStorageDeviceType::SSD;
@@ -629,6 +678,8 @@ const TCHAR* LexToString(EStorageDeviceType StorageType)
 		return TEXT("SSD");
 	case EStorageDeviceType::NVMe:
 		return TEXT("NVMe");
+	case EStorageDeviceType::Hybrid:
+		return TEXT("Hybrid");
 	case EStorageDeviceType::Unknown:
 		[[fallthrough]];
 	default:
diff --git a/Engine/Source/Runtime/Core/Public/Windows/WindowsPlatformMisc.h b/Engine/Source/Runtime/Core/Public/Windows/WindowsPlatformMisc.h
--- a/Engine/Source/Runtime/Core/Public/Windows/WindowsPlatformMisc.h
+++ b/Engine/Source/Runtime/Core/Public/Windows/WindowsPlatformMisc.h
@@ -67,6 +67,8 @@ enum class EStorageDeviceType : uint8
 	SSD = 2,
 	/** Drive is an NVMe . */
 	NVMe = 3,
+	/** Drive is a hybrid SSD/HDD */
+	Hybrid = 4,
 
 	Other = 0xff
 };
