(On behalf of [REDACTED_HANDLE]): Fixing a typo in ConfigContext that would occasionally cause the wrong branch to be used for loading

#rnx

diff --git a/Engine/Source/Runtime/Core/Private/Misc/ConfigContext.cpp b/Engine/Source/Runtime/Core/Private/Misc/ConfigContext.cpp
--- a/Engine/Source/Runtime/Core/Private/Misc/ConfigContext.cpp
+++ b/Engine/Source/Runtime/Core/Private/Misc/ConfigContext.cpp
@@ -215,7 +215,7 @@ bool FConfigContext::Load(const TCHAR* InBaseIniName, FString& OutFinalFilename)
 	// we will need the original InBaseIniName
 	if (bIsForPluginModification)
 	{
-		Branch = ConfigSystem->FindBranch(InBaseIniName, DestIniFilename);
+		Branch = ConfigSystem->FindBranch(InBaseIniName, InBaseIniName);
 		// if not found, make one, so we can add dynamic layers to an empty branch
 		// note: this is unexpected for now since we are going to be looking in KnownConfigFiles
 		if (Branch == nullptr)
