Changed dlopen to globally load symbols (RTLD_GLOBAL) to fix errors with unresolved symbols when loading dylibs

diff --git a/Engine/Source/Runtime/Core/Private/Mac/MacPlatformProcess.cpp b/Engine/Source/Runtime/Core/Private/Mac/MacPlatformProcess.cpp
--- a/Engine/Source/Runtime/Core/Private/Mac/MacPlatformProcess.cpp
+++ b/Engine/Source/Runtime/Core/Private/Mac/MacPlatformProcess.cpp
@@ -50,13 +50,13 @@ static void* GetDllHandleImpl(NSString* DylibPath, NSString* ExecutableFolder)
 		{
 			DylibName = [DylibPath lastPathComponent];
 		}
-		Handle = dlopen([[@"@rpath" stringByAppendingPathComponent:DylibName] fileSystemRepresentation], RTLD_NOLOAD | RTLD_LAZY | RTLD_LOCAL);
+		Handle = dlopen([[@"@rpath" stringByAppendingPathComponent:DylibName] fileSystemRepresentation], RTLD_NOLOAD | RTLD_LAZY | RTLD_GLOBAL);
 	}
 	
 	if (!Handle)
 	{
 		// Not loaded yet, so try to open it
-		Handle = dlopen([DylibPath fileSystemRepresentation], RTLD_LAZY | RTLD_LOCAL);
+		Handle = dlopen([DylibPath fileSystemRepresentation], RTLD_LAZY | RTLD_GLOBAL);
 	}
 	
 	if (!Handle && FParse::Param(FCommandLine::Get(), TEXT("dllerrors")))
