Getting AutoRTFM headers to compile on MSVC.

diff --git a/Engine/Source/Runtime/Core/Private/AutoRTFM/TaggedPtr.h b/Engine/Source/Runtime/Core/Private/AutoRTFM/TaggedPtr.h
--- a/Engine/Source/Runtime/Core/Private/AutoRTFM/TaggedPtr.h
+++ b/Engine/Source/Runtime/Core/Private/AutoRTFM/TaggedPtr.h
@@ -58,7 +58,7 @@ namespace AutoRTFM
 			return *this;
 		}
 
-		UE_AUTORTFM_FORCEINLINE TTaggedPtr& operator=(TTaggedPtr& Other)
+		UE_AUTORTFM_FORCEINLINE TTaggedPtr& operator=(const TTaggedPtr& Other)
 		{
 			Payload = Other.Payload;
 			return *this;
@@ -158,7 +158,7 @@ namespace AutoRTFM
 			return *this;
 		}
 
-		UE_AUTORTFM_FORCEINLINE TTaggedPtr& operator=(TTaggedPtr& Other)
+		UE_AUTORTFM_FORCEINLINE TTaggedPtr& operator=(const TTaggedPtr& Other)
 		{
 			Payload = Other.Payload;
 			return *this;
diff --git a/Engine/Source/Runtime/Core/Private/AutoRTFM/Utils.h b/Engine/Source/Runtime/Core/Private/AutoRTFM/Utils.h
--- a/Engine/Source/Runtime/Core/Private/AutoRTFM/Utils.h
+++ b/Engine/Source/Runtime/Core/Private/AutoRTFM/Utils.h
@@ -45,11 +45,18 @@ template<size_t A, size_t B> struct PrettyStaticAssert final
 
 } // namespace AutoRTFM
 
-#define ASSERT(exp) do { if (UNLIKELY(!(exp))) { UE_DEBUG_BREAK(); PrettyAbort(__FILE__, __LINE__, __PRETTY_FUNCTION__, #exp); } } while (false)
+#if defined(_MSC_VER) && !defined(__clang__)
+	#define ASSERT(exp) do { if (UNLIKELY(!(exp))) { UE_DEBUG_BREAK(); PrettyAbort(__FILE__, __LINE__, __FUNCSIG__, #exp); } } while (false)
+#else
+	#define ASSERT(exp) do { if (UNLIKELY(!(exp))) { UE_DEBUG_BREAK(); PrettyAbort(__FILE__, __LINE__, __PRETTY_FUNCTION__, #exp); } } while (false)
+#endif
 
+#if defined(__has_feature)
+	#if __has_feature(address_sanitizer)
+		#define AUTORTFM_NO_ASAN [[clang::no_sanitize("address")]]
+	#endif
+#endif
 
-#if defined(__has_feature) && __has_feature(address_sanitizer)
-#define AUTORTFM_NO_ASAN [[clang::no_sanitize("address")]]
-#else
-#define AUTORTFM_NO_ASAN
+#if !defined(AUTORTFM_NO_ASAN)
+	#define AUTORTFM_NO_ASAN
 #endif
diff --git a/Engine/Source/Runtime/Core/Private/AutoRTFM/WriteLog.h b/Engine/Source/Runtime/Core/Private/AutoRTFM/WriteLog.h
--- a/Engine/Source/Runtime/Core/Private/AutoRTFM/WriteLog.h
+++ b/Engine/Source/Runtime/Core/Private/AutoRTFM/WriteLog.h
@@ -16,8 +16,8 @@ namespace AutoRTFM
         void* Copy;
 
         UE_AUTORTFM_FORCEINLINE FWriteLogEntry() = default;
-        UE_AUTORTFM_FORCEINLINE FWriteLogEntry(FWriteLogEntry&) = default;
-        UE_AUTORTFM_FORCEINLINE FWriteLogEntry& operator=(FWriteLogEntry&) = default;
+        UE_AUTORTFM_FORCEINLINE FWriteLogEntry(const FWriteLogEntry&) = default;
+        UE_AUTORTFM_FORCEINLINE FWriteLogEntry& operator=(const FWriteLogEntry&) = default;
 
         UE_AUTORTFM_FORCEINLINE explicit FWriteLogEntry(void* Original, size_t Size, void* Copy) :
             OriginalAndSize(Original), Copy(Copy)
