AutomationTest pointer comparisons

diff --git a/Engine/Source/Runtime/Core/Public/Misc/AutomationTest.h b/Engine/Source/Runtime/Core/Public/Misc/AutomationTest.h
--- a/Engine/Source/Runtime/Core/Public/Misc/AutomationTest.h
+++ b/Engine/Source/Runtime/Core/Public/Misc/AutomationTest.h
@@ -2239,13 +2239,13 @@ public:
 	}
 
 	/**
-	 * Logs an error if the two values are the same object in memory.
+	 * Logs an error if the two values reference the same object in memory.
 	 *
 	 * @param Description - Description text for the test.
-	 * @param A - The first value.
-	 * @param B - The second value.
+	 * @param A - The first reference.
+	 * @param B - The second reference.
 	 *
-	 * @see TestSame
+	 * @see TestSame, TestNotSamePtr
 	 */
 	template<typename ValueType>
 	FORCEINLINE bool TestNotSame(const TCHAR* Description, const ValueType& Actual, const ValueType& Expected)
@@ -2263,6 +2263,39 @@ public:
 		return TestNotSame(*Description, Actual, Expected);
 	}
 
+	/**
+	 * Logs an error if the two pointers point to the same object in memory.
+	 *
+	 * @param Description - Description text for the test.
+	 * @param A - The first pointer.
+	 * @param B - The second pointer.
+	 *
+	 * @see TestNotSamePtr, TestSame
+	 */
+	template<typename ValueType>
+	FORCEINLINE bool TestNotSamePtr(const TCHAR* Description, const ValueType* Actual, const ValueType* Expected)
+	{
+		if (nullptr == Actual)
+		{
+			AddWarning(FString::Printf(TEXT("%s: Actual value is nullptr."), Description));
+		}
+		if (nullptr == Expected)
+		{
+			AddWarning(FString::Printf(TEXT("%s: Expected value is nullptr, which may be unintended. If intentional consider instead using TestNotNull()."), Description));
+		}
+		if (Actual == Expected)
+		{
+			AddError(FString::Printf(TEXT("%s: The two pointers are the same."), Description));
+			return false;
+		}
+		return true;
+	}
+
+	template<typename ValueType> bool TestNotSamePtr(const FString& Description, const ValueType* Actual, const ValueType* Expected)
+	{
+		return TestNotSamePtr(*Description, Actual, Expected);
+	}
+
 	/**
 	 * Logs an error if the specified pointer is not NULL.
 	 *
@@ -2279,13 +2312,13 @@ public:
 	}
 
 	/**
-	 * Logs an error if the two values are not the same object in memory.
+	 * Logs an error if the two values do not reference the same object in memory.
 	 *
 	 * @param Description - Description text for the test.
-	 * @param Actual - The actual value.
-	 * @param Expected - The expected value.
+	 * @param Actual - The actual reference.
+	 * @param Expected - The expected reference.
 	 *
-	 * @see TestNotSame
+	 * @see TestNotSame, TestSamePtr
 	 */
 	template<typename ValueType>
 	FORCEINLINE bool TestSame(const TCHAR* Description, const ValueType& Actual, const ValueType& Expected)
@@ -2303,6 +2336,39 @@ public:
 		return TestSame(*Description, Actual, Expected);
 	}
 
+	/**
+	 * Logs an error if the two pointers do not point to the same object in memory.
+	 *
+	 * @param Description - Description text for the test.
+	 * @param Actual - The actual pointer.
+	 * @param Expected - The expected pointer.
+	 *
+	 * @see TestNotSamePtr, TestNotSame
+	 */
+	template<typename ValueType>
+	FORCEINLINE bool TestSamePtr(const TCHAR* Description, const ValueType* Actual, const ValueType* Expected)
+	{
+		if (nullptr == Actual)
+		{
+			AddWarning(FString::Printf(TEXT("%s: Actual value is nullptr."), Description));
+		}
+		if (nullptr == Expected)
+		{
+			AddWarning(FString::Printf(TEXT("%s: Expected value is nullptr, which may be unintended. If intentional consider instead using TestNull()."), Description));
+		}
+		if (Actual != Expected)
+		{
+			AddError(FString::Printf(TEXT("%s: The two pointers are not the same."), Description));
+			return false;
+		}
+		return true;
+	}
+
+	template<typename ValueType> bool TestSamePtr(const FString& Description, const ValueType* Actual, const ValueType* Expected)
+	{
+		return TestSamePtr(*Description, Actual, Expected);
+	}
+
 	/**
 	 * Logs an error if the specified Boolean value is not true.
 	 *
@@ -4259,6 +4325,30 @@ public: \
 		return false;\
 	}
 
+#define UTEST_SAME_PTR(What, Actual, Expected)\
+	if (!TestSamePtr(What, Actual, Expected))\
+	{\
+		return false;\
+	}
+
+#define UTEST_SAME_PTR_EXPR(Actual, Expected)\
+	if (!TestSamePtr(FString::Printf(TEXT("%s == %s"), TEXT(#Actual), TEXT(#Expected)), Actual, Expected))\
+	{\
+		return false;\
+	}
+
+#define UTEST_NOT_SAME_PTR(What, Actual, Expected)\
+	if (!TestNotSamePtr(What, Actual, Expected))\
+	{\
+		return false;\
+	}
+
+#define UTEST_NOT_SAME_PTR_EXPR(Actual, Expected)\
+	if (!TestNotSamePtr(FString::Printf(TEXT("%s != %s"), TEXT(#Actual), TEXT(#Expected)), Actual, Expected))\
+	{\
+		return false;\
+	}
+
 #define UTEST_TRUE(What, Value)\
 	if (!TestTrue(What, Value))\
 	{\
diff --git a/Engine/Source/Runtime/Engine/Private/Tests/EngineAutomationTests.cpp b/Engine/Source/Runtime/Engine/Private/Tests/EngineAutomationTests.cpp
--- a/Engine/Source/Runtime/Engine/Private/Tests/EngineAutomationTests.cpp
+++ b/Engine/Source/Runtime/Engine/Private/Tests/EngineAutomationTests.cpp
@@ -857,15 +857,12 @@ bool FAutomationTestInequalityPointer::RunTest(const FString& Parameters)
 	UWorld* SameWorld = World;
 	UWorld* OtherWorld = UWorld::CreateWorld(EWorldType::Game, false);
 	
-	TestSame("Identity stack primitive", StackPointer, StackPointer);
-	TestSame("Identity world object", World, World);
-	/*
-	* Disabled due to UE-207121
-	* TestSame("Same stack primitive", SameStackPointer, StackPointer);
-	* TestSame("Same world object", SameWorld, World);
-	* TestNotSame("Other stack primitive", OtherStackPointer, StackPointer);
-	* TestNotSame("Other world object", OtherWorld, World);
-	*/
+	TestSamePtr("Identity stack primitive", StackPointer, StackPointer);
+	TestSamePtr("Identity world object", World, World);
+	TestSamePtr("Same stack primitive", SameStackPointer, StackPointer);
+	TestSamePtr("Same world object", SameWorld, World);
+	TestNotSamePtr("Other stack primitive", OtherStackPointer, StackPointer);
+	TestNotSamePtr("Other world object", OtherWorld, World);
 	TestNotNull("Stack primitive not null", StackPointer);
 	TestNotNull("Constructed World object not null", World);
 	TestNull("Nullptr", nullptr);
@@ -876,6 +873,30 @@ bool FAutomationTestInequalityPointer::RunTest(const FString& Parameters)
 	return true;
 }
 
+IMPLEMENT_SIMPLE_AUTOMATION_TEST(FAutomationTestInequalityReference, "TestFramework.Validation.TestInequalityReference", EAutomationTestFlags::ApplicationContextMask | EAutomationTestFlags::EngineFilter);
+bool FAutomationTestInequalityReference::RunTest(const FString& Parameters)
+{
+	int32 StackValue = 42;
+	int32& StackRef = StackValue;
+	int32& SameStackRef = StackValue;
+	int32 OtherStackValue = 42;
+	int32& OtherStackRef = OtherStackValue;
+	int32&& StackRValRef = 42;
+	int32& StackLValFromRvalRef = StackRValRef;
+
+	TestSame("Identity primitive", StackRef, StackRef);
+	TestSame("Identity value", StackValue, StackRef);
+	TestSame("Same primitive", SameStackRef, StackRef);
+	TestSame("Identity rvalue", StackRValRef, StackRValRef);
+	TestSame("Same rvalue and lvalue", StackLValFromRvalRef, StackRValRef);
+	TestNotSame("Other primitive", OtherStackRef, StackRef);
+	TestNotSame("Other value", OtherStackValue, StackRef);
+	TestNotSame("Other rvalue", StackRValRef, StackRef);
+	TestNotSame("Other lvalue from rvalue", StackLValFromRvalRef, StackRef);
+
+	return true;
+}
+
 IMPLEMENT_SIMPLE_AUTOMATION_TEST(FAutomationTestInequalityInt32, "TestFramework.Validation.TestInequalityInt32", EAutomationTestFlags::ApplicationContextMask | EAutomationTestFlags::EngineFilter);
 bool FAutomationTestInequalityInt32::RunTest(const FString& Parameters)
 {
@@ -1204,6 +1225,18 @@ bool FAutomationSameNotSameEXPR::RunTest(const FString& Parameters)
 	return true;
 }
 
+IMPLEMENT_CUSTOM_SIMPLE_AUTOMATION_TEST(FAutomationSameNotSamePtrEXPR, FAutomationUTestMacrosExpr, "TestFramework.Validation.UTestSameNotSamePtr", EAutomationTestFlags::ApplicationContextMask | EAutomationTestFlags::EngineFilter);
+bool FAutomationSameNotSamePtrEXPR::RunTest(const FString& Parameters)
+{
+
+	UTEST_SAME_PTR_EXPR(&ActualFStringValue, &ActualFStringValue);
+	UTEST_SAME_PTR(CustomDescriptionString, &ActualFStringValue, &ActualFStringValue);
+	UTEST_NOT_SAME_PTR_EXPR(&ActualFStringValue, &ExpectedFStringValueLowerCase);
+	UTEST_NOT_SAME_PTR(CustomDescriptionString, &ActualFStringValue, &ExpectedFStringValueLowerCase);
+
+	return true;
+}
+
 IMPLEMENT_CUSTOM_SIMPLE_AUTOMATION_TEST(FAutomationTrueFalseEXPR, FAutomationUTestMacrosExpr, "TestFramework.Validation.UTestTrueFalse", EAutomationTestFlags::ApplicationContextMask | EAutomationTestFlags::EngineFilter);
 bool FAutomationTrueFalseEXPR::RunTest(const FString& Parameters)
 {
