Added a family of string splitting functions as UE::String::Split*

Extended UE::String::Find(First|Last)OfAny[Char] with an optional match index output parameter.

diff --git a/Engine/Source/Runtime/Core/Private/String/Find.cpp b/Engine/Source/Runtime/Core/Private/String/Find.cpp
--- a/Engine/Source/Runtime/Core/Private/String/Find.cpp
+++ b/Engine/Source/Runtime/Core/Private/String/Find.cpp
@@ -100,7 +100,7 @@ static inline int32 FindLast(TStringView<CharType> View, TStringView<CharType> S
 }
 
 template <typename CharType>
-static inline int32 FindFirstOfAny(TStringView<CharType> View, TConstArrayView<TStringView<CharType>> Search, ESearchCase::Type SearchCase)
+static inline int32 FindFirstOfAny(TStringView<CharType> View, TConstArrayView<TStringView<CharType>> Search, ESearchCase::Type SearchCase, int32* OutMatchIndex)
 {
 	check(Algo::NoneOf(Search, &TStringView<CharType>::IsEmpty));
 	switch (Search.Num())
@@ -108,7 +108,14 @@ static inline int32 FindFirstOfAny(TStringView<CharType> View, TConstArrayView<T
 	case 0:
 		return INDEX_NONE;
 	case 1:
-		return String::FindFirst(View, Search[0], SearchCase);
+	{
+		const int32 Index = String::FindFirst(View, Search[0], SearchCase);
+		if (Index != INDEX_NONE && OutMatchIndex)
+		{
+			*OutMatchIndex = 0;
+		}
+		return Index;
+	}
 	default:
 		if (Algo::AllOf(Search, [](const TStringView<CharType>& Pattern) { return Pattern.Len() == 1; }))
 		{
@@ -118,7 +125,7 @@ static inline int32 FindFirstOfAny(TStringView<CharType> View, TConstArrayView<T
 			{
 				SearchChars.Add(Pattern[0]);
 			}
-			return String::FindFirstOfAnyChar(View, SearchChars, SearchCase);
+			return String::FindFirstOfAnyChar(View, SearchChars, SearchCase, OutMatchIndex);
 		}
 		break;
 	}
@@ -129,8 +136,12 @@ static inline int32 FindFirstOfAny(TStringView<CharType> View, TConstArrayView<T
 	{
 		const TStringView<CharType> RemainingView(ViewBegin + ViewIndex, ViewLen - ViewIndex);
 		auto MatchPattern = [&RemainingView, SearchCase](const TStringView<CharType>& Pattern) { return RemainingView.StartsWith(Pattern, SearchCase); };
-		if (Algo::FindByPredicate(Search, MatchPattern))
+		if (const TStringView<CharType>* Match = Algo::FindByPredicate(Search, MatchPattern))
 		{
+			if (OutMatchIndex)
+			{
+				*OutMatchIndex = UE_PTRDIFF_TO_INT32(Match - Search.GetData());
+			}
 			return ViewIndex;
 		}
 	}
@@ -138,7 +149,7 @@ static inline int32 FindFirstOfAny(TStringView<CharType> View, TConstArrayView<T
 }
 
 template <typename CharType>
-static inline int32 FindLastOfAny(TStringView<CharType> View, TConstArrayView<TStringView<CharType>> Search, ESearchCase::Type SearchCase)
+static inline int32 FindLastOfAny(TStringView<CharType> View, TConstArrayView<TStringView<CharType>> Search, ESearchCase::Type SearchCase, int32* OutMatchIndex)
 {
 	check(Algo::NoneOf(Search, &TStringView<CharType>::IsEmpty));
 	switch (Search.Num())
@@ -146,7 +157,14 @@ static inline int32 FindLastOfAny(TStringView<CharType> View, TConstArrayView<TS
 	case 0:
 		return INDEX_NONE;
 	case 1:
-		return String::FindLast(View, Search[0], SearchCase);
+	{
+		const int32 Index = String::FindLast(View, Search[0], SearchCase);
+		if (Index != INDEX_NONE && OutMatchIndex)
+		{
+			*OutMatchIndex = 0;
+		}
+		return Index;
+	}
 	default:
 		if (Algo::AllOf(Search, [](const TStringView<CharType>& Pattern) { return Pattern.Len() == 1; }))
 		{
@@ -156,12 +174,11 @@ static inline int32 FindLastOfAny(TStringView<CharType> View, TConstArrayView<TS
 			{
 				SearchChars.Add(Pattern[0]);
 			}
-			return String::FindLastOfAnyChar(View, SearchChars, SearchCase);
+			return String::FindLastOfAnyChar(View, SearchChars, SearchCase, OutMatchIndex);
 		}
 		break;
 	}
 
-
 	const CharType* const ViewBegin = View.GetData();
 	const int32 ViewLen = View.Len();
 	if (ViewLen == 0)
@@ -172,8 +189,12 @@ static inline int32 FindLastOfAny(TStringView<CharType> View, TConstArrayView<TS
 	{
 		const TStringView<CharType> RemainingView(ViewBegin + ViewIndex, ViewLen - ViewIndex);
 		auto MatchPattern = [&RemainingView, SearchCase](const TStringView<CharType>& Pattern) { return RemainingView.StartsWith(Pattern, SearchCase); };
-		if (Algo::FindByPredicate(Search, MatchPattern))
+		if (const TStringView<CharType>* Match = Algo::FindByPredicate(Search, MatchPattern))
 		{
+			if (OutMatchIndex)
+			{
+				*OutMatchIndex = UE_PTRDIFF_TO_INT32(Match - Search.GetData());
+			}
 			return ViewIndex;
 		}
 	}
@@ -243,14 +264,21 @@ static inline int32 FindLastChar(TStringView<CharType> View, CharType Search, ES
 }
 
 template <typename CharType>
-static inline int32 FindFirstOfAnyChar(TStringView<CharType> View, TConstArrayView<CharType> Search, ESearchCase::Type SearchCase)
+static inline int32 FindFirstOfAnyChar(TStringView<CharType> View, TConstArrayView<CharType> Search, ESearchCase::Type SearchCase, int32* OutMatchIndex)
 {
 	switch (Search.Num())
 	{
 	case 0:
 		return INDEX_NONE;
 	case 1:
-		return String::FindFirstChar(View, Search[0], SearchCase);
+	{
+		const int32 Index = String::FindFirstChar(View, Search[0], SearchCase);
+		if (OutMatchIndex)
+		{
+			*OutMatchIndex = 0;
+		}
+		return Index;
+	}
 	default:
 		break;
 	}
@@ -261,8 +289,12 @@ static inline int32 FindFirstOfAnyChar(TStringView<CharType> View, TConstArrayVi
 	{
 		for (const CharType* ViewIt = ViewBegin; ViewIt < ViewEnd; ++ViewIt)
 		{
-			if (Algo::Find(Search, *ViewIt))
+			if (const CharType* Match = Algo::Find(Search, *ViewIt))
 			{
+				if (OutMatchIndex)
+				{
+					*OutMatchIndex = UE_PTRDIFF_TO_INT32(Match - Search.GetData());
+				}
 				return int32(ViewIt - ViewBegin);
 			}
 		}
@@ -278,8 +310,12 @@ static inline int32 FindFirstOfAnyChar(TStringView<CharType> View, TConstArrayVi
 		}
 		for (const CharType* ViewIt = ViewBegin; ViewIt < ViewEnd; ++ViewIt)
 		{
-			if (Algo::Find(SearchUpper, TChar<CharType>::ToUpper(*ViewIt)))
+			if (const CharType* Match = Algo::Find(SearchUpper, TChar<CharType>::ToUpper(*ViewIt)))
 			{
+				if (OutMatchIndex)
+				{
+					*OutMatchIndex = UE_PTRDIFF_TO_INT32(Match - SearchUpper.GetData());
+				}
 				return int32(ViewIt - ViewBegin);
 			}
 		}
@@ -288,14 +324,21 @@ static inline int32 FindFirstOfAnyChar(TStringView<CharType> View, TConstArrayVi
 }
 
 template <typename CharType>
-static inline int32 FindLastOfAnyChar(TStringView<CharType> View, TConstArrayView<CharType> Search, ESearchCase::Type SearchCase)
+static inline int32 FindLastOfAnyChar(TStringView<CharType> View, TConstArrayView<CharType> Search, ESearchCase::Type SearchCase, int32* OutMatchIndex)
 {
 	switch (Search.Num())
 	{
 	case 0:
 		return INDEX_NONE;
 	case 1:
-		return String::FindLastChar(View, Search[0], SearchCase);
+	{
+		const int32 Index = String::FindLastChar(View, Search[0], SearchCase);
+		if (OutMatchIndex)
+		{
+			*OutMatchIndex = 0;
+		}
+		return Index;
+	}
 	default:
 		break;
 	}
@@ -310,8 +353,12 @@ static inline int32 FindLastOfAnyChar(TStringView<CharType> View, TConstArrayVie
 	{
 		for (const CharType* ViewIt = ViewEnd - 1; ViewIt >= ViewBegin; --ViewIt)
 		{
-			if (Algo::Find(Search, *ViewIt))
+			if (const CharType* Match = Algo::Find(Search, *ViewIt))
 			{
+				if (OutMatchIndex)
+				{
+					*OutMatchIndex = UE_PTRDIFF_TO_INT32(Match - Search.GetData());
+				}
 				return int32(ViewIt - ViewBegin);
 			}
 		}
@@ -327,8 +374,12 @@ static inline int32 FindLastOfAnyChar(TStringView<CharType> View, TConstArrayVie
 		}
 		for (const CharType* ViewIt = ViewEnd - 1; ViewIt >= ViewBegin; --ViewIt)
 		{
-			if (Algo::Find(SearchUpper, TChar<CharType>::ToUpper(*ViewIt)))
+			if (const CharType* Match = Algo::Find(SearchUpper, TChar<CharType>::ToUpper(*ViewIt)))
 			{
+				if (OutMatchIndex)
+				{
+					*OutMatchIndex = UE_PTRDIFF_TO_INT32(Match - SearchUpper.GetData());
+				}
 				return int32(ViewIt - ViewBegin);
 			}
 		}
@@ -361,24 +412,34 @@ int32 FindLast(FWideStringView View, FWideStringView Search, ESearchCase::Type S
 	return Private::FindLast(View, Search, SearchCase);
 }
 
-int32 FindFirstOfAny(FUtf8StringView View, TConstArrayView<FUtf8StringView> Search, ESearchCase::Type SearchCase)
+int32 FindFirstOfAny(FAnsiStringView View, TConstArrayView<FAnsiStringView> Search, ESearchCase::Type SearchCase, int32* OutMatchIndex)
+{
+	return Private::FindFirstOfAny(View, Search, SearchCase, OutMatchIndex);
+}
+
+int32 FindFirstOfAny(FUtf8StringView View, TConstArrayView<FUtf8StringView> Search, ESearchCase::Type SearchCase, int32* OutMatchIndex)
+{
+	return Private::FindFirstOfAny(View, Search, SearchCase, OutMatchIndex);
+}
+
+int32 FindFirstOfAny(FWideStringView View, TConstArrayView<FWideStringView> Search, ESearchCase::Type SearchCase, int32* OutMatchIndex)
 {
-	return Private::FindFirstOfAny(View, Search, SearchCase);
+	return Private::FindFirstOfAny(View, Search, SearchCase, OutMatchIndex);
 }
 
-int32 FindFirstOfAny(FWideStringView View, TConstArrayView<FWideStringView> Search, ESearchCase::Type SearchCase)
+int32 FindLastOfAny(FAnsiStringView View, TConstArrayView<FAnsiStringView> Search, ESearchCase::Type SearchCase, int32* OutMatchIndex)
 {
-	return Private::FindFirstOfAny(View, Search, SearchCase);
+	return Private::FindLastOfAny(View, Search, SearchCase, OutMatchIndex);
 }
 
-int32 FindLastOfAny(FUtf8StringView View, TConstArrayView<FUtf8StringView> Search, ESearchCase::Type SearchCase)
+int32 FindLastOfAny(FUtf8StringView View, TConstArrayView<FUtf8StringView> Search, ESearchCase::Type SearchCase, int32* OutMatchIndex)
 {
-	return Private::FindLastOfAny(View, Search, SearchCase);
+	return Private::FindLastOfAny(View, Search, SearchCase, OutMatchIndex);
 }
 
-int32 FindLastOfAny(FWideStringView View, TConstArrayView<FWideStringView> Search, ESearchCase::Type SearchCase)
+int32 FindLastOfAny(FWideStringView View, TConstArrayView<FWideStringView> Search, ESearchCase::Type SearchCase, int32* OutMatchIndex)
 {
-	return Private::FindLastOfAny(View, Search, SearchCase);
+	return Private::FindLastOfAny(View, Search, SearchCase, OutMatchIndex);
 }
 
 int32 FindFirstChar(FUtf8StringView View, ANSICHAR Search, ESearchCase::Type SearchCase)
@@ -411,34 +472,34 @@ int32 FindLastChar(FWideStringView View, WIDECHAR Search, ESearchCase::Type Sear
 	return Private::FindLastChar(View, Search, SearchCase);
 }
 
-int32 FindFirstOfAnyChar(FUtf8StringView View, TConstArrayView<ANSICHAR> Search, ESearchCase::Type SearchCase)
+int32 FindFirstOfAnyChar(FUtf8StringView View, TConstArrayView<ANSICHAR> Search, ESearchCase::Type SearchCase, int32* OutMatchIndex)
 {
-	return Private::FindFirstOfAnyChar(View, MakeArrayView((const UTF8CHAR*)Search.GetData(), Search.Num()), SearchCase);
+	return Private::FindFirstOfAnyChar(View, MakeArrayView((const UTF8CHAR*)Search.GetData(), Search.Num()), SearchCase, OutMatchIndex);
 }
 
-int32 FindFirstOfAnyChar(FUtf8StringView View, TConstArrayView<UTF8CHAR> Search, ESearchCase::Type SearchCase)
+int32 FindFirstOfAnyChar(FUtf8StringView View, TConstArrayView<UTF8CHAR> Search, ESearchCase::Type SearchCase, int32* OutMatchIndex)
 {
-	return Private::FindFirstOfAnyChar(View, Search, SearchCase);
+	return Private::FindFirstOfAnyChar(View, Search, SearchCase, OutMatchIndex);
 }
 
-int32 FindFirstOfAnyChar(FWideStringView View, TConstArrayView<WIDECHAR> Search, ESearchCase::Type SearchCase)
+int32 FindFirstOfAnyChar(FWideStringView View, TConstArrayView<WIDECHAR> Search, ESearchCase::Type SearchCase, int32* OutMatchIndex)
 {
-	return Private::FindFirstOfAnyChar(View, Search, SearchCase);
+	return Private::FindFirstOfAnyChar(View, Search, SearchCase, OutMatchIndex);
 }
 
-int32 FindLastOfAnyChar(FUtf8StringView View, TConstArrayView<ANSICHAR> Search, ESearchCase::Type SearchCase)
+int32 FindLastOfAnyChar(FUtf8StringView View, TConstArrayView<ANSICHAR> Search, ESearchCase::Type SearchCase, int32* OutMatchIndex)
 {
-	return Private::FindLastOfAnyChar(View, MakeArrayView((const UTF8CHAR*)Search.GetData(), Search.Num()), SearchCase);
+	return Private::FindLastOfAnyChar(View, MakeArrayView((const UTF8CHAR*)Search.GetData(), Search.Num()), SearchCase, OutMatchIndex);
 }
 
-int32 FindLastOfAnyChar(FUtf8StringView View, TConstArrayView<UTF8CHAR> Search, ESearchCase::Type SearchCase)
+int32 FindLastOfAnyChar(FUtf8StringView View, TConstArrayView<UTF8CHAR> Search, ESearchCase::Type SearchCase, int32* OutMatchIndex)
 {
-	return Private::FindLastOfAnyChar(View, Search, SearchCase);
+	return Private::FindLastOfAnyChar(View, Search, SearchCase, OutMatchIndex);
 }
 
-int32 FindLastOfAnyChar(FWideStringView View, TConstArrayView<WIDECHAR> Search, ESearchCase::Type SearchCase)
+int32 FindLastOfAnyChar(FWideStringView View, TConstArrayView<WIDECHAR> Search, ESearchCase::Type SearchCase, int32* OutMatchIndex)
 {
-	return Private::FindLastOfAnyChar(View, Search, SearchCase);
+	return Private::FindLastOfAnyChar(View, Search, SearchCase, OutMatchIndex);
 }
 
 } // UE::String
diff --git a/Engine/Source/Runtime/Core/Private/String/Split.cpp b/Engine/Source/Runtime/Core/Private/String/Split.cpp
new file mode 100644
--- /dev/null
+++ b/Engine/Source/Runtime/Core/Private/String/Split.cpp
@@ -0,0 +1,210 @@
+// Copyright Epic Games, Inc. All Rights Reserved.
+
+#include "String/Split.h"
+
+#include "Containers/StringView.h"
+#include "String/Find.h"
+
+namespace UE::String::Private
+{
+
+template <typename CharType>
+static inline bool SplitOnIndex(TStringView<CharType> View, int32 Index, int32 Size, TStringView<CharType>& OutLeft, TStringView<CharType>& OutRight)
+{
+	if (Index == INDEX_NONE)
+	{
+		return false;
+	}
+	OutLeft = View.Left(Index);
+	OutRight = View.RightChop(Index + Size);
+	return true;
+}
+
+template <typename CharType>
+static inline bool SplitFirst(TStringView<CharType> View, TStringView<CharType> Search, TStringView<CharType>& OutLeft, TStringView<CharType>& OutRight, ESearchCase::Type SearchCase)
+{
+	return SplitOnIndex(View, String::FindFirst(View, Search, SearchCase), Search.Len(), OutLeft, OutRight);
+}
+
+template <typename CharType>
+static inline bool SplitLast(TStringView<CharType> View, TStringView<CharType> Search, TStringView<CharType>& OutLeft, TStringView<CharType>& OutRight, ESearchCase::Type SearchCase)
+{
+	return SplitOnIndex(View, String::FindLast(View, Search, SearchCase), Search.Len(), OutLeft, OutRight);
+}
+
+template <typename CharType>
+static inline bool SplitFirstOfAny(TStringView<CharType> View, TConstArrayView<TStringView<CharType>> Search, TStringView<CharType>& OutLeft, TStringView<CharType>& OutRight, ESearchCase::Type SearchCase, int32* OutMatchIndex)
+{
+	int32 MatchIndex = INDEX_NONE;
+	int32 ViewIndex = String::FindFirstOfAny(View, Search, SearchCase, &MatchIndex);
+	int32 MatchSize = MatchIndex != INDEX_NONE ? Search[MatchIndex].Len() : 0;
+	if (MatchIndex != INDEX_NONE && OutMatchIndex)
+	{
+		*OutMatchIndex = MatchIndex;
+	}
+	return SplitOnIndex(View, ViewIndex, MatchSize, OutLeft, OutRight);
+}
+
+template <typename CharType>
+static inline bool SplitLastOfAny(TStringView<CharType> View, TConstArrayView<TStringView<CharType>> Search, TStringView<CharType>& OutLeft, TStringView<CharType>& OutRight, ESearchCase::Type SearchCase, int32* OutMatchIndex)
+{
+	int32 MatchIndex = INDEX_NONE;
+	int32 ViewIndex = String::FindLastOfAny(View, Search, SearchCase, &MatchIndex);
+	int32 MatchSize = MatchIndex != INDEX_NONE ? Search[MatchIndex].Len() : 0;
+	if (MatchIndex != INDEX_NONE && OutMatchIndex)
+	{
+		*OutMatchIndex = MatchIndex;
+	}
+	return SplitOnIndex(View, ViewIndex, MatchSize, OutLeft, OutRight);
+}
+
+template <typename CharType>
+static inline bool SplitFirstChar(TStringView<CharType> View, CharType Search, TStringView<CharType>& OutLeft, TStringView<CharType>& OutRight, ESearchCase::Type SearchCase)
+{
+	return SplitOnIndex(View, String::FindFirstChar(View, Search, SearchCase), 1, OutLeft, OutRight);
+}
+
+template <typename CharType>
+static inline bool SplitLastChar(TStringView<CharType> View, CharType Search, TStringView<CharType>& OutLeft, TStringView<CharType>& OutRight, ESearchCase::Type SearchCase)
+{
+	return SplitOnIndex(View, String::FindLastChar(View, Search, SearchCase), 1, OutLeft, OutRight);
+}
+
+template <typename CharType>
+static inline bool SplitFirstOfAnyChar(TStringView<CharType> View, TConstArrayView<CharType> Search, TStringView<CharType>& OutLeft, TStringView<CharType>& OutRight, ESearchCase::Type SearchCase, int32* OutMatchIndex)
+{
+	return SplitOnIndex(View, String::FindFirstOfAnyChar(View, Search, SearchCase, OutMatchIndex), 1, OutLeft, OutRight);
+}
+
+template <typename CharType>
+static inline bool SplitLastOfAnyChar(TStringView<CharType> View, TConstArrayView<CharType> Search, TStringView<CharType>& OutLeft, TStringView<CharType>& OutRight, ESearchCase::Type SearchCase, int32* OutMatchIndex)
+{
+	return SplitOnIndex(View, String::FindLastOfAnyChar(View, Search, SearchCase, OutMatchIndex), 1, OutLeft, OutRight);
+}
+
+} // UE::String::Private
+
+namespace UE::String
+{
+
+bool SplitFirst(FAnsiStringView View, FAnsiStringView Search, FAnsiStringView& OutLeft, FAnsiStringView& OutRight, ESearchCase::Type SearchCase)
+{
+	return Private::SplitFirst(View, Search, OutLeft, OutRight, SearchCase);
+}
+
+bool SplitFirst(FUtf8StringView View, FUtf8StringView Search, FUtf8StringView& OutLeft, FUtf8StringView& OutRight, ESearchCase::Type SearchCase)
+{
+	return Private::SplitFirst(View, Search, OutLeft, OutRight, SearchCase);
+}
+
+bool SplitFirst(FWideStringView View, FWideStringView Search, FWideStringView& OutLeft, FWideStringView& OutRight, ESearchCase::Type SearchCase)
+{
+	return Private::SplitFirst(View, Search, OutLeft, OutRight, SearchCase);
+}
+
+bool SplitLast(FAnsiStringView View, FAnsiStringView Search, FAnsiStringView& OutLeft, FAnsiStringView& OutRight, ESearchCase::Type SearchCase)
+{
+	return Private::SplitLast(View, Search, OutLeft, OutRight, SearchCase);
+}
+
+bool SplitLast(FUtf8StringView View, FUtf8StringView Search, FUtf8StringView& OutLeft, FUtf8StringView& OutRight, ESearchCase::Type SearchCase)
+{
+	return Private::SplitLast(View, Search, OutLeft, OutRight, SearchCase);
+}
+
+bool SplitLast(FWideStringView View, FWideStringView Search, FWideStringView& OutLeft, FWideStringView& OutRight, ESearchCase::Type SearchCase)
+{
+	return Private::SplitLast(View, Search, OutLeft, OutRight, SearchCase);
+}
+
+bool SplitFirstOfAny(FAnsiStringView View, TConstArrayView<FAnsiStringView> Search, FAnsiStringView& OutLeft, FAnsiStringView& OutRight, ESearchCase::Type SearchCase, int32* OutMatchIndex)
+{
+	return Private::SplitFirstOfAny(View, Search, OutLeft, OutRight, SearchCase, OutMatchIndex);
+}
+
+bool SplitFirstOfAny(FUtf8StringView View, TConstArrayView<FUtf8StringView> Search, FUtf8StringView& OutLeft, FUtf8StringView& OutRight, ESearchCase::Type SearchCase, int32* OutMatchIndex)
+{
+	return Private::SplitFirstOfAny(View, Search, OutLeft, OutRight, SearchCase, OutMatchIndex);
+}
+
+bool SplitFirstOfAny(FWideStringView View, TConstArrayView<FWideStringView> Search, FWideStringView& OutLeft, FWideStringView& OutRight, ESearchCase::Type SearchCase, int32* OutMatchIndex)
+{
+	return Private::SplitFirstOfAny(View, Search, OutLeft, OutRight, SearchCase, OutMatchIndex);
+}
+
+bool SplitLastOfAny(FAnsiStringView View, TConstArrayView<FAnsiStringView> Search, FAnsiStringView& OutLeft, FAnsiStringView& OutRight, ESearchCase::Type SearchCase, int32* OutMatchIndex)
+{
+	return Private::SplitLastOfAny(View, Search, OutLeft, OutRight, SearchCase, OutMatchIndex);
+}
+
+bool SplitLastOfAny(FUtf8StringView View, TConstArrayView<FUtf8StringView> Search, FUtf8StringView& OutLeft, FUtf8StringView& OutRight, ESearchCase::Type SearchCase, int32* OutMatchIndex)
+{
+	return Private::SplitLastOfAny(View, Search, OutLeft, OutRight, SearchCase, OutMatchIndex);
+}
+
+bool SplitLastOfAny(FWideStringView View, TConstArrayView<FWideStringView> Search, FWideStringView& OutLeft, FWideStringView& OutRight, ESearchCase::Type SearchCase, int32* OutMatchIndex)
+{
+	return Private::SplitLastOfAny(View, Search, OutLeft, OutRight, SearchCase, OutMatchIndex);
+}
+
+bool SplitFirstChar(FAnsiStringView View, ANSICHAR Search, FAnsiStringView& OutLeft, FAnsiStringView& OutRight, ESearchCase::Type SearchCase)
+{
+	return Private::SplitFirstChar(View, Search, OutLeft, OutRight, SearchCase);
+}
+
+bool SplitFirstChar(FUtf8StringView View, UTF8CHAR Search, FUtf8StringView& OutLeft, FUtf8StringView& OutRight, ESearchCase::Type SearchCase)
+{
+	return Private::SplitFirstChar(View, Search, OutLeft, OutRight, SearchCase);
+}
+
+bool SplitFirstChar(FWideStringView View, WIDECHAR Search, FWideStringView& OutLeft, FWideStringView& OutRight, ESearchCase::Type SearchCase)
+{
+	return Private::SplitFirstChar(View, Search, OutLeft, OutRight, SearchCase);
+}
+
+bool SplitLastChar(FAnsiStringView View, ANSICHAR Search, FAnsiStringView& OutLeft, FAnsiStringView& OutRight, ESearchCase::Type SearchCase)
+{
+	return Private::SplitLastChar(View, Search, OutLeft, OutRight, SearchCase);
+}
+
+bool SplitLastChar(FUtf8StringView View, UTF8CHAR Search, FUtf8StringView& OutLeft, FUtf8StringView& OutRight, ESearchCase::Type SearchCase)
+{
+	return Private::SplitLastChar(View, Search, OutLeft, OutRight, SearchCase);
+}
+
+bool SplitLastChar(FWideStringView View, WIDECHAR Search, FWideStringView& OutLeft, FWideStringView& OutRight, ESearchCase::Type SearchCase)
+{
+	return Private::SplitLastChar(View, Search, OutLeft, OutRight, SearchCase);
+}
+
+bool SplitFirstOfAnyChar(FAnsiStringView View, TConstArrayView<ANSICHAR> Search, FAnsiStringView& OutLeft, FAnsiStringView& OutRight, ESearchCase::Type SearchCase, int32* OutMatchIndex)
+{
+	return Private::SplitFirstOfAnyChar(View, Search, OutLeft, OutRight, SearchCase, OutMatchIndex);
+}
+
+bool SplitFirstOfAnyChar(FUtf8StringView View, TConstArrayView<UTF8CHAR> Search, FUtf8StringView& OutLeft, FUtf8StringView& OutRight, ESearchCase::Type SearchCase, int32* OutMatchIndex)
+{
+	return Private::SplitFirstOfAnyChar(View, Search, OutLeft, OutRight, SearchCase, OutMatchIndex);
+}
+
+bool SplitFirstOfAnyChar(FWideStringView View, TConstArrayView<WIDECHAR> Search, FWideStringView& OutLeft, FWideStringView& OutRight, ESearchCase::Type SearchCase, int32* OutMatchIndex)
+{
+	return Private::SplitFirstOfAnyChar(View, Search, OutLeft, OutRight, SearchCase, OutMatchIndex);
+}
+
+bool SplitLastOfAnyChar(FAnsiStringView View, TConstArrayView<ANSICHAR> Search, FAnsiStringView& OutLeft, FAnsiStringView& OutRight, ESearchCase::Type SearchCase, int32* OutMatchIndex)
+{
+	return Private::SplitLastOfAnyChar(View, Search, OutLeft, OutRight, SearchCase, OutMatchIndex);
+}
+
+bool SplitLastOfAnyChar(FUtf8StringView View, TConstArrayView<UTF8CHAR> Search, FUtf8StringView& OutLeft, FUtf8StringView& OutRight, ESearchCase::Type SearchCase, int32* OutMatchIndex)
+{
+	return Private::SplitLastOfAnyChar(View, Search, OutLeft, OutRight, SearchCase, OutMatchIndex);
+}
+
+bool SplitLastOfAnyChar(FWideStringView View, TConstArrayView<WIDECHAR> Search, FWideStringView& OutLeft, FWideStringView& OutRight, ESearchCase::Type SearchCase, int32* OutMatchIndex)
+{
+	return Private::SplitLastOfAnyChar(View, Search, OutLeft, OutRight, SearchCase, OutMatchIndex);
+}
+
+} // UE::String
diff --git a/Engine/Source/Runtime/Core/Public/String/Find.h b/Engine/Source/Runtime/Core/Public/String/Find.h
--- a/Engine/Source/Runtime/Core/Public/String/Find.h
+++ b/Engine/Source/Runtime/Core/Public/String/Find.h
@@ -39,10 +39,12 @@ CORE_API int32 FindLast(FWideStringView View, FWideStringView Search, ESearchCas
  * @param View The string to search within.
  * @param Search The strings to search for.
  * @param SearchCase Whether the comparison should ignore case.
+ * @param OutMatchIndex Receives the index of the search string that matched.
  * @return The position at which any search string was found, or INDEX_NONE if not found.
  */
-CORE_API int32 FindFirstOfAny(FUtf8StringView View, TConstArrayView<FUtf8StringView> Search, ESearchCase::Type SearchCase = ESearchCase::CaseSensitive);
-CORE_API int32 FindFirstOfAny(FWideStringView View, TConstArrayView<FWideStringView> Search, ESearchCase::Type SearchCase = ESearchCase::CaseSensitive);
+CORE_API int32 FindFirstOfAny(FAnsiStringView View, TConstArrayView<FAnsiStringView> Search, ESearchCase::Type SearchCase = ESearchCase::CaseSensitive, int32* OutMatchIndex = nullptr);
+CORE_API int32 FindFirstOfAny(FUtf8StringView View, TConstArrayView<FUtf8StringView> Search, ESearchCase::Type SearchCase = ESearchCase::CaseSensitive, int32* OutMatchIndex = nullptr);
+CORE_API int32 FindFirstOfAny(FWideStringView View, TConstArrayView<FWideStringView> Search, ESearchCase::Type SearchCase = ESearchCase::CaseSensitive, int32* OutMatchIndex = nullptr);
 
 /**
  * Search the view for the last occurrence of any search string.
@@ -50,10 +52,12 @@ CORE_API int32 FindFirstOfAny(FWideStringView View, TConstArrayView<FWideStringV
  * @param View The string to search within.
  * @param Search The strings to search for.
  * @param SearchCase Whether the comparison should ignore case.
+ * @param OutMatchIndex Receives the index of the search string that matched.
  * @return The position at which any search string was found, or INDEX_NONE if not found.
  */
-CORE_API int32 FindLastOfAny(FUtf8StringView View, TConstArrayView<FUtf8StringView> Search, ESearchCase::Type SearchCase = ESearchCase::CaseSensitive);
-CORE_API int32 FindLastOfAny(FWideStringView View, TConstArrayView<FWideStringView> Search, ESearchCase::Type SearchCase = ESearchCase::CaseSensitive);
+CORE_API int32 FindLastOfAny(FAnsiStringView View, TConstArrayView<FAnsiStringView> Search, ESearchCase::Type SearchCase = ESearchCase::CaseSensitive, int32* OutMatchIndex = nullptr);
+CORE_API int32 FindLastOfAny(FUtf8StringView View, TConstArrayView<FUtf8StringView> Search, ESearchCase::Type SearchCase = ESearchCase::CaseSensitive, int32* OutMatchIndex = nullptr);
+CORE_API int32 FindLastOfAny(FWideStringView View, TConstArrayView<FWideStringView> Search, ESearchCase::Type SearchCase = ESearchCase::CaseSensitive, int32* OutMatchIndex = nullptr);
 
 /**
  * Search the view for the first occurrence of the search character.
@@ -85,11 +89,12 @@ CORE_API int32 FindLastChar(FWideStringView View, WIDECHAR Search, ESearchCase::
  * @param View The string to search within.
  * @param Search The characters to search for.
  * @param SearchCase Whether the comparison should ignore case.
+ * @param OutMatchIndex Receives the index of the search character that matched.
  * @return The position at which any search character was found, or INDEX_NONE if not found.
  */
-CORE_API int32 FindFirstOfAnyChar(FUtf8StringView View, TConstArrayView<ANSICHAR> Search, ESearchCase::Type SearchCase = ESearchCase::CaseSensitive);
-CORE_API int32 FindFirstOfAnyChar(FUtf8StringView View, TConstArrayView<UTF8CHAR> Search, ESearchCase::Type SearchCase = ESearchCase::CaseSensitive);
-CORE_API int32 FindFirstOfAnyChar(FWideStringView View, TConstArrayView<WIDECHAR> Search, ESearchCase::Type SearchCase = ESearchCase::CaseSensitive);
+CORE_API int32 FindFirstOfAnyChar(FUtf8StringView View, TConstArrayView<ANSICHAR> Search, ESearchCase::Type SearchCase = ESearchCase::CaseSensitive, int32* OutMatchIndex = nullptr);
+CORE_API int32 FindFirstOfAnyChar(FUtf8StringView View, TConstArrayView<UTF8CHAR> Search, ESearchCase::Type SearchCase = ESearchCase::CaseSensitive, int32* OutMatchIndex = nullptr);
+CORE_API int32 FindFirstOfAnyChar(FWideStringView View, TConstArrayView<WIDECHAR> Search, ESearchCase::Type SearchCase = ESearchCase::CaseSensitive, int32* OutMatchIndex = nullptr);
 
 /**
  * Search the view for the last occurrence of any search character.
@@ -97,10 +102,11 @@ CORE_API int32 FindFirstOfAnyChar(FWideStringView View, TConstArrayView<WIDECHAR
  * @param View The string to search within.
  * @param Search The characters to search for.
  * @param SearchCase Whether the comparison should ignore case.
+ * @param OutMatchIndex Receives the index of the search character that matched.
  * @return The position at which any search character was found, or INDEX_NONE if not found.
  */
-CORE_API int32 FindLastOfAnyChar(FUtf8StringView View, TConstArrayView<ANSICHAR> Search, ESearchCase::Type SearchCase = ESearchCase::CaseSensitive);
-CORE_API int32 FindLastOfAnyChar(FUtf8StringView View, TConstArrayView<UTF8CHAR> Search, ESearchCase::Type SearchCase = ESearchCase::CaseSensitive);
-CORE_API int32 FindLastOfAnyChar(FWideStringView View, TConstArrayView<WIDECHAR> Search, ESearchCase::Type SearchCase = ESearchCase::CaseSensitive);
+CORE_API int32 FindLastOfAnyChar(FUtf8StringView View, TConstArrayView<ANSICHAR> Search, ESearchCase::Type SearchCase = ESearchCase::CaseSensitive, int32* OutMatchIndex = nullptr);
+CORE_API int32 FindLastOfAnyChar(FUtf8StringView View, TConstArrayView<UTF8CHAR> Search, ESearchCase::Type SearchCase = ESearchCase::CaseSensitive, int32* OutMatchIndex = nullptr);
+CORE_API int32 FindLastOfAnyChar(FWideStringView View, TConstArrayView<WIDECHAR> Search, ESearchCase::Type SearchCase = ESearchCase::CaseSensitive, int32* OutMatchIndex = nullptr);
 
 } // UE::String
diff --git a/Engine/Source/Runtime/Core/Public/String/Split.h b/Engine/Source/Runtime/Core/Public/String/Split.h
new file mode 100644
--- /dev/null
+++ b/Engine/Source/Runtime/Core/Public/String/Split.h
@@ -0,0 +1,129 @@
+// Copyright Epic Games, Inc. All Rights Reserved.
+
+#pragma once
+
+#include "Containers/ArrayView.h"
+#include "Containers/StringFwd.h"
+#include "CoreTypes.h"
+#include "Misc/CString.h"
+
+namespace UE::String
+{
+
+/**
+ * Split the view on the first occurrence of the search string.
+ *
+ * @param View The string to split.
+ * @param Search The string to search for and split on.
+ * @param OutLeft Receives the part of the view to the left of the search string if it is found.
+ * @param OutRight Receives the part of the view to the right of the search string if it is found.
+ * @param SearchCase Whether the comparison should ignore case.
+ * @return true if the search string was found and the outputs were written, otherwise false.
+ */
+CORE_API bool SplitFirst(FAnsiStringView View, FAnsiStringView Search, FAnsiStringView& OutLeft, FAnsiStringView& OutRight, ESearchCase::Type SearchCase = ESearchCase::CaseSensitive);
+CORE_API bool SplitFirst(FUtf8StringView View, FUtf8StringView Search, FUtf8StringView& OutLeft, FUtf8StringView& OutRight, ESearchCase::Type SearchCase = ESearchCase::CaseSensitive);
+CORE_API bool SplitFirst(FWideStringView View, FWideStringView Search, FWideStringView& OutLeft, FWideStringView& OutRight, ESearchCase::Type SearchCase = ESearchCase::CaseSensitive);
+
+/**
+ * Split the view on the last occurrence of the search string.
+ *
+ * @param View The string to split.
+ * @param Search The string to search for and split on.
+ * @param OutLeft Receives the part of the view to the left of the search string if it is found.
+ * @param OutRight Receives the part of the view to the right of the search string if it is found.
+ * @param SearchCase Whether the comparison should ignore case.
+ * @return true if the search string was found and the outputs were written, otherwise false.
+ */
+CORE_API bool SplitLast(FAnsiStringView View, FAnsiStringView Search, FAnsiStringView& OutLeft, FAnsiStringView& OutRight, ESearchCase::Type SearchCase = ESearchCase::CaseSensitive);
+CORE_API bool SplitLast(FUtf8StringView View, FUtf8StringView Search, FUtf8StringView& OutLeft, FUtf8StringView& OutRight, ESearchCase::Type SearchCase = ESearchCase::CaseSensitive);
+CORE_API bool SplitLast(FWideStringView View, FWideStringView Search, FWideStringView& OutLeft, FWideStringView& OutRight, ESearchCase::Type SearchCase = ESearchCase::CaseSensitive);
+
+/**
+ * Split the view on the first occurrence of any search string.
+ *
+ * @param View The string to split.
+ * @param Search The strings to search for and split on.
+ * @param OutLeft Receives the part of the view to the left of the search string if it is found.
+ * @param OutRight Receives the part of the view to the right of the search string if it is found.
+ * @param SearchCase Whether the comparison should ignore case.
+ * @param OutMatchIndex Receives the index of the search string that matched.
+ * @return true if the search string was found and the outputs were written, otherwise false.
+ */
+CORE_API bool SplitFirstOfAny(FAnsiStringView View, TConstArrayView<FAnsiStringView> Search, FAnsiStringView& OutLeft, FAnsiStringView& OutRight, ESearchCase::Type SearchCase = ESearchCase::CaseSensitive, int32* OutMatchIndex = nullptr);
+CORE_API bool SplitFirstOfAny(FUtf8StringView View, TConstArrayView<FUtf8StringView> Search, FUtf8StringView& OutLeft, FUtf8StringView& OutRight, ESearchCase::Type SearchCase = ESearchCase::CaseSensitive, int32* OutMatchIndex = nullptr);
+CORE_API bool SplitFirstOfAny(FWideStringView View, TConstArrayView<FWideStringView> Search, FWideStringView& OutLeft, FWideStringView& OutRight, ESearchCase::Type SearchCase = ESearchCase::CaseSensitive, int32* OutMatchIndex = nullptr);
+
+/**
+ * Split the view on the last occurrence of any search string.
+ *
+ * @param View The string to split.
+ * @param Search The strings to search for and split on.
+ * @param OutLeft Receives the part of the view to the left of the search string if it is found.
+ * @param OutRight Receives the part of the view to the right of the search string if it is found.
+ * @param SearchCase Whether the comparison should ignore case.
+ * @param OutMatchIndex Receives the index of the search string that matched.
+ * @return true if the search string was found and the outputs were written, otherwise false.
+ */
+CORE_API bool SplitLastOfAny(FAnsiStringView View, TConstArrayView<FAnsiStringView> Search, FAnsiStringView& OutLeft, FAnsiStringView& OutRight, ESearchCase::Type SearchCase = ESearchCase::CaseSensitive, int32* OutMatchIndex = nullptr);
+CORE_API bool SplitLastOfAny(FUtf8StringView View, TConstArrayView<FUtf8StringView> Search, FUtf8StringView& OutLeft, FUtf8StringView& OutRight, ESearchCase::Type SearchCase = ESearchCase::CaseSensitive, int32* OutMatchIndex = nullptr);
+CORE_API bool SplitLastOfAny(FWideStringView View, TConstArrayView<FWideStringView> Search, FWideStringView& OutLeft, FWideStringView& OutRight, ESearchCase::Type SearchCase = ESearchCase::CaseSensitive, int32* OutMatchIndex = nullptr);
+
+/**
+ * Split the view on the first occurrence of the search character.
+ *
+ * @param View The string to split.
+ * @param Search The character to search for and split on.
+ * @param OutLeft Receives the part of the view to the left of the search character if it is found.
+ * @param OutRight Receives the part of the view to the right of the search character if it is found.
+ * @param SearchCase Whether the comparison should ignore case.
+ * @return true if the search character was found and the outputs were written, otherwise false.
+ */
+CORE_API bool SplitFirstChar(FAnsiStringView View, ANSICHAR Search, FAnsiStringView& OutLeft, FAnsiStringView& OutRight, ESearchCase::Type SearchCase = ESearchCase::CaseSensitive);
+CORE_API bool SplitFirstChar(FUtf8StringView View, UTF8CHAR Search, FUtf8StringView& OutLeft, FUtf8StringView& OutRight, ESearchCase::Type SearchCase = ESearchCase::CaseSensitive);
+CORE_API bool SplitFirstChar(FWideStringView View, WIDECHAR Search, FWideStringView& OutLeft, FWideStringView& OutRight, ESearchCase::Type SearchCase = ESearchCase::CaseSensitive);
+
+/**
+ * Split the view on the last occurrence of the search character.
+ *
+ * @param View The string to split.
+ * @param Search The character to search for and split on.
+ * @param OutLeft Receives the part of the view to the left of the search character if it is found.
+ * @param OutRight Receives the part of the view to the right of the search character if it is found.
+ * @param SearchCase Whether the comparison should ignore case.
+ * @return true if the search character was found and the outputs were written, otherwise false.
+ */
+CORE_API bool SplitLastChar(FAnsiStringView View, ANSICHAR Search, FAnsiStringView& OutLeft, FAnsiStringView& OutRight, ESearchCase::Type SearchCase = ESearchCase::CaseSensitive);
+CORE_API bool SplitLastChar(FUtf8StringView View, UTF8CHAR Search, FUtf8StringView& OutLeft, FUtf8StringView& OutRight, ESearchCase::Type SearchCase = ESearchCase::CaseSensitive);
+CORE_API bool SplitLastChar(FWideStringView View, WIDECHAR Search, FWideStringView& OutLeft, FWideStringView& OutRight, ESearchCase::Type SearchCase = ESearchCase::CaseSensitive);
+
+/**
+ * Split the view on the first occurrence of any search character.
+ *
+ * @param View The string to split.
+ * @param Search The characters to search for and split on.
+ * @param OutLeft Receives the part of the view to the left of the search character if it is found.
+ * @param OutRight Receives the part of the view to the right of the search character if it is found.
+ * @param SearchCase Whether the comparison should ignore case.
+ * @param OutMatchIndex Receives the index of the search character that matched.
+ * @return true if the search character was found and the outputs were written, otherwise false.
+ */
+CORE_API bool SplitFirstOfAnyChar(FAnsiStringView View, TConstArrayView<ANSICHAR> Search, FAnsiStringView& OutLeft, FAnsiStringView& OutRight, ESearchCase::Type SearchCase = ESearchCase::CaseSensitive, int32* OutMatchIndex = nullptr);
+CORE_API bool SplitFirstOfAnyChar(FUtf8StringView View, TConstArrayView<UTF8CHAR> Search, FUtf8StringView& OutLeft, FUtf8StringView& OutRight, ESearchCase::Type SearchCase = ESearchCase::CaseSensitive, int32* OutMatchIndex = nullptr);
+CORE_API bool SplitFirstOfAnyChar(FWideStringView View, TConstArrayView<WIDECHAR> Search, FWideStringView& OutLeft, FWideStringView& OutRight, ESearchCase::Type SearchCase = ESearchCase::CaseSensitive, int32* OutMatchIndex = nullptr);
+
+/**
+ * Split the view on the last occurrence of any search character.
+ *
+ * @param View The string to split.
+ * @param Search The characters to search for and split on.
+ * @param OutLeft Receives the part of the view to the left of the search character if it is found.
+ * @param OutRight Receives the part of the view to the right of the search character if it is found.
+ * @param SearchCase Whether the comparison should ignore case.
+ * @param OutMatchIndex Receives the index of the search character that matched.
+ * @return true if the search character was found and the outputs were written, otherwise false.
+ */
+CORE_API bool SplitLastOfAnyChar(FAnsiStringView View, TConstArrayView<ANSICHAR> Search, FAnsiStringView& OutLeft, FAnsiStringView& OutRight, ESearchCase::Type SearchCase = ESearchCase::CaseSensitive, int32* OutMatchIndex = nullptr);
+CORE_API bool SplitLastOfAnyChar(FUtf8StringView View, TConstArrayView<UTF8CHAR> Search, FUtf8StringView& OutLeft, FUtf8StringView& OutRight, ESearchCase::Type SearchCase = ESearchCase::CaseSensitive, int32* OutMatchIndex = nullptr);
+CORE_API bool SplitLastOfAnyChar(FWideStringView View, TConstArrayView<WIDECHAR> Search, FWideStringView& OutLeft, FWideStringView& OutRight, ESearchCase::Type SearchCase = ESearchCase::CaseSensitive, int32* OutMatchIndex = nullptr);
+
+} // UE::String
diff --git a/Engine/Source/Runtime/Core/Tests/String/FindTest.cpp b/Engine/Source/Runtime/Core/Tests/String/FindTest.cpp
--- a/Engine/Source/Runtime/Core/Tests/String/FindTest.cpp
+++ b/Engine/Source/Runtime/Core/Tests/String/FindTest.cpp
@@ -57,7 +57,6 @@ TEST_CASE_NAMED(FStringFindTest, "System::Core::String::Find", "[Core][String][S
 
 		CHECK(FindLast(TEXT("AbCABCAbCABC"), TEXT("AbC")) == 6);
 		CHECK(FindLast(TEXT("AbCABCAbCABC"), TEXT("ABC")) == 9);
-		CHECK(FindLast(TEXT("AbCABCAbCABC"), TEXT("Bc"), ESearchCase::IgnoreCase) == 10);
 		CHECK(FindLast(TEXT("AbCABCAbCABC"), TEXT("ab")) == INDEX_NONE);
 		CHECK(FindLast(TEXT("AbCABCAbCABC"), TEXT("CD"), ESearchCase::IgnoreCase) == INDEX_NONE);
 		CHECK(FindLast(TEXT("AbCABCAbCABC"), TEXT("BC")) == 10);
@@ -103,10 +102,18 @@ TEST_CASE_NAMED(FStringFindTest, "System::Core::String::Find", "[Core][String][S
 		CHECK(FindFirstOfAny(TEXT("ABC"), {TEXT("abc"), TEXT("bc")}, ESearchCase::IgnoreCase) == 0);
 		CHECK(FindFirstOfAny(TEXT("AB"), {TEXT("ABC"), TEXT("ABD")}) == INDEX_NONE);
 
-		CHECK(FindFirstOfAny("AbCABCAbCABC", {"CABc", "ABC"}) == 3);
+		CHECK(FindFirstOfAny(ANSITEXTVIEW("AbCABCAbCABC"), {"CABc", "ABC"}) == 3);
 
 		CHECK(FindFirstOfAny(FStringView(nullptr, 0), {TEXT("ABC"), TEXT("ABD")}) == INDEX_NONE); //-V575
 		CHECK(FindFirstOfAny(FStringView(), {TEXT("ABC"), TEXT("ABD")}) == INDEX_NONE);
+
+		int32 MatchIndex = INDEX_NONE;
+		FindFirstOfAny(TEXT("AbCABCAbCABC"), {TEXT("ABC")}, ESearchCase::CaseSensitive, &MatchIndex);
+		CHECK(MatchIndex == 0);
+		FindFirstOfAny(TEXT("AbCABCAbCABC"), {TEXT("CAbC"), TEXT("ABC")}, ESearchCase::CaseSensitive, &MatchIndex);
+		CHECK(MatchIndex == 1);
+		FindFirstOfAny(TEXT("AbCABCAbCABC"), {TEXT("DEF"), TEXT("ABCD"), TEXT("Cab")}, ESearchCase::IgnoreCase, &MatchIndex);
+		CHECK(MatchIndex == 2);
 	}
 
 	SECTION("FindLastOfAny")
@@ -136,10 +143,18 @@ TEST_CASE_NAMED(FStringFindTest, "System::Core::String::Find", "[Core][String][S
 		CHECK(FindLastOfAny(TEXT("ABC"), {TEXT("abc"), TEXT("bc")}, ESearchCase::IgnoreCase) == 1);
 		CHECK(FindLastOfAny(TEXT("AB"), {TEXT("ABC"), TEXT("ABD")}) == INDEX_NONE);
 
-		CHECK(FindLastOfAny("AbCABCAbCABC", {"CABc", "ABC"}) == 9);
+		CHECK(FindLastOfAny(ANSITEXTVIEW("AbCABCAbCABC"), {"CABc", "ABC"}) == 9);
 
 		CHECK(FindLastOfAny(FStringView(nullptr, 0), { TEXT("ABC"), TEXT("ABD") }) == INDEX_NONE); //-V575
 		CHECK(FindLastOfAny(FStringView(), { TEXT("ABC"), TEXT("ABD") }) == INDEX_NONE);
+
+		int32 MatchIndex = INDEX_NONE;
+		FindLastOfAny(TEXT("AbCABCAbCABC"), {TEXT("AbC")}, ESearchCase::CaseSensitive, &MatchIndex);
+		CHECK(MatchIndex == 0);
+		FindLastOfAny(TEXT("AbCABCAbCABC"), {TEXT("CAbC"), TEXT("AbC")}, ESearchCase::CaseSensitive, &MatchIndex);
+		CHECK(MatchIndex == 1);
+		FindLastOfAny(TEXT("AbCABCAbCABC"), {TEXT("DEF"), TEXT("CAB"), TEXT("AbC")}, ESearchCase::IgnoreCase, &MatchIndex);
+		CHECK(MatchIndex == 2);
 	}
 
 	SECTION("FindFirstChar")
@@ -200,6 +215,14 @@ TEST_CASE_NAMED(FStringFindTest, "System::Core::String::Find", "[Core][String][S
 
 		CHECK(FindFirstOfAnyChar(FStringView(nullptr, 0), { TEXT('A'), TEXT('B') }) == INDEX_NONE); //-V575
 		CHECK(FindFirstOfAnyChar(FStringView(), { TEXT('A'), TEXT('B') }) == INDEX_NONE);
+
+		int32 MatchIndex = INDEX_NONE;
+		FindFirstOfAnyChar(TEXT("AbCABCAbCABC"), {TEXT('B')}, ESearchCase::CaseSensitive, &MatchIndex);
+		CHECK(MatchIndex == 0);
+		FindFirstOfAnyChar(TEXT("AbCABCAbCABC"), {TEXT('B'), TEXT('A')}, ESearchCase::CaseSensitive, &MatchIndex);
+		CHECK(MatchIndex == 1);
+		FindFirstOfAnyChar(TEXT("AbCABCAbCABC"), {TEXT('D'), TEXT('C'), TEXT('B')}, ESearchCase::IgnoreCase, &MatchIndex);
+		CHECK(MatchIndex == 2);
 	}
 
 	SECTION("FindLastOfAnyChar")
@@ -220,6 +243,14 @@ TEST_CASE_NAMED(FStringFindTest, "System::Core::String::Find", "[Core][String][S
 
 		CHECK(FindLastOfAnyChar(FStringView(nullptr, 0), { TEXT('A'), TEXT('B') }) == INDEX_NONE); //-V575
 		CHECK(FindLastOfAnyChar(FStringView(), { TEXT('A'), TEXT('B') }) == INDEX_NONE);
+
+		int32 MatchIndex = INDEX_NONE;
+		FindLastOfAnyChar(TEXT("AbCABCAbCABC"), {TEXT('A')}, ESearchCase::CaseSensitive, &MatchIndex);
+		CHECK(MatchIndex == 0);
+		FindLastOfAnyChar(TEXT("AbCABCAbCABC"), {TEXT('b'), TEXT('A')}, ESearchCase::CaseSensitive, &MatchIndex);
+		CHECK(MatchIndex == 1);
+		FindLastOfAnyChar(TEXT("AbCABCAbCABC"), {TEXT('E'), TEXT('D'), TEXT('a')}, ESearchCase::IgnoreCase, &MatchIndex);
+		CHECK(MatchIndex == 2);
 	}
 }
 
diff --git a/Engine/Source/Runtime/Core/Tests/String/SplitTest.cpp b/Engine/Source/Runtime/Core/Tests/String/SplitTest.cpp
new file mode 100644
--- /dev/null
+++ b/Engine/Source/Runtime/Core/Tests/String/SplitTest.cpp
@@ -0,0 +1,310 @@
+// Copyright Epic Games, Inc. All Rights Reserved.
+
+#if WITH_TESTS
+
+#include "String/Split.h"
+
+#include "Containers/StringView.h"
+#include "Tests/TestHarnessAdapter.h"
+
+namespace UE::String
+{
+
+TEST_CASE_NAMED(FStringSplitTest, "System::Core::String::Split", "[Core][String][SmokeFilter]")
+{
+	FStringView Left;
+	FStringView Right;
+
+	SECTION("SplitFirst")
+	{
+		const auto TestSplitFirst = []<typename T>(const T* View, const T* Search, const T* ExpectedLeft, const T* ExpectedRight, ESearchCase::Type SearchCase = ESearchCase::CaseSensitive)
+		{
+			TStringView<T> ActualLeft, ActualRight;
+			//CAPTURE(View, Search, SearchCase);
+			CHECK(SplitFirst(View, Search, ActualLeft, ActualRight, SearchCase));
+			//CAPTURE(ActualLeft, ExpectedLeft, ActualRight, ExpectedRight);
+			CHECK(ActualLeft.Equals(ExpectedLeft));
+			CHECK(ActualRight.Equals(ExpectedRight));
+		};
+
+		TestSplitFirst(TEXT("AbCABCAbCABC"), TEXT("A"), TEXT(""), TEXT("bCABCAbCABC"));
+		TestSplitFirst(TEXT("AbCABCAbCABC"), TEXT("a"), TEXT(""), TEXT("bCABCAbCABC"), ESearchCase::IgnoreCase);
+		TestSplitFirst(TEXT("AbCABCAbCABC"), TEXT("b"), TEXT("A"), TEXT("CABCAbCABC"));
+		TestSplitFirst(TEXT("AbCABCAbCABC"), TEXT("B"), TEXT("AbCA"), TEXT("CAbCABC"));
+		TestSplitFirst(TEXT("AbCABCAbCABC"), TEXT("B"), TEXT("A"), TEXT("CABCAbCABC"), ESearchCase::IgnoreCase);
+		TestSplitFirst(TEXT("AbCABCAbCABD"), TEXT("D"), TEXT("AbCABCAbCAB"), TEXT(""));
+		TestSplitFirst(TEXT("AbCABCAbCABD"), TEXT("d"), TEXT("AbCABCAbCAB"), TEXT(""), ESearchCase::IgnoreCase);
+		CHECK_FALSE(SplitFirst(TEXT("AbCABCAbCABC"), TEXT("a"), Left, Right));
+		CHECK_FALSE(SplitFirst(TEXT("AbCABCAbCABC"), TEXT("D"), Left, Right, ESearchCase::IgnoreCase));
+
+		TestSplitFirst(TEXT("AbCABCAbCABC"), TEXT("AbC"), TEXT(""), TEXT("ABCAbCABC"));
+		TestSplitFirst(TEXT("AbCABCAbCABC"), TEXT("ABC"), TEXT("AbC"), TEXT("AbCABC"));
+		TestSplitFirst(TEXT("AbCABCAbCABC"), TEXT("Bc"), TEXT("A"), TEXT("ABCAbCABC"), ESearchCase::IgnoreCase);
+		TestSplitFirst(TEXT("AbCABCAbCABD"), TEXT("BD"), TEXT("AbCABCAbCA"), TEXT(""));
+		TestSplitFirst(TEXT("AbCABCAbCABD"), TEXT("Bd"), TEXT("AbCABCAbCA"), TEXT(""), ESearchCase::IgnoreCase);
+		CHECK_FALSE(SplitFirst(TEXT("AbCABCAbCABC"), TEXT("ab"), Left, Right));
+		CHECK_FALSE(SplitFirst(TEXT("AbCABCAbCABC"), TEXT("CD"), Left, Right, ESearchCase::IgnoreCase));
+
+		TestSplitFirst(TEXT("A"), TEXT("A"), TEXT(""), TEXT(""));
+		TestSplitFirst(TEXT("A"), TEXT("a"), TEXT(""), TEXT(""), ESearchCase::IgnoreCase);
+		TestSplitFirst(TEXT("ABC"), TEXT("ABC"), TEXT(""), TEXT(""));
+		TestSplitFirst(TEXT("ABC"), TEXT("abc"), TEXT(""), TEXT(""), ESearchCase::IgnoreCase);
+		CHECK_FALSE(SplitFirst(TEXT(""), TEXT("A"), Left, Right));
+		CHECK_FALSE(SplitFirst(TEXT("AB"), TEXT("ABC"), Left, Right));
+
+		TestSplitFirst("AbCABCAbCABC", "ABC", "AbC", "AbCABC");
+
+		CHECK_FALSE(SplitFirst(FStringView(nullptr, 0), TEXT("SearchTerm"), Left, Right)); //-V575
+		CHECK_FALSE(SplitFirst(FStringView(), TEXT("SearchTerm"), Left, Right));
+	}
+
+	SECTION("SplitLast")
+	{
+		const auto TestSplitLast = []<typename T>(const T* View, const T* Search, const T* ExpectedLeft, const T* ExpectedRight, ESearchCase::Type SearchCase = ESearchCase::CaseSensitive)
+		{
+			TStringView<T> ActualLeft, ActualRight;
+			//CAPTURE(View, Search, SearchCase);
+			CHECK(SplitLast(View, Search, ActualLeft, ActualRight, SearchCase));
+			//CAPTURE(ActualLeft, ExpectedLeft, ActualRight, ExpectedRight);
+			CHECK(ActualLeft.Equals(ExpectedLeft));
+			CHECK(ActualRight.Equals(ExpectedRight));
+		};
+
+		TestSplitLast(TEXT("AbCABCAbCABC"), TEXT("b"), TEXT("AbCABCA"), TEXT("CABC"));
+		TestSplitLast(TEXT("AbCABCAbCABC"), TEXT("B"), TEXT("AbCABCAbCA"), TEXT("C"));
+		TestSplitLast(TEXT("AbCABCAbCABC"), TEXT("b"), TEXT("AbCABCAbCA"), TEXT("C"), ESearchCase::IgnoreCase);
+		TestSplitLast(TEXT("AbCABCAbCABD"), TEXT("D"), TEXT("AbCABCAbCAB"), TEXT(""));
+		TestSplitLast(TEXT("AbCABCAbCABD"), TEXT("d"), TEXT("AbCABCAbCAB"), TEXT(""), ESearchCase::IgnoreCase);
+		CHECK_FALSE(SplitLast(TEXT("AbCABCAbCABC"), TEXT("a"), Left, Right));
+		CHECK_FALSE(SplitLast(TEXT("AbCABCAbCABC"), TEXT("D"), Left, Right, ESearchCase::IgnoreCase));
+
+		TestSplitLast(TEXT("AbCABCAbCABC"), TEXT("AbC"), TEXT("AbCABC"), TEXT("ABC"));
+		TestSplitLast(TEXT("AbCABCAbCABC"), TEXT("ABC"), TEXT("AbCABCAbC"), TEXT(""));
+		TestSplitLast(TEXT("AbCABCAbCABC"), TEXT("BC"), TEXT("AbCABCAbCA"), TEXT(""));
+		TestSplitLast(TEXT("AbCABCAbCABC"), TEXT("Bc"), TEXT("AbCABCAbCA"), TEXT(""), ESearchCase::IgnoreCase);
+		CHECK_FALSE(SplitLast(TEXT("AbCABCAbCABC"), TEXT("ab"), Left, Right));
+		CHECK_FALSE(SplitLast(TEXT("AbCABCAbCABC"), TEXT("CD"), Left, Right, ESearchCase::IgnoreCase));
+
+		TestSplitLast(TEXT("A"), TEXT("A"), TEXT(""), TEXT(""));
+		TestSplitLast(TEXT("A"), TEXT("A"), TEXT(""), TEXT(""), ESearchCase::IgnoreCase);
+		TestSplitLast(TEXT("ABC"), TEXT("ABC"), TEXT(""), TEXT(""));
+		TestSplitLast(TEXT("ABC"), TEXT("abc"), TEXT(""), TEXT(""), ESearchCase::IgnoreCase);
+		CHECK_FALSE(SplitLast(TEXT(""), TEXT("A"), Left, Right));
+		CHECK_FALSE(SplitLast(TEXT("AB"), TEXT("ABC"), Left, Right));
+
+		TestSplitLast("AbCABCAbCABC", "ABC", "AbCABCAbC", "");
+
+		CHECK_FALSE(SplitLast(FStringView(nullptr, 0), TEXT("SearchTerm"), Left, Right)); //-V575
+		CHECK_FALSE(SplitLast(FStringView(), TEXT("SearchTerm"), Left, Right));
+	}
+
+	SECTION("SplitFirstOfAny")
+	{
+		const auto TestSplitFirstOfAny = []<typename T>(const T* View, TConstArrayView<TStringView<T>> Search, const T* ExpectedLeft, const T* ExpectedRight, ESearchCase::Type SearchCase = ESearchCase::CaseSensitive)
+		{
+			TStringView<T> ActualLeft, ActualRight;
+			//CAPTURE(View, Search, SearchCase);
+			CHECK(SplitFirstOfAny(View, Search, ActualLeft, ActualRight, SearchCase));
+			//CAPTURE(ActualLeft, ExpectedLeft, ActualRight, ExpectedRight);
+			CHECK(ActualLeft.Equals(ExpectedLeft));
+			CHECK(ActualRight.Equals(ExpectedRight));
+		};
+
+		TestSplitFirstOfAny(TEXT("AbCABCAbcABC"), {TEXT("A"), TEXT("B")}, TEXT(""), TEXT("bCABCAbcABC"));
+		TestSplitFirstOfAny(TEXT("AbCABCAbcABC"), {TEXT("a"), TEXT("B")}, TEXT(""), TEXT("bCABCAbcABC"), ESearchCase::IgnoreCase);
+		TestSplitFirstOfAny(TEXT("AbCABCAbcABC"), {TEXT("b")}, TEXT("A"), TEXT("CABCAbcABC"));
+		TestSplitFirstOfAny(TEXT("AbCABCAbcABC"), {TEXT("c"), TEXT("B")}, TEXT("AbCA"), TEXT("CAbcABC"));
+		TestSplitFirstOfAny(TEXT("AbCABCAbcABC"), {TEXT("c"), TEXT("B")}, TEXT("A"), TEXT("CABCAbcABC"), ESearchCase::IgnoreCase);
+		TestSplitFirstOfAny(TEXT("AbCABCAbcABD"), {TEXT("E"), TEXT("D")}, TEXT("AbCABCAbcAB"), TEXT(""));
+		TestSplitFirstOfAny(TEXT("AbCABCAbcABD"), {TEXT("E"), TEXT("d")}, TEXT("AbCABCAbcAB"), TEXT(""), ESearchCase::IgnoreCase);
+		CHECK_FALSE(SplitFirstOfAny(TEXT("AbCABCAbcABC"), {TEXT("D"), TEXT("a")}, Left, Right));
+		CHECK_FALSE(SplitFirstOfAny(TEXT("AbCABCAbcABC"), {TEXT("E"), TEXT("D")}, Left, Right, ESearchCase::IgnoreCase));
+
+		TestSplitFirstOfAny(TEXT("AbCABCAbCABC"), {TEXT("AbC")}, TEXT(""), TEXT("ABCAbCABC"));
+		TestSplitFirstOfAny(TEXT("AbCABCAbCABC"), {TEXT("CABc"), TEXT("ABC")}, TEXT("AbC"), TEXT("AbCABC"));
+		TestSplitFirstOfAny(TEXT("AbCABCAbCABC"), {TEXT("ABD"), TEXT("Bc")}, TEXT("A"), TEXT("ABCAbCABC"), ESearchCase::IgnoreCase);
+		TestSplitFirstOfAny(TEXT("AbCABCAbCABD"), {TEXT("BD"), TEXT("CABB")}, TEXT("AbCABCAbCA"), TEXT(""));
+		TestSplitFirstOfAny(TEXT("AbCABCAbCABD"), {TEXT("Bd"), TEXT("CABB")}, TEXT("AbCABCAbCA"), TEXT(""), ESearchCase::IgnoreCase);
+		CHECK_FALSE(SplitFirstOfAny(TEXT("AbCABCAbCABC"), {TEXT("bc"), TEXT("ab")}, Left, Right));
+		CHECK_FALSE(SplitFirstOfAny(TEXT("AbCABCAbCABC"), {TEXT("DA"), TEXT("CD")}, Left, Right, ESearchCase::IgnoreCase));
+
+		TestSplitFirstOfAny(TEXT("A"), {TEXT("A"), TEXT("B")}, TEXT(""), TEXT(""));
+		TestSplitFirstOfAny(TEXT("A"), {TEXT("A"), TEXT("B")}, TEXT(""), TEXT(""), ESearchCase::IgnoreCase);
+		TestSplitFirstOfAny(TEXT("ABC"), {TEXT("ABC"), TEXT("BC")}, TEXT(""), TEXT(""));
+		TestSplitFirstOfAny(TEXT("ABC"), {TEXT("abc"), TEXT("bc")}, TEXT(""), TEXT(""), ESearchCase::IgnoreCase);
+		CHECK_FALSE(SplitFirstOfAny(TEXT(""), {TEXT("A"), TEXT("B")}, Left, Right));
+		CHECK_FALSE(SplitFirstOfAny(TEXT("AB"), {TEXT("ABC"), TEXT("ABD")}, Left, Right));
+
+		TestSplitFirstOfAny("AbCABCAbCABC", {"CABc", "ABC"}, "AbC", "AbCABC");
+
+		CHECK_FALSE(SplitFirstOfAny(FStringView(nullptr, 0), {TEXT("ABC"), TEXT("ABD")}, Left, Right)); //-V575
+		CHECK_FALSE(SplitFirstOfAny(FStringView(), {TEXT("ABC"), TEXT("ABD")}, Left, Right));
+	}
+
+	SECTION("SplitLastOfAny")
+	{
+		const auto TestSplitLastOfAny = []<typename T>(const T* View, TConstArrayView<TStringView<T>> Search, const T* ExpectedLeft, const T* ExpectedRight, ESearchCase::Type SearchCase = ESearchCase::CaseSensitive)
+		{
+			TStringView<T> ActualLeft, ActualRight;
+			//CAPTURE(View, Search, SearchCase);
+			CHECK(SplitLastOfAny(View, Search, ActualLeft, ActualRight, SearchCase));
+			//CAPTURE(ActualLeft, ExpectedLeft, ActualRight, ExpectedRight);
+			CHECK(ActualLeft.Equals(ExpectedLeft));
+			CHECK(ActualRight.Equals(ExpectedRight));
+		};
+
+		TestSplitLastOfAny(TEXT("AbCABCAbcABC"), {TEXT("a"), TEXT("b")}, TEXT("AbCABCA"), TEXT("cABC"));
+		TestSplitLastOfAny(TEXT("AbCABCAbcABC"), {TEXT("a"), TEXT("b")}, TEXT("AbCABCAbcA"), TEXT("C"), ESearchCase::IgnoreCase);
+		TestSplitLastOfAny(TEXT("AbCABCAbcABC"), {TEXT("b")}, TEXT("AbCABCA"), TEXT("cABC"));
+		TestSplitLastOfAny(TEXT("AbCABCAbcABC"), {TEXT("c"), TEXT("B")}, TEXT("AbCABCAbcA"), TEXT("C"));
+		TestSplitLastOfAny(TEXT("AbCABCAbcABC"), {TEXT("c"), TEXT("B")}, TEXT("AbCABCAbcAB"), TEXT(""), ESearchCase::IgnoreCase);
+		TestSplitLastOfAny(TEXT("AbCABCAbcABD"), {TEXT("E"), TEXT("D")}, TEXT("AbCABCAbcAB"), TEXT(""));
+		TestSplitLastOfAny(TEXT("AbCABCAbcABD"), {TEXT("E"), TEXT("d")}, TEXT("AbCABCAbcAB"), TEXT(""), ESearchCase::IgnoreCase);
+		CHECK_FALSE(SplitLastOfAny(TEXT("AbCABCAbcABC"), {TEXT("D"), TEXT("a")}, Left, Right));
+		CHECK_FALSE(SplitLastOfAny(TEXT("AbCABCAbcABC"), {TEXT("E"), TEXT("D")}, Left, Right, ESearchCase::IgnoreCase));
+
+		TestSplitLastOfAny(TEXT("AbCABCAbCABC"), {TEXT("AbC")}, TEXT("AbCABC"), TEXT("ABC"));
+		TestSplitLastOfAny(TEXT("AbCABCAbCABC"), {TEXT("CABc"), TEXT("ABC")}, TEXT("AbCABCAbC"), TEXT(""));
+		TestSplitLastOfAny(TEXT("AbCABCAbCABC"), {TEXT("ABD"), TEXT("Bc")}, TEXT("AbCABCAbCA"), TEXT(""), ESearchCase::IgnoreCase);
+		TestSplitLastOfAny(TEXT("AbCABCAbCABD"), {TEXT("BD"), TEXT("CABB")}, TEXT("AbCABCAbCA"), TEXT(""));
+		TestSplitLastOfAny(TEXT("AbCABCAbCABD"), {TEXT("Bd"), TEXT("CABB")}, TEXT("AbCABCAbCA"), TEXT(""), ESearchCase::IgnoreCase);
+		CHECK_FALSE(SplitLastOfAny(TEXT("AbCABCAbCABC"), {TEXT("bc"), TEXT("ab")}, Left, Right));
+		CHECK_FALSE(SplitLastOfAny(TEXT("AbCABCAbCABC"), {TEXT("DA"), TEXT("CD")}, Left, Right, ESearchCase::IgnoreCase));
+
+		TestSplitLastOfAny(TEXT("A"), {TEXT("A"), TEXT("B")}, TEXT(""), TEXT(""));
+		TestSplitLastOfAny(TEXT("A"), {TEXT("A"), TEXT("B")}, TEXT(""), TEXT(""), ESearchCase::IgnoreCase);
+		TestSplitLastOfAny(TEXT("ABC"), {TEXT("ABC"), TEXT("BC")}, TEXT("A"), TEXT(""));
+		TestSplitLastOfAny(TEXT("ABC"), {TEXT("abc"), TEXT("bc")}, TEXT("A"), TEXT(""), ESearchCase::IgnoreCase);
+		CHECK_FALSE(SplitLastOfAny(TEXT(""), {TEXT("A"), TEXT("B")}, Left, Right));
+		CHECK_FALSE(SplitLastOfAny(TEXT("AB"), {TEXT("ABC"), TEXT("ABD")}, Left, Right));
+
+		TestSplitLastOfAny("AbCABCAbCABC", {"CABc", "ABC"}, "AbCABCAbC", "");
+
+		CHECK_FALSE(SplitLastOfAny(FStringView(nullptr, 0), {TEXT("ABC"), TEXT("ABD")}, Left, Right)); //-V575
+		CHECK_FALSE(SplitLastOfAny(FStringView(), {TEXT("ABC"), TEXT("ABD")}, Left, Right));
+	}
+
+	SECTION("SplitFirstChar")
+	{
+		const auto TestSplitFirstChar = []<typename T>(const T* View, T Search, const T* ExpectedLeft, const T* ExpectedRight, ESearchCase::Type SearchCase = ESearchCase::CaseSensitive)
+		{
+			TStringView<T> ActualLeft, ActualRight;
+			//CAPTURE(View, Search, SearchCase);
+			CHECK(SplitFirstChar(View, Search, ActualLeft, ActualRight, SearchCase));
+			//CAPTURE(ActualLeft, ExpectedLeft, ActualRight, ExpectedRight);
+			CHECK(ActualLeft.Equals(ExpectedLeft));
+			CHECK(ActualRight.Equals(ExpectedRight));
+		};
+
+		TestSplitFirstChar(TEXT("AbCABCAbCABC"), TEXT('b'), TEXT("A"), TEXT("CABCAbCABC"));
+		TestSplitFirstChar(TEXT("AbCABCAbCABC"), TEXT('B'), TEXT("AbCA"), TEXT("CAbCABC"));
+		TestSplitFirstChar(TEXT("AbCABCAbCABC"), TEXT('B'), TEXT("A"), TEXT("CABCAbCABC"), ESearchCase::IgnoreCase);
+		TestSplitFirstChar(TEXT("AbCABCAbCABD"), TEXT('D'), TEXT("AbCABCAbCAB"), TEXT(""));
+		TestSplitFirstChar(TEXT("AbCABCAbCABD"), TEXT('d'), TEXT("AbCABCAbCAB"), TEXT(""), ESearchCase::IgnoreCase);
+		CHECK_FALSE(SplitFirstChar(TEXT("AbCABCAbCABC"), TEXT('a'), Left, Right));
+		CHECK_FALSE(SplitFirstChar(TEXT("AbCABCAbCABC"), TEXT('D'), Left, Right, ESearchCase::IgnoreCase));
+
+		TestSplitFirstChar(TEXT("A"), TEXT('A'), TEXT(""), TEXT(""));
+		TestSplitFirstChar(TEXT("A"), TEXT('a'), TEXT(""), TEXT(""), ESearchCase::IgnoreCase);
+		CHECK_FALSE(SplitFirstChar(TEXT(""), TEXT('A'), Left, Right));
+
+		TestSplitFirstChar("AbCABCAbCABC", 'B', "AbCA", "CAbCABC");
+
+		CHECK_FALSE(SplitFirstChar(FStringView(nullptr, 0), TEXT('A'), Left, Right)); //-V575
+		CHECK_FALSE(SplitFirstChar(FStringView(), TEXT('A'), Left, Right));
+	}
+
+	SECTION("SplitLastChar")
+	{
+		const auto TestSplitLastChar = []<typename T>(const T* View, T Search, const T* ExpectedLeft, const T* ExpectedRight, ESearchCase::Type SearchCase = ESearchCase::CaseSensitive)
+		{
+			TStringView<T> ActualLeft, ActualRight;
+			//CAPTURE(View, Search, SearchCase);
+			CHECK(SplitLastChar(View, Search, ActualLeft, ActualRight, SearchCase));
+			//CAPTURE(ActualLeft, ExpectedLeft, ActualRight, ExpectedRight);
+			CHECK(ActualLeft.Equals(ExpectedLeft));
+			CHECK(ActualRight.Equals(ExpectedRight));
+		};
+
+		TestSplitLastChar(TEXT("AbCABCAbCABC"), TEXT('b'), TEXT("AbCABCA"), TEXT("CABC"));
+		TestSplitLastChar(TEXT("AbCABCAbCABC"), TEXT('B'), TEXT("AbCABCAbCA"), TEXT("C"));
+		TestSplitLastChar(TEXT("AbCABCAbCABC"), TEXT('b'), TEXT("AbCABCAbCA"), TEXT("C"), ESearchCase::IgnoreCase);
+		TestSplitLastChar(TEXT("AbCABCAbCABD"), TEXT('D'), TEXT("AbCABCAbCAB"), TEXT(""));
+		TestSplitLastChar(TEXT("AbCABCAbCABD"), TEXT('d'), TEXT("AbCABCAbCAB"), TEXT(""), ESearchCase::IgnoreCase);
+		CHECK_FALSE(SplitLastChar(TEXT("AbCABCAbCABC"), TEXT('a'), Left, Right));
+		CHECK_FALSE(SplitLastChar(TEXT("AbCABCAbCABC"), TEXT('D'), Left, Right, ESearchCase::IgnoreCase));
+
+		TestSplitLastChar(TEXT("A"), TEXT('A'), TEXT(""), TEXT(""));
+		TestSplitLastChar(TEXT("A"), TEXT('a'), TEXT(""), TEXT(""), ESearchCase::IgnoreCase);
+		CHECK_FALSE(SplitLastChar(TEXT(""), TEXT('A'), Left, Right));
+
+		TestSplitLastChar("AbCABCAbCABC", 'B', "AbCABCAbCA", "C");
+
+		CHECK_FALSE(SplitLastChar(FStringView(nullptr, 0), TEXT('A'), Left, Right)); //-V575
+		CHECK_FALSE(SplitLastChar(FStringView(), TEXT('A'), Left, Right));
+	}
+
+	SECTION("SplitFirstOfAnyChar")
+	{
+		const auto TestSplitFirstOfAnyChar = []<typename T>(const T* View, TConstArrayView<T> Search, const T* ExpectedLeft, const T* ExpectedRight, ESearchCase::Type SearchCase = ESearchCase::CaseSensitive)
+		{
+			TStringView<T> ActualLeft, ActualRight;
+			//CAPTURE(View, Search, SearchCase);
+			CHECK(SplitFirstOfAnyChar(View, Search, ActualLeft, ActualRight, SearchCase));
+			//CAPTURE(ActualLeft, ExpectedLeft, ActualRight, ExpectedRight);
+			CHECK(ActualLeft.Equals(ExpectedLeft));
+			CHECK(ActualRight.Equals(ExpectedRight));
+		};
+
+		TestSplitFirstOfAnyChar(TEXT("AbCABCAbcABC"), {TEXT('b')}, TEXT("A"), TEXT("CABCAbcABC"));
+		TestSplitFirstOfAnyChar(TEXT("AbCABCAbcABC"), {TEXT('c'), TEXT('B')}, TEXT("AbCA"), TEXT("CAbcABC"));
+		TestSplitFirstOfAnyChar(TEXT("AbCABCAbcABC"), {TEXT('c'), TEXT('B')}, TEXT("A"), TEXT("CABCAbcABC"), ESearchCase::IgnoreCase);
+		TestSplitFirstOfAnyChar(TEXT("AbCABCAbcABD"), {TEXT('E'), TEXT('D')}, TEXT("AbCABCAbcAB"), TEXT(""));
+		TestSplitFirstOfAnyChar(TEXT("AbCABCAbcABD"), {TEXT('E'), TEXT('d')}, TEXT("AbCABCAbcAB"), TEXT(""), ESearchCase::IgnoreCase);
+		CHECK_FALSE(SplitFirstOfAnyChar(TEXT("AbCABCAbcABC"), {TEXT('D'), TEXT('a')}, Left, Right));
+		CHECK_FALSE(SplitFirstOfAnyChar(TEXT("AbCABCAbcABC"), {TEXT('E'), TEXT('D')}, Left, Right, ESearchCase::IgnoreCase));
+
+		TestSplitFirstOfAnyChar(TEXT("A"), {TEXT('A'), TEXT('B')}, TEXT(""), TEXT(""));
+		TestSplitFirstOfAnyChar(TEXT("A"), {TEXT('A'), TEXT('B')}, TEXT(""), TEXT(""), ESearchCase::IgnoreCase);
+		CHECK_FALSE(SplitFirstOfAnyChar(TEXT(""), {TEXT('A'), TEXT('B')}, Left, Right));
+
+		TestSplitFirstOfAnyChar("AbCABCAbcABC", {'c', 'B'}, "AbCA", "CAbcABC");
+
+		CHECK_FALSE(SplitFirstOfAnyChar(FStringView(nullptr, 0), {TEXT('A'), TEXT('B')}, Left, Right)); //-V575
+		CHECK_FALSE(SplitFirstOfAnyChar(FStringView(), {TEXT('A'), TEXT('B')}, Left, Right));
+	}
+
+	SECTION("SplitLastOfAnyChar")
+	{
+		const auto TestSplitLastOfAnyChar = []<typename T>(const T* View, TConstArrayView<T> Search, const T* ExpectedLeft, const T* ExpectedRight, ESearchCase::Type SearchCase = ESearchCase::CaseSensitive)
+		{
+			TStringView<T> ActualLeft, ActualRight;
+			//CAPTURE(View, Search, SearchCase);
+			CHECK(SplitLastOfAnyChar(View, Search, ActualLeft, ActualRight, SearchCase));
+			//CAPTURE(ActualLeft, ExpectedLeft, ActualRight, ExpectedRight);
+			CHECK(ActualLeft.Equals(ExpectedLeft));
+			CHECK(ActualRight.Equals(ExpectedRight));
+		};
+
+		TestSplitLastOfAnyChar(TEXT("AbCABCAbcABC"), {TEXT('b')}, TEXT("AbCABCA"), TEXT("cABC"));
+		TestSplitLastOfAnyChar(TEXT("AbCABCAbcABC"), {TEXT('c'), TEXT('B')}, TEXT("AbCABCAbcA"), TEXT("C"));
+		TestSplitLastOfAnyChar(TEXT("AbCABCAbcABC"), {TEXT('c'), TEXT('B')}, TEXT("AbCABCAbcAB"), TEXT(""), ESearchCase::IgnoreCase);
+		TestSplitLastOfAnyChar(TEXT("AbCABCAbcABD"), {TEXT('E'), TEXT('D')}, TEXT("AbCABCAbcAB"), TEXT(""));
+		TestSplitLastOfAnyChar(TEXT("AbCABCAbcABD"), {TEXT('E'), TEXT('d')}, TEXT("AbCABCAbcAB"), TEXT(""), ESearchCase::IgnoreCase);
+		CHECK_FALSE(SplitLastOfAnyChar(TEXT("AbCABCAbcABC"), {TEXT('D'), TEXT('a')}, Left, Right));
+		CHECK_FALSE(SplitLastOfAnyChar(TEXT("AbCABCAbcABC"), {TEXT('E'), TEXT('D')}, Left, Right, ESearchCase::IgnoreCase));
+
+		TestSplitLastOfAnyChar(TEXT("A"), {TEXT('A'), TEXT('B')}, TEXT(""), TEXT(""));
+		TestSplitLastOfAnyChar(TEXT("A"), {TEXT('A'), TEXT('B')}, TEXT(""), TEXT(""), ESearchCase::IgnoreCase);
+		CHECK_FALSE(SplitLastOfAnyChar(TEXT(""), {TEXT('A'), TEXT('B')}, Left, Right));
+
+		TestSplitLastOfAnyChar("AbCABCAbcABC", {'c', 'B'}, "AbCABCAbcA", "C");
+
+		CHECK_FALSE(SplitLastOfAnyChar(FStringView(nullptr, 0), {TEXT('A'), TEXT('B')}, Left, Right)); //-V575
+		CHECK_FALSE(SplitLastOfAnyChar(FStringView(), {TEXT('A'), TEXT('B')}, Left, Right));
+	}
+}
+
+} // UE::String
+
+#endif // WITH_TESTS
