[ExecutionResource]
  - Maintain per-task execution resource to avoid busy wait tasks to copy another task's resources

diff --git a/Engine/Source/Runtime/Core/Private/Experimental/Misc/ExecutionResource.cpp b/Engine/Source/Runtime/Core/Private/Experimental/Misc/ExecutionResource.cpp
--- a/Engine/Source/Runtime/Core/Private/Experimental/Misc/ExecutionResource.cpp
+++ b/Engine/Source/Runtime/Core/Private/Experimental/Misc/ExecutionResource.cpp
@@ -4,20 +4,85 @@
 #include "Containers/Array.h"
 #include "HAL/Platform.h"
 #include "Templates/UnrealTemplate.h"
+#include "Async/Fundamental/Task.h"
+#include "Containers/Map.h"
 
 namespace ExecutionResourceImpl
 {
-	static thread_local TArray<TRefCountPtr<IExecutionResource>> ExecutionResourceStack;
+	class FState
+	{
+	public:
+		void Push(const LowLevelTasks::FTask* ActiveTask, TRefCountPtr<IExecutionResource> ExecutionResource)
+		{
+			if (MainTask == INVALID_TASK || MainTask == UPTRINT(ActiveTask))
+			{
+				MainTask = UPTRINT(ActiveTask);
+				MainStack.Push(MoveTemp(ExecutionResource));
+			}
+			else
+			{
+				AdditionalStacks.FindOrAdd(ActiveTask).Push(MoveTemp(ExecutionResource));
+			}
+		}
+
+		void Pop(const LowLevelTasks::FTask* ActiveTask)
+		{
+			if (MainTask == UPTRINT(ActiveTask))
+			{
+				MainStack.Pop(EAllowShrinking::No);
+				if (MainStack.IsEmpty())
+				{
+					MainTask = INVALID_TASK;
+					MainStack.Shrink();
+				}
+			}
+			else
+			{
+				TArray<TRefCountPtr<IExecutionResource>>& Stack = AdditionalStacks.FindChecked(ActiveTask);
+				Stack.Pop();
+				if (Stack.IsEmpty())
+				{
+					AdditionalStacks.Remove(ActiveTask);
+				}
+			}
+		}
+
+		TArray<TRefCountPtr<IExecutionResource>>* Get(const LowLevelTasks::FTask* ActiveTask)
+		{
+			if (MainTask == UPTRINT(ActiveTask))
+			{
+				return &MainStack;
+			}
+			else
+			{
+				return AdditionalStacks.Find(ActiveTask);
+			}
+		}
+
+	private:
+		// Define the invalid task as something other than nullptr for the case when we're not called as part of a backend task.
+		static constexpr UPTRINT INVALID_TASK = 1;
+
+		// Fast path to avoid using the map in the normal case
+		UPTRINT                                  MainTask = INVALID_TASK;
+		TArray<TRefCountPtr<IExecutionResource>> MainStack;
+
+		// Will be used during busy waits (i.e. multiple active task for the same thread_local)
+		// Because we don't want busy wait task to end up with the resource of another task.
+		TMap<const LowLevelTasks::FTask*, TArray<TRefCountPtr<IExecutionResource>>> AdditionalStacks;
+	};
+
+	static thread_local FState State;
 }
 
 FExecutionResourceContextScope::FExecutionResourceContextScope(TRefCountPtr<IExecutionResource> ExecutionResource)
 {
-	ExecutionResourceImpl::ExecutionResourceStack.Push(MoveTemp(ExecutionResource));
+	ExecutionResourceImpl::State.Push(LowLevelTasks::FTask::GetActiveTask(), MoveTemp(ExecutionResource));
 }
 
 FExecutionResourceContextScope::~FExecutionResourceContextScope()
 {
-	ExecutionResourceImpl::ExecutionResourceStack.Pop();
+	ExecutionResourceImpl::State.Pop(LowLevelTasks::FTask::GetActiveTask());
 }
 
 class FCompositeExecutionResource : public IExecutionResource, public FThreadSafeRefCountedObject
@@ -49,11 +114,12 @@ private:
 
 TRefCountPtr<IExecutionResource> FExecutionResourceContext::Get()
 {
-	using namespace ExecutionResourceImpl;
-	if (ExecutionResourceStack.IsEmpty())
+	TArray<TRefCountPtr<IExecutionResource>>* Stack = ExecutionResourceImpl::State.Get(LowLevelTasks::FTask::GetActiveTask());
+
+	if (Stack == nullptr || Stack->IsEmpty())
 	{
 		return nullptr;
 	}
 
-	return new FCompositeExecutionResource(ExecutionResourceStack);
-}
+	return new FCompositeExecutionResource(*Stack);
+}
\ No newline at end of file
diff --git a/Engine/Source/Runtime/Core/Tests/Tasks/TasksTest.cpp b/Engine/Source/Runtime/Core/Tests/Tasks/TasksTest.cpp
--- a/Engine/Source/Runtime/Core/Tests/Tasks/TasksTest.cpp
+++ b/Engine/Source/Runtime/Core/Tests/Tasks/TasksTest.cpp
@@ -7,6 +7,7 @@
 #include "Tasks/Pipe.h"
 #include "Tasks/TaskConcurrencyLimiter.h"
 #include "HAL/Thread.h"
+#include "Experimental/Misc/ExecutionResource.h"
 #include "Async/ParallelFor.h"
 #include "Async/ManualResetEvent.h"
 #include "Tests/TestHarnessAdapter.h"
@@ -1716,6 +1717,108 @@ namespace UE { namespace TasksTests
 		// Do not exit the test before all blocked workers' tasks are done.
 		LowLevelTasks::BusyWaitForTasks<LowLevelTasks::FTask>(WorkerBlockers);
 	}
+
+	// Test that validates each tasks have their own execution resource stack and that busy wait cannot interfere with other tasks' resources.
+	TEST_CASE_NAMED(FTasksCaptureExecutionResourceInsideBusyWait, "System::Core::Async::Tasks::CaptureExecutionResourceInsideBusyWait", "[.][ApplicationContextMask][EngineFilter]")
+	{
+		using namespace LowLevelTasks;
+		
+		FPlatformProcess::Sleep(0.1f); // give workers time to fall asleep, to avoid any reserve worker messing around
+
+		uint32 NumWorkers = LowLevelTasks::FScheduler::Get().GetNumWorkers();
+
+		// Block all workers but one so we can prevent the busywait task from running outside the busywait
+		FTaskEvent ResumeEvent{ UE_SOURCE_LOCATION };
+		TArray<LowLevelTasks::FTask> WorkerBlockers = BlockWorkers(ResumeEvent, NumWorkers - 1);
+
+		std::atomic<bool>     NormalResourceHeld{ false };
+		std::atomic<bool>     BusyWaitResourceHeld{ false };
+
+		class FResource : public FThreadSafeRefCountedObject, public IExecutionResource
+		{
+			std::atomic<bool>& ResourceHeld;
+		public:
+			FResource(std::atomic<bool>& ResourceToHold)
+				: ResourceHeld(ResourceToHold)
+			{
+				ResourceHeld = true;
+			}
+			virtual ~FResource() override
+			{
+				ResourceHeld = false;
+			}
+			uint32 AddRef() const override
+			{
+				return FThreadSafeRefCountedObject::AddRef();
+			}
+
+			uint32 Release() const override
+			{
+				return FThreadSafeRefCountedObject::Release();
+			}
+
+			uint32 GetRefCount() const override
+			{
+				return FThreadSafeRefCountedObject::GetRefCount();
+			}
+		};
+
+		TRefCountPtr<IExecutionResource> NormalResource;
+		TRefCountPtr<IExecutionResource> BusyWaitResource;
+
+		FTask NormalTask =
+			Launch(TEXT("NormalTask"),
+				[&]()
+				{
+					verify(!LowLevelTasks::FScheduler::Get().IsBusyWaiting());
+
+					FExecutionResourceContextScope ExecutionResourceScope(new FResource(NormalResourceHeld));
+
+					NormalResource = FExecutionResourceContext::Get();
+
+					BusyWaitUntil([&]() { return BusyWaitResourceHeld.load(); });
+				}
+			);
+
+		FTask BusyWaitTask =
+			Launch(TEXT("BusyWaitTask"),
+				[&]()
+				{
+					verify(LowLevelTasks::FScheduler::Get().IsBusyWaiting());
+
+					FExecutionResourceContextScope ExecutionResourceScope(new FResource(BusyWaitResourceHeld));
+
+					BusyWaitResource = FExecutionResourceContext::Get();
+				}
+		);
+
+		NormalTask.Wait();
+		
+		// This should have been completed as part of the normal task
+		verify(BusyWaitTask.IsCompleted());
+
+		// Verify now that both resource are held
+		verify(NormalResource.IsValid());
+		verify(BusyWaitResource.IsValid());
+
+		verify(NormalResourceHeld.load());
+		verify(BusyWaitResourceHeld.load());
+
+		// Now confirm that we can release the normal resource and its not also being held by the busy wait resource
+		NormalResource.SafeRelease();
+		verify(!NormalResourceHeld.load());
+		verify(BusyWaitResourceHeld.load());
+
+		// Now verify that we can release the busy wait resource too.
+		BusyWaitResource.SafeRelease();
+		verify(!BusyWaitResourceHeld.load());
+
+		// Unblock all workers
+		ResumeEvent.Trigger();
+
+		// Do not exit the test before all blocked workers' tasks are done.
+		LowLevelTasks::BusyWaitForTasks<LowLevelTasks::FTask>(WorkerBlockers);
+	}
 }}
 
 #endif // WITH_TESTS
