IoStore: Move FIoStoreToc::GetTocChunkInfo to FIoStoreTocResource::GetTocChunkInfo and split the concerns of FIoStoreToc into FIoStoreTocBuilder and FIoStoreTocReader.
Prepare for moving IoStoreWriter out of Core to IoStoreUtilities.

#rnx
#tests Packaging and running cooked LyraGame

diff --git a/Engine/Source/Runtime/Core/Internal/IO/IoStore.h b/Engine/Source/Runtime/Core/Internal/IO/IoStore.h
--- a/Engine/Source/Runtime/Core/Internal/IO/IoStore.h
+++ b/Engine/Source/Runtime/Core/Internal/IO/IoStore.h
@@ -213,6 +213,8 @@ struct FIoStoreTocResource
 	TArray<FIoStoreTocOnDemandChunkMeta> OnDemandChunkMeta;
 	TArray<FIoStoreTocOnDemandCompressedBlockMeta> OnDemandCompressedBlockMeta;
 
+	[[nodiscard]] CORE_API FIoStoreTocChunkInfo GetTocChunkInfo(int32 TocEntryIndex) const;
+
 	[[nodiscard]] CORE_API static FIoStatus Read(const TCHAR* TocFilePath, EIoStoreTocReadOptions ReadOptions, FIoStoreTocResource& OutTocResource);
 
 	[[nodiscard]] CORE_API static TIoStatusOr<uint64> Write(const TCHAR* TocFilePath, FIoStoreTocResource& TocResource, uint32 CompressionBlockSize, uint64 MaxPartitionSize, const FIoContainerSettings& ContainerSettings);
diff --git a/Engine/Source/Runtime/Core/Private/IO/IoStore.cpp b/Engine/Source/Runtime/Core/Private/IO/IoStore.cpp
--- a/Engine/Source/Runtime/Core/Private/IO/IoStore.cpp
+++ b/Engine/Source/Runtime/Core/Private/IO/IoStore.cpp
@@ -8,23 +8,18 @@
 #include "Async/ParallelFor.h"
 #include "Containers/Map.h"
 #include "DerivedDataCacheInterface.h"
-#include "DerivedDataPluginInterface.h"
 #include "Features/IModularFeatures.h"
 #include "HAL/FileManager.h"
 #include "HAL/PlatformFileManager.h"
 #include "IO/IoDirectoryIndex.h"
-#include "Interfaces/ITargetPlatform.h"
 #include "Misc/Compression.h"
-#include "Misc/ConfigCacheIni.h"
 #include "Misc/CoreDelegates.h"
 #include "Misc/Paths.h"
 #include "Misc/StringBuilder.h"
-#include "Modules/ModuleManager.h"
 #include "ProfilingDebugging/CountersTrace.h"
 #include "Serialization/BufferWriter.h"
 #include "Serialization/LargeMemoryReader.h"
 #include "Serialization/MemoryWriter.h"
-#include "Serialization/MemoryWriter.h"
 #include "Tasks/Task.h"
 #include "Templates/UniquePtr.h"
 
@@ -549,24 +544,14 @@ FIoStoreWriterContext::FProgress FIoStoreWriterContext::GetProgress() const
 	return Impl->GetProgress();
 }
 
-class FIoStoreToc
+class FIoStoreTocBuilder
 {
 public:
-	FIoStoreToc()
+	FIoStoreTocBuilder()
 	{
 		FMemory::Memzero(&Toc.Header, sizeof(FIoStoreTocHeader));
 	}
 
-	void Initialize()
-	{
-		ChunkIdToIndex.Empty(Toc.ChunkIds.Num());
-
-		for (int32 ChunkIndex = 0; ChunkIndex < Toc.ChunkIds.Num(); ++ChunkIndex)
-		{
-			ChunkIdToIndex.Add(Toc.ChunkIds[ChunkIndex], ChunkIndex);
-		}
-	}
-
 	int32 AddChunkEntry(const FIoChunkId& ChunkId, const FIoOffsetAndLength& OffsetLength, const FIoStoreTocEntryMeta& Meta, const FIoStoreTocOnDemandChunkMeta& OnDemandData)
 	{
 		int32& Index = ChunkIdToIndex.FindOrAdd(ChunkId);
@@ -639,16 +624,6 @@ public:
 		return ChunkIdToIndex.Find(ChunkId);
 	}
 
-	const FIoOffsetAndLength* GetOffsetAndLength(const FIoChunkId& ChunkId) const
-	{
-		if (const int32* Index = ChunkIdToIndex.Find(ChunkId))
-		{
-			return &Toc.ChunkOffsetLengths[*Index];
-		}
-
-		return nullptr;
-	}
-
 	void GetFileNamesToIndex(TArray<FStringView>& OutFileNames) const
 	{
 		OutFileNames.Empty(ChunkIdToFileName.Num());
@@ -663,63 +638,179 @@ public:
 		return ChunkIdToFileName.Find(ChunkId);
 	}
 
-	FIoStoreTocChunkInfo GetTocChunkInfo(const int32 TocEntryIndex, const TMap<int32, FString>* InChunkFileNamesMap) const
+	FIoStoreTocChunkInfo GetTocChunkInfo(int32 TocEntryIndex) const
 	{
-		const FIoStoreTocResource& TocResource = GetTocResource();
-		const FIoStoreTocEntryMeta& Meta = TocResource.ChunkMetas[TocEntryIndex];
-		const FIoOffsetAndLength& OffsetLength = TocResource.ChunkOffsetLengths[TocEntryIndex];
+		FIoStoreTocChunkInfo ChunkInfo = Toc.GetTocChunkInfo(TocEntryIndex);
 
-		const bool bIsContainerCompressed = EnumHasAnyFlags(TocResource.Header.ContainerFlags, EIoContainerFlags::Compressed);
+		ChunkInfo.FileName = FString::Printf(TEXT("<%s>"), *LexToString(ChunkInfo.ChunkType));
+		ChunkInfo.bHasValidFileName = false;
 
-		FIoStoreTocChunkInfo ChunkInfo;
-		ChunkInfo.Id = TocResource.ChunkIds[TocEntryIndex];
-		ChunkInfo.ChunkType = ChunkInfo.Id.GetChunkType();
-		ChunkInfo.Hash = Meta.ChunkHash;
-		ChunkInfo.bIsCompressed = EnumHasAnyFlags(Meta.Flags, FIoStoreTocEntryMetaFlags::Compressed);
-		ChunkInfo.bIsMemoryMapped = EnumHasAnyFlags(Meta.Flags, FIoStoreTocEntryMetaFlags::MemoryMapped);
-		ChunkInfo.bForceUncompressed = bIsContainerCompressed && !EnumHasAnyFlags(Meta.Flags, FIoStoreTocEntryMetaFlags::Compressed);
-		ChunkInfo.Offset = OffsetLength.GetOffset();
-		ChunkInfo.Size = OffsetLength.GetLength();
+		return ChunkInfo;
+	}
+
+private:
+	FIoStoreTocResource Toc;
+	TMap<FIoChunkId, int32> ChunkIdToIndex;
+	TMap<FIoChunkId, FString> ChunkIdToFileName;
+};
 
-		const FString* FindFileName = 0;
-		if (InChunkFileNamesMap &&
-			(FindFileName = InChunkFileNamesMap->Find(TocEntryIndex)) != nullptr)
+class FIoStoreTocReader
+{
+public:
+	FIoStoreTocReader()
+	{
+		FMemory::Memzero(&Toc.Header, sizeof(FIoStoreTocHeader));
+	}
+
+	[[nodiscard]] FIoStatus Read(const TCHAR* TocFilePath, const TMap<FGuid, FAES::FAESKey>& DecryptionKeys)
+	{
+		FIoStatus TocStatus = FIoStoreTocResource::Read(TocFilePath, EIoStoreTocReadOptions::ReadAll, Toc);
+		if (!TocStatus.IsOk())
 		{
-			ChunkInfo.FileName = *FindFileName;
-			ChunkInfo.bHasValidFileName = true;
+			return TocStatus;
 		}
-		else
+
+		ChunkIdToIndex.Empty(Toc.ChunkIds.Num());
+
+		for (int32 ChunkIndex = 0; ChunkIndex < Toc.ChunkIds.Num(); ++ChunkIndex)
 		{
-			ChunkInfo.FileName = FString::Printf(TEXT("<%s>"), *LexToString(ChunkInfo.ChunkType));
-			ChunkInfo.bHasValidFileName = false;
+			ChunkIdToIndex.Add(Toc.ChunkIds[ChunkIndex], ChunkIndex);
 		}
 
-		const uint64 CompressionBlockSize = TocResource.Header.CompressionBlockSize;
-		int32 FirstBlockIndex = int32(ChunkInfo.Offset / CompressionBlockSize);
-		int32 LastBlockIndex = int32((Align(ChunkInfo.Offset + ChunkInfo.Size, CompressionBlockSize) - 1) / CompressionBlockSize);
+		if (EnumHasAnyFlags(Toc.Header.ContainerFlags, EIoContainerFlags::Encrypted))
+		{
+			const FAES::FAESKey* FindKey = DecryptionKeys.Find(Toc.Header.EncryptionKeyGuid);
+			if (!FindKey)
+			{
+				return FIoStatusBuilder(EIoErrorCode::FileOpenFailed) << TEXT("Missing decryption key for IoStore container file '") << TocFilePath << TEXT("'");
+			}
+			DecryptionKey = *FindKey;
+		}
 
-		ChunkInfo.NumCompressedBlocks = LastBlockIndex - FirstBlockIndex + 1;
-		ChunkInfo.OffsetOnDisk = TocResource.CompressionBlocks[FirstBlockIndex].GetOffset();
-		ChunkInfo.CompressedSize = 0;
-		ChunkInfo.PartitionIndex = -1;
-		for (int32 BlockIndex = FirstBlockIndex; BlockIndex <= LastBlockIndex; ++BlockIndex)
+		if (EnumHasAnyFlags(Toc.Header.ContainerFlags, EIoContainerFlags::Indexed) &&
+			Toc.DirectoryIndexBuffer.Num() > 0)
 		{
-			const FIoStoreTocCompressedBlockEntry& CompressionBlock = TocResource.CompressionBlocks[BlockIndex];
-			ChunkInfo.CompressedSize += CompressionBlock.GetCompressedSize();
-			if (ChunkInfo.PartitionIndex < 0)
+			FIoStatus DirectoryIndexStatus = DirectoryIndexReader.Initialize(Toc.DirectoryIndexBuffer, DecryptionKey);
+			if (!DirectoryIndexStatus.IsOk())
 			{
-				ChunkInfo.PartitionIndex = int32(CompressionBlock.GetOffset() / TocResource.Header.PartitionSize);
+				return DirectoryIndexStatus;
 			}
+			DirectoryIndexReader.IterateDirectoryIndex(
+				FIoDirectoryIndexHandle::RootDirectory(),
+				TEXT(""),
+				[this](FStringView Filename, uint32 TocEntryIndex) -> bool
+				{
+					AddFileName(TocEntryIndex, Filename);
+					return true;
+				});
+		}
+
+		return TocStatus;
+	}
+
+	FIoStoreTocResource& GetTocResource()
+	{
+		return Toc;
+	}
+
+	const FIoStoreTocResource& GetTocResource() const
+	{
+		return Toc;
+	}
+
+	const FAES::FAESKey& GetDecryptionKey() const
+	{
+		return DecryptionKey;
+	}
+
+	const FIoDirectoryIndexReader& GetDirectoryIndexReader() const
+	{
+		return DirectoryIndexReader;
+	}
+
+	const int32* GetTocEntryIndex(const FIoChunkId& ChunkId) const
+	{
+		return ChunkIdToIndex.Find(ChunkId);
+	}
+
+	const FIoOffsetAndLength* GetOffsetAndLength(const FIoChunkId& ChunkId) const
+	{
+		if (const int32* Index = ChunkIdToIndex.Find(ChunkId))
+		{
+			return &Toc.ChunkOffsetLengths[*Index];
+		}
+
+		return nullptr;
+	}
+
+	FIoStoreTocChunkInfo GetTocChunkInfo(int32 TocEntryIndex) const
+	{
+		FIoStoreTocChunkInfo ChunkInfo = Toc.GetTocChunkInfo(TocEntryIndex);
+
+		if (const FString* FileName = IndexToFileName.Find(TocEntryIndex); FileName != nullptr)
+		{
+			ChunkInfo.FileName = *FileName;
+			ChunkInfo.bHasValidFileName = true;
+		}
+		else
+		{
+			ChunkInfo.FileName = FString::Printf(TEXT("<%s>"), *LexToString(ChunkInfo.ChunkType));
+			ChunkInfo.bHasValidFileName = false;
 		}
 		return ChunkInfo;
 	}
 
 private:
-	TMap<FIoChunkId, int32> ChunkIdToIndex;
+	void AddFileName(int32 TocEntryIndex, FStringView Filename)
+	{
+		IndexToFileName.Add(TocEntryIndex, FString(Filename));
+	}
+
 	FIoStoreTocResource Toc;
-	TMap<FIoChunkId, FString> ChunkIdToFileName;
+	FIoDirectoryIndexReader DirectoryIndexReader;
+	FAES::FAESKey DecryptionKey;
+	TMap<FIoChunkId, int32> ChunkIdToIndex;
+	TMap<int32, FString> IndexToFileName;
 };
 
+FIoStoreTocChunkInfo FIoStoreTocResource::GetTocChunkInfo(int32 TocEntryIndex) const
+{
+	const FIoStoreTocEntryMeta& Meta = ChunkMetas[TocEntryIndex];
+	const FIoOffsetAndLength& OffsetLength = ChunkOffsetLengths[TocEntryIndex];
+
+	const bool bIsContainerCompressed = EnumHasAnyFlags(Header.ContainerFlags, EIoContainerFlags::Compressed);
+
+	FIoStoreTocChunkInfo ChunkInfo;
+	ChunkInfo.Id = ChunkIds[TocEntryIndex];
+	ChunkInfo.ChunkType = ChunkInfo.Id.GetChunkType();
+	ChunkInfo.Hash = Meta.ChunkHash;
+	ChunkInfo.bHasValidFileName = false;
+	ChunkInfo.bIsCompressed = EnumHasAnyFlags(Meta.Flags, FIoStoreTocEntryMetaFlags::Compressed);
+	ChunkInfo.bIsMemoryMapped = EnumHasAnyFlags(Meta.Flags, FIoStoreTocEntryMetaFlags::MemoryMapped);
+	ChunkInfo.bForceUncompressed = bIsContainerCompressed && !EnumHasAnyFlags(Meta.Flags, FIoStoreTocEntryMetaFlags::Compressed);
+	ChunkInfo.Offset = OffsetLength.GetOffset();
+	ChunkInfo.Size = OffsetLength.GetLength();
+
+	const uint64 CompressionBlockSize = Header.CompressionBlockSize;
+	int32 FirstBlockIndex = int32(ChunkInfo.Offset / CompressionBlockSize);
+	int32 LastBlockIndex = int32((Align(ChunkInfo.Offset + ChunkInfo.Size, CompressionBlockSize) - 1) / CompressionBlockSize);
+
+	ChunkInfo.NumCompressedBlocks = LastBlockIndex - FirstBlockIndex + 1;
+	ChunkInfo.OffsetOnDisk = CompressionBlocks[FirstBlockIndex].GetOffset();
+	ChunkInfo.CompressedSize = 0;
+	ChunkInfo.PartitionIndex = -1;
+	for (int32 BlockIndex = FirstBlockIndex; BlockIndex <= LastBlockIndex; ++BlockIndex)
+	{
+		const FIoStoreTocCompressedBlockEntry& CompressionBlock = CompressionBlocks[BlockIndex];
+		ChunkInfo.CompressedSize += CompressionBlock.GetCompressedSize();
+		if (ChunkInfo.PartitionIndex < 0)
+		{
+			ChunkInfo.PartitionIndex = int32(CompressionBlock.GetOffset() / Header.PartitionSize);
+		}
+	}
+	return ChunkInfo;
+}
+
 class FIoStoreWriter
 	: public IIoStoreWriter
 {
@@ -756,11 +847,11 @@ public:
 
 	void EnumerateChunks(TFunction<bool(FIoStoreTocChunkInfo&&)>&& Callback) const
 	{
-		const FIoStoreTocResource& TocResource = Toc.GetTocResource();
+		const FIoStoreTocResource& TocResource = TocBuilder.GetTocResource();
 
 		for (int32 ChunkIndex = 0; ChunkIndex < TocResource.ChunkIds.Num(); ++ChunkIndex)
 		{
-			FIoStoreTocChunkInfo ChunkInfo = Toc.GetTocChunkInfo(ChunkIndex, 0);
+			FIoStoreTocChunkInfo ChunkInfo = TocBuilder.GetTocChunkInfo(ChunkIndex);
 			if (!Callback(MoveTemp(ChunkInfo)))
 			{
 				break;
@@ -1243,7 +1334,7 @@ public:
 			}
 		}
 
-		FIoStoreTocResource& TocResource = Toc.GetTocResource();
+		FIoStoreTocResource& TocResource = TocBuilder.GetTocResource();
 
 		GeneratePerfectHashes(TocResource, *FPaths::GetBaseFilename(TocFilePath));
 
@@ -1251,7 +1342,7 @@ public:
 		{
 			TRACE_CPUPROFILER_EVENT_SCOPE(BuildIndex);
 			TArray<FStringView> FilesToIndex;
-			Toc.GetFileNamesToIndex(FilesToIndex);
+			TocBuilder.GetFileNamesToIndex(FilesToIndex);
 
 			FString MountPoint = IoDirectoryIndexUtils::GetCommonRootPath(FilesToIndex);
 			FIoDirectoryIndexWriter DirectoryIndexWriter;
@@ -1260,7 +1351,7 @@ public:
 			uint32 TocEntryIndex = 0;
 			for (const FIoChunkId& ChunkId : TocResource.ChunkIds)
 			{
-				const FString* ChunkFileName = Toc.GetFileName(ChunkId);
+				const FString* ChunkFileName = TocBuilder.GetFileName(ChunkId);
 				if (ChunkFileName)
 				{
 					const uint32 FileEntryIndex = DirectoryIndexWriter.AddFile(*ChunkFileName);
@@ -1834,7 +1925,7 @@ private:
 	void WriteEntry(FIoStoreWriteQueueEntry* Entry)
 	{
 		TRACE_CPUPROFILER_EVENT_SCOPE(WriteEntry);
-		const int32* FindExistingIndex = Toc.GetTocEntryIndex(Entry->ChunkId);
+		const int32* FindExistingIndex = TocBuilder.GetTocEntryIndex(Entry->ChunkId);
 		if (FindExistingIndex)
 		{
 			// afaict this should never happen so add a warning. If there's a legit reason for it
@@ -1843,11 +1934,11 @@ private:
 			UE_LOG(LogIoStore, Warning, TEXT("ChunkId was added twice in container %s, %s, file %s hash %s vs %s"), 
 				*FPaths::GetBaseFilename(TocFilePath),
 				*LexToString(Entry->ChunkId), *Entry->Options.FileName,
-				*Toc.GetTocResource().ChunkMetas[*FindExistingIndex].ChunkHash.ToString(),
+				*TocBuilder.GetTocResource().ChunkMetas[*FindExistingIndex].ChunkHash.ToString(),
 				*Entry->ChunkHash.ToString()
 				);
 
-			checkf(Toc.GetTocResource().ChunkMetas[*FindExistingIndex].ChunkHash == Entry->ChunkHash, TEXT("Chunk id has already been added with different content"));
+			checkf(TocBuilder.GetTocResource().ChunkMetas[*FindExistingIndex].ChunkHash == Entry->ChunkHash, TEXT("Chunk id has already been added with different content"));
 			for (FChunkBlock& ChunkBlock : Entry->ChunkBlocks)
 			{
 				WriterContext->FreeCompressionBuffer(ChunkBlock.IoBuffer, Entry->ChunkBlocks.Num());
@@ -1937,12 +2028,12 @@ private:
 		uint64 OffsetInChunk = 0;
 		for (const FChunkBlock& ChunkBlock : Entry->ChunkBlocks)
 		{
-			FIoStoreTocCompressedBlockEntry& BlockEntry = Toc.AddCompressionBlockEntry();
+			FIoStoreTocCompressedBlockEntry& BlockEntry = TocBuilder.AddCompressionBlockEntry();
 			BlockEntry.SetOffset(TargetPartition->Index * WriterSettings.MaxPartitionSize + TargetPartition->Offset + OffsetInChunk);
 			OffsetInChunk += ChunkBlock.Size;
 			BlockEntry.SetCompressedSize(uint32(ChunkBlock.CompressedSize));
 			BlockEntry.SetUncompressedSize(uint32(ChunkBlock.UncompressedSize));
-			BlockEntry.SetCompressionMethodIndex(Toc.AddCompressionMethodEntry(ChunkBlock.CompressionMethod));
+			BlockEntry.SetCompressionMethodIndex(TocBuilder.AddCompressionMethodEntry(ChunkBlock.CompressionMethod));
 
 			// We do this here so that we get the total size of data excluding the encryption alignment
 			TotalEntryCompressedSize += ChunkBlock.CompressedSize;
@@ -1958,25 +2049,25 @@ private:
 
 			if (ContainerSettings.IsSigned())
 			{
-				FSHAHash& Signature = Toc.AddBlockSignatureEntry();
+				FSHAHash& Signature = TocBuilder.AddBlockSignatureEntry();
 				Signature = ChunkBlock.Signature;
 			}
 
 			if (ContainerSettings.IsOnDemand())
 			{
-				FIoStoreTocOnDemandCompressedBlockMeta& Meta = Toc.AddCompressionBlockMetaEntry();
+				FIoStoreTocOnDemandCompressedBlockMeta& Meta = TocBuilder.AddCompressionBlockMetaEntry();
 				Meta.DiskHash = ChunkBlock.DiskHash;
 			}
 		}
 
 		FIoStoreTocOnDemandChunkMeta OnDemandChunkMeta{ Entry->ChunkDiskHash };
 
-		const int32 TocEntryIndex = Toc.AddChunkEntry(Entry->ChunkId, OffsetLength, ChunkMeta, OnDemandChunkMeta);
+		const int32 TocEntryIndex = TocBuilder.AddChunkEntry(Entry->ChunkId, OffsetLength, ChunkMeta, OnDemandChunkMeta);
 		check(TocEntryIndex != INDEX_NONE);
 
 		if (ContainerSettings.IsIndexed() && Entry->Options.FileName.Len() > 0)
 		{
-			Toc.AddToFileIndex(Entry->ChunkId, Entry->Options.FileName);
+			TocBuilder.AddToFileIndex(Entry->ChunkId, Entry->Options.FileName);
 		}
 
 		const uint64 RegionStartOffset = TargetPartition->Offset;
@@ -2025,7 +2116,7 @@ private:
 	FIoStoreWriterContextImpl*	WriterContext = nullptr;
 	FIoContainerSettings		ContainerSettings;
 	FString						TocFilePath;
-	FIoStoreToc					Toc;
+	FIoStoreTocBuilder			TocBuilder;
 	TArray<uint8>				PaddingBuffer;
 	TArray<FPartition>			Partitions;
 	TArray<FIoStoreWriteQueueEntry*> Entries;
@@ -2367,14 +2458,13 @@ public:
 		TocFilePath.Append(InContainerPath);
 		TocFilePath.Append(TEXT(".utoc"));
 
-		FIoStoreTocResource& TocResource = Toc.GetTocResource();
-		FIoStatus TocStatus = FIoStoreTocResource::Read(*TocFilePath, EIoStoreTocReadOptions::ReadAll, TocResource);
+		FIoStatus TocStatus = TocReader.Read(*TocFilePath, InDecryptionKeys);
 		if (!TocStatus.IsOk())
 		{
 			return TocStatus;
 		}
 
-		Toc.Initialize();
+		FIoStoreTocResource& TocResource = TocReader.GetTocResource();
 
 		IPlatformFile& Ipf = FPlatformFileManager::Get().GetPlatformFile();
 		ContainerFileAccessors.Reserve(TocResource.Header.PartitionCount);
@@ -2395,55 +2485,27 @@ public:
 			}
 		}
 
-		if (EnumHasAnyFlags(TocResource.Header.ContainerFlags, EIoContainerFlags::Encrypted))
-		{
-			const FAES::FAESKey* FindKey = InDecryptionKeys.Find(TocResource.Header.EncryptionKeyGuid);
-			if (!FindKey)
-			{
-				return FIoStatusBuilder(EIoErrorCode::FileOpenFailed) << TEXT("Missing decryption key for IoStore container file '") << *TocFilePath << TEXT("'");
-			}
-			DecryptionKey = *FindKey;
-		}
-
-		if (EnumHasAnyFlags(TocResource.Header.ContainerFlags, EIoContainerFlags::Indexed) &&
-			TocResource.DirectoryIndexBuffer.Num() > 0)
-		{
-			FIoStatus DirectoryIndexStatus = DirectoryIndexReader.Initialize(TocResource.DirectoryIndexBuffer, DecryptionKey);
-			if (!DirectoryIndexStatus.IsOk())
-			{
-				return DirectoryIndexStatus;
-			}
-			DirectoryIndexReader.IterateDirectoryIndex(
-				FIoDirectoryIndexHandle::RootDirectory(),
-				TEXT(""),
-				[this](FStringView Filename, uint32 TocEntryIndex) -> bool
-				{
-					ChunkFileNamesMap.Add(TocEntryIndex, FString(Filename));
-					return true;
-				});
-		}
-
 		return FIoStatus::Ok;
 	}
 
 	FIoContainerId GetContainerId() const
 	{
-		return Toc.GetTocResource().Header.ContainerId;
+		return TocReader.GetTocResource().Header.ContainerId;
 	}
 
 	uint32 GetVersion() const
 	{
-		return Toc.GetTocResource().Header.Version;
+		return TocReader.GetTocResource().Header.Version;
 	}
 
 	EIoContainerFlags GetContainerFlags() const
 	{
-		return Toc.GetTocResource().Header.ContainerFlags;
+		return TocReader.GetTocResource().Header.ContainerFlags;
 	}
 
 	FGuid GetEncryptionKeyGuid() const
 	{
-		return Toc.GetTocResource().Header.EncryptionKeyGuid;
+		return TocReader.GetTocResource().Header.EncryptionKeyGuid;
 	}
 
 	FString GetContainerName() const
@@ -2453,16 +2515,16 @@ public:
 
 	int32 GetChunkCount() const 
 	{
-		return Toc.GetTocResource().ChunkIds.Num();
+		return TocReader.GetTocResource().ChunkIds.Num();
 	}
 
 	void EnumerateChunks(TFunction<bool(FIoStoreTocChunkInfo&&)>&& Callback) const
 	{
-		const FIoStoreTocResource& TocResource = Toc.GetTocResource();
+		const FIoStoreTocResource& TocResource = TocReader.GetTocResource();
 
 		for (int32 ChunkIndex = 0; ChunkIndex < TocResource.ChunkIds.Num(); ++ChunkIndex)
 		{
-			FIoStoreTocChunkInfo ChunkInfo = Toc.GetTocChunkInfo(ChunkIndex, &ChunkFileNamesMap);
+			FIoStoreTocChunkInfo ChunkInfo = TocReader.GetTocChunkInfo(ChunkIndex);
 			if (!Callback(MoveTemp(ChunkInfo)))
 			{
 				break;
@@ -2472,10 +2534,10 @@ public:
 
 	TIoStatusOr<FIoStoreTocChunkInfo> GetChunkInfo(const FIoChunkId& ChunkId) const
 	{
-		const int32* TocEntryIndex = Toc.GetTocEntryIndex(ChunkId);
+		const int32* TocEntryIndex = TocReader.GetTocEntryIndex(ChunkId);
 		if (TocEntryIndex)
 		{
-			return Toc.GetTocChunkInfo(*TocEntryIndex, &ChunkFileNamesMap);
+			return TocReader.GetTocChunkInfo(*TocEntryIndex);
 		}
 		else
 		{
@@ -2485,11 +2547,11 @@ public:
 
 	TIoStatusOr<FIoStoreTocChunkInfo> GetChunkInfo(const uint32 TocEntryIndex) const
 	{
-		const FIoStoreTocResource& TocResource = Toc.GetTocResource();
+		const FIoStoreTocResource& TocResource = TocReader.GetTocResource();
 
 		if (TocEntryIndex < uint32(TocResource.ChunkIds.Num()))
 		{
-			return Toc.GetTocChunkInfo(TocEntryIndex, &ChunkFileNamesMap);
+			return TocReader.GetTocChunkInfo(TocEntryIndex);
 		}
 		else
 		{
@@ -2499,11 +2561,11 @@ public:
 
 	TIoStatusOr<FIoStoreCompressedChunkInfo> GetChunkCompressedInfo(const FIoChunkId& ChunkId) const
 	{
-		const int32* TocEntryIndex = Toc.GetTocEntryIndex(ChunkId);
+		const int32* TocEntryIndex = TocReader.GetTocEntryIndex(ChunkId);
 		if (TocEntryIndex != nullptr)
 		{
 			// Find where in the virtual file the chunk exists.
-			const FIoOffsetAndLength* OffsetAndLength = Toc.GetOffsetAndLength(ChunkId);
+			const FIoOffsetAndLength* OffsetAndLength = TocReader.GetOffsetAndLength(ChunkId);
 			if (!OffsetAndLength)
 			{
 				return FIoStatus(EIoErrorCode::NotFound, TEXT("Unknown chunk ID"));
@@ -2513,7 +2575,7 @@ public:
 			const uint64 ResolvedSize = OffsetAndLength->GetLength();
 
 			// Find what compressed blocks this read straddles.
-			const FIoStoreTocResource& TocResource = Toc.GetTocResource();
+			const FIoStoreTocResource& TocResource = TocReader.GetTocResource();
 			const uint64 CompressionBlockSize = TocResource.Header.CompressionBlockSize;
 
 			const int32 FirstBlockIndex = int32(ResolvedOffset / CompressionBlockSize);
@@ -2578,7 +2640,7 @@ public:
 			std::atomic_bool bUncompressFailed { false };
 		};
 
-		const FIoOffsetAndLength* OffsetAndLength = Toc.GetOffsetAndLength(ChunkId);
+		const FIoOffsetAndLength* OffsetAndLength = TocReader.GetOffsetAndLength(ChunkId);
 		if (!OffsetAndLength )
 		{
 			// Currently there's no way to make a task with a valid result that just emplaces
@@ -2592,7 +2654,7 @@ public:
 		const uint64 RequestedOffset = Options.GetOffset();
 		const uint64 ResolvedOffset = OffsetAndLength->GetOffset() + RequestedOffset;
 		const uint64 ResolvedSize = RequestedOffset <= OffsetAndLength->GetLength() ? FMath::Min(Options.GetSize(), OffsetAndLength->GetLength() - RequestedOffset) : 0;
-		const FIoStoreTocResource& TocResource = Toc.GetTocResource();
+		const FIoStoreTocResource& TocResource = TocReader.GetTocResource();
 		const uint64 CompressionBlockSize = TocResource.Header.CompressionBlockSize;
 		const int32 FirstBlockIndex = int32(ResolvedOffset / CompressionBlockSize);
 		const int32 LastBlockIndex = int32((Align(ResolvedOffset + ResolvedSize, CompressionBlockSize) - 1) / CompressionBlockSize);
@@ -2638,7 +2700,7 @@ public:
 					{
 						uint8* CompressedSource = State->CompressedBuffer.GetData() + CompressedSourceOffset;
 						uint8* UncompressedDestination = State->UncompressedBuffer->Data() + UncompressedDestinationOffset;
-						const FIoStoreTocResource& TocResource = Toc.GetTocResource();
+						const FIoStoreTocResource& TocResource = TocReader.GetTocResource();
 						const FIoStoreTocCompressedBlockEntry& CompressionBlock = TocResource.CompressionBlocks[BlockIndex];
 						const uint32 RawSize = Align(CompressionBlock.GetCompressedSize(), FAES::AESBlockSize);
 						const uint32 UncompressedSize = CompressionBlock.GetUncompressedSize();
@@ -2647,7 +2709,7 @@ public:
 						{
 							TRACE_CPUPROFILER_EVENT_SCOPE(Decrypt);
 							check(CompressedSource + RawSize <= State->CompressedBuffer.GetData() + State->CompressedSize);
-							FAES::DecryptData(CompressedSource, RawSize, DecryptionKey);
+							FAES::DecryptData(CompressedSource, RawSize, TocReader.GetDecryptionKey());
 						}
 						if (CompressionMethod.IsNone())
 						{
@@ -2719,7 +2781,7 @@ public:
 	{
 		TRACE_CPUPROFILER_EVENT_SCOPE(ReadChunk);
 
-		const FIoOffsetAndLength* OffsetAndLength = Toc.GetOffsetAndLength(ChunkId);
+		const FIoOffsetAndLength* OffsetAndLength = TocReader.GetOffsetAndLength(ChunkId);
 		if (!OffsetAndLength)
 		{
 			return FIoStatus(EIoErrorCode::NotFound, TEXT("Unknown chunk ID"));
@@ -2733,7 +2795,7 @@ public:
 			ResolvedSize = FMath::Min(Options.GetSize(), OffsetAndLength->GetLength() - RequestedOffset);
 		}
 
-		const FIoStoreTocResource& TocResource = Toc.GetTocResource();
+		const FIoStoreTocResource& TocResource = TocReader.GetTocResource();
 		const uint64 CompressionBlockSize = TocResource.Header.CompressionBlockSize;
 		FIoBuffer UncompressedBuffer(ResolvedSize);
 		if (ResolvedSize == 0)
@@ -2812,7 +2874,7 @@ public:
 			if (EnumHasAnyFlags(TocResource.Header.ContainerFlags, EIoContainerFlags::Encrypted))
 			{
 				TRACE_CPUPROFILER_EVENT_SCOPE(Decrypt);
-				FAES::DecryptData(CompressedBuffers[OurBufferIndex].GetData(), RawSize, DecryptionKey);
+				FAES::DecryptData(CompressedBuffers[OurBufferIndex].GetData(), RawSize, TocReader.GetDecryptionKey());
 			}
 
 			FName CompressionMethod = TocResource.CompressionMethods[CompressionBlock.GetCompressionMethodIndex()];
@@ -2863,7 +2925,7 @@ public:
 		TRACE_CPUPROFILER_EVENT_SCOPE(ReadChunkCompressed);
 
 		// Find where in the virtual file the chunk exists.
-		const FIoOffsetAndLength* OffsetAndLength = Toc.GetOffsetAndLength(ChunkId);
+		const FIoOffsetAndLength* OffsetAndLength = TocReader.GetOffsetAndLength(ChunkId);
 		if (!OffsetAndLength)
 		{
 			return FIoStatus(EIoErrorCode::NotFound, TEXT("Unknown chunk ID"));
@@ -2879,7 +2941,7 @@ public:
 		}
 
 		// Find what compressed blocks this read straddles.
-		const FIoStoreTocResource& TocResource = Toc.GetTocResource();
+		const FIoStoreTocResource& TocResource = TocReader.GetTocResource();
 		const uint64 CompressionBlockSize = TocResource.Header.CompressionBlockSize;
 		int32 FirstBlockIndex = int32(ResolvedOffset / CompressionBlockSize);
 		int32 LastBlockIndex = int32((Align(ResolvedOffset + ResolvedSize, CompressionBlockSize) - 1) / CompressionBlockSize);
@@ -2950,7 +3012,7 @@ public:
 			{
 				FIoStoreCompressedBlockInfo& OutputBlock = Result.Blocks[BlockIndex - FirstBlockIndex];
 				uint8* Buffer = OutputBuffer + OutputBlock.OffsetInBuffer;
-				FAES::DecryptData(Buffer, OutputBlock.AlignedSize, DecryptionKey);
+				FAES::DecryptData(Buffer, OutputBlock.AlignedSize, TocReader.GetDecryptionKey());
 			}
 		}
 		return Result;
@@ -2958,12 +3020,12 @@ public:
 
 	const FIoDirectoryIndexReader& GetDirectoryIndexReader() const
 	{
-		return DirectoryIndexReader;
+		return TocReader.GetDirectoryIndexReader();
 	}
 
 	bool TocChunkContainsBlockIndex(const int32 TocEntryIndex, const int32 BlockIndex) const
 	{
-		const FIoStoreTocResource& TocResource = Toc.GetTocResource();
+		const FIoStoreTocResource& TocResource = TocReader.GetTocResource();
 		const FIoOffsetAndLength& OffsetLength = TocResource.ChunkOffsetLengths[TocEntryIndex];
 
 		const uint64 CompressionBlockSize = TocResource.Header.CompressionBlockSize;
@@ -2975,24 +3037,24 @@ public:
 
 	uint32 GetCompressionBlockSize() const
 	{
-		return Toc.GetTocResource().Header.CompressionBlockSize;
+		return TocReader.GetTocResource().Header.CompressionBlockSize;
 	}
 	
 	const TArray<FName>& GetCompressionMethods() const
 	{
-		return Toc.GetTocResource().CompressionMethods;
+		return TocReader.GetTocResource().CompressionMethods;
 	}
 
 	bool EnumerateCompressedBlocksForChunk(const FIoChunkId& ChunkId, TFunction<bool(const FIoStoreTocCompressedBlockInfo&)>&& Callback) const
 	{
-		const FIoOffsetAndLength* OffsetAndLength = Toc.GetOffsetAndLength(ChunkId);
+		const FIoOffsetAndLength* OffsetAndLength = TocReader.GetOffsetAndLength(ChunkId);
 		if (!OffsetAndLength)
 		{
 			return false;
 		}
 
 		// Find what compressed blocks this chunk straddles.
-		const FIoStoreTocResource& TocResource = Toc.GetTocResource();
+		const FIoStoreTocResource& TocResource = TocReader.GetTocResource();
 		const uint64 CompressionBlockSize = TocResource.Header.CompressionBlockSize;
 		int32 FirstBlockIndex = int32(OffsetAndLength->GetOffset() / CompressionBlockSize);
 		int32 LastBlockIndex = int32((Align(OffsetAndLength->GetOffset() + OffsetAndLength->GetLength(), CompressionBlockSize) - 1) / CompressionBlockSize);
@@ -3017,7 +3079,7 @@ public:
 
 	void EnumerateCompressedBlocks(TFunction<bool(const FIoStoreTocCompressedBlockInfo&)>&& Callback) const
 	{
-		const FIoStoreTocResource& TocResource = Toc.GetTocResource();
+		const FIoStoreTocResource& TocResource = TocReader.GetTocResource();
 
 		for (int32 BlockIndex = 0; BlockIndex < TocResource.CompressionBlocks.Num(); ++BlockIndex)
 		{
@@ -3040,7 +3102,7 @@ public:
 	{
 		TStringBuilder<256> Sb;
 
-		for (uint32 PartitionIndex = 0; PartitionIndex < Toc.GetTocResource().Header.PartitionCount; ++PartitionIndex)
+		for (uint32 PartitionIndex = 0; PartitionIndex < TocReader.GetTocResource().Header.PartitionCount; ++PartitionIndex)
 		{
 			Sb.Reset();
 			Sb.Append(ContainerPath);
@@ -3056,12 +3118,9 @@ public:
 private:
 
 
-	FIoStoreToc Toc;
-	FAES::FAESKey DecryptionKey;
+	FIoStoreTocReader TocReader;
 	TArray<TUniquePtr<FContainerFileAccess>> ContainerFileAccessors;
 	FString ContainerPath;
-	FIoDirectoryIndexReader DirectoryIndexReader;
-	TMap<int32, FString> ChunkFileNamesMap;
 };
 
 FIoStoreReader::FIoStoreReader()
