nop comments only about slow hash functions being texture import bottlenecks

diff --git a/Engine/Source/Runtime/Core/Private/Compression/CompressedBuffer.cpp b/Engine/Source/Runtime/Core/Private/Compression/CompressedBuffer.cpp
--- a/Engine/Source/Runtime/Core/Private/Compression/CompressedBuffer.cpp
+++ b/Engine/Source/Runtime/Core/Private/Compression/CompressedBuffer.cpp
@@ -271,7 +271,7 @@ private:
 
 FCompositeBuffer FBlockEncoder::Compress(const FCompositeBuffer& RawData, const uint64 BlockSize) const
 {
-	TRACE_CPUPROFILER_EVENT_SCOPE(FBlockEncoder::Compress);
+	TRACE_CPUPROFILER_EVENT_SCOPE(FBlockEncoder.Compress);
 
 	checkf(FMath::IsPowerOfTwo(BlockSize) && BlockSize <= MAX_uint32,
 		TEXT("BlockSize must be a 32-bit power of two but was %" UINT64_FMT "."), BlockSize);
@@ -381,7 +381,9 @@ FCompositeBuffer FBlockEncoder::Compress(const FCompositeBuffer& RawData, const
 			if (BlockContexts.Num() > 1)
 			{
 				// Hash the raw data in parallel with encoding the blocks when there are multiple threads.
-				TRACE_CPUPROFILER_EVENT_SCOPE(FBlockEncoder::Compress::RawHash);
+				// @todo : on many-core systems like ThreadRippers this is the bottleneck (slower than compression)
+				//	use FXxHash64::HashBufferChunked instead.
+				TRACE_CPUPROFILER_EVENT_SCOPE(FBlockEncoder.RawHash);
 				RawHash.Update(RawData);
 			}
 		});
diff --git a/Engine/Source/Runtime/Interchange/Core/Public/InterchangeSourceData.h b/Engine/Source/Runtime/Interchange/Core/Public/InterchangeSourceData.h
--- a/Engine/Source/Runtime/Interchange/Core/Public/InterchangeSourceData.h
+++ b/Engine/Source/Runtime/Interchange/Core/Public/InterchangeSourceData.h
@@ -87,6 +87,8 @@ private:
 		FileContentHashCache.Reset();
 		if (!Filename.IsEmpty())
 		{
+			// @todo : this is slow and not multi-threaded.  Use FXxHash64::HashBufferChunked instead?
+
 			FileContentHashCache = FMD5Hash::HashFile(*Filename);
 		}
 	}
