Add cvar MessageLog.InfoMessagesAtDisplayVerbosity to let MessageLog emit info messages at Display level.

diff --git a/Engine/Source/Runtime/Core/Private/Logging/MessageLog.cpp b/Engine/Source/Runtime/Core/Private/Logging/MessageLog.cpp
--- a/Engine/Source/Runtime/Core/Private/Logging/MessageLog.cpp
+++ b/Engine/Source/Runtime/Core/Private/Logging/MessageLog.cpp
@@ -1,6 +1,8 @@
 // Copyright Epic Games, Inc. All Rights Reserved.
 
 #include "Logging/MessageLog.h"
+
+#include "HAL/ConsoleManager.h"
 #include "Internationalization/Internationalization.h"
 #include "Logging/IMessageLog.h"
 
@@ -14,6 +16,13 @@ LLM_DEFINE_TAG(EngineMisc_MessageLog);
 
 namespace UE::MessageLog::Private
 {
+	bool LogInfoMessagesAtDisplayVerbosity = 0;
+	FAutoConsoleVariableRef CVarLogInfoMessagesAtDisplayVerbosity(
+		TEXT("MessageLog.InfoMessagesAtDisplayVerbosity"),
+		LogInfoMessagesAtDisplayVerbosity,
+		TEXT("MessageLog messages at level 'info' will be mirrored to the log at 'Display' verbosity rather than 'Log' verbosity."),
+		ECVF_Default);
+	
 	using FOverrideStackMap = TSortedMap<FName, TArray<const FMessageLogScopedOverride*>, FDefaultAllocator, FNameFastLess>;
 
 	/** Get the internal stack of FMessageLogScopedOverride's */
@@ -343,6 +352,7 @@ ELogVerbosity::Type FMessageLog::GetLogVerbosity( EMessageSeverity::Type InSever
 	case EMessageSeverity::Warning:
 		return ELogVerbosity::Warning;
 	case EMessageSeverity::Info:
+		return UE::MessageLog::Private::LogInfoMessagesAtDisplayVerbosity ? ELogVerbosity::Display : ELogVerbosity::Log;
 	default:
 		return ELogVerbosity::Log;
 	}
