Reenable text for GPU scripts being stored in the DDC as it is currently used during the compilation of the GPUComputeScript

diff --git a/Engine/Plugins/FX/Niagara/Source/Niagara/Classes/NiagaraScript.h b/Engine/Plugins/FX/Niagara/Source/Niagara/Classes/NiagaraScript.h
--- a/Engine/Plugins/FX/Niagara/Source/Niagara/Classes/NiagaraScript.h
+++ b/Engine/Plugins/FX/Niagara/Source/Niagara/Classes/NiagaraScript.h
@@ -525,7 +525,8 @@ public:
 	UPROPERTY(Transient)
 	FString LastHlslTranslation;
 
-	UPROPERTY(Transient)
+	// Note that this is currently needed to be non-transient because of how we kick off compilation of GPUComputeScripts
+	UPROPERTY()
 	FString LastHlslTranslationGPU;
 
 	UPROPERTY(Transient)
diff --git a/Engine/Source/Runtime/Core/Private/UObject/FortniteReleaseBranchCustomObjectVersion.cpp b/Engine/Source/Runtime/Core/Private/UObject/FortniteReleaseBranchCustomObjectVersion.cpp
--- a/Engine/Source/Runtime/Core/Private/UObject/FortniteReleaseBranchCustomObjectVersion.cpp
+++ b/Engine/Source/Runtime/Core/Private/UObject/FortniteReleaseBranchCustomObjectVersion.cpp
@@ -6,7 +6,7 @@ TMap<FGuid, FGuid> FFortniteReleaseBranchCustomObjectVersion::GetSystemGuids()
 	TMap<FGuid, FGuid> SystemGuids;
 	const FDevSystemGuids& DevGuids = FDevSystemGuids::Get();
 
-	SystemGuids.Add(DevGuids.Niagara_LatestScriptCompileVersion, FGuid("BF16049E73584021A385C4918F5C85BE"));
+	SystemGuids.Add(DevGuids.Niagara_LatestScriptCompileVersion, FGuid("5F589BAA37394A51BF718ECB10D73C19"));
 	SystemGuids.Add(DevGuids.SkeletalMeshDerivedDataVersion, FGuid("5705956EC7134274A5A1BD3FC74DB3A9"));
 	return SystemGuids;
 }
