Make it so we can turn off AutoRTFM per-island by setting the metadata to 0. Also added a force enable mode to the CVar for our testing.

diff --git a/Engine/Source/Runtime/Core/Private/AutoRTFM/API.cpp b/Engine/Source/Runtime/Core/Private/AutoRTFM/API.cpp
--- a/Engine/Source/Runtime/Core/Private/AutoRTFM/API.cpp
+++ b/Engine/Source/Runtime/Core/Private/AutoRTFM/API.cpp
@@ -12,7 +12,7 @@ namespace
 
 	void UpdateAutoRTFMRuntimeCrashData()
 	{
-		FGenericCrashContext::SetGameData(TEXT("IsAutoRTFMRuntimeEnabled"), GAutoRTFMRuntimeEnabled == AutoRTFM::EAutoRTFMEnabledState::AutoRTFM_Enabled ? TEXT("true") : TEXT("false"));
+		FGenericCrashContext::SetGameData(TEXT("IsAutoRTFMRuntimeEnabled"), AutoRTFM::ForTheRuntime::IsAutoRTFMRuntimeEnabled() ? TEXT("true") : TEXT("false"));
 	}
 }
 
@@ -39,10 +39,15 @@ namespace AutoRTFM
 		{
 			// #noop if AutoRTFM is not compiled in, as GAutoRTFMRuntimeEnabled is a static constexpr when no AutoRTFM compiled in
 #if UE_AUTORTFM
-		// If we have ForcedDisabled AutoRTFM from the CVar we will not change it from other means
-			if (GAutoRTFMRuntimeEnabled == EAutoRTFMEnabledState::AutoRTFM_ForcedDisabled)
+			switch (GAutoRTFMRuntimeEnabled)
 			{
-				UE_LOG(LogCore, Log, TEXT("Enabling AutoRTFM is disabled due to GAutoRTFMRuntimeEnabled set to forced disabled"));
+			default:
+				break;
+			case EAutoRTFMEnabledState::AutoRTFM_ForcedDisabled:
+				UE_LOG(LogCore, Log, TEXT("Ignoring changing AutoRTFM runtime state due to GAutoRTFMRuntimeEnabled being set to forced disabled."));
+				return false;
+			case EAutoRTFMEnabledState::AutoRTFM_ForcedEnabled:
+				UE_LOG(LogCore, Log, TEXT("Ignoring changing AutoRTFM runtime state due to GAutoRTFMRuntimeEnabled being set to forced enabled."));
 				return false;
 			}
 
@@ -58,7 +63,14 @@ namespace AutoRTFM
 
 		bool IsAutoRTFMRuntimeEnabled()
 		{
-			return GAutoRTFMRuntimeEnabled == EAutoRTFMEnabledState::AutoRTFM_Enabled;
+			switch (GAutoRTFMRuntimeEnabled)
+			{
+			default:
+				return false;
+			case EAutoRTFMEnabledState::AutoRTFM_Enabled:
+			case EAutoRTFMEnabledState::AutoRTFM_ForcedEnabled:
+				return true;
+			}
 		}
 	}
 }
diff --git a/Engine/Source/Runtime/Core/Public/AutoRTFM/AutoRTFM.h b/Engine/Source/Runtime/Core/Public/AutoRTFM/AutoRTFM.h
--- a/Engine/Source/Runtime/Core/Public/AutoRTFM/AutoRTFM.h
+++ b/Engine/Source/Runtime/Core/Public/AutoRTFM/AutoRTFM.h
@@ -519,13 +519,16 @@ namespace ForTheRuntime
 	enum EAutoRTFMEnabledState
 	{
 		// Disable AutoRTFM.
-		AutoRTFM_Disabled,
+		AutoRTFM_Disabled = 0,
 
 		// Enable AutoRTFM.
 		AutoRTFM_Enabled,
 
 		// Force disable AutoRTFM - once set the AutoRTFM runtime cannot be re-enabled.
 		AutoRTFM_ForcedDisabled,
+
+		// Force enable AutoRTFM - once set the AutoRTFM runtime cannot be re-enabled.
+		AutoRTFM_ForcedEnabled,
 	};
 
 	// Set whether the AutoRTFM runtime is enabled or disabled.
