- Pass down the StorageServerConnection if a request needs a staging buffer to decode into it, since the memory used can be uncached or write-combined, usually GPU resource

diff --git a/Engine/Source/Runtime/Core/Public/GenericPlatform/GenericPlatformFile.h b/Engine/Source/Runtime/Core/Public/GenericPlatform/GenericPlatformFile.h
--- a/Engine/Source/Runtime/Core/Public/GenericPlatform/GenericPlatformFile.h
+++ b/Engine/Source/Runtime/Core/Public/GenericPlatform/GenericPlatformFile.h
@@ -31,8 +31,10 @@ enum EAsyncIOPriorityAndFlags
 	AIOP_PRIORITY_MASK = 0x000000ff,
 
 	// Flags - combine with priorities if needed
-	AIOP_FLAG_PRECACHE	=	0x00000100,
-	AIOP_FLAG_DONTCACHE	=	0x00000200,
+	AIOP_FLAG_PRECACHE			= 0x00000100,
+	AIOP_FLAG_DONTCACHE			= 0x00000200,
+	// Use it to specify memory where read from cpu can be slow
+	AIOP_FLAG_HW_TARGET_MEMORY	= 0x00000400,	
 
 	// Priorities
 	AIOP_MIN = 0,
diff --git a/Engine/Source/Runtime/Core/Public/IO/IoDispatcher.h b/Engine/Source/Runtime/Core/Public/IO/IoDispatcher.h
--- a/Engine/Source/Runtime/Core/Public/IO/IoDispatcher.h
+++ b/Engine/Source/Runtime/Core/Public/IO/IoDispatcher.h
@@ -150,6 +150,16 @@ private:
 
 //////////////////////////////////////////////////////////////////////////
 
+enum class EIoReadOptionsFlags : uint32
+{
+	None = 0,
+	/**
+	 * Use this flag to inform the decompressor that the memory is uncached or write-combined and therefore the usage of staging might be needed if reading directly from the original memory
+	 */
+	HardwareTargetBuffer = 1 << 0,
+};
+ENUM_CLASS_FLAGS(EIoReadOptionsFlags);
+
 class FIoReadOptions
 {
 public:
@@ -166,6 +176,13 @@ public:
 		, TargetVa(InTargetVa)
 	{ }
 
+	FIoReadOptions(uint64 InOffset, uint64 InSize, void* InTargetVa, EIoReadOptionsFlags InFlags)
+		: RequestedOffset(InOffset)
+		, RequestedSize(InSize)
+		, TargetVa(InTargetVa)
+		, Flags(InFlags)
+	{ }
+
 	~FIoReadOptions() = default;
 
 	void SetRange(uint64 Offset, uint64 Size)
@@ -179,6 +196,11 @@ public:
 		TargetVa = InTargetVa;
 	}
 
+	void SetFlags(EIoReadOptionsFlags InValue)
+	{
+		Flags = InValue;
+	}
+
 	uint64 GetOffset() const
 	{
 		return RequestedOffset;
@@ -194,10 +216,16 @@ public:
 		return TargetVa;
 	}
 
+	EIoReadOptionsFlags GetFlags() const
+	{
+		return Flags;
+	}
+
 private:
 	uint64	RequestedOffset = 0;
 	uint64	RequestedSize = ~uint64(0);
 	void* TargetVa = nullptr;
+	EIoReadOptionsFlags Flags = EIoReadOptionsFlags::None;
 };
 
 //////////////////////////////////////////////////////////////////////////
diff --git a/Engine/Source/Runtime/CoreUObject/Private/Serialization/BulkDataStreaming.cpp b/Engine/Source/Runtime/CoreUObject/Private/Serialization/BulkDataStreaming.cpp
--- a/Engine/Source/Runtime/CoreUObject/Private/Serialization/BulkDataStreaming.cpp
+++ b/Engine/Source/Runtime/CoreUObject/Private/Serialization/BulkDataStreaming.cpp
@@ -1092,7 +1092,7 @@ FBulkDataBatchRequest::FBatchBuilder& FBulkDataBatchRequest::FBatchBuilder::Read
 
 	GetBatch().Read(
 		BulkData.BulkChunkId,
-		FIoReadOptions(ReadOffset, ReadSize, Dst.GetData()),
+		FIoReadOptions(ReadOffset, ReadSize, Dst.GetData(), (Priority & AIOP_FLAG_HW_TARGET_MEMORY) ? EIoReadOptionsFlags::HardwareTargetBuffer : EIoReadOptionsFlags::None),
 		Priority,
 		FIoReadCallback(),
 		OutRequest);
diff --git a/Engine/Source/Runtime/Engine/Private/Streaming/Texture2DStreamIn_IO.cpp b/Engine/Source/Runtime/Engine/Private/Streaming/Texture2DStreamIn_IO.cpp
--- a/Engine/Source/Runtime/Engine/Private/Streaming/Texture2DStreamIn_IO.cpp
+++ b/Engine/Source/Runtime/Engine/Private/Streaming/Texture2DStreamIn_IO.cpp
@@ -108,7 +108,7 @@ void FTexture2DStreamIn_IO::SetIORequests(const FContext& Context)
 				}
 
 				FIoBuffer Dst(FIoBuffer::Wrap, MipData[MipIndex].Data, BulkDataSize);
-				Batch.Read(MipMap.BulkData, 0, BulkDataSize, Priority | AIOP_FLAG_DONTCACHE, Dst);
+				Batch.Read(MipMap.BulkData, 0, BulkDataSize, Priority | AIOP_FLAG_DONTCACHE | AIOP_FLAG_HW_TARGET_MEMORY, Dst);
 			}
 			else // Bulk data size can only be 0 when not available, in which case, we need to recache the file state.
 			{
@@ -197,7 +197,7 @@ void FTexture2DStreamIn_IO::SetIORequests(const FContext& Context)
 			IORequests[MipIndex] = MipMap.BulkData.CreateStreamingRequest(
 				0,
 				BulkDataSize,
-				Priority | AIOP_FLAG_DONTCACHE,
+				Priority | AIOP_FLAG_DONTCACHE | AIOP_FLAG_HW_TARGET_MEMORY,
 				&AsyncFileCallBack,
 				(uint8*)MipData[MipIndex].Data);
 		}
diff --git a/Engine/Source/Runtime/StorageServerClient/Private/StorageServerConnection.cpp b/Engine/Source/Runtime/StorageServerClient/Private/StorageServerConnection.cpp
--- a/Engine/Source/Runtime/StorageServerClient/Private/StorageServerConnection.cpp
+++ b/Engine/Source/Runtime/StorageServerClient/Private/StorageServerConnection.cpp
@@ -344,7 +344,7 @@ void FStorageServerResponse::Serialize(void* V, int64 Length)
 	}
 }
 
-int64 FStorageServerResponse::SerializeChunk(FStorageServerSerializationContext& Context, FIoBuffer& OutChunk, void* TargetVa, uint64 RawOffset, uint64 RawSize)
+int64 FStorageServerResponse::SerializeChunk(FStorageServerSerializationContext& Context, FIoBuffer& OutChunk, void* TargetVa, uint64 RawOffset, uint64 RawSize, bool bHardwareTargetBuffer)
 {
 	if (ContentLength == 0)
 	{
@@ -377,7 +377,7 @@ int64 FStorageServerResponse::SerializeChunk(FStorageServerSerializationContext&
 			const uint64 ChunkSize = FMath::Min(Compressed.GetRawSize(), RawSize);
 			OutChunk = TargetVa ? FIoBuffer(FIoBuffer::Wrap, TargetVa, ChunkSize) : FIoBuffer(ChunkSize);
 			const uint64 CompressedOffset = GetCompressedOffset(Compressed, RawOffset);
-			if (Context.Decoder.TryDecompressTo(OutChunk.GetMutableView(), CompressedOffset))
+			if (Context.Decoder.TryDecompressTo(OutChunk.GetMutableView(), CompressedOffset, bHardwareTargetBuffer ? ECompressedBufferDecompressFlags::IntermediateBuffer :ECompressedBufferDecompressFlags::None ))
 			{
 				return ChunkSize;
 			}
diff --git a/Engine/Source/Runtime/StorageServerClient/Private/StorageServerIoDispatcherBackend.cpp b/Engine/Source/Runtime/StorageServerClient/Private/StorageServerIoDispatcherBackend.cpp
--- a/Engine/Source/Runtime/StorageServerClient/Private/StorageServerIoDispatcherBackend.cpp
+++ b/Engine/Source/Runtime/StorageServerClient/Private/StorageServerIoDispatcherBackend.cpp
@@ -342,7 +342,8 @@ void FStorageServerIoDispatcherBackend::FBatch::DoThreadedWork()
 			{
 				TRACE_CPUPROFILER_EVENT_SCOPE(SerializeResponse);
 				FIoBuffer Chunk;
-				if (Response.SerializeChunk(*SerializationContext, Chunk, Request->Options.GetTargetVa(), Request->Options.GetOffset(), Request->Options.GetSize()))
+				if (Response.SerializeChunk(*SerializationContext, Chunk, Request->Options.GetTargetVa(), Request->Options.GetOffset(), Request->Options.GetSize(), 
+					EnumHasAnyFlags(Request->Options.GetFlags(), EIoReadOptionsFlags::HardwareTargetBuffer) ? true : false ))
 				{
 					Request->SetResult(Chunk);
 				}
diff --git a/Engine/Source/Runtime/StorageServerClient/Public/StorageServerConnection.h b/Engine/Source/Runtime/StorageServerClient/Public/StorageServerConnection.h
--- a/Engine/Source/Runtime/StorageServerClient/Public/StorageServerConnection.h
+++ b/Engine/Source/Runtime/StorageServerClient/Public/StorageServerConnection.h
@@ -129,12 +129,12 @@ public:
 
 	void Serialize(void* V, int64 Length) override;
 
-	int64 SerializeChunk(FStorageServerSerializationContext& Context, FIoBuffer& OutChunk, void* TargetVa = nullptr, uint64 RawOffset = 0, uint64 RawSize = MAX_uint64);
+	int64 SerializeChunk(FStorageServerSerializationContext& Context, FIoBuffer& OutChunk, void* TargetVa = nullptr, uint64 RawOffset = 0, uint64 RawSize = MAX_uint64, bool  bHardwareTargetBuffer = false);
 	
-	inline int64 SerializeChunk(FIoBuffer& OutChunk, void* TargetVa = nullptr, uint64 RawOffset = 0, uint64 RawSize = MAX_uint64)
+	inline int64 SerializeChunk(FIoBuffer& OutChunk, void* TargetVa = nullptr, uint64 RawOffset = 0, uint64 RawSize = MAX_uint64, bool bHardwareTargetBuffer = false)
 	{
 		FStorageServerSerializationContext SerializationContext;
-		return SerializeChunk(SerializationContext, OutChunk, TargetVa, RawOffset, RawSize);
+		return SerializeChunk(SerializationContext, OutChunk, TargetVa, RawOffset, RawSize, bHardwareTargetBuffer);
 	}
 
 	int64 SerializeChunkTo(FMutableMemoryView Memory, uint64 RawOffset = 0);
