- Fix some outdated comments

diff --git a/Engine/Source/Runtime/Core/Public/Misc/ConfigHierarchy.h b/Engine/Source/Runtime/Core/Public/Misc/ConfigHierarchy.h
--- a/Engine/Source/Runtime/Core/Public/Misc/ConfigHierarchy.h
+++ b/Engine/Source/Runtime/Core/Public/Misc/ConfigHierarchy.h
@@ -45,7 +45,7 @@ inline FConfigLayer GConfigLayers[] =
 
 
 /// <summary>
-/// Plugins don't need to look at the same number of insane layers. Here PROJECT is the Plugin dir
+/// Plugins don't need to look at the same number of insane layers.
 /// </summary>
 inline FConfigLayer GPluginLayers[] =
 {
@@ -60,7 +60,7 @@ inline FConfigLayer GPluginLayers[] =
 	{ TEXT("PluginPlatform"),			TEXT("{PLUGIN}/Config/{PLATFORM}/{PLATFORM}{TYPE}.ini") },
 	// Project/Default.ini
 	{ TEXT("ProjectDefault"),			TEXT("{PROJECT}/Config/Default{TYPE}.ini") },
-	// Project/Platform/.ini
+	// Project/Platform/Platform*.ini
 	{ TEXT("ProjectDefault"),			TEXT("{PROJECT}/Config/{PLATFORM}/{PLATFORM}{TYPE}.ini") },
 };
 
