Cap the amount of logspam when hunting down compressed shader issues.

diff --git a/Engine/Source/Runtime/Core/Private/Misc/Compression.cpp b/Engine/Source/Runtime/Core/Private/Misc/Compression.cpp
--- a/Engine/Source/Runtime/Core/Private/Misc/Compression.cpp
+++ b/Engine/Source/Runtime/Core/Private/Misc/Compression.cpp
@@ -953,9 +953,13 @@ bool FCompression::UncompressMemory(FName FormatName, void* UncompressedBuffer,
 		}
 		// Always log an error
 		UE_LOG(LogCompression, Error, TEXT("FCompression::UncompressMemory - Failed to uncompress memory (%d/%d) from address %p using format %s, this may indicate the asset is corrupt!"), CompressedSize, UncompressedSize, CompressedBuffer, *FormatName.ToString());
-		// this extra logging is added to understand shader decompression errors, see UE-159777
+		// this extra logging is added to understand shader decompression errors, see UE-159777. However in unrelated
+		// corruption issues this gets hit a lot causing massive log sizes. Since for UE-159777 we crash afterwards, we
+		// can safely limit to one instance for the purpose of diagnosing this.
+		static std::atomic_int32_t has_logged = 0;
+		bool bAllowLog = has_logged.exchange(1, std::memory_order_relaxed) == 0;
 		const int32 MaxSizeToLogOut = 16384;
-		if (FormatName == NAME_Oodle && CompressedSize <= MaxSizeToLogOut)
+		if (bAllowLog && FormatName == NAME_Oodle && CompressedSize <= MaxSizeToLogOut)
 		{
 			UE_LOG(LogCompression, Error, TEXT("FCompression::UncompressMemory - Logging compressed data (%d bytes) as a hex dump for investigation"), CompressedSize);
 			FCompressionUtil::LogHexDump(reinterpret_cast<const uint8*>(CompressedBuffer), CompressedSize, 0, CompressedSize);
