Add EAutomatinTestFlags::ProgramContext flag to filter tests for Program target.

diff --git a/Engine/Source/Runtime/Core/Private/Misc/AutomationTest.cpp b/Engine/Source/Runtime/Core/Private/Misc/AutomationTest.cpp
--- a/Engine/Source/Runtime/Core/Private/Misc/AutomationTest.cpp
+++ b/Engine/Source/Runtime/Core/Private/Misc/AutomationTest.cpp
@@ -154,6 +154,7 @@ CORE_API const TMap<FString, EAutomationTestFlags::Type>& EAutomationTestFlags::
 		{ TEXT("ClientContext"), Type::ClientContext},
 		{ TEXT("ServerContext"), Type::ServerContext},
 		{ TEXT("CommandletContext"), Type::CommandletContext},
+		{ TEXT("ProgramContext"), Type::ProgramContext},
 		{ TEXT("ApplicationContextMask"), Type::ApplicationContextMask},
 		{ TEXT("NonNullRHI"), Type::NonNullRHI},
 		{ TEXT("RequiresUser"), Type::RequiresUser},
@@ -680,8 +681,9 @@ void FAutomationTestFramework::GetValidTestNames( TArray<FAutomationTestInfo>& T
 	// Determine required application type (Editor, Game, or Commandlet)
 	const bool bRunningCommandlet = IsRunningCommandlet();
 	const bool bRunningEditor = GIsEditor && !bRunningCommandlet;
-	const bool bRunningClient = !GIsEditor && !IsRunningDedicatedServer();
-	const bool bRunningServer = !GIsEditor && !IsRunningClientOnly();
+	const bool bRunningClient = !GIsEditor && !IsRunningDedicatedServer() && !FPlatformProperties::IsProgram();
+	const bool bRunningServer = !GIsEditor && !IsRunningClientOnly() && !FPlatformProperties::IsProgram();
+	const bool bRunningProgram = !GIsEditor && FPlatformProperties::IsProgram();
 
 	//application flags
 	uint32 ApplicationSupportFlags = 0;
@@ -701,6 +703,10 @@ void FAutomationTestFramework::GetValidTestNames( TArray<FAutomationTestInfo>& T
 	{
 		ApplicationSupportFlags |= EAutomationTestFlags::CommandletContext;
 	}
+	if ( bRunningProgram )
+	{
+		ApplicationSupportFlags |= EAutomationTestFlags::ProgramContext;
+	}
 
 	//Feature support - assume valid RHI until told otherwise
 	uint32 FeatureSupportFlags = EAutomationTestFlags::FeatureMask;
diff --git a/Engine/Source/Runtime/Core/Public/Misc/AutomationTest.h b/Engine/Source/Runtime/Core/Public/Misc/AutomationTest.h
--- a/Engine/Source/Runtime/Core/Public/Misc/AutomationTest.h
+++ b/Engine/Source/Runtime/Core/Public/Misc/AutomationTest.h
@@ -89,7 +89,9 @@ struct EAutomationTestFlags
 		ServerContext = 0x00000004,
 		// Test is suitable for running within a commandlet
 		CommandletContext = 0x00000008,
-		ApplicationContextMask = EditorContext | ClientContext | ServerContext | CommandletContext,
+		// Test is suitable for running within program application (not editor, nor game)
+		ProgramContext = 0x00000010,
+		ApplicationContextMask = EditorContext | ClientContext | ServerContext | CommandletContext | ProgramContext,
 
 		//~ Features required for the test - not specifying means it is valid for any feature combination
 		// Test requires a non-null RHI to run correctly
