Improved error reporting when collecting disk information

- made some call Synchronous instead of SemiSynchronous so errors are reported immediately
- downgraded Warning to Log
- added more information in the log

#tests in FrontEnd
#rnx

diff --git a/Engine/Source/Runtime/Core/Private/Windows/WindowsPlatformMisc.cpp b/Engine/Source/Runtime/Core/Private/Windows/WindowsPlatformMisc.cpp
--- a/Engine/Source/Runtime/Core/Private/Windows/WindowsPlatformMisc.cpp
+++ b/Engine/Source/Runtime/Core/Private/Windows/WindowsPlatformMisc.cpp
@@ -418,10 +418,27 @@ namespace
 
 	static void LogStorageInformationWarning(HRESULT HRes,  const TCHAR* message)
 	{
-		IErrorInfo* Error = nullptr;
-		GetErrorInfo(0, &Error);
-		_com_error error(HRes, Error);
-		UE_LOG(LogWindows, Warning, TEXT("%s [%s]"), message, error.ErrorMessage());
+		IErrorInfo* ErrorInfo = nullptr;
+		GetErrorInfo(0, &ErrorInfo);
+		if (ErrorInfo)
+		{
+			BSTR ErrorMessage;
+			ErrorInfo->GetDescription(&ErrorMessage);
+			if (ErrorMessage && *ErrorMessage)
+			{
+				UE_LOG(LogWindows, Log, TEXT("%s [%s]"), message, ErrorMessage);
+			}
+			else
+			{
+				UE_LOG(LogWindows, Log, TEXT("%s [error code %x]"), message, HRes);
+			}
+			::SysFreeString(ErrorMessage);
+			ErrorInfo->Release();
+		}
+		else
+		{
+			UE_LOG(LogWindows, Log, TEXT("%s [error code %x]"), message, HRes);
+		}
 	}
 
 	static bool CollectStorageInformation()
@@ -479,7 +496,7 @@ namespace
 		hres = WbemServices->ExecQuery(
 			bstr_t("WQL"),
 			bstr_t("SELECT * FROM MSFT_DiskToPartition"),
-			WBEM_FLAG_FORWARD_ONLY | WBEM_FLAG_RETURN_IMMEDIATELY,
+			WBEM_FLAG_FORWARD_ONLY,
 			NULL,
 			&StorageEnumerator);
 
@@ -580,7 +597,7 @@ namespace
 		hres = WbemServices->ExecQuery(
 			bstr_t("WQL"),
 			bstr_t("SELECT * FROM MSFT_PhysicalDisk"),
-			WBEM_FLAG_FORWARD_ONLY | WBEM_FLAG_RETURN_IMMEDIATELY,
+			WBEM_FLAG_FORWARD_ONLY,
 			NULL,
 			&StorageEnumerator);
 
