[TaskGraph]
  - Fix use-after-free in both WaitForAnyTaskCompleted and UE::Tasks::WaitAny

#rnx
#tests Stress FTaskGraphWaitForAnyTask under ASAN/TSAN for both TASKGRAPH_NEW_FRONTEND on/off

diff --git a/Engine/Source/Runtime/Core/Private/Async/TaskGraph.cpp b/Engine/Source/Runtime/Core/Private/Async/TaskGraph.cpp
--- a/Engine/Source/Runtime/Core/Private/Async/TaskGraph.cpp
+++ b/Engine/Source/Runtime/Core/Private/Async/TaskGraph.cpp
@@ -16,6 +16,7 @@
 #include "HAL/ThreadSafeCounter.h"
 #include "Misc/NoopCounter.h"
 #include "Misc/ScopeLock.h"
+#include "Async/ManualResetEvent.h"
 #include "Containers/LockFreeList.h"
 #include "Templates/Function.h"
 #include "Stats/Stats.h"
@@ -24,6 +25,7 @@
 #include "HAL/IConsoleManager.h"
 #include "Misc/App.h"
 #include "Misc/Fork.h"
+#include "Misc/Timeout.h"
 #include "Async/TaskGraphInterfaces.h"
 #include "HAL/LowLevelMemTracker.h"
 #include "HAL/ThreadHeartBeat.h"
@@ -2811,28 +2813,55 @@ static FAutoConsoleCommand TaskThreadPriorityCmd(
 
 int32 WaitForAnyTaskCompleted(const FGraphEventArray& GraphEvents, FTimespan Timeout /*= FTimespan::MaxValue()*/)
 {
+#if TASKGRAPH_NEW_FRONTEND
+	return UE::Tasks::WaitAny(GraphEvents, Timeout);
+#else
 	if (UNLIKELY(GraphEvents.Num() == 0))
 	{
 		return INDEX_NONE;
 	}
 
-	FSharedEventRef SystemEvent;
-	std::atomic<int32> CompletedTaskIndex;
+	// Avoid memory allocations if any of the events are already completed
+	for (int32 Index = 0; Index < GraphEvents.Num(); ++Index)
+	{
+		if (GraphEvents[Index]->IsComplete())
+		{
+			return Index;
+		}
+	}
 
-	for (int32 Index = 0; Index != GraphEvents.Num(); ++Index)
+	struct FSharedData
+	{
+		UE::FManualResetEvent Event;
+		std::atomic<int32>    CompletedTaskIndex{ 0 };
+	};
+
+	// Shared data usage is important to avoid the variable to go out of scope
+	// before all the task have been run even if we exit after the first event
+	// is triggered.
+	TSharedRef<FSharedData> SharedData = MakeShared<FSharedData>();
+
+	for (int32 Index = 0; Index < GraphEvents.Num(); ++Index)
 	{
 		FFunctionGraphTask::CreateAndDispatchWhenReady(
-			[SystemEvent, Index, &CompletedTaskIndex] 
-			{ 
-				CompletedTaskIndex.store(Index, std::memory_order_relaxed);
-				SystemEvent->Trigger(); 
-			}, 
-			TStatId{}, 
-			GraphEvents[Index]
+			[SharedData, Index]
+			{
+				SharedData->CompletedTaskIndex.store(Index, std::memory_order_relaxed);
+				SharedData->Event.Notify();
+			},
+			TStatId{},
+			GraphEvents[Index],
+			ENamedThreads::AnyHiPriThreadHiPriTask /* Run as soon as possible. */
 		);
 	}
 
-	return SystemEvent->Wait(Timeout) ? CompletedTaskIndex.load(std::memory_order_relaxed) : INDEX_NONE;
+	if (SharedData->Event.WaitFor(UE::FMonotonicTimeSpan::FromMilliseconds(Timeout.GetTotalMilliseconds())))
+	{
+		return SharedData->CompletedTaskIndex.load(std::memory_order_relaxed);
+	}
+
+	return INDEX_NONE;
+#endif
 }
 
 FGraphEventRef AnyTaskCompleted(const FGraphEventArray& GraphEvents)
diff --git a/Engine/Source/Runtime/Core/Public/Tasks/Task.h b/Engine/Source/Runtime/Core/Public/Tasks/Task.h
--- a/Engine/Source/Runtime/Core/Public/Tasks/Task.h
+++ b/Engine/Source/Runtime/Core/Public/Tasks/Task.h
@@ -4,6 +4,7 @@
 
 #include "Tasks/TaskPrivate.h"
 #include "Async/Fundamental/Task.h"
+#include "Async/ManualResetEvent.h"
 #include "Containers/StaticArray.h"
 #include "HAL/Event.h"
 #include "HAL/IConsoleManager.h"
@@ -391,6 +392,18 @@ namespace UE::Tasks
 		{
 			Array[Index] = Task.Pimpl.GetReference();
 		}
+
+		template<typename HigherLevelTaskType, std::enable_if_t<std::is_same_v<HigherLevelTaskType, FTask>>* = nullptr>
+		bool IsCompleted(const HigherLevelTaskType& Prerequisite)
+		{
+			return Prerequisite.IsCompleted();
+		}
+
+		template<typename HigherLevelTaskType, std::enable_if_t<std::is_same_v<HigherLevelTaskType, FGraphEventRef>>* = nullptr>
+		bool IsCompleted(const HigherLevelTaskType& Prerequisite)
+		{
+			return Prerequisite.IsValid() ? Prerequisite->IsCompleted() : false;
+		}
 	}
 
 	template<typename... TaskTypes, 
@@ -423,16 +436,33 @@ namespace UE::Tasks
 			return INDEX_NONE;
 		}
 
-		FSharedEventRef Event;
-		TSharedPtr<std::atomic<int32>> CompletedTaskIndex = MakeShared<std::atomic<int32>>(INDEX_NONE);
+		// Avoid memory allocations if any of the events are already completed
+		for (int32 Index = 0; Index < Tasks.Num(); ++Index)
+		{
+			if (Private::IsCompleted(Tasks[Index]))
+			{
+				return Index;
+			}
+		}
+
+		struct FSharedData
+		{
+			UE::FManualResetEvent Event;
+			std::atomic<int32>    CompletedTaskIndex{ 0 };
+		};
 
-		for (int32 Index = 0; Index != Tasks.Num(); ++Index)
+		// Shared data usage is important to avoid the variable to go out of scope
+		// before all the task have been run even if we exit after the first event
+		// is triggered.
+		TSharedRef<FSharedData> SharedData = MakeShared<FSharedData>();
+
+		for (int32 Index = 0; Index < Tasks.Num(); ++Index)
 		{
 			Launch(UE_SOURCE_LOCATION, 
-				[Event, Index, CompletedTaskIndex] 
+				[SharedData, Index]
 				{ 
-					CompletedTaskIndex->store(Index, std::memory_order_relaxed);
-					Event->Trigger(); 
+					SharedData->CompletedTaskIndex.store(Index, std::memory_order_relaxed);
+					SharedData->Event.Notify();
 				}, 
 				Prerequisites(Tasks[Index]),
 				ETaskPriority::Default, 
@@ -440,7 +470,12 @@ namespace UE::Tasks
 			);
 		}
 
-		return Event->Wait(Timeout) ? CompletedTaskIndex->load(std::memory_order_relaxed) : INDEX_NONE;
+		if (SharedData->Event.WaitFor(UE::FMonotonicTimeSpan::FromMilliseconds(Timeout.GetTotalMilliseconds())))
+		{
+			return SharedData->CompletedTaskIndex.load(std::memory_order_relaxed);
+		}
+
+		return INDEX_NONE;
 	}
 
 	// Returns a task that gets completed as soon as any of the given tasks gets completed
diff --git a/Engine/Source/Runtime/Core/Public/Tasks/TaskPrivate.h b/Engine/Source/Runtime/Core/Public/Tasks/TaskPrivate.h
--- a/Engine/Source/Runtime/Core/Public/Tasks/TaskPrivate.h
+++ b/Engine/Source/Runtime/Core/Public/Tasks/TaskPrivate.h
@@ -260,6 +260,14 @@ namespace UE::Tasks
 				return Prerequisite.IsValid() ? AddPrerequisites(*Prerequisite.Pimpl) : false;
 			}
 
+			// The task will be executed only when all prerequisites are completed. The task type must be a task handle that holds a pointer to
+			// Must not be called concurrently
+			template<typename HigherLevelTaskType, std::enable_if_t<std::is_same_v<HigherLevelTaskType, FGraphEventRef>>* = nullptr>
+			bool AddPrerequisites(const HigherLevelTaskType& Prerequisite)
+			{
+				return Prerequisite.IsValid() ? AddPrerequisites(*Prerequisite.GetReference()) : false;
+			}
+
 protected:
 			// The task will be executed only when all prerequisites are completed.
 			// Must not be called concurrently.
