CompactBinary: Add operator<<(FCbWriter&, const FString&)  so that defining operator<< for some other type that can implicitly construct from FString does not cause an ambiguous overload error.
#rnx

diff --git a/Engine/Source/Runtime/Core/Private/Serialization/CompactBinaryWriter.cpp b/Engine/Source/Runtime/Core/Private/Serialization/CompactBinaryWriter.cpp
--- a/Engine/Source/Runtime/Core/Private/Serialization/CompactBinaryWriter.cpp
+++ b/Engine/Source/Runtime/Core/Private/Serialization/CompactBinaryWriter.cpp
@@ -2,6 +2,8 @@
 
 #include "Serialization/CompactBinaryWriter.h"
 
+#include "Containers/AnsiString.h"
+#include "Containers/Utf8String.h"
 #include "HAL/PlatformString.h"
 #include "HAL/UnrealMemory.h"
 #include "IO/IoHash.h"
@@ -626,6 +628,24 @@ FCbWriter& FCbWriter::operator<<(const FTimespan Value)
 	return *this;
 }
 
+FCbWriter& FCbWriter::operator<<(const FString& Value)
+{
+	*this << FStringView(Value);
+	return *this;
+}
+
+FCbWriter& FCbWriter::operator<<(const FAnsiString& Value)
+{
+	*this << FAnsiStringView(Value);
+	return *this;
+}
+
+FCbWriter& FCbWriter::operator<<(const FUtf8String& Value)
+{
+	*this << FUtf8StringView(Value);
+	return *this;
+}
+
 FCbWriter& FCbWriter::operator<<(FName Value)
 {
 	*this << WriteToUtf8String<FName::StringBufferSize>(Value).ToView();
diff --git a/Engine/Source/Runtime/Core/Public/Serialization/CompactBinaryWriter.h b/Engine/Source/Runtime/Core/Public/Serialization/CompactBinaryWriter.h
--- a/Engine/Source/Runtime/Core/Public/Serialization/CompactBinaryWriter.h
+++ b/Engine/Source/Runtime/Core/Public/Serialization/CompactBinaryWriter.h
@@ -435,6 +435,9 @@ public:
 		return *this;
 	}
 
+	CORE_API FCbWriter& operator<<(const FString& Value);
+	CORE_API FCbWriter& operator<<(const FAnsiString& Value);
+	CORE_API FCbWriter& operator<<(const FUtf8String& Value);
 	CORE_API FCbWriter& operator<<(FName Value);
 
 	template <typename T, typename Allocator,
