Exposes `AttendedStatusToString` and `AddCrashProperty` properly.

diff --git a/Engine/Source/Runtime/Core/Public/GenericPlatform/GenericPlatformCrashContext.h b/Engine/Source/Runtime/Core/Public/GenericPlatform/GenericPlatformCrashContext.h
--- a/Engine/Source/Runtime/Core/Public/GenericPlatform/GenericPlatformCrashContext.h
+++ b/Engine/Source/Runtime/Core/Public/GenericPlatform/GenericPlatformCrashContext.h
@@ -141,7 +141,7 @@ enum class EUnattendedStatus : uint8
 	Unattended
 };
 
-const TCHAR* AttendedStatusToString(const EUnattendedStatus Status);
+CORE_API const TCHAR* AttendedStatusToString(const EUnattendedStatus Status);
 
 #define CR_MAX_ERROR_MESSAGE_CHARS 2048
 #define CR_MAX_DIRECTORY_CHARS 256
@@ -746,7 +746,7 @@ private:
 	static void SerializeUserSettings(FString& Buffer);
 
 	/** Writes a common property to the buffer. */
-	static void AddCrashPropertyInternal(FString& Buffer, FStringView PropertyName, FStringView PropertyValue);
+	CORE_API static void AddCrashPropertyInternal(FString& Buffer, FStringView PropertyName, FStringView PropertyValue);
 
 	/** Writes a common property to the buffer. */
 	template <typename Type>
