[Unit Conversion]
- Update NoExportTypes.h to add missing units, and to correct the order of units that were already there. This was breaking serialization of entries missing from the mirrored enum in NoExportTypes.h
- Update the comment in UnitConversion.h to reference NoExportTypes.h instead of Object.h
- Make sure that EUnit in both UnitConversion.h and NoExportTypes.h match exactly, so they can be more easily diffed.
- Split some units onto their own lines, and neaten up the comments.

diff --git a/Engine/Source/Runtime/Core/Public/Math/UnitConversion.h b/Engine/Source/Runtime/Core/Public/Math/UnitConversion.h
--- a/Engine/Source/Runtime/Core/Public/Math/UnitConversion.h
+++ b/Engine/Source/Runtime/Core/Public/Math/UnitConversion.h
@@ -14,67 +14,82 @@ class FText;
 template<typename NumericType> struct FNumericUnit;
 
 /** Enum *must* be zero-indexed and sequential. Must be grouped by relevance and ordered by magnitude. */
-/** Enum *must* match the mirrored enum that exists in CoreUObject/Classes/Object.h for the purposes of UObject reflection */
+/** Enum *must* match the mirrored enum that exists in CoreUObject/NoExportTypes.h for the purposes of UObject reflection */
 enum class EUnit : uint8
 {
-	/** Scalar distance/length units */
+	/** Scalar distance/length unit. */
 	Micrometers, Millimeters, Centimeters, Meters, Kilometers,
 	Inches, Feet, Yards, Miles,
 	Lightyears,
 
-	/** Angular units */
+	/** Angular unit. */
 	Degrees, Radians,
 
-	/** Speed units */
+	/** Speed unit. */
 	CentimetersPerSecond, MetersPerSecond, KilometersPerHour, MilesPerHour,
 
-	/** Angular speed units */
+	/** Angular speed unit. */
 	DegreesPerSecond, RadiansPerSecond,
-	
-	/** Acceleration units */
+
+	/** Acceleration unit. */
 	CentimetersPerSecondSquared, MetersPerSecondSquared,
 
-	/** Temperature units */
+	/** Temperature unit. */
 	Celsius, Farenheit, Kelvin,
 
-	/** Mass units */
+	/** Mass unit. */
 	Micrograms, Milligrams, Grams, Kilograms, MetricTons,
 	Ounces, Pounds, Stones,
 
-	/** Density units */
+	/** Density unit. */
 	GramsPerCubicCentimeter, GramsPerCubicMeter, KilogramsPerCubicCentimeter, KilogramsPerCubicMeter,
 
-	/** Force units */
+	/** Force unit. */
 	Newtons, PoundsForce, KilogramsForce, KilogramCentimetersPerSecondSquared,
 
-	/** Torque Units */
+	/** Torque unit. */
 	NewtonMeters, KilogramCentimetersSquaredPerSecondSquared,
 
-	/** Impulse Units */
+	/** Impulse unit. */
 	NewtonSeconds, KilogramCentimeters, KilogramMeters,
 
-	/** Frequency units */
+	/** Frequency unit. */
 	Hertz, Kilohertz, Megahertz, Gigahertz, RevolutionsPerMinute,
 
-	/** Data Size units */
+	/** Data Size unit. */
 	Bytes, Kilobytes, Megabytes, Gigabytes, Terabytes,
 
-	/** Luminous flux units, luminous intensity, illuminance, luminance, exposure value */
-	Lumens, Candela, Lux, CandelaPerMeter2, ExposureValue,
+	/** Luminous flux unit. */
+	Lumens,
+	
+	/** Luminous intensity unit. */
+	Candela,
+	
+	/** Illuminance unit. */
+	Lux,
+	
+	/** Luminance unit. */
+	CandelaPerMeter2,
+	
+	/** Exposure value unit. */
+	ExposureValue,
 
-	/** Time units */
+	/** Time unit. */
 	Nanoseconds, Microseconds, Milliseconds, Seconds, Minutes, Hours, Days, Months, Years,
 
-	/** Pixel density units */
+	/** Pixel density unit. */
 	PixelsPerInch,
 
-	/** Arbitrary multipliers */
-	Percentage,	Multiplier,
+	/** Percentage. */
+	Percentage,
+
+	/** Arbitrary multiplier. */
+	Multiplier,
 
-	/** Stress units */
+	/** Stress unit. */
 	Pascals, KiloPascals, MegaPascals, GigaPascals,
 
-	/** Symbolic entry, not specifiable on meta data */
+	/** Symbolic entry, not specifiable on meta data. */
 	Unspecified
 };
 
diff --git a/Engine/Source/Runtime/CoreUObject/Public/UObject/NoExportTypes.h b/Engine/Source/Runtime/CoreUObject/Public/UObject/NoExportTypes.h
--- a/Engine/Source/Runtime/CoreUObject/Public/UObject/NoExportTypes.h
+++ b/Engine/Source/Runtime/CoreUObject/Public/UObject/NoExportTypes.h
@@ -360,54 +360,77 @@ UENUM(BlueprintType)
 enum class EUnit : uint8
 {
 	/** Scalar distance/length unit. */
-	Micrometers, Millimeters, Centimeters, Meters, Kilometers, Inches, Feet, Yards, Miles, Lightyears,
-	
-	/** Angular units */
+	Micrometers, Millimeters, Centimeters, Meters, Kilometers,
+	Inches, Feet, Yards, Miles,
+	Lightyears,
+
+	/** Angular unit. */
 	Degrees, Radians,
-	
-	/** Speed units */
+
+	/** Speed unit. */
 	CentimetersPerSecond, MetersPerSecond, KilometersPerHour, MilesPerHour,
 
-	/** Angular speed units */
+	/** Angular speed unit. */
 	DegreesPerSecond, RadiansPerSecond,
-	
-	/** Acceleration units */
+
+	/** Acceleration unit. */
 	CentimetersPerSecondSquared, MetersPerSecondSquared,
 
-	/** Temperature units */
+	/** Temperature unit. */
 	Celsius, Farenheit, Kelvin,
-	
-	/** Mass units */
-	Micrograms, Milligrams, Grams, Kilograms, MetricTons, Ounces, Pounds, Stones,
-	
-	/** Force units */
+
+	/** Mass unit. */
+	Micrograms, Milligrams, Grams, Kilograms, MetricTons,
+	Ounces, Pounds, Stones,
+
+	/** Density unit. */
+	GramsPerCubicCentimeter, GramsPerCubicMeter, KilogramsPerCubicCentimeter, KilogramsPerCubicMeter,
+
+	/** Force unit. */
 	Newtons, PoundsForce, KilogramsForce, KilogramCentimetersPerSecondSquared,
 
-	/** Torque units */
+	/** Torque unit. */
 	NewtonMeters, KilogramCentimetersSquaredPerSecondSquared,
 
-	/** Impulse units */
+	/** Impulse unit. */
 	NewtonSeconds, KilogramCentimeters, KilogramMeters,
-	
-	/** Frequency units */
+
+	/** Frequency unit. */
 	Hertz, Kilohertz, Megahertz, Gigahertz, RevolutionsPerMinute,
-	
-	/** Data Size units */	
+
+	/** Data Size unit. */
 	Bytes, Kilobytes, Megabytes, Gigabytes, Terabytes,
-	
-	/** Luminous flux units */	
+
+	/** Luminous flux unit. */
 	Lumens,
 	
-	/** Time units */	
-	Milliseconds, Seconds, Minutes, Hours, Days, Months, Years,
+	/** Luminous intensity unit. */
+	Candela,
+	
+	/** Illuminance unit. */
+	Lux,
+	
+	/** Luminance unit. */
+	CandelaPerMeter2,
+	
+	/** Exposure value unit. */
+	ExposureValue,
 
-	/** Arbitrary multiplier */	
-	Multiplier,
+	/** Time unit. */
+	Nanoseconds, Microseconds, Milliseconds, Seconds, Minutes, Hours, Days, Months, Years,
 
+	/** Pixel density unit. */
+	PixelsPerInch,
 
-	/** Percentage */
+	/** Percentage. */
 	Percentage,
 
+	/** Arbitrary multiplier. */
+	Multiplier,
+
+	/** Stress unit. */
+	Pascals, KiloPascals, MegaPascals, GigaPascals,
+
 	/** Symbolic entry, not specifiable on meta data. */
 	Unspecified
 };
