Exposes `AddHeader`/`AddFooter`/`BeginSection`/`EndSection`/`AddSection` properly by adding the missing export macros.

diff --git a/Engine/Source/Runtime/Core/Public/GenericPlatform/GenericPlatformCrashContext.h b/Engine/Source/Runtime/Core/Public/GenericPlatform/GenericPlatformCrashContext.h
--- a/Engine/Source/Runtime/Core/Public/GenericPlatform/GenericPlatformCrashContext.h
+++ b/Engine/Source/Runtime/Core/Public/GenericPlatform/GenericPlatformCrashContext.h
@@ -772,14 +772,14 @@ private:
 
 public:  // Allows this helper functionality to be present for clients to write their own types of crash reports.
 	/** Writes header information to the buffer. */
-	static void AddHeader(FString& Buffer);
+	CORE_API static void AddHeader(FString& Buffer);
 
 	/** Writes footer to the buffer. */
-	static void AddFooter(FString& Buffer);
+	CORE_API static void AddFooter(FString& Buffer);
 
-	static void BeginSection(FString& Buffer, const TCHAR* SectionName);
-	static void EndSection(FString& Buffer, const TCHAR* SectionName);
-	static void AddSection(FString& Buffer, const TCHAR* SectionName, const FString& SectionContent);
+	CORE_API static void BeginSection(FString& Buffer, const TCHAR* SectionName);
+	CORE_API static void EndSection(FString& Buffer, const TCHAR* SectionName);
+	CORE_API static void AddSection(FString& Buffer, const TCHAR* SectionName, const FString& SectionContent);
 
 private:
 	/** Called once when GConfig is initialized. Opportunity to cache values from config. */
