Added a check to MakeShareable to ensure that you're not trying to call it on an already-shared TSharedFromThis object, as this will result in a double delete.

diff --git a/Engine/Source/Runtime/Core/Public/Templates/SharedPointer.h b/Engine/Source/Runtime/Core/Public/Templates/SharedPointer.h
--- a/Engine/Source/Runtime/Core/Public/Templates/SharedPointer.h
+++ b/Engine/Source/Runtime/Core/Public/Templates/SharedPointer.h
@@ -2171,6 +2171,12 @@ template< class CastToType, class CastFromType, ESPMode Mode >
 template< class ObjectType >
 [[nodiscard]] FORCEINLINE SharedPointerInternals::TRawPtrProxy< ObjectType > MakeShareable( ObjectType* InObject )
 {
+	if constexpr (IsDerivedFromSharedFromThis<ObjectType>())
+	{
+		// If this goes off, you should probably be using Ptr->AsShared() or Ptr->AsWeak() instead.
+		checkf(!InObject || !InObject->DoesSharedInstanceExist(), TEXT("Trying to share an already-shared object"));
+	}
+
 	return SharedPointerInternals::TRawPtrProxy< ObjectType >( InObject );
 }
 
@@ -2184,6 +2190,12 @@ template< class ObjectType >
 template< class ObjectType, class DeleterType >
 [[nodiscard]] FORCEINLINE SharedPointerInternals::TRawPtrProxyWithDeleter< ObjectType, DeleterType > MakeShareable( ObjectType* InObject, DeleterType&& InDeleter )
 {
+	if constexpr (IsDerivedFromSharedFromThis<ObjectType>())
+	{
+		// If this goes off, you should probably be using Ptr->AsShared() or Ptr->AsWeak() instead.
+		checkf(!InObject || !InObject->DoesSharedInstanceExist(), TEXT("Trying to share an already-shared object"));
+	}
+
 	return SharedPointerInternals::TRawPtrProxyWithDeleter< ObjectType, DeleterType >( InObject, Forward< DeleterType >( InDeleter ) );
 }
 
