Added TArrayView support to TSet's constructor and fixed TArrayView deduction in its Append overload.
Added deduction guide for TSet to allow CTAD.

diff --git a/Engine/Source/Runtime/Core/Public/Containers/Set.h b/Engine/Source/Runtime/Core/Public/Containers/Set.h
--- a/Engine/Source/Runtime/Core/Public/Containers/Set.h
+++ b/Engine/Source/Runtime/Core/Public/Containers/Set.h
@@ -317,10 +317,10 @@ public:
 		*this = Copy;
 	}
 
-	FORCEINLINE explicit TSet(const TArray<ElementType>& InArray)
+	FORCEINLINE explicit TSet(TArrayView<const ElementType> InArrayView)
 		: HashSize(0)
 	{
-		Append(InArray);
+		Append(InArrayView);
 	}
 
 	FORCEINLINE explicit TSet(TArray<ElementType>&& InArray)
@@ -774,28 +774,7 @@ public:
 		return FSetElementId(NewHashIndex);
 	}
 
-	template<typename ViewSizeType>
-	void Append(TArrayView<ElementType, ViewSizeType> InElements)
-	{
-		Reserve(Elements.Num() + InElements.Num());
-		for (const ElementType& Element : InElements)
-		{
-			Add(Element);
-		}
-	}
-
-	template<typename ViewSizeType>
-	void Append(TArrayView<const ElementType, ViewSizeType> InElements)
-	{
-		Reserve(Elements.Num() + InElements.Num());
-		for (const ElementType& Element : InElements)
-		{
-			Add(Element);
-		}
-	}
-
-	template<typename ArrayAllocator>
-	void Append(const TArray<ElementType, ArrayAllocator>& InElements)
+	void Append(TArrayView<const ElementType> InElements)
 	{
 		Reserve(Elements.Num() + InElements.Num());
 		for (const ElementType& Element : InElements)
@@ -1809,6 +1788,9 @@ public:
 	FORCEINLINE TRangedForConstIterator end() const   { return TRangedForConstIterator(Elements.end());   }
 };
 
+template <typename RangeType>
+TSet(RangeType&&) -> TSet<TElementType_T<RangeType>>;
+
 namespace Freeze
 {
 	template<typename ElementType, typename KeyFuncs, typename Allocator>
