Localization.DumpLiveTable improvements

* Dumps can now be written to a file as well as the log, by specifying a filename via -DumpFile.
* Dumps are now sorted for better diffing.
* Dumps no longer trigger hang detection if they take a while to output.

diff --git a/Engine/Source/Runtime/Core/Private/Internationalization/TextLocalizationManager.cpp b/Engine/Source/Runtime/Core/Private/Internationalization/TextLocalizationManager.cpp
--- a/Engine/Source/Runtime/Core/Private/Internationalization/TextLocalizationManager.cpp
+++ b/Engine/Source/Runtime/Core/Private/Internationalization/TextLocalizationManager.cpp
@@ -11,6 +11,8 @@
 #include "GenericPlatform/GenericPlatformFile.h"
 #include "HAL/FileManager.h"
 #include "HAL/LowLevelMemTracker.h"
+#include "HAL/ThreadHeartBeat.h"
+#include "Misc/FileHelper.h"
 #include "Misc/Parse.h"
 #include "Misc/ScopeLock.h"
 #include "Misc/ScopeRWLock.h"
@@ -80,27 +82,29 @@ static FAutoConsoleCommand CmdDumpLiveTable(
 		TOptional<FString> NamespaceFilter;
 		TOptional<FString> KeyFilter;
 		TOptional<FString> DisplayStringFilter;
+		TOptional<FString> DumpFile;
 
 		for (const FString& Arg : Args)
 		{
 			if (!ParseOptionalStringArg(*Arg, TEXT("Namespace="), NamespaceFilter) &&
 				!ParseOptionalStringArg(*Arg, TEXT("Key="), KeyFilter) &&
-				!ParseOptionalStringArg(*Arg, TEXT("DisplayString="), DisplayStringFilter))
+				!ParseOptionalStringArg(*Arg, TEXT("DisplayString="), DisplayStringFilter) && 
+				!ParseOptionalStringArg(*Arg, TEXT("DumpFile="), DumpFile))
 			{
 				UE_LOG(LogLocalization, Warning, TEXT("Unknown argument '%s' passed to Localization.DumpLiveTable!"), *Arg);
 			}
 		}
 
-		auto GetConsoleResponseLogCategoryPtr = []() -> const FLogCategoryBase*
+		if (DumpFile.IsSet())
 		{
-#if NO_LOGGING
-			return nullptr;
-#else
-			return &LogConsoleResponse;
+			FTextLocalizationManager::Get().DumpLiveTable(DumpFile.GetValue(), NamespaceFilter.GetPtrOrNull(), KeyFilter.GetPtrOrNull(), DisplayStringFilter.GetPtrOrNull());
+		}
+		else
+		{
+#if !NO_LOGGING
+			FTextLocalizationManager::Get().DumpLiveTable(NamespaceFilter.GetPtrOrNull(), KeyFilter.GetPtrOrNull(), DisplayStringFilter.GetPtrOrNull(), &LogConsoleResponse);
 #endif
-		};
-
-		FTextLocalizationManager::Get().DumpLiveTable(NamespaceFilter.GetPtrOrNull(), KeyFilter.GetPtrOrNull(), DisplayStringFilter.GetPtrOrNull(), GetConsoleResponseLogCategoryPtr());
+		}
 	}));
 #endif
 }
@@ -606,32 +610,73 @@ void FTextLocalizationManager::CompactDataStructures()
 }
 
 #if ENABLE_LOC_TESTING
+void FTextLocalizationManager::DumpLiveTableImpl(const FString* NamespaceFilter, const FString* KeyFilter, const FString* DisplayStringFilter, TFunctionRef<void(const FTextId& Id, const FTextConstDisplayStringRef& DisplayString)> Callback) const
+{
+	FSlowHeartBeatScope SuspendHeartBeat;
+
+	FDisplayStringLookupTable DisplayStringLookupTableToDump;
+	{
+		auto PassesFilter = [](const FString& Str, const FString* Filter)
+		{
+			return !Filter || Str.MatchesWildcard(*Filter, ESearchCase::IgnoreCase); // Note: This is case insensitive since its used from a debug command
+		};
+
+		FScopeLock ScopeLock(&DisplayStringLookupTableCS);
+		DisplayStringLookupTableToDump.Reserve(DisplayStringLookupTable.Num());
+		for (const auto& DisplayStringPair : DisplayStringLookupTable)
+		{
+			if (PassesFilter(DisplayStringPair.Key.GetNamespace().GetChars(), NamespaceFilter) &&
+				PassesFilter(DisplayStringPair.Key.GetKey().GetChars(), KeyFilter) &&
+				PassesFilter(**DisplayStringPair.Value.DisplayString, DisplayStringFilter))
+			{
+				DisplayStringLookupTableToDump.Add(DisplayStringPair.Key, DisplayStringPair.Value);
+			}
+		}
+		DisplayStringLookupTableToDump.KeySort([](const FTextId& A, const FTextId& B)
+		{
+			const int32 NamespaceResult = FCString::Strcmp(A.GetNamespace().GetChars(), B.GetNamespace().GetChars());
+			if (NamespaceResult != 0)
+			{
+				return NamespaceResult < 0;
+			}
+			return FCString::Strcmp(A.GetKey().GetChars(), B.GetKey().GetChars()) < 0;
+		});
+	}
+
+	for (const auto& DisplayStringPair : DisplayStringLookupTableToDump)
+	{
+		Callback(DisplayStringPair.Key, DisplayStringPair.Value.DisplayString);
+	}
+}
+
 void FTextLocalizationManager::DumpLiveTable(const FString* NamespaceFilter, const FString* KeyFilter, const FString* DisplayStringFilter, const FLogCategoryBase* CategoryOverride) const
 {
 #if !NO_LOGGING
 	const FLogCategoryBase& Category = CategoryOverride ? *CategoryOverride : LogLocalization;
 
-	auto PassesFilter = [](const FString& Str, const FString* Filter)
-	{
-		return !Filter || Str.MatchesWildcard(*Filter, ESearchCase::IgnoreCase); // Note: This is case insensitive since its used from a debug command
-	};
-
 	UE_LOG_REF(Category, Display, TEXT("----------------------------------------------------------------------"));
 
-	FScopeLock ScopeLock(&DisplayStringLookupTableCS);
-	for (const auto& DisplayStringPair : DisplayStringLookupTable)
+	DumpLiveTableImpl(NamespaceFilter, KeyFilter, DisplayStringFilter, [&Category](const FTextId& Id, const FTextConstDisplayStringRef& DisplayString)
 	{
-		if (PassesFilter(DisplayStringPair.Key.GetNamespace().GetChars(), NamespaceFilter) &&
-			PassesFilter(DisplayStringPair.Key.GetKey().GetChars(), KeyFilter) &&
-			PassesFilter(**DisplayStringPair.Value.DisplayString, DisplayStringFilter))
-		{
-			UE_LOG_REF(Category, Display, TEXT("LiveTableEntry: Namespace: '%s', Key: '%s', DisplayString: '%s'"), DisplayStringPair.Key.GetNamespace().GetChars(), DisplayStringPair.Key.GetKey().GetChars(), **DisplayStringPair.Value.DisplayString);
-		}
-	}
+		UE_LOG_REF(Category, Display, TEXT("LiveTableEntry: Namespace: '%s', Key: '%s', DisplayString: '%s'"), Id.GetNamespace().GetChars(), Id.GetKey().GetChars(), **DisplayString);
+	});
 
 	UE_LOG_REF(Category, Display, TEXT("----------------------------------------------------------------------"));
 #endif // !NO_LOGGING
 }
+
+void FTextLocalizationManager::DumpLiveTable(const FString& OutputFilename, const FString* NamespaceFilter, const FString* KeyFilter, const FString* DisplayStringFilter) const
+{
+	FString DumpString;
+
+	DumpLiveTableImpl(NamespaceFilter, KeyFilter, DisplayStringFilter, [&DumpString](const FTextId& Id, const FTextConstDisplayStringRef& DisplayString)
+	{
+		DumpString += FString::Printf(TEXT("LiveTableEntry: Namespace: '%s', Key: '%s', DisplayString: '%s'"), Id.GetNamespace().GetChars(), Id.GetKey().GetChars(), **DisplayString);
+		DumpString += LINE_TERMINATOR;
+	});
+
+	FFileHelper::SaveStringToFile(DumpString, *OutputFilename, FFileHelper::EEncodingOptions::ForceUTF8WithoutBOM);
+}
 #endif
 
 FString FTextLocalizationManager::GetRequestedLanguageName() const
diff --git a/Engine/Source/Runtime/Core/Public/Internationalization/TextLocalizationManager.h b/Engine/Source/Runtime/Core/Public/Internationalization/TextLocalizationManager.h
--- a/Engine/Source/Runtime/Core/Public/Internationalization/TextLocalizationManager.h
+++ b/Engine/Source/Runtime/Core/Public/Internationalization/TextLocalizationManager.h
@@ -123,7 +123,11 @@ public:
 	 * Dumps the current live table state to the log, optionally filtering it based on the given wildcard arguments.
 	 * @note Calling this function with no filters specified will dump the entire live table.
 	 */
+private:
+	CORE_API void DumpLiveTableImpl(const FString* NamespaceFilter, const FString* KeyFilter, const FString* DisplayStringFilter, TFunctionRef<void(const FTextId& Id, const FTextConstDisplayStringRef& DisplayString)> Callback) const;
+public:
 	CORE_API void DumpLiveTable(const FString* NamespaceFilter = nullptr, const FString* KeyFilter = nullptr, const FString* DisplayStringFilter = nullptr, const FLogCategoryBase* CategoryOverride = nullptr) const;
+	CORE_API void DumpLiveTable(const FString& OutputFilename, const FString* NamespaceFilter = nullptr, const FString* KeyFilter = nullptr, const FString* DisplayStringFilter = nullptr) const;
 #endif
 
 	/**
