Silence errant PVS warning

#rnx

diff --git a/Engine/Source/Runtime/Core/Private/Internationalization/TextLocalizationManager.cpp b/Engine/Source/Runtime/Core/Private/Internationalization/TextLocalizationManager.cpp
--- a/Engine/Source/Runtime/Core/Private/Internationalization/TextLocalizationManager.cpp
+++ b/Engine/Source/Runtime/Core/Private/Internationalization/TextLocalizationManager.cpp
@@ -658,7 +658,7 @@ void FTextLocalizationManager::DumpLiveTable(const FString* NamespaceFilter, con
 
 	DumpLiveTableImpl(NamespaceFilter, KeyFilter, DisplayStringFilter, [&Category](const FTextId& Id, const FTextConstDisplayStringRef& DisplayString)
 	{
-		UE_LOG_REF(Category, Display, TEXT("LiveTableEntry: Namespace: '%s', Key: '%s', DisplayString: '%s'"), Id.GetNamespace().GetChars(), Id.GetKey().GetChars(), **DisplayString);
+		UE_LOG_REF(Category, Display, TEXT("LiveTableEntry: Namespace: '%s', Key: '%s', DisplayString: '%s'"), Id.GetNamespace().GetChars(), Id.GetKey().GetChars(), **DisplayString); //-V510
 	});
 
 	UE_LOG_REF(Category, Display, TEXT("----------------------------------------------------------------------"));
