Remove the additional metadata added in CL 28039514 from the IoStoreToc.

#rnx

- The metadata was originally added so that we could generate a toc for OnDemand streaming data if needed at runtime, but we no longer have any plans to support this functionality so we can remove it.
- The change still allows .utoc files that contain the additional metadata to be read. Since the data was at the end of the structure we can just ignore it and the code does not need to consider it's existence.
- Strictly speaking there was no need to add a new entry to EIoStoreTocVersion but we have chosen to do so on the off chance that we need to identify .utoc files that were created while this code was active, for what ever reason.

diff --git a/Engine/Source/Runtime/Core/Internal/IO/IoStore.h b/Engine/Source/Runtime/Core/Internal/IO/IoStore.h
--- a/Engine/Source/Runtime/Core/Internal/IO/IoStore.h
+++ b/Engine/Source/Runtime/Core/Internal/IO/IoStore.h
@@ -28,6 +28,7 @@ enum class EIoStoreTocVersion : uint8
 	PerfectHash,
 	PerfectHashWithOverflow,
 	OnDemandMetaData,
+	RemovedOnDemandMetaData,
 	LatestPlusOne,
 	Latest = LatestPlusOne - 1
 };
@@ -93,18 +94,6 @@ struct FIoStoreTocEntryMeta
 	FIoStoreTocEntryMetaFlags Flags;
 };
 
-struct FIoStoreTocOnDemandChunkMeta
-{
-	/** Hash of the chunk on disk after both compression and encryption */
-	FIoHash DiskHash;
-};
-
-struct FIoStoreTocOnDemandCompressedBlockMeta
-{
-	/** Hash of the block on disk */
-	FIoHash DiskHash;
-};
-
 /**
  * Compression block entry.
  */
@@ -210,9 +199,6 @@ struct FIoStoreTocResource
 	
 	TArray<FIoStoreTocEntryMeta> ChunkMetas;
 
-	TArray<FIoStoreTocOnDemandChunkMeta> OnDemandChunkMeta;
-	TArray<FIoStoreTocOnDemandCompressedBlockMeta> OnDemandCompressedBlockMeta;
-
 	[[nodiscard]] CORE_API FIoStoreTocChunkInfo GetTocChunkInfo(int32 TocEntryIndex) const;
 
 	[[nodiscard]] CORE_API static FIoStatus Read(const TCHAR* TocFilePath, EIoStoreTocReadOptions ReadOptions, FIoStoreTocResource& OutTocResource);
diff --git a/Engine/Source/Runtime/Core/Private/IO/IoStore.cpp b/Engine/Source/Runtime/Core/Private/IO/IoStore.cpp
--- a/Engine/Source/Runtime/Core/Private/IO/IoStore.cpp
+++ b/Engine/Source/Runtime/Core/Private/IO/IoStore.cpp
@@ -200,7 +200,6 @@ struct FIoStoreWriteQueueEntry
 	FGraphEventRef BeginCompressionBarrier;
 	FGraphEventRef FinishCompressionBarrier;
 	FGraphEventRef FinishEncryptionAndSigningBarrier;
-	FGraphEventRef BeginOnDemandDataBarrier;
 	FGraphEventRef BeginWriteBarrier;
 	FGraphEventRef WriteFinishedEvent;
 	TAtomic<int32> CompressedBlocksCount{ 0 };
@@ -310,12 +309,10 @@ public:
 		BeginCompressionQueue.CompleteAdding();
 		BeginEncryptionAndSigningQueue.CompleteAdding();
 		FinishEncryptionAndSigningQueue.CompleteAdding();
-		OnDemandDataQueue.CompleteAdding();
 		WriterQueue.CompleteAdding();
 		BeginCompressionThread.Wait();
 		BeginEncryptionAndSigningThread.Wait();
 		FinishEncryptionAndSigningThread.Wait();
-		BeginOnDemandDataThread.Wait();
 		WriterThread.Wait();
 		if (CompressionBufferAvailableEvent)
 		{
@@ -357,7 +354,6 @@ public:
 		BeginCompressionThread = Async(EAsyncExecution::Thread, [this]() { BeginCompressionThreadFunc(); });
 		BeginEncryptionAndSigningThread = Async(EAsyncExecution::Thread, [this]() { BeginEncryptionAndSigningThreadFunc(); });
 		FinishEncryptionAndSigningThread = Async(EAsyncExecution::Thread, [this]() { FinishEncryptionAndSigningThreadFunc(); });
-		BeginOnDemandDataThread = Async(EAsyncExecution::Thread, [this]() { BeginOnDemandDataThreadFunc(); });
 		WriterThread = Async(EAsyncExecution::Thread, [this]() { WriterThreadFunc(); });
 
 		return FIoStatus::Ok;
@@ -473,7 +469,6 @@ private:
 	void BeginCompressionThreadFunc();
 	void BeginEncryptionAndSigningThreadFunc();
 	void FinishEncryptionAndSigningThreadFunc();
-	void BeginOnDemandDataThreadFunc();
 	void WriterThreadFunc();
 
 	FIoStoreWriterSettings WriterSettings;
@@ -482,12 +477,10 @@ private:
 	TFuture<void> BeginCompressionThread;
 	TFuture<void> BeginEncryptionAndSigningThread;
 	TFuture<void> FinishEncryptionAndSigningThread;
-	TFuture<void> BeginOnDemandDataThread;
 	TFuture<void> WriterThread;
 	FIoStoreWriteQueue BeginCompressionQueue;
 	FIoStoreWriteQueue BeginEncryptionAndSigningQueue;
 	FIoStoreWriteQueue FinishEncryptionAndSigningQueue;
-	FIoStoreWriteQueue OnDemandDataQueue;
 	FIoStoreWriteQueue WriterQueue;
 	TAtomic<uint64> TotalChunksCount{ 0 };
 	TAtomic<uint64> HashedChunksCount{ 0 };
@@ -552,7 +545,7 @@ public:
 		FMemory::Memzero(&Toc.Header, sizeof(FIoStoreTocHeader));
 	}
 
-	int32 AddChunkEntry(const FIoChunkId& ChunkId, const FIoOffsetAndLength& OffsetLength, const FIoStoreTocEntryMeta& Meta, const FIoStoreTocOnDemandChunkMeta& OnDemandData)
+	int32 AddChunkEntry(const FIoChunkId& ChunkId, const FIoOffsetAndLength& OffsetLength, const FIoStoreTocEntryMeta& Meta)
 	{
 		int32& Index = ChunkIdToIndex.FindOrAdd(ChunkId);
 
@@ -561,7 +554,6 @@ public:
 			Index = Toc.ChunkIds.Add(ChunkId);
 			Toc.ChunkOffsetLengths.Add(OffsetLength);
 			Toc.ChunkMetas.Add(Meta);
-			Toc.OnDemandChunkMeta.Add(OnDemandData);
 
 			return Index;
 		}
@@ -579,11 +571,6 @@ public:
 		return Toc.ChunkBlockSignatures.AddDefaulted_GetRef();
 	}
 
-	FIoStoreTocOnDemandCompressedBlockMeta& AddCompressionBlockMetaEntry()
-	{
-		return Toc.OnDemandCompressedBlockMeta.AddDefaulted_GetRef();
-	}
-
 	uint8 AddCompressionMethodEntry(FName CompressionMethod)
 	{
 		if (CompressionMethod == NAME_None)
@@ -950,7 +937,6 @@ public:
 		Entry->BeginCompressionBarrier = FGraphEvent::CreateGraphEvent();
 		Entry->FinishCompressionBarrier = FGraphEvent::CreateGraphEvent();
 		Entry->FinishEncryptionAndSigningBarrier = FGraphEvent::CreateGraphEvent();
-		Entry->BeginOnDemandDataBarrier = FGraphEvent::CreateGraphEvent();
 		Entry->BeginWriteBarrier = FGraphEvent::CreateGraphEvent();
 		Entry->WriteFinishedEvent = FGraphEvent::CreateGraphEvent();
 		
@@ -1074,8 +1060,6 @@ public:
 		OutTocOffsetAndLengths.SetNum(ChunkCount);
 		TArray<FIoStoreTocEntryMeta> OutTocChunkMetas;
 		OutTocChunkMetas.SetNum(ChunkCount);
-		TArray<FIoStoreTocOnDemandChunkMeta> OutTocOnDemandChunkMeta;
-		OutTocOnDemandChunkMeta.SetNum(ChunkCount);
 		TArray<int32> OutTocChunkHashSeeds;
 		OutTocChunkHashSeeds.SetNumZeroed(SeedCount);
 		TArray<int32> OutTocChunkIndicesWithoutPerfectHash;
@@ -1235,7 +1219,6 @@ public:
 					OutTocChunkIds[Slot] = ChunkId;
 					OutTocOffsetAndLengths[Slot] = TocResource.ChunkOffsetLengths[ChunkIndex];
 					OutTocChunkMetas[Slot] = TocResource.ChunkMetas[ChunkIndex];
-					OutTocOnDemandChunkMeta[Slot] = TocResource.OnDemandChunkMeta[ChunkIndex];
 				}
 			}
 		}
@@ -1257,7 +1240,6 @@ public:
 				OutTocChunkIds[Slot] = ChunkId;
 				OutTocOffsetAndLengths[Slot] = TocResource.ChunkOffsetLengths[ChunkIndex];
 				OutTocChunkMetas[Slot] = TocResource.ChunkMetas[ChunkIndex];
-				OutTocOnDemandChunkMeta[Slot] = TocResource.OnDemandChunkMeta[ChunkIndex];
 			}
 		}
 
@@ -1273,7 +1255,6 @@ public:
 				OutTocChunkIds[Slot] = ChunkId;
 				OutTocOffsetAndLengths[Slot] = TocResource.ChunkOffsetLengths[OverflowEntryIndex];
 				OutTocChunkMetas[Slot] = TocResource.ChunkMetas[OverflowEntryIndex];
-				OutTocOnDemandChunkMeta[Slot] = TocResource.OnDemandChunkMeta[OverflowEntryIndex];
 				OverflowEntryIndex = Slot;
 			}
 		}
@@ -1287,7 +1268,6 @@ public:
 		TocResource.ChunkIds = MoveTemp(OutTocChunkIds);
 		TocResource.ChunkOffsetLengths = MoveTemp(OutTocOffsetAndLengths);
 		TocResource.ChunkMetas = MoveTemp(OutTocChunkMetas);
-		TocResource.OnDemandChunkMeta = MoveTemp(OutTocOnDemandChunkMeta);
 		TocResource.ChunkPerfectHashSeeds = MoveTemp(OutTocChunkHashSeeds);
 		TocResource.ChunkIndicesWithoutPerfectHash = MoveTemp(OutTocChunkIndicesWithoutPerfectHash);
 
@@ -1895,33 +1875,6 @@ private:
 		}
 	}
 
-	void BeginOnDemandData(FIoStoreWriteQueueEntry* Entry)
-	{
-		if (ContainerSettings.IsOnDemand())
-		{
-			// TODO: Worth threading this branch? (needs profiling on larger data sets)
-			FIoHashBuilder HashBuilder;
-			for (FChunkBlock& Block : Entry->ChunkBlocks)
-			{
-				check(Align(Block.CompressedSize, FAES::AESBlockSize) == Block.Size);
-
-				// Note that the IoBuffer size is not the size of the data!
-				FMemoryView BlockData = MakeMemoryView(Block.IoBuffer->Data(), Block.Size);
-		
-				HashBuilder.Update(BlockData);
-				Block.DiskHash = FIoHash::HashBuffer(BlockData);
-			}
-
-			Entry->ChunkDiskHash = HashBuilder.Finalize();
-
-			Entry->BeginWriteBarrier->DispatchSubsequents();
-		}
-		else
-		{
-			Entry->BeginWriteBarrier->DispatchSubsequents();
-		}
-	}
-
 	void WriteEntry(FIoStoreWriteQueueEntry* Entry)
 	{
 		TRACE_CPUPROFILER_EVENT_SCOPE(WriteEntry);
@@ -2053,16 +2006,9 @@ private:
 				Signature = ChunkBlock.Signature;
 			}
 
-			if (ContainerSettings.IsOnDemand())
-			{
-				FIoStoreTocOnDemandCompressedBlockMeta& Meta = TocBuilder.AddCompressionBlockMetaEntry();
-				Meta.DiskHash = ChunkBlock.DiskHash;
-			}
 		}
 
-		FIoStoreTocOnDemandChunkMeta OnDemandChunkMeta{ Entry->ChunkDiskHash };
-
-		const int32 TocEntryIndex = TocBuilder.AddChunkEntry(Entry->ChunkId, OffsetLength, ChunkMeta, OnDemandChunkMeta);
+		const int32 TocEntryIndex = TocBuilder.AddChunkEntry(Entry->ChunkId, OffsetLength, ChunkMeta);
 		check(TocEntryIndex != INDEX_NONE);
 
 		if (ContainerSettings.IsIndexed() && Entry->Options.FileName.Len() > 0)
@@ -2304,37 +2250,13 @@ void FIoStoreWriterContextImpl::FinishEncryptionAndSigningThreadFunc()
 		{
 			FIoStoreWriteQueueEntry* Next = Entry->Next;
 			Entry->FinishEncryptionAndSigningBarrier->Wait();
-			OnDemandDataQueue.Enqueue(Entry);
+			WriterQueue.Enqueue(Entry);
 			Entry->CompressedSize = 0;
 			for (const FChunkBlock& ChunkBlock : Entry->ChunkBlocks)
 			{
 				Entry->CompressedSize += ChunkBlock.Size;
 			}
-			Entry->BeginOnDemandDataBarrier->DispatchSubsequents();
-
-			Entry = Next;
-		}
-	}
-}
-
-void FIoStoreWriterContextImpl::BeginOnDemandDataThreadFunc()
-{
-	for (;;)
-	{
-		FIoStoreWriteQueueEntry* Entry = OnDemandDataQueue.DequeueOrWait();
-		if (!Entry)
-		{
-			return;
-		}
-
-		while (Entry)
-		{
-			FIoStoreWriteQueueEntry* Next = Entry->Next;
-			Entry->BeginOnDemandDataBarrier->Wait();
-
-			WriterQueue.Enqueue(Entry);
-
-			Entry->Writer->BeginOnDemandData(Entry);
+			Entry->BeginWriteBarrier->DispatchSubsequents();
 
 			Entry = Next;
 		}
@@ -2559,73 +2481,6 @@ public:
 		}
 	}
 
-	TIoStatusOr<FIoStoreCompressedChunkInfo> GetChunkCompressedInfo(const FIoChunkId& ChunkId) const
-	{
-		const int32* TocEntryIndex = TocReader.GetTocEntryIndex(ChunkId);
-		if (TocEntryIndex != nullptr)
-		{
-			// Find where in the virtual file the chunk exists.
-			const FIoOffsetAndLength* OffsetAndLength = TocReader.GetOffsetAndLength(ChunkId);
-			if (!OffsetAndLength)
-			{
-				return FIoStatus(EIoErrorCode::NotFound, TEXT("Unknown chunk ID"));
-			}
-
-			const uint64 ResolvedOffset = OffsetAndLength->GetOffset();
-			const uint64 ResolvedSize = OffsetAndLength->GetLength();
-
-			// Find what compressed blocks this read straddles.
-			const FIoStoreTocResource& TocResource = TocReader.GetTocResource();
-			const uint64 CompressionBlockSize = TocResource.Header.CompressionBlockSize;
-
-			const int32 FirstBlockIndex = int32(ResolvedOffset / CompressionBlockSize);
-			const int32 LastBlockIndex = int32((Align(ResolvedOffset + ResolvedSize, CompressionBlockSize) - 1) / CompressionBlockSize);
-
-			// Determine size of the result and set up output buffers
-			uint64 TotalCompressedSize = 0;
-			for (int32 BlockIndex = FirstBlockIndex; BlockIndex <= LastBlockIndex; ++BlockIndex)
-			{
-				const FIoStoreTocCompressedBlockEntry& CompressionBlock = TocResource.CompressionBlocks[BlockIndex];
-				TotalCompressedSize += CompressionBlock.GetCompressedSize();
-			}
-
-			FIoStoreCompressedChunkInfo Info;
-			if (!TocResource.OnDemandChunkMeta.IsEmpty())
-			{
-				Info.DiskHash = TocResource.OnDemandChunkMeta[*TocEntryIndex].DiskHash;
-			}
-
-			Info.Blocks.Reserve(LastBlockIndex + 1 - FirstBlockIndex);
-			Info.UncompressedOffset = ResolvedOffset % CompressionBlockSize;
-			Info.UncompressedSize = ResolvedSize;
-			Info.TotalCompressedSize = TotalCompressedSize;
-
-			// Set up the result blocks.
-			uint64 CurrentOffset = 0;
-			for (int32 BlockIndex = FirstBlockIndex; BlockIndex <= LastBlockIndex; ++BlockIndex)
-			{
-				const FIoStoreTocCompressedBlockEntry& CompressionBlock = TocResource.CompressionBlocks[BlockIndex];
-				FIoStoreCompressedBlockInfo& BlockInfo = Info.Blocks.AddDefaulted_GetRef();
-				if (!TocResource.OnDemandCompressedBlockMeta.IsEmpty())
-				{
-					BlockInfo.DiskHash = TocResource.OnDemandCompressedBlockMeta[BlockIndex].DiskHash;
-				}
-				BlockInfo.CompressionMethod = TocResource.CompressionMethods[CompressionBlock.GetCompressionMethodIndex()];
-				BlockInfo.CompressedSize = CompressionBlock.GetCompressedSize();
-				BlockInfo.UncompressedSize = CompressionBlock.GetUncompressedSize();
-				BlockInfo.OffsetInBuffer = CurrentOffset;
-				BlockInfo.AlignedSize = Align(CompressionBlock.GetCompressedSize(), FAES::AESBlockSize);
-				CurrentOffset += BlockInfo.AlignedSize;
-			}
-
-			return Info;
-		}
-		else
-		{
-			return FIoStatus(EIoErrorCode::NotFound, TEXT("Not found"));
-		}
-	}
-
 	UE::Tasks::TTask<TIoStatusOr<FIoBuffer>> ReadAsync(const FIoChunkId& ChunkId, const FIoReadOptions& Options) const
 	{
 		TRACE_CPUPROFILER_EVENT_SCOPE(ReadChunkAsync);
@@ -2969,10 +2824,7 @@ public:
 		{
 			const FIoStoreTocCompressedBlockEntry& CompressionBlock = TocResource.CompressionBlocks[BlockIndex];
 			FIoStoreCompressedBlockInfo& BlockInfo = Result.Blocks.AddDefaulted_GetRef();
-			if (!TocResource.OnDemandCompressedBlockMeta.IsEmpty())
-			{
-				BlockInfo.DiskHash = TocResource.OnDemandCompressedBlockMeta[BlockIndex].DiskHash;
-			}
+			
 			BlockInfo.CompressionMethod = TocResource.CompressionMethods[CompressionBlock.GetCompressionMethodIndex()];
 			BlockInfo.CompressedSize = CompressionBlock.GetCompressedSize();
 			BlockInfo.UncompressedSize = CompressionBlock.GetUncompressedSize();
@@ -3183,11 +3035,6 @@ TIoStatusOr<FIoStoreTocChunkInfo> FIoStoreReader::GetChunkInfo(const uint32 TocE
 	return Impl->GetChunkInfo(TocEntryIndex);
 }
 
-TIoStatusOr<FIoStoreCompressedChunkInfo> FIoStoreReader::GetChunkCompressedInfo(const FIoChunkId& Chunk) const
-{
-	return Impl->GetChunkCompressedInfo(Chunk);
-}
-
 TIoStatusOr<FIoBuffer> FIoStoreReader::Read(const FIoChunkId& Chunk, const FIoReadOptions& Options) const
 {
 	return Impl->Read(Chunk, Options);
@@ -3277,14 +3124,10 @@ FIoStatus FIoStoreTocResource::Read(const TCHAR* TocFilePath, EIoStoreTocReadOpt
 		return FIoStatusBuilder(EIoErrorCode::CorruptToc) << TEXT("Too new TOC header version while reading '") << TocFilePath << TEXT("'");
 	}
 
-	const bool bHasOnDemandMetaData = Header.Version >= static_cast<uint8>(EIoStoreTocVersion::OnDemandMetaData) && EnumHasAnyFlags(Header.ContainerFlags, EIoContainerFlags::OnDemand);
-
 	const uint64 TotalTocSize = TocFileHandle->Size() - sizeof(FIoStoreTocHeader);
 	const uint64 TocMetaSize = Header.TocEntryCount * sizeof(FIoStoreTocEntryMeta);
-	const uint64 TocOnDemandMetaSize = !bHasOnDemandMetaData ? 0 :	(Header.TocEntryCount * sizeof(FIoStoreTocOnDemandChunkMeta)) +
-																	(Header.TocCompressedBlockEntryCount * sizeof(FIoStoreTocOnDemandCompressedBlockMeta));
 
-	const uint64 DefaultTocSize = TotalTocSize - (Header.DirectoryIndexSize + TocMetaSize + TocOnDemandMetaSize);
+	const uint64 DefaultTocSize = TotalTocSize - (Header.DirectoryIndexSize + TocMetaSize);
 	uint64 TocSize = DefaultTocSize;
 
 	if (EnumHasAnyFlags(ReadOptions, EIoStoreTocReadOptions::ReadTocMeta))
@@ -3407,20 +3250,6 @@ FIoStatus FIoStoreTocResource::Read(const TCHAR* TocFilePath, EIoStoreTocReadOpt
 
 		const FIoStoreTocEntryMeta* ChunkMetas = reinterpret_cast<const FIoStoreTocEntryMeta*>(TocMeta);
 		OutTocResource.ChunkMetas = MakeArrayView<FIoStoreTocEntryMeta const>(ChunkMetas, Header.TocEntryCount);
-
-		// OnDemand 
-		if (bHasOnDemandMetaData)
-		{
-			uint8 const* OnDemandPtr = TocMeta + TocMetaSize;
-
-			const FIoStoreTocOnDemandChunkMeta* OnDemandChunkMeta = reinterpret_cast<const FIoStoreTocOnDemandChunkMeta*>(OnDemandPtr);
-			OutTocResource.OnDemandChunkMeta = MakeArrayView<FIoStoreTocOnDemandChunkMeta const>(OnDemandChunkMeta, Header.TocEntryCount);
-
-			OnDemandPtr += sizeof(FIoStoreTocOnDemandChunkMeta) * Header.TocEntryCount;
-
-			const FIoStoreTocOnDemandCompressedBlockMeta* OnDemandCompressedBlockMeta = reinterpret_cast<const FIoStoreTocOnDemandCompressedBlockMeta*>(OnDemandPtr);
-			OutTocResource.OnDemandCompressedBlockMeta = MakeArrayView<FIoStoreTocOnDemandCompressedBlockMeta const>(OnDemandCompressedBlockMeta, Header.TocCompressedBlockEntryCount);
-		}
 	}
 
 	if (Header.Version < static_cast<uint8>(EIoStoreTocVersion::PartitionSize))
@@ -3605,26 +3434,9 @@ TIoStatusOr<uint64> FIoStoreTocResource::Write(
 	}
 
 	// Meta data (EIoStoreTocReadOptions::ReadTocMeta)
+	if (!WriteArray(TocFileHandle.Get(), TocResource.ChunkMetas))
 	{
-		// First write out the standard meta info
-		if (!WriteArray(TocFileHandle.Get(), TocResource.ChunkMetas))
-		{
-			return FIoStatus(EIoErrorCode::WriteError, TEXT("Failed to write chunk meta data"));
-		}
-
-		// Write out the metadata for the 'OnDemand' feature if the container supports this.
-		if (EnumHasAnyFlags(TocHeader.ContainerFlags, EIoContainerFlags::OnDemand))
-		{
-			if (!WriteArray(TocFileHandle.Get(), TocResource.OnDemandChunkMeta))
-			{
-				return FIoStatus(EIoErrorCode::WriteError, TEXT("Failed to write on demand chunk meta data"));
-			}
-
-			if (!WriteArray(TocFileHandle.Get(), TocResource.OnDemandCompressedBlockMeta))
-			{
-				return FIoStatus(EIoErrorCode::WriteError, TEXT("Failed to write on demand compressed block meta data"));
-			}
-		}
+		return FIoStatus(EIoErrorCode::WriteError, TEXT("Failed to write chunk meta data"));
 	}
 
 	TocFileHandle->Flush(true);
diff --git a/Engine/Source/Runtime/Core/Public/IO/IoDispatcher.h b/Engine/Source/Runtime/Core/Public/IO/IoDispatcher.h
--- a/Engine/Source/Runtime/Core/Public/IO/IoDispatcher.h
+++ b/Engine/Source/Runtime/Core/Public/IO/IoDispatcher.h
@@ -650,12 +650,6 @@ struct FIoStoreTocCompressedBlockInfo
 
 struct FIoStoreCompressedBlockInfo
 {
-	/**
-	* Hash of the block on disk. Note that this can be all zero if the hash info was not computed when
-	* the utoc was created.
-	*/
-	FIoHash DiskHash;
-
 	/** Name of the method used to compress the block. */
 	FName CompressionMethod;
 	/** The size of relevant data in the block (i.e. what you pass to decompress). */
@@ -668,28 +662,6 @@ struct FIoStoreCompressedBlockInfo
 	uint64 OffsetInBuffer;
 };
 
-struct FIoStoreCompressedChunkInfo
-{
-	/** Info about the blocks that the chunk is split up into. */
-	TArray<FIoStoreCompressedBlockInfo> Blocks;
-
-	/**
-	* Hash of the compressed chunk on disk. Note that this can be all zero if the hash info was
-	* not computed when the utoc was created.
-	*/
-	FIoHash DiskHash;
-
-	/** There is where the data starts in IoBuffer(for when you pass in a data range via FIoReadOptions). */
-	uint64 UncompressedOffset = 0;
-	/**
-	 * This is the total size requested via FIoReadOptions. Notably, if you requested a narrow range, you could
-	 * add up all the block uncompressed sizes and it would be larger than this.
-	 */
-	uint64 UncompressedSize = 0;
-	/** This is the total size of compressed data, which is less than IoBuffer size due to padding for decryption. */
-	uint64 TotalCompressedSize = 0;
-};
-
 struct FIoStoreCompressedReadResult
 {
 	/** The buffer containing the chunk. */
@@ -799,7 +771,6 @@ public:
 	CORE_API void EnumerateChunks(TFunction<bool(FIoStoreTocChunkInfo&&)>&& Callback) const;
 	CORE_API TIoStatusOr<FIoStoreTocChunkInfo> GetChunkInfo(const FIoChunkId& Chunk) const;
 	CORE_API TIoStatusOr<FIoStoreTocChunkInfo> GetChunkInfo(const uint32 TocEntryIndex) const;
-	CORE_API TIoStatusOr<FIoStoreCompressedChunkInfo> GetChunkCompressedInfo(const FIoChunkId& Chunk) const;
 
 	// Reads the chunk off the disk, decrypting/decompressing as necessary.
 	CORE_API TIoStatusOr<FIoBuffer> Read(const FIoChunkId& Chunk, const FIoReadOptions& Options) const;
