Publish AnticheatProvider in the abnormal termination context.
Change PlatformName to be stored as TCHAR array like all other properties.

#tests locally in FrontEnd

diff --git a/Engine/Source/Programs/CrashReportClient/Private/CrashReportClientApp.cpp b/Engine/Source/Programs/CrashReportClient/Private/CrashReportClientApp.cpp
--- a/Engine/Source/Programs/CrashReportClient/Private/CrashReportClientApp.cpp
+++ b/Engine/Source/Programs/CrashReportClient/Private/CrashReportClientApp.cpp
@@ -874,19 +874,6 @@ bool FindAndCopyValue(const TMap<FString, FString>& Map, const FString& Key, TCH
 	return false;
 }
 
-template <size_t Size>
-bool FindAndCopyValue(const TMap<FString, FString>& Map, const FString& Key, char(&OutValue)[Size])
-{
-	const FString* ValueString = Map.Find(Key);
-	if (ValueString != nullptr)
-	{
-		FCStringAnsi::Strncpy(OutValue, StringCast<ANSICHAR>(**ValueString).Get(), Size);
-		return true;
-	}
-
-	return false;
-}
-
 static bool LoadTempCrashContextFromFile(FSharedCrashContext& CrashContext, uint64 ProcessID)
 {
 	const FString TempContextFilePath = FGenericCrashContext::GetTempSessionContextFilePath(ProcessID);
@@ -941,6 +928,7 @@ static bool LoadTempCrashContextFromFile(FSharedCrashContext& CrashContext, uint
 	FindAndCopyValue(ContextProperties, TEXT("Misc.PrimaryGPUBrand"), SessionContext.PrimaryGPUBrand);
 	FindAndCopyValue(ContextProperties, TEXT("Misc.OSVersionMajor"), SessionContext.OsVersion);
 	FindAndCopyValue(ContextProperties, TEXT("Misc.OSVersionMinor"), SessionContext.OsSubVersion);
+	FindAndCopyValue(ContextProperties, TEXT("Misc.AnticheatProvider"), SessionContext.AnticheatProvider);
 	FindAndParseValue(ContextProperties, TEXT("MemoryStats.AvailablePhysical"), SessionContext.MemoryStats.AvailablePhysical);
 	FindAndParseValue(ContextProperties, TEXT("MemoryStats.AvailableVirtual"), SessionContext.MemoryStats.AvailableVirtual);
 	FindAndParseValue(ContextProperties, TEXT("MemoryStats.UsedPhysical"), SessionContext.MemoryStats.UsedPhysical);
diff --git a/Engine/Source/Runtime/Core/Private/GenericPlatform/GenericPlatformCrashContext.cpp b/Engine/Source/Runtime/Core/Private/GenericPlatform/GenericPlatformCrashContext.cpp
--- a/Engine/Source/Runtime/Core/Private/GenericPlatform/GenericPlatformCrashContext.cpp
+++ b/Engine/Source/Runtime/Core/Private/GenericPlatform/GenericPlatformCrashContext.cpp
@@ -315,8 +315,8 @@ void FGenericCrashContext::Initialize()
 	NCached::Set(NCached::Session.UserName, FPlatformProcess::UserName());
 	NCached::Set(NCached::Session.DefaultLocale, *FPlatformMisc::GetDefaultLocale());
 
-	NCached::Set(NCached::Session.PlatformName, FPlatformProperties::PlatformName());
-	NCached::Set(NCached::Session.PlatformNameIni, FPlatformProperties::IniPlatformName());
+	NCached::Set(NCached::Session.PlatformName, ANSI_TO_TCHAR(FPlatformProperties::PlatformName()));
+	NCached::Set(NCached::Session.PlatformNameIni, ANSI_TO_TCHAR(FPlatformProperties::IniPlatformName()));
 	NCached::Set(NCached::Session.AttendedStatus, AttendedStatusToString(EUnattendedStatus::Unknown));
 
 	// Information that cannot be gathered if command line is not initialized (e.g. crash during static init)
@@ -658,6 +658,13 @@ void FGenericCrashContext::UpdateLocalizedStrings()
 #endif
 }
 
+void FGenericCrashContext::SetAnticheatProvider(const FString& AnticheatProvider)
+{
+	NCached::Set(NCached::Session.AnticheatProvider, *AnticheatProvider);
+
+	SerializeTempCrashContextToFile();
+}
+
 FGenericCrashContext::FGenericCrashContext(ECrashContextType InType, const TCHAR* InErrorMessage)
 	: Type(InType)
 	, CrashedThreadId(~uint32(0))
@@ -841,6 +848,7 @@ void FGenericCrashContext::SerializeSessionContext(FString& Buffer)
 	AddCrashPropertyInternal(Buffer, TEXT("Misc.PrimaryGPUBrand"), NCached::Session.PrimaryGPUBrand);
 	AddCrashPropertyInternal(Buffer, TEXT("Misc.OSVersionMajor"), NCached::Session.OsVersion);
 	AddCrashPropertyInternal(Buffer, TEXT("Misc.OSVersionMinor"), NCached::Session.OsSubVersion);
+	AddCrashPropertyInternal(Buffer, TEXT("Misc.AnticheatProvider"), NCached::Session.AnticheatProvider);
 
 	// FPlatformMemory::GetConstants is called in the GCreateMalloc, so we can assume it is always valid.
 	{
diff --git a/Engine/Source/Runtime/Core/Public/GenericPlatform/GenericPlatformCrashContext.h b/Engine/Source/Runtime/Core/Public/GenericPlatform/GenericPlatformCrashContext.h
--- a/Engine/Source/Runtime/Core/Public/GenericPlatform/GenericPlatformCrashContext.h
+++ b/Engine/Source/Runtime/Core/Public/GenericPlatform/GenericPlatformCrashContext.h
@@ -204,8 +204,9 @@ struct FSessionContext
 	TCHAR 					GameStateName[CR_MAX_GENERIC_FIELD_CHARS];
 	TCHAR 					CrashConfigFilePath[CR_MAX_DIRECTORY_CHARS];
 	TCHAR					AttendedStatus[CR_MAX_GENERIC_FIELD_CHARS];
-	char					PlatformName[CR_MAX_GENERIC_FIELD_CHARS];
-	char					PlatformNameIni[CR_MAX_GENERIC_FIELD_CHARS];
+	TCHAR					PlatformName[CR_MAX_GENERIC_FIELD_CHARS];
+	TCHAR					PlatformNameIni[CR_MAX_GENERIC_FIELD_CHARS];
+	TCHAR					AnticheatProvider[CR_MAX_GENERIC_FIELD_CHARS];
 	FPlatformMemoryStats	MemoryStats;
 };
 
@@ -610,6 +611,9 @@ public:
 	  */
 	CORE_API static void SetMemoryStats(const FPlatformMemoryStats& MemoryStats);
 
+	/** Sets the Anticheat client provider. */
+	CORE_API static void SetAnticheatProvider(const FString& AnticheatProvider);
+
 	/** Attempts to create the output report directory. */
 	CORE_API static bool CreateCrashReportDirectory(const TCHAR* CrashGUIDRoot, int32 CrashIndex, FString& OutCrashDirectoryAbsolute);
 
