FTextKey/FTextId memory crunching

 * UE_TEXTKEY_USE_SLAB_ALLOCATOR (default enabled) controls whether the internal strings are allocated in slabs, or as separate allocations
 * UE_TEXTKEY_STORE_EMBEDDED_HASH (default enabled) controls whether FTextKey/FTextId store an embedded hash value, or whether they calculate it on-demand

#rnx

diff --git a/Engine/Source/Runtime/Core/Private/Internationalization/TextKey.cpp b/Engine/Source/Runtime/Core/Private/Internationalization/TextKey.cpp
--- a/Engine/Source/Runtime/Core/Private/Internationalization/TextKey.cpp
+++ b/Engine/Source/Runtime/Core/Private/Internationalization/TextKey.cpp
@@ -1,100 +1,36 @@
 // Copyright Epic Games, Inc. All Rights Reserved.
 
 #include "Internationalization/TextKey.h"
-#include "Containers/Map.h"
+
 #include "Containers/ContainerAllocationPolicies.h"
+#include "Containers/Map.h"
+#include "CoreGlobals.h"
+#include "HAL/LowLevelMemTracker.h"
 #include "Hash/CityHash.h"
+#include "Logging/LogMacros.h"
 #include "Misc/ByteSwap.h"
-#include "HAL/LowLevelMemTracker.h"
 #include "Misc/LazySingleton.h"
 #include "Misc/ScopeLock.h"
-#include "Logging/LogMacros.h"
 
 DEFINE_LOG_CATEGORY_STATIC(LogTextKey, Log, All);
 
+#ifndef UE_TEXTKEY_USE_SLAB_ALLOCATOR
+	#define UE_TEXTKEY_USE_SLAB_ALLOCATOR (1)
+#endif
+
 class FTextKeyState
 {
 public:
-	void FindOrAdd(const TCHAR* InStr, const int32 InStrLen, const TCHAR*& OutStrPtr, uint32& OutStrHash)
-	{
-		check(*InStr != 0);
-
-		FScopeLock ScopeLock(&SynchronizationObject);
-
-		const FKeyData SrcKey(InStr, InStrLen);
-		const FString* StrPtr = KeysTable.Find(SrcKey);
-		if (!StrPtr)
-		{
-			LLM_SCOPE(ELLMTag::Localization);
-			FString StrCopy = CopyString(InStrLen, InStr); // Need to copy the string here so we can reference its internal allocation as the key
-			const FKeyData DestKey(*StrCopy, StrCopy.Len(), SrcKey.StrHash);
-			StrPtr = &KeysTable.Add(DestKey, MoveTemp(StrCopy));
-			check(DestKey.Str == **StrPtr); // The move must have moved the allocation we referenced in the key
-		}
-
-		OutStrPtr = **StrPtr;
-		OutStrHash = SrcKey.StrHash;
-	}
-
-	void FindOrAdd(const TCHAR* InStr, const int32 InStrLen, const uint32 InStrHash, const TCHAR*& OutStrPtr)
-	{
-		check(*InStr != 0);
-
-		FScopeLock ScopeLock(&SynchronizationObject);
-
-		const FKeyData SrcKey(InStr, InStrLen, InStrHash);
-		const FString* StrPtr = KeysTable.Find(SrcKey);
-		if (!StrPtr)
-		{
-			LLM_SCOPE(ELLMTag::Localization);
-			FString StrCopy = CopyString(InStrLen, InStr); // Need to copy the string here so we can reference its internal allocation as the key
-			const FKeyData DestKey(*StrCopy, StrCopy.Len(), SrcKey.StrHash);
-			StrPtr = &KeysTable.Add(DestKey, MoveTemp(StrCopy));
-			check(DestKey.Str == **StrPtr); // The move must have moved the allocation we referenced in the key
-		}
-
-		OutStrPtr = **StrPtr;
-	}
-
-	void FindOrAdd(const FString& InStr, const TCHAR*& OutStrPtr, uint32& OutStrHash)
+	void FindOrAdd(FStringView InStr, FTextKey& OutTextKey)
 	{
 		check(!InStr.IsEmpty());
-
-		FScopeLock ScopeLock(&SynchronizationObject);
-
-		const FKeyData SrcKey(*InStr, InStr.Len());
-		const FString* StrPtr = KeysTable.Find(SrcKey);
-		if (!StrPtr)
-		{
-			LLM_SCOPE(ELLMTag::Localization);
-			FString StrCopy = InStr; // Need to copy the string here so we can reference its internal allocation as the key
-			const FKeyData DestKey(*StrCopy, StrCopy.Len(), SrcKey.StrHash);
-			StrPtr = &KeysTable.Add(DestKey, MoveTemp(StrCopy));
-			check(DestKey.Str == **StrPtr); // The move must have moved the allocation we referenced in the key
-		}
-		
-		OutStrPtr = **StrPtr;
-		OutStrHash = SrcKey.StrHash;
+		FindOrAddImpl(FKeyData(InStr), OutTextKey);
 	}
 
-	void FindOrAdd(FString&& InStr, const TCHAR*& OutStrPtr, uint32& OutStrHash)
+	void FindOrAdd(FStringView InStr, const uint32 InStrHash, FTextKey& OutTextKey)
 	{
 		check(!InStr.IsEmpty());
-
-		FScopeLock ScopeLock(&SynchronizationObject);
-
-		const FKeyData SrcKey(*InStr, InStr.Len());
-		const FString* StrPtr = KeysTable.Find(SrcKey);
-		if (!StrPtr)
-		{
-			LLM_SCOPE(ELLMTag::Localization);
-			const FKeyData DestKey(*InStr, InStr.Len(), SrcKey.StrHash);
-			StrPtr = &KeysTable.Add(DestKey, MoveTemp(InStr));
-			check(DestKey.Str == **StrPtr); // The move must have moved the allocation we referenced in the key
-		}
-
-		OutStrPtr = **StrPtr;
-		OutStrHash = SrcKey.StrHash;
+		FindOrAddImpl(FKeyData(InStr, InStrHash), OutTextKey);
 	}
 
 	void Shrink()
@@ -117,20 +53,38 @@ public:
 private:
 	struct FKeyData
 	{
-		FKeyData(const TCHAR* InStr, const int32 InStrLen)
-			: Str(InStr)
-			, StrLen(InStrLen)
-			, StrHash(TextKeyUtil::HashString(InStr, InStrLen)) // Note: This hash gets serialized so *DO NOT* change it without fixing the serialization to discard the old hash method
+		explicit FKeyData(FStringView InStr)
+			: Str(InStr.GetData())
+			, StrLen(InStr.Len())
+#if UE_TEXTKEY_STORE_EMBEDDED_HASH
+			, StrHash(TextKeyUtil::HashString(InStr)) // Note: This hash gets serialized so *DO NOT* change it without fixing the serialization to discard the old hash method (also update FTextKey::GetTypeHash)
+#endif
 		{
 		}
 
-		FKeyData(const TCHAR* InStr, const int32 InStrLen, const uint32 InStrHash)
-			: Str(InStr)
-			, StrLen(InStrLen)
+		FKeyData(FStringView InStr, const uint32 InStrHash)
+			: Str(InStr.GetData())
+			, StrLen(InStr.Len())
+#if UE_TEXTKEY_STORE_EMBEDDED_HASH
 			, StrHash(InStrHash)
+#endif
 		{
 		}
 
+		FKeyData(FStringView InStr, const FKeyData& InOther)
+			: Str(InStr.GetData())
+			, StrLen(InStr.Len())
+#if UE_TEXTKEY_STORE_EMBEDDED_HASH
+			, StrHash(InOther.StrHash)
+#endif
+		{
+		}
+
+		FStringView ToView() const
+		{
+			return FStringView(Str, StrLen);
+		}
+
 		friend FORCEINLINE bool operator==(const FKeyData& A, const FKeyData& B)
 		{
 			// We can use Memcmp here as we know we're comparing two blocks of the same size and don't care about lexical ordering
@@ -145,26 +99,156 @@ private:
 
 		friend FORCEINLINE uint32 GetTypeHash(const FKeyData& A)
 		{
+#if UE_TEXTKEY_STORE_EMBEDDED_HASH
 			return A.StrHash;
+#else
+			return TextKeyUtil::HashString(A.ToView());
+#endif
 		}
 
 		const TCHAR* Str;
 		int32 StrLen;
+#if UE_TEXTKEY_STORE_EMBEDDED_HASH
 		uint32 StrHash;
+#endif
 	};
 
-	FORCEINLINE FString CopyString(const int32 InStrLen, const TCHAR* InStr)
+	void FindOrAddImpl(const FKeyData& KeyData, FTextKey& OutTextKey)
+	{
+		const TCHAR* StrPtr = FindOrAddString(KeyData);
+		check(StrPtr);
+
+		OutTextKey.StrPtr = StrPtr;
+#if UE_TEXTKEY_STORE_EMBEDDED_HASH
+		OutTextKey.StrHash = KeyData.StrHash;
+#endif
+	}
+
+	const TCHAR* FindOrAddString(const FKeyData& KeyData)
 	{
-		// We do this rather than use the FString constructor directly, 
-		// as this method avoids slack being added to the allocation
-		FString Str;
-		Str.Reserve(InStrLen);
-		Str.AppendChars(InStr, InStrLen);
-		return Str;
+		FScopeLock ScopeLock(&SynchronizationObject);
+#if UE_TEXTKEY_USE_SLAB_ALLOCATOR
+		const TCHAR* StrPtr = KeysTable.FindRef(KeyData);
+		if (!StrPtr)
+		{
+			LLM_SCOPE(ELLMTag::Localization);
+			StrPtr = StringAllocations.Add(KeyData.ToView());
+			KeysTable.Add(FKeyData(StrPtr, KeyData), StrPtr);
+		}
+		return StrPtr;
+#else
+		const FString* StrPtr = KeysTable.Find(KeyData);
+		if (!StrPtr)
+		{
+			LLM_SCOPE(ELLMTag::Localization);
+
+			// Need to copy the string here so we can reference its internal allocation as the key
+			FString StrCopy;
+			{
+				// We do this rather than use the FString constructor directly, 
+				// as this method avoids slack being added to the allocation
+				StrCopy.Reserve(KeyData.StrLen);
+				StrCopy.AppendChars(KeyData.Str, KeyData.StrLen);
+			}
+
+			FKeyData DestKey(StrCopy, KeyData);
+			StrPtr = &KeysTable.Add(DestKey, MoveTemp(StrCopy));
+			check(DestKey.Str == **StrPtr); // The move must have moved the allocation we referenced in the key
+		}
+		return **StrPtr;
+#endif
 	}
 
+#if UE_TEXTKEY_USE_SLAB_ALLOCATOR
+	class FStringSlabAllocator
+	{
+	public:
+		FStringSlabAllocator() = default;
+
+		~FStringSlabAllocator()
+		{
+			uint64 TotalNumElementsWasted = 0;
+			for (FStringSlab& Slab : Slabs)
+			{
+				TotalNumElementsWasted += (SlabSizeInElements - Slab.NumElementsUsed);
+				FMemory::Free(Slab.Allocation);
+			}
+			UE_LOG(LogLocalization, Log, TEXT("FTextKey slab allocator allocated %d slabs (of %d elements) and wasted %d elements (%d bytes)"), Slabs.Num(), SlabSizeInElements, TotalNumElementsWasted, TotalNumElementsWasted * sizeof(TCHAR));
+		}
+
+		const TCHAR* Add(FStringView InStr)
+		{
+			const int32 NumSlabElementsNeeded = InStr.Len() + 1;
+			FStringSlab& Slab = GetSlab(NumSlabElementsNeeded);
+
+			TCHAR* StringPtr = Slab.Allocation + Slab.NumElementsUsed;
+			Slab.NumElementsUsed += NumSlabElementsNeeded;
+
+			FMemory::Memcpy(StringPtr, InStr.GetData(), InStr.Len() * sizeof(TCHAR));
+			*(StringPtr + InStr.Len()) = 0;
+
+			return StringPtr;
+		}
+
+	private:
+		static const int32 SlabSizeInElements = 8192;
+
+		struct FStringSlab
+		{
+			TCHAR* Allocation = nullptr;
+			int32 NumElementsUsed = 0;
+		};
+
+		FStringSlab& GetSlab(const int32 NumSlabElementsNeeded)
+		{
+			if (Slabs.Num() > 0)
+			{
+				// Always try the last slab first
+				if (FStringSlab& Slab = Slabs.Last();
+					(Slab.NumElementsUsed + NumSlabElementsNeeded) <= SlabSizeInElements)
+				{
+					return Slab;
+				}
+
+				if (Slabs.Num() > 1)
+				{
+					// We only add to the last slab in the array, so if we've run out of space, merge it back into the array based 
+					// on its current used size (sorted most used first), and then check to see if the new last slab has space
+					FStringSlab SlabToMerge = Slabs.Pop(EAllowShrinking::No);
+					const int32 MergeIndex = Algo::LowerBound(Slabs, SlabToMerge, [](const FStringSlab& A, const FStringSlab& B)
+					{
+						return A.NumElementsUsed > B.NumElementsUsed;
+					});
+					Slabs.Insert(SlabToMerge, MergeIndex);
+					if (MergeIndex != Slabs.Num())
+					{
+						if (FStringSlab& Slab = Slabs.Last();
+							(Slab.NumElementsUsed + NumSlabElementsNeeded) <= SlabSizeInElements)
+						{
+							return Slab;
+						}
+					}
+				}
+			}
+
+			// If no slabs have space then just allocate a new one
+			checkf(NumSlabElementsNeeded <= SlabSizeInElements, TEXT("Tried to allocate a FTextKey string of %d elements, which is larger than the allowed slab size of %d elements!"), NumSlabElementsNeeded, SlabSizeInElements);
+			FStringSlab& Slab = Slabs.AddDefaulted_GetRef();
+			Slab.Allocation = reinterpret_cast<TCHAR*>(FMemory::Malloc(SlabSizeInElements * sizeof(TCHAR)));
+			return Slab;
+		}
+
+		TArray<FStringSlab> Slabs;
+	};
+#endif
+
 	FCriticalSection SynchronizationObject;
+#if UE_TEXTKEY_USE_SLAB_ALLOCATOR
+	FStringSlabAllocator StringAllocations;
+	TMap<FKeyData, const TCHAR*> KeysTable;
+#else
 	TMap<FKeyData, FString> KeysTable;
+#endif
 };
 
 namespace TextKeyUtil
@@ -284,31 +368,7 @@ FTextKey::FTextKey()
 	Reset();
 }
 
-FTextKey::FTextKey(const TCHAR* InStr)
-{
-	if (*InStr == 0)
-	{
-		Reset();
-	}
-	else
-	{
-		FTextKeyState::GetState().FindOrAdd(InStr, FCString::Strlen(InStr), StrPtr, StrHash);
-	}
-}
-
-FTextKey::FTextKey(const FString& InStr)
-{
-	if (InStr.IsEmpty())
-	{
-		Reset();
-	}
-	else
-	{
-		FTextKeyState::GetState().FindOrAdd(InStr, StrPtr, StrHash);
-	}
-}
-
-FTextKey::FTextKey(FString&& InStr)
+FTextKey::FTextKey(FStringView InStr)
 {
 	if (InStr.IsEmpty())
 	{
@@ -316,7 +376,7 @@ FTextKey::FTextKey(FString&& InStr)
 	}
 	else
 	{
-		FTextKeyState::GetState().FindOrAdd(MoveTemp(InStr), StrPtr, StrHash);
+		FTextKeyState::GetState().FindOrAdd(InStr, *this);
 	}
 }
 
@@ -333,12 +393,12 @@ void FTextKey::SerializeAsString(FArchive& Ar)
 		}
 		else
 		{
-			FTextKeyState::GetState().FindOrAdd(StrBuffer.GetData(), StrBuffer.Num() - 1, StrPtr, StrHash);
+			FTextKeyState::GetState().FindOrAdd(FStringView(StrBuffer.GetData(), StrBuffer.Num() - 1), *this);
 		}
 	}
 	else
 	{
-		TextKeyUtil::SaveKeyString(Ar, StrPtr);
+		TextKeyUtil::SaveKeyString(Ar, GetChars());
 	}
 }
 
@@ -346,7 +406,8 @@ void FTextKey::SerializeWithHash(FArchive& Ar)
 {
 	if (Ar.IsLoading())
 	{
-		Ar << StrHash;
+		uint32 TmpStrHash = 0;
+		Ar << TmpStrHash;
 
 		TextKeyUtil::FInlineStringBuffer StrBuffer;
 		TextKeyUtil::LoadKeyString(Ar, StrBuffer);
@@ -357,14 +418,15 @@ void FTextKey::SerializeWithHash(FArchive& Ar)
 		}
 		else
 		{
-			FTextKeyState::GetState().FindOrAdd(StrBuffer.GetData(), StrBuffer.Num() - 1, StrHash, StrPtr);
+			FTextKeyState::GetState().FindOrAdd(FStringView(StrBuffer.GetData(), StrBuffer.Num() - 1), TmpStrHash, *this);
 		}
 	}
 	else
 	{
-		Ar << StrHash;
+		uint32 TmpStrHash = GetTypeHash(*this);
+		Ar << TmpStrHash;
 
-		TextKeyUtil::SaveKeyString(Ar, StrPtr);
+		TextKeyUtil::SaveKeyString(Ar, GetChars());
 	}
 }
 
@@ -384,14 +446,15 @@ void FTextKey::SerializeDiscardHash(FArchive& Ar)
 		}
 		else
 		{
-			FTextKeyState::GetState().FindOrAdd(StrBuffer.GetData(), StrBuffer.Num() - 1, StrPtr, StrHash);
+			FTextKeyState::GetState().FindOrAdd(FStringView(StrBuffer.GetData(), StrBuffer.Num() - 1), *this);
 		}
 	}
 	else
 	{
-		Ar << StrHash;
+		uint32 TmpStrHash = GetTypeHash(*this);
+		Ar << TmpStrHash;
 
-		TextKeyUtil::SaveKeyString(Ar, StrPtr);
+		TextKeyUtil::SaveKeyString(Ar, GetChars());
 	}
 }
 
@@ -410,12 +473,12 @@ void FTextKey::SerializeAsString(FStructuredArchiveSlot Slot)
 			}
 			else
 			{
-				FTextKeyState::GetState().FindOrAdd(MoveTemp(TmpStr), StrPtr, StrHash);
+				FTextKeyState::GetState().FindOrAdd(TmpStr, *this);
 			}
 		}
 		else
 		{
-			FString TmpStr = StrPtr;
+			FString TmpStr = GetChars();
 			Slot << TmpStr;
 		}
 	}
@@ -434,7 +497,8 @@ void FTextKey::SerializeWithHash(FStructuredArchiveSlot Slot)
 
 		if (Slot.GetUnderlyingArchive().IsLoading())
 		{
-			Record << SA_VALUE(TEXT("Hash"), StrHash);
+			uint32 TmpStrHash = 0;
+			Record << SA_VALUE(TEXT("Hash"), TmpStrHash);
 
 			FString TmpStr;
 			Record << SA_VALUE(TEXT("Str"), TmpStr);
@@ -445,14 +509,15 @@ void FTextKey::SerializeWithHash(FStructuredArchiveSlot Slot)
 			}
 			else
 			{
-				FTextKeyState::GetState().FindOrAdd(*TmpStr, TmpStr.Len(), StrHash, StrPtr);
+				FTextKeyState::GetState().FindOrAdd(TmpStr, TmpStrHash, *this);
 			}
 		}
 		else
 		{
-			Record << SA_VALUE(TEXT("Hash"), StrHash);
+			uint32 TmpStrHash = GetTypeHash(*this);
+			Record << SA_VALUE(TEXT("Hash"), TmpStrHash);
 
-			FString TmpStr = StrPtr;
+			FString TmpStr = GetChars();
 			Record << SA_VALUE(TEXT("Str"), TmpStr);
 		}
 	}
@@ -483,14 +548,15 @@ void FTextKey::SerializeDiscardHash(FStructuredArchiveSlot Slot)
 			}
 			else
 			{
-				FTextKeyState::GetState().FindOrAdd(MoveTemp(TmpStr), StrPtr, StrHash);
+				FTextKeyState::GetState().FindOrAdd(TmpStr, *this);
 			}
 		}
 		else
 		{
-			Record << SA_VALUE(TEXT("Hash"), StrHash);
+			uint32 TmpStrHash = GetTypeHash(*this);
+			Record << SA_VALUE(TEXT("Hash"), TmpStrHash);
 
-			FString TmpStr = StrPtr;
+			FString TmpStr = GetChars();
 			Record << SA_VALUE(TEXT("Str"), TmpStr);
 		}
 	}
@@ -504,7 +570,9 @@ void FTextKey::SerializeDiscardHash(FStructuredArchiveSlot Slot)
 void FTextKey::Reset()
 {
 	StrPtr = TEXT("");
+#if UE_TEXTKEY_STORE_EMBEDDED_HASH
 	StrHash = 0;
+#endif
 }
 
 void FTextKey::CompactDataStructures()
diff --git a/Engine/Source/Runtime/Core/Public/Internationalization/TextKey.h b/Engine/Source/Runtime/Core/Public/Internationalization/TextKey.h
--- a/Engine/Source/Runtime/Core/Public/Internationalization/TextKey.h
+++ b/Engine/Source/Runtime/Core/Public/Internationalization/TextKey.h
@@ -5,9 +5,14 @@
 #include "CoreTypes.h"
 #include "Templates/TypeHash.h"
 #include "Containers/StringConv.h"
+#include "Containers/StringView.h"
 #include "Containers/UnrealString.h"
 #include "Serialization/StructuredArchive.h"
 
+#ifndef UE_TEXTKEY_STORE_EMBEDDED_HASH
+	#define UE_TEXTKEY_STORE_EMBEDDED_HASH (1)
+#endif
+
 namespace TextKeyUtil
 {
 	/** Utility to produce a hash for a UTF-16 string (as used by FTextKey) */
@@ -50,6 +55,16 @@ namespace TextKeyUtil
 	{
 		return HashString(*InStr, InStr.Len(), InBaseHash);
 	}
+
+	/** Utility to produce a hash for a string (as used by FTextKey) */
+	FORCEINLINE uint32 HashString(FStringView InStr)
+	{
+		return HashString(InStr.GetData(), InStr.Len());
+	}
+	FORCEINLINE uint32 HashString(FStringView InStr, const uint32 InBaseHash)
+	{
+		return HashString(InStr.GetData(), InStr.Len(), InBaseHash);
+	}
 }
 
 /**
@@ -60,9 +75,17 @@ class FTextKey
 {
 public:
 	CORE_API FTextKey();
-	CORE_API FTextKey(const TCHAR* InStr);
-	CORE_API FTextKey(const FString& InStr);
-	CORE_API FTextKey(FString&& InStr);
+	CORE_API FTextKey(FStringView InStr);
+	
+	FTextKey(const TCHAR* InStr)
+		: FTextKey(FStringView(InStr))
+	{
+	}
+
+	FTextKey(const FString& InStr)
+		: FTextKey(FStringView(InStr))
+	{
+	}
 
 	/** Get the underlying chars buffer this text key represents */
 	FORCEINLINE const TCHAR* GetChars() const
@@ -85,7 +108,11 @@ public:
 	/** Get the hash of this text key */
 	friend FORCEINLINE uint32 GetTypeHash(const FTextKey& A)
 	{
+#if UE_TEXTKEY_STORE_EMBEDDED_HASH
 		return A.StrHash;
+#else
+		return TextKeyUtil::HashString(A.StrPtr);
+#endif
 	}
 
 	/** Serialize this text key as if it were an FString */
@@ -122,15 +149,24 @@ public:
 	static CORE_API void TearDown();
 
 private:
-	/** Pointer to the string buffer we reference from the internal table */
+	/** Pointer to the null-terminated string buffer we reference */
 	const TCHAR* StrPtr;
 
+#if UE_TEXTKEY_STORE_EMBEDDED_HASH
 	/** Hash of this text key */
 	uint32 StrHash;
+#endif
 
-	FTextKey(const TCHAR* Str, uint32 Hash) : StrPtr(Str), StrHash(Hash) {}
+	FTextKey(const TCHAR* Str, uint32 Hash)
+		: StrPtr(Str)
+#if UE_TEXTKEY_STORE_EMBEDDED_HASH
+		, StrHash(Hash)
+#endif
+	{
+	}
 
 	friend class FTextId;
+	friend class FTextKeyState;
 };
 
 /**
@@ -145,23 +181,29 @@ public:
 	}
 
 	FTextId(const FTextKey& InNamespace, const FTextKey& InKey)
-		: NamespaceStr(InNamespace.StrPtr)
-		, KeyStr(InKey.StrPtr)
-		, NamespaceHash(InNamespace.StrHash)
-		, KeyHash(InKey.StrHash)
 	{
+		SetNamespace(InNamespace);
+		SetKey(InKey);
 	}
 
 	/** Get the namespace component of this text identity */
 	FORCEINLINE FTextKey GetNamespace() const
 	{
+#if UE_TEXTKEY_STORE_EMBEDDED_HASH
 		return FTextKey(NamespaceStr, NamespaceHash);
+#else
+		return FTextKey(NamespaceStr, 0);
+#endif
 	}
 
 	/** Get the key component of this text identity */
 	FORCEINLINE FTextKey GetKey() const
 	{
+#if UE_TEXTKEY_STORE_EMBEDDED_HASH
 		return FTextKey(KeyStr, KeyHash);
+#else
+		return FTextKey(KeyStr, 0);
+#endif
 	}
 
 	/** Compare for equality */
@@ -179,94 +221,89 @@ public:
 	/** Get the hash of this text identity */
 	friend FORCEINLINE uint32 GetTypeHash(const FTextId& A)
 	{
+#if UE_TEXTKEY_STORE_EMBEDDED_HASH
 		return HashCombine(A.NamespaceHash, A.KeyHash);
+#else
+		return HashCombine(TextKeyUtil::HashString(A.NamespaceStr), TextKeyUtil::HashString(A.KeyStr));
+#endif
 	}
 
 	/** Serialize this text identity as if it were FStrings */
 	void SerializeAsString(FArchive& Ar)
 	{
-		FTextKey Namespace = FTextKey(NamespaceStr, NamespaceHash);
+		FTextKey Namespace = GetNamespace();
 		Namespace.SerializeAsString(Ar);
-		NamespaceStr = Namespace.StrPtr;
-		NamespaceHash = Namespace.StrHash;
+		SetNamespace(Namespace);
 
-		FTextKey Key = FTextKey(KeyStr, KeyHash);
+		FTextKey Key = GetKey();
 		Key.SerializeAsString(Ar);
-		KeyStr = Key.StrPtr;
-		KeyHash = Key.StrHash;
+		SetKey(Key);
 	}
 
 	/** Serialize this text identity including its hash values (this method is sensitive to hashing algorithm changes, so only use it for generated files that can be rebuilt from another source) */
 	void SerializeWithHash(FArchive& Ar)
 	{
-		FTextKey Namespace = FTextKey(NamespaceStr, NamespaceHash);
+		FTextKey Namespace = GetNamespace();
 		Namespace.SerializeWithHash(Ar);
-		NamespaceStr = Namespace.StrPtr;
-		NamespaceHash = Namespace.StrHash;
+		SetNamespace(Namespace);
 
-		FTextKey Key = FTextKey(KeyStr, KeyHash);
+		FTextKey Key = GetKey();
 		Key.SerializeWithHash(Ar);
-		KeyStr = Key.StrPtr;
-		KeyHash = Key.StrHash;
+		SetKey(Key);
 	}
 
 	/** Serialize this text identity including its hash values, discarding the hash on load (to upgrade from an older hashing algorithm) */
 	void SerializeDiscardHash(FArchive& Ar)
 	{
-		FTextKey Namespace = FTextKey(NamespaceStr, NamespaceHash);
+		FTextKey Namespace = GetNamespace();
 		Namespace.SerializeDiscardHash(Ar);
-		NamespaceStr = Namespace.StrPtr;
-		NamespaceHash = Namespace.StrHash;
+		SetNamespace(Namespace);
 
-		FTextKey Key = FTextKey(KeyStr, KeyHash);
+		FTextKey Key = GetKey();
 		Key.SerializeDiscardHash(Ar);
-		KeyStr = Key.StrPtr;
-		KeyHash = Key.StrHash;
+		SetKey(Key);
 	}
 
 	/** Serialize this text identity as if it were FStrings */
 	void SerializeAsString(FStructuredArchiveSlot Slot)
 	{
 		FStructuredArchiveRecord Record = Slot.EnterRecord();
-		FTextKey Namespace = FTextKey(NamespaceStr, NamespaceHash);
+
+		FTextKey Namespace = GetNamespace();
 		Namespace.SerializeAsString(Record.EnterField(TEXT("Namespace")));
-		NamespaceStr = Namespace.StrPtr;
-		NamespaceHash = Namespace.StrHash;
+		SetNamespace(Namespace);
 
-		FTextKey Key = FTextKey(KeyStr, KeyHash);
+		FTextKey Key = GetKey();
 		Key.SerializeAsString(Record.EnterField(TEXT("Key")));
-		KeyStr = Key.StrPtr;
-		KeyHash = Key.StrHash;
+		SetKey(Key);
 	}
 
 	/** Serialize this text identity including its hash values (this method is sensitive to hashing algorithm changes, so only use it for generated files that can be rebuilt from another source) */
 	void SerializeWithHash(FStructuredArchiveSlot Slot)
 	{
 		FStructuredArchiveRecord Record = Slot.EnterRecord();
-		FTextKey Namespace = FTextKey(NamespaceStr, NamespaceHash);
+
+		FTextKey Namespace = GetNamespace();
 		Namespace.SerializeWithHash(Record.EnterField(TEXT("Namespace")));
-		NamespaceStr = Namespace.StrPtr;
-		NamespaceHash = Namespace.StrHash;
+		SetNamespace(Namespace);
 
-		FTextKey Key = FTextKey(KeyStr, KeyHash);
+		FTextKey Key = GetKey();
 		Key.SerializeWithHash(Record.EnterField(TEXT("Key")));
-		KeyStr = Key.StrPtr;
-		KeyHash = Key.StrHash;
+		SetKey(Key);
 	}
 
 	/** Serialize this text identity including its hash values, discarding the hash on load (to upgrade from an older hashing algorithm) */
 	void SerializeDiscardHash(FStructuredArchiveSlot Slot)
 	{
 		FStructuredArchiveRecord Record = Slot.EnterRecord();
-		FTextKey Namespace = FTextKey(NamespaceStr, NamespaceHash);
+
+		FTextKey Namespace = GetNamespace();
 		Namespace.SerializeDiscardHash(Record.EnterField(TEXT("Namespace")));
-		NamespaceStr = Namespace.StrPtr;
-		NamespaceHash = Namespace.StrHash;
+		SetNamespace(Namespace);
 
-		FTextKey Key = FTextKey(KeyStr, KeyHash);
+		FTextKey Key = GetKey();
 		Key.SerializeDiscardHash(Record.EnterField(TEXT("Key")));
-		KeyStr = Key.StrPtr;
-		KeyHash = Key.StrHash;
+		SetKey(Key);
 	}
 
 	/** Is this text identity empty? */
@@ -279,12 +316,32 @@ public:
 	FORCEINLINE void Reset()
 	{
 		NamespaceStr = KeyStr = TEXT("");
+#if UE_TEXTKEY_STORE_EMBEDDED_HASH
 		NamespaceHash = KeyHash = 0;
+#endif
 	}
 
 private:
+	void SetNamespace(const FTextKey& InNamespace)
+	{
+		NamespaceStr = InNamespace.StrPtr;
+#if UE_TEXTKEY_STORE_EMBEDDED_HASH
+		NamespaceHash = InNamespace.StrHash;
+#endif
+	}
+
+	void SetKey(const FTextKey& InKey)
+	{
+		KeyStr = InKey.StrPtr;
+#if UE_TEXTKEY_STORE_EMBEDDED_HASH
+		KeyHash = InKey.StrHash;
+#endif
+	}
+
 	const TCHAR* NamespaceStr;
 	const TCHAR* KeyStr;
+#if UE_TEXTKEY_STORE_EMBEDDED_HASH
 	uint32 NamespaceHash;
 	uint32 KeyHash;
+#endif
 };
