Move the ENABLE_LOC_TESTING display string backup to a separate map

This is only used when using the leet or keys culture, and makes the size of FDisplayStringEntry consistent between builds with and without ENABLE_LOC_TESTING enabled

#rnx

diff --git a/Engine/Source/Runtime/Core/Private/Internationalization/TextLocalizationManager.cpp b/Engine/Source/Runtime/Core/Private/Internationalization/TextLocalizationManager.cpp
--- a/Engine/Source/Runtime/Core/Private/Internationalization/TextLocalizationManager.cpp
+++ b/Engine/Source/Runtime/Core/Private/Internationalization/TextLocalizationManager.cpp
@@ -582,7 +582,7 @@ FTextLocalizationManager::~FTextLocalizationManager()
 void FTextLocalizationManager::DumpMemoryInfo() const
 {
 	{
-		FScopeLock ScopeLock(&DisplayStringLookupTableCS);
+		FScopeLock ScopeLock(&DisplayStringTableCS);
 		UE_LOG(LogTextLocalizationManager, Log, TEXT("DisplayStringLookupTable.GetAllocatedSize()=%d elems=%d"), DisplayStringLookupTable.GetAllocatedSize(), DisplayStringLookupTable.Num());
 	}
 	{
@@ -598,7 +598,7 @@ void FTextLocalizationManager::CompactDataStructures()
 
 	double StartTime = FPlatformTime::Seconds();
 	{
-		FScopeLock ScopeLock(&DisplayStringLookupTableCS);
+		FScopeLock ScopeLock(&DisplayStringTableCS);
 		DisplayStringLookupTable.Shrink();
 	}
 	{
@@ -621,7 +621,7 @@ void FTextLocalizationManager::DumpLiveTableImpl(const FString* NamespaceFilter,
 			return !Filter || Str.MatchesWildcard(*Filter, ESearchCase::IgnoreCase); // Note: This is case insensitive since its used from a debug command
 		};
 
-		FScopeLock ScopeLock(&DisplayStringLookupTableCS);
+		FScopeLock ScopeLock(&DisplayStringTableCS);
 		DisplayStringLookupTableToDump.Reserve(DisplayStringLookupTable.Num());
 		for (const auto& DisplayStringPair : DisplayStringLookupTable)
 		{
@@ -829,7 +829,7 @@ FTextConstDisplayStringPtr FTextLocalizationManager::FindDisplayString(const FTe
 		return nullptr;
 	}
 
-	FScopeLock ScopeLock(&DisplayStringLookupTableCS);
+	FScopeLock ScopeLock(&DisplayStringTableCS);
 
 	const FTextId TextId(Namespace, Key);
 
@@ -886,7 +886,7 @@ FTextConstDisplayStringRef FTextLocalizationManager::GetDisplayString(const FTex
 
 	const FString& SourceString = GetSourceStringRef();
 
-	FScopeLock ScopeLock(&DisplayStringLookupTableCS);
+	FScopeLock ScopeLock(&DisplayStringTableCS);
 
 	// Hack fix for old assets that don't have namespace/key info.
 	if (Namespace.IsEmpty() && Key.IsEmpty())
@@ -1016,7 +1016,7 @@ FTextConstDisplayStringRef FTextLocalizationManager::GetDisplayString(const FTex
 #if WITH_EDITORONLY_DATA
 bool FTextLocalizationManager::GetLocResID(const FTextKey& Namespace, const FTextKey& Key, FString& OutLocResId) const
 {
-	FScopeLock ScopeLock(&DisplayStringLookupTableCS);
+	FScopeLock ScopeLock(&DisplayStringTableCS);
 
 	const FTextId TextId(Namespace, Key);
 
@@ -1071,7 +1071,7 @@ bool FTextLocalizationManager::AddDisplayString(const FTextDisplayStringRef& Dis
 	TRACE_CPUPROFILER_EVENT_SCOPE(FTextLocalizationManager::AddDisplayString);
 	LLM_SCOPE(ELLMTag::Localization);
 
-	FScopeLock ScopeLock(&DisplayStringLookupTableCS);
+	FScopeLock ScopeLock(&DisplayStringTableCS);
 
 	const FTextId TextId(Namespace, Key);
 
@@ -1515,7 +1515,7 @@ void FTextLocalizationManager::UpdateFromNative(FTextLocalizationResource&& Text
 
 	// Lock while updating the tables
 	{
-		FScopeLock ScopeLock(&DisplayStringLookupTableCS);
+		FScopeLock ScopeLock(&DisplayStringTableCS);
 
 		DisplayStringLookupTable.Reserve(TextLocalizationResource.Entries.Num());
 
@@ -1523,7 +1523,7 @@ void FTextLocalizationManager::UpdateFromNative(FTextLocalizationResource&& Text
 		// Note: This code doesn't handle "leet-ification" itself as it is resetting everything to a known "good" state ("leet-ification" happens later on the "good" native text)
 		for (auto& EntryPair : TextLocalizationResource.Entries)
 		{
-			const FTextId TextId = EntryPair.Key;
+			const FTextId& TextId = EntryPair.Key;
 			FTextLocalizationResource::FEntry& NewEntry = EntryPair.Value;
 
 			FDisplayStringEntry* LiveEntry = DisplayStringLookupTable.Find(TextId);
@@ -1542,7 +1542,7 @@ void FTextLocalizationManager::UpdateFromNative(FTextLocalizationResource&& Text
 					LiveEntry->LocResID = NewEntry.LocResID;
 #endif	// WITH_EDITORONLY_DATA
 #if ENABLE_LOC_TESTING
-					LiveEntry->NativeStringBackup.Reset();
+					DisplayStringBackupTable.Remove(TextId);
 #endif	// ENABLE_LOC_TESTING
 				}
 			}
@@ -1566,16 +1566,17 @@ void FTextLocalizationManager::UpdateFromNative(FTextLocalizationResource&& Text
 #if ENABLE_LOC_TESTING || USE_STABLE_LOCALIZATION_KEYS
 		for (auto& DisplayStringPair : DisplayStringLookupTable)
 		{
+			const FTextId& TextId = DisplayStringPair.Key;
 			FDisplayStringEntry& LiveEntry = DisplayStringPair.Value;
 
 #if USE_STABLE_LOCALIZATION_KEYS
 			// In builds with stable keys enabled, we have to update the display strings from the "clean" version of the text (if the sources match) as this is the only version that is translated
 			{
-				const FTextKey LiveNamespace = DisplayStringPair.Key.GetNamespace();
+				const FTextKey LiveNamespace = TextId.GetNamespace();
 				const FTextKey DisplayNamespace = TextNamespaceUtil::StripPackageNamespace(LiveNamespace.GetChars());
 				if (LiveNamespace != DisplayNamespace)
 				{
-					const FDisplayStringEntry* DisplayStringEntry = DisplayStringLookupTable.Find(FTextId(DisplayNamespace, DisplayStringPair.Key.GetKey()));
+					const FDisplayStringEntry* DisplayStringEntry = DisplayStringLookupTable.Find(FTextId(DisplayNamespace, TextId.GetKey()));
 					if (DisplayStringEntry && ((LiveEntry.SourceStringHash == DisplayStringEntry->SourceStringHash) || LiveEntry.IsEmpty()))
 					{
 						// this is to account for the case where the LiveString is empty and we are overwriting the value 
@@ -1587,7 +1588,7 @@ void FTextLocalizationManager::UpdateFromNative(FTextLocalizationResource&& Text
 						LiveEntry.LocResID = DisplayStringEntry->LocResID;
 #endif	// WITH_EDITORONLY_DATA
 #if ENABLE_LOC_TESTING
-						LiveEntry.NativeStringBackup.Reset();
+						DisplayStringBackupTable.Remove(TextId);
 #endif	// ENABLE_LOC_TESTING
 					}
 				}
@@ -1596,10 +1597,10 @@ void FTextLocalizationManager::UpdateFromNative(FTextLocalizationResource&& Text
 
 #if ENABLE_LOC_TESTING
 			// Restore the pre-leet state (if any)
-			if (LiveEntry.NativeStringBackup)
+			if (FTextConstDisplayStringPtr DisplayStringBackup;
+				DisplayStringBackupTable.RemoveAndCopyValue(TextId, DisplayStringBackup))
 			{
-				LiveEntry.DisplayString = LiveEntry.NativeStringBackup.ToSharedRef();
-				LiveEntry.NativeStringBackup.Reset();
+				LiveEntry.DisplayString = DisplayStringBackup.ToSharedRef();
 #if WITH_EDITORONLY_DATA
 				LiveEntry.LocResID = FTextKey();
 #endif	// WITH_EDITORONLY_DATA
@@ -1630,14 +1631,14 @@ void FTextLocalizationManager::UpdateFromLocalizations(FTextLocalizationResource
 
 	// Lock while updating the tables
 	{
-		FScopeLock ScopeLock(&DisplayStringLookupTableCS);
+		FScopeLock ScopeLock(&DisplayStringTableCS);
 
 		DisplayStringLookupTable.Reserve(TextLocalizationResource.Entries.Num());
 
 		// Add/update entries
 		for (auto& EntryPair : TextLocalizationResource.Entries)
 		{
-			const FTextId TextId = EntryPair.Key;
+			const FTextId& TextId = EntryPair.Key;
 			FTextLocalizationResource::FEntry& NewEntry = EntryPair.Value;
 
 			FDisplayStringEntry* LiveEntry = DisplayStringLookupTable.Find(TextId);
@@ -1691,14 +1692,15 @@ void FTextLocalizationManager::UpdateFromLocalizations(FTextLocalizationResource
 		{
 			for (auto& DisplayStringPair : DisplayStringLookupTable)
 			{
+				const FTextId& TextId = DisplayStringPair.Key;
 				FDisplayStringEntry& LiveEntry = DisplayStringPair.Value;
 
 				// In builds with stable keys enabled, we have to update the display strings from the "clean" version of the text (if the sources match) as this is the only version that is translated
-				const FTextKey LiveNamespace = DisplayStringPair.Key.GetNamespace();
+				const FTextKey LiveNamespace = TextId.GetNamespace();
 				const FTextKey DisplayNamespace = TextNamespaceUtil::StripPackageNamespace(LiveNamespace.GetChars());
 				if (LiveNamespace != DisplayNamespace)
 				{
-					const FDisplayStringEntry* DisplayStringEntry = DisplayStringLookupTable.Find(FTextId(DisplayNamespace, DisplayStringPair.Key.GetKey()));
+					const FDisplayStringEntry* DisplayStringEntry = DisplayStringLookupTable.Find(FTextId(DisplayNamespace, TextId.GetKey()));
 
 					// If the source string hashes are are the same, we can replace the display string.
 					// Otherwise, it would suggest the source string has changed and the new localization may be based off of an old source string.
@@ -1868,12 +1870,13 @@ bool FTextLocalizationManager::IsLocalizationLocked() const
 void FTextLocalizationManager::LeetifyAllDisplayStrings()
 {
 	// Lock while updating the tables
-	FScopeLock ScopeLock(&DisplayStringLookupTableCS);
+	FScopeLock ScopeLock(&DisplayStringTableCS);
 
+	DisplayStringBackupTable.Reset();
 	for (auto& DisplayStringPair : DisplayStringLookupTable)
 	{
 		FDisplayStringEntry& LiveEntry = DisplayStringPair.Value;
-		LiveEntry.NativeStringBackup = LiveEntry.DisplayString;
+		DisplayStringBackupTable.Add(DisplayStringPair.Key, LiveEntry.DisplayString);
 
 		if (!LiveEntry.DisplayString->IsEmpty())
 		{
@@ -1887,12 +1890,14 @@ void FTextLocalizationManager::LeetifyAllDisplayStrings()
 void FTextLocalizationManager::KeyifyAllDisplayStrings()
 {
 	// Lock while updating the tables
-	FScopeLock ScopeLock(&DisplayStringLookupTableCS);
+	FScopeLock ScopeLock(&DisplayStringTableCS);
 
+	DisplayStringBackupTable.Reset();
 	for (auto& DisplayStringPair : DisplayStringLookupTable)
 	{
 		FDisplayStringEntry& LiveEntry = DisplayStringPair.Value;
-		LiveEntry.NativeStringBackup = LiveEntry.DisplayString;
+		DisplayStringBackupTable.Add(DisplayStringPair.Key, LiveEntry.DisplayString);
+
 		// We want to show the identity in terms of key, namespace. This is to try and fit into the constraints of UI text blocks and at least let the key component be visible to easily identify a piece of text.
 		// If the key/namespace pair is too long, the Slate.LogPaintedText cvar can be used to see the entire thing.
 		FString KeyNamespaceDisplay = FString::Printf(TEXT("%s, %s"), DisplayStringPair.Key.GetKey().GetChars(), DisplayStringPair.Key.GetNamespace().GetChars());
diff --git a/Engine/Source/Runtime/Core/Public/Internationalization/TextLocalizationManager.h b/Engine/Source/Runtime/Core/Public/Internationalization/TextLocalizationManager.h
--- a/Engine/Source/Runtime/Core/Public/Internationalization/TextLocalizationManager.h
+++ b/Engine/Source/Runtime/Core/Public/Internationalization/TextLocalizationManager.h
@@ -53,9 +53,6 @@ private:
 		FTextConstDisplayStringRef DisplayString;
 #if WITH_EDITORONLY_DATA
 		FTextKey LocResID;
-#endif
-#if ENABLE_LOC_TESTING
-		FTextConstDisplayStringPtr NativeStringBackup;
 #endif
 		uint32 SourceStringHash;
 
@@ -88,8 +85,11 @@ private:
 		return InitializedFlags != ETextLocalizationManagerInitializedFlags::None;
 	}
 
-	mutable FCriticalSection DisplayStringLookupTableCS;
+	mutable FCriticalSection DisplayStringTableCS;
 	FDisplayStringLookupTable DisplayStringLookupTable;
+#if ENABLE_LOC_TESTING
+	TMap<FTextId, FTextConstDisplayStringPtr> DisplayStringBackupTable;
+#endif
 
 	mutable FRWLock TextRevisionRW;
 	TMap<FTextId, uint16> LocalTextRevisions;
